[AutoRTFM] Update audio component map for transactional safety.

Additionally, `AudioComponentIDCounter` is now protected by the audio-component-map critical section, instead of being incremented atomically. This will allow it to cleanly roll back if a transaction failure occurs, while maintaining tightly packed audio-component IDs.

diff --git a/Engine/Source/Runtime/Engine/Classes/Components/AudioComponent.h b/Engine/Source/Runtime/Engine/Classes/Components/AudioComponent.h
--- a/Engine/Source/Runtime/Engine/Classes/Components/AudioComponent.h
+++ b/Engine/Source/Runtime/Engine/Classes/Components/AudioComponent.h
@@ -8,6 +8,7 @@
 #include "Engine/EngineTypes.h"
 #include "IAudioParameterTransmitter.h"
 #include "Math/RandomStream.h"
+#include "Misc/TransactionallySafeCriticalSection.h"
 #include "Quartz/AudioMixerQuantizedCommands.h"
 #include "Sound/QuartzQuantizationUtilities.h"
 #include "Sound/QuartzSubscription.h"
@@ -548,7 +549,7 @@ public:
 
 	static ENGINE_API uint64 AudioComponentIDCounter;
 	static ENGINE_API TMap<uint64, UAudioComponent*> AudioIDToComponentMap;
-	static ENGINE_API FCriticalSection AudioIDToComponentMapLock;
+	static ENGINE_API FTransactionallySafeCriticalSection AudioIDToComponentMapLock;
 
 private:
 	// Data to hold pending quartz commands 
diff --git a/Engine/Source/Runtime/Engine/Private/Components/AudioComponent.cpp b/Engine/Source/Runtime/Engine/Private/Components/AudioComponent.cpp
--- a/Engine/Source/Runtime/Engine/Private/Components/AudioComponent.cpp
+++ b/Engine/Source/Runtime/Engine/Private/Components/AudioComponent.cpp
@@ -8,6 +8,7 @@
 #include "Engine/Texture2D.h"
 #include "Engine/World.h"
 #include "Kismet/GameplayStatics.h"
+#include "Misc/ScopeLock.h"
 #include "Quartz/AudioMixerClockHandle.h"
 #include "Quartz/QuartzSubsystem.h"
 #include "Sound/SoundCue.h"
@@ -65,7 +66,7 @@ FAutoConsoleVariableRef CVarWorldlessGetAudioTimeBehavior(
 
 uint64 UAudioComponent::AudioComponentIDCounter = 0;
 TMap<uint64, UAudioComponent*> UAudioComponent::AudioIDToComponentMap;
-FCriticalSection UAudioComponent::AudioIDToComponentMapLock;
+FTransactionallySafeCriticalSection UAudioComponent::AudioIDToComponentMapLock;
 
 #if WITH_EDITORONLY_DATA
 static const TCHAR* GAudioSpriteAssetNameAutoActivate = TEXT("/Engine/EditorResources/AudioIcons/S_AudioComponent_AutoActivate.S_AudioComponent_AutoActivate");
@@ -118,23 +119,23 @@ UAudioComponent::UAudioComponent(const FObjectInitializer& ObjectInitializer)
 	EnvelopeFollowerReleaseTime = 100;
 
 	AudioDeviceID = INDEX_NONE;
-	AudioComponentID = FPlatformAtomics::InterlockedIncrement(reinterpret_cast<volatile int64*>(&AudioComponentIDCounter));
 
 	RandomStream.Initialize(FApp::bUseFixedSeed ? GetFName() : NAME_None);
 
 	{
 		// TODO: Consider only putting played/active components in to the map
-		FScopeLock Lock(&AudioIDToComponentMapLock);
+		UE::TScopeLock Lock(AudioIDToComponentMapLock);
+		AudioComponentID = ++AudioComponentIDCounter;
 		AudioIDToComponentMap.Add(AudioComponentID, this);
 	}
 }
 
 UAudioComponent* UAudioComponent::GetAudioComponentFromID(uint64 AudioComponentID)
 {
-	//although we should be in the game thread when calling this function, async loading makes it possible/common for these
-	//components to be constructed outside of the game thread. this means we need a lock around anything that deals with the
-	//AudioIDToComponentMap.
-	FScopeLock Lock(&AudioIDToComponentMapLock);
+	// Although we should be in the game thread when calling this function, async loading makes it possible/common for these
+	// components to be constructed outside of the game thread. This means we need a lock around anything that deals with the
+	// AudioIDToComponentMap.
+	UE::TScopeLock Lock(AudioIDToComponentMapLock);
 	return AudioIDToComponentMap.FindRef(AudioComponentID);
 }
 
@@ -149,7 +150,7 @@ void UAudioComponent::BeginDestroy()
 	ResetParameters();
 
 	{
-		FScopeLock Lock(&AudioIDToComponentMapLock);
+		UE::TScopeLock Lock(AudioIDToComponentMapLock);
 		AudioIDToComponentMap.Remove(AudioComponentID);
 	}
 
