[Selection] ISM components shouldn't create hit proxies if they are not instanced creation method

diff --git a/Engine/Source/Runtime/Engine/Private/InstancedStaticMesh.cpp b/Engine/Source/Runtime/Engine/Private/InstancedStaticMesh.cpp
--- a/Engine/Source/Runtime/Engine/Private/InstancedStaticMesh.cpp
+++ b/Engine/Source/Runtime/Engine/Private/InstancedStaticMesh.cpp
@@ -2319,7 +2319,7 @@ FMatrix UInstancedStaticMeshComponent::GetRenderMatrix() const
 
 void UInstancedStaticMeshComponent::CreateHitProxyData(TArray<TRefCountPtr<HHitProxy>>& HitProxies)
 {
-	if (GIsEditor && bHasPerInstanceHitProxies)
+	if (GIsEditor && bHasPerInstanceHitProxies && (CreationMethod == EComponentCreationMethod::Instance))
 	{
 		QUICK_SCOPE_CYCLE_COUNTER(STAT_UInstancedStaticMeshComponent_CreateHitProxyData);
 		
@@ -5119,7 +5119,7 @@ FInstanceDataFlags UInstancedStaticMeshComponent::MakeInstanceDataFlags(bool bAn
 	Flags.bHasPerInstanceRandom = bAnyMaterialHasPerInstanceRandom;
 	Flags.bHasPerInstanceCustomData = bAnyMaterialHasPerInstanceCustomData && NumCustomDataFloats != 0;
 #if WITH_EDITOR
-	Flags.bHasPerInstanceEditorData = GIsEditor != 0 && bHasPerInstanceHitProxies;
+	Flags.bHasPerInstanceEditorData = GIsEditor != 0 && bHasPerInstanceHitProxies && (CreationMethod == EComponentCreationMethod::Instance);
 #endif
 	
 	Flags.bHasPerInstanceLMSMUVBias = IsStaticLightingAllowed();
