Dynamic Anim Montage does not use Asset frame-rate
#fix ensure that we match common-frame-rate to asset frame-rate

diff --git a/Engine/Source/Runtime/Engine/Private/Animation/AnimMontage.cpp b/Engine/Source/Runtime/Engine/Private/Animation/AnimMontage.cpp
--- a/Engine/Source/Runtime/Engine/Private/Animation/AnimMontage.cpp
+++ b/Engine/Source/Runtime/Engine/Private/Animation/AnimMontage.cpp
@@ -3051,6 +3051,9 @@ UAnimMontage* UAnimMontage::CreateSlotAnimationAsDynamicMontage_WithBlendSetting
 	NewMontage->BlendProfileOut = BlendOutSettings.BlendProfile;
 
 	NewMontage->BlendOutTriggerTime = InBlendOutTriggerTime;
+
+	NewMontage->CommonTargetFrameRate = Asset->GetSamplingFrameRate();
+	
 	return NewMontage;
 }
 
