Fix for crash in MRQ-based editor workflows
#okfirgithub public
#swarm https://p4-swarm.epicgames.net/reviews/36991290

diff --git a/Engine/Source/Runtime/Engine/Private/SampleBufferIO.cpp b/Engine/Source/Runtime/Engine/Private/SampleBufferIO.cpp
--- a/Engine/Source/Runtime/Engine/Private/SampleBufferIO.cpp
+++ b/Engine/Source/Runtime/Engine/Private/SampleBufferIO.cpp
@@ -644,6 +644,8 @@ namespace Audio
 
 	void FAsyncSoundWavePCMWriteWorker::DoWork()
 	{
+	    check (Writer);
+		
 		switch (TaskType)
 		{
 			case Audio::ESoundWavePCMWriteTaskType::GenerateSoundWave:
