Disable VRS when Substrate is enabled.

diff --git a/Engine/Source/Runtime/Engine/Private/Materials/MaterialShared.cpp b/Engine/Source/Runtime/Engine/Private/Materials/MaterialShared.cpp
--- a/Engine/Source/Runtime/Engine/Private/Materials/MaterialShared.cpp
+++ b/Engine/Source/Runtime/Engine/Private/Materials/MaterialShared.cpp
@@ -1762,9 +1762,12 @@ bool FMaterialResource::HasMaterialAttributesConnected() const { return (Materia
 #endif
 EMaterialShadingRate FMaterialResource::GetShadingRate() const { return Material->ShadingRate; }
 bool FMaterialResource::IsVariableRateShadingAllowed() const {
-	// Automatically disable VRS on anything using pixel discard as coarse shading causes the whole block to get 
-	// discarded resulting in noticeable artifacts
-	return Material->bAllowVariableRateShading && !IsMasked();
+	// * Automatically disable VRS on anything using pixel discard as coarse shading causes the whole block to get 
+	//   discarded resulting in noticeable artifacts
+	// * For now disable VRS when Substrate is enabled due to incorrect data replication between neighbor pixels
+	return Material->bAllowVariableRateShading 
+		&& !IsMasked() 
+		&& !Substrate::IsSubstrateEnabled();
 }
 FString FMaterialResource::GetBaseMaterialPathName() const { return Material->GetPathName(); }
 FString FMaterialResource::GetDebugName() const
