Substrate - Fixed eye legacy shading model iris normal and tangent to work even without the specific output nodes.

#tests Tested PC dx12
#rnx

diff --git a/Engine/Shaders/Private/Definitions.usf b/Engine/Shaders/Private/Definitions.usf
--- a/Engine/Shaders/Private/Definitions.usf
+++ b/Engine/Shaders/Private/Definitions.usf
@@ -89,6 +89,14 @@
 #define SUBSTRATE_TOP_LAYER_TYPE                            uint
 #endif
 
+#ifndef SUBSTRATE_LEGACY_IRIS_NORMAL
+#define SUBSTRATE_LEGACY_IRIS_NORMAL						0
+#endif
+
+#ifndef SUBSTRATE_LEGACY_IRIS_TANGENT
+#define SUBSTRATE_LEGACY_IRIS_TANGENT						0
+#endif
+
 #ifndef MATERIAL_SHADINGMODEL_DEFAULT_LIT
 #define MATERIAL_SHADINGMODEL_DEFAULT_LIT				0
 #endif
@@ -262,7 +270,7 @@
 #endif
 
 #ifndef SUBSTRATE_MATERIAL_EXPORT_LEGACY_BLEND_MODE
-#define SUBSTRATE_MATERIAL_EXPORT_LEGACY_BLEND_MODE		0
+#define SUBSTRATE_MATERIAL_EXPORT_LEGACY_BLEND_MODE			0
 #endif
 
 #ifndef SUBSTRATE_USE_FULLYSIMPLIFIED_MATERIAL
@@ -270,10 +278,10 @@
 #endif
 
 // Those define usage should match the meaning of ESubstrateMaterialExport enum from MaterialCompiler.h
-#define SUBSTRATE_MATERIAL_EXPORT_DIFFUSE_COLOR				1
+#define SUBSTRATE_MATERIAL_EXPORT_DIFFUSE_COLOR					1
 #define SUBSTRATE_MATERIAL_EXPORT_NORMAL						2
 #define SUBSTRATE_MATERIAL_EXPORT_EMISSIVE						3
-#define SUBSTRATE_MATERIAL_EXPORT_TRANSMITTANCE				4
+#define SUBSTRATE_MATERIAL_EXPORT_TRANSMITTANCE					4
 #define SUBSTRATE_MATERIAL_EXPORT_MATERIAL_PREVIEW				5
 
 #define SUBSTRATE_MATERIAL_EXPORT_CONTEXT_OPAQUE				0
@@ -282,10 +290,10 @@
 #define SUBSTRATE_MATERIAL_EXPORT_LEGACY_BLEND_OPAQUE			0
 #define SUBSTRATE_MATERIAL_EXPORT_LEGACY_BLEND_MASKED			1
 #define SUBSTRATE_MATERIAL_EXPORT_LEGACY_BLEND_TRANSLUCENT		2
-#define SUBSTRATE_MATERIAL_EXPORT_LEGACY_BLEND_ADDITIVE		3
-#define SUBSTRATE_MATERIAL_EXPORT_LEGACY_BLEND_MODULATE		4
+#define SUBSTRATE_MATERIAL_EXPORT_LEGACY_BLEND_ADDITIVE			3
+#define SUBSTRATE_MATERIAL_EXPORT_LEGACY_BLEND_MODULATE			4
 #define SUBSTRATE_MATERIAL_EXPORT_LEGACY_BLEND_ALPHACOMPOSITE	5
-#define SUBSTRATE_MATERIAL_EXPORT_LEGACY_BLEND_ALPHAHOLDOUT	6
+#define SUBSTRATE_MATERIAL_EXPORT_LEGACY_BLEND_ALPHAHOLDOUT		6
 
 #define SUBSTRATE_MATERIAL_EXPORT_EXECUTED						(SUBSTRATE_MATERIAL_EXPORT_TYPE > 0)
 #define SUBSTRATE_MATERIAL_EXPORT_FROM_OPAQUE					(SUBSTRATE_MATERIAL_EXPORT_EXECUTED && (SUBSTRATE_MATERIAL_EXPORT_CONTEXT == SUBSTRATE_MATERIAL_EXPORT_CONTEXT_OPAQUE))
diff --git a/Engine/Shaders/Private/MaterialTemplate.ush b/Engine/Shaders/Private/MaterialTemplate.ush
--- a/Engine/Shaders/Private/MaterialTemplate.ush
+++ b/Engine/Shaders/Private/MaterialTemplate.ush
@@ -4656,20 +4656,20 @@ void GetEyeNormals(
 	inout float3 OutIrisPlaneNormal)
 {
 #if IRIS_NORMAL
-  #if NUM_MATERIAL_OUTPUTS_CLEARCOATBOTTOMNORMAL > 0
+  #if NUM_MATERIAL_OUTPUTS_CLEARCOATBOTTOMNORMAL > 0 || SUBSTRATE_LEGACY_IRIS_NORMAL > 0
 	OutIrisNormal = normalize(InClearCoatNormal);
   #else
 	OutIrisNormal = InNormal;
   #endif
 
-  #if NUM_MATERIAL_OUTPUTS_GETTANGENTOUTPUT > 0
+  #if NUM_MATERIAL_OUTPUTS_GETTANGENTOUTPUT > 0 || SUBSTRATE_LEGACY_IRIS_TANGENT > 0
 	OutIrisPlaneNormal = normalize(InCustomTangent);
   #else
 	OutIrisPlaneNormal = OutIrisNormal;
   #endif
 #else
   #if NUM_MATERIAL_OUTPUTS_GETTANGENTOUTPUT > 0
-	OutIrisNormal = normalize(InCustomTangent);
+	OutIrisNormal = normalize(InCustomTangent);	// Weird that the tangent ends up in the normal, but that was the legacy behavior.
 	OutIrisPlaneNormal = OutIrisNormal;
   #else
 	OutIrisNormal = InNormal;
diff --git a/Engine/Source/Runtime/Core/Private/UObject/FortniteMainBranchObjectVersion.cpp b/Engine/Source/Runtime/Core/Private/UObject/FortniteMainBranchObjectVersion.cpp
--- a/Engine/Source/Runtime/Core/Private/UObject/FortniteMainBranchObjectVersion.cpp
+++ b/Engine/Source/Runtime/Core/Private/UObject/FortniteMainBranchObjectVersion.cpp
@@ -14,6 +14,6 @@ TMap<FGuid, FGuid> FFortniteMainBranchObjectVersion::GetSystemGuids()
 	SystemGuids.Add(DevGuids.Niagara_LatestScriptCompileVersion, FGuid("F8D995224324470E97ED718D33939A0E"));
 	SystemGuids.Add(DevGuids.SkeletalMeshDerivedDataVersion, FGuid("B9AEAA3EE7AC4FCFB3860E8F3E06DD75"));
 	SystemGuids.Add(DevGuids.STATICMESH_DERIVEDDATA_VER, FGuid("3231A08A87A248A09D959335B0F4060A"));
-	SystemGuids.Add(DevGuids.MaterialTranslationDDCVersion, FGuid("ED57D4E5371B47AB9D15BCB57144F9D2"));
+	SystemGuids.Add(DevGuids.MaterialTranslationDDCVersion, FGuid("38F4505C9CB84A8B9C9648D2E4CF0ED6"));
 	return SystemGuids;
 }
diff --git a/Engine/Source/Runtime/Engine/Private/Materials/HLSLMaterialTranslator.cpp b/Engine/Source/Runtime/Engine/Private/Materials/HLSLMaterialTranslator.cpp
--- a/Engine/Source/Runtime/Engine/Private/Materials/HLSLMaterialTranslator.cpp
+++ b/Engine/Source/Runtime/Engine/Private/Materials/HLSLMaterialTranslator.cpp
@@ -330,6 +330,8 @@ struct FHLSLMaterialTranslator::FEnvironmentDefines
 	uint32 SubstrateClampedClosureCount;
 	bool SubstrateIsComplexSpecialPath;
 	bool bSubstrateComplexSpecialPath;
+	bool bSubstrateLegacyIrisNormal;
+	bool bSubstrateLegacyIrisTangent;
 	uint8 NumCustomizedUVs;
 	bool bTextureSampleDebug;
 	TArray<TPair<FString, int32>> SubstrateDefines;
@@ -428,6 +430,8 @@ struct FHLSLMaterialTranslator::FEnvironmentDefines
 		Ar << bSubstrateFastPath;
 		Ar << SubstrateClampedClosureCount;
 		Ar << bSubstrateComplexSpecialPath;
+		Ar << bSubstrateLegacyIrisNormal;
+		Ar << bSubstrateLegacyIrisTangent;
 		Ar << NumCustomizedUVs;
 		Ar << bTextureSampleDebug;
 		Ar << SubstrateDefines;
@@ -2896,12 +2900,13 @@ void FHLSLMaterialTranslator::GetMaterialEnvironment(EShaderPlatform InPlatform,
 		OutEnvironment.SetDefine(TEXT("SUBSTRATE_MATERIAL_EXPORT_LEGACY_BLEND_MODE"), EnvironmentDefines->SubstrateMaterialExportLegacyBlendMode);
 		OutEnvironment.SetDefine(TEXT("SUBSTRATE_OPTIMIZED_UNLIT"), EnvironmentDefines->bSubstrateOptimizedUnlit);
 
-		{
-			OutEnvironment.SetDefine(TEXT("SUBSTRATE_SINGLEPATH"), EnvironmentDefines->bSubstrateSinglePath);
-			OutEnvironment.SetDefine(TEXT("SUBSTRATE_FASTPATH"), EnvironmentDefines->bSubstrateFastPath);
-			OutEnvironment.SetDefine(TEXT("SUBSTRATE_CLAMPED_CLOSURE_COUNT"), EnvironmentDefines->SubstrateClampedClosureCount);
-			OutEnvironment.SetDefine(TEXT("SUBSTRATE_COMPLEXSPECIALPATH"), EnvironmentDefines->bSubstrateComplexSpecialPath);
-		}
+		OutEnvironment.SetDefine(TEXT("SUBSTRATE_LEGACY_IRIS_NORMAL"), EnvironmentDefines->bSubstrateLegacyIrisNormal);
+		OutEnvironment.SetDefine(TEXT("SUBSTRATE_LEGACY_IRIS_TANGENT"), EnvironmentDefines->bSubstrateLegacyIrisTangent);
+
+		OutEnvironment.SetDefine(TEXT("SUBSTRATE_SINGLEPATH"), EnvironmentDefines->bSubstrateSinglePath);
+		OutEnvironment.SetDefine(TEXT("SUBSTRATE_FASTPATH"), EnvironmentDefines->bSubstrateFastPath);
+		OutEnvironment.SetDefine(TEXT("SUBSTRATE_CLAMPED_CLOSURE_COUNT"), EnvironmentDefines->SubstrateClampedClosureCount);
+		OutEnvironment.SetDefine(TEXT("SUBSTRATE_COMPLEXSPECIALPATH"), EnvironmentDefines->bSubstrateComplexSpecialPath);
 	}
 
 	for (int i = 0; i < EnvironmentDefines->SubstrateDefines.Num(); ++i)
@@ -15976,6 +15981,8 @@ void FHLSLMaterialTranslator::PrepareEnvironmentDefines()
 			EnvironmentDefines->bSubstrateComplexSpecialPath = SubstrateCtx.SubstrateMaterialComplexity.IsComplexSpecial();
 		}
 
+		bool EyeIrisNormalPluggedIn = false;
+		bool EyeIrisTangentPluggedIn = false;
 
 		FString SubstrateMaterialDescription;
 		for (uint32 SubstrateCompilationContextIndex = 0; SubstrateCompilationContextIndex < ESubstrateCompilationContext::SCC_MAX; ++SubstrateCompilationContextIndex)
@@ -16037,6 +16044,10 @@ void FHLSLMaterialTranslator::PrepareEnvironmentDefines()
 				for (int32 OpIt = 0; OpIt < SubstrateCtx.SubstrateMaterialExpressionRegisteredOperators.Num(); ++OpIt)
 				{
 					FSubstrateOperator& BSDFOperator = SubstrateCtx.SubstrateMaterialExpressionRegisteredOperators[OpIt];
+
+					EyeIrisNormalPluggedIn |= !!(BSDFOperator.BSDFFeatures & ESubstrateBsdfFeature_EyeIrisNormalPluggedIn);
+					EyeIrisTangentPluggedIn |= !!(BSDFOperator.BSDFFeatures & ESubstrateBsdfFeature_EyeIrisTangentPluggedIn);
+
 					if (BSDFOperator.BSDFIndex == INDEX_NONE || BSDFOperator.IsDiscarded())
 					{
 						continue;	// not a BSDF or if discarded (i.e. not the root of a parameter blending subtree), then there is no local basis to register
@@ -16094,6 +16105,9 @@ void FHLSLMaterialTranslator::PrepareEnvironmentDefines()
 #endif // WITH_EDITOR
 		}
 		MaterialCompilationOutput.SubstrateMaterialCompilationOutput.SubstrateMaterialDescription = SubstrateMaterialDescription;
+
+		EnvironmentDefines->bSubstrateLegacyIrisNormal = EyeIrisNormalPluggedIn;
+		EnvironmentDefines->bSubstrateLegacyIrisTangent = EyeIrisTangentPluggedIn;
 	}
 
 	EnvironmentDefines->bTextureSampleDebug = IsDebugTextureSampleEnabled();
diff --git a/Engine/Source/Runtime/Engine/Private/Materials/MaterialExpressionSubstrate.cpp b/Engine/Source/Runtime/Engine/Private/Materials/MaterialExpressionSubstrate.cpp
--- a/Engine/Source/Runtime/Engine/Private/Materials/MaterialExpressionSubstrate.cpp
+++ b/Engine/Source/Runtime/Engine/Private/Materials/MaterialExpressionSubstrate.cpp
@@ -428,7 +428,7 @@ uint32 UMaterialExpressionSubstrateShadingModels::GetInputType(int32 InputIndex)
 	else if (InputIndex == 14) return MCT_Float3; // WaterAbsorptionCoefficients
 	else if (InputIndex == 15) return MCT_Float1; // WaterPhaseG
 	else if (InputIndex == 16) return MCT_Float3; // ColorScaleBehindWater
-	else if (InputIndex == 17) return MCT_Float3; // ClearCoatNormal
+	else if (InputIndex == 17) return MCT_Float3; // ClearCoatNormal / IrisNormal
 	else if (InputIndex == 18) return MCT_Float3; // CustomTangent
 	else if (InputIndex == 19) return MCT_Float1; // ThinTranslucentSurfaceCoverage
 	else if (InputIndex == 20) return MCT_ShadingModel; // ShadingModel
@@ -504,8 +504,15 @@ FName UMaterialExpressionSubstrateShadingModels::GetInputName(int32 InputIndex)
 	else if (InputIndex == 14)	return TEXT("Water Absorption Coefficients");
 	else if (InputIndex == 15)	return TEXT("Water Phase G");
 	else if (InputIndex == 16)	return TEXT("Color Scale BehindWater");
-	else if (InputIndex == 17)	return TEXT("Clear Coat Normal");
-	else if (InputIndex == 18)	return TEXT("Custom Tangent");
+	else if (InputIndex == 17)	return TEXT("Iris Normal");
+	else if (InputIndex == 18)
+	{
+		if (!bShadingModelFromExpression && ShadingModelOverride == MSM_Eye)
+		{
+			return TEXT("Iris Tangent");
+		}
+		return TEXT("Custom Tangent");
+	}
 	else if (InputIndex == 19)	return TEXT("Thin Translucent Surface Coverage");
 	else if (InputIndex == 20)	return TEXT("Single Shading Model");
 	else if (InputIndex == 21)	return TEXT("Shading Model From Expression");
@@ -592,6 +599,14 @@ FSubstrateOperator* UMaterialExpressionSubstrateShadingModels::SubstrateGenerate
 	// Get the shading models resulting from the UMaterial::RebuildShadingModelField().
 	FMaterialShadingModelField ShadingModels = Compiler->GetMaterialShadingModels();
 
+	bool bEyeIrisNormalPluggedIn = ClearCoatNormal.IsConnected();
+	bool bEyeIrisTangentPluggedIn = CustomTangent.IsConnected();
+	auto ApplyEyeIrisUsedFeatures = [&](FSubstrateOperator* Operator)
+	{
+		if (bEyeIrisNormalPluggedIn)	Operator->BSDFFeatures |= ESubstrateBsdfFeature_EyeIrisNormalPluggedIn;
+		if (bEyeIrisTangentPluggedIn)	Operator->BSDFFeatures |= ESubstrateBsdfFeature_EyeIrisTangentPluggedIn;
+	};
+
 	// Logic about shading models and complexity should match UMaterialExpressionSubstrateShadingModels::Compile.
 	const bool bHasShadingModelFromExpression = ShadingModel.IsConnected(); // We keep HasShadingModelFromExpression in case all shading models cannot be safely recovered from material functions.
 	if ((ShadingModels.CountShadingModels() > 1) || bHasShadingModelFromExpression)
@@ -610,7 +625,12 @@ FSubstrateOperator* UMaterialExpressionSubstrateShadingModels::SubstrateGenerate
 		uint32 BsdfFeatures = ESubstrateBsdfFeature_SSS | ESubstrateBsdfFeature_Fuzz;
 		if (ShadingModels.HasShadingModel(MSM_Eye))  { BsdfFeatures |= ESubstrateBsdfFeature_Eye; }
 		if (ShadingModels.HasShadingModel(MSM_Hair)) { BsdfFeatures |= ESubstrateBsdfFeature_Hair; }
-		return AddDefaultWorstCase(BsdfFeatures);
+		FSubstrateOperator* Operator = AddDefaultWorstCase(BsdfFeatures);
+		if (ShadingModels.HasShadingModel(MSM_Eye) || bHasShadingModelFromExpression)
+		{
+			ApplyEyeIrisUsedFeatures(Operator);
+		}
+		return Operator;
 	}
 	// else
 	{
@@ -675,6 +695,8 @@ FSubstrateOperator* UMaterialExpressionSubstrateShadingModels::SubstrateGenerate
 			FSubstrateOperator& Operator = Compiler->SubstrateCompilationRegisterOperator(SUBSTRATE_OPERATOR_BSDF_LEGACY, Compiler->SubstrateTreeStackGetPathUniqueId(), this, Parent, Compiler->SubstrateTreeStackGetParentPathUniqueId());
 			Operator.BSDFType = SUBSTRATE_BSDF_TYPE_EYE;
 			Operator.BSDFFeatures = ESubstrateBsdfFeature_Eye;
+			ApplyEyeIrisUsedFeatures(&Operator);
+
 			Operator.ThicknessIndex = ThicknessIndex;
 			Operator.bBSDFWritesEmissive = EmissiveColor.IsConnected();
 			return &Operator;
diff --git a/Engine/Source/Runtime/Engine/Public/Materials/MaterialExpressionSubstrate.h b/Engine/Source/Runtime/Engine/Public/Materials/MaterialExpressionSubstrate.h
--- a/Engine/Source/Runtime/Engine/Public/Materials/MaterialExpressionSubstrate.h
+++ b/Engine/Source/Runtime/Engine/Public/Materials/MaterialExpressionSubstrate.h
@@ -165,7 +165,7 @@ class UMaterialExpressionSubstrateShadingModels : public UMaterialExpressionSubs
 	FExpressionInput ColorScaleBehindWater;
 
 	/**
-	 * Take the bottom clear coat surface normal as input. The normal is considered tangent or world space according to the space properties on the main material node. (type = float3, unit = unitless, defaults to vertex normal)
+	 * The iris normal. (type = float3, unit = unitless, defaults to vertex normal)
 	 */
 	UPROPERTY()
 	FExpressionInput ClearCoatNormal;
diff --git a/Engine/Source/Runtime/Engine/Public/Rendering/SubstrateMaterialShared.h b/Engine/Source/Runtime/Engine/Public/Rendering/SubstrateMaterialShared.h
--- a/Engine/Source/Runtime/Engine/Public/Rendering/SubstrateMaterialShared.h
+++ b/Engine/Source/Runtime/Engine/Public/Rendering/SubstrateMaterialShared.h
@@ -36,7 +36,9 @@ enum ESubstrateBsdfFeature
 	ESubstrateBsdfFeature_Glint                            = 1u<<6u,
 	ESubstrateBsdfFeature_SpecularProfile                  = 1u<<7u,
 	ESubstrateBsdfFeature_Eye                              = 1u<<8u,
-	ESubstrateBsdfFeature_Hair                             = 1u<<9u
+	ESubstrateBsdfFeature_EyeIrisNormalPluggedIn           = 1u<<9u,
+	ESubstrateBsdfFeature_EyeIrisTangentPluggedIn          = 1u<<10u,
+	ESubstrateBsdfFeature_Hair                             = 1u<<11u
 };
 
 struct FSubstrateOperator
