Fix invalid LOD arrays synchronization between StaticMesh & StaticMeshComponent in UStaticMeshComponent::UpdateStaticLightingData
Also made sure to inialize the MapBuildDataID in there

diff --git a/Engine/Source/Runtime/Engine/Private/Components/StaticMeshComponent.cpp b/Engine/Source/Runtime/Engine/Private/Components/StaticMeshComponent.cpp
--- a/Engine/Source/Runtime/Engine/Private/Components/StaticMeshComponent.cpp
+++ b/Engine/Source/Runtime/Engine/Private/Components/StaticMeshComponent.cpp
@@ -3667,10 +3667,15 @@ void UStaticMeshComponent::UpdateStaticLightingData()
 	if (!bAllowStaticLighting)
 		return;
 
-	if (HasStaticLighting() && LODData.IsEmpty() && StaticMesh)
+	if (HasStaticLighting() && StaticMesh && (LODData.Num() != StaticMesh->GetNumLODs()))
 	{		
-		SetLODDataCount(StaticMesh->GetNumLODs(), LODData.Num());
+		SetLODDataCount(StaticMesh->GetNumLODs(), LODData.Num());	
 	}
+
+	for (int32 LODIndex = 0; LODIndex < LODData.Num(); LODIndex++)
+	{
+		LODData[LODIndex].CreateMapBuildDataId(LODIndex);
+	}	
 }
 
 void UStaticMeshComponent::UpdateMapBuildDataId()
