Add logging when recreating cloth actors for Skeletal Meshes with the name of the asset.

Hoping this will help us track down a crash in live.

diff --git a/Engine/Source/Runtime/Engine/Private/SkeletalMeshComponentPhysics.cpp b/Engine/Source/Runtime/Engine/Private/SkeletalMeshComponentPhysics.cpp
--- a/Engine/Source/Runtime/Engine/Private/SkeletalMeshComponentPhysics.cpp
+++ b/Engine/Source/Runtime/Engine/Private/SkeletalMeshComponentPhysics.cpp
@@ -2777,6 +2777,7 @@ void USkeletalMeshComponent::RecreateClothingActors()
 			if (World && World->bShouldSimulatePhysics && World->GetPhysicsScene())
 #endif
 			{
+				UE_LOG(LogSkeletalMesh, Log, TEXT("USkeletalMeshComponent: Recreating Clothing Actors for '%s' with '%s'"), *GetName(), *SkelMesh->GetName());
 				TArray<UClothingAssetBase*> AssetsInUse;
 				SkelMesh->GetClothingAssetsInUse(AssetsInUse);
 
