First Person: Changed UPrimitiveComponent::bIsFirstPerson into an enum to support additional values in the future. This also replaces SetIsFirstPerson() with SetFirstPersonPrimitiveType(). Note that there's intentionally no deprecation handling in this CL as the feature is still experimental and no public full engine release has shipped with it yet.

#rnx

#changelist validated
#virtualized

diff --git a/Engine/Source/Runtime/Engine/Classes/Components/PrimitiveComponent.h b/Engine/Source/Runtime/Engine/Classes/Components/PrimitiveComponent.h
--- a/Engine/Source/Runtime/Engine/Classes/Components/PrimitiveComponent.h
+++ b/Engine/Source/Runtime/Engine/Classes/Components/PrimitiveComponent.h
@@ -124,6 +124,16 @@ enum class ERayTracingGroupCullingPriority : uint8
 	CP_8_QUICKLY_CULL UMETA(DisplayName = "8 - Quickly cull")
 };
 
+/** Determines how a primitive interacts with first person rendering. */
+UENUM()
+enum class EFirstPersonPrimitiveType : uint8
+{
+	/** Primitive does not interact with first person rendering. */
+	None UMETA(DisplayName = "None"),
+	/** Primitive is rendered as first person and affected by first person properties on the camera. */
+	FirstPerson UMETA(DisplayName = "First Person"),
+};
+
 /** Converts a stencil mask from the editor's USTRUCT version to the version the renderer uses. */
 struct FRendererStencilMaskEvaluation
 {
@@ -650,9 +660,9 @@ public:
 	UPROPERTY()
 	uint8 bRayTracingFarField : 1;
 
-	/** If this is true, the camera FirstPersonFieldOfView and FirstPersonScale parameters will be used on this component. These parameters can be used to render the component with a different field of view and a smaller depth range such that clipping with the scene can be avoided. This is useful for rendering first person view geometry. */
-	UPROPERTY(EditAnywhere, AdvancedDisplay, BlueprintReadWrite, Category = Rendering)
-	uint8 bIsFirstPerson : 1;
+	/** If this is set to FirstPerson, the camera FirstPersonFieldOfView and FirstPersonScale parameters will be used on this component. These parameters can be used to render the component with a different field of view and a smaller depth range such that clipping with the scene can be avoided. This is useful for rendering first person view geometry. */
+	UPROPERTY(EditAnywhere, AdvancedDisplay, BlueprintReadOnly, Category = Rendering)
+	EFirstPersonPrimitiveType FirstPersonPrimitiveType;
 
 protected:
 	/** Result of last call to AreAllCollideableDescendantsRelative(). */
@@ -1974,9 +1984,9 @@ public:
 	UFUNCTION(BlueprintCallable, Category = "Rendering")
 	ENGINE_API void SetHiddenInSceneCapture(bool bValue);
 
-	/** Sets bIsFirstPerson property and marks the render state dirty. */
+	/** Sets FirstPersonPrimitiveType property and marks the render state dirty. */
 	UFUNCTION(BlueprintCallable, Category = "Rendering")
-	ENGINE_API void SetIsFirstPerson(bool bValue);
+	ENGINE_API void SetFirstPersonPrimitiveType(EFirstPersonPrimitiveType Value);
 
 	/**
 	 * Count of all component overlap events (begin or end) ever generated for any components.
diff --git a/Engine/Source/Runtime/Engine/Private/Components/PrimitiveComponent.cpp b/Engine/Source/Runtime/Engine/Private/Components/PrimitiveComponent.cpp
--- a/Engine/Source/Runtime/Engine/Private/Components/PrimitiveComponent.cpp
+++ b/Engine/Source/Runtime/Engine/Private/Components/PrimitiveComponent.cpp
@@ -408,7 +408,7 @@ UPrimitiveComponent::UPrimitiveComponent(const FObjectInitializer& ObjectInitial
 	bVisibleInSceneCaptureOnly = false;
 	bHiddenInSceneCapture = false;
 
-	bIsFirstPerson = false;
+	FirstPersonPrimitiveType = EFirstPersonPrimitiveType::None;
 }
 
 bool UPrimitiveComponent::UsesOnlyUnlitMaterials() const
@@ -1228,7 +1228,7 @@ void UPrimitiveComponent::PostEditChangeProperty(FPropertyChangedEvent& Property
 		}
 
 		// bIsFirstPerson can be toggled at runtime and needs to propagate to the scene proxy.
-		if (PropertyName == GET_MEMBER_NAME_CHECKED(UPrimitiveComponent, bIsFirstPerson))
+		if (PropertyName == GET_MEMBER_NAME_CHECKED(UPrimitiveComponent, FirstPersonPrimitiveType))
 		{
 			MarkRenderStateDirty();
 		}
@@ -4608,11 +4608,11 @@ void UPrimitiveComponent::SetHiddenInSceneCapture(bool bValue)
 	}
 }
 
-void UPrimitiveComponent::SetIsFirstPerson(bool bValue)
+void UPrimitiveComponent::SetFirstPersonPrimitiveType(EFirstPersonPrimitiveType Value)
 {
-	if (bIsFirstPerson != bValue)
+	if (FirstPersonPrimitiveType != Value)
 	{
-		bIsFirstPerson = bValue;
+		FirstPersonPrimitiveType = Value;
 		MarkRenderStateDirty();
 	}
 }
diff --git a/Engine/Source/Runtime/Engine/Private/PrimitiveSceneProxy.cpp b/Engine/Source/Runtime/Engine/Private/PrimitiveSceneProxy.cpp
--- a/Engine/Source/Runtime/Engine/Private/PrimitiveSceneProxy.cpp
+++ b/Engine/Source/Runtime/Engine/Private/PrimitiveSceneProxy.cpp
@@ -308,7 +308,7 @@ void FPrimitiveSceneProxyDesc::InitializeFromPrimitiveComponent(const UPrimitive
 	bRayTracingFarField = InComponent->bRayTracingFarField;
 	bHoldout = InComponent->bHoldout;
 	bWantsEditorEffects = InComponent->bWantsEditorEffects;
-	bIsFirstPerson = InComponent->bIsFirstPerson;
+	bIsFirstPerson = InComponent->FirstPersonPrimitiveType == EFirstPersonPrimitiveType::FirstPerson;
 	
 	bIsVisible = InComponent->IsVisible();
 	bIsVisibleEditor = InComponent->GetVisibleFlag();
