Fix crash in Material interface VTableHelper constructor.

diff --git a/Engine/Source/Editor/UnrealEd/Private/MaterialGraph.cpp b/Engine/Source/Editor/UnrealEd/Private/MaterialGraph.cpp
--- a/Engine/Source/Editor/UnrealEd/Private/MaterialGraph.cpp
+++ b/Engine/Source/Editor/UnrealEd/Private/MaterialGraph.cpp
@@ -45,6 +45,7 @@ UMaterialGraph::UMaterialGraph(const FObjectInitializer& ObjectInitializer)
 UMaterialGraph::UMaterialGraph() = default;
 
 UMaterialGraph::UMaterialGraph(FVTableHelper& Helper)
+	: Super(Helper)
 {
 }
 
diff --git a/Engine/Source/Runtime/Engine/Private/Materials/MaterialInterface.cpp b/Engine/Source/Runtime/Engine/Private/Materials/MaterialInterface.cpp
--- a/Engine/Source/Runtime/Engine/Private/Materials/MaterialInterface.cpp
+++ b/Engine/Source/Runtime/Engine/Private/Materials/MaterialInterface.cpp
@@ -150,6 +150,7 @@ namespace UE::MaterialInterface
 UMaterialInterface::UMaterialInterface() = default;
 
 UMaterialInterface::UMaterialInterface(FVTableHelper& Helper)
+	: Super(Helper)
 {
 }
 
@@ -2049,8 +2050,9 @@ UMaterialInterfaceEditorOnlyData::UMaterialInterfaceEditorOnlyData()
 }
 
 UMaterialInterfaceEditorOnlyData::UMaterialInterfaceEditorOnlyData(FVTableHelper& Helper)
+	: Super(Helper)
 {
-}
+} 
 
 UMaterialInterfaceEditorOnlyData::~UMaterialInterfaceEditorOnlyData()
 {
