Fixed bug in HLSLMaterialTranslator causing some materials not to translate properly.

Issue was caused by hacky code in the Translator outputting the pixel members locals code epilog based on the last evaluated property instead of the property that contains the largest chunk. This caused an issue when some expression was connected to the Displacement property (the last to be evaluated), but the chunk of code that generates this output is shorter than that needed by other material outputs, resulting in some Locals to be referenced by other material outputs but not outputted to the HLSL shader.

More details:
The if in 3295 (before fix) is cutting out the generation of the code that evalutes the final value that is fed to a material output to the Local used by that material output. Unfortunately, another part of the translator decides which output's "prologue" code to use simply selecting the last property (fixed code).

That won't do, as if you connect something to Displacement (the last property) and this happens to use less code than some other previous property, locals used by the other properties will be cut off and not output to the HLSL.

The fix selects the largest chunk of code in the evaluated properties instead.

diff --git a/Engine/Source/Runtime/Engine/Private/Materials/HLSLMaterialTranslator.cpp b/Engine/Source/Runtime/Engine/Private/Materials/HLSLMaterialTranslator.cpp
--- a/Engine/Source/Runtime/Engine/Private/Materials/HLSLMaterialTranslator.cpp
+++ b/Engine/Source/Runtime/Engine/Private/Materials/HLSLMaterialTranslator.cpp
@@ -2988,6 +2988,7 @@ TBitArray<> FHLSLMaterialTranslator::GetVertexInterpolatorsOffsets(FString& Vert
 void FHLSLMaterialTranslator::GetSharedInputsMaterialCode(FString& PixelMembersDeclaration, FString& NormalAssignment, FString& PixelMembersInitializationEpilog, ECompiledPartialDerivativeVariation DerivativeVariation)
 {
 	int32 LastProperty = -1;
+	int32 TranslatedCodeChunkDefinitionsLargestSize = 0;
 	FString PixelInputInitializerValues;
 	FString NormalInitializerValue;
 
@@ -3008,15 +3009,17 @@ void FHLSLMaterialTranslator::GetSharedInputsMaterialCode(FString& PixelMembersD
 
 		if (!bEnableExecutionFlow)
 		{
-		// Normal requires its own separate initializer
-		if (Property == MP_Normal)
-		{
+			// Normal requires its own separate initializer
+			if (Property == MP_Normal)
+			{
 				NormalInitializerValue = FString::Printf(TEXT("\tPixelMaterialInputs.%s = %s;\n"), *PropertyName, *DerivativeVariations[DerivativeVariation].TranslatedCodeChunks[Property]);
-		}
-		else
-		{
-				if (DerivativeVariations[DerivativeVariation].TranslatedCodeChunkDefinitions[Property].Len() > 0)
+			}
+			else
+			{
+				const FString& TranslatedCodeChunkDefinitions = DerivativeVariations[DerivativeVariation].TranslatedCodeChunkDefinitions[Property];
+				if (TranslatedCodeChunkDefinitions.Len() >= TranslatedCodeChunkDefinitionsLargestSize)
 				{
+					TranslatedCodeChunkDefinitionsLargestSize = TranslatedCodeChunkDefinitions.Len();
 					LastProperty = Property;
 				}
 			}
