Build Fix

Fix error caught by build machine - duplicate cvar name

diff --git a/Engine/Source/Runtime/Engine/Private/ISMPartition/ISMComponentDescriptor.cpp b/Engine/Source/Runtime/Engine/Private/ISMPartition/ISMComponentDescriptor.cpp
--- a/Engine/Source/Runtime/Engine/Private/ISMPartition/ISMComponentDescriptor.cpp
+++ b/Engine/Source/Runtime/Engine/Private/ISMPartition/ISMComponentDescriptor.cpp
@@ -16,7 +16,7 @@
 #include "Algo/Transform.h"
 
 bool GISMComponentDescriptorCopiesCanEverAffectNavigation = true;
-static FAutoConsoleVariableRef CVarDefaultUseSubObjectReplicationList(
+static FAutoConsoleVariableRef CVarSMComponentDescriptorCopiesCanEverAffectNavigation(
 	TEXT("ISM.ComponentDescriptor.CanEverAffectNavigation"),
 	GISMComponentDescriptorCopiesCanEverAffectNavigation,
 	TEXT("If true CanEverAffectNavigation property will be copied from the StaticMesh and applied to the ISM"));
