Renamed the "Max" parameter in RandomInteger and RandomInteger64 to "Max Exclusive".

This only affects Blueprints. The "Max" sometimes throws off users thinking that it means inclusive max, when it's actually exclusive. While the tooltip does mention this, renaming the parameter helps enforce the meaning.

diff --git a/Engine/Source/Runtime/Engine/Classes/Kismet/KismetMathLibrary.h b/Engine/Source/Runtime/Engine/Classes/Kismet/KismetMathLibrary.h
--- a/Engine/Source/Runtime/Engine/Classes/Kismet/KismetMathLibrary.h
+++ b/Engine/Source/Runtime/Engine/Classes/Kismet/KismetMathLibrary.h
@@ -377,7 +377,7 @@ class UKismetMathLibrary : public UBlueprintFunctionLibrary
 
 	/** Returns a uniformly distributed random number between 0 and Max - 1 */
 	UFUNCTION(BlueprintPure, Category="Math|Random", meta=(NotBlueprintThreadSafe))
-	static ENGINE_API int32 RandomInteger(int32 Max);
+	static ENGINE_API int32 RandomInteger(UPARAM(DisplayName = "Exclusive Max") int32 Max);
 
 	/** Return a random integer between Min and Max (>= Min and <= Max) */
 	UFUNCTION(BlueprintPure, Category="Math|Random", meta = (NotBlueprintThreadSafe))
@@ -480,7 +480,7 @@ class UKismetMathLibrary : public UBlueprintFunctionLibrary
 
 	/** Returns a uniformly distributed random number between 0 and Max - 1 */
 	UFUNCTION(BlueprintPure, Category="Math|Random", meta=(NotBlueprintThreadSafe))
-	static ENGINE_API int64 RandomInteger64(int64 Max);
+	static ENGINE_API int64 RandomInteger64(UPARAM(DisplayName = "Exclusive Max") int64 Max);
 
 	/** Return a random integer64 between Min and Max (>= Min and <= Max) */
 	UFUNCTION(BlueprintPure, Category="Math|Random", meta = (NotBlueprintThreadSafe))
