[Engine]
* Changing ParticleModule system to use context for spawn/update/etc in order to try to remove direct access to UComponent. Goal is to get rid of dependency so Cascade can run without needing a component

diff --git a/Engine/Source/Editor/Cascade/Private/CascadePreviewViewportClient.cpp b/Engine/Source/Editor/Cascade/Private/CascadePreviewViewportClient.cpp
--- a/Engine/Source/Editor/Cascade/Private/CascadePreviewViewportClient.cpp
+++ b/Engine/Source/Editor/Cascade/Private/CascadePreviewViewportClient.cpp
@@ -580,7 +580,7 @@ void FCascadeEdPreviewViewportClient::Draw(const FSceneView* View, FPrimitiveDra
 					UParticleModule* Module = LODLevel->Modules[j];
 					if (Module && Module->bSupported3DDrawMode && Module->b3DDrawMode)
 					{
-						Module->Render3DPreview(EmitterInst, View, PDI);
+						Module->Render3DPreview({ *EmitterInst, View, PDI });
 					}
 				}
 			}
diff --git a/Engine/Source/Runtime/Engine/Classes/Particles/Acceleration/ParticleModuleAcceleration.h b/Engine/Source/Runtime/Engine/Classes/Particles/Acceleration/ParticleModuleAcceleration.h
--- a/Engine/Source/Runtime/Engine/Classes/Particles/Acceleration/ParticleModuleAcceleration.h
+++ b/Engine/Source/Runtime/Engine/Classes/Particles/Acceleration/ParticleModuleAcceleration.h
@@ -48,8 +48,8 @@ class UParticleModuleAcceleration : public UParticleModuleAccelerationBase
 
 	//Begin UParticleModule Interface
 	virtual void CompileModule( struct FParticleEmitterBuildInfo& EmitterInfo ) override;
-	virtual void Spawn(FParticleEmitterInstance* Owner, int32 Offset, float SpawnTime, FBaseParticle* ParticleBase) override;
-	virtual void Update(FParticleEmitterInstance* Owner, int32 Offset, float DeltaTime) override;
+	virtual void Spawn(const FSpawnContext& Context) override;
+	virtual void Update(const FUpdateContext& Context) override;
 	virtual uint32 RequiredBytes(UParticleModuleTypeDataBase* TypeData) override;
 
 #if WITH_EDITOR
diff --git a/Engine/Source/Runtime/Engine/Classes/Particles/Acceleration/ParticleModuleAccelerationConstant.h b/Engine/Source/Runtime/Engine/Classes/Particles/Acceleration/ParticleModuleAccelerationConstant.h
--- a/Engine/Source/Runtime/Engine/Classes/Particles/Acceleration/ParticleModuleAccelerationConstant.h
+++ b/Engine/Source/Runtime/Engine/Classes/Particles/Acceleration/ParticleModuleAccelerationConstant.h
@@ -25,8 +25,8 @@ class UParticleModuleAccelerationConstant : public UParticleModuleAccelerationBa
 
 	//Begin UParticleModule Interface
 	virtual void CompileModule( struct FParticleEmitterBuildInfo& EmitterInfo ) override;
-	virtual void Spawn(FParticleEmitterInstance* Owner, int32 Offset, float SpawnTime, FBaseParticle* ParticleBase) override;
-	virtual void Update(FParticleEmitterInstance* Owner, int32 Offset, float DeltaTime) override;
+	virtual void Spawn(const FSpawnContext& Context) override;
+	virtual void Update(const FUpdateContext& Context) override;
 	//End UParticleModule Interface
 };
 
diff --git a/Engine/Source/Runtime/Engine/Classes/Particles/Acceleration/ParticleModuleAccelerationDrag.h b/Engine/Source/Runtime/Engine/Classes/Particles/Acceleration/ParticleModuleAccelerationDrag.h
--- a/Engine/Source/Runtime/Engine/Classes/Particles/Acceleration/ParticleModuleAccelerationDrag.h
+++ b/Engine/Source/Runtime/Engine/Classes/Particles/Acceleration/ParticleModuleAccelerationDrag.h
@@ -40,7 +40,7 @@ class UParticleModuleAccelerationDrag : public UParticleModuleAccelerationBase
 
 	//Begin UParticleModule Interface
 	virtual void CompileModule( struct FParticleEmitterBuildInfo& EmitterInfo ) override;
-	virtual void Update( FParticleEmitterInstance* Owner, int32 Offset, float DeltaTime ) override;
+	virtual void Update(const FUpdateContext& Context) override;
 	//End UParticleModule Interface
 
 	virtual void PostLoad() override;
diff --git a/Engine/Source/Runtime/Engine/Classes/Particles/Acceleration/ParticleModuleAccelerationOverLifetime.h b/Engine/Source/Runtime/Engine/Classes/Particles/Acceleration/ParticleModuleAccelerationOverLifetime.h
--- a/Engine/Source/Runtime/Engine/Classes/Particles/Acceleration/ParticleModuleAccelerationOverLifetime.h
+++ b/Engine/Source/Runtime/Engine/Classes/Particles/Acceleration/ParticleModuleAccelerationOverLifetime.h
@@ -38,7 +38,7 @@ class UParticleModuleAccelerationOverLifetime : public UParticleModuleAccelerati
 	//End UObject Interface
 
 	//Begin UParticleModule Interface
-	virtual void Update(FParticleEmitterInstance* Owner, int32 Offset, float DeltaTime) override;
+	virtual void Update(const FUpdateContext& Context) override;
 	//End UParticleModule Interface
 };
 
diff --git a/Engine/Source/Runtime/Engine/Classes/Particles/Attractor/ParticleModuleAttractorLine.h b/Engine/Source/Runtime/Engine/Classes/Particles/Attractor/ParticleModuleAttractorLine.h
--- a/Engine/Source/Runtime/Engine/Classes/Particles/Attractor/ParticleModuleAttractorLine.h
+++ b/Engine/Source/Runtime/Engine/Classes/Particles/Attractor/ParticleModuleAttractorLine.h
@@ -43,8 +43,8 @@ class UParticleModuleAttractorLine : public UParticleModuleAttractorBase
 	//End UObject Interface
 
 	//Begin UParticleModule Interface
-	virtual void Update(FParticleEmitterInstance* Owner, int32 Offset, float DeltaTime) override;
-	virtual void Render3DPreview(FParticleEmitterInstance* Owner, const FSceneView* View,FPrimitiveDrawInterface* PDI) override;
+	virtual void Update(const FUpdateContext& Context) override;
+	virtual void Render3DPreview(const FPreviewContext& Context) override;
 	//End UParticleModule Interface
 };
 
diff --git a/Engine/Source/Runtime/Engine/Classes/Particles/Attractor/ParticleModuleAttractorParticle.h b/Engine/Source/Runtime/Engine/Classes/Particles/Attractor/ParticleModuleAttractorParticle.h
--- a/Engine/Source/Runtime/Engine/Classes/Particles/Attractor/ParticleModuleAttractorParticle.h
+++ b/Engine/Source/Runtime/Engine/Classes/Particles/Attractor/ParticleModuleAttractorParticle.h
@@ -90,8 +90,8 @@ class UParticleModuleAttractorParticle : public UParticleModuleAttractorBase
 	//End UObject Interface
 
 	//Begin UParticleModule Interface
-	virtual void Spawn(FParticleEmitterInstance* Owner, int32 Offset, float SpawnTime, FBaseParticle* ParticleBase) override;
-	virtual void Update(FParticleEmitterInstance* Owner, int32 Offset, float DeltaTime) override;
+	virtual void Spawn(const FSpawnContext& Context) override;
+	virtual void Update(const FUpdateContext& Context) override;
 	virtual uint32 RequiredBytes(UParticleModuleTypeDataBase* TypeData) override;
 	//End UParticleModule Interface
 };
diff --git a/Engine/Source/Runtime/Engine/Classes/Particles/Attractor/ParticleModuleAttractorPoint.h b/Engine/Source/Runtime/Engine/Classes/Particles/Attractor/ParticleModuleAttractorPoint.h
--- a/Engine/Source/Runtime/Engine/Classes/Particles/Attractor/ParticleModuleAttractorPoint.h
+++ b/Engine/Source/Runtime/Engine/Classes/Particles/Attractor/ParticleModuleAttractorPoint.h
@@ -80,8 +80,8 @@ class UParticleModuleAttractorPoint : public UParticleModuleAttractorBase
 	//End UObject Interface
 
 	//Begin UParticleModule Interface
-	virtual void Update(FParticleEmitterInstance* Owner, int32 Offset, float DeltaTime) override;
-	virtual void Render3DPreview(FParticleEmitterInstance* Owner, const FSceneView* View,FPrimitiveDrawInterface* PDI) override;
+	virtual void Update(const FUpdateContext& Context) override;
+	virtual void Render3DPreview(const FPreviewContext& Context) override;
 	//End UParticleModule Interface
 };
 
diff --git a/Engine/Source/Runtime/Engine/Classes/Particles/Attractor/ParticleModuleAttractorPointGravity.h b/Engine/Source/Runtime/Engine/Classes/Particles/Attractor/ParticleModuleAttractorPointGravity.h
--- a/Engine/Source/Runtime/Engine/Classes/Particles/Attractor/ParticleModuleAttractorPointGravity.h
+++ b/Engine/Source/Runtime/Engine/Classes/Particles/Attractor/ParticleModuleAttractorPointGravity.h
@@ -49,7 +49,7 @@ class UParticleModuleAttractorPointGravity : public UParticleModuleAttractorBase
 
 	//Begin UParticleModule Interface
 	virtual void CompileModule(struct FParticleEmitterBuildInfo& EmitterInfo) override;
-	virtual void Render3DPreview(FParticleEmitterInstance* Owner, const FSceneView* View,FPrimitiveDrawInterface* PDI) override;
+	virtual void Render3DPreview(const FPreviewContext& Context) override;
 	//End UParticleModule Interface
 };
 
diff --git a/Engine/Source/Runtime/Engine/Classes/Particles/Beam/ParticleModuleBeamModifier.h b/Engine/Source/Runtime/Engine/Classes/Particles/Beam/ParticleModuleBeamModifier.h
--- a/Engine/Source/Runtime/Engine/Classes/Particles/Beam/ParticleModuleBeamModifier.h
+++ b/Engine/Source/Runtime/Engine/Classes/Particles/Beam/ParticleModuleBeamModifier.h
@@ -109,8 +109,8 @@ class UParticleModuleBeamModifier : public UParticleModuleBeamBase
 	//End Uobject Interface
 
 	//Begin UParticleModule Interface
-	virtual void Spawn(FParticleEmitterInstance* Owner, int32 Offset, float SpawnTime, FBaseParticle* ParticleBase) override;
-	virtual void Update(FParticleEmitterInstance* Owner, int32 Offset, float DeltaTime) override;
+	virtual void Spawn(const FSpawnContext& Context) override;
+	virtual void Update(const FUpdateContext& Context) override;
 	virtual uint32 RequiredBytes(UParticleModuleTypeDataBase* TypeData) override;
 	virtual void AutoPopulateInstanceProperties(UParticleSystemComponent* PSysComp) override;
 	virtual void GetCurveObjects(TArray<FParticleCurvePair>& OutCurves) override;
diff --git a/Engine/Source/Runtime/Engine/Classes/Particles/Beam/ParticleModuleBeamNoise.h b/Engine/Source/Runtime/Engine/Classes/Particles/Beam/ParticleModuleBeamNoise.h
--- a/Engine/Source/Runtime/Engine/Classes/Particles/Beam/ParticleModuleBeamNoise.h
+++ b/Engine/Source/Runtime/Engine/Classes/Particles/Beam/ParticleModuleBeamNoise.h
@@ -135,8 +135,8 @@ class UParticleModuleBeamNoise : public UParticleModuleBeamBase
 	//End UObject Interface
 
 	//Begin UParticleModule Interface
-	virtual void Spawn(FParticleEmitterInstance* Owner, int32 Offset, float SpawnTime, FBaseParticle* ParticleBase) override;
-	virtual void Update(FParticleEmitterInstance* Owner, int32 Offset, float DeltaTime) override;
+	virtual void Spawn(const FSpawnContext& Context) override;
+	virtual void Update(const FUpdateContext& Context) override;
 	virtual void SetToSensibleDefaults(UParticleEmitter* Owner) override;
 	//End UParticleModule Interface
 
diff --git a/Engine/Source/Runtime/Engine/Classes/Particles/Beam/ParticleModuleBeamSource.h b/Engine/Source/Runtime/Engine/Classes/Particles/Beam/ParticleModuleBeamSource.h
--- a/Engine/Source/Runtime/Engine/Classes/Particles/Beam/ParticleModuleBeamSource.h
+++ b/Engine/Source/Runtime/Engine/Classes/Particles/Beam/ParticleModuleBeamSource.h
@@ -77,8 +77,8 @@ class UParticleModuleBeamSource : public UParticleModuleBeamBase
 	//End UObject Interface
 
 	//Begin UParticleModule Interface
-	virtual void Spawn(FParticleEmitterInstance* Owner, int32 Offset, float SpawnTime, FBaseParticle* ParticleBase) override;
-	virtual void Update(FParticleEmitterInstance* Owner, int32 Offset, float DeltaTime) override;
+	virtual void Spawn(const FSpawnContext& Context) override;
+	virtual void Update(const FUpdateContext& Context) override;
 	virtual uint32 RequiredBytes(UParticleModuleTypeDataBase* TypeData) override;
 	virtual void AutoPopulateInstanceProperties(UParticleSystemComponent* PSysComp) override;
 	virtual void GetParticleSysParamsUtilized(TArray<FString>& ParticleSysParamList) override;
diff --git a/Engine/Source/Runtime/Engine/Classes/Particles/Beam/ParticleModuleBeamTarget.h b/Engine/Source/Runtime/Engine/Classes/Particles/Beam/ParticleModuleBeamTarget.h
--- a/Engine/Source/Runtime/Engine/Classes/Particles/Beam/ParticleModuleBeamTarget.h
+++ b/Engine/Source/Runtime/Engine/Classes/Particles/Beam/ParticleModuleBeamTarget.h
@@ -80,8 +80,8 @@ class UParticleModuleBeamTarget : public UParticleModuleBeamBase
 	//End UObject Interface
 
 	//Begin UParticleModule Interface
-	virtual void Spawn(FParticleEmitterInstance* Owner, int32 Offset, float SpawnTime, FBaseParticle* ParticleBase) override;
-	virtual void Update(FParticleEmitterInstance* Owner, int32 Offset, float DeltaTime) override;
+	virtual void Spawn(const FSpawnContext& Context) override;
+	virtual void Update(const FUpdateContext& Context) override;
 	virtual void AutoPopulateInstanceProperties(UParticleSystemComponent* PSysComp) override;
 	virtual void GetParticleSysParamsUtilized(TArray<FString>& ParticleSysParamList) override;
 	//End UParticleModule Interface
diff --git a/Engine/Source/Runtime/Engine/Classes/Particles/Camera/ParticleModuleCameraOffset.h b/Engine/Source/Runtime/Engine/Classes/Particles/Camera/ParticleModuleCameraOffset.h
--- a/Engine/Source/Runtime/Engine/Classes/Particles/Camera/ParticleModuleCameraOffset.h
+++ b/Engine/Source/Runtime/Engine/Classes/Particles/Camera/ParticleModuleCameraOffset.h
@@ -59,8 +59,8 @@ class UParticleModuleCameraOffset : public UParticleModuleCameraBase
 	//End UObject Interface
 
 	//Begin UParticleModule Interface
-	virtual void Spawn(FParticleEmitterInstance* Owner, int32 Offset, float SpawnTime, FBaseParticle* ParticleBase) override;
-	virtual void Update(FParticleEmitterInstance* Owner, int32 Offset, float DeltaTime) override;
+	virtual void Spawn(const FSpawnContext& Context) override;
+	virtual void Update(const FUpdateContext& Context) override;
 	virtual uint32 RequiredBytes(UParticleModuleTypeDataBase* TypeData) override;
 	//End UParticleModule Interface
 
diff --git a/Engine/Source/Runtime/Engine/Classes/Particles/Collision/ParticleModuleCollision.h b/Engine/Source/Runtime/Engine/Classes/Particles/Collision/ParticleModuleCollision.h
--- a/Engine/Source/Runtime/Engine/Classes/Particles/Collision/ParticleModuleCollision.h
+++ b/Engine/Source/Runtime/Engine/Classes/Particles/Collision/ParticleModuleCollision.h
@@ -173,8 +173,8 @@ class UParticleModuleCollision : public UParticleModuleCollisionBase
 	//End UObject Interface
 
 	//Begin UParticleModule Interface
-	virtual void Spawn(FParticleEmitterInstance* Owner, int32 Offset, float SpawnTime, FBaseParticle* ParticleBase) override;
-	virtual void Update(FParticleEmitterInstance* Owner, int32 Offset, float DeltaTime) override;
+	virtual void Spawn(const FSpawnContext& Context) override;
+	virtual void Update(const FUpdateContext& Context) override;
 	virtual uint32 RequiredBytes(UParticleModuleTypeDataBase* TypeData) override;
 	virtual uint32 RequiredBytesPerInstance() override;
 	virtual uint32 PrepPerInstanceBlock(FParticleEmitterInstance* Owner, void* InstData) override;
diff --git a/Engine/Source/Runtime/Engine/Classes/Particles/Color/ParticleModuleColor.h b/Engine/Source/Runtime/Engine/Classes/Particles/Color/ParticleModuleColor.h
--- a/Engine/Source/Runtime/Engine/Classes/Particles/Color/ParticleModuleColor.h
+++ b/Engine/Source/Runtime/Engine/Classes/Particles/Color/ParticleModuleColor.h
@@ -47,7 +47,7 @@ class UParticleModuleColor : public UParticleModuleColorBase
 
 	//Begin UParticleModule Interface
 	ENGINE_API virtual	bool AddModuleCurvesToEditor(UInterpCurveEdSetup* EdSetup, TArray<const FCurveEdEntry*>& OutCurveEntries) override;
-	ENGINE_API virtual void Spawn(FParticleEmitterInstance* Owner, int32 Offset, float SpawnTime, FBaseParticle* ParticleBase) override;
+	ENGINE_API virtual void Spawn(const FSpawnContext& Context) override;
 	ENGINE_API virtual void CompileModule( struct FParticleEmitterBuildInfo& EmitterInfo ) override;
 	ENGINE_API virtual void SetToSensibleDefaults(UParticleEmitter* Owner) override;
 	//End UParticleModule Interface
@@ -60,7 +60,7 @@ class UParticleModuleColor : public UParticleModuleColorBase
 	 *	@param	SpawnTime			The time of the spawn
 	 *	@param	InRandomStream		The random stream to use for retrieving random values
 	 */
-	ENGINE_API void SpawnEx(FParticleEmitterInstance* Owner, int32 Offset, float SpawnTime, struct FRandomStream* InRandomStream, FBaseParticle* ParticleBase);
+	ENGINE_API void SpawnEx(const FSpawnContext& Context, struct FRandomStream* InRandomStream);
 
 };
 
diff --git a/Engine/Source/Runtime/Engine/Classes/Particles/Color/ParticleModuleColorOverLife.h b/Engine/Source/Runtime/Engine/Classes/Particles/Color/ParticleModuleColorOverLife.h
--- a/Engine/Source/Runtime/Engine/Classes/Particles/Color/ParticleModuleColorOverLife.h
+++ b/Engine/Source/Runtime/Engine/Classes/Particles/Color/ParticleModuleColorOverLife.h
@@ -45,8 +45,8 @@ class UParticleModuleColorOverLife : public UParticleModuleColorBase
 
 	//Begin UParticleModule Interface
 	virtual	bool AddModuleCurvesToEditor(UInterpCurveEdSetup* EdSetup, TArray<const FCurveEdEntry*>& OutCurveEntries) override;
-	virtual void Spawn(FParticleEmitterInstance* Owner, int32 Offset, float SpawnTime, FBaseParticle* ParticleBase) override;
-	virtual void Update(FParticleEmitterInstance* Owner, int32 Offset, float DeltaTime) override;
+	virtual void Spawn(const FSpawnContext& Context) override;
+	virtual void Update(const FUpdateContext& Context) override;
 	virtual void CompileModule( struct FParticleEmitterBuildInfo& EmitterInfo ) override;
 	virtual void SetToSensibleDefaults(UParticleEmitter* Owner) override;
 	//End UParticleModule Interface
diff --git a/Engine/Source/Runtime/Engine/Classes/Particles/Color/ParticleModuleColorScaleOverLife.h b/Engine/Source/Runtime/Engine/Classes/Particles/Color/ParticleModuleColorScaleOverLife.h
--- a/Engine/Source/Runtime/Engine/Classes/Particles/Color/ParticleModuleColorScaleOverLife.h
+++ b/Engine/Source/Runtime/Engine/Classes/Particles/Color/ParticleModuleColorScaleOverLife.h
@@ -48,8 +48,8 @@ class UParticleModuleColorScaleOverLife : public UParticleModuleColorBase
 
 	//Begin UParticleModule Interface
 	virtual void CompileModule( struct FParticleEmitterBuildInfo& EmitterInfo ) override;
-	virtual void Spawn(FParticleEmitterInstance* Owner, int32 Offset, float SpawnTime, FBaseParticle* ParticleBase) override;
-	virtual void Update(FParticleEmitterInstance* Owner, int32 Offset, float DeltaTime) override;
+	virtual void Spawn(const FSpawnContext& Context) override;
+	virtual void Update(const FUpdateContext& Context) override;
 	virtual void SetToSensibleDefaults(UParticleEmitter* Owner) override;
 #if WITH_EDITOR
 	virtual int32 GetNumberOfCustomMenuOptions() const override;
diff --git a/Engine/Source/Runtime/Engine/Classes/Particles/Event/ParticleModuleEventGenerator.h b/Engine/Source/Runtime/Engine/Classes/Particles/Event/ParticleModuleEventGenerator.h
--- a/Engine/Source/Runtime/Engine/Classes/Particles/Event/ParticleModuleEventGenerator.h
+++ b/Engine/Source/Runtime/Engine/Classes/Particles/Event/ParticleModuleEventGenerator.h
@@ -85,8 +85,8 @@ class UParticleModuleEventGenerator : public UParticleModuleEventBase
 	//End UObject Interface
 
 	//Begin UParticleModule Interface
-	virtual void	Spawn(FParticleEmitterInstance* Owner, int32 Offset, float SpawnTime, FBaseParticle* ParticleBase) override;
-	virtual void	Update(FParticleEmitterInstance* Owner, int32 Offset, float DeltaTime) override;
+	virtual void	Spawn(const FSpawnContext& Context) override;
+	virtual void	Update(const FUpdateContext& Context) override;
 	virtual uint32	RequiredBytes(UParticleModuleTypeDataBase* TypeData) override;
 	virtual uint32	RequiredBytesPerInstance() override;
 	virtual uint32	PrepPerInstanceBlock(FParticleEmitterInstance* Owner, void* InstData) override;
diff --git a/Engine/Source/Runtime/Engine/Classes/Particles/Kill/ParticleModuleKillBox.h b/Engine/Source/Runtime/Engine/Classes/Particles/Kill/ParticleModuleKillBox.h
--- a/Engine/Source/Runtime/Engine/Classes/Particles/Kill/ParticleModuleKillBox.h
+++ b/Engine/Source/Runtime/Engine/Classes/Particles/Kill/ParticleModuleKillBox.h
@@ -50,8 +50,8 @@ class UParticleModuleKillBox : public UParticleModuleKillBase
 	//End UObject Interface
 
 	//Begin UParticleModule Interface
-	virtual void	Update(FParticleEmitterInstance* Owner, int32 Offset, float DeltaTime) override;
-	virtual void	Render3DPreview(FParticleEmitterInstance* Owner, const FSceneView* View,FPrimitiveDrawInterface* PDI) override;
+	virtual void	Update(const FUpdateContext& Context) override;
+	virtual void	Render3DPreview(const FPreviewContext& Context) override;
 	//End UParticleModule Interface
 };
 
diff --git a/Engine/Source/Runtime/Engine/Classes/Particles/Kill/ParticleModuleKillHeight.h b/Engine/Source/Runtime/Engine/Classes/Particles/Kill/ParticleModuleKillHeight.h
--- a/Engine/Source/Runtime/Engine/Classes/Particles/Kill/ParticleModuleKillHeight.h
+++ b/Engine/Source/Runtime/Engine/Classes/Particles/Kill/ParticleModuleKillHeight.h
@@ -46,8 +46,8 @@ class UParticleModuleKillHeight : public UParticleModuleKillBase
 	//End UObject Interface
 
 	//Begin UParticleModule Interface
-	virtual void	Update(FParticleEmitterInstance* Owner, int32 Offset, float DeltaTime) override;
-	virtual void	Render3DPreview(FParticleEmitterInstance* Owner, const FSceneView* View,FPrimitiveDrawInterface* PDI) override;
+	virtual void	Update(const FUpdateContext& Context) override;
+	virtual void	Render3DPreview(const FPreviewContext& Context) override;
 	//End UParticleModule Interface
 };
 
diff --git a/Engine/Source/Runtime/Engine/Classes/Particles/Lifetime/ParticleModuleLifetime.h b/Engine/Source/Runtime/Engine/Classes/Particles/Lifetime/ParticleModuleLifetime.h
--- a/Engine/Source/Runtime/Engine/Classes/Particles/Lifetime/ParticleModuleLifetime.h
+++ b/Engine/Source/Runtime/Engine/Classes/Particles/Lifetime/ParticleModuleLifetime.h
@@ -33,13 +33,13 @@ class UParticleModuleLifetime : public UParticleModuleLifetimeBase
 
 	//Begin UParticleModule Interface
 	virtual void CompileModule( struct FParticleEmitterBuildInfo& EmitterInfo ) override;
-	virtual void Spawn(FParticleEmitterInstance* Owner, int32 Offset, float SpawnTime, FBaseParticle* ParticleBase) override;
+	virtual void Spawn(const FSpawnContext& Context) override;
 	virtual void SetToSensibleDefaults(UParticleEmitter* Owner) override;
 	//End UParticleModule Interface
 
 	//~ Begin UParticleModuleLifetimeBase Interface
 	virtual float GetMaxLifetime() override;
-	virtual float GetLifetimeValue(FParticleEmitterInstance* Owner, float InTime, UObject* Data = NULL) override;
+	virtual float GetLifetimeValue(const FContext& Context, float InTime, UObject* Data = NULL) override;
 	//~ End UParticleModuleLifetimeBase Interface
 
 	/**
@@ -50,7 +50,7 @@ class UParticleModuleLifetime : public UParticleModuleLifetimeBase
 	 *	@param	SpawnTime			The time of the spawn
 	 *	@param	InRandomStream		The random stream to use for retrieving random values
 	 */
-	void SpawnEx(FParticleEmitterInstance* Owner, int32 Offset, float SpawnTime, struct FRandomStream* InRandomStream, FBaseParticle* ParticleBase);
+	void SpawnEx(const FSpawnContext& Context, struct FRandomStream* InRandomStream);
 };
 
 
diff --git a/Engine/Source/Runtime/Engine/Classes/Particles/Lifetime/ParticleModuleLifetimeBase.h b/Engine/Source/Runtime/Engine/Classes/Particles/Lifetime/ParticleModuleLifetimeBase.h
--- a/Engine/Source/Runtime/Engine/Classes/Particles/Lifetime/ParticleModuleLifetimeBase.h
+++ b/Engine/Source/Runtime/Engine/Classes/Particles/Lifetime/ParticleModuleLifetimeBase.h
@@ -31,7 +31,7 @@ class UParticleModuleLifetimeBase : public UParticleModule
 	 *
 	 *	@return	float		The Lifetime value
 	 */
-	virtual float	GetLifetimeValue(FParticleEmitterInstance* Owner, float InTime, UObject* Data = NULL)
+	virtual float	GetLifetimeValue(const FContext& Context, float InTime, UObject* Data = NULL)
 		PURE_VIRTUAL(UParticleModuleLifetimeBase::GetLifetimeValue,return 0.0f;);
 };
 
diff --git a/Engine/Source/Runtime/Engine/Classes/Particles/Lifetime/ParticleModuleLifetime_Seeded.h b/Engine/Source/Runtime/Engine/Classes/Particles/Lifetime/ParticleModuleLifetime_Seeded.h
--- a/Engine/Source/Runtime/Engine/Classes/Particles/Lifetime/ParticleModuleLifetime_Seeded.h
+++ b/Engine/Source/Runtime/Engine/Classes/Particles/Lifetime/ParticleModuleLifetime_Seeded.h
@@ -29,7 +29,7 @@ class UParticleModuleLifetime_Seeded : public UParticleModuleLifetime
 	//End UParticleModule Interface
 
 	//~ Begin UParticleModuleLifetimeBase Interface
-	virtual float	GetLifetimeValue(FParticleEmitterInstance* Owner, float InTime, UObject* Data = NULL) override;
+	virtual float	GetLifetimeValue(const FContext& Context, float InTime, UObject* Data = NULL) override;
 	//~ End UParticleModuleLifetimeBase Interface
 };
 
diff --git a/Engine/Source/Runtime/Engine/Classes/Particles/Light/ParticleModuleLight.h b/Engine/Source/Runtime/Engine/Classes/Particles/Light/ParticleModuleLight.h
--- a/Engine/Source/Runtime/Engine/Classes/Particles/Light/ParticleModuleLight.h
+++ b/Engine/Source/Runtime/Engine/Classes/Particles/Light/ParticleModuleLight.h
@@ -105,15 +105,15 @@ class UParticleModuleLight : public UParticleModuleLightBase
 
 
 	//Begin UParticleModule Interface
-	virtual void Spawn(FParticleEmitterInstance* Owner, int32 Offset, float SpawnTime, FBaseParticle* ParticleBase) override;
-	virtual void Update(FParticleEmitterInstance* Owner, int32 Offset, float DeltaTime) override;
+	virtual void Spawn(const FSpawnContext& Context) override;
+	virtual void Update(const FUpdateContext& Context) override;
 	virtual uint32 RequiredBytes(UParticleModuleTypeDataBase* TypeData) override;
 	virtual void SetToSensibleDefaults(UParticleEmitter* Owner) override;
 	virtual EModuleType	GetModuleType() const override { return EPMT_Light; }
-	virtual void Render3DPreview(FParticleEmitterInstance* Owner, const FSceneView* View,FPrimitiveDrawInterface* PDI) override;
+	virtual void Render3DPreview(const FPreviewContext& Context) override;
 	//End UParticleModule Interface
 
-	void SpawnEx(FParticleEmitterInstance* Owner, int32 Offset, float SpawnTime, struct FRandomStream* InRandomStream, FBaseParticle* ParticleBase);
+	void SpawnEx(const FSpawnContext& Context, struct FRandomStream* InRandomStream);
 
 	virtual bool CanTickInAnyThread() override;
 
diff --git a/Engine/Source/Runtime/Engine/Classes/Particles/Location/ParticleModuleLocation.h b/Engine/Source/Runtime/Engine/Classes/Particles/Location/ParticleModuleLocation.h
--- a/Engine/Source/Runtime/Engine/Classes/Particles/Location/ParticleModuleLocation.h
+++ b/Engine/Source/Runtime/Engine/Classes/Particles/Location/ParticleModuleLocation.h
@@ -52,7 +52,7 @@ protected:
 	 *	@param	SpawnTime			The time of the spawn
 	 *	@param	InRandomStream		The random stream to use for retrieving random values
 	 */
-	ENGINE_API virtual void SpawnEx(FParticleEmitterInstance* Owner, int32 Offset, float SpawnTime, struct FRandomStream* InRandomStream, FBaseParticle* ParticleBase);
+	ENGINE_API virtual void SpawnEx(const FSpawnContext& Context, struct FRandomStream* InRandomStream);
 
 public:
 	//Begin UObject Interface
@@ -63,8 +63,8 @@ public:
 	//End UObject Interface
 
 	//Begin UParticleModule Interface
-	ENGINE_API virtual void Spawn(FParticleEmitterInstance* Owner, int32 Offset, float SpawnTime, FBaseParticle* ParticleBase) override;
-	ENGINE_API virtual void Render3DPreview(FParticleEmitterInstance* Owner, const FSceneView* View,FPrimitiveDrawInterface* PDI) override;
+	ENGINE_API virtual void Spawn(const FSpawnContext& Context) override;
+	ENGINE_API virtual void Render3DPreview(const FPreviewContext& Context) override;
 	//End UParticleModule Interface
 };
 
diff --git a/Engine/Source/Runtime/Engine/Classes/Particles/Location/ParticleModuleLocationBoneSocket.h b/Engine/Source/Runtime/Engine/Classes/Particles/Location/ParticleModuleLocationBoneSocket.h
--- a/Engine/Source/Runtime/Engine/Classes/Particles/Location/ParticleModuleLocationBoneSocket.h
+++ b/Engine/Source/Runtime/Engine/Classes/Particles/Location/ParticleModuleLocationBoneSocket.h
@@ -132,9 +132,9 @@ class UParticleModuleLocationBoneSocket : public UParticleModuleLocationBase
 
 	//~ Begin UParticleModule Interface
 	ENGINE_API virtual void	PostLoad() override;
-	ENGINE_API virtual void	Spawn(FParticleEmitterInstance* Owner, int32 Offset, float SpawnTime, FBaseParticle* ParticleBase) override;
-	ENGINE_API virtual void	Update(FParticleEmitterInstance* Owner, int32 Offset, float DeltaTime) override;
-	ENGINE_API virtual void	FinalUpdate(FParticleEmitterInstance* Owner, int32 Offset, float DeltaTime) override;
+	ENGINE_API virtual void	Spawn(const FSpawnContext& Context) override;
+	ENGINE_API virtual void	Update(const FUpdateContext& Context) override;
+	ENGINE_API virtual void	FinalUpdate(const FUpdateContext& Context) override;
 	ENGINE_API virtual uint32	RequiredBytes(UParticleModuleTypeDataBase* TypeData) override;
 	ENGINE_API virtual uint32	RequiredBytesPerInstance() override;
 	ENGINE_API virtual uint32	PrepPerInstanceBlock(FParticleEmitterInstance* Owner, void* InstData) override;
@@ -159,7 +159,7 @@ class UParticleModuleLocationBoneSocket : public UParticleModuleLocationBase
 	 *
 	 *	@return	USkeletalMeshComponent*		The skeletal mesh component to use as the source
 	 */
-	ENGINE_API void GetSkeletalMeshComponentSource(FParticleEmitterInstance* Owner, FModuleLocationBoneSocketInstancePayload* InstancePayload);
+	ENGINE_API void GetSkeletalMeshComponentSource(const FContext& Context, FModuleLocationBoneSocketInstancePayload* InstancePayload);
 
 	/**
 	 *	Retrieve the position for the given socket index.
diff --git a/Engine/Source/Runtime/Engine/Classes/Particles/Location/ParticleModuleLocationDirect.h b/Engine/Source/Runtime/Engine/Classes/Particles/Location/ParticleModuleLocationDirect.h
--- a/Engine/Source/Runtime/Engine/Classes/Particles/Location/ParticleModuleLocationDirect.h
+++ b/Engine/Source/Runtime/Engine/Classes/Particles/Location/ParticleModuleLocationDirect.h
@@ -59,8 +59,8 @@ class UParticleModuleLocationDirect : public UParticleModuleLocationBase
 	//~ End UObject Interface
 
 	//~ Begin UParticleModule Interface
-	ENGINE_API virtual void	Spawn(FParticleEmitterInstance* Owner, int32 Offset, float SpawnTime, FBaseParticle* ParticleBase) override;
-	ENGINE_API virtual void	Update(FParticleEmitterInstance* Owner, int32 Offset, float DeltaTime) override;
+	ENGINE_API virtual void	Spawn(const FSpawnContext& Context) override;
+	ENGINE_API virtual void	Update(const FUpdateContext& Context) override;
 	ENGINE_API virtual uint32	RequiredBytes(UParticleModuleTypeDataBase* TypeData) override;
 	//~ End UParticleModule Interface
 };
diff --git a/Engine/Source/Runtime/Engine/Classes/Particles/Location/ParticleModuleLocationEmitter.h b/Engine/Source/Runtime/Engine/Classes/Particles/Location/ParticleModuleLocationEmitter.h
--- a/Engine/Source/Runtime/Engine/Classes/Particles/Location/ParticleModuleLocationEmitter.h
+++ b/Engine/Source/Runtime/Engine/Classes/Particles/Location/ParticleModuleLocationEmitter.h
@@ -66,7 +66,7 @@ class UParticleModuleLocationEmitter : public UParticleModuleLocationBase
 
 
 	//~ Begin UParticleModule Interface
-	ENGINE_API virtual void	Spawn(FParticleEmitterInstance* Owner, int32 Offset, float SpawnTime, FBaseParticle* ParticleBase) override;
+	ENGINE_API virtual void	Spawn(const FSpawnContext& Context) override;
 	ENGINE_API virtual uint32	RequiredBytesPerInstance() override;
 	//~ End UParticleModule Interface
 };
diff --git a/Engine/Source/Runtime/Engine/Classes/Particles/Location/ParticleModuleLocationEmitterDirect.h b/Engine/Source/Runtime/Engine/Classes/Particles/Location/ParticleModuleLocationEmitterDirect.h
--- a/Engine/Source/Runtime/Engine/Classes/Particles/Location/ParticleModuleLocationEmitterDirect.h
+++ b/Engine/Source/Runtime/Engine/Classes/Particles/Location/ParticleModuleLocationEmitterDirect.h
@@ -28,8 +28,8 @@ class UParticleModuleLocationEmitterDirect : public UParticleModuleLocationBase
 
 
 	//~ Begin UParticleModule Interface
-	ENGINE_API virtual void	Spawn(FParticleEmitterInstance* Owner, int32 Offset, float SpawnTime, FBaseParticle* ParticleBase) override;
-	ENGINE_API virtual void	Update(FParticleEmitterInstance* Owner, int32 Offset, float DeltaTime) override;
+	ENGINE_API virtual void	Spawn(const FSpawnContext& Context) override;
+	ENGINE_API virtual void	Update(const FUpdateContext& Context) override;
 	//~ End UParticleModule Interface
 
 };
diff --git a/Engine/Source/Runtime/Engine/Classes/Particles/Location/ParticleModuleLocationPrimitiveCylinder.h b/Engine/Source/Runtime/Engine/Classes/Particles/Location/ParticleModuleLocationPrimitiveCylinder.h
--- a/Engine/Source/Runtime/Engine/Classes/Particles/Location/ParticleModuleLocationPrimitiveCylinder.h
+++ b/Engine/Source/Runtime/Engine/Classes/Particles/Location/ParticleModuleLocationPrimitiveCylinder.h
@@ -61,8 +61,8 @@ class UParticleModuleLocationPrimitiveCylinder : public UParticleModuleLocationP
 	//~ End UObject Interface
 
 	//~ Begin UParticleModule Interface
-	ENGINE_API virtual void	Spawn(FParticleEmitterInstance* Owner, int32 Offset, float SpawnTime, FBaseParticle* ParticleBase) override;
-	ENGINE_API virtual void	Render3DPreview(FParticleEmitterInstance* Owner, const FSceneView* View,FPrimitiveDrawInterface* PDI) override;
+	ENGINE_API virtual void	Spawn(const FSpawnContext& Context) override;
+	ENGINE_API virtual void	Render3DPreview(const FPreviewContext& Context) override;
 	//~ End UParticleModule Interface
 
 	/**
@@ -73,7 +73,7 @@ class UParticleModuleLocationPrimitiveCylinder : public UParticleModuleLocationP
 	 *	@param	SpawnTime			The time of the spawn
 	 *	@param	InRandomStream		The random stream to use for retrieving random values
 	 */
-	ENGINE_API void SpawnEx(FParticleEmitterInstance* Owner, int32 Offset, float SpawnTime, struct FRandomStream* InRandomStream, FBaseParticle* ParticleBase);
+	ENGINE_API void SpawnEx(const FSpawnContext& Context, struct FRandomStream* InRandomStream);
 };
 
 
diff --git a/Engine/Source/Runtime/Engine/Classes/Particles/Location/ParticleModuleLocationPrimitiveSphere.h b/Engine/Source/Runtime/Engine/Classes/Particles/Location/ParticleModuleLocationPrimitiveSphere.h
--- a/Engine/Source/Runtime/Engine/Classes/Particles/Location/ParticleModuleLocationPrimitiveSphere.h
+++ b/Engine/Source/Runtime/Engine/Classes/Particles/Location/ParticleModuleLocationPrimitiveSphere.h
@@ -34,8 +34,8 @@ class UParticleModuleLocationPrimitiveSphere : public UParticleModuleLocationPri
 	//~ End UObject Interface
 
 	//~ Begin UParticleModule Interface
-	ENGINE_API virtual void	Spawn(FParticleEmitterInstance* Owner, int32 Offset, float SpawnTime, FBaseParticle* ParticleBase) override;
-	ENGINE_API virtual void	Render3DPreview(FParticleEmitterInstance* Owner, const FSceneView* View,FPrimitiveDrawInterface* PDI) override;
+	ENGINE_API virtual void	Spawn(const FSpawnContext& Context) override;
+	ENGINE_API virtual void	Render3DPreview(const FPreviewContext& Context) override;
 	//~ End UParticleModule Interface
 
 	/**
@@ -46,7 +46,7 @@ class UParticleModuleLocationPrimitiveSphere : public UParticleModuleLocationPri
 	 *	@param	SpawnTime			The time of the spawn
 	 *	@param	InRandomStream		The random stream to use for retrieving random values
 	 */
-	ENGINE_API void SpawnEx(FParticleEmitterInstance* Owner, int32 Offset, float SpawnTime, struct FRandomStream* InRandomStream, FBaseParticle* ParticleBase);
+	ENGINE_API void SpawnEx(const FSpawnContext& Context, struct FRandomStream* InRandomStream);
 };
 
 
diff --git a/Engine/Source/Runtime/Engine/Classes/Particles/Location/ParticleModuleLocationPrimitiveTriangle.h b/Engine/Source/Runtime/Engine/Classes/Particles/Location/ParticleModuleLocationPrimitiveTriangle.h
--- a/Engine/Source/Runtime/Engine/Classes/Particles/Location/ParticleModuleLocationPrimitiveTriangle.h
+++ b/Engine/Source/Runtime/Engine/Classes/Particles/Location/ParticleModuleLocationPrimitiveTriangle.h
@@ -45,8 +45,8 @@ class UParticleModuleLocationPrimitiveTriangle : public UParticleModuleLocationB
 	//~ End UObject Interface
 
 	//~ Begin UParticleModule Interface
-	ENGINE_API virtual void	Spawn(FParticleEmitterInstance* Owner, int32 Offset, float SpawnTime, FBaseParticle* ParticleBase) override;
-	ENGINE_API virtual void	Render3DPreview(FParticleEmitterInstance* Owner, const FSceneView* View,FPrimitiveDrawInterface* PDI) override;
+	ENGINE_API virtual void	Spawn(const FSpawnContext& Context) override;
+	ENGINE_API virtual void	Render3DPreview(const FPreviewContext& Context) override;
 	//~ End UParticleModule Interface
 
 	/**
@@ -57,7 +57,7 @@ class UParticleModuleLocationPrimitiveTriangle : public UParticleModuleLocationB
 	 *	@param	SpawnTime			The time of the spawn
 	 *	@param	InRandomStream		The random stream to use for retrieving random values
 	 */
-	ENGINE_API void SpawnEx(FParticleEmitterInstance* Owner, int32 Offset, float SpawnTime, struct FRandomStream* InRandomStream, FBaseParticle* ParticleBase);
+	ENGINE_API void SpawnEx(const FSpawnContext& Context, struct FRandomStream* InRandomStream);
 };
 
 
diff --git a/Engine/Source/Runtime/Engine/Classes/Particles/Location/ParticleModuleLocationSkelVertSurface.h b/Engine/Source/Runtime/Engine/Classes/Particles/Location/ParticleModuleLocationSkelVertSurface.h
--- a/Engine/Source/Runtime/Engine/Classes/Particles/Location/ParticleModuleLocationSkelVertSurface.h
+++ b/Engine/Source/Runtime/Engine/Classes/Particles/Location/ParticleModuleLocationSkelVertSurface.h
@@ -165,9 +165,9 @@ class UParticleModuleLocationSkelVertSurface : public UParticleModuleLocationBas
 	//~ End UObject Interface
 
 	//Begin UParticleModule Interface
-	ENGINE_API virtual void Spawn(FParticleEmitterInstance* Owner, int32 Offset, float SpawnTime, FBaseParticle* ParticleBase) override;
-	ENGINE_API virtual void Update(FParticleEmitterInstance* Owner, int32 Offset, float DeltaTime) override;
-	ENGINE_API virtual void FinalUpdate(FParticleEmitterInstance* Owner, int32 Offset, float DeltaTime) override;
+	ENGINE_API virtual void Spawn(const FSpawnContext& Context) override;
+	ENGINE_API virtual void Update(const FUpdateContext& Context) override;
+	ENGINE_API virtual void FinalUpdate(const FUpdateContext& Context) override;
 	ENGINE_API virtual uint32	PrepPerInstanceBlock(FParticleEmitterInstance* Owner, void* InstData) override;
 	ENGINE_API virtual uint32	RequiredBytes(UParticleModuleTypeDataBase* TypeData) override;
 	ENGINE_API virtual uint32	RequiredBytesPerInstance() override;
@@ -192,7 +192,7 @@ class UParticleModuleLocationSkelVertSurface : public UParticleModuleLocationBas
 	 *	@param	Owner						The payload data that is being setup.
 	 *
 	 */
-	ENGINE_API void GetSkeletalMeshComponentSource(FParticleEmitterInstance* Owner, FModuleLocationVertSurfaceInstancePayload* InstancePayload);
+	ENGINE_API void GetSkeletalMeshComponentSource(const FContext& Context, FModuleLocationVertSurfaceInstancePayload* InstancePayload);
 
 	/**
 	 *	Retrieve the position for the given socket index.
@@ -206,7 +206,7 @@ class UParticleModuleLocationSkelVertSurface : public UParticleModuleLocationBas
 	 *	
 	 *	@return	bool					true if successful, false if not
 	 */
-	ENGINE_API bool GetParticleLocation(FParticleEmitterInstance* Owner, FModuleLocationVertSurfaceInstancePayload* InstPayload, USkeletalMeshComponent* InSkelMeshComponent, int32 InPrimaryVertexIndex, FVector& OutPosition, FQuat& OutRotation, bool bSpawning = false);
+	ENGINE_API bool GetParticleLocation(const FContext& Context, FModuleLocationVertSurfaceInstancePayload* InstPayload, USkeletalMeshComponent* InSkelMeshComponent, int32 InPrimaryVertexIndex, FVector& OutPosition, FQuat& OutRotation, bool bSpawning = false);
 
 	/**
 	 *  Check to see if the vert is influenced by a bone on our approved list.
diff --git a/Engine/Source/Runtime/Engine/Classes/Particles/Location/ParticleModuleLocationWorldOffset.h b/Engine/Source/Runtime/Engine/Classes/Particles/Location/ParticleModuleLocationWorldOffset.h
--- a/Engine/Source/Runtime/Engine/Classes/Particles/Location/ParticleModuleLocationWorldOffset.h
+++ b/Engine/Source/Runtime/Engine/Classes/Particles/Location/ParticleModuleLocationWorldOffset.h
@@ -18,7 +18,7 @@ class UParticleModuleLocationWorldOffset : public UParticleModuleLocation
 
 protected:
 	//Begin UParticleModuleLocation Interface
-	ENGINE_API virtual void SpawnEx(FParticleEmitterInstance* Owner, int32 Offset, float SpawnTime, struct FRandomStream* InRandomStream, FBaseParticle* ParticleBase) override;
+	ENGINE_API virtual void SpawnEx(const FSpawnContext& Context, struct FRandomStream* InRandomStream) override;
 	//End UParticleModuleLocation Interface
 };
 
diff --git a/Engine/Source/Runtime/Engine/Classes/Particles/Location/ParticleModuleSourceMovement.h b/Engine/Source/Runtime/Engine/Classes/Particles/Location/ParticleModuleSourceMovement.h
--- a/Engine/Source/Runtime/Engine/Classes/Particles/Location/ParticleModuleSourceMovement.h
+++ b/Engine/Source/Runtime/Engine/Classes/Particles/Location/ParticleModuleSourceMovement.h
@@ -31,7 +31,7 @@ class UParticleModuleSourceMovement : public UParticleModuleLocationBase
 	//End UObject Interface
 
 	//Begin UParticleModule Interface
-	ENGINE_API virtual void FinalUpdate(FParticleEmitterInstance* Owner, int32 Offset, float DeltaTime) override;
+	ENGINE_API virtual void FinalUpdate(const FUpdateContext& Context) override;
 	ENGINE_API virtual bool CanTickInAnyThread() override;
 	//End UParticleModule Interface
 
diff --git a/Engine/Source/Runtime/Engine/Classes/Particles/Material/ParticleModuleMeshMaterial.h b/Engine/Source/Runtime/Engine/Classes/Particles/Material/ParticleModuleMeshMaterial.h
--- a/Engine/Source/Runtime/Engine/Classes/Particles/Material/ParticleModuleMeshMaterial.h
+++ b/Engine/Source/Runtime/Engine/Classes/Particles/Material/ParticleModuleMeshMaterial.h
@@ -21,7 +21,7 @@ class UParticleModuleMeshMaterial : public UParticleModuleMaterialBase
 
 
 	//Begin UParticleModule Interface
-	virtual void	Spawn(FParticleEmitterInstance* Owner, int32 Offset, float SpawnTime, FBaseParticle* ParticleBase) override;
+	virtual void	Spawn(const FSpawnContext& Context) override;
 	virtual uint32	RequiredBytesPerInstance() override;
 
 #if WITH_EDITOR
diff --git a/Engine/Source/Runtime/Engine/Classes/Particles/Orbit/ParticleModuleOrbit.h b/Engine/Source/Runtime/Engine/Classes/Particles/Orbit/ParticleModuleOrbit.h
--- a/Engine/Source/Runtime/Engine/Classes/Particles/Orbit/ParticleModuleOrbit.h
+++ b/Engine/Source/Runtime/Engine/Classes/Particles/Orbit/ParticleModuleOrbit.h
@@ -124,8 +124,8 @@ class UParticleModuleOrbit : public UParticleModuleOrbitBase
 
 	//Begin UParticleModule Interface
 	virtual void CompileModule( struct FParticleEmitterBuildInfo& EmitterInfo ) override;
-	virtual void	Spawn(FParticleEmitterInstance* Owner, int32 Offset, float SpawnTime, FBaseParticle* ParticleBase) override;
-	virtual void	Update(FParticleEmitterInstance* Owner, int32 Offset, float DeltaTime) override;
+	virtual void	Spawn(const FSpawnContext& Context) override;
+	virtual void	Update(const FUpdateContext& Context) override;
 	virtual uint32	RequiredBytes(UParticleModuleTypeDataBase* TypeData) override;
 	virtual uint32	RequiredBytesPerInstance() override;
 #if WITH_EDITOR
diff --git a/Engine/Source/Runtime/Engine/Classes/Particles/Orientation/ParticleModuleOrientationAxisLock.h b/Engine/Source/Runtime/Engine/Classes/Particles/Orientation/ParticleModuleOrientationAxisLock.h
--- a/Engine/Source/Runtime/Engine/Classes/Particles/Orientation/ParticleModuleOrientationAxisLock.h
+++ b/Engine/Source/Runtime/Engine/Classes/Particles/Orientation/ParticleModuleOrientationAxisLock.h
@@ -67,8 +67,8 @@ class UParticleModuleOrientationAxisLock : public UParticleModuleOrientationBase
 	//~ End UObject Interface
 
 	//~ Begin UParticleModule Interface
-	virtual void	Spawn(FParticleEmitterInstance* Owner, int32 Offset, float SpawnTime, FBaseParticle* ParticleBase) override;
-	virtual void	Update(FParticleEmitterInstance* Owner, int32 Offset, float DeltaTime) override;
+	virtual void	Spawn(const FSpawnContext& Context) override;
+	virtual void	Update(const FUpdateContext& Context) override;
 	//~ End UParticleModule Interface
 
 	//@todo document
diff --git a/Engine/Source/Runtime/Engine/Classes/Particles/Parameter/ParticleModuleParameterDynamic.h b/Engine/Source/Runtime/Engine/Classes/Particles/Parameter/ParticleModuleParameterDynamic.h
--- a/Engine/Source/Runtime/Engine/Classes/Particles/Parameter/ParticleModuleParameterDynamic.h
+++ b/Engine/Source/Runtime/Engine/Classes/Particles/Parameter/ParticleModuleParameterDynamic.h
@@ -121,8 +121,8 @@ class UParticleModuleParameterDynamic : public UParticleModuleParameterBase
 	//~ Begin UObject Interface
 
 	//~ Begin UParticleModule Interface
-	virtual void	Spawn(FParticleEmitterInstance* Owner, int32 Offset, float SpawnTime, FBaseParticle* ParticleBase) override;
-	virtual void	Update(FParticleEmitterInstance* Owner, int32 Offset, float DeltaTime) override;
+	virtual void	Spawn(const FSpawnContext& Context) override;
+	virtual void	Update(const FUpdateContext& Context) override;
 	virtual uint32	RequiredBytes(UParticleModuleTypeDataBase* TypeData) override;
 	virtual void SetToSensibleDefaults(UParticleEmitter* Owner) override;
 	virtual void	GetCurveObjects(TArray<FParticleCurvePair>& OutCurves) override;
@@ -145,7 +145,7 @@ class UParticleModuleParameterDynamic : public UParticleModuleParameterBase
 	 *	@param	SpawnTime			The time of the spawn
 	 *	@param	InRandomStream		The random stream to use for retrieving random values
 	 */
-	void SpawnEx(FParticleEmitterInstance* Owner, int32 Offset, float SpawnTime, struct FRandomStream* InRandomStream, FBaseParticle* ParticleBase);
+	void SpawnEx(const FSpawnContext& Context, struct FRandomStream* InRandomStream);
 	
 	/**
 	 *	Update the parameter names with the given material...
diff --git a/Engine/Source/Runtime/Engine/Classes/Particles/ParticleModule.h b/Engine/Source/Runtime/Engine/Classes/Particles/ParticleModule.h
--- a/Engine/Source/Runtime/Engine/Classes/Particles/ParticleModule.h
+++ b/Engine/Source/Runtime/Engine/Classes/Particles/ParticleModule.h
@@ -233,6 +233,14 @@ class UParticleModule : public UObject
 	 */
 	ENGINE_API virtual void CompileModule( struct FParticleEmitterBuildInfo& EmitterInfo );
 
+	struct FContext
+	{
+		FParticleEmitterInstance& Owner;
+		const FTransform& GetTransform() const;
+		UObject* GetDistributionData() const;
+		FContext(FParticleEmitterInstance& Ow) : Owner(Ow) {}
+	};
+
 	/**
 	 *	Called on a particle that is freshly spawned by the emitter.
 	 *	
@@ -240,7 +248,15 @@ class UParticleModule : public UObject
 	 *	@param	Offset		The modules offset into the data payload of the particle.
 	 *	@param	SpawnTime	The time of the spawn.
 	 */
-	ENGINE_API virtual void	Spawn(FParticleEmitterInstance* Owner, int32 Offset, float SpawnTime, FBaseParticle* ParticleBase);
+	struct FSpawnContext : FContext
+	{
+		int32 Offset;
+		float SpawnTime;
+		FBaseParticle* ParticleBase;
+		FSpawnContext(FParticleEmitterInstance& Ow, int32 Of, float St, FBaseParticle* Pb) : FContext(Ow), Offset(Of), SpawnTime(St), ParticleBase(Pb) {}
+	};
+	ENGINE_API virtual void	Spawn(const FSpawnContext& Context);
+
 	/**
 	 *	Called on a particle that is being updated by its emitter.
 	 *	
@@ -248,7 +264,13 @@ class UParticleModule : public UObject
 	 *	@param	Offset		The modules offset into the data payload of the particle.
 	 *	@param	DeltaTime	The time since the last update.
 	 */
-	ENGINE_API virtual void	Update(FParticleEmitterInstance* Owner, int32 Offset, float DeltaTime);
+	struct FUpdateContext : FContext
+	{
+		int32 Offset;
+		float DeltaTime;
+		FUpdateContext(FParticleEmitterInstance& Ow, int32 Of, float Dt) : FContext(Ow), Offset(Of), DeltaTime(Dt) {}
+	};
+	ENGINE_API virtual void	Update(const FUpdateContext& Context);
 	/**
 	 *	Called on an emitter when all other update operations have taken place
 	 *	INCLUDING bounding box cacluations!
@@ -257,7 +279,7 @@ class UParticleModule : public UObject
 	 *	@param	Offset		The modules offset into the data payload of the particle.
 	 *	@param	DeltaTime	The time since the last update.
 	 */
-	ENGINE_API virtual void	FinalUpdate(FParticleEmitterInstance* Owner, int32 Offset, float DeltaTime);
+	ENGINE_API virtual void	FinalUpdate(const FUpdateContext& Context);
 
 	/**
 	 *	Returns the number of bytes that the module requires in the particle payload block.
@@ -341,7 +363,13 @@ class UParticleModule : public UObject
 	 *	@param	View		The scene view that is being rendered.
 	 *	@param	PDI			The FPrimitiveDrawInterface to use for rendering.
 	 */
-	virtual void Render3DPreview(FParticleEmitterInstance* Owner, const FSceneView* View,FPrimitiveDrawInterface* PDI)	{};
+	struct FPreviewContext : FContext
+	{
+		const FSceneView* View;
+		FPrimitiveDrawInterface* PDI;
+		FPreviewContext(FParticleEmitterInstance& Ow, const FSceneView* Vi, FPrimitiveDrawInterface* Pd) : FContext(Ow), View(Vi), PDI(Pd) {}
+	};
+	virtual void Render3DPreview(const FPreviewContext& Context) {}
 
 	/**
 	 *	Retrieve the ModuleType of this module.
@@ -546,7 +574,7 @@ class UParticleModule : public UObject
 	 *
 	 * @return FRandomStream&	The random stream to use for the provided instance.
 	 */
-	ENGINE_API FRandomStream& GetRandomStream(FParticleEmitterInstance* Owner);
+	ENGINE_API FRandomStream& GetRandomStream(const FContext& Context);
 
 #if WITH_EDITOR
 	ENGINE_API virtual void PostLoadSubobjects( FObjectInstancingGraph* OuterInstanceGraph ) override;
diff --git a/Engine/Source/Runtime/Engine/Classes/Particles/Rotation/ParticleModuleMeshRotation.h b/Engine/Source/Runtime/Engine/Classes/Particles/Rotation/ParticleModuleMeshRotation.h
--- a/Engine/Source/Runtime/Engine/Classes/Particles/Rotation/ParticleModuleMeshRotation.h
+++ b/Engine/Source/Runtime/Engine/Classes/Particles/Rotation/ParticleModuleMeshRotation.h
@@ -35,7 +35,7 @@ class UParticleModuleMeshRotation : public UParticleModuleRotationBase
 	//End UObject Interface
 
 	//Begin UParticleModule Interface
-	virtual void	Spawn(FParticleEmitterInstance* Owner, int32 Offset, float SpawnTime, FBaseParticle* ParticleBase) override;
+	virtual void	Spawn(const FSpawnContext& Context) override;
 	virtual bool	TouchesMeshRotation() const override { return true; }
 	//End UParticleModule Interface
 
@@ -50,7 +50,7 @@ class UParticleModuleMeshRotation : public UParticleModuleRotationBase
 	 *	@param	SpawnTime			The time of the spawn
 	 *	@param	InRandomStream		The random stream to use for retrieving random values
 	 */
-	void SpawnEx(FParticleEmitterInstance* Owner, int32 Offset, float SpawnTime, struct FRandomStream* InRandomStream, FBaseParticle* ParticleBase);
+	void SpawnEx(const FSpawnContext& Context, struct FRandomStream* InRandomStream);
 
 };
 
diff --git a/Engine/Source/Runtime/Engine/Classes/Particles/Rotation/ParticleModuleRotation.h b/Engine/Source/Runtime/Engine/Classes/Particles/Rotation/ParticleModuleRotation.h
--- a/Engine/Source/Runtime/Engine/Classes/Particles/Rotation/ParticleModuleRotation.h
+++ b/Engine/Source/Runtime/Engine/Classes/Particles/Rotation/ParticleModuleRotation.h
@@ -34,7 +34,7 @@ class UParticleModuleRotation : public UParticleModuleRotationBase
 	//End UObject Interface
 
 	//Begin UParticleModule Interface
-	virtual void	Spawn(FParticleEmitterInstance* Owner, int32 Offset, float SpawnTime, FBaseParticle* ParticleBase) override;
+	virtual void	Spawn(const FSpawnContext& Context) override;
 	//End UParticleModule Interface
 
 	/**
@@ -45,7 +45,7 @@ class UParticleModuleRotation : public UParticleModuleRotationBase
 	 *	@param	SpawnTime			The time of the spawn
 	 *	@param	InRandomStream		The random stream to use for retrieving random values
 	 */
-	void SpawnEx(FParticleEmitterInstance* Owner, int32 Offset, float SpawnTime, struct FRandomStream* InRandomStream, FBaseParticle* ParticleBase);
+	void SpawnEx(const FSpawnContext& Context, struct FRandomStream* InRandomStream);
 };
 
 
diff --git a/Engine/Source/Runtime/Engine/Classes/Particles/Rotation/ParticleModuleRotationOverLifetime.h b/Engine/Source/Runtime/Engine/Classes/Particles/Rotation/ParticleModuleRotationOverLifetime.h
--- a/Engine/Source/Runtime/Engine/Classes/Particles/Rotation/ParticleModuleRotationOverLifetime.h
+++ b/Engine/Source/Runtime/Engine/Classes/Particles/Rotation/ParticleModuleRotationOverLifetime.h
@@ -41,7 +41,7 @@ class UParticleModuleRotationOverLifetime : public UParticleModuleRotationBase
 	//End UObject Interface
 
 	//Begin UParticleModule Interface
-	virtual void	Update(FParticleEmitterInstance* Owner, int32 Offset, float DeltaTime) override;
+	virtual void	Update(const FUpdateContext& Context) override;
 	//End UParticleModule Interface
 };
 
diff --git a/Engine/Source/Runtime/Engine/Classes/Particles/RotationRate/ParticleModuleMeshRotationRate.h b/Engine/Source/Runtime/Engine/Classes/Particles/RotationRate/ParticleModuleMeshRotationRate.h
--- a/Engine/Source/Runtime/Engine/Classes/Particles/RotationRate/ParticleModuleMeshRotationRate.h
+++ b/Engine/Source/Runtime/Engine/Classes/Particles/RotationRate/ParticleModuleMeshRotationRate.h
@@ -32,7 +32,7 @@ class UParticleModuleMeshRotationRate : public UParticleModuleRotationRateBase
 	//End UObject Interface
 
 	//~ Begin UParticleModule Interface
-	virtual void	Spawn(FParticleEmitterInstance* Owner, int32 Offset, float SpawnTime, FBaseParticle* ParticleBase) override;
+	virtual void	Spawn(const FSpawnContext& Context) override;
 	virtual void SetToSensibleDefaults(UParticleEmitter* Owner) override;
 	virtual bool	TouchesMeshRotation() const override { return true; }
 	//~ End UParticleModule Interface
@@ -48,7 +48,7 @@ class UParticleModuleMeshRotationRate : public UParticleModuleRotationRateBase
 	 *	@param	SpawnTime			The time of the spawn
 	 *	@param	InRandomStream		The random stream to use for retrieving random values
 	 */
-	void SpawnEx(FParticleEmitterInstance* Owner, int32 Offset, float SpawnTime, struct FRandomStream* InRandomStream, FBaseParticle* ParticleBase);
+	void SpawnEx(const FSpawnContext& Context, struct FRandomStream* InRandomStream);
 
 };
 
diff --git a/Engine/Source/Runtime/Engine/Classes/Particles/RotationRate/ParticleModuleMeshRotationRateMultiplyLife.h b/Engine/Source/Runtime/Engine/Classes/Particles/RotationRate/ParticleModuleMeshRotationRateMultiplyLife.h
--- a/Engine/Source/Runtime/Engine/Classes/Particles/RotationRate/ParticleModuleMeshRotationRateMultiplyLife.h
+++ b/Engine/Source/Runtime/Engine/Classes/Particles/RotationRate/ParticleModuleMeshRotationRateMultiplyLife.h
@@ -33,8 +33,8 @@ class UParticleModuleMeshRotationRateMultiplyLife : public UParticleModuleRotati
 	//End UObject Interface
 
 	//~ Begin UParticleModule Interface
-	virtual void	Spawn(FParticleEmitterInstance* Owner, int32 Offset, float SpawnTime, FBaseParticle* ParticleBase) override;
-	virtual void	Update(FParticleEmitterInstance* Owner, int32 Offset, float DeltaTime) override;
+	virtual void	Spawn(const FSpawnContext& Context) override;
+	virtual void	Update(const FUpdateContext& Context) override;
 	virtual void SetToSensibleDefaults(UParticleEmitter* Owner) override;
 	//~ End UParticleModule Interface
 
diff --git a/Engine/Source/Runtime/Engine/Classes/Particles/RotationRate/ParticleModuleMeshRotationRateOverLife.h b/Engine/Source/Runtime/Engine/Classes/Particles/RotationRate/ParticleModuleMeshRotationRateOverLife.h
--- a/Engine/Source/Runtime/Engine/Classes/Particles/RotationRate/ParticleModuleMeshRotationRateOverLife.h
+++ b/Engine/Source/Runtime/Engine/Classes/Particles/RotationRate/ParticleModuleMeshRotationRateOverLife.h
@@ -39,8 +39,8 @@ class UParticleModuleMeshRotationRateOverLife : public UParticleModuleRotationRa
 	//End UObject Interface
 
 	//~ Begin UParticleModule Interface
-	virtual void Spawn(FParticleEmitterInstance* Owner, int32 Offset, float SpawnTime, FBaseParticle* ParticleBase) override;
-	virtual void Update(FParticleEmitterInstance* Owner, int32 Offset, float DeltaTime) override;
+	virtual void Spawn(const FSpawnContext& Context) override;
+	virtual void Update(const FUpdateContext& Context) override;
 	virtual void SetToSensibleDefaults(UParticleEmitter* Owner) override;
 	virtual bool TouchesMeshRotation() const override { return true; }
 	//~ End UParticleModule Interface
diff --git a/Engine/Source/Runtime/Engine/Classes/Particles/RotationRate/ParticleModuleRotationRate.h b/Engine/Source/Runtime/Engine/Classes/Particles/RotationRate/ParticleModuleRotationRate.h
--- a/Engine/Source/Runtime/Engine/Classes/Particles/RotationRate/ParticleModuleRotationRate.h
+++ b/Engine/Source/Runtime/Engine/Classes/Particles/RotationRate/ParticleModuleRotationRate.h
@@ -36,7 +36,7 @@ class UParticleModuleRotationRate : public UParticleModuleRotationRateBase
 
 	//~ Begin UParticleModule Interface
 	virtual void CompileModule( struct FParticleEmitterBuildInfo& EmitterInfo ) override;
-	virtual void Spawn(FParticleEmitterInstance* Owner, int32 Offset, float SpawnTime, FBaseParticle* ParticleBase) override;
+	virtual void Spawn(const FSpawnContext& Context) override;
 	virtual void SetToSensibleDefaults(UParticleEmitter* Owner) override;
 	//~ Begin UParticleModule Interface
 
@@ -48,7 +48,7 @@ class UParticleModuleRotationRate : public UParticleModuleRotationRateBase
 	 *	@param	SpawnTime			The time of the spawn
 	 *	@param	InRandomStream		The random stream to use for retrieving random values
 	 */
-	void SpawnEx(FParticleEmitterInstance* Owner, int32 Offset, float SpawnTime, struct FRandomStream* InRandomStream, FBaseParticle* ParticleBase);
+	void SpawnEx(const FSpawnContext& Context, struct FRandomStream* InRandomStream);
 };
 
 
diff --git a/Engine/Source/Runtime/Engine/Classes/Particles/RotationRate/ParticleModuleRotationRateMultiplyLife.h b/Engine/Source/Runtime/Engine/Classes/Particles/RotationRate/ParticleModuleRotationRateMultiplyLife.h
--- a/Engine/Source/Runtime/Engine/Classes/Particles/RotationRate/ParticleModuleRotationRateMultiplyLife.h
+++ b/Engine/Source/Runtime/Engine/Classes/Particles/RotationRate/ParticleModuleRotationRateMultiplyLife.h
@@ -35,8 +35,8 @@ class UParticleModuleRotationRateMultiplyLife : public UParticleModuleRotationRa
 	//End UObject Interface
 
 	//~ Begin UParticleModule Interface
-	virtual void Spawn(FParticleEmitterInstance* Owner, int32 Offset, float SpawnTime, FBaseParticle* ParticleBase) override;
-	virtual void Update(FParticleEmitterInstance* Owner, int32 Offset, float DeltaTime) override;
+	virtual void Spawn(const FSpawnContext& Context) override;
+	virtual void Update(const FUpdateContext& Context) override;
 	virtual void SetToSensibleDefaults(UParticleEmitter* Owner) override;
 	// End UParticleModule Interface
 };
diff --git a/Engine/Source/Runtime/Engine/Classes/Particles/Size/ParticleModuleSize.h b/Engine/Source/Runtime/Engine/Classes/Particles/Size/ParticleModuleSize.h
--- a/Engine/Source/Runtime/Engine/Classes/Particles/Size/ParticleModuleSize.h
+++ b/Engine/Source/Runtime/Engine/Classes/Particles/Size/ParticleModuleSize.h
@@ -36,7 +36,7 @@ class UParticleModuleSize : public UParticleModuleSizeBase
 
 	//~ Begin UParticleModule Interface
 	virtual void CompileModule( struct FParticleEmitterBuildInfo& EmitterInfo ) override;
-	virtual void Spawn(FParticleEmitterInstance* Owner, int32 Offset, float SpawnTime, FBaseParticle* ParticleBase) override;
+	virtual void Spawn(const FSpawnContext& Context) override;
 	//~ End UParticleModule Interface
 
 	/**
@@ -47,7 +47,7 @@ class UParticleModuleSize : public UParticleModuleSizeBase
 	 *	@param	SpawnTime			The time of the spawn
 	 *	@param	InRandomStream		The random stream to use for retrieving random values
 	 */
-	void SpawnEx(FParticleEmitterInstance* Owner, int32 Offset, float SpawnTime, struct FRandomStream* InRandomStream, FBaseParticle* ParticleBase);
+	void SpawnEx(const FSpawnContext& Context, struct FRandomStream* InRandomStream);
 };
 
 
diff --git a/Engine/Source/Runtime/Engine/Classes/Particles/Size/ParticleModuleSizeMultiplyLife.h b/Engine/Source/Runtime/Engine/Classes/Particles/Size/ParticleModuleSizeMultiplyLife.h
--- a/Engine/Source/Runtime/Engine/Classes/Particles/Size/ParticleModuleSizeMultiplyLife.h
+++ b/Engine/Source/Runtime/Engine/Classes/Particles/Size/ParticleModuleSizeMultiplyLife.h
@@ -58,8 +58,8 @@ class UParticleModuleSizeMultiplyLife : public UParticleModuleSizeBase
 
 	//~ Begin UParticleModule Interface
 	virtual void CompileModule( struct FParticleEmitterBuildInfo& EmitterInfo ) override;
-	virtual void Spawn(FParticleEmitterInstance* Owner, int32 Offset, float SpawnTime, FBaseParticle* ParticleBase) override;
-	virtual void	Update(FParticleEmitterInstance* Owner, int32 Offset, float DeltaTime) override;
+	virtual void Spawn(const FSpawnContext& Context) override;
+	virtual void	Update(const FUpdateContext& Context) override;
 	virtual void SetToSensibleDefaults(UParticleEmitter* Owner) override;
 	virtual bool   IsSizeMultiplyLife() override { return true; };
 
diff --git a/Engine/Source/Runtime/Engine/Classes/Particles/Size/ParticleModuleSizeScale.h b/Engine/Source/Runtime/Engine/Classes/Particles/Size/ParticleModuleSizeScale.h
--- a/Engine/Source/Runtime/Engine/Classes/Particles/Size/ParticleModuleSizeScale.h
+++ b/Engine/Source/Runtime/Engine/Classes/Particles/Size/ParticleModuleSizeScale.h
@@ -50,8 +50,8 @@ class UParticleModuleSizeScale : public UParticleModuleSizeBase
 
 	//~ Begin UParticleModule Interface
 	virtual void CompileModule( struct FParticleEmitterBuildInfo& EmitterInfo ) override;
-	virtual void Spawn(FParticleEmitterInstance* Owner, int32 Offset, float SpawnTime, FBaseParticle* ParticleBase) override;
-	virtual void Update(FParticleEmitterInstance* Owner, int32 Offset, float DeltaTime) override;
+	virtual void Spawn(const FSpawnContext& Context) override;
+	virtual void Update(const FUpdateContext& Context) override;
 	virtual void SetToSensibleDefaults(UParticleEmitter* Owner) override;
 
 #if WITH_EDITOR
diff --git a/Engine/Source/Runtime/Engine/Classes/Particles/Size/ParticleModuleSizeScaleBySpeed.h b/Engine/Source/Runtime/Engine/Classes/Particles/Size/ParticleModuleSizeScaleBySpeed.h
--- a/Engine/Source/Runtime/Engine/Classes/Particles/Size/ParticleModuleSizeScaleBySpeed.h
+++ b/Engine/Source/Runtime/Engine/Classes/Particles/Size/ParticleModuleSizeScaleBySpeed.h
@@ -28,7 +28,7 @@ class UParticleModuleSizeScaleBySpeed : public UParticleModuleSizeBase
 
 
 	//~ Begin UParticleModule Interface
-	virtual void Update(FParticleEmitterInstance* Owner, int32 Offset, float DeltaTime) override;
+	virtual void Update(const FUpdateContext& Context) override;
 	virtual void CompileModule(struct FParticleEmitterBuildInfo& EmitterInfo) override;
 	// End UParticleModule Interface
 };
diff --git a/Engine/Source/Runtime/Engine/Classes/Particles/Spawn/ParticleModuleSpawn.h b/Engine/Source/Runtime/Engine/Classes/Particles/Spawn/ParticleModuleSpawn.h
--- a/Engine/Source/Runtime/Engine/Classes/Particles/Spawn/ParticleModuleSpawn.h
+++ b/Engine/Source/Runtime/Engine/Classes/Particles/Spawn/ParticleModuleSpawn.h
@@ -56,7 +56,7 @@ class UParticleModuleSpawn : public UParticleModuleSpawnBase
 	//~ End UParticleModule Interface
 
 	//~ Begin UParticleModuleSpawnBase Interface
-	virtual bool GetSpawnAmount(FParticleEmitterInstance* Owner, int32 Offset, float OldLeftover, 
+	virtual bool GetSpawnAmount(const FContext& Context, int32 Offset, float OldLeftover, 
 		float DeltaTime, int32& Number, float& Rate) override;
 	virtual float GetMaximumSpawnRate() override;
 	virtual float GetEstimatedSpawnRate() override;
diff --git a/Engine/Source/Runtime/Engine/Classes/Particles/Spawn/ParticleModuleSpawnBase.h b/Engine/Source/Runtime/Engine/Classes/Particles/Spawn/ParticleModuleSpawnBase.h
--- a/Engine/Source/Runtime/Engine/Classes/Particles/Spawn/ParticleModuleSpawnBase.h
+++ b/Engine/Source/Runtime/Engine/Classes/Particles/Spawn/ParticleModuleSpawnBase.h
@@ -51,7 +51,7 @@ class UParticleModuleSpawnBase : public UParticleModule
 	 *	@return	bool		false if the SpawnRate should be ignored.
 	 *						true if the SpawnRate should still be processed.
 	 */
-	virtual bool GetSpawnAmount(FParticleEmitterInstance* Owner, int32 Offset, float OldLeftover, 
+	virtual bool GetSpawnAmount(const FContext& Context, int32 Offset, float OldLeftover, 
 		float DeltaTime, int32& Number, float& Rate)
 	{
 		return bProcessSpawnRate;
diff --git a/Engine/Source/Runtime/Engine/Classes/Particles/Spawn/ParticleModuleSpawnPerUnit.h b/Engine/Source/Runtime/Engine/Classes/Particles/Spawn/ParticleModuleSpawnPerUnit.h
--- a/Engine/Source/Runtime/Engine/Classes/Particles/Spawn/ParticleModuleSpawnPerUnit.h
+++ b/Engine/Source/Runtime/Engine/Classes/Particles/Spawn/ParticleModuleSpawnPerUnit.h
@@ -86,7 +86,7 @@ class UParticleModuleSpawnPerUnit : public UParticleModuleSpawnBase
 	//~ End UParticleModule Interface
 
 	//~ Begin UParticleModuleSpawnBase Interface
-	virtual bool GetSpawnAmount(FParticleEmitterInstance* Owner, int32 Offset, float OldLeftover, 
+	virtual bool GetSpawnAmount(const FContext& Context, int32 Offset, float OldLeftover, 
 		float DeltaTime, int32& Number, float& Rate) override;
 	//~ End UParticleModuleSpawnBase Interface
 };
diff --git a/Engine/Source/Runtime/Engine/Classes/Particles/SubUV/ParticleModuleSubUV.h b/Engine/Source/Runtime/Engine/Classes/Particles/SubUV/ParticleModuleSubUV.h
--- a/Engine/Source/Runtime/Engine/Classes/Particles/SubUV/ParticleModuleSubUV.h
+++ b/Engine/Source/Runtime/Engine/Classes/Particles/SubUV/ParticleModuleSubUV.h
@@ -55,8 +55,8 @@ class UParticleModuleSubUV : public UParticleModuleSubUVBase
 
 	//~ Begin UParticleModule Interface
 	virtual void CompileModule( struct FParticleEmitterBuildInfo& EmitterInfo ) override;
-	virtual void Spawn(FParticleEmitterInstance* Owner, int32 Offset, float SpawnTime, FBaseParticle* ParticleBase) override;
-	virtual void Update(FParticleEmitterInstance* Owner, int32 Offset, float DeltaTime) override;
+	virtual void Spawn(const FSpawnContext& Context) override;
+	virtual void Update(const FUpdateContext& Context) override;
 	virtual void SetToSensibleDefaults(UParticleEmitter* Owner) override;
 	//~ End UParticleModule Interface
 
@@ -71,7 +71,7 @@ class UParticleModuleSubUV : public UParticleModuleSubUVBase
 	 *
 	 *	@return	float					The image index with interpolation amount as the fractional portion.
 	 */
-	virtual float DetermineImageIndex(FParticleEmitterInstance* Owner, int32 Offset, FBaseParticle* Particle, 
+	virtual float DetermineImageIndex(const FContext& Context, int32 Offset, FBaseParticle* Particle, 
 		EParticleSubUVInterpMethod InterpMethod, FFullSubUVPayload& SubUVPayload, float DeltaTime);
 
 #if WITH_EDITOR
diff --git a/Engine/Source/Runtime/Engine/Classes/Particles/SubUV/ParticleModuleSubUVMovie.h b/Engine/Source/Runtime/Engine/Classes/Particles/SubUV/ParticleModuleSubUVMovie.h
--- a/Engine/Source/Runtime/Engine/Classes/Particles/SubUV/ParticleModuleSubUVMovie.h
+++ b/Engine/Source/Runtime/Engine/Classes/Particles/SubUV/ParticleModuleSubUVMovie.h
@@ -50,14 +50,14 @@ class UParticleModuleSubUVMovie : public UParticleModuleSubUV
 	//End UObject Interface
 
 	//~ Begin UParticleModule Interface
-	virtual void Spawn(FParticleEmitterInstance* Owner, int32 Offset, float SpawnTime, FBaseParticle* ParticleBase) override;
+	virtual void Spawn(const FSpawnContext& Context) override;
 	virtual uint32 RequiredBytes(UParticleModuleTypeDataBase* TypeData) override;
 	virtual void SetToSensibleDefaults(UParticleEmitter* Owner) override;
 	virtual void GetCurveObjects(TArray<FParticleCurvePair>& OutCurves) override;
 	//~ End UParticleModule Interface
 	
 	//~ Begin UParticleModuleSubUV Interface
-	virtual float DetermineImageIndex(FParticleEmitterInstance* Owner, int32 Offset, FBaseParticle* Particle, 
+	virtual float DetermineImageIndex(const FContext& Context, int32 Offset, FBaseParticle* Particle, 
 		EParticleSubUVInterpMethod InterpMethod, FFullSubUVPayload& SubUVPayload, float DeltaTime) override;
 	//~ End UParticleModuleSubUV Interface
 	
diff --git a/Engine/Source/Runtime/Engine/Classes/Particles/TypeData/ParticleModuleTypeDataBeam2.h b/Engine/Source/Runtime/Engine/Classes/Particles/TypeData/ParticleModuleTypeDataBeam2.h
--- a/Engine/Source/Runtime/Engine/Classes/Particles/TypeData/ParticleModuleTypeDataBeam2.h
+++ b/Engine/Source/Runtime/Engine/Classes/Particles/TypeData/ParticleModuleTypeDataBeam2.h
@@ -210,8 +210,8 @@ class UParticleModuleTypeDataBeam2 : public UParticleModuleTypeDataBase
 
 
 	//~ Begin UParticleModule Interface
-	virtual void	Spawn(FParticleEmitterInstance* Owner, int32 Offset, float SpawnTime, FBaseParticle* ParticleBase) override;
-	virtual void	Update(FParticleEmitterInstance* Owner, int32 Offset, float DeltaTime) override;
+	virtual void	Spawn(const FSpawnContext& Context) override;
+	virtual void	Update(const FUpdateContext& Context) override;
 	virtual uint32	RequiredBytes(UParticleModuleTypeDataBase* TypeData) override;
 	virtual	bool	AddModuleCurvesToEditor(UInterpCurveEdSetup* EdSetup, TArray<const FCurveEdEntry*>& OutCurveEntries) override;
 	virtual bool CanTickInAnyThread() override
diff --git a/Engine/Source/Runtime/Engine/Classes/Particles/Velocity/ParticleModuleVelocity.h b/Engine/Source/Runtime/Engine/Classes/Particles/Velocity/ParticleModuleVelocity.h
--- a/Engine/Source/Runtime/Engine/Classes/Particles/Velocity/ParticleModuleVelocity.h
+++ b/Engine/Source/Runtime/Engine/Classes/Particles/Velocity/ParticleModuleVelocity.h
@@ -43,7 +43,7 @@ class UParticleModuleVelocity : public UParticleModuleVelocityBase
 	//~ End UObject Interface
 
 	//~ Begin UParticleModule Interface
-	virtual void Spawn(FParticleEmitterInstance* Owner, int32 Offset, float SpawnTime, FBaseParticle* ParticleBase) override;
+	virtual void Spawn(const FSpawnContext& Context) override;
 	//~ Begin UParticleModule Interface
 
 	/**
@@ -54,7 +54,7 @@ class UParticleModuleVelocity : public UParticleModuleVelocityBase
 	 *	@param	SpawnTime			The time of the spawn
 	 *	@param	InRandomStream		The random stream to use for retrieving random values
 	 */
-	void SpawnEx(FParticleEmitterInstance* Owner, int32 Offset, float SpawnTime, struct FRandomStream* InRandomStream, FBaseParticle* ParticleBase);
+	void SpawnEx(const FSpawnContext& Context, struct FRandomStream* InRandomStream);
 };
 
 
diff --git a/Engine/Source/Runtime/Engine/Classes/Particles/Velocity/ParticleModuleVelocityCone.h b/Engine/Source/Runtime/Engine/Classes/Particles/Velocity/ParticleModuleVelocityCone.h
--- a/Engine/Source/Runtime/Engine/Classes/Particles/Velocity/ParticleModuleVelocityCone.h
+++ b/Engine/Source/Runtime/Engine/Classes/Particles/Velocity/ParticleModuleVelocityCone.h
@@ -39,8 +39,8 @@ class UParticleModuleVelocityCone : public UParticleModuleVelocityBase
 	//~ End UObject Interface
 
 	//~ Begin UParticleModule Interface
-	virtual void Spawn(FParticleEmitterInstance* Owner, int32 Offset, float SpawnTime, FBaseParticle* ParticleBase) override;
-	virtual void	Render3DPreview(FParticleEmitterInstance* Owner, const FSceneView* View,FPrimitiveDrawInterface* PDI) override;
+	virtual void Spawn(const FSpawnContext& Context) override;
+	virtual void	Render3DPreview(const FPreviewContext& Context) override;
 	//~ Begin UParticleModule Interface
 	
 	/**
@@ -51,7 +51,7 @@ class UParticleModuleVelocityCone : public UParticleModuleVelocityBase
 	 *	@param	SpawnTime			The time of the spawn
 	 *	@param	InRandomStream		The random stream to use for retrieving random values
 	 */
-	void			SpawnEx(FParticleEmitterInstance* Owner, int32 Offset, float SpawnTime, struct FRandomStream* InRandomStream, FBaseParticle* ParticleBase);
+	void			SpawnEx(const FSpawnContext& Context, struct FRandomStream* InRandomStream);
 
 };
 
diff --git a/Engine/Source/Runtime/Engine/Classes/Particles/Velocity/ParticleModuleVelocityInheritParent.h b/Engine/Source/Runtime/Engine/Classes/Particles/Velocity/ParticleModuleVelocityInheritParent.h
--- a/Engine/Source/Runtime/Engine/Classes/Particles/Velocity/ParticleModuleVelocityInheritParent.h
+++ b/Engine/Source/Runtime/Engine/Classes/Particles/Velocity/ParticleModuleVelocityInheritParent.h
@@ -34,7 +34,7 @@ class UParticleModuleVelocityInheritParent : public UParticleModuleVelocityBase
 	//~ End UObject Interface
 
 	//~ Begin UParticleModule Interface
-	virtual void Spawn(FParticleEmitterInstance* Owner, int32 Offset, float SpawnTime, FBaseParticle* ParticleBase) override;
+	virtual void Spawn(const FSpawnContext& Context) override;
 	//~ Begin UParticleModule Interface
 };
 
diff --git a/Engine/Source/Runtime/Engine/Classes/Particles/Velocity/ParticleModuleVelocityOverLifetime.h b/Engine/Source/Runtime/Engine/Classes/Particles/Velocity/ParticleModuleVelocityOverLifetime.h
--- a/Engine/Source/Runtime/Engine/Classes/Particles/Velocity/ParticleModuleVelocityOverLifetime.h
+++ b/Engine/Source/Runtime/Engine/Classes/Particles/Velocity/ParticleModuleVelocityOverLifetime.h
@@ -41,8 +41,8 @@ class UParticleModuleVelocityOverLifetime : public UParticleModuleVelocityBase
 	//~ End UObject Interface
 
 	//~ Begin UParticleModule Interface
-	virtual void Spawn(FParticleEmitterInstance* Owner, int32 Offset, float SpawnTime, FBaseParticle* ParticleBase) override;
-	virtual void	Update(FParticleEmitterInstance* Owner, int32 Offset, float DeltaTime) override;
+	virtual void Spawn(const FSpawnContext& Context) override;
+	virtual void	Update(const FUpdateContext& Context) override;
 	//~ Begin UParticleModule Interface
 };
 
diff --git a/Engine/Source/Runtime/Engine/Private/ParticleSystemSceneProxy.cpp b/Engine/Source/Runtime/Engine/Private/ParticleSystemSceneProxy.cpp
--- a/Engine/Source/Runtime/Engine/Private/ParticleSystemSceneProxy.cpp
+++ b/Engine/Source/Runtime/Engine/Private/ParticleSystemSceneProxy.cpp
@@ -596,7 +596,7 @@ void DrawParticleSystemHelpers(UParticleSystemComponent* InPSysComp, const FScen
 					UParticleModule* Module = LODLevel->Modules[ModuleIdx];
 					if (Module && Module->bSupported3DDrawMode && Module->b3DDrawMode)
 					{
-						Module->Render3DPreview(EmitterInst, View, PDI);
+						Module->Render3DPreview({ *EmitterInst, View, PDI });
 					}
 				}
 			}
diff --git a/Engine/Source/Runtime/Engine/Private/Particles/ParticleBeam2EmitterInstance.cpp b/Engine/Source/Runtime/Engine/Private/Particles/ParticleBeam2EmitterInstance.cpp
--- a/Engine/Source/Runtime/Engine/Private/Particles/ParticleBeam2EmitterInstance.cpp
+++ b/Engine/Source/Runtime/Engine/Private/Particles/ParticleBeam2EmitterInstance.cpp
@@ -583,25 +583,25 @@ void FParticleBeam2EmitterInstance::Tick_ModulePostUpdate(float DeltaTime, UPart
 		// The order of the update here is VERY important
 		if (BeamModule_Source && BeamModule_Source->bEnabled)
 		{
-			BeamModule_Source->Update(this, GetModuleDataOffset(BeamModule_Source), DeltaTime);
+			BeamModule_Source->Update({*this, (int)GetModuleDataOffset(BeamModule_Source), DeltaTime});
 		}
 		if (BeamModule_SourceModifier && BeamModule_SourceModifier->bEnabled)
 		{
 			int32 TempOffset = BeamModule_SourceModifier_Offset;
-			BeamModule_SourceModifier->Update(this, TempOffset, DeltaTime);
+			BeamModule_SourceModifier->Update({*this, TempOffset, DeltaTime});
 		}
 		if (BeamModule_Target && BeamModule_Target->bEnabled)
 		{
-			BeamModule_Target->Update(this, GetModuleDataOffset(BeamModule_Target), DeltaTime);
+			BeamModule_Target->Update({*this, (int)GetModuleDataOffset(BeamModule_Target), DeltaTime});
 		}
 		if (BeamModule_TargetModifier && BeamModule_TargetModifier->bEnabled)
 		{
 			int32 TempOffset = BeamModule_TargetModifier_Offset;
-			BeamModule_TargetModifier->Update(this, TempOffset, DeltaTime);
+			BeamModule_TargetModifier->Update({*this, TempOffset, DeltaTime});
 		}
 		if (BeamModule_Noise && BeamModule_Noise->bEnabled)
 		{
-			BeamModule_Noise->Update(this, GetModuleDataOffset(BeamModule_Noise), DeltaTime);
+			BeamModule_Noise->Update({*this, (int)GetModuleDataOffset(BeamModule_Noise), DeltaTime});
 		}
 
 		FParticleEmitterInstance::Tick_ModulePostUpdate(DeltaTime, InCurrentLODLevel);
@@ -902,30 +902,30 @@ void FParticleBeam2EmitterInstance::PostSpawn(FBaseParticle* Particle, float Int
 	UParticleLODLevel* LODLevel = SpriteTemplate->GetCurrentLODLevel(this);
 	if (BeamModule_Source && BeamModule_Source->bEnabled)
 	{
-		BeamModule_Source->Spawn(this, GetModuleDataOffset(BeamModule_Source), SpawnTime, Particle);
+		BeamModule_Source->Spawn({*this, (int32)GetModuleDataOffset(BeamModule_Source), SpawnTime, Particle});
 	}
 	if (BeamModule_SourceModifier && BeamModule_SourceModifier->bEnabled)
 	{
 		int32 ModifierOffset = BeamModule_SourceModifier_Offset;
-		BeamModule_SourceModifier->Spawn(this, ModifierOffset, SpawnTime, Particle);
+		BeamModule_SourceModifier->Spawn({*this, ModifierOffset, SpawnTime, Particle});
 	}
 	if (BeamModule_Target && BeamModule_Target->bEnabled)
 	{
-		BeamModule_Target->Spawn(this, GetModuleDataOffset(BeamModule_Target), SpawnTime, Particle);
+		BeamModule_Target->Spawn({*this, (int32)GetModuleDataOffset(BeamModule_Target), SpawnTime, Particle});
 	}
 	if (BeamModule_TargetModifier && BeamModule_TargetModifier->bEnabled)
 	{
 		int32 ModifierOffset = BeamModule_TargetModifier_Offset;
-		BeamModule_TargetModifier->Spawn(this, ModifierOffset, SpawnTime, Particle);
+		BeamModule_TargetModifier->Spawn({*this, ModifierOffset, SpawnTime, Particle});
 	}
 	if (BeamModule_Noise && BeamModule_Noise->bEnabled)
 	{
-		BeamModule_Noise->Spawn(this, GetModuleDataOffset(BeamModule_Noise), SpawnTime, Particle);
+		BeamModule_Noise->Spawn({*this, (int32)GetModuleDataOffset(BeamModule_Noise), SpawnTime, Particle});
 	}
 	if (LODLevel->TypeDataModule)
 	{
 		//@todo. Need to track TypeData offset into payload!
-		LODLevel->TypeDataModule->Spawn(this, TypeDataOffset, SpawnTime, Particle);
+		LODLevel->TypeDataModule->Spawn({*this, TypeDataOffset, SpawnTime, Particle});
 	}
 
 	FParticleEmitterInstance::PostSpawn( Particle, InterpolationPercentage, SpawnTime );
diff --git a/Engine/Source/Runtime/Engine/Private/Particles/ParticleBeamModules.cpp b/Engine/Source/Runtime/Engine/Private/Particles/ParticleBeamModules.cpp
--- a/Engine/Source/Runtime/Engine/Private/Particles/ParticleBeamModules.cpp
+++ b/Engine/Source/Runtime/Engine/Private/Particles/ParticleBeamModules.cpp
@@ -94,8 +94,10 @@ void UParticleModuleTypeDataBeam2::PostInitProperties()
 	}
 }
 
-void UParticleModuleTypeDataBeam2::Spawn(FParticleEmitterInstance* Owner, int32 Offset, float SpawnTime, FBaseParticle* ParticleBase)
+void UParticleModuleTypeDataBeam2::Spawn(const FSpawnContext& Context)
 {
+	FParticleEmitterInstance* Owner = &Context.Owner;
+
 	// Make sure that the owner is a beam emitter instance.
 	FParticleBeam2EmitterInstance*	BeamInst	= (FParticleBeam2EmitterInstance*)Owner;
 	check(Owner->bIsBeam);
@@ -202,8 +204,9 @@ void UParticleModuleTypeDataBeam2::Spawn(FParticleEmitterInstance* Owner, int32
 }
 
 
-void UParticleModuleTypeDataBeam2::Update(FParticleEmitterInstance* Owner, int32 Offset, float DeltaTime)
+void UParticleModuleTypeDataBeam2::Update(const FUpdateContext& Context)
 {
+	FParticleEmitterInstance* Owner = &Context.Owner;
 	UParticleSystemComponent*		Component	= Owner->Component;
 	FParticleBeam2EmitterInstance*	BeamInst	= (FParticleBeam2EmitterInstance*)Owner;
 	check(Owner->bIsBeam);
@@ -300,7 +303,7 @@ void UParticleModuleTypeDataBeam2::Update(FParticleEmitterInstance* Owner, int32
 				FVector Direction	= BeamData->TargetPoint - Particle.Location;
 				Direction.Normalize();
 				// Setup the offset and the current distance travelled
-				const FVector BeamOffset		= Direction * Speed * DeltaTime;
+				const FVector BeamOffset		= Direction * Speed * Context.DeltaTime;
 				FVector	Sum			= Particle.Location + BeamOffset;
 				if ((FMath::Abs(Sum.X - BeamData->TargetPoint.X) < LockRadius) && 
 					(FMath::Abs(Sum.Y - BeamData->TargetPoint.Y) < LockRadius) &&
@@ -1079,8 +1082,10 @@ void UParticleModuleBeamModifier::PostInitProperties()
 	}
 }
 
-void UParticleModuleBeamModifier::Spawn(FParticleEmitterInstance* Owner, int32 Offset, float SpawnTime, FBaseParticle* ParticleBase)
+void UParticleModuleBeamModifier::Spawn(const FSpawnContext& Context)
 {
+	FParticleEmitterInstance* Owner = &Context.Owner;
+
 	// Make sure that the owner is a beam emitter instance.
 	FParticleBeam2EmitterInstance*	BeamInst	= (FParticleBeam2EmitterInstance*)Owner;
 	if (BeamInst == NULL)
@@ -1097,6 +1102,7 @@ void UParticleModuleBeamModifier::Spawn(FParticleEmitterInstance* Owner, int32 O
 		FBeamParticleModifierPayloadData* TargetModifierPayload = NULL;
 
 		// Get the pointers to the data, but save the module offset that is passed in.
+		int32 Offset = Context.Offset;
 		GetDataPointers(Owner, (const uint8*)&Particle, Offset, BeamDataPayload, SourceModifierPayload, TargetModifierPayload);
 
 		FBeamParticleModifierPayloadData* ModifierPayload = (ModifierType == PEB2MT_Source) ? 
@@ -1108,7 +1114,7 @@ void UParticleModuleBeamModifier::Spawn(FParticleEmitterInstance* Owner, int32 O
 			ModifierPayload->bModifyPosition = PositionOptions.bModify;
 			if (PositionOptions.bModify == true)
 			{
-				ModifierPayload->Position = Position.GetValue(Owner->EmitterTime, Owner->Component);
+				ModifierPayload->Position = Position.GetValue(Owner->EmitterTime, Context.GetDistributionData());
 				ModifierPayload->bScalePosition = PositionOptions.bScale;
 			}
 
@@ -1116,7 +1122,7 @@ void UParticleModuleBeamModifier::Spawn(FParticleEmitterInstance* Owner, int32 O
 			ModifierPayload->bModifyTangent = TangentOptions.bModify;
 			if (TangentOptions.bModify == true)
 			{
-				ModifierPayload->Tangent = Tangent.GetValue(Owner->EmitterTime, Owner->Component);
+				ModifierPayload->Tangent = Tangent.GetValue(Owner->EmitterTime, Context.GetDistributionData());
 				ModifierPayload->bScaleTangent = TangentOptions.bScale;
 			}
 
@@ -1124,15 +1130,17 @@ void UParticleModuleBeamModifier::Spawn(FParticleEmitterInstance* Owner, int32 O
 			ModifierPayload->bModifyStrength = StrengthOptions.bModify;
 			if (StrengthOptions.bModify == true)
 			{
-				ModifierPayload->Strength = Strength.GetValue(Owner->EmitterTime, Owner->Component);
+				ModifierPayload->Strength = Strength.GetValue(Owner->EmitterTime, Context.GetDistributionData());
 				ModifierPayload->bScaleStrength = StrengthOptions.bScale;
 			}
 		}
 	}
 }
 
-void UParticleModuleBeamModifier::Update(FParticleEmitterInstance* Owner, int32 Offset, float DeltaTime)
+void UParticleModuleBeamModifier::Update(const FUpdateContext& Context)
 {
+	FParticleEmitterInstance* Owner = &Context.Owner;
+
 	// Make sure that the owner is a beam emitter instance.
 	FParticleBeam2EmitterInstance*	BeamInst	= (FParticleBeam2EmitterInstance*)Owner;
 	if (BeamInst == NULL || !Owner->bIsBeam)
@@ -1158,7 +1166,7 @@ void UParticleModuleBeamModifier::Update(FParticleEmitterInstance* Owner, int32
 			ModifierPayload->bModifyPosition = PositionOptions.bModify;
 			if ((PositionOptions.bModify == true) && (PositionOptions.bLock == false))
 			{
-				ModifierPayload->Position = Position.GetValue(Owner->EmitterTime, Owner->Component);
+				ModifierPayload->Position = Position.GetValue(Owner->EmitterTime, Context.GetDistributionData());
 				ModifierPayload->bScalePosition = PositionOptions.bScale;
 			}
 
@@ -1166,7 +1174,7 @@ void UParticleModuleBeamModifier::Update(FParticleEmitterInstance* Owner, int32
 			ModifierPayload->bModifyTangent = TangentOptions.bModify;
 			if ((TangentOptions.bModify == true) && (TangentOptions.bLock == false))
 			{
-				ModifierPayload->Tangent = Tangent.GetValue(Owner->EmitterTime, Owner->Component);
+				ModifierPayload->Tangent = Tangent.GetValue(Owner->EmitterTime, Context.GetDistributionData());
 				ModifierPayload->bScaleTangent = TangentOptions.bScale;
 			}
 
@@ -1174,7 +1182,7 @@ void UParticleModuleBeamModifier::Update(FParticleEmitterInstance* Owner, int32
 			ModifierPayload->bModifyStrength = StrengthOptions.bModify;
 			if ((StrengthOptions.bModify == true) && (StrengthOptions.bLock == false))
 			{
-				ModifierPayload->Strength = Strength.GetValue(Owner->EmitterTime, Owner->Component);
+				ModifierPayload->Strength = Strength.GetValue(Owner->EmitterTime, Context.GetDistributionData());
 				ModifierPayload->bScaleStrength = StrengthOptions.bScale;
 			}
 		}
@@ -1363,7 +1371,7 @@ void UParticleModuleBeamNoise::PostInitProperties()
 	}
 }
 
-void UParticleModuleBeamNoise::Spawn(FParticleEmitterInstance* Owner, int32 Offset, float SpawnTime, FBaseParticle* ParticleBase)
+void UParticleModuleBeamNoise::Spawn(const FSpawnContext& Context)
 {
 	if (bLowFreq_Enabled == false)
 	{
@@ -1371,6 +1379,8 @@ void UParticleModuleBeamNoise::Spawn(FParticleEmitterInstance* Owner, int32 Offs
 		return;
 	}
 
+	FParticleEmitterInstance* Owner = &Context.Owner;
+
 	// Make sure that the owner is a beam emitter instance.
 	FParticleBeam2EmitterInstance*	BeamInst	= (FParticleBeam2EmitterInstance*)Owner;
 	if (!BeamInst || !bLowFreq_Enabled || (Frequency == 0) || !Owner->bIsBeam)
@@ -1378,10 +1388,9 @@ void UParticleModuleBeamNoise::Spawn(FParticleEmitterInstance* Owner, int32 Offs
 		return;
 	}
 
-	UParticleSystemComponent*		Component	= Owner->Component;
 	UParticleModuleTypeDataBeam2*	BeamTD		= BeamInst->BeamTypeData;
 
-	FRandomStream& RandomStream = GetRandomStream(Owner);
+	FRandomStream& RandomStream = GetRandomStream(Context);
 
 	SPAWN_INIT;
 
@@ -1443,16 +1452,16 @@ void UParticleModuleBeamNoise::Spawn(FParticleEmitterInstance* Owner, int32 Offs
 		{
 			Extreme = 0;
 		}
-		TargetNoisePoints[ii] = NoiseRange.GetValue(StepSize * ii, Owner->Component, Extreme);
+		TargetNoisePoints[ii] = NoiseRange.GetValue(StepSize * ii, Context.GetDistributionData(), Extreme);
 		if (bSmooth)
 		{
 			Extreme = -Extreme;
-			NextNoisePoints[ii] = NoiseRange.GetValue(StepSize * ii, Owner->Component, Extreme);
+			NextNoisePoints[ii] = NoiseRange.GetValue(StepSize * ii, Context.GetDistributionData(), Extreme);
 		}
 	}
 }
 
-void UParticleModuleBeamNoise::Update(FParticleEmitterInstance* Owner, int32 Offset, float DeltaTime)
+void UParticleModuleBeamNoise::Update(const FUpdateContext& Context)
 {
 	if (bLowFreq_Enabled == false)
 	{
@@ -1460,6 +1469,8 @@ void UParticleModuleBeamNoise::Update(FParticleEmitterInstance* Owner, int32 Off
 		return;
 	}
 
+	FParticleEmitterInstance* Owner = &Context.Owner;
+
 	// Make sure that the owner is a beam emitter instance and there is noise.
 	FParticleBeam2EmitterInstance*	BeamInst	= (FParticleBeam2EmitterInstance*)Owner;
 	if (!BeamInst || (Frequency == 0) || !Owner->bIsBeam)
@@ -1527,32 +1538,32 @@ void UParticleModuleBeamNoise::Update(FParticleEmitterInstance* Owner, int32 Off
 					//@todo. Add support for moving noise points!
 					// Check the times...
 					check(NoiseRate);
-					*NoiseRate += DeltaTime;
+					*NoiseRate += Context.DeltaTime;
 					if (*NoiseRate > NoiseLockTime)
 					{
 						if (bSmooth)
 						{
 							for (int32 ii = 0; ii < (Freq + 1); ii++)
 							{
-								NextNoisePoints[ii] = NoiseRange.GetValue(StepSize * ii, Owner->Component, Extreme);
+								NextNoisePoints[ii] = NoiseRange.GetValue(StepSize * ii, Context.GetDistributionData(), Extreme);
 							}
 						}
 						else
 						{
 							for (int32 ii = 0; ii < (Freq + 1); ii++)
 							{
-								TargetNoisePoints[ii] = NoiseRange.GetValue(StepSize * ii, Owner->Component, Extreme);
+								TargetNoisePoints[ii] = NoiseRange.GetValue(StepSize * ii, Context.GetDistributionData(), Extreme);
 							}
 						}
 						*NoiseRate	= 0.0f;
 					}
-					*NoiseDelta	= DeltaTime;
+					*NoiseDelta	= Context.DeltaTime;
 				}
 				else
 				{
 					for (int32 ii = 0; ii < (Freq + 1); ii++)
 					{
-						TargetNoisePoints[ii] = NoiseRange.GetValue(StepSize * ii, Owner->Component, Extreme);
+						TargetNoisePoints[ii] = NoiseRange.GetValue(StepSize * ii, Context.GetDistributionData(), Extreme);
 					}
 				}
 			}
@@ -1683,8 +1694,10 @@ void UParticleModuleBeamSource::PostInitProperties()
 	}
 }
 
-void UParticleModuleBeamSource::Spawn(FParticleEmitterInstance* Owner, int32 Offset, float SpawnTime, FBaseParticle* ParticleBase)
+void UParticleModuleBeamSource::Spawn(const FSpawnContext& Context)
 {
+	FParticleEmitterInstance* Owner = &Context.Owner;
+
 	// Make sure that the owner is a beam emitter instance.
 	FParticleBeam2EmitterInstance*	BeamInst	= (FParticleBeam2EmitterInstance*)Owner;
 	if (!BeamInst || !Owner->bIsBeam)
@@ -1714,6 +1727,7 @@ void UParticleModuleBeamSource::Spawn(FParticleEmitterInstance* Owner, int32 Off
 		TaperValues, NoiseDistanceScale, SourceModifier, TargetModifier);
 
 	// Resolve the source data
+	int32 Offset = Context.Offset;
 	ResolveSourceData(BeamInst, BeamData, (uint8*)ParticleBase, Offset, BeamInst->ActiveParticles, true, SourceModifier);
 
 	// Set the location and clear the initial data flags
@@ -1725,7 +1739,7 @@ void UParticleModuleBeamSource::Spawn(FParticleEmitterInstance* Owner, int32 Off
 	BeamData->TriangleCount				= 0;
 }
 
-void UParticleModuleBeamSource::Update(FParticleEmitterInstance* Owner, int32 Offset, float DeltaTime)
+void UParticleModuleBeamSource::Update(const FUpdateContext& Context)
 {
 	// If the source is locked, don't perform any update
 	if (bLockSource && bLockSourceTangent && bLockSourceStength)
@@ -1733,6 +1747,8 @@ void UParticleModuleBeamSource::Update(FParticleEmitterInstance* Owner, int32 Of
 		return;
 	}
 
+	FParticleEmitterInstance* Owner = &Context.Owner;
+
 	// Make sure that the owner is a beam emitter instance.
 	FParticleBeam2EmitterInstance*	BeamInst	= (FParticleBeam2EmitterInstance*)Owner;
 	if (!BeamInst || !Owner->bIsBeam)
@@ -2172,8 +2188,10 @@ void UParticleModuleBeamTarget::PostInitProperties()
 	}
 }
 
-void UParticleModuleBeamTarget::Spawn(FParticleEmitterInstance* Owner, int32 Offset, float SpawnTime, FBaseParticle* ParticleBase)
+void UParticleModuleBeamTarget::Spawn(const FSpawnContext& Context)
 {
+	FParticleEmitterInstance* Owner = &Context.Owner;
+
 	FParticleBeam2EmitterInstance*	BeamInst	= (FParticleBeam2EmitterInstance*)Owner;
 	if (!BeamInst || !Owner->bIsBeam)
 	{
@@ -2201,16 +2219,18 @@ void UParticleModuleBeamTarget::Spawn(FParticleEmitterInstance* Owner, int32 Off
 		TaperValues, NoiseDistanceScale, SourceModifier, TargetModifier);
 	CurrentOffset	= TempOffset;
 
+	int32 Offset = Context.Offset;
 	ResolveTargetData(BeamInst, BeamData, (uint8*)ParticleBase, Offset, BeamInst->ActiveParticles, true, TargetModifier);
 }
 
-void UParticleModuleBeamTarget::Update(FParticleEmitterInstance* Owner, int32 Offset, float DeltaTime)
+void UParticleModuleBeamTarget::Update(const FUpdateContext& Context)
 {
 	if (bLockTarget && bLockTargetTangent && bLockTargetStength)
 	{
 		return;
 	}
 
+	FParticleEmitterInstance* Owner = &Context.Owner;
 	FParticleBeam2EmitterInstance*	BeamInst	= (FParticleBeam2EmitterInstance*)Owner;
 	if (!BeamInst || !Owner->bIsBeam)
 	{
diff --git a/Engine/Source/Runtime/Engine/Private/Particles/ParticleEmitterInstances.cpp b/Engine/Source/Runtime/Engine/Private/Particles/ParticleEmitterInstances.cpp
--- a/Engine/Source/Runtime/Engine/Private/Particles/ParticleEmitterInstances.cpp
+++ b/Engine/Source/Runtime/Engine/Private/Particles/ParticleEmitterInstances.cpp
@@ -1077,7 +1077,7 @@ void FParticleEmitterInstance::Tick_ModuleUpdate(float DeltaTime, UParticleLODLe
 		UParticleModule* CurrentModule = InCurrentLODLevel->UpdateModules[ModuleIndex];
 		if (CurrentModule && CurrentModule->bEnabled && CurrentModule->bUpdateModule)
 		{
-			CurrentModule->Update(this, GetModuleDataOffset(HighestLODLevel->UpdateModules[ModuleIndex]), DeltaTime);
+			CurrentModule->Update({ *this, (int32)GetModuleDataOffset(HighestLODLevel->UpdateModules[ModuleIndex]), DeltaTime });
 		}
 	}
 }
@@ -1093,7 +1093,7 @@ void FParticleEmitterInstance::Tick_ModulePostUpdate(float DeltaTime, UParticleL
 	// Handle the TypeData module
 	if (InCurrentLODLevel->TypeDataModule)
 	{
-		InCurrentLODLevel->TypeDataModule->Update(this, TypeDataOffset, DeltaTime);
+		InCurrentLODLevel->TypeDataModule->Update({ *this, TypeDataOffset, DeltaTime });
 	}
 }
 
@@ -1112,14 +1112,14 @@ void FParticleEmitterInstance::Tick_ModuleFinalUpdate(float DeltaTime, UParticle
 		UParticleModule* CurrentModule = InCurrentLODLevel->UpdateModules[ModuleIndex];
 		if (CurrentModule && CurrentModule->bEnabled && CurrentModule->bFinalUpdateModule)
 		{
-			CurrentModule->FinalUpdate(this, GetModuleDataOffset(HighestLODLevel->UpdateModules[ModuleIndex]), DeltaTime);
+			CurrentModule->FinalUpdate({ *this, (int32)GetModuleDataOffset(HighestLODLevel->UpdateModules[ModuleIndex]), DeltaTime });
 		}
 	}
 
 
 	if (InCurrentLODLevel->TypeDataModule && InCurrentLODLevel->TypeDataModule->bEnabled && InCurrentLODLevel->TypeDataModule->bFinalUpdateModule)
 	{
-		InCurrentLODLevel->TypeDataModule->FinalUpdate(this, GetModuleDataOffset(HighestLODLevel->TypeDataModule), DeltaTime);
+		InCurrentLODLevel->TypeDataModule->FinalUpdate({ *this, (int32)GetModuleDataOffset(HighestLODLevel->TypeDataModule), DeltaTime });
 	}
 }
 
@@ -1640,7 +1640,7 @@ float FParticleEmitterInstance::GetCurrentBurstRateOffset(float& DeltaTime, int3
 	UParticleLODLevel* LODLevel	= GetCurrentLODLevelChecked();
 	if (LODLevel->SpawnModule->BurstList.Num() > 0)
 	{
-		FRandomStream& RandomStream = LODLevel->SpawnModule->GetRandomStream(this);
+		FRandomStream& RandomStream = LODLevel->SpawnModule->GetRandomStream({*this});
 
 		// For each burst in the list
 		for (int32 BurstIdx = 0; BurstIdx < LODLevel->SpawnModule->BurstList.Num(); BurstIdx++)
@@ -2010,7 +2010,7 @@ float FParticleEmitterInstance::Spawn(float DeltaTime)
 				// Update the spawn rate
 				int32 Number = 0;
 				float Rate = 0.0f;
-				if (SpawnModule->GetSpawnAmount(this, Offset, OldLeftover, DeltaTime, Number, Rate) == false)
+				if (SpawnModule->GetSpawnAmount({*this}, Offset, OldLeftover, DeltaTime, Number, Rate) == false)
 				{
 					bProcessSpawnRate = false;
 				}
@@ -2242,7 +2242,7 @@ void FParticleEmitterInstance::SpawnParticles( int32 Count, float StartTime, flo
 				if (SpawnModule->bEnabled)
 				{
 					UParticleModule* OffsetModule = HighestLODLevel->SpawnModules[ModuleIndex];
-					SpawnModule->Spawn(this, GetModuleDataOffset(OffsetModule), SpawnTime, Particle);
+					SpawnModule->Spawn({ *this, (int32)GetModuleDataOffset(OffsetModule), SpawnTime, Particle });
 
 					ensureMsgf(!Particle->Location.ContainsNaN(), TEXT("NaN in Particle Location. Template: %s, Component: %s"), Component ? *GetNameSafe(Component->Template) : TEXT("UNKNOWN"), *GetPathNameSafe(Component));
 				}
@@ -2668,7 +2668,7 @@ void FParticleEmitterInstance::SetupEmitterDuration()
 		UParticleLODLevel* TempLOD = SpriteTemplate->LODLevels[LODIndex];
 		UParticleModuleRequired* RequiredModule = TempLOD->RequiredModule;
 
-		FRandomStream& RandomStream = RequiredModule->GetRandomStream(this);
+		FRandomStream& RandomStream = RequiredModule->GetRandomStream({*this});
 
 		CurrentDelay = RequiredModule->EmitterDelay + Component->EmitterDelay;
 		if (RequiredModule->bEmitterDelayUseRange)
diff --git a/Engine/Source/Runtime/Engine/Private/Particles/ParticleGpuSimulation.cpp b/Engine/Source/Runtime/Engine/Private/Particles/ParticleGpuSimulation.cpp
--- a/Engine/Source/Runtime/Engine/Private/Particles/ParticleGpuSimulation.cpp
+++ b/Engine/Source/Runtime/Engine/Private/Particles/ParticleGpuSimulation.cpp
@@ -4137,7 +4137,7 @@ private:
 		{
 			int32 Number = 0;
 			float Rate = 0.0f;
-			if (EmitterInfo.SpawnPerUnitModule->GetSpawnAmount(this, 0, 0.0f, DeltaSeconds, Number, Rate) == false)
+			if (EmitterInfo.SpawnPerUnitModule->GetSpawnAmount({*this}, 0, 0.0f, DeltaSeconds, Number, Rate) == false)
 			{
 				SpawnRate = Rate;
 			}
@@ -4233,7 +4233,7 @@ private:
 				UParticleModule* SpawnModule = EmitterInfo.SpawnModules[ModuleIndex];
 				if (SpawnModule->bEnabled)
 				{
-					SpawnModule->Spawn(this, /*Offset=*/ 0, SpawnInfo.StartTime, TempParticle);
+					SpawnModule->Spawn({ *this, /*Offset=*/0, SpawnInfo.StartTime, TempParticle });
 				}
 			}
 
@@ -4317,7 +4317,7 @@ private:
 			UParticleModule* CurrentModule	= InCurrentLODLevel->UpdateModules[ModuleIndex];
 			if (CurrentModule && CurrentModule->bEnabled && CurrentModule->bUpdateModule && CurrentModule->bUpdateForGPUEmitter)
 			{
-				CurrentModule->Update(this, GetModuleDataOffset(HighestLODLevel->UpdateModules[ModuleIndex]), DeltaTime);
+				CurrentModule->Update({ *this, (int32)GetModuleDataOffset(HighestLODLevel->UpdateModules[ModuleIndex]), DeltaTime });
 			}
 		}
 	}
@@ -4336,7 +4336,7 @@ private:
 			UParticleModule* CurrentModule	= InCurrentLODLevel->UpdateModules[ModuleIndex];
 			if (CurrentModule && CurrentModule->bEnabled && CurrentModule->bFinalUpdateModule && CurrentModule->bUpdateForGPUEmitter)
 			{
-				CurrentModule->FinalUpdate(this, GetModuleDataOffset(HighestLODLevel->UpdateModules[ModuleIndex]), DeltaTime);
+				CurrentModule->FinalUpdate({ *this, (int32)GetModuleDataOffset(HighestLODLevel->UpdateModules[ModuleIndex]), DeltaTime });
 			}
 		}
 	}
diff --git a/Engine/Source/Runtime/Engine/Private/Particles/ParticleModules.cpp b/Engine/Source/Runtime/Engine/Private/Particles/ParticleModules.cpp
--- a/Engine/Source/Runtime/Engine/Private/Particles/ParticleModules.cpp
+++ b/Engine/Source/Runtime/Engine/Private/Particles/ParticleModules.cpp
@@ -142,16 +142,16 @@ void UParticleModule::CompileModule( FParticleEmitterBuildInfo& EmitterInfo )
 	}
 }
 
-void UParticleModule::Spawn(FParticleEmitterInstance* Owner, int32 Offset, float SpawnTime, FBaseParticle* ParticleBase)
+void UParticleModule::Spawn(const FSpawnContext& Context)
 {
 }
 
-void UParticleModule::Update(FParticleEmitterInstance* Owner, int32 Offset, float DeltaTime)
+void UParticleModule::Update(const FUpdateContext& Context)
 {
 }
 
 
-void UParticleModule::FinalUpdate(FParticleEmitterInstance* Owner, int32 Offset, float DeltaTime)
+void UParticleModule::FinalUpdate(const FUpdateContext& Context)
 {
 }
 
@@ -777,10 +777,10 @@ bool UParticleModule::IsUsedInGPUEmitter()const
 	return false;
 }
 
-FRandomStream& UParticleModule::GetRandomStream(FParticleEmitterInstance* Owner)
+FRandomStream& UParticleModule::GetRandomStream(const FContext& Context)
 {
-	FParticleRandomSeedInstancePayload* Payload = Owner->GetModuleRandomSeedInstanceData(this);
-	FRandomStream& RandomStream = (Payload != nullptr) ? Payload->RandomStream : Owner->Component->RandomStream;
+	FParticleRandomSeedInstancePayload* Payload = Context.Owner.GetModuleRandomSeedInstanceData(this);
+	FRandomStream& RandomStream = (Payload != nullptr) ? Payload->RandomStream : Context.Owner.Component->RandomStream;
 	return RandomStream;
 }
 
@@ -960,19 +960,20 @@ void UParticleModuleSourceMovement::PostEditChangeProperty(FPropertyChangedEvent
 }
 #endif // WITH_EDITOR
 
-void UParticleModuleSourceMovement::FinalUpdate(FParticleEmitterInstance* Owner, int32 Offset, float DeltaTime)
+void UParticleModuleSourceMovement::FinalUpdate(const FUpdateContext& Context)
 {
-	Super::FinalUpdate(Owner, Offset, DeltaTime);
+	Super::FinalUpdate(Context);
 	// If enabled, update the source offset translation for all particles...
+	FParticleEmitterInstance* Owner = &Context.Owner;
 	if (Owner && Owner->Component)
 	{
-		FVector FrameDiff = Owner->Component->GetComponentLocation() - Owner->Component->OldPosition;
+		FVector FrameDiff = Context.GetTransform().GetLocation() - Owner->Component->OldPosition;
 		BEGIN_UPDATE_LOOP;
 		{
 			// Rough estimation of the particle being alive for more than a frame
 			if (Particle.RelativeTime > (2.0f * DeltaTime * Particle.OneOverMaxLifetime))
 			{
-				FVector	Scale = SourceMovementScale.GetValue(Particle.RelativeTime, Owner->Component);
+				FVector	Scale = SourceMovementScale.GetValue(Particle.RelativeTime, Context.GetDistributionData());
 				Particle.Location += (Scale * FrameDiff);
 			}
 		}
@@ -1005,11 +1006,11 @@ UParticleModuleOrientationAxisLock::UParticleModuleOrientationAxisLock(const FOb
 
 //    uint8 LockAxisFlags;
 //    FVector LockAxis;
-void UParticleModuleOrientationAxisLock::Spawn(FParticleEmitterInstance* Owner, int32 Offset, float SpawnTime, FBaseParticle* ParticleBase)
+void UParticleModuleOrientationAxisLock::Spawn(const FSpawnContext& Context)
 {
 }
 
-void UParticleModuleOrientationAxisLock::Update(FParticleEmitterInstance* Owner, int32 Offset, float DeltaTime)
+void UParticleModuleOrientationAxisLock::Update(const FUpdateContext& Context)
 {
 }
 
@@ -1407,20 +1408,22 @@ void UParticleModuleMeshRotation::PostEditChangeProperty(FPropertyChangedEvent&
 }
 #endif // WITH_EDITOR
 
-void UParticleModuleMeshRotation::Spawn(FParticleEmitterInstance* Owner, int32 Offset, float SpawnTime, FBaseParticle* ParticleBase)
+void UParticleModuleMeshRotation::Spawn(const FSpawnContext& Context)
 {
-	SpawnEx(Owner, Offset, SpawnTime, &GetRandomStream(Owner), ParticleBase);
+	SpawnEx(Context, &GetRandomStream(Context));
 }
 
 
-void UParticleModuleMeshRotation::SpawnEx(FParticleEmitterInstance* Owner, int32 Offset, float SpawnTime, struct FRandomStream* InRandomStream, FBaseParticle* ParticleBase)
+void UParticleModuleMeshRotation::SpawnEx(const FSpawnContext& Context, struct FRandomStream* InRandomStream)
 {
+	FParticleEmitterInstance* Owner = &Context.Owner;
+
 	SPAWN_INIT;
 	{
 		const int32 MeshRotationOffset = Owner->GetMeshRotationOffset();
 		if (MeshRotationOffset)
 		{
-			FVector Rotation = StartRotation.GetValue(Owner->EmitterTime, Owner->Component, 0, InRandomStream);
+			FVector Rotation = StartRotation.GetValue(Owner->EmitterTime, Context.GetDistributionData(), 0, InRandomStream);
 			if (bInheritParent)
 			{
 				FRotator	Rotator	= Owner->Component->GetComponentRotation();
@@ -1505,19 +1508,21 @@ void UParticleModuleMeshRotationRate::PostEditChangeProperty(FPropertyChangedEve
 }
 #endif // WITH_EDITOR
 
-void UParticleModuleMeshRotationRate::Spawn(FParticleEmitterInstance* Owner, int32 Offset, float SpawnTime, FBaseParticle* ParticleBase)
+void UParticleModuleMeshRotationRate::Spawn(const FSpawnContext& Context)
 {
-	SpawnEx(Owner, Offset, SpawnTime, &GetRandomStream(Owner), ParticleBase);
+	SpawnEx(Context, &GetRandomStream(Context));
 }
 
-void UParticleModuleMeshRotationRate::SpawnEx(FParticleEmitterInstance* Owner, int32 Offset, float SpawnTime, struct FRandomStream* InRandomStream, FBaseParticle* ParticleBase)
+void UParticleModuleMeshRotationRate::SpawnEx(const FSpawnContext& Context, struct FRandomStream* InRandomStream)
 {
+	FParticleEmitterInstance* Owner = &Context.Owner;
+
 	SPAWN_INIT;
 	{
 		const int32 MeshRotationOffset = Owner->GetMeshRotationOffset();
 		if (MeshRotationOffset)
 		{
-			FVector StartRate = StartRotationRate.GetValue(Owner->EmitterTime, Owner->Component, 0, InRandomStream);// * ((float)PI/180.f);
+			FVector StartRate = StartRotationRate.GetValue(Owner->EmitterTime, Context.GetDistributionData(), 0, InRandomStream);// * ((float)PI/180.f);
 			FVector3f StartValue;
 			StartValue.X = StartRate.X * 360.0f;
 			StartValue.Y = StartRate.Y * 360.0f;
@@ -1596,29 +1601,32 @@ void UParticleModuleMeshRotationRateMultiplyLife::PostEditChangeProperty(FProper
 #endif // WITH_EDITOR
 
 
-void UParticleModuleMeshRotationRateMultiplyLife::Spawn(FParticleEmitterInstance* Owner, int32 Offset, float SpawnTime, FBaseParticle* ParticleBase)
+void UParticleModuleMeshRotationRateMultiplyLife::Spawn(const FSpawnContext& Context)
 {
+	FParticleEmitterInstance* Owner = &Context.Owner;
+
 	const int32 MeshRotationOffset = Owner->GetMeshRotationOffset();
 	if (MeshRotationOffset)
 	{
 		SPAWN_INIT;
 		{
 			FMeshRotationPayloadData* PayloadData = (FMeshRotationPayloadData*)((uint8*)&Particle + MeshRotationOffset);
-			FVector3f RateScale(LifeMultiplier.GetValue(Particle.RelativeTime, Owner->Component));
+			FVector3f RateScale(LifeMultiplier.GetValue(Particle.RelativeTime, Context.GetDistributionData()));
 			PayloadData->RotationRate *= RateScale;
 		}
 	}
 }
 
-void UParticleModuleMeshRotationRateMultiplyLife::Update(FParticleEmitterInstance* Owner, int32 Offset, float DeltaTime)
+void UParticleModuleMeshRotationRateMultiplyLife::Update(const FUpdateContext& Context)
 {
+	FParticleEmitterInstance* Owner = &Context.Owner;
 	const int32 MeshRotationOffset = Owner->GetMeshRotationOffset();
 	if (MeshRotationOffset)
 	{
 		BEGIN_UPDATE_LOOP;
 		{
 			FMeshRotationPayloadData* PayloadData = (FMeshRotationPayloadData*)((uint8*)&Particle + MeshRotationOffset);
-			FVector3f RateScale(LifeMultiplier.GetValue(Particle.RelativeTime, Owner->Component));
+			FVector3f RateScale(LifeMultiplier.GetValue(Particle.RelativeTime, Context.GetDistributionData()));
 			PayloadData->RotationRate *= RateScale;
 		}
 		END_UPDATE_LOOP;
@@ -1671,15 +1679,16 @@ void UParticleModuleMeshRotationRateOverLife::PostEditChangeProperty(FPropertyCh
 }
 #endif // WITH_EDITOR
 
-void UParticleModuleMeshRotationRateOverLife::Spawn(FParticleEmitterInstance* Owner, int32 Offset, float SpawnTime, FBaseParticle* ParticleBase)
+void UParticleModuleMeshRotationRateOverLife::Spawn(const FSpawnContext& Context)
 {
+	FParticleEmitterInstance* Owner = &Context.Owner;
 	const int32 MeshRotationOffset = Owner->GetMeshRotationOffset();
 	if (MeshRotationOffset)
 	{
 		SPAWN_INIT;
 		{
 			FMeshRotationPayloadData* PayloadData = (FMeshRotationPayloadData*)((uint8*)&Particle + MeshRotationOffset);
-			FVector3f RateValue(RotRate.GetValue(Particle.RelativeTime, Owner->Component));// * ((float)PI/180.f);
+			FVector3f RateValue(RotRate.GetValue(Particle.RelativeTime, Context.GetDistributionData()));// * ((float)PI/180.f);
 			RateValue.X = RateValue.X * 360.0f;
 			RateValue.Y = RateValue.Y * 360.0f;
 			RateValue.Z = RateValue.Z * 360.0f;
@@ -1696,8 +1705,9 @@ void UParticleModuleMeshRotationRateOverLife::Spawn(FParticleEmitterInstance* Ow
 	}
 }
 
-void UParticleModuleMeshRotationRateOverLife::Update(FParticleEmitterInstance* Owner, int32 Offset, float DeltaTime)
+void UParticleModuleMeshRotationRateOverLife::Update(const FUpdateContext& Context)
 {
+	FParticleEmitterInstance* Owner = &Context.Owner;
 	const int32 MeshRotationOffset = Owner->GetMeshRotationOffset();
 	if (MeshRotationOffset)
 	{
@@ -1708,7 +1718,7 @@ void UParticleModuleMeshRotationRateOverLife::Update(FParticleEmitterInstance* O
 			BEGIN_UPDATE_LOOP;
 			{
 				PayloadData = (FMeshRotationPayloadData*)((uint8*)&Particle + MeshRotationOffset);
-				RateValue = RotRate.GetValue(Particle.RelativeTime, Owner->Component);// * ((float)PI/180.f);
+				RateValue = RotRate.GetValue(Particle.RelativeTime, Context.GetDistributionData());// * ((float)PI/180.f);
 				RateValue.X = RateValue.X * 360.0f;
 				RateValue.Y = RateValue.Y * 360.0f;
 				RateValue.Z = RateValue.Z * 360.0f;
@@ -1721,7 +1731,7 @@ void UParticleModuleMeshRotationRateOverLife::Update(FParticleEmitterInstance* O
 			BEGIN_UPDATE_LOOP;
 			{
 				PayloadData = (FMeshRotationPayloadData*)((uint8*)&Particle + MeshRotationOffset);
-				RateValue = RotRate.GetValue(Particle.RelativeTime, Owner->Component);// * ((float)PI/180.f);
+				RateValue = RotRate.GetValue(Particle.RelativeTime, Context.GetDistributionData());// * ((float)PI/180.f);
 				RateValue.X = RateValue.X * 360.0f;
 				RateValue.Y = RateValue.Y * 360.0f;
 				RateValue.Z = RateValue.Z * 360.0f;
@@ -1780,16 +1790,18 @@ void UParticleModuleRotation::PostEditChangeProperty(FPropertyChangedEvent& Prop
 }
 #endif // WITH_EDITOR
 
-void UParticleModuleRotation::Spawn(FParticleEmitterInstance* Owner, int32 Offset, float SpawnTime, FBaseParticle* ParticleBase)
+void UParticleModuleRotation::Spawn(const FSpawnContext& Context)
 {
-	SpawnEx(Owner, Offset, SpawnTime, &GetRandomStream(Owner), ParticleBase);
+	SpawnEx(Context, &GetRandomStream(Context));
 }
 
-void UParticleModuleRotation::SpawnEx(FParticleEmitterInstance* Owner, int32 Offset, float SpawnTime, struct FRandomStream* InRandomStream, FBaseParticle* ParticleBase)
+void UParticleModuleRotation::SpawnEx(const FSpawnContext& Context, struct FRandomStream* InRandomStream)
 {
+	FParticleEmitterInstance* Owner = &Context.Owner;
+
 	SPAWN_INIT;
 	{
-		Particle.Rotation += (UE_PI/180.f) * 360.0f * StartRotation.GetValue(Owner->EmitterTime, Owner->Component, InRandomStream);
+		Particle.Rotation += (UE_PI/180.f) * 360.0f * StartRotation.GetValue(Owner->EmitterTime, Context.GetDistributionData(), InRandomStream);
 	}
 }
 
@@ -1859,16 +1871,18 @@ void UParticleModuleRotationRate::PostEditChangeProperty(FPropertyChangedEvent&
 }
 #endif // WITH_EDITOR
 
-void UParticleModuleRotationRate::Spawn(FParticleEmitterInstance* Owner, int32 Offset, float SpawnTime, FBaseParticle* ParticleBase)
+void UParticleModuleRotationRate::Spawn(const FSpawnContext& Context)
 {
-	SpawnEx(Owner, Offset, SpawnTime, &GetRandomStream(Owner), ParticleBase);
+	SpawnEx(Context, &GetRandomStream(Context));
 }
 
-void UParticleModuleRotationRate::SpawnEx(FParticleEmitterInstance* Owner, int32 Offset, float SpawnTime, struct FRandomStream* InRandomStream, FBaseParticle* ParticleBase)
+void UParticleModuleRotationRate::SpawnEx(const FSpawnContext& Context, struct FRandomStream* InRandomStream)
 {
+	FParticleEmitterInstance* Owner = &Context.Owner;
+
 	SPAWN_INIT;
 	{
-		float StartRotRate = (UE_PI/180.f) * 360.0f * StartRotationRate.GetValue(Owner->EmitterTime, Owner->Component, InRandomStream);
+		float StartRotRate = (UE_PI/180.f) * 360.0f * StartRotationRate.GetValue(Owner->EmitterTime, Context.GetDistributionData(), InRandomStream);
 		Particle.RotationRate += StartRotRate;
 		Particle.BaseRotationRate += StartRotRate;
 	}
@@ -1942,13 +1956,14 @@ void UParticleModuleRotationOverLifetime::PostEditChangeProperty(FPropertyChange
 }
 #endif // WITH_EDITOR
 
-void UParticleModuleRotationOverLifetime::Update(FParticleEmitterInstance* Owner, int32 Offset, float DeltaTime)
+void UParticleModuleRotationOverLifetime::Update(const FUpdateContext& Context)
 {
+	FParticleEmitterInstance* Owner = &Context.Owner;
 	if (Scale)
 	{
 		BEGIN_UPDATE_LOOP;
 		{
-			float Rotation = RotationOverLife.GetValue(Particle.RelativeTime, Owner->Component);
+			float Rotation = RotationOverLife.GetValue(Particle.RelativeTime, Context.GetDistributionData());
 			// For now, we are just using the X-value
 			Particle.Rotation = (Particle.Rotation * (Rotation * (UE_PI/180.f) * 360.0f));
 		}
@@ -1958,7 +1973,7 @@ void UParticleModuleRotationOverLifetime::Update(FParticleEmitterInstance* Owner
 	{
 		BEGIN_UPDATE_LOOP;
 		{
-			float Rotation = RotationOverLife.GetValue(Particle.RelativeTime, Owner->Component);
+			float Rotation = RotationOverLife.GetValue(Particle.RelativeTime, Context.GetDistributionData());
 			// For now, we are just using the X-value
 			Particle.Rotation = (Particle.Rotation + (Rotation * (UE_PI/180.f) * 360.0f));
 		}
@@ -2074,8 +2089,10 @@ bool UParticleModuleSubUV::IsValidForLODLevel(UParticleLODLevel* LODLevel, FStri
 }
 #endif // WITH_EDITOR
 
-void UParticleModuleSubUV::Spawn(FParticleEmitterInstance* Owner, int32 Offset, float SpawnTime, FBaseParticle* ParticleBase)
+void UParticleModuleSubUV::Spawn(const FSpawnContext& Context)
 {
+	FParticleEmitterInstance* Owner = &Context.Owner;
+
 	check(Owner->SpriteTemplate);
 
 	UParticleLODLevel* LODLevel	= Owner->SpriteTemplate->GetCurrentLODLevel(Owner);
@@ -2097,13 +2114,14 @@ void UParticleModuleSubUV::Spawn(FParticleEmitterInstance* Owner, int32 Offset,
 			PARTICLE_ELEMENT(FFullSubUVPayload, SubUVPayload);
 			CurrentOffset	= TempOffset;
 
-			SubUVPayload.ImageIndex = DetermineImageIndex(Owner, Offset, &Particle, InterpMethod, SubUVPayload, SpawnTime);
+			SubUVPayload.ImageIndex = DetermineImageIndex(Context, Context.Offset, &Particle, InterpMethod, SubUVPayload, Context.SpawnTime);
 		}
 	}
 }
 
-void UParticleModuleSubUV::Update(FParticleEmitterInstance* Owner, int32 Offset, float DeltaTime)
+void UParticleModuleSubUV::Update(const FUpdateContext& Context)
 {
+	FParticleEmitterInstance* Owner = &Context.Owner;
 	check(Owner->SpriteTemplate);
 
 	UParticleLODLevel* LODLevel	= Owner->SpriteTemplate->GetCurrentLODLevel(Owner);
@@ -2139,14 +2157,16 @@ void UParticleModuleSubUV::Update(FParticleEmitterInstance* Owner, int32 Offset,
 			PARTICLE_ELEMENT(FFullSubUVPayload, SubUVPayload);
 			CurrentOffset	= TempOffset;
 
-			SubUVPayload.ImageIndex = DetermineImageIndex(Owner, Offset, &Particle, InterpMethod, SubUVPayload, DeltaTime);
+			SubUVPayload.ImageIndex = DetermineImageIndex(Context, Offset, &Particle, InterpMethod, SubUVPayload, DeltaTime);
 		END_UPDATE_LOOP;
 	}
 }
 
-float UParticleModuleSubUV::DetermineImageIndex(FParticleEmitterInstance* Owner, int32 Offset, FBaseParticle* Particle, 
+float UParticleModuleSubUV::DetermineImageIndex(const FContext& Context, int32 Offset, FBaseParticle* Particle, 
 	EParticleSubUVInterpMethod InterpMethod, FFullSubUVPayload& SubUVPayload, float DeltaTime)
 {
+	FParticleEmitterInstance* Owner = &Context.Owner;
+
 	UParticleLODLevel* LODLevel	= Owner->SpriteTemplate->GetCurrentLODLevel(Owner);
 	check(LODLevel);
 
@@ -2162,18 +2182,18 @@ float UParticleModuleSubUV::DetermineImageIndex(FParticleEmitterInstance* Owner,
 	{
 		if (bUseRealTime == false)
 		{
-			ImageIndex = SubImageIndex.GetValue(Particle->RelativeTime, Owner->Component);
+			ImageIndex = SubImageIndex.GetValue(Particle->RelativeTime, Context.GetDistributionData());
 		}
 		else
 		{
 			UWorld* World = Owner->Component->GetWorld();
 			if ((World != NULL) && (World->GetWorldSettings() != NULL))
 			{
-				ImageIndex = SubImageIndex.GetValue(Particle->RelativeTime / World->GetWorldSettings()->GetEffectiveTimeDilation(), Owner->Component);
+				ImageIndex = SubImageIndex.GetValue(Particle->RelativeTime / World->GetWorldSettings()->GetEffectiveTimeDilation(), Context.GetDistributionData());
 			}
 			else
 			{
-				ImageIndex = SubImageIndex.GetValue(Particle->RelativeTime, Owner->Component);
+				ImageIndex = SubImageIndex.GetValue(Particle->RelativeTime, Context.GetDistributionData());
 			}
 		}
 
@@ -2188,7 +2208,7 @@ float UParticleModuleSubUV::DetermineImageIndex(FParticleEmitterInstance* Owner,
 			((Particle->RelativeTime - SubUVPayload.RandomImageTime) > LODLevel->RequiredModule->RandomImageTime) ||
 			(SubUVPayload.RandomImageTime == 0.0f))
 		{
-			ImageIndex = GetRandomStream(Owner).RandHelper(TotalSubImages);
+			ImageIndex = GetRandomStream(Context).RandHelper(TotalSubImages);
 			SubUVPayload.RandomImageTime	= Particle->RelativeTime;
 		}
 
@@ -2270,8 +2290,9 @@ struct FSubUVMovieParticlePayload
 };
 
 
-void UParticleModuleSubUVMovie::Spawn(FParticleEmitterInstance* Owner, int32 Offset, float SpawnTime, FBaseParticle* ParticleBase)
+void UParticleModuleSubUVMovie::Spawn(const FSpawnContext& Context)
 {
+	FParticleEmitterInstance* Owner = &Context.Owner;
 	check(Owner->SpriteTemplate);
 
 	UParticleLODLevel* LODLevel	= Owner->SpriteTemplate->GetCurrentLODLevel(Owner);
@@ -2287,7 +2308,7 @@ void UParticleModuleSubUVMovie::Spawn(FParticleEmitterInstance* Owner, int32 Off
 	// Movies only work w/ Linear modes...
 	if ((InterpMethod != PSUVIM_Linear) && (InterpMethod != PSUVIM_Linear_Blend))
 	{
-		return UParticleModuleSubUV::Spawn(Owner, Offset, SpawnTime, ParticleBase);
+		return UParticleModuleSubUV::Spawn(Context);
 	}
 
 	if (!LODLevel->TypeDataModule || LODLevel->TypeDataModule->SupportsSubUV())
@@ -2310,7 +2331,7 @@ void UParticleModuleSubUVMovie::Spawn(FParticleEmitterInstance* Owner, int32 Off
 			PARTICLE_ELEMENT(FFullSubUVPayload, SubUVPayload);
 			CurrentOffset	= TempOffset;
 
-			const float UserSetFrameRate = FrameRate.GetValue(bUseEmitterTime ? Owner->EmitterTime : Particle.RelativeTime, Owner->Component);
+			const float UserSetFrameRate = FrameRate.GetValue(bUseEmitterTime ? Owner->EmitterTime : Particle.RelativeTime, Context.GetDistributionData());
 			PARTICLE_ELEMENT(FSubUVMovieParticlePayload, MoviePayload);
 			MoviePayload.Time = 0.0f;
 			if (StartingFrame > 1)
@@ -2320,7 +2341,7 @@ void UParticleModuleSubUVMovie::Spawn(FParticleEmitterInstance* Owner, int32 Off
 			}
 			else if (StartingFrame == 0)
 			{
-				FRandomStream& RandomStream = GetRandomStream(Owner);
+				FRandomStream& RandomStream = GetRandomStream(Context);
 				MoviePayload.Time = FMath::TruncToFloat(RandomStream.FRand() * (TotalSubImages-1));
 			}
 
@@ -2337,11 +2358,12 @@ uint32 UParticleModuleSubUVMovie::RequiredBytes(UParticleModuleTypeDataBase* Typ
 }
 
 
-float UParticleModuleSubUVMovie::DetermineImageIndex(FParticleEmitterInstance* Owner, int32 Offset, FBaseParticle* Particle, 
+float UParticleModuleSubUVMovie::DetermineImageIndex(const FContext& Context, int32 Offset, FBaseParticle* Particle, 
 	EParticleSubUVInterpMethod InterpMethod, FFullSubUVPayload& SubUVPayload, float DeltaTime)
 {
+	FParticleEmitterInstance* Owner = &Context.Owner;
 	FSubUVMovieParticlePayload& MoviePayload = *((FSubUVMovieParticlePayload*)((uint8*)Particle + Offset));
-	float UserSetFrameRate = FrameRate.GetValue(bUseEmitterTime ? Owner->EmitterTime : Particle->RelativeTime, Owner->Component);
+	float UserSetFrameRate = FrameRate.GetValue(bUseEmitterTime ? Owner->EmitterTime : Particle->RelativeTime, Context.GetDistributionData());
 	if (bUseRealTime == false)
 	{
 		MoviePayload.Time += DeltaTime;
@@ -2429,17 +2451,19 @@ void UParticleModuleRotationRateMultiplyLife::PostEditChangeProperty(FPropertyCh
 }
 #endif // WITH_EDITOR
 
-void UParticleModuleRotationRateMultiplyLife::Spawn(FParticleEmitterInstance* Owner, int32 Offset, float SpawnTime, FBaseParticle* ParticleBase)
+void UParticleModuleRotationRateMultiplyLife::Spawn(const FSpawnContext& Context)
 {
+	FParticleEmitterInstance* Owner = &Context.Owner;
 	SPAWN_INIT;
-	float RateScale = LifeMultiplier.GetValue(Particle.RelativeTime, Owner->Component);
+	float RateScale = LifeMultiplier.GetValue(Particle.RelativeTime, Context.GetDistributionData());
 	Particle.RotationRate *= RateScale;
 }
 
-void UParticleModuleRotationRateMultiplyLife::Update(FParticleEmitterInstance* Owner, int32 Offset, float DeltaTime)
+void UParticleModuleRotationRateMultiplyLife::Update(const FUpdateContext& Context)
 {
+	FParticleEmitterInstance* Owner = &Context.Owner;
 	BEGIN_UPDATE_LOOP;
-		float RateScale = LifeMultiplier.GetValue(Particle.RelativeTime, Owner->Component);
+		float RateScale = LifeMultiplier.GetValue(Particle.RelativeTime, Context.GetDistributionData());
 		Particle.RotationRate *= RateScale;
 	END_UPDATE_LOOP;
 }
@@ -2489,16 +2513,17 @@ void UParticleModuleAccelerationConstant::CompileModule( FParticleEmitterBuildIn
 	EmitterInfo.ConstantAcceleration = Acceleration;
 }
 
-void UParticleModuleAccelerationConstant::Spawn(FParticleEmitterInstance* Owner, int32 Offset, float SpawnTime, FBaseParticle* ParticleBase)
+void UParticleModuleAccelerationConstant::Spawn(const FSpawnContext& Context)
 {
+	FParticleEmitterInstance* Owner = &Context.Owner;
 	SPAWN_INIT;
 	UParticleLODLevel* LODLevel	= Owner->SpriteTemplate->GetCurrentLODLevel(Owner);
 	check(LODLevel);
 	if (bAlwaysInWorldSpace && LODLevel->RequiredModule->bUseLocalSpace)
 	{
-		FVector3f LocalAcceleration(Owner->Component->GetComponentTransform().InverseTransformVector(Acceleration));
-		Particle.Velocity		+= LocalAcceleration * SpawnTime;
-		Particle.BaseVelocity	+= LocalAcceleration * SpawnTime;
+		FVector3f LocalAcceleration(Context.GetTransform().InverseTransformVector(Acceleration));
+		Particle.Velocity		+= LocalAcceleration * Context.SpawnTime;
+		Particle.BaseVelocity	+= LocalAcceleration * Context.SpawnTime;
 	}
 	else
 	{
@@ -2507,13 +2532,14 @@ void UParticleModuleAccelerationConstant::Spawn(FParticleEmitterInstance* Owner,
 		{
 			LocalAcceleration = FVector4f(Owner->EmitterToSimulation.TransformVector((FVector)LocalAcceleration));
 		}
-		Particle.Velocity		+= LocalAcceleration * SpawnTime;
-		Particle.BaseVelocity	+= LocalAcceleration * SpawnTime;
+		Particle.Velocity		+= LocalAcceleration * Context.SpawnTime;
+		Particle.BaseVelocity	+= LocalAcceleration * Context.SpawnTime;
 	}
 }
 
-void UParticleModuleAccelerationConstant::Update(FParticleEmitterInstance* Owner, int32 Offset, float DeltaTime)
+void UParticleModuleAccelerationConstant::Update(const FUpdateContext& Context)
 {
+	FParticleEmitterInstance* Owner = &Context.Owner;
 	if ((Owner == NULL) || (Owner->ActiveParticles <= 0) || 
 		(Owner->ParticleData == NULL) || (Owner->ParticleIndices == NULL))
 	{
@@ -2525,7 +2551,7 @@ void UParticleModuleAccelerationConstant::Update(FParticleEmitterInstance* Owner
 	FPlatformMisc::Prefetch(Owner->ParticleData, (Owner->ParticleIndices[0] * Owner->ParticleStride) + PLATFORM_CACHE_LINE_SIZE);
 	if (bAlwaysInWorldSpace && LODLevel->RequiredModule->bUseLocalSpace)
 	{
-		FTransform Mat = Owner->Component->GetComponentTransform();
+		FTransform Mat = Context.GetTransform();
 		FVector3f LocalAcceleration(Mat.InverseTransformVector(Acceleration));
 		BEGIN_UPDATE_LOOP;
 		{
@@ -2624,11 +2650,12 @@ bool UParticleModuleAccelerationDrag::IsValidForLODLevel(UParticleLODLevel* LODL
 }
 #endif // WITH_EDITOR
 
-void UParticleModuleAccelerationDrag::Update(FParticleEmitterInstance* Owner, int32 Offset, float DeltaTime)
+void UParticleModuleAccelerationDrag::Update(const FUpdateContext& Context)
 {
+	FParticleEmitterInstance* Owner = &Context.Owner;
 	BEGIN_UPDATE_LOOP;
 	{
-		FVector3f Drag  = Particle.Velocity * -DragCoefficientRaw.GetValue(Particle.RelativeTime, Owner->Component);
+		FVector3f Drag  = Particle.Velocity * -DragCoefficientRaw.GetValue(Particle.RelativeTime, Context.GetDistributionData());
 		Particle.Velocity		+= Drag * DeltaTime;
 		Particle.BaseVelocity	+= Drag * DeltaTime;
 	}
@@ -2759,9 +2786,11 @@ void UParticleModuleAttractorPointGravity::PostEditChangeProperty(FPropertyChang
 }
 #endif // WITH_EDITOR
 
-void UParticleModuleAttractorPointGravity::Render3DPreview(FParticleEmitterInstance* Owner, const FSceneView* View,FPrimitiveDrawInterface* PDI)
+void UParticleModuleAttractorPointGravity::Render3DPreview(const FPreviewContext& Context)
 {
 #if WITH_EDITOR
+	FPrimitiveDrawInterface* PDI = Context.PDI;
+
 	// Draw a wire star at the position.
 	DrawWireStar(PDI,Position, 10.0f, ModuleEditorColor, SDPG_World);
 
@@ -2812,23 +2841,34 @@ void UParticleModuleAcceleration::PostEditChangeProperty(FPropertyChangedEvent&
 }
 #endif // WITH_EDITOR
 
-void UParticleModuleAcceleration::Spawn(FParticleEmitterInstance* Owner, int32 Offset, float SpawnTime, FBaseParticle* ParticleBase)
+const FTransform& UParticleModuleAcceleration::FContext::GetTransform() const
+{
+	return Owner.Component->GetComponentTransform();
+}
+
+UObject* UParticleModuleAcceleration::FContext::GetDistributionData() const
+{
+	return Owner.Component;
+}
+
+void UParticleModuleAcceleration::Spawn(const FSpawnContext& Context)
 {
+	FParticleEmitterInstance* Owner = &Context.Owner;
 	SPAWN_INIT;
 	PARTICLE_ELEMENT(FVector3f, UsedAcceleration);
-	UsedAcceleration = FVector3f(Acceleration.GetValue(Owner->EmitterTime, Owner->Component));
+	UsedAcceleration = FVector3f(Acceleration.GetValue(Owner->EmitterTime, Context.GetDistributionData()));
 	if ((bApplyOwnerScale == true) && Owner && Owner->Component)
 	{
-		FVector3f Scale(Owner->Component->GetComponentTransform().GetScale3D());
+		FVector3f Scale(Context.GetTransform().GetScale3D());
 		UsedAcceleration *= Scale;
 	}
 	UParticleLODLevel* LODLevel	= Owner->SpriteTemplate->GetCurrentLODLevel(Owner);
 	check(LODLevel);
 	if (bAlwaysInWorldSpace && LODLevel->RequiredModule->bUseLocalSpace)
 	{
-		FVector3f TempUsedAcceleration = (FVector3f)Owner->Component->GetComponentTransform().InverseTransformVector(FVector(UsedAcceleration));
-		Particle.Velocity		+= TempUsedAcceleration * SpawnTime;
-		Particle.BaseVelocity	+= TempUsedAcceleration * SpawnTime;
+		FVector3f TempUsedAcceleration = (FVector3f)Context.GetTransform().InverseTransformVector(FVector(UsedAcceleration));
+		Particle.Velocity		+= TempUsedAcceleration * Context.SpawnTime;
+		Particle.BaseVelocity	+= TempUsedAcceleration * Context.SpawnTime;
 	}
 	else
 	{
@@ -2836,13 +2876,14 @@ void UParticleModuleAcceleration::Spawn(FParticleEmitterInstance* Owner, int32 O
 		{
 			UsedAcceleration = (FVector4f)Owner->EmitterToSimulation.TransformVector(FVector(UsedAcceleration));
 		}
-		Particle.Velocity		+= UsedAcceleration * SpawnTime;
-		Particle.BaseVelocity	+= UsedAcceleration * SpawnTime;
+		Particle.Velocity		+= UsedAcceleration * Context.SpawnTime;
+		Particle.BaseVelocity	+= UsedAcceleration * Context.SpawnTime;
 	}
 }
 
-void UParticleModuleAcceleration::Update(FParticleEmitterInstance* Owner, int32 Offset, float DeltaTime)
+void UParticleModuleAcceleration::Update(const FUpdateContext& Context)
 {
+	FParticleEmitterInstance* Owner = &Context.Owner;
 	if ((Owner == NULL) || (Owner->ActiveParticles <= 0) || 
 		(Owner->ParticleData == NULL) || (Owner->ParticleIndices == NULL))
 	{
@@ -2854,7 +2895,7 @@ void UParticleModuleAcceleration::Update(FParticleEmitterInstance* Owner, int32
 	FPlatformMisc::Prefetch(Owner->ParticleData, (Owner->ParticleIndices[0] * Owner->ParticleStride) + PLATFORM_CACHE_LINE_SIZE);
 	if (bAlwaysInWorldSpace && LODLevel->RequiredModule->bUseLocalSpace)
 	{
-		FTransform Mat = Owner->Component->GetComponentTransform();
+		FTransform Mat = Context.GetTransform();
 		BEGIN_UPDATE_LOOP;
 		{
 			FVector3f& UsedAcceleration = *((FVector3f*)(ParticleBase + CurrentOffset));																\
@@ -2938,16 +2979,17 @@ void UParticleModuleAccelerationOverLifetime::PostEditChangeProperty(FPropertyCh
 }
 #endif // WITH_EDITOR
 
-void UParticleModuleAccelerationOverLifetime::Update(FParticleEmitterInstance* Owner, int32 Offset, float DeltaTime)
+void UParticleModuleAccelerationOverLifetime::Update(const FUpdateContext& Context)
 {
+	FParticleEmitterInstance* Owner = &Context.Owner;
 	UParticleLODLevel* LODLevel	= Owner->SpriteTemplate->GetCurrentLODLevel(Owner);
 	check(LODLevel);
 	if (bAlwaysInWorldSpace && LODLevel->RequiredModule->bUseLocalSpace)
 	{
-		FTransform Mat = Owner->Component->GetComponentTransform();
+		FTransform Mat = Context.GetTransform();
 		BEGIN_UPDATE_LOOP;
 			// Acceleration should always be in world space...
-			FVector Accel = AccelOverLife.GetValue(Particle.RelativeTime, Owner->Component);
+			FVector Accel = AccelOverLife.GetValue(Particle.RelativeTime, Context.GetDistributionData());
 			Accel = Mat.InverseTransformVector(Accel);
 			Particle.Velocity		+= (FVector3f)Accel * DeltaTime;
 			Particle.BaseVelocity	+= (FVector3f)Accel * DeltaTime;
@@ -2957,7 +2999,7 @@ void UParticleModuleAccelerationOverLifetime::Update(FParticleEmitterInstance* O
 	{
 		BEGIN_UPDATE_LOOP;
 		// Acceleration should always be in world space...
-		FVector Accel = AccelOverLife.GetValue(Particle.RelativeTime, Owner->Component);
+		FVector Accel = AccelOverLife.GetValue(Particle.RelativeTime, Context.GetDistributionData());
 		Particle.Velocity		+= (FVector3f)Accel * DeltaTime;
 		Particle.BaseVelocity	+= (FVector3f)Accel * DeltaTime;
 		END_UPDATE_LOOP;
@@ -3052,18 +3094,19 @@ static TAutoConsoleVariable<int32> CVarParticleDefaultLightInverseExposureBlend(
 	ECVF_Default
 );
 
-void UParticleModuleLight::SpawnEx(FParticleEmitterInstance* Owner, int32 Offset, float SpawnTime, struct FRandomStream* InRandomStream, FBaseParticle* ParticleBase)
+void UParticleModuleLight::SpawnEx(const FSpawnContext& Context, struct FRandomStream* InRandomStream)
 {
 	int32 ParticleLightQuality = CVarParticleLightQuality.GetValueOnAnyThread();
 	if (ParticleLightQuality > 0)
 	{
+		FParticleEmitterInstance* Owner = &Context.Owner;
 		SPAWN_INIT;
 		PARTICLE_ELEMENT(FLightParticlePayload, LightData);
-		const float Brightness = BrightnessOverLife.GetValue(Particle.RelativeTime, Owner->Component, InRandomStream);
-		LightData.ColorScale = (FVector3f)ColorScaleOverLife.GetValue(Particle.RelativeTime, Owner->Component, 0, InRandomStream) * Brightness;
-		LightData.RadiusScale = RadiusScale.GetValue(Owner->EmitterTime, Owner->Component, InRandomStream);
+		const float Brightness = BrightnessOverLife.GetValue(Particle.RelativeTime, Context.GetDistributionData(), InRandomStream);
+		LightData.ColorScale = (FVector3f)ColorScaleOverLife.GetValue(Particle.RelativeTime, Context.GetDistributionData(), 0, InRandomStream) * Brightness;
+		LightData.RadiusScale = RadiusScale.GetValue(Owner->EmitterTime, Context.GetDistributionData(), InRandomStream);
 		// Exponent of 0 is interpreted by renderer as inverse squared falloff
-		LightData.LightExponent = bUseInverseSquaredFalloff ? 0 : LightExponent.GetValue(Owner->EmitterTime, Owner->Component, InRandomStream);
+		LightData.LightExponent = bUseInverseSquaredFalloff ? 0 : LightExponent.GetValue(Owner->EmitterTime, Context.GetDistributionData(), InRandomStream);
 		LightData.InverseExposureBlend = bOverrideInverseExposureBlend ? InverseExposureBlend : CVarParticleDefaultLightInverseExposureBlend.GetValueOnAnyThread();
 		const float RandomNumber = InRandomStream->GetFraction();
 		LightData.bValid = RandomNumber < SpawnFraction;
@@ -3187,13 +3230,14 @@ void UParticleModuleLight::UpdateHQLight(UPointLightComponent* PointLightCompone
 	}
 }
 
-void UParticleModuleLight::Spawn(FParticleEmitterInstance* Owner, int32 Offset, float SpawnTime, FBaseParticle* ParticleBase)
+void UParticleModuleLight::Spawn(const FSpawnContext& Context)
 {
-	SpawnEx(Owner, Offset, SpawnTime, &GetRandomStream(Owner), ParticleBase);
+	SpawnEx(Context, &GetRandomStream(Context));
 }
 
-void UParticleModuleLight::Update(FParticleEmitterInstance* Owner, int32 Offset, float DeltaTime)
+void UParticleModuleLight::Update(const FUpdateContext& Context)
 {
+	FParticleEmitterInstance* Owner = &Context.Owner;
 	if ((Owner == NULL) || (Owner->ActiveParticles <= 0) || 
 		(Owner->ParticleData == NULL) || (Owner->ParticleIndices == NULL))
 	{
@@ -3220,8 +3264,8 @@ void UParticleModuleLight::Update(FParticleEmitterInstance* Owner, int32 Offset,
 	BEGIN_UPDATE_LOOP;
 	{
 		PARTICLE_ELEMENT(FLightParticlePayload,	Data);
-		const float Brightness = BrightnessOverLife.GetValue(Particle.RelativeTime, Owner->Component);
-		Data.ColorScale = (FVector3f)ColorScaleOverLife.GetValue(Particle.RelativeTime, Owner->Component) * Brightness;
+		const float Brightness = BrightnessOverLife.GetValue(Particle.RelativeTime, Context.GetDistributionData());
+		Data.ColorScale = (FVector3f)ColorScaleOverLife.GetValue(Particle.RelativeTime, Context.GetDistributionData()) * Brightness;
 
 		if (bHighQualityLights && (Data.LightId != 0))
 		{
@@ -3288,9 +3332,10 @@ void UParticleModuleLight::SetToSensibleDefaults(UParticleEmitter* Owner)
 	}
 }
 
-void UParticleModuleLight::Render3DPreview(FParticleEmitterInstance* Owner, const FSceneView* View,FPrimitiveDrawInterface* PDI)
+void UParticleModuleLight::Render3DPreview(const FPreviewContext& Context)
 {
 #if WITH_EDITOR
+	FParticleEmitterInstance* Owner = &Context.Owner;
 	if ((Owner == NULL) || (Owner->ActiveParticles <= 0) || 
 		(Owner->ParticleData == NULL) || (Owner->ParticleIndices == NULL))
 	{
@@ -3303,13 +3348,15 @@ void UParticleModuleLight::Render3DPreview(FParticleEmitterInstance* Owner, cons
 		int32 Offset = 0;
 		UParticleLODLevel* LODLevel	= Owner->SpriteTemplate->GetCurrentLODLevel(Owner);
 		const bool bLocalSpace = LODLevel->RequiredModule->bUseLocalSpace;
-		const FVector Scale = Owner->Component->GetComponentTransform().GetScale3D();
+		const FVector Scale = Context.GetTransform().GetScale3D();
 		const FMatrix LocalToWorld = Owner->EmitterToSimulation * Owner->SimulationToWorld;
 		check(LODLevel);
 
 		const uint8* ParticleData = Owner->ParticleData;
 		const uint16* ParticleIndices = Owner->ParticleIndices;
 
+		FPrimitiveDrawInterface* PDI = Context.PDI;
+
 		for (int32 i = 0; i < Owner->ActiveParticles; i++)
 		{
 			DECLARE_PARTICLE_CONST(Particle, ParticleData + Owner->ParticleStride * ParticleIndices[i]);
@@ -3618,18 +3665,19 @@ void UParticleModuleKillBox::PostEditChangeProperty(FPropertyChangedEvent& Prope
 //	uint32 bAbsolute:1;
 	//## END PROPS ParticleModuleKillBox
 
-void UParticleModuleKillBox::Update(FParticleEmitterInstance* Owner, int32 Offset, float DeltaTime)
+void UParticleModuleKillBox::Update(const FUpdateContext& Context)
 {
+	FParticleEmitterInstance* Owner = &Context.Owner;
 	UParticleLODLevel* LODLevel	= Owner->SpriteTemplate->GetCurrentLODLevel(Owner);
 
-	FVector LL = LowerLeftCorner.GetValue(Owner->EmitterTime, Owner->Component);
-	FVector UR = UpperRightCorner.GetValue(Owner->EmitterTime, Owner->Component);
+	FVector LL = LowerLeftCorner.GetValue(Owner->EmitterTime, Context.GetDistributionData());
+	FVector UR = UpperRightCorner.GetValue(Owner->EmitterTime, Context.GetDistributionData());
 	FVector CheckLL = LL.ComponentMin(UR);
 	FVector CheckUR = UR.ComponentMax(LL);
 	if (bAbsolute == false)
 	{
-		CheckLL += Owner->Component->GetComponentLocation();
-		CheckUR += Owner->Component->GetComponentLocation();
+		CheckLL += Context.GetTransform().GetLocation();
+		CheckUR += Context.GetTransform().GetLocation();
 	}
 	FBox CheckBox = FBox(CheckLL, CheckUR);
 
@@ -3639,11 +3687,11 @@ void UParticleModuleKillBox::Update(FParticleEmitterInstance* Owner, int32 Offse
 
 		if (LODLevel->RequiredModule->bUseLocalSpace)
 		{
-			Position = Owner->Component->GetComponentTransform().TransformVector(Position);
+			Position = Context.GetTransform().TransformVector(Position);
 		}
 		else if ((bAxisAlignedAndFixedSize == false) && (bAbsolute == false))
 		{
-			Position = Owner->Component->GetComponentTransform().InverseTransformPosition(Position) + Owner->Component->GetComponentTransform().GetLocation();
+			Position = Context.GetTransform().InverseTransformPosition(Position) + Context.GetTransform().GetLocation();
 		}
 
 		// Determine if the particle is inside the box
@@ -3662,11 +3710,12 @@ void UParticleModuleKillBox::Update(FParticleEmitterInstance* Owner, int32 Offse
 	END_UPDATE_LOOP;
 }
 
-void UParticleModuleKillBox::Render3DPreview(FParticleEmitterInstance* Owner, const FSceneView* View,FPrimitiveDrawInterface* PDI)
+void UParticleModuleKillBox::Render3DPreview(const FPreviewContext& Context)
 {
 #if WITH_EDITOR
-	FVector CheckLL = LowerLeftCorner.GetValue(Owner->EmitterTime, Owner->Component);
-	FVector CheckUR = UpperRightCorner.GetValue(Owner->EmitterTime, Owner->Component);
+	FParticleEmitterInstance* Owner = &Context.Owner;
+	FVector CheckLL = LowerLeftCorner.GetValue(Owner->EmitterTime, Context.GetDistributionData());
+	FVector CheckUR = UpperRightCorner.GetValue(Owner->EmitterTime, Context.GetDistributionData());
 	
 	TArray<FVector> KillboxVerts;
 	if (bAxisAlignedAndFixedSize == false)
@@ -3688,16 +3737,18 @@ void UParticleModuleKillBox::Render3DPreview(FParticleEmitterInstance* Owner, co
 		{
 			for (int32 i = 0; i < 8; ++i)
 			{
-				KillboxVerts[i] = Owner->Component->GetComponentTransform().TransformPosition(KillboxVerts[i]);
+				KillboxVerts[i] = Context.GetTransform().TransformPosition(KillboxVerts[i]);
 			}
 		}
 		else
 		{
-			CheckLL += Owner->Component->GetComponentLocation();
-			CheckUR += Owner->Component->GetComponentLocation();
+			CheckLL += Context.GetTransform().GetLocation();
+			CheckUR += Context.GetTransform().GetLocation();
 		}
 	}
 
+	FPrimitiveDrawInterface* PDI = Context.PDI;
+
 	if (bAxisAlignedAndFixedSize == false)
 	{
 		PDI->DrawLine(KillboxVerts[0], KillboxVerts[1], ModuleEditorColor, SDPG_World);
@@ -3761,20 +3812,21 @@ void UParticleModuleKillHeight::PostEditChangeProperty(FPropertyChangedEvent& Pr
 //	struct FRawDistributionFloat Height;
 //	uint32 bAbsolute:1;
 	//## END PROPS ParticleModuleKillHeight
-void UParticleModuleKillHeight::Update(FParticleEmitterInstance* Owner, int32 Offset, float DeltaTime)
+void UParticleModuleKillHeight::Update(const FUpdateContext& Context)
 {
+	FParticleEmitterInstance* Owner = &Context.Owner;
 	UParticleLODLevel* LODLevel	= Owner->SpriteTemplate->GetCurrentLODLevel(Owner);
 
-	float CheckHeight = Height.GetValue(Owner->EmitterTime, Owner->Component);
+	float CheckHeight = Height.GetValue(Owner->EmitterTime, Context.GetDistributionData());
 	if (bApplyPSysScale == true)
 	{
-		FVector OwnerScale = Owner->Component->GetComponentTransform().GetScale3D();
+		FVector OwnerScale = Context.GetTransform().GetScale3D();
 		CheckHeight *= OwnerScale.Z;
 	}
 
 	if (bAbsolute == false)
 	{
-		CheckHeight += Owner->Component->GetComponentLocation().Z;
+		CheckHeight += Context.GetTransform().GetLocation().Z;
 	}
 
 	BEGIN_UPDATE_LOOP;
@@ -3783,7 +3835,7 @@ void UParticleModuleKillHeight::Update(FParticleEmitterInstance* Owner, int32 Of
 
 		if (LODLevel->RequiredModule->bUseLocalSpace)
 		{
-			Position = Owner->Component->GetComponentTransform().TransformVector(Position);
+			Position = Context.GetTransform().TransformVector(Position);
 		}
 
 		if ((bFloor == true) && (Position.Z < CheckHeight))
@@ -3800,19 +3852,20 @@ void UParticleModuleKillHeight::Update(FParticleEmitterInstance* Owner, int32 Of
 	END_UPDATE_LOOP;
 }
 
-void UParticleModuleKillHeight::Render3DPreview(FParticleEmitterInstance* Owner, const FSceneView* View,FPrimitiveDrawInterface* PDI)
+void UParticleModuleKillHeight::Render3DPreview(const FPreviewContext& Context)
 {
 #if WITH_EDITOR
 	FVector OwnerPosition(0.0f);
+	FParticleEmitterInstance* Owner = &Context.Owner;
 	if ((Owner != NULL) && (Owner->Component != NULL))
 	{
-		OwnerPosition = Owner->Component->GetComponentLocation();
+		OwnerPosition = Context.GetTransform().GetLocation();
 
-		float CheckHeight = Height.GetValue(Owner->EmitterTime, Owner->Component);
+		float CheckHeight = Height.GetValue(Owner->EmitterTime, Context.GetDistributionData());
 		float ScaleValue = 1.0f;
 		if (bApplyPSysScale == true)
 		{
-			FVector OwnerScale = Owner->Component->GetComponentTransform().GetScale3D();
+			FVector OwnerScale = Context.GetTransform().GetScale3D();
 			ScaleValue = OwnerScale.Z;
 		}
 		CheckHeight *= ScaleValue;
@@ -3827,6 +3880,8 @@ void UParticleModuleKillHeight::Render3DPreview(FParticleEmitterInstance* Owner,
 		FVector Pt3 = FVector(OwnerPosition.X - Offset, OwnerPosition.Y + Offset, CheckHeight);
 		FVector Pt4 = FVector(OwnerPosition.X + Offset, OwnerPosition.Y + Offset, CheckHeight);
 
+		FPrimitiveDrawInterface* PDI = Context.PDI;
+
 		PDI->DrawLine(Pt1, Pt2, ModuleEditorColor, SDPG_World);
 		PDI->DrawLine(Pt1, Pt3, ModuleEditorColor, SDPG_World);
 		PDI->DrawLine(Pt2, Pt4, ModuleEditorColor, SDPG_World);
@@ -3884,16 +3939,18 @@ void UParticleModuleLifetime::CompileModule( struct FParticleEmitterBuildInfo& E
 	EmitterInfo.SpawnModules.Add( this );
 }
 
-void UParticleModuleLifetime::Spawn(FParticleEmitterInstance* Owner, int32 Offset, float SpawnTime, FBaseParticle* ParticleBase)
+void UParticleModuleLifetime::Spawn(const FSpawnContext& Context)
 {
-	SpawnEx(Owner, Offset, SpawnTime, &GetRandomStream(Owner), ParticleBase);
+	SpawnEx(Context, &GetRandomStream(Context));
 }
 
-void UParticleModuleLifetime::SpawnEx(FParticleEmitterInstance* Owner, int32 Offset, float SpawnTime, struct FRandomStream* InRandomStream, FBaseParticle* ParticleBase)
+void UParticleModuleLifetime::SpawnEx(const FSpawnContext& Context, struct FRandomStream* InRandomStream)
 {
 	SPAWN_INIT;
 	{
-		float MaxLifetime = Lifetime.GetValue(Owner->EmitterTime, Owner->Component, InRandomStream);
+		FParticleEmitterInstance* Owner = &Context.Owner;
+
+		float MaxLifetime = Lifetime.GetValue(Owner->EmitterTime, Context.GetDistributionData(), InRandomStream);
 		if(Particle.OneOverMaxLifetime > 0.f)
 		{
 			// Another module already modified lifetime.
@@ -3905,7 +3962,7 @@ void UParticleModuleLifetime::SpawnEx(FParticleEmitterInstance* Owner, int32 Off
 			Particle.OneOverMaxLifetime = MaxLifetime > 0.f ? 1.f / MaxLifetime : 0.f;
 		}
 		//If the relative time is already > 1.0f then we don't want to be setting it. Some modules use this to mark a particle as dead during spawn.
-		Particle.RelativeTime = Particle.RelativeTime > 1.0f ? Particle.RelativeTime : SpawnTime * Particle.OneOverMaxLifetime;
+		Particle.RelativeTime = Particle.RelativeTime > 1.0f ? Particle.RelativeTime : Context.SpawnTime * Particle.OneOverMaxLifetime;
 	}
 }
 
@@ -3928,7 +3985,7 @@ float UParticleModuleLifetime::GetMaxLifetime()
 	return Max;
 }
 
-float UParticleModuleLifetime::GetLifetimeValue(FParticleEmitterInstance* Owner, float InTime, UObject* Data)
+float UParticleModuleLifetime::GetLifetimeValue(const FContext& Context, float InTime, UObject* Data)
 {
 	return Lifetime.GetValue(InTime, Data);
 }
@@ -3961,9 +4018,9 @@ void UParticleModuleLifetime_Seeded::EmitterLoopingNotify(FParticleEmitterInstan
 	}
 }
 
-float UParticleModuleLifetime_Seeded::GetLifetimeValue(FParticleEmitterInstance* Owner, float InTime, UObject* Data )
+float UParticleModuleLifetime_Seeded::GetLifetimeValue(const FContext& Context, float InTime, UObject* Data )
 {
-	return Lifetime.GetValue(InTime, Data, &GetRandomStream(Owner));
+	return Lifetime.GetValue(InTime, Data, &GetRandomStream(Context));
 }
 
 /*-----------------------------------------------------------------------------
@@ -4014,7 +4071,7 @@ void UParticleModuleAttractorLine::PostEditChangeProperty(FPropertyChangedEvent&
 }
 #endif // WITH_EDITOR
 
-void UParticleModuleAttractorLine::Update(FParticleEmitterInstance* Owner, int32 Offset, float DeltaTime)
+void UParticleModuleAttractorLine::Update(const FUpdateContext& Context)
 {
 	FVector Line = EndPoint1 - EndPoint0;
 
@@ -4027,9 +4084,11 @@ void UParticleModuleAttractorLine::Update(FParticleEmitterInstance* Owner, int32
 	FVector LineNorm = Line;
 	LineNorm.Normalize();
 
+	FParticleEmitterInstance* Owner = &Context.Owner;
+
 	BEGIN_UPDATE_LOOP;
 		// Determine the position of the particle projected on the line
-		FVector AdjustedLocation = Particle.Location - Owner->Component->GetComponentLocation();
+		FVector AdjustedLocation = Particle.Location - Context.GetTransform().GetLocation();
 		FVector EP02Particle = AdjustedLocation - EndPoint0;
 
 		FVector ProjectedParticle = Line * (Line | EP02Particle) / (Line | Line);
@@ -4068,14 +4127,14 @@ void UParticleModuleAttractorLine::Update(FParticleEmitterInstance* Owner, int32
 		if (bProcess)
 		{
 			// Look up the Range and Strength at that position on the line
-			float AttractorRange = Range.GetValue(fRatio, Owner->Component);
+			float AttractorRange = Range.GetValue(fRatio, Context.GetDistributionData());
 			FVector LineToPoint = AdjustedLocation - ProjectedParticle;
 			float Distance = LineToPoint.Size();
 
 			if ((AttractorRange > 0) && (Distance <= AttractorRange))
 			{
 				// Adjust the strength based on the range ratio
-				float AttractorStrength = Strength.GetValue((AttractorRange - Distance) / AttractorRange, Owner->Component);
+				float AttractorStrength = Strength.GetValue((AttractorRange - Distance) / AttractorRange, Context.GetDistributionData());
 				FVector Direction = LineToPoint^Line;
 				// Adjust the VELOCITY of the particle based on the attractor... 
 				Particle.Velocity += (FVector3f)Direction * AttractorStrength * DeltaTime;
@@ -4084,15 +4143,19 @@ void UParticleModuleAttractorLine::Update(FParticleEmitterInstance* Owner, int32
 	END_UPDATE_LOOP;
 }
 
-void UParticleModuleAttractorLine::Render3DPreview(FParticleEmitterInstance* Owner, const FSceneView* View,FPrimitiveDrawInterface* PDI)
+void UParticleModuleAttractorLine::Render3DPreview(const FPreviewContext& Context)
 {
 #if WITH_EDITOR
+	FPrimitiveDrawInterface* PDI = Context.PDI;
+
 	PDI->DrawLine(EndPoint0, EndPoint1, ModuleEditorColor, SDPG_World);
 
+	FParticleEmitterInstance* Owner = &Context.Owner;
+
 	UParticleLODLevel* LODLevel = Owner->SpriteTemplate->GetLODLevel(0);
 	check(LODLevel);
 	float CurrRatio = Owner->EmitterTime / LODLevel->RequiredModule->EmitterDuration;
-	float LineRange = Range.GetValue(CurrRatio, Owner->Component);
+	float LineRange = Range.GetValue(CurrRatio, Context.GetDistributionData());
 
 	// Determine the position of the range at this time.
 	FVector LinePos = EndPoint0 + CurrRatio * (EndPoint1 - EndPoint0);
@@ -4162,8 +4225,10 @@ void UParticleModuleAttractorParticle::PostEditChangeProperty(FPropertyChangedEv
 }
 #endif // WITH_EDITOR
 
-void UParticleModuleAttractorParticle::Spawn(FParticleEmitterInstance* Owner, int32 Offset, float SpawnTime, FBaseParticle* ParticleBase)
+void UParticleModuleAttractorParticle::Spawn(const FSpawnContext& Context)
 {
+	FParticleEmitterInstance* Owner = &Context.Owner;
+
 	// We need to look up the emitter instance...
 	// This may not need to be done every Spawn, but in the short term, it will to be safe.
 	// (In the editor, the source emitter may be deleted, etc...)
@@ -4197,7 +4262,7 @@ void UParticleModuleAttractorParticle::Spawn(FParticleEmitterInstance* Owner, in
 			switch (SelectionMethod)
 			{
 			case EAPSM_Random:
-				LastSelIndex		= GetRandomStream(Owner).RandHelper(AttractorEmitterInst->ActiveParticles);
+				LastSelIndex		= GetRandomStream(Context).RandHelper(AttractorEmitterInst->ActiveParticles);
 				Data.SourceIndex	= LastSelIndex;
 				break;
 			case EAPSM_Sequential:
@@ -4229,8 +4294,10 @@ void UParticleModuleAttractorParticle::Spawn(FParticleEmitterInstance* Owner, in
 }
 
 
-void UParticleModuleAttractorParticle::Update(FParticleEmitterInstance* Owner, int32 Offset, float DeltaTime)
+void UParticleModuleAttractorParticle::Update(const FUpdateContext& Context)
 {
+	FParticleEmitterInstance* Owner = &Context.Owner;
+
 	// We need to look up the emitter instance...
 	// This may not need to be done every Spawn, but in the short term, it will to be safe.
 	// (In the editor, the source emitter may be deleted, etc...)
@@ -4300,18 +4367,18 @@ void UParticleModuleAttractorParticle::Update(FParticleEmitterInstance* Owner, i
 			CONTINUE_UPDATE_LOOP;
 		}
 
-		float	AttractorRange		= Range.GetValue(Source->RelativeTime, Owner->Component);
+		float	AttractorRange		= Range.GetValue(Source->RelativeTime, Context.GetDistributionData());
 		FVector SrcLocation			= Source->Location;
 		FVector	ParticleLocation	= Particle.Location;
 		if (bUseLocalSpace != bSrcUseLocalSpace)
 		{
 			if (bSrcUseLocalSpace)
 			{
-				SrcLocation = Owner->Component->GetComponentTransform().TransformVector(SrcLocation);
+				SrcLocation = Context.GetTransform().TransformVector(SrcLocation);
 			}
 			if (bUseLocalSpace)
 			{
-				ParticleLocation = Owner->Component->GetComponentTransform().TransformVector(Particle.Location);
+				ParticleLocation = Context.GetTransform().TransformVector(Particle.Location);
 			}
 		}
 
@@ -4330,7 +4397,7 @@ void UParticleModuleAttractorParticle::Update(FParticleEmitterInstance* Owner, i
 			else
 			{
 				// on emitter time
-				AttractorStrength = Strength.GetValue(Source->RelativeTime, Owner->Component);
+				AttractorStrength = Strength.GetValue(Source->RelativeTime, Context.GetDistributionData());
 			}
 
 			// Adjust the VELOCITY of the particle based on the attractor... 
@@ -4410,14 +4477,15 @@ void UParticleModuleAttractorPoint::PostEditChangeProperty(FPropertyChangedEvent
 }
 #endif // WITH_EDITOR
 
-void UParticleModuleAttractorPoint::Update(FParticleEmitterInstance* Owner, int32 Offset, float DeltaTime)
+void UParticleModuleAttractorPoint::Update(const FUpdateContext& Context)
 {
+	FParticleEmitterInstance* Owner = &Context.Owner;
 	check(Owner);
 	UParticleSystemComponent* Component = Owner->Component;
 
 	// Grab the position of the attractor in Emitter time???
-	FVector AttractorPosition = Position.GetValue(Owner->EmitterTime, Component);
-	float AttractorRange = Range.GetValue(Owner->EmitterTime, Component);
+	FVector AttractorPosition = Position.GetValue(Owner->EmitterTime, Context.GetDistributionData());
+	float AttractorRange = Range.GetValue(Owner->EmitterTime, Context.GetDistributionData());
 
 	FVector Scale = FVector(1.0f, 1.0f, 1.0f);
 
@@ -4456,13 +4524,13 @@ void UParticleModuleAttractorPoint::Update(FParticleEmitterInstance* Owner, int3
 				}
 				else
 				{
-					AttractorStrength = Strength.GetValue((AttractorRange - Distance) / AttractorRange, Component);
+					AttractorStrength = Strength.GetValue((AttractorRange - Distance) / AttractorRange, Context.GetDistributionData());
 				}
 			}
 			else
 			{
 				// on emitter time
-				AttractorStrength = Strength.GetValue(Owner->EmitterTime, Component);
+				AttractorStrength = Strength.GetValue(Owner->EmitterTime, Context.GetDistributionData());
 			}
 			if ( (LODLevel->RequiredModule->bUseLocalSpace == false) && ( bUseWorldSpacePosition == false ) )
 			{
@@ -4489,12 +4557,15 @@ void UParticleModuleAttractorPoint::Update(FParticleEmitterInstance* Owner, int3
 	END_UPDATE_LOOP;
 }
 
-void UParticleModuleAttractorPoint::Render3DPreview(FParticleEmitterInstance* Owner, const FSceneView* View,FPrimitiveDrawInterface* PDI)
+void UParticleModuleAttractorPoint::Render3DPreview(const FPreviewContext& Context)
 {
 #if WITH_EDITOR
-	FVector PointPos = Position.GetValue(Owner->EmitterTime, Owner->Component);
-//    float PointStr = Strength.GetValue(Owner->EmitterTime, Owner->Component);
-	float PointRange = Range.GetValue(Owner->EmitterTime, Owner->Component);
+	FParticleEmitterInstance* Owner = &Context.Owner;
+	FVector PointPos = Position.GetValue(Owner->EmitterTime, Context.GetDistributionData());
+//    float PointStr = Strength.GetValue(Owner->EmitterTime, Context.GetDistributionData());
+	float PointRange = Range.GetValue(Owner->EmitterTime, Context.GetDistributionData());
+
+	FPrimitiveDrawInterface* PDI = Context.PDI;
 
 	// Draw a wire star at the position.
 	DrawWireStar(PDI,PointPos, 10.0f, ModuleEditorColor, SDPG_World);
diff --git a/Engine/Source/Runtime/Engine/Private/Particles/ParticleModules_Camera.cpp b/Engine/Source/Runtime/Engine/Private/Particles/ParticleModules_Camera.cpp
--- a/Engine/Source/Runtime/Engine/Private/Particles/ParticleModules_Camera.cpp
+++ b/Engine/Source/Runtime/Engine/Private/Particles/ParticleModules_Camera.cpp
@@ -69,8 +69,10 @@ bool UParticleModuleCameraOffset::CanTickInAnyThread()
 }
 
 
-void UParticleModuleCameraOffset::Spawn(FParticleEmitterInstance* Owner, int32 Offset, float SpawnTime, FBaseParticle* ParticleBase)
+void UParticleModuleCameraOffset::Spawn(const FSpawnContext& Context)
 {
+	FParticleEmitterInstance* Owner = &Context.Owner;
+
 	check(Owner && Owner->Component)
 	float ScaleFactor = 1.0f;
 
@@ -84,9 +86,9 @@ void UParticleModuleCameraOffset::Spawn(FParticleEmitterInstance* Owner, int32 O
 	}
 	SPAWN_INIT;
 	{
-		CurrentOffset = Owner ? ((Owner->CameraPayloadOffset != 0) ? Owner->CameraPayloadOffset : Offset) : Offset;
+		CurrentOffset = Owner ? ((Owner->CameraPayloadOffset != 0) ? Owner->CameraPayloadOffset : Context.Offset) : Context.Offset;
 		PARTICLE_ELEMENT(FCameraOffsetParticlePayload, CameraPayload);
-		float CameraOffsetValue = CameraOffset.GetValue(Particle.RelativeTime, Owner->Component) * ScaleFactor;
+		float CameraOffsetValue = CameraOffset.GetValue(Particle.RelativeTime, Context.GetDistributionData()) * ScaleFactor;
 		if (UpdateMethod == EPCOUM_DirectSet)
 		{
 			CameraPayload.BaseOffset = CameraOffsetValue;
@@ -103,15 +105,16 @@ void UParticleModuleCameraOffset::Spawn(FParticleEmitterInstance* Owner, int32 O
 	}
 }
 
-void UParticleModuleCameraOffset::Update(FParticleEmitterInstance* Owner, int32 Offset, float DeltaTime)
+void UParticleModuleCameraOffset::Update(const FUpdateContext& Context)
 {
 	if (bSpawnTimeOnly == false)
 	{
+		FParticleEmitterInstance* Owner = &Context.Owner;																	\
 		BEGIN_UPDATE_LOOP;
 		{
 			CurrentOffset = Owner ? ((Owner->CameraPayloadOffset != 0) ? Owner->CameraPayloadOffset : Offset) : Offset;
 			PARTICLE_ELEMENT(FCameraOffsetParticlePayload, CameraPayload);
-			float CameraOffsetValue = CameraOffset.GetValue(Particle.RelativeTime, Owner->Component);
+			float CameraOffsetValue = CameraOffset.GetValue(Particle.RelativeTime, Context.GetDistributionData());
 			if (UpdateMethod == EPCOUM_Additive)
 			{
 				CameraPayload.Offset += CameraOffsetValue;
diff --git a/Engine/Source/Runtime/Engine/Private/Particles/ParticleModules_Collision.cpp b/Engine/Source/Runtime/Engine/Private/Particles/ParticleModules_Collision.cpp
--- a/Engine/Source/Runtime/Engine/Private/Particles/ParticleModules_Collision.cpp
+++ b/Engine/Source/Runtime/Engine/Private/Particles/ParticleModules_Collision.cpp
@@ -148,17 +148,20 @@ uint32 UParticleModuleCollision::PrepPerInstanceBlock(FParticleEmitterInstance*
 	return 0;
 }
 
-void UParticleModuleCollision::Spawn(FParticleEmitterInstance* Owner, int32 Offset, float SpawnTime, FBaseParticle* ParticleBase)
+void UParticleModuleCollision::Spawn(const FSpawnContext& Context)
 {
 	SCOPE_CYCLE_COUNTER(STAT_ParticleCollisionTime);
+
+	FParticleEmitterInstance* Owner = &Context.Owner;
+
 	SPAWN_INIT;
 	{
 		PARTICLE_ELEMENT(FParticleCollisionPayload, CollisionPayload);
-		CollisionPayload.UsedDampingFactor = (FVector3f)DampingFactor.GetValue(Owner->EmitterTime, Owner->Component);
-		CollisionPayload.UsedDampingFactorRotation = (FVector3f)DampingFactorRotation.GetValue(Owner->EmitterTime, Owner->Component);
-		CollisionPayload.UsedCollisions = FMath::RoundToInt(MaxCollisions.GetValue(Owner->EmitterTime, Owner->Component));
-		CollisionPayload.Delay = DelayAmount.GetValue(Owner->EmitterTime, Owner->Component);
-		if (CollisionPayload.Delay > SpawnTime)
+		CollisionPayload.UsedDampingFactor = (FVector3f)DampingFactor.GetValue(Owner->EmitterTime, Context.GetDistributionData());
+		CollisionPayload.UsedDampingFactorRotation = (FVector3f)DampingFactorRotation.GetValue(Owner->EmitterTime, Context.GetDistributionData());
+		CollisionPayload.UsedCollisions = FMath::RoundToInt(MaxCollisions.GetValue(Owner->EmitterTime, Context.GetDistributionData()));
+		CollisionPayload.Delay = DelayAmount.GetValue(Owner->EmitterTime, Context.GetDistributionData());
+		if (CollisionPayload.Delay > Context.SpawnTime)
 		{
 			Particle.Flags |= STATE_Particle_DelayCollisions;
 			Particle.Flags &= ~STATE_Particle_CollisionHasOccurred;
@@ -174,9 +177,10 @@ static FAutoConsoleVariableRef ParticleCollisionIgnoreInvisibleTime(
 	ECVF_Default
 	);
 
-void UParticleModuleCollision::Update(FParticleEmitterInstance* Owner, int32 Offset, float DeltaTime)
+void UParticleModuleCollision::Update(const FUpdateContext& Context)
 {
 	SCOPE_CYCLE_COUNTER(STAT_ParticleCollisionTime);
+	FParticleEmitterInstance* Owner = &Context.Owner;
 	check(Owner);
 	check(Owner->Component);
 	UWorld* World = Owner->Component->GetWorld();
@@ -328,7 +332,7 @@ void UParticleModuleCollision::Update(FParticleEmitterInstance* Owner, int32 Off
 		FVector			OldLocation;
 
 		// Location won't be calculated till after tick so we need to calculate an intermediate one here.
-		Location	= Particle.Location + (FVector)Particle.Velocity * DeltaTime;
+		Location	= Particle.Location + (FVector)Particle.Velocity * Context.DeltaTime;
 		if (LODLevel->RequiredModule->bUseLocalSpace)
 		{
 			// Transform the location and old location into world space
@@ -444,7 +448,7 @@ void UParticleModuleCollision::Update(FParticleEmitterInstance* Owner, int32 Off
 							if(PrimitiveComponent && PrimitiveComponent->IsAnySimulatingPhysics())
 							{
 								FVector vImpulse;
-								vImpulse = -(NewVelocity - OldVelocity) * ParticleMass.GetValue(Particle.RelativeTime, Owner->Component);
+								vImpulse = -(NewVelocity - OldVelocity) * ParticleMass.GetValue(Particle.RelativeTime, Context.GetDistributionData());
 								PrimitiveComponent->AddImpulseAtLocation(vImpulse, Hit.Location, Hit.BoneName);
 							}
 						}
@@ -474,7 +478,7 @@ void UParticleModuleCollision::Update(FParticleEmitterInstance* Owner, int32 Off
 							if(PrimitiveComponent && PrimitiveComponent->IsAnySimulatingPhysics())
 							{
 								FVector vImpulse;
-								vImpulse = -(vNewVelocity - vOldVelocity) * ParticleMass.GetValue(Particle.RelativeTime, Owner->Component);
+								vImpulse = -(vNewVelocity - vOldVelocity) * ParticleMass.GetValue(Particle.RelativeTime, Context.GetDistributionData());
 								PrimitiveComponent->AddImpulseAtLocation(vImpulse, Hit.Location, Hit.BoneName);
 							}
 						}
diff --git a/Engine/Source/Runtime/Engine/Private/Particles/ParticleModules_Color.cpp b/Engine/Source/Runtime/Engine/Private/Particles/ParticleModules_Color.cpp
--- a/Engine/Source/Runtime/Engine/Private/Particles/ParticleModules_Color.cpp
+++ b/Engine/Source/Runtime/Engine/Private/Particles/ParticleModules_Color.cpp
@@ -152,18 +152,20 @@ bool UParticleModuleColor::AddModuleCurvesToEditor(UInterpCurveEdSetup* EdSetup,
 	return bNewCurve;
 }
 
-void UParticleModuleColor::Spawn(FParticleEmitterInstance* Owner, int32 Offset, float SpawnTime, FBaseParticle* ParticleBase)
+void UParticleModuleColor::Spawn(const FSpawnContext& Context)
 {
-	SpawnEx(Owner, Offset, SpawnTime, &GetRandomStream(Owner), ParticleBase);
+	SpawnEx(Context, &GetRandomStream(Context));
 }
 
 
-void UParticleModuleColor::SpawnEx(FParticleEmitterInstance* Owner, int32 Offset, float SpawnTime, struct FRandomStream* InRandomStream, FBaseParticle* ParticleBase)
+void UParticleModuleColor::SpawnEx(const FSpawnContext& Context, struct FRandomStream* InRandomStream)
 {
+	FParticleEmitterInstance* Owner = &Context.Owner;
+
 	SPAWN_INIT;
 	{
-		FVector3f ColorVec	= (FVector3f)StartColor.GetValue(Owner->EmitterTime, Owner->Component, 0, InRandomStream);
-		float	Alpha		= StartAlpha.GetValue(Owner->EmitterTime, Owner->Component, InRandomStream);
+		FVector3f ColorVec	= (FVector3f)StartColor.GetValue(Owner->EmitterTime, Context.GetDistributionData(), 0, InRandomStream);
+		float	Alpha		= StartAlpha.GetValue(Owner->EmitterTime, Context.GetDistributionData(), InRandomStream);
 		Particle_SetColorFromVector(ColorVec, Alpha, Particle.Color);
 		Particle.BaseColor	= Particle.Color;
 	}
@@ -348,11 +350,13 @@ bool UParticleModuleColorOverLife::AddModuleCurvesToEditor(UInterpCurveEdSetup*
 	return bNewCurve;
 }
 
-void UParticleModuleColorOverLife::Spawn(FParticleEmitterInstance* Owner, int32 Offset, float SpawnTime, FBaseParticle* ParticleBase)
+void UParticleModuleColorOverLife::Spawn(const FSpawnContext& Context)
 {
+	FParticleEmitterInstance* Owner = &Context.Owner;
+
 	SPAWN_INIT;
-	FVector ColorVec	= ColorOverLife.GetValue(Particle.RelativeTime, Owner->Component);
-	float	fAlpha		= AlphaOverLife.GetValue(Particle.RelativeTime, Owner->Component);
+	FVector ColorVec	= ColorOverLife.GetValue(Particle.RelativeTime, Context.GetDistributionData());
+	float	fAlpha		= AlphaOverLife.GetValue(Particle.RelativeTime, Context.GetDistributionData());
 	Particle.Color.R	 = UE_REAL_TO_FLOAT(ColorVec.X);
 	Particle.BaseColor.R = UE_REAL_TO_FLOAT(ColorVec.X);
 	Particle.Color.G	 = UE_REAL_TO_FLOAT(ColorVec.Y);
@@ -363,8 +367,9 @@ void UParticleModuleColorOverLife::Spawn(FParticleEmitterInstance* Owner, int32
 	Particle.BaseColor.A = fAlpha;
 }
 
-void UParticleModuleColorOverLife::Update(FParticleEmitterInstance* Owner, int32 Offset, float DeltaTime)
+void UParticleModuleColorOverLife::Update(const FUpdateContext& Context)
 {
+	FParticleEmitterInstance* Owner = &Context.Owner;
 	if ((Owner == NULL) || (Owner->ActiveParticles <= 0) || 
 		(Owner->ParticleData == NULL) || (Owner->ParticleIndices == NULL))
 	{
@@ -392,8 +397,8 @@ void UParticleModuleColorOverLife::Update(FParticleEmitterInstance* Owner, int32
 		float	fAlpha;
 		BEGIN_UPDATE_LOOP;
 		{
-			ColorVec = ColorOverLife.GetValue(Particle.RelativeTime, Owner->Component);
-			fAlpha = AlphaOverLife.GetValue(Particle.RelativeTime, Owner->Component);
+			ColorVec = ColorOverLife.GetValue(Particle.RelativeTime, Context.GetDistributionData());
+			fAlpha = AlphaOverLife.GetValue(Particle.RelativeTime, Context.GetDistributionData());
 			FPlatformMisc::Prefetch(ParticleData, (ParticleIndices[i+1] * ParticleStride));
 			FPlatformMisc::Prefetch(ParticleData, (ParticleIndices[i+1] * ParticleStride) + PLATFORM_CACHE_LINE_SIZE);
 			Particle.Color.R = UE_REAL_TO_FLOAT(ColorVec.X);
@@ -533,21 +538,23 @@ void UParticleModuleColorScaleOverLife::PostEditChangeProperty(FPropertyChangedE
 }
 #endif // WITH_EDITOR
 
-void UParticleModuleColorScaleOverLife::Spawn(FParticleEmitterInstance* Owner, int32 Offset, float SpawnTime, FBaseParticle* ParticleBase)
+void UParticleModuleColorScaleOverLife::Spawn(const FSpawnContext& Context)
 {
+	FParticleEmitterInstance* Owner = &Context.Owner;
+
 	SPAWN_INIT;
 	FVector ColorVec;
 	float	fAlpha;
 
 	if (bEmitterTime)
 	{
-		ColorVec	= ColorScaleOverLife.GetValue(Owner->EmitterTime, Owner->Component);
-		fAlpha		= AlphaScaleOverLife.GetValue(Owner->EmitterTime, Owner->Component);
+		ColorVec	= ColorScaleOverLife.GetValue(Owner->EmitterTime, Context.GetDistributionData());
+		fAlpha		= AlphaScaleOverLife.GetValue(Owner->EmitterTime, Context.GetDistributionData());
 	}
 	else
 	{
-		ColorVec	= ColorScaleOverLife.GetValue(Particle.RelativeTime, Owner->Component);
-		fAlpha		= AlphaScaleOverLife.GetValue(Particle.RelativeTime, Owner->Component);
+		ColorVec	= ColorScaleOverLife.GetValue(Particle.RelativeTime, Context.GetDistributionData());
+		fAlpha		= AlphaScaleOverLife.GetValue(Particle.RelativeTime, Context.GetDistributionData());
 	}
 
 	Particle.Color.R *= UE_REAL_TO_FLOAT(ColorVec.X);
@@ -556,8 +563,9 @@ void UParticleModuleColorScaleOverLife::Spawn(FParticleEmitterInstance* Owner, i
 	Particle.Color.A *= fAlpha;
 }
 
-void UParticleModuleColorScaleOverLife::Update(FParticleEmitterInstance* Owner, int32 Offset, float DeltaTime)
+void UParticleModuleColorScaleOverLife::Update(const FUpdateContext& Context)
 { 
+	FParticleEmitterInstance* Owner = &Context.Owner;
 	const FRawDistribution* FastColorScaleOverLife = ColorScaleOverLife.GetFastRawDistribution();
 	const FRawDistribution* FastAlphaScaleOverLife = AlphaScaleOverLife.GetFastRawDistribution();
 	FVector3f ColorVec;	// LWC_TODO: WRONG? Force ColorVec to floats for serialization. Particle.Color is FVector. Update to FVector3f also?
@@ -597,8 +605,8 @@ void UParticleModuleColorScaleOverLife::Update(FParticleEmitterInstance* Owner,
 		if (bEmitterTime)
 		{
 			BEGIN_UPDATE_LOOP;
-				ColorVec = FVector3f(ColorScaleOverLife.GetValue(Owner->EmitterTime, Owner->Component));
-				fAlpha = AlphaScaleOverLife.GetValue(Owner->EmitterTime, Owner->Component);
+				ColorVec = FVector3f(ColorScaleOverLife.GetValue(Owner->EmitterTime, Context.GetDistributionData()));
+				fAlpha = AlphaScaleOverLife.GetValue(Owner->EmitterTime, Context.GetDistributionData());
 				Particle.Color.R *= ColorVec.X;
 				Particle.Color.G *= ColorVec.Y;
 				Particle.Color.B *= ColorVec.Z;
@@ -608,8 +616,8 @@ void UParticleModuleColorScaleOverLife::Update(FParticleEmitterInstance* Owner,
 		else
 		{
 			BEGIN_UPDATE_LOOP;
-				ColorVec = FVector3f(ColorScaleOverLife.GetValue(Particle.RelativeTime, Owner->Component));
-				fAlpha = AlphaScaleOverLife.GetValue(Particle.RelativeTime, Owner->Component);
+				ColorVec = FVector3f(ColorScaleOverLife.GetValue(Particle.RelativeTime, Context.GetDistributionData()));
+				fAlpha = AlphaScaleOverLife.GetValue(Particle.RelativeTime, Context.GetDistributionData());
 				Particle.Color.R *= ColorVec.X;
 				Particle.Color.G *= ColorVec.Y;
 				Particle.Color.B *= ColorVec.Z;
diff --git a/Engine/Source/Runtime/Engine/Private/Particles/ParticleModules_Event.cpp b/Engine/Source/Runtime/Engine/Private/Particles/ParticleModules_Event.cpp
--- a/Engine/Source/Runtime/Engine/Private/Particles/ParticleModules_Event.cpp
+++ b/Engine/Source/Runtime/Engine/Private/Particles/ParticleModules_Event.cpp
@@ -50,11 +50,11 @@ UParticleModuleEventGenerator::UParticleModuleEventGenerator(const FObjectInitia
 	bUpdateModule = true;
 }
 
-void UParticleModuleEventGenerator::Spawn(FParticleEmitterInstance* Owner, int32 Offset, float SpawnTime, FBaseParticle* ParticleBase)
+void UParticleModuleEventGenerator::Spawn(const FSpawnContext& Context)
 {
 }
 
-void UParticleModuleEventGenerator::Update(FParticleEmitterInstance* Owner, int32 Offset, float DeltaTime)
+void UParticleModuleEventGenerator::Update(const FUpdateContext& Context)
 {
 }
 
diff --git a/Engine/Source/Runtime/Engine/Private/Particles/ParticleModules_Location.cpp b/Engine/Source/Runtime/Engine/Private/Particles/ParticleModules_Location.cpp
--- a/Engine/Source/Runtime/Engine/Private/Particles/ParticleModules_Location.cpp
+++ b/Engine/Source/Runtime/Engine/Private/Particles/ParticleModules_Location.cpp
@@ -91,14 +91,15 @@ void UParticleModuleLocation::PostEditChangeProperty(FPropertyChangedEvent& Prop
 }
 #endif // WITH_EDITOR
 
-void UParticleModuleLocation::Spawn(FParticleEmitterInstance* Owner, int32 Offset, float SpawnTime, FBaseParticle* ParticleBase)
+void UParticleModuleLocation::Spawn(const FSpawnContext& Context)
 {
-	SpawnEx(Owner, Offset, SpawnTime, &GetRandomStream(Owner), ParticleBase);
+	SpawnEx(Context, &GetRandomStream(Context));
 }
 
-void UParticleModuleLocation::SpawnEx(FParticleEmitterInstance* Owner, int32 Offset, float SpawnTime, struct FRandomStream* InRandomStream, FBaseParticle* ParticleBase)
+void UParticleModuleLocation::SpawnEx(const FSpawnContext& Context, struct FRandomStream* InRandomStream)
 {
 	SPAWN_INIT;
+	FParticleEmitterInstance* Owner = &Context.Owner;
 	UParticleLODLevel* LODLevel	= Owner->SpriteTemplate->GetCurrentLODLevel(Owner);
 	check(LODLevel);
 	FVector LocationOffset;
@@ -110,7 +111,7 @@ void UParticleModuleLocation::SpawnEx(FParticleEmitterInstance* Owner, int32 Off
 
 		if(RandomNum > DistributeThreshold)
 		{
-			LocationOffset = StartLocation.GetValue(Owner->EmitterTime, Owner->Component, 0, InRandomStream);
+			LocationOffset = StartLocation.GetValue(Owner->EmitterTime, Context.GetDistributionData(), 0, InRandomStream);
 		}
 		else
 		{
@@ -122,7 +123,7 @@ void UParticleModuleLocation::SpawnEx(FParticleEmitterInstance* Owner, int32 Off
 	}
 	else
 	{
-		LocationOffset = StartLocation.GetValue(Owner->EmitterTime, Owner->Component, 0, InRandomStream);
+		LocationOffset = StartLocation.GetValue(Owner->EmitterTime, Context.GetDistributionData(), 0, InRandomStream);
 	}
 
 	LocationOffset = Owner->EmitterToSimulation.TransformVector(LocationOffset);
@@ -130,11 +131,14 @@ void UParticleModuleLocation::SpawnEx(FParticleEmitterInstance* Owner, int32 Off
 	ensureMsgf(!Particle.Location.ContainsNaN(), TEXT("NaN in Particle Location. Template: %s, Component: %s"), Owner->Component ? *GetNameSafe(Owner->Component->Template) : TEXT("UNKNOWN"), *GetPathNameSafe(Owner->Component));
 }
 
-void UParticleModuleLocation::Render3DPreview(FParticleEmitterInstance* Owner, const FSceneView* View,FPrimitiveDrawInterface* PDI)
+void UParticleModuleLocation::Render3DPreview(const FPreviewContext& Context)
 {
 #if WITH_EDITOR
+	FParticleEmitterInstance* Owner = &Context.Owner;
 	check(Owner);
 
+	FPrimitiveDrawInterface* PDI = Context.PDI;
+
 	// Draw the location as a wire star
 	FVector Position(0.0f);
 
@@ -166,11 +170,11 @@ void UParticleModuleLocation::Render3DPreview(FParticleEmitterInstance* Owner, c
 			UDistributionVectorConstantCurve* Curve = CastChecked<UDistributionVectorConstantCurve>(StartLocation.Distribution);
 
 			//Curve->
-			Position = StartLocation.GetValue(0.0f, Owner->Component);
+			Position = StartLocation.GetValue(0.0f, Context.GetDistributionData());
 		}
 		else if (StartLocation.Distribution->IsA(UDistributionVectorConstant::StaticClass()))
 		{
-			Position = StartLocation.GetValue(0.0f, Owner->Component);
+			Position = StartLocation.GetValue(0.0f, Context.GetDistributionData());
 		}
 	}
 
@@ -207,21 +211,22 @@ UParticleModuleLocationWorldOffset::UParticleModuleLocationWorldOffset(const FOb
 {
 }
 
-void UParticleModuleLocationWorldOffset::SpawnEx(FParticleEmitterInstance* Owner, int32 Offset, float SpawnTime, struct FRandomStream* InRandomStream, FBaseParticle* ParticleBase)
+void UParticleModuleLocationWorldOffset::SpawnEx(const FSpawnContext& Context, struct FRandomStream* InRandomStream)
 {
 	SPAWN_INIT;
+	FParticleEmitterInstance* Owner = &Context.Owner;
 	UParticleLODLevel* LODLevel = Owner->SpriteTemplate->GetCurrentLODLevel(Owner);
 	check(LODLevel);
 	if (LODLevel->RequiredModule->bUseLocalSpace == false)
 	{
 		// Nothing to do here... the distribution value is already being in world space
-		Particle.Location += StartLocation.GetValue(Owner->EmitterTime, Owner->Component, 0, InRandomStream);
+		Particle.Location += StartLocation.GetValue(Owner->EmitterTime, Context.GetDistributionData(), 0, InRandomStream);
 	}
 	else
 	{
 		// We need to inverse transform the location so that the bUseLocalSpace transform uses the proper value
-		FMatrix InvMat = Owner->Component->GetComponentTransform().ToMatrixWithScale().InverseFast();
-		FVector StartLoc = StartLocation.GetValue(Owner->EmitterTime, Owner->Component, 0, InRandomStream);
+		FMatrix InvMat = Context.GetTransform().ToMatrixWithScale().InverseFast();
+		FVector StartLoc = StartLocation.GetValue(Owner->EmitterTime, Context.GetDistributionData(), 0, InRandomStream);
 		Particle.Location += InvMat.TransformVector(StartLoc);
 	}
 	ensureMsgf(!Particle.Location.ContainsNaN(), TEXT("NaN in Particle Location. Template: %s, Component: %s"), Owner->Component ? *GetNameSafe(Owner->Component->Template) : TEXT("UNKNOWN"), *GetPathNameSafe(Owner->Component));
@@ -298,31 +303,32 @@ void UParticleModuleLocationDirect::PostEditChangeProperty(FPropertyChangedEvent
 }
 #endif // WITH_EDITOR
 
-void UParticleModuleLocationDirect::Spawn(FParticleEmitterInstance* Owner, int32 Offset, float SpawnTime, FBaseParticle* ParticleBase)
+void UParticleModuleLocationDirect::Spawn(const FSpawnContext& Context)
 {
 	SPAWN_INIT;
-
+	FParticleEmitterInstance* Owner = &Context.Owner;
 	UParticleLODLevel* LODLevel	= Owner->SpriteTemplate->GetCurrentLODLevel(Owner);
 	check(LODLevel);
 	if (LODLevel->RequiredModule->bUseLocalSpace)
 	{
-		Particle.Location = Location.GetValue(Particle.RelativeTime, Owner->Component);
+		Particle.Location = Location.GetValue(Particle.RelativeTime, Context.GetDistributionData());
 	}
 	else
 	{
-		FVector StartLoc	= Location.GetValue(Particle.RelativeTime, Owner->Component);
-		StartLoc = Owner->Component->GetComponentTransform().TransformPosition(StartLoc);
+		FVector StartLoc	= Location.GetValue(Particle.RelativeTime, Context.GetDistributionData());
+		StartLoc = Context.GetTransform().TransformPosition(StartLoc);
 		Particle.Location	= StartLoc;
 	}
 
 	PARTICLE_ELEMENT(FVector, LocOffset);
-	LocOffset	= LocationOffset.GetValue(Owner->EmitterTime, Owner->Component);
+	LocOffset	= LocationOffset.GetValue(Owner->EmitterTime, Context.GetDistributionData());
 	Particle.Location += LocOffset;
 	ensureMsgf(!Particle.Location.ContainsNaN(), TEXT("NaN in Particle Location. Template: %s, Component: %s"), Owner->Component ? *GetNameSafe(Owner->Component->Template) : TEXT("UNKNOWN"), *GetPathNameSafe(Owner->Component));
 }
 
-void UParticleModuleLocationDirect::Update(FParticleEmitterInstance* Owner, int32 Offset, float DeltaTime)
+void UParticleModuleLocationDirect::Update(const FUpdateContext& Context)
 {
+	FParticleEmitterInstance* Owner = &Context.Owner;
 	const FTransform& OwnerTM = Owner->Component->GetAsyncComponentToWorld();
 	BEGIN_UPDATE_LOOP;
 	{
@@ -331,16 +337,16 @@ void UParticleModuleLocationDirect::Update(FParticleEmitterInstance* Owner, int3
 		check(LODLevel);
 		if (LODLevel->RequiredModule->bUseLocalSpace)
 		{
-			NewLoc = Location.GetValue(Particle.RelativeTime, Owner->Component);
+			NewLoc = Location.GetValue(Particle.RelativeTime, Context.GetDistributionData());
 		}
 		else
 		{
-			FVector Loc			= Location.GetValue(Particle.RelativeTime, Owner->Component);
+			FVector Loc			= Location.GetValue(Particle.RelativeTime, Context.GetDistributionData());
 			Loc = OwnerTM.TransformPosition(Loc);
 			NewLoc	= Loc;
 		}
 
-		FVector	Scale	= ScaleFactor.GetValue(Particle.RelativeTime, Owner->Component);
+		FVector	Scale	= ScaleFactor.GetValue(Particle.RelativeTime, Context.GetDistributionData());
 
 		PARTICLE_ELEMENT(FVector, LocOffset);
 		NewLoc += LocOffset;
@@ -387,8 +393,10 @@ UParticleModuleLocationEmitter::UParticleModuleLocationEmitter(const FObjectInit
 	InheritSourceRotationScale = 1.0f;
 }
 
-void UParticleModuleLocationEmitter::Spawn(FParticleEmitterInstance* Owner, int32 Offset, float SpawnTime, FBaseParticle* ParticleBase)
+void UParticleModuleLocationEmitter::Spawn(const FSpawnContext& Context)
 {
+	FParticleEmitterInstance* Owner = &Context.Owner;
+
 	// We need to look up the emitter instance...
 	// This may not need to be done every Spawn, but in the short term, it will to be safe.
 	// (In the editor, the source emitter may be deleted, etc...)
@@ -419,7 +427,7 @@ void UParticleModuleLocationEmitter::Spawn(FParticleEmitterInstance* Owner, int3
 	bool bSourceIsInLocalSpace = LocationEmitterInst->CurrentLODLevel->RequiredModule->bUseLocalSpace;
 	bool bInLocalSpace = Owner->CurrentLODLevel->RequiredModule->bUseLocalSpace;
 
-	FRandomStream& RandomStream = GetRandomStream(Owner);
+	FRandomStream& RandomStream = GetRandomStream(Context);
 
 	SPAWN_INIT;
 		{
@@ -549,8 +557,10 @@ UParticleModuleLocationEmitterDirect::UParticleModuleLocationEmitterDirect(const
 	EmitterName = ConstructorStatics.NAME_None;
 }
 
-void UParticleModuleLocationEmitterDirect::Spawn(FParticleEmitterInstance* Owner, int32 Offset, float SpawnTime, FBaseParticle* ParticleBase)
+void UParticleModuleLocationEmitterDirect::Spawn(const FSpawnContext& Context)
 {
+	FParticleEmitterInstance* Owner = &Context.Owner;
+
 	// We need to look up the emitter instance...
 	// This may not need to be done every Spawn, but in the short term, it will to be safe.
 	// (In the editor, the source emitter may be deleted, etc...)
@@ -590,8 +600,10 @@ void UParticleModuleLocationEmitterDirect::Spawn(FParticleEmitterInstance* Owner
 		}
 } 
 
-void UParticleModuleLocationEmitterDirect::Update(FParticleEmitterInstance* Owner, int32 Offset, float DeltaTime)
+void UParticleModuleLocationEmitterDirect::Update(const FUpdateContext& Context)
 {
+	FParticleEmitterInstance* Owner = &Context.Owner;
+
 	// We need to look up the emitter instance...
 	// This may not need to be done every Spawn, but in the short term, it will to be safe.
 	// (In the editor, the source emitter may be deleted, etc...)
@@ -804,21 +816,22 @@ void UParticleModuleLocationPrimitiveTriangle::PostEditChangeProperty(FPropertyC
 }
 #endif // WITH_EDITOR
 
-void UParticleModuleLocationPrimitiveTriangle::Spawn(FParticleEmitterInstance* Owner, int32 Offset, float SpawnTime, FBaseParticle* ParticleBase)
+void UParticleModuleLocationPrimitiveTriangle::Spawn(const FSpawnContext& Context)
 {
-	SpawnEx(Owner, Offset, SpawnTime, &GetRandomStream(Owner), ParticleBase);
+	SpawnEx(Context, &GetRandomStream(Context));
 }
 
-void UParticleModuleLocationPrimitiveTriangle::SpawnEx(FParticleEmitterInstance* Owner, int32 Offset, float SpawnTime, struct FRandomStream* InRandomStream, FBaseParticle* ParticleBase)
+void UParticleModuleLocationPrimitiveTriangle::SpawnEx(const FSpawnContext& Context, struct FRandomStream* InRandomStream)
 {
 	SPAWN_INIT;
+	FParticleEmitterInstance* Owner = &Context.Owner;
 	UParticleLODLevel* LODLevel	= Owner->SpriteTemplate->GetCurrentLODLevel(Owner);
 	check(LODLevel);
 
-	FVector TriOffset = StartOffset.GetValue(Owner->EmitterTime, Owner->Component, 0, InRandomStream);
-	float TriHeight = Height.GetValue(Owner->EmitterTime, Owner->Component, InRandomStream);
-	float TriAngle = Angle.GetValue(Owner->EmitterTime, Owner->Component, InRandomStream);
-	float TriThickness = Thickness.GetValue(Owner->EmitterTime, Owner->Component, InRandomStream);
+	FVector TriOffset = StartOffset.GetValue(Owner->EmitterTime, Context.GetDistributionData(), 0, InRandomStream);
+	float TriHeight = Height.GetValue(Owner->EmitterTime, Context.GetDistributionData(), InRandomStream);
+	float TriAngle = Angle.GetValue(Owner->EmitterTime, Context.GetDistributionData(), InRandomStream);
+	float TriThickness = Thickness.GetValue(Owner->EmitterTime, Context.GetDistributionData(), InRandomStream);
 	float BaseLength = TriHeight * FMath::Tan(0.5f * TriAngle * UE_PI / 180.0f);
 
 	FVector Corners[3];
@@ -847,10 +860,11 @@ void UParticleModuleLocationPrimitiveTriangle::SpawnEx(FParticleEmitterInstance*
 	ensureMsgf(!Particle.Location.ContainsNaN(), TEXT("NaN in Particle Location. Template: %s, Component: %s"), Owner->Component ? *GetNameSafe(Owner->Component->Template) : TEXT("UNKNOWN"), *GetPathNameSafe(Owner->Component));
 }
 
-void UParticleModuleLocationPrimitiveTriangle::Render3DPreview(FParticleEmitterInstance* Owner, const FSceneView* View,FPrimitiveDrawInterface* PDI)
+void UParticleModuleLocationPrimitiveTriangle::Render3DPreview(const FPreviewContext& Context)
 {
 #if WITH_EDITOR
 	FMatrix LocalToWorld = FMatrix::Identity;
+	FParticleEmitterInstance* Owner = &Context.Owner;
 	if (Owner != NULL)
 	{
 		LocalToWorld = Owner->EmitterToSimulation * Owner->SimulationToWorld;
@@ -875,7 +889,9 @@ void UParticleModuleLocationPrimitiveTriangle::Render3DPreview(FParticleEmitterI
 		}
 		FVector ThicknessDir(0.0f, 0.0f, 0.5f * TriThickness);
 		ThicknessDir = LocalToWorld.TransformVector(ThicknessDir);
-		
+
+		FPrimitiveDrawInterface* PDI = Context.PDI;
+
 		FVector CenterPos = Corners[0] / 3.0f + Corners[1] / 3.0f + Corners[2] / 3.0f;
 		DrawWireStar(PDI, CenterPos, 10.0f, ModuleEditorColor, SDPG_World);	
 
@@ -949,12 +965,12 @@ void UParticleModuleLocationPrimitiveCylinder::PostEditChangeProperty(FPropertyC
 }
 #endif // WITH_EDITOR
 
-void UParticleModuleLocationPrimitiveCylinder::Spawn(FParticleEmitterInstance* Owner, int32 Offset, float SpawnTime, FBaseParticle* ParticleBase)
+void UParticleModuleLocationPrimitiveCylinder::Spawn(const FSpawnContext& Context)
 {
-	SpawnEx(Owner, Offset, SpawnTime, &GetRandomStream(Owner), ParticleBase);
+	SpawnEx(Context, &GetRandomStream(Context));
 }
 
-void UParticleModuleLocationPrimitiveCylinder::SpawnEx(FParticleEmitterInstance* Owner, int32 Offset, float SpawnTime, struct FRandomStream* InRandomStream, FBaseParticle* ParticleBase)
+void UParticleModuleLocationPrimitiveCylinder::SpawnEx(const FSpawnContext& Context, struct FRandomStream* InRandomStream)
 {
 	SPAWN_INIT;
 
@@ -978,12 +994,14 @@ void UParticleModuleLocationPrimitiveCylinder::SpawnEx(FParticleEmitterInstance*
 		break;
 	}
 
+	FParticleEmitterInstance* Owner = &Context.Owner;
+
 	// Determine the start location for the sphere
-	FVector vStartLoc = StartLocation.GetValue(Owner->EmitterTime, Owner->Component, 0, InRandomStream);
+	FVector vStartLoc = StartLocation.GetValue(Owner->EmitterTime, Context.GetDistributionData(), 0, InRandomStream);
 
 	FVector vOffset(0.0f);
-	float	fStartRadius	= StartRadius.GetValue(Owner->EmitterTime, Owner->Component, InRandomStream);
-	float	fStartHeight	= StartHeight.GetValue(Owner->EmitterTime, Owner->Component, InRandomStream) / 2.0f;
+	float	fStartRadius	= StartRadius.GetValue(Owner->EmitterTime, Context.GetDistributionData(), InRandomStream);
+	float	fStartHeight	= StartHeight.GetValue(Owner->EmitterTime, Context.GetDistributionData(), InRandomStream) / 2.0f;
 
 
 	// Determine the unit direction
@@ -1078,7 +1096,7 @@ void UParticleModuleLocationPrimitiveCylinder::SpawnEx(FParticleEmitterInstance*
 		{
 			vVelocity[HeightIndex]	= 0.0f;
 		}
-		vVelocity	*= VelocityScale.GetValue(Owner->EmitterTime, Owner->Component, InRandomStream);
+		vVelocity	*= VelocityScale.GetValue(Owner->EmitterTime, Context.GetDistributionData(), InRandomStream);
 		vVelocity = Owner->EmitterToSimulation.TransformVector(vVelocity);
 
 		Particle.Velocity		+= (FVector3f)vVelocity;
@@ -1088,19 +1106,22 @@ void UParticleModuleLocationPrimitiveCylinder::SpawnEx(FParticleEmitterInstance*
 	ensureMsgf(!Particle.Velocity.ContainsNaN(), TEXT("NaN in Particle Velocity. Template: %s, Component: %s"), Owner->Component ? *GetNameSafe(Owner->Component->Template) : TEXT("UNKNOWN"), *GetPathNameSafe(Owner->Component));
 }
 
-void UParticleModuleLocationPrimitiveCylinder::Render3DPreview(FParticleEmitterInstance* Owner, const FSceneView* View,FPrimitiveDrawInterface* PDI)
+void UParticleModuleLocationPrimitiveCylinder::Render3DPreview(const FPreviewContext& Context)
 {
 #if WITH_EDITOR
 	// Draw the location as a wire star
 	FVector Position = FVector::ZeroVector;
 	FVector OwnerScale = FVector(1.0f);
 	FMatrix LocalToWorld = FMatrix::Identity;
+	FParticleEmitterInstance* Owner = &Context.Owner;
 	if (Owner != NULL)
 	{
 		LocalToWorld = Owner->EmitterToSimulation * Owner->SimulationToWorld;
 		OwnerScale = LocalToWorld.GetScaleVector();
 	}
 
+	FPrimitiveDrawInterface* PDI = Context.PDI;
+
 	Position = LocalToWorld.TransformPosition(Position);
 	DrawWireStar(PDI, Position, 10.0f, ModuleEditorColor, SDPG_World);
 
@@ -1127,7 +1148,7 @@ void UParticleModuleLocationPrimitiveCylinder::Render3DPreview(FParticleEmitterI
 			if (Owner != NULL)
 			{
 				//pkCurve->
-				Position = StartLocation.GetValue(0.0f, Owner->Component);
+				Position = StartLocation.GetValue(0.0f, Context.GetDistributionData());
 			}
 		}
 	}
@@ -1140,8 +1161,8 @@ void UParticleModuleLocationPrimitiveCylinder::Render3DPreview(FParticleEmitterI
 	float fStartHeight = 1.0f;
 	if(Owner && Owner->Component)
 	{
-		fStartRadius = StartRadius.GetValue(Owner->EmitterTime, Owner->Component);
-		fStartHeight = StartHeight.GetValue(Owner->EmitterTime, Owner->Component) / 2.0f;
+		fStartRadius = StartRadius.GetValue(Owner->EmitterTime, Context.GetDistributionData());
+		fStartHeight = StartHeight.GetValue(Owner->EmitterTime, Context.GetDistributionData()) / 2.0f;
 	}
 
 	FVector	TransformedAxis[3];
@@ -1238,17 +1259,18 @@ void UParticleModuleLocationPrimitiveSphere::PostEditChangeProperty(FPropertyCha
 }
 #endif // WITH_EDITOR
 
-void UParticleModuleLocationPrimitiveSphere::Spawn(FParticleEmitterInstance* Owner, int32 Offset, float SpawnTime, FBaseParticle* ParticleBase)
+void UParticleModuleLocationPrimitiveSphere::Spawn(const FSpawnContext& Context)
 {
-	SpawnEx(Owner, Offset, SpawnTime, &GetRandomStream(Owner), ParticleBase);
+	SpawnEx(Context, &GetRandomStream(Context));
 }
 
-void UParticleModuleLocationPrimitiveSphere::SpawnEx(FParticleEmitterInstance* Owner, int32 Offset, float SpawnTime, struct FRandomStream* InRandomStream, FBaseParticle* ParticleBase)
+void UParticleModuleLocationPrimitiveSphere::SpawnEx(const FSpawnContext& Context, struct FRandomStream* InRandomStream)
 {
 	SPAWN_INIT;
+	FParticleEmitterInstance* Owner = &Context.Owner;
 
 	// Determine the start location for the sphere
-	FVector vStartLoc = StartLocation.GetValue(Owner->EmitterTime, Owner->Component, 0, InRandomStream);
+	FVector vStartLoc = StartLocation.GetValue(Owner->EmitterTime, Context.GetDistributionData(), 0, InRandomStream);
 
 	// Determine the unit direction
 	FVector vUnitDir;
@@ -1264,7 +1286,7 @@ void UParticleModuleLocationPrimitiveSphere::SpawnEx(FParticleEmitterInstance* O
 	}
 
 	// Determine the position
-	float	fStartRadius	= StartRadius.GetValue(Owner->EmitterTime, Owner->Component, InRandomStream);
+	float	fStartRadius	= StartRadius.GetValue(Owner->EmitterTime, Context.GetDistributionData(), InRandomStream);
 	FVector vStartRadius	= FVector(fStartRadius);
 	FVector vOffset			= vUnitDir * vStartRadius;
 
@@ -1305,7 +1327,7 @@ void UParticleModuleLocationPrimitiveSphere::SpawnEx(FParticleEmitterInstance* O
 
 	if (Velocity)
 	{
-		FVector vVelocity		 = (vOffset - vStartLoc) * VelocityScale.GetValue(Owner->EmitterTime, Owner->Component, InRandomStream);
+		FVector vVelocity		 = (vOffset - vStartLoc) * VelocityScale.GetValue(Owner->EmitterTime, Context.GetDistributionData(), InRandomStream);
 		vVelocity = Owner->EmitterToSimulation.TransformVector(vVelocity);
 		Particle.Velocity		+= (FVector3f)vVelocity;
 		Particle.BaseVelocity	+= (FVector3f)vVelocity;
@@ -1314,16 +1336,20 @@ void UParticleModuleLocationPrimitiveSphere::SpawnEx(FParticleEmitterInstance* O
 	ensureMsgf(!Particle.Velocity.ContainsNaN(), TEXT("NaN in Particle Velocity. Template: %s, Component: %s"), Owner->Component ? *GetNameSafe(Owner->Component->Template) : TEXT("UNKNOWN"), *GetPathNameSafe(Owner->Component));
 }
 
-void UParticleModuleLocationPrimitiveSphere::Render3DPreview(FParticleEmitterInstance* Owner, const FSceneView* View,FPrimitiveDrawInterface* PDI)
+void UParticleModuleLocationPrimitiveSphere::Render3DPreview(const FPreviewContext& Context)
 {
 #if WITH_EDITOR
 	FVector Position(0.0f);
 
 	// Draw the location as a wire star
+	FParticleEmitterInstance* Owner = &Context.Owner;
 	if ((Owner != NULL) && (Owner->Component != NULL))
 	{
 		Position = Owner->SimulationToWorld.TransformPosition(Owner->EmitterToSimulation.GetOrigin());
 	}
+
+	FPrimitiveDrawInterface* PDI = Context.PDI;
+
 	DrawWireStar(PDI, Position, 10.0f, ModuleEditorColor, SDPG_World);
 
 	if (StartLocation.Distribution)
@@ -1349,7 +1375,7 @@ void UParticleModuleLocationPrimitiveSphere::Render3DPreview(FParticleEmitterIns
 			if (Owner != NULL)
 			{
 				//pkCurve->
-				Position = StartLocation.GetValue(0.0f, Owner->Component);
+				Position = StartLocation.GetValue(0.0f, Context.GetDistributionData());
 			}
 		}
 	}
@@ -1369,7 +1395,7 @@ void UParticleModuleLocationPrimitiveSphere::Render3DPreview(FParticleEmitterIns
 
 	if ((Owner != NULL) && (Owner->Component != NULL))
 	{
-		fRadius = StartRadius.GetValue(Owner->EmitterTime, Owner->Component);
+		fRadius = StartRadius.GetValue(Owner->EmitterTime, Context.GetDistributionData());
 		vAxis[0]	= Owner->SimulationToWorld.TransformVector(Owner->EmitterToSimulation.GetScaledAxis( EAxis::X ));
 		vAxis[1]	= Owner->SimulationToWorld.TransformVector(Owner->EmitterToSimulation.GetScaledAxis( EAxis::Y ));
 		vAxis[2]	= Owner->SimulationToWorld.TransformVector(Owner->EmitterToSimulation.GetScaledAxis( EAxis::Z ));
@@ -1548,8 +1574,10 @@ void UParticleModuleLocationBoneSocket::PostLoad()
 	SetSourceIndexMode();
 }
 
-void UParticleModuleLocationBoneSocket::Spawn(FParticleEmitterInstance* Owner, int32 Offset, float SpawnTime, FBaseParticle* ParticleBase)
+void UParticleModuleLocationBoneSocket::Spawn(const FSpawnContext& Context)
 {
+	FParticleEmitterInstance* Owner = &Context.Owner;
+
 	FModuleLocationBoneSocketInstancePayload* InstancePayload = 
 		(FModuleLocationBoneSocketInstancePayload*)(Owner->GetModuleInstanceData(this));
 	if (InstancePayload == NULL)
@@ -1557,10 +1585,10 @@ void UParticleModuleLocationBoneSocket::Spawn(FParticleEmitterInstance* Owner, i
 		return;
 	}
 
-	FRandomStream& RandomStream = GetRandomStream(Owner);
+	FRandomStream& RandomStream = GetRandomStream(Context);
 
 	// Setup the source skeletal mesh component...
-	GetSkeletalMeshComponentSource(Owner, InstancePayload);
+	GetSkeletalMeshComponentSource(Context, InstancePayload);
 
 	// Early out if source component is still invalid 
 	if (!InstancePayload->SourceComponent.IsValid())
@@ -1584,7 +1612,7 @@ void UParticleModuleLocationBoneSocket::Spawn(FParticleEmitterInstance* Owner, i
 	{
 		SPAWN_INIT
 		{
-			FModuleLocationBoneSocketParticlePayload* ParticlePayload = (FModuleLocationBoneSocketParticlePayload*)((uint8*)&Particle + Offset);
+			FModuleLocationBoneSocketParticlePayload* ParticlePayload = (FModuleLocationBoneSocketParticlePayload*)((uint8*)&Particle + Context.Offset);
 			ParticlePayload->SourceIndex = SourceIndex;
 			Particle.Location = SourceLocation;
 			ensureMsgf(!Particle.Location.ContainsNaN(), TEXT("NaN in Particle Location. Template: %s, Component: %s"), Owner->Component ? *GetNameSafe(Owner->Component->Template) : TEXT("UNKNOWN"), *GetPathNameSafe(Owner->Component));
@@ -1600,7 +1628,7 @@ void UParticleModuleLocationBoneSocket::Spawn(FParticleEmitterInstance* Owner, i
 				PayloadData->Rotation = (FVector3f)RotationQuat.Euler();
 				if (Owner->CurrentLODLevel->RequiredModule->bUseLocalSpace == true)
 				{
-					PayloadData->Rotation = (FVector3f)Owner->Component->GetComponentTransform().InverseTransformVectorNoScale((FVector)PayloadData->Rotation);
+					PayloadData->Rotation = (FVector3f)Context.GetTransform().InverseTransformVectorNoScale((FVector)PayloadData->Rotation);
 				}
 			}
 		}
@@ -1633,8 +1661,9 @@ void UParticleModuleLocationBoneSocket::UpdatePrevBoneLocationsAndVelocities(FMo
 	}
 }
 
-void UParticleModuleLocationBoneSocket::Update(FParticleEmitterInstance* Owner, int32 Offset, float DeltaTime)
+void UParticleModuleLocationBoneSocket::Update(const FUpdateContext& Context)
 {
+	FParticleEmitterInstance* Owner = &Context.Owner;
 	FModuleLocationBoneSocketInstancePayload* InstancePayload = 
 		(FModuleLocationBoneSocketInstancePayload*)(Owner->GetModuleInstanceData(this));
 	if (!InstancePayload->SourceComponent.IsValid())
@@ -1645,7 +1674,7 @@ void UParticleModuleLocationBoneSocket::Update(FParticleEmitterInstance* Owner,
 	USkeletalMeshComponent* SourceComponent = InstancePayload->SourceComponent.Get();
 	if (InheritingBoneVelocity())
 	{
-		UpdatePrevBoneLocationsAndVelocities(InstancePayload, SourceComponent, DeltaTime);
+		UpdatePrevBoneLocationsAndVelocities(InstancePayload, SourceComponent, Context.DeltaTime);
 	}
 
 	if (bUpdatePositionEachFrame == false)
@@ -1689,9 +1718,10 @@ void UParticleModuleLocationBoneSocket::Update(FParticleEmitterInstance* Owner,
 	END_UPDATE_LOOP;
 }
 
-void UParticleModuleLocationBoneSocket::FinalUpdate(FParticleEmitterInstance* Owner, int32 Offset, float DeltaTime)
+void UParticleModuleLocationBoneSocket::FinalUpdate(const FUpdateContext& Context)
 {
-	Super::FinalUpdate(Owner, Offset, DeltaTime);
+	Super::FinalUpdate(Context);
+	FParticleEmitterInstance* Owner = &Context.Owner;
 
 	FModuleLocationBoneSocketInstancePayload* InstancePayload = 
 		(FModuleLocationBoneSocketInstancePayload*)(Owner->GetModuleInstanceData(this));
@@ -1743,7 +1773,7 @@ void UParticleModuleLocationBoneSocket::FinalUpdate(FParticleEmitterInstance* Ow
 	}
 
 	//Select a new set of bones to spawn from next frame.
-	RegeneratePreSelectedIndices(InstancePayload, SourceComponent, GetRandomStream(Owner));
+	RegeneratePreSelectedIndices(InstancePayload, SourceComponent, GetRandomStream(Context));
 }
 
 uint32 UParticleModuleLocationBoneSocket::RequiredBytes(UParticleModuleTypeDataBase* TypeData)
@@ -1937,8 +1967,9 @@ int32 UParticleModuleLocationBoneSocket::GetMaxSourceIndex(FModuleLocationBoneSo
 	return 0;
 }
 
-void UParticleModuleLocationBoneSocket::GetSkeletalMeshComponentSource(FParticleEmitterInstance* Owner, FModuleLocationBoneSocketInstancePayload* InstancePayload)
+void UParticleModuleLocationBoneSocket::GetSkeletalMeshComponentSource(const FContext& Context, FModuleLocationBoneSocketInstancePayload* InstancePayload)
 {
+	FParticleEmitterInstance* Owner = &Context.Owner;
 	if (Owner == NULL)
 	{
 		InstancePayload->SourceComponent = nullptr;
@@ -1992,7 +2023,7 @@ void UParticleModuleLocationBoneSocket::GetSkeletalMeshComponentSource(FParticle
 		if (NewSkelComp)
 		{
 			InstancePayload->SourceComponent = NewSkelComp;
-			RegeneratePreSelectedIndices(InstancePayload, NewSkelComp, GetRandomStream(Owner));
+			RegeneratePreSelectedIndices(InstancePayload, NewSkelComp, GetRandomStream(Context));
 		}
 	}
 }
@@ -2238,8 +2269,9 @@ void UParticleModuleLocationSkelVertSurface::PostEditChangeProperty(FPropertyCha
 }
 #endif // WITH_EDITOR
 
-void UParticleModuleLocationSkelVertSurface::Spawn(FParticleEmitterInstance* Owner, int32 Offset, float SpawnTime, FBaseParticle* ParticleBase)
+void UParticleModuleLocationSkelVertSurface::Spawn(const FSpawnContext& Context)
 {
+	FParticleEmitterInstance* Owner = &Context.Owner;
 	SCOPE_CYCLE_COUNTER(STAT_ParticleSkelMeshSurfTime);
 	FModuleLocationVertSurfaceInstancePayload* InstancePayload = 
 		(FModuleLocationVertSurfaceInstancePayload*)(Owner->GetModuleInstanceData(this));
@@ -2248,9 +2280,9 @@ void UParticleModuleLocationSkelVertSurface::Spawn(FParticleEmitterInstance* Own
 		return;
 	}
 
-	FRandomStream& RandomStream = GetRandomStream(Owner);
+	FRandomStream& RandomStream = GetRandomStream(Context);
 	
-	GetSkeletalMeshComponentSource(Owner, InstancePayload);
+	GetSkeletalMeshComponentSource(Context, InstancePayload);
 
 	if (!InstancePayload->MeshIsValid())
 	{
@@ -2341,11 +2373,11 @@ void UParticleModuleLocationSkelVertSurface::Spawn(FParticleEmitterInstance* Own
 	FQuat SourceRotation;
 	const int32 MeshRotationOffset = Owner->GetMeshRotationOffset();
 	const bool bMeshRotationActive = MeshRotationOffset > 0 && Owner->IsMeshRotationActive();
-	if (GetParticleLocation(Owner, InstancePayload, SourceComponent, SourceIndex, SourceLocation, SourceRotation, true) == true)
+	if (GetParticleLocation(Context, InstancePayload, SourceComponent, SourceIndex, SourceLocation, SourceRotation, true) == true)
 	{
 		SPAWN_INIT
 		{
-			FModuleLocationVertSurfaceParticlePayload* ParticlePayload = (FModuleLocationVertSurfaceParticlePayload*)((uint8*)&Particle + Offset);
+			FModuleLocationVertSurfaceParticlePayload* ParticlePayload = (FModuleLocationVertSurfaceParticlePayload*)((uint8*)&Particle + Context.Offset);
 			ParticlePayload->SourceIndex = SourceIndex;
 			Particle.Location = SourceLocation;
 			ensureMsgf(!Particle.Location.ContainsNaN(), TEXT("NaN in Particle Location. Template: %s, Component: %s"), Owner->Component ? *GetNameSafe(Owner->Component->Template) : TEXT("UNKNOWN"), *GetPathNameSafe(Owner->Component));
@@ -2437,7 +2469,7 @@ void UParticleModuleLocationSkelVertSurface::Spawn(FParticleEmitterInstance* Own
 				FVector Rot = SourceRotation.Euler();
 				if (Owner->CurrentLODLevel->RequiredModule->bUseLocalSpace == true)
 				{
-					Rot = Owner->Component->GetComponentTransform().InverseTransformVectorNoScale(Rot);
+					Rot = Context.GetTransform().InverseTransformVectorNoScale(Rot);
 				}
 				PayloadData->Rotation = (FVector3f)Rot;
 				PayloadData->InitRotation = (FVector3f)Rot;
@@ -2453,14 +2485,15 @@ void UParticleModuleLocationSkelVertSurface::Spawn(FParticleEmitterInstance* Own
 	}
 }
 
-void UParticleModuleLocationSkelVertSurface::Update(FParticleEmitterInstance* Owner, int32 Offset, float DeltaTime)
+void UParticleModuleLocationSkelVertSurface::Update(const FUpdateContext& Context)
 {
 	SCOPE_CYCLE_COUNTER(STAT_ParticleSkelMeshSurfTime);
 	
+	FParticleEmitterInstance* Owner = &Context.Owner;
 	FModuleLocationVertSurfaceInstancePayload* InstancePayload = 
 		(FModuleLocationVertSurfaceInstancePayload*)(Owner->GetModuleInstanceData(this));
 
-	GetSkeletalMeshComponentSource(Owner, InstancePayload);
+	GetSkeletalMeshComponentSource(Context, InstancePayload);
 
 	if (!InstancePayload->MeshIsValid())
 	{
@@ -2471,7 +2504,7 @@ void UParticleModuleLocationSkelVertSurface::Update(FParticleEmitterInstance* Ow
 
 	if(bInheritBoneVelocity)
 	{
-		const float InvDeltaTime = (DeltaTime > 0.0f) ? 1.0f / DeltaTime : 0.0f;
+		const float InvDeltaTime = (Context.DeltaTime > 0.0f) ? 1.0f / Context.DeltaTime : 0.0f;
 
 		// Calculate velocities to be used when spawning particles later this frame
 		for(int32 ValidBoneIndex = 0; ValidBoneIndex < InstancePayload->NumValidAssociatedBoneIndices; ++ValidBoneIndex)
@@ -2503,12 +2536,12 @@ void UParticleModuleLocationSkelVertSurface::Update(FParticleEmitterInstance* Ow
 	const bool bMeshRotationActive = MeshRotationOffset > 0 && Owner->IsMeshRotationActive();
 	const FTransform& OwnerTM = Owner->Component->GetAsyncComponentToWorld();
 
-	FRandomStream& RandomStream = GetRandomStream(Owner);
+	FRandomStream& RandomStream = GetRandomStream(Context);
 
 	BEGIN_UPDATE_LOOP;
 	{
 		FModuleLocationVertSurfaceParticlePayload* ParticlePayload = (FModuleLocationVertSurfaceParticlePayload*)((uint8*)&Particle + Offset);
-		if (GetParticleLocation(Owner, InstancePayload, SourceComponent, ParticlePayload->SourceIndex, SourceLocation, SourceRotation) == true)
+		if (GetParticleLocation(Context, InstancePayload, SourceComponent, ParticlePayload->SourceIndex, SourceLocation, SourceRotation) == true)
 		{
 			Particle.Location = SourceLocation;
 			ensureMsgf(!Particle.Location.ContainsNaN(), TEXT("NaN in Particle Location. Template: %s, Component: %s"), Owner->Component ? *GetNameSafe(Owner->Component->Template) : TEXT("UNKNOWN"), *GetPathNameSafe(Owner->Component));
@@ -2534,10 +2567,11 @@ void UParticleModuleLocationSkelVertSurface::Update(FParticleEmitterInstance* Ow
 	END_UPDATE_LOOP;
 }
 
-void UParticleModuleLocationSkelVertSurface::FinalUpdate(FParticleEmitterInstance* Owner, int32 Offset, float DeltaTime)
+void UParticleModuleLocationSkelVertSurface::FinalUpdate(const FUpdateContext& Context)
 {
-	Super::FinalUpdate(Owner, Offset, DeltaTime);
+	Super::FinalUpdate(Context);
 
+	FParticleEmitterInstance* Owner = &Context.Owner;
 	FModuleLocationVertSurfaceInstancePayload* InstancePayload = 
 		(FModuleLocationVertSurfaceInstancePayload*)(Owner->GetModuleInstanceData(this));
 	if (!InstancePayload->MeshIsValid())
@@ -2708,8 +2742,9 @@ bool UParticleModuleLocationSkelVertSurface::IsValidForLODLevel(UParticleLODLeve
 
 #endif
 
-void UParticleModuleLocationSkelVertSurface::GetSkeletalMeshComponentSource(FParticleEmitterInstance* Owner, FModuleLocationVertSurfaceInstancePayload* InstancePayload)
+void UParticleModuleLocationSkelVertSurface::GetSkeletalMeshComponentSource(const FContext& Context, FModuleLocationVertSurfaceInstancePayload* InstancePayload)
 {
+	FParticleEmitterInstance* Owner = &Context.Owner;
 	if (Owner == NULL)
 	{
 		return;
@@ -2804,7 +2839,7 @@ void UParticleModuleLocationSkelVertSurface::GetSkeletalMeshComponentSource(FPar
 }
 
 
-bool UParticleModuleLocationSkelVertSurface::GetParticleLocation(FParticleEmitterInstance* Owner, FModuleLocationVertSurfaceInstancePayload* InstPayload,
+bool UParticleModuleLocationSkelVertSurface::GetParticleLocation(const FContext& Context, FModuleLocationVertSurfaceInstancePayload* InstPayload,
 	USkeletalMeshComponent* InSkelMeshComponent, int32 InPrimaryVertexIndex, 
 	FVector& OutPosition, FQuat& OutRotation, bool bSpawning /* = false*/)
 {
@@ -2886,9 +2921,9 @@ bool UParticleModuleLocationSkelVertSurface::GetParticleLocation(FParticleEmitte
 		}
 	}
 
-	if (Owner->CurrentLODLevel->RequiredModule->bUseLocalSpace == true)
+	if (Context.Owner.CurrentLODLevel->RequiredModule->bUseLocalSpace == true)
 	{
-		OutPosition = Owner->Component->GetComponentTransform().InverseTransformPosition(OutPosition);
+		OutPosition = Context.GetTransform().InverseTransformPosition(OutPosition);
 	}
 
 	OutPosition += UniversalOffset;
diff --git a/Engine/Source/Runtime/Engine/Private/Particles/ParticleModules_Material.cpp b/Engine/Source/Runtime/Engine/Private/Particles/ParticleModules_Material.cpp
--- a/Engine/Source/Runtime/Engine/Private/Particles/ParticleModules_Material.cpp
+++ b/Engine/Source/Runtime/Engine/Private/Particles/ParticleModules_Material.cpp
@@ -36,7 +36,7 @@ UParticleModuleMeshMaterial::UParticleModuleMeshMaterial(const FObjectInitialize
 //	TArray<class UMaterialInstance*> MeshMaterials;
 	//## END PROPS ParticleModuleMeshMaterial
 
-void UParticleModuleMeshMaterial::Spawn(FParticleEmitterInstance* Owner, int32 Offset, float SpawnTime, FBaseParticle* ParticleBase)
+void UParticleModuleMeshMaterial::Spawn(const FSpawnContext& Context)
 {
 
 }
diff --git a/Engine/Source/Runtime/Engine/Private/Particles/ParticleModules_Orbit.cpp b/Engine/Source/Runtime/Engine/Private/Particles/ParticleModules_Orbit.cpp
--- a/Engine/Source/Runtime/Engine/Private/Particles/ParticleModules_Orbit.cpp
+++ b/Engine/Source/Runtime/Engine/Private/Particles/ParticleModules_Orbit.cpp
@@ -109,8 +109,10 @@ void UParticleModuleOrbit::PostEditChangeProperty(FPropertyChangedEvent& Propert
 }
 #endif // WITH_EDITOR
 
-void UParticleModuleOrbit::Spawn(FParticleEmitterInstance* Owner, int32 Offset, float SpawnTime, FBaseParticle* ParticleBase)
+void UParticleModuleOrbit::Spawn(const FSpawnContext& Context)
 {
+	FParticleEmitterInstance* Owner = &Context.Owner;
+
 	SPAWN_INIT;
 	{
 		PARTICLE_ELEMENT(FOrbitChainModuleInstancePayload, OrbitPayload);
@@ -121,11 +123,11 @@ void UParticleModuleOrbit::Spawn(FParticleEmitterInstance* Owner, int32 Offset,
 			FVector LocalOffset;
 			if (OffsetOptions.bUseEmitterTime == false)
 			{
-				LocalOffset = OffsetAmount.GetValue(Particle.RelativeTime, Owner->Component);
+				LocalOffset = OffsetAmount.GetValue(Particle.RelativeTime, Context.GetDistributionData());
 			}
 			else
 			{
-				LocalOffset = OffsetAmount.GetValue(Owner->EmitterTime, Owner->Component);
+				LocalOffset = OffsetAmount.GetValue(Owner->EmitterTime, Context.GetDistributionData());
 			}
 			OrbitPayload.BaseOffset += (FVector3f)LocalOffset;
 			OrbitPayload.PreviousOffset = OrbitPayload.Offset;
@@ -138,11 +140,11 @@ void UParticleModuleOrbit::Spawn(FParticleEmitterInstance* Owner, int32 Offset,
 			FVector LocalRotation;
 			if (RotationOptions.bUseEmitterTime == false)
 			{
-				LocalRotation = RotationAmount.GetValue(Particle.RelativeTime, Owner->Component);
+				LocalRotation = RotationAmount.GetValue(Particle.RelativeTime, Context.GetDistributionData());
 			}
 			else
 			{
-				LocalRotation = RotationAmount.GetValue(Owner->EmitterTime, Owner->Component);
+				LocalRotation = RotationAmount.GetValue(Owner->EmitterTime, Context.GetDistributionData());
 			}
 			OrbitPayload.Rotation += (FVector3f)LocalRotation;
 		}
@@ -153,11 +155,11 @@ void UParticleModuleOrbit::Spawn(FParticleEmitterInstance* Owner, int32 Offset,
 			FVector LocalRotationRate;
 			if (RotationRateOptions.bUseEmitterTime == false)
 			{
-				LocalRotationRate = RotationRateAmount.GetValue(Particle.RelativeTime, Owner->Component);
+				LocalRotationRate = RotationRateAmount.GetValue(Particle.RelativeTime, Context.GetDistributionData());
 			}
 			else
 			{
-				LocalRotationRate = RotationRateAmount.GetValue(Owner->EmitterTime, Owner->Component);
+				LocalRotationRate = RotationRateAmount.GetValue(Owner->EmitterTime, Context.GetDistributionData());
 			}
 			OrbitPayload.BaseRotationRate += (FVector3f)LocalRotationRate;
 			OrbitPayload.RotationRate += (FVector3f)LocalRotationRate;
@@ -165,8 +167,9 @@ void UParticleModuleOrbit::Spawn(FParticleEmitterInstance* Owner, int32 Offset,
 	}
 }
 
-void UParticleModuleOrbit::Update(FParticleEmitterInstance* Owner, int32 Offset, float DeltaTime)
+void UParticleModuleOrbit::Update(const FUpdateContext& Context)
 {
+	FParticleEmitterInstance* Owner = &Context.Owner;
 	BEGIN_UPDATE_LOOP;
 	{
 		PARTICLE_ELEMENT(FOrbitChainModuleInstancePayload, OrbitPayload);
@@ -177,11 +180,11 @@ void UParticleModuleOrbit::Update(FParticleEmitterInstance* Owner, int32 Offset,
 			FVector LocalOffset;
 			if (OffsetOptions.bUseEmitterTime == false)
 			{
-				LocalOffset = OffsetAmount.GetValue(Particle.RelativeTime, Owner->Component);
+				LocalOffset = OffsetAmount.GetValue(Particle.RelativeTime, Context.GetDistributionData());
 			}
 			else
 			{
-				LocalOffset = OffsetAmount.GetValue(Owner->EmitterTime, Owner->Component);
+				LocalOffset = OffsetAmount.GetValue(Owner->EmitterTime, Context.GetDistributionData());
 			}
 
 			//@todo. Do we need to update the base offset here???
@@ -196,11 +199,11 @@ void UParticleModuleOrbit::Update(FParticleEmitterInstance* Owner, int32 Offset,
 			FVector LocalRotation;
 			if (RotationOptions.bUseEmitterTime == false)
 			{
-				LocalRotation = RotationAmount.GetValue(Particle.RelativeTime, Owner->Component);
+				LocalRotation = RotationAmount.GetValue(Particle.RelativeTime, Context.GetDistributionData());
 			}
 			else
 			{
-				LocalRotation = RotationAmount.GetValue(Owner->EmitterTime, Owner->Component);
+				LocalRotation = RotationAmount.GetValue(Owner->EmitterTime, Context.GetDistributionData());
 			}
 			OrbitPayload.Rotation += (FVector3f)LocalRotation;
 		}
@@ -212,11 +215,11 @@ void UParticleModuleOrbit::Update(FParticleEmitterInstance* Owner, int32 Offset,
 			FVector LocalRotationRate;
 			if (RotationRateOptions.bUseEmitterTime == false)
 			{
-				LocalRotationRate = RotationRateAmount.GetValue(Particle.RelativeTime, Owner->Component);
+				LocalRotationRate = RotationRateAmount.GetValue(Particle.RelativeTime, Context.GetDistributionData());
 			}
 			else
 			{
-				LocalRotationRate = RotationRateAmount.GetValue(Owner->EmitterTime, Owner->Component);
+				LocalRotationRate = RotationRateAmount.GetValue(Owner->EmitterTime, Context.GetDistributionData());
 			}
 			//@todo. Do we need to update the base rotationrate here???
 //			OrbitPayload.BaseRotationRate += LocalRotationRate;
diff --git a/Engine/Source/Runtime/Engine/Private/Particles/ParticleModules_Parameter.cpp b/Engine/Source/Runtime/Engine/Private/Particles/ParticleModules_Parameter.cpp
--- a/Engine/Source/Runtime/Engine/Private/Particles/ParticleModules_Parameter.cpp
+++ b/Engine/Source/Runtime/Engine/Private/Particles/ParticleModules_Parameter.cpp
@@ -134,13 +134,15 @@ bool UParticleModuleParameterDynamic::CanTickInAnyThread()
 }
 
 
-void UParticleModuleParameterDynamic::Spawn(FParticleEmitterInstance* Owner, int32 Offset, float SpawnTime, FBaseParticle* ParticleBase)
+void UParticleModuleParameterDynamic::Spawn(const FSpawnContext& Context)
 {
-	SpawnEx(Owner, Offset, SpawnTime, &GetRandomStream(Owner), ParticleBase);
+	SpawnEx(Context, &GetRandomStream(Context));
 }
 
-void UParticleModuleParameterDynamic::SpawnEx(FParticleEmitterInstance* Owner, int32 Offset, float SpawnTime, struct FRandomStream* InRandomStream, FBaseParticle* ParticleBase)
+void UParticleModuleParameterDynamic::SpawnEx(const FSpawnContext& Context, struct FRandomStream* InRandomStream)
 {
+	FParticleEmitterInstance* Owner = &Context.Owner;
+
 	SPAWN_INIT;
 	{
 		PARTICLE_ELEMENT(FEmitterDynamicParameterPayload, DynamicPayload);
@@ -163,7 +165,7 @@ void UParticleModuleParameterDynamic::SpawnEx(FParticleEmitterInstance* Owner, i
 	}
 }
 
-void UParticleModuleParameterDynamic::Update(FParticleEmitterInstance* Owner, int32 Offset, float DeltaTime)
+void UParticleModuleParameterDynamic::Update(const FUpdateContext& Context)
 {
 	if (UpdateFlags == EDPU_UPDATE_NONE)
 	{
@@ -171,13 +173,14 @@ void UParticleModuleParameterDynamic::Update(FParticleEmitterInstance* Owner, in
 		return;
 	}
 
+	FParticleEmitterInstance* Owner = &Context.Owner;
 	if ((Owner == NULL) || (Owner->ActiveParticles <= 0) || 
 		(Owner->ParticleData == NULL) || (Owner->ParticleIndices == NULL))
 	{
 		return;
 	}
 
-	FRandomStream& RandomStream = GetRandomStream(Owner);
+	FRandomStream& RandomStream = GetRandomStream(Context);
 
 	// 
 	int32 ParameterIndex = ParticleDynamicParameter_GetIndexFromFlag(UpdateFlags);
diff --git a/Engine/Source/Runtime/Engine/Private/Particles/ParticleModules_Size.cpp b/Engine/Source/Runtime/Engine/Private/Particles/ParticleModules_Size.cpp
--- a/Engine/Source/Runtime/Engine/Private/Particles/ParticleModules_Size.cpp
+++ b/Engine/Source/Runtime/Engine/Private/Particles/ParticleModules_Size.cpp
@@ -83,15 +83,17 @@ void UParticleModuleSize::PostEditChangeProperty(FPropertyChangedEvent& Property
 }
 #endif // WITH_EDITOR
 
-void UParticleModuleSize::Spawn(FParticleEmitterInstance* Owner, int32 Offset, float SpawnTime, FBaseParticle* ParticleBase)
+void UParticleModuleSize::Spawn(const FSpawnContext& Context)
 {
-	SpawnEx(Owner, Offset, SpawnTime, &GetRandomStream(Owner), ParticleBase);
+	SpawnEx(Context, &GetRandomStream(Context));
 }
 
-void UParticleModuleSize::SpawnEx(FParticleEmitterInstance* Owner, int32 Offset, float SpawnTime, struct FRandomStream* InRandomStream, FBaseParticle* ParticleBase)
+void UParticleModuleSize::SpawnEx(const FSpawnContext& Context, struct FRandomStream* InRandomStream)
 {
+	FParticleEmitterInstance* Owner = &Context.Owner;
+
 	SPAWN_INIT;
-	FVector Size		 = StartSize.GetValue(Owner->EmitterTime, Owner->Component, 0, InRandomStream);
+	FVector Size		 = StartSize.GetValue(Owner->EmitterTime, Context.GetDistributionData(), 0, InRandomStream);
 	Particle.Size	+= (FVector3f)Size;
 
 	AdjustParticleBaseSizeForUVFlipping(Size, Owner->CurrentLODLevel->RequiredModule->UVFlippingMode, *InRandomStream);
@@ -188,10 +190,12 @@ bool UParticleModuleSizeMultiplyLife::IsValidForLODLevel(UParticleLODLevel* LODL
 
 #endif // WITH_EDITOR
 
-void UParticleModuleSizeMultiplyLife::Spawn(FParticleEmitterInstance* Owner, int32 Offset, float SpawnTime, FBaseParticle* ParticleBase)
+void UParticleModuleSizeMultiplyLife::Spawn(const FSpawnContext& Context)
 {
+	FParticleEmitterInstance* Owner = &Context.Owner;
+
 	SPAWN_INIT;
-	FVector SizeScale = LifeMultiplier.GetValue(Particle.RelativeTime, Owner->Component);
+	FVector SizeScale = LifeMultiplier.GetValue(Particle.RelativeTime, Context.GetDistributionData());
 	if(MultiplyX)
 	{
 		Particle.Size.X *= SizeScale.X;
@@ -206,8 +210,9 @@ void UParticleModuleSizeMultiplyLife::Spawn(FParticleEmitterInstance* Owner, int
 	}
 }
 
-void UParticleModuleSizeMultiplyLife::Update(FParticleEmitterInstance* Owner, int32 Offset, float DeltaTime)
+void UParticleModuleSizeMultiplyLife::Update(const FUpdateContext& Context)
 {
+	FParticleEmitterInstance* Owner = &Context.Owner;
 	if ((Owner == NULL) || (Owner->ActiveParticles <= 0) || 
 		(Owner->ParticleData == NULL) || (Owner->ParticleIndices == NULL))
 	{
@@ -235,7 +240,7 @@ void UParticleModuleSizeMultiplyLife::Update(FParticleEmitterInstance* Owner, in
 		{
 			BEGIN_UPDATE_LOOP
 			{
-				FVector3f SizeScale(LifeMultiplier.GetValue(Particle.RelativeTime, Owner->Component));
+				FVector3f SizeScale(LifeMultiplier.GetValue(Particle.RelativeTime, Context.GetDistributionData()));
 				FPlatformMisc::Prefetch(ParticleData, (ParticleIndices[i+1] * ParticleStride));
 				FPlatformMisc::Prefetch(ParticleData, (ParticleIndices[i+1] * ParticleStride) + PLATFORM_CACHE_LINE_SIZE);
 				Particle.Size.X *= SizeScale.X;
@@ -256,7 +261,7 @@ void UParticleModuleSizeMultiplyLife::Update(FParticleEmitterInstance* Owner, in
 			int32 Index = MultiplyX ? 0 : (MultiplyY ? 1 : 2);
 			BEGIN_UPDATE_LOOP
 			{
-				FVector SizeScale = LifeMultiplier.GetValue(Particle.RelativeTime, Owner->Component);
+				FVector SizeScale = LifeMultiplier.GetValue(Particle.RelativeTime, Context.GetDistributionData());
 				FPlatformMisc::Prefetch(ParticleData, (ParticleIndices[i+1] * ParticleStride));
 				FPlatformMisc::Prefetch(ParticleData, (ParticleIndices[i+1] * ParticleStride) + PLATFORM_CACHE_LINE_SIZE);
 				Particle.Size[Index] *= SizeScale[Index];
@@ -267,7 +272,7 @@ void UParticleModuleSizeMultiplyLife::Update(FParticleEmitterInstance* Owner, in
 		{
 			BEGIN_UPDATE_LOOP
 			{
-				FVector3f SizeScale(LifeMultiplier.GetValue(Particle.RelativeTime, Owner->Component));
+				FVector3f SizeScale(LifeMultiplier.GetValue(Particle.RelativeTime, Context.GetDistributionData()));
 				FPlatformMisc::Prefetch(ParticleData, (ParticleIndices[i+1] * ParticleStride));
 				FPlatformMisc::Prefetch(ParticleData, (ParticleIndices[i+1] * ParticleStride) + PLATFORM_CACHE_LINE_SIZE);
 				if(MultiplyX)
@@ -366,17 +371,19 @@ bool UParticleModuleSizeScale::IsValidForLODLevel(UParticleLODLevel* LODLevel, F
 
 #endif // WITH_EDITOR
 
-void UParticleModuleSizeScale::Spawn(FParticleEmitterInstance* Owner, int32 Offset, float SpawnTime, FBaseParticle* ParticleBase)
+void UParticleModuleSizeScale::Spawn(const FSpawnContext& Context)
 {
+	FParticleEmitterInstance* Owner = &Context.Owner;
 	SPAWN_INIT;
-	FVector3f ScaleFactor(SizeScale.GetValue(Particle.RelativeTime, Owner->Component));
+	FVector3f ScaleFactor(SizeScale.GetValue(Particle.RelativeTime, Context.GetDistributionData()));
 	Particle.Size = GetParticleBaseSize(*ParticleBase) * ScaleFactor;
 }
 
-void UParticleModuleSizeScale::Update(FParticleEmitterInstance* Owner, int32 Offset, float DeltaTime)
+void UParticleModuleSizeScale::Update(const FUpdateContext& Context)
 {
+	FParticleEmitterInstance* Owner = &Context.Owner;
 	BEGIN_UPDATE_LOOP;
-		FVector3f ScaleFactor(SizeScale.GetValue(Particle.RelativeTime, Owner->Component));
+		FVector3f ScaleFactor(SizeScale.GetValue(Particle.RelativeTime, Context.GetDistributionData()));
 		Particle.Size = GetParticleBaseSize(Particle) * ScaleFactor;
 	END_UPDATE_LOOP;
 }
@@ -402,11 +409,12 @@ UParticleModuleSizeScaleBySpeed::UParticleModuleSizeScaleBySpeed(const FObjectIn
 	MaxScale.Y = 1;
 }
 
-void UParticleModuleSizeScaleBySpeed::Update(FParticleEmitterInstance* Owner, int32 Offset, float DeltaTime)
+void UParticleModuleSizeScaleBySpeed::Update(const FUpdateContext& Context)
 {
 	FVector Scale(SpeedScale.X, SpeedScale.Y, 1.0f);
 	FVector ScaleMax(MaxScale.X, MaxScale.Y, 1.0f);
 
+	FParticleEmitterInstance* Owner = &Context.Owner;
 	BEGIN_UPDATE_LOOP;
 		FVector Size = Scale * Particle.Velocity.Size();
 		Size = Size.ComponentMax(FVector(1.0f));
diff --git a/Engine/Source/Runtime/Engine/Private/Particles/ParticleModules_Spawn.cpp b/Engine/Source/Runtime/Engine/Private/Particles/ParticleModules_Spawn.cpp
--- a/Engine/Source/Runtime/Engine/Private/Particles/ParticleModules_Spawn.cpp
+++ b/Engine/Source/Runtime/Engine/Private/Particles/ParticleModules_Spawn.cpp
@@ -104,10 +104,9 @@ void UParticleModuleSpawn::PostEditChangeProperty(FPropertyChangedEvent& Propert
 }
 #endif // WITH_EDITOR
 
-bool UParticleModuleSpawn::GetSpawnAmount(FParticleEmitterInstance* Owner, 
+bool UParticleModuleSpawn::GetSpawnAmount(const FContext& Context, 
 	int32 Offset, float OldLeftover, float DeltaTime, int32& Number, float& InRate)
 {
-	check(Owner);
 	return false;
 }
 
@@ -301,16 +300,17 @@ uint32 UParticleModuleSpawnPerUnit::RequiredBytesPerInstance()
 	return sizeof(FParticleSpawnPerUnitInstancePayload);
 }
 
-bool UParticleModuleSpawnPerUnit::GetSpawnAmount(FParticleEmitterInstance* Owner, 
+bool UParticleModuleSpawnPerUnit::GetSpawnAmount(const FContext& Context, 
 	int32 Offset, float OldLeftover, float DeltaTime, int32& Number, float& Rate)
 {
+	FParticleEmitterInstance* Owner = &Context.Owner;
 	check(Owner);
 
 	bool bMoved = false;
 	FParticleSpawnPerUnitInstancePayload* SPUPayload = NULL;
 	float NewTravelLeftover = 0.0f;
 
-	float ParticlesPerUnit = SpawnPerUnit.GetValue(Owner->EmitterTime, Owner->Component) / UnitScalar;
+	float ParticlesPerUnit = SpawnPerUnit.GetValue(Owner->EmitterTime, Context.GetDistributionData()) / UnitScalar;
 	// Allow for PPU of 0.0f to allow for 'turning off' an emitter when moving
 	if (ParticlesPerUnit >= 0.0f)
 	{
diff --git a/Engine/Source/Runtime/Engine/Private/Particles/ParticleModules_Velocity.cpp b/Engine/Source/Runtime/Engine/Private/Particles/ParticleModules_Velocity.cpp
--- a/Engine/Source/Runtime/Engine/Private/Particles/ParticleModules_Velocity.cpp
+++ b/Engine/Source/Runtime/Engine/Private/Particles/ParticleModules_Velocity.cpp
@@ -73,22 +73,24 @@ void UParticleModuleVelocity::PostEditChangeProperty(FPropertyChangedEvent& Prop
 }
 #endif // WITH_EDITOR
 
-void UParticleModuleVelocity::Spawn(FParticleEmitterInstance* Owner, int32 Offset, float SpawnTime, FBaseParticle* ParticleBase)
+void UParticleModuleVelocity::Spawn(const FSpawnContext& Context)
 {
-	SpawnEx(Owner, Offset, SpawnTime, &GetRandomStream(Owner), ParticleBase);
+	SpawnEx(Context, &GetRandomStream(Context));
 }
 
-void UParticleModuleVelocity::SpawnEx(FParticleEmitterInstance* Owner, int32 Offset, float SpawnTime, struct FRandomStream* InRandomStream, FBaseParticle* ParticleBase)
+void UParticleModuleVelocity::SpawnEx(const FSpawnContext& Context, struct FRandomStream* InRandomStream)
 {
+	FParticleEmitterInstance* Owner = &Context.Owner;
+
 	SPAWN_INIT;
 	{
-		FVector Vel = StartVelocity.GetValue(Owner->EmitterTime, Owner->Component, 0, InRandomStream);
+		FVector Vel = StartVelocity.GetValue(Owner->EmitterTime, Context.GetDistributionData(), 0, InRandomStream);
 		FVector FromOrigin = (Particle.Location - Owner->EmitterToSimulation.GetOrigin()).GetSafeNormal();
 
 		FVector OwnerScale(1.0f);
 		if ((bApplyOwnerScale == true) && Owner->Component)
 		{
-			OwnerScale = Owner->Component->GetComponentTransform().GetScale3D();
+			OwnerScale = Context.GetTransform().GetScale3D();
 		}
 
 		UParticleLODLevel* LODLevel	= Owner->SpriteTemplate->GetCurrentLODLevel(Owner);
@@ -109,7 +111,7 @@ void UParticleModuleVelocity::SpawnEx(FParticleEmitterInstance* Owner, int32 Off
 			Vel = Owner->EmitterToSimulation.TransformVector(Vel);
 		}
 		Vel *= OwnerScale;
-		Vel += FromOrigin * StartVelocityRadial.GetValue(Owner->EmitterTime, Owner->Component, InRandomStream) * OwnerScale;
+		Vel += FromOrigin * StartVelocityRadial.GetValue(Owner->EmitterTime, Context.GetDistributionData(), InRandomStream) * OwnerScale;
 		Particle.Velocity		+= (FVector3f)Vel;
 		Particle.BaseVelocity	+= (FVector3f)Vel;
 	}
@@ -171,8 +173,10 @@ void UParticleModuleVelocityInheritParent::PostEditChangeProperty(FPropertyChang
 }
 #endif // WITH_EDITOR
 
-void UParticleModuleVelocityInheritParent::Spawn(FParticleEmitterInstance* Owner, int32 Offset, float SpawnTime, FBaseParticle* ParticleBase)
+void UParticleModuleVelocityInheritParent::Spawn(const FSpawnContext& Context)
 {
+	FParticleEmitterInstance* Owner = &Context.Owner;
+
 	SPAWN_INIT;
 
 	FVector Vel(0.0f);
@@ -181,14 +185,14 @@ void UParticleModuleVelocityInheritParent::Spawn(FParticleEmitterInstance* Owner
 	check(LODLevel);
 	if (LODLevel->RequiredModule->bUseLocalSpace)
 	{
-		Vel = Owner->Component->GetComponentTransform().InverseTransformVector(Owner->Component->PartSysVelocity);
+		Vel = Context.GetTransform().InverseTransformVector(Owner->Component->PartSysVelocity);
 	}
 	else
 	{
 		Vel = Owner->Component->PartSysVelocity;
 	}
 
-	FVector vScale = Scale.GetValue(Owner->EmitterTime, Owner->Component);
+	FVector vScale = Scale.GetValue(Owner->EmitterTime, Context.GetDistributionData());
 
 	Vel *= vScale;
 
@@ -233,24 +237,27 @@ void UParticleModuleVelocityOverLifetime::PostEditChangeProperty(FPropertyChange
 }
 #endif // WITH_EDITOR
 
-void UParticleModuleVelocityOverLifetime::Spawn(FParticleEmitterInstance* Owner, int32 Offset, float SpawnTime, FBaseParticle* ParticleBase)
+void UParticleModuleVelocityOverLifetime::Spawn(const FSpawnContext& Context)
 {
 	if (Absolute)
 	{
+		FParticleEmitterInstance* Owner = &Context.Owner;
+
 		SPAWN_INIT;
 		FVector OwnerScale(1.0f);
 		if ((bApplyOwnerScale == true) && Owner && Owner->Component)
 		{
-			OwnerScale = Owner->Component->GetComponentTransform().GetScale3D();
+			OwnerScale = Context.GetTransform().GetScale3D();
 		}
-		FVector Vel = VelOverLife.GetValue(Particle.RelativeTime, Owner->Component) * OwnerScale;
+		FVector Vel = VelOverLife.GetValue(Particle.RelativeTime, Context.GetDistributionData()) * OwnerScale;
 		Particle.Velocity		= (FVector3f)Vel;
 		Particle.BaseVelocity	= (FVector3f)Vel;
 	}
 }
 
-void UParticleModuleVelocityOverLifetime::Update(FParticleEmitterInstance* Owner, int32 Offset, float DeltaTime)
+void UParticleModuleVelocityOverLifetime::Update(const FUpdateContext& Context)
 {
+	FParticleEmitterInstance* Owner = &Context.Owner;
 	check(Owner && Owner->Component);
 	UParticleLODLevel* LODLevel	= Owner->SpriteTemplate->GetCurrentLODLevel(Owner);
 	check(LODLevel);
@@ -270,7 +277,7 @@ void UParticleModuleVelocityOverLifetime::Update(FParticleEmitterInstance* Owner
 				const FMatrix LocalToWorld = OwnerTM.ToMatrixNoScale();
 				BEGIN_UPDATE_LOOP;
 				{
-					Vel = VelOverLife.GetValue(Particle.RelativeTime, Owner->Component);
+					Vel = VelOverLife.GetValue(Particle.RelativeTime, Context.GetDistributionData());
 					Particle.Velocity = FVector4f(LocalToWorld.TransformVector(Vel) * OwnerScale);
 				}
 				END_UPDATE_LOOP;
@@ -279,7 +286,7 @@ void UParticleModuleVelocityOverLifetime::Update(FParticleEmitterInstance* Owner
 			{
 				BEGIN_UPDATE_LOOP;
 				{
-					Particle.Velocity = FVector3f(VelOverLife.GetValue(Particle.RelativeTime, Owner->Component) * OwnerScale);
+					Particle.Velocity = FVector3f(VelOverLife.GetValue(Particle.RelativeTime, Context.GetDistributionData()) * OwnerScale);
 				}
 				END_UPDATE_LOOP;
 			}
@@ -290,7 +297,7 @@ void UParticleModuleVelocityOverLifetime::Update(FParticleEmitterInstance* Owner
 			{
 				BEGIN_UPDATE_LOOP;
 				{
-					Particle.Velocity = FVector3f(VelOverLife.GetValue(Particle.RelativeTime, Owner->Component) * OwnerScale);
+					Particle.Velocity = FVector3f(VelOverLife.GetValue(Particle.RelativeTime, Context.GetDistributionData()) * OwnerScale);
 				}
 				END_UPDATE_LOOP;
 			}
@@ -301,7 +308,7 @@ void UParticleModuleVelocityOverLifetime::Update(FParticleEmitterInstance* Owner
 				const FMatrix InvMat = LocalToWorld.InverseFast();
 				BEGIN_UPDATE_LOOP;
 				{
-					Vel = VelOverLife.GetValue(Particle.RelativeTime, Owner->Component);
+					Vel = VelOverLife.GetValue(Particle.RelativeTime, Context.GetDistributionData());
 					Particle.Velocity = FVector4f(InvMat.TransformVector(Vel) * OwnerScale);
 				}
 				END_UPDATE_LOOP;
@@ -318,7 +325,7 @@ void UParticleModuleVelocityOverLifetime::Update(FParticleEmitterInstance* Owner
 				const FMatrix LocalToWorld = OwnerTM.ToMatrixNoScale();
 				BEGIN_UPDATE_LOOP;
 				{
-					Vel = VelOverLife.GetValue(Particle.RelativeTime, Owner->Component);
+					Vel = VelOverLife.GetValue(Particle.RelativeTime, Context.GetDistributionData());
 					Particle.Velocity *= FVector4f(LocalToWorld.TransformVector(Vel) * OwnerScale);
 				}
 				END_UPDATE_LOOP;
@@ -327,7 +334,7 @@ void UParticleModuleVelocityOverLifetime::Update(FParticleEmitterInstance* Owner
 			{
 				BEGIN_UPDATE_LOOP;
 				{
-					Particle.Velocity *= FVector3f(VelOverLife.GetValue(Particle.RelativeTime, Owner->Component) * OwnerScale);
+					Particle.Velocity *= FVector3f(VelOverLife.GetValue(Particle.RelativeTime, Context.GetDistributionData()) * OwnerScale);
 				}
 				END_UPDATE_LOOP;
 			}
@@ -338,7 +345,7 @@ void UParticleModuleVelocityOverLifetime::Update(FParticleEmitterInstance* Owner
 			{
 				BEGIN_UPDATE_LOOP;
 				{
-					Particle.Velocity *= FVector3f(VelOverLife.GetValue(Particle.RelativeTime, Owner->Component) * OwnerScale);
+					Particle.Velocity *= FVector3f(VelOverLife.GetValue(Particle.RelativeTime, Context.GetDistributionData()) * OwnerScale);
 				}
 				END_UPDATE_LOOP;
 			}
@@ -349,7 +356,7 @@ void UParticleModuleVelocityOverLifetime::Update(FParticleEmitterInstance* Owner
 				const FMatrix InvMat = LocalToWorld.InverseFast();
 				BEGIN_UPDATE_LOOP;
 				{
-					Vel = VelOverLife.GetValue(Particle.RelativeTime, Owner->Component);
+					Vel = VelOverLife.GetValue(Particle.RelativeTime, Context.GetDistributionData());
 					Particle.Velocity *= FVector4f(InvMat.TransformVector(Vel) * OwnerScale);
 				}
 				END_UPDATE_LOOP;
@@ -399,33 +406,35 @@ void UParticleModuleVelocityCone::PostEditChangeProperty(FPropertyChangedEvent&
 }
 #endif // WITH_EDITOR
 
-void UParticleModuleVelocityCone::Spawn(FParticleEmitterInstance* Owner, int32 Offset, float SpawnTime, FBaseParticle* ParticleBase)
+void UParticleModuleVelocityCone::Spawn(const FSpawnContext& Context)
 {
-	SpawnEx(Owner, Offset, SpawnTime, &GetRandomStream(Owner), ParticleBase);
+	SpawnEx(Context, &GetRandomStream(Context));
 }
 
-void UParticleModuleVelocityCone::SpawnEx(FParticleEmitterInstance* Owner, int32 Offset, float SpawnTime, struct FRandomStream* InRandomStream, FBaseParticle* ParticleBase)
+void UParticleModuleVelocityCone::SpawnEx(const FSpawnContext& Context, struct FRandomStream* InRandomStream)
 {
 	static const float TwoPI = 2.0f * UE_PI;
 	static const float ToRads = UE_PI / 180.0f;
 	static const int32 UUPerRad = 10430;
 	static const FVector DefaultDirection(0.0f, 0.0f, 1.0f);
-	
+
+	FParticleEmitterInstance* Owner = &Context.Owner;
+
 	// Calculate the owning actor's scale
 	UParticleLODLevel* LODLevel	= Owner->SpriteTemplate->GetCurrentLODLevel(Owner);
 	check(LODLevel);
 	FVector OwnerScale(1.0f);
 	if ((bApplyOwnerScale == true) && Owner->Component)
 	{
-		OwnerScale = Owner->Component->GetComponentTransform().GetScale3D();
+		OwnerScale = Context.GetTransform().GetScale3D();
 	}
 	
 	// Spawn particles
 	SPAWN_INIT
 	{
 		// Calculate the default position (prior to the Direction vector being factored in)
-		const float SpawnAngle = Angle.GetValue(Owner->EmitterTime, Owner->Component, InRandomStream);
-		const float SpawnVelocity = Velocity.GetValue(Owner->EmitterTime, Owner->Component, InRandomStream);
+		const float SpawnAngle = Angle.GetValue(Owner->EmitterTime, Context.GetDistributionData(), InRandomStream);
+		const float SpawnVelocity = Velocity.GetValue(Owner->EmitterTime, Context.GetDistributionData(), InRandomStream);
 		const float LatheAngle = InRandomStream->FRand() * TwoPI;
 		const FRotator DefaultDirectionRotater((int32)(SpawnAngle * ToRads * UUPerRad), (int32)(LatheAngle * UUPerRad), 0);
 		const FRotationMatrix DefaultDirectionRotation(DefaultDirectionRotater);
@@ -461,11 +470,11 @@ void UParticleModuleVelocityCone::SpawnEx(FParticleEmitterInstance* Owner, int32
 		// Transform according to world and local space flags 
 		if (!LODLevel->RequiredModule->bUseLocalSpace && !bInWorldSpace)
 		{
-			SpawnDirection = Owner->Component->GetComponentTransform().TransformVector(SpawnDirection);
+			SpawnDirection = Context.GetTransform().TransformVector(SpawnDirection);
 		}
 		else if (LODLevel->RequiredModule->bUseLocalSpace && bInWorldSpace)
 		{
-			SpawnDirection = Owner->Component->GetComponentTransform().InverseTransformVector(SpawnDirection);
+			SpawnDirection = Context.GetTransform().InverseTransformVector(SpawnDirection);
 		}
 
 		// Set final velocity vector
@@ -475,7 +484,7 @@ void UParticleModuleVelocityCone::SpawnEx(FParticleEmitterInstance* Owner, int32
 	}
 }
 
-void UParticleModuleVelocityCone::Render3DPreview(FParticleEmitterInstance* Owner, const FSceneView* View,FPrimitiveDrawInterface* PDI)
+void UParticleModuleVelocityCone::Render3DPreview(const FPreviewContext& Context)
 {
 #if WITH_EDITOR
 	float ConeMaxAngle = 0.0f;
@@ -487,6 +496,7 @@ void UParticleModuleVelocityCone::Render3DPreview(FParticleEmitterInstance* Owne
 	Velocity.GetOutRange(ConeMinVelocity, ConeMaxVelocity);
 
 	float MaxLifetime = 0.0f;
+	FParticleEmitterInstance* Owner = &Context.Owner;
 	TArray<TObjectPtr<UParticleModule>>& Modules = Owner->SpriteTemplate->GetCurrentLODLevel(Owner)->Modules;
 	for (int32 ModuleIndex = 0; ModuleIndex < Modules.Num(); ModuleIndex++)
 	{
@@ -538,13 +548,13 @@ void UParticleModuleVelocityCone::Render3DPreview(FParticleEmitterInstance* Owne
 		{
 			if (bApplyOwnerScale == true)
 			{
-				OwnerScale = Owner->Component->GetComponentTransform().GetScale3D();
+				OwnerScale = Context.GetTransform().GetScale3D();
 			}
 
 			OwnerRotation = FQuatRotationMatrix(Actor->GetActorQuat());
 		}
-	  LocalToWorldOrigin = Owner->Component->GetComponentTransform().GetLocation();
-	  LocalToWorld = Owner->Component->GetComponentTransform().ToMatrixWithScale().RemoveTranslation();
+	  LocalToWorldOrigin = Context.GetTransform().GetLocation();
+	  LocalToWorld = Context.GetTransform().ToMatrixWithScale().RemoveTranslation();
 	  LocalToWorld.RemoveScaling();
 	}
 	FMatrix Transform;
@@ -583,6 +593,8 @@ void UParticleModuleVelocityCone::Render3DPreview(FParticleEmitterInstance* Owne
 	TArray<FVector> OuterVerts;
 	TArray<FVector> InnerVerts;
 
+	FPrimitiveDrawInterface* PDI = Context.PDI;
+
 	// Draw inner and outer cones
 	DrawWireCone(PDI, InnerVerts, Transform, ConeRadius, ConeMinAngle, ConeSides, ModuleEditorColor, SDPG_World);
 	DrawWireCone(PDI, OuterVerts, Transform, ConeRadius, ConeMaxAngle, ConeSides, ModuleEditorColor, SDPG_World);
diff --git a/Engine/Source/Runtime/Engine/Private/Particles/ParticleTrail2EmitterInstance.cpp b/Engine/Source/Runtime/Engine/Private/Particles/ParticleTrail2EmitterInstance.cpp
--- a/Engine/Source/Runtime/Engine/Private/Particles/ParticleTrail2EmitterInstance.cpp
+++ b/Engine/Source/Runtime/Engine/Private/Particles/ParticleTrail2EmitterInstance.cpp
@@ -1201,7 +1201,7 @@ void FParticleRibbonEmitterInstance::GetParticleLifetimeAndSize(int32 InTrailIdx
 			UParticleModuleLifetime* LifetimeModule = Cast<UParticleModuleLifetime>(LODLevel->SpawnModules[ModuleIdx]);
 			if (LifetimeModule != NULL)
 			{
-				float MaxLifetime = LifetimeModule->GetLifetimeValue(this, EmitterTime, Component);
+				float MaxLifetime = LifetimeModule->GetLifetimeValue({*this}, EmitterTime, Component);
 				if (CurrLifetime > 0.f)
 				{
 					// Another module already modified lifetime.
@@ -1379,7 +1379,7 @@ float FParticleRibbonEmitterInstance::Spawn(float DeltaTime)
 			SetDeadIndex(TrailData->TrailIndex, ParticleIndex);
 			if (LODLevel->TypeDataModule)
 			{
-				LODLevel->TypeDataModule->Spawn(this, TypeDataOffset, SpawnTime, Particle);
+				LODLevel->TypeDataModule->Spawn({ *this, TypeDataOffset, SpawnTime, Particle });
 			}
 
 			for (int32 ModuleIndex = 0; ModuleIndex < LODLevel->SpawnModules.Num(); ModuleIndex++)
@@ -1388,7 +1388,7 @@ float FParticleRibbonEmitterInstance::Spawn(float DeltaTime)
 				if (SpawnModule->bEnabled)
 				{
 					UParticleModule* OffsetModule	= LODLevel->SpawnModules[ModuleIndex];
-					SpawnModule->Spawn(this, GetModuleDataOffset(OffsetModule), SpawnTime, Particle);
+					SpawnModule->Spawn({ *this, (int32)GetModuleDataOffset(OffsetModule), SpawnTime, Particle });
 				}
 			}
 			PostSpawn(Particle, 1.f - float(SpawnIdx + 1) / float(SpawnNumber), SpawnTime);
@@ -1876,13 +1876,13 @@ bool FParticleRibbonEmitterInstance::Spawn_Source(float DeltaTime)
 						continue;
 					}
 
-					SpawnModule->Spawn(this, GetModuleDataOffset(SpawnModule), SpawnTime, Particle);
+					SpawnModule->Spawn({ *this, (int32)GetModuleDataOffset(SpawnModule), SpawnTime, Particle });
 				}
 
 				if (LODLevel->TypeDataModule)
 				{
 					//@todo. Need to track TypeData offset into payload!
-					LODLevel->TypeDataModule->Spawn(this, TypeDataOffset, SpawnTime, Particle);
+					LODLevel->TypeDataModule->Spawn({ *this, TypeDataOffset, SpawnTime, Particle });
 				}
 				const float InterpolationPercentage = GSkipRibbonSpawnInterp ? 0.f : (1.f - float(SpawnIdx + 1) / float(MovementSpawnCount));
 				PostSpawn(Particle, InterpolationPercentage, SpawnTime);
@@ -2265,7 +2265,7 @@ bool FParticleRibbonEmitterInstance::ResolveSourcePoint(int32 InTrailIdx,
 						{
 							case EPSSM_Random:
 							{
-								FRandomStream& RandomStream = SourceModule->GetRandomStream(this);
+							FRandomStream& RandomStream = SourceModule->GetRandomStream({*this});
 								Index = RandomStream.RandHelper(SourceEmitter->ActiveParticles);
 							}
 								break;
@@ -3123,7 +3123,7 @@ void FParticleAnimTrailEmitterInstance::SpawnParticle( int32& StartParticleIndex
 			continue;
 		}
 
-		SpawnModule->Spawn(this, GetModuleDataOffset(SpawnModule), SpawnTime, Particle);
+		SpawnModule->Spawn({ *this, (int32)GetModuleDataOffset(SpawnModule), SpawnTime, Particle });
 	}
 
 	if ((1.0f / Particle->OneOverMaxLifetime) < 0.001f)
@@ -3134,7 +3134,7 @@ void FParticleAnimTrailEmitterInstance::SpawnParticle( int32& StartParticleIndex
 	if (LODLevel->TypeDataModule)
 	{
 		//@todo. Need to track TypeData offset into payload!
-		LODLevel->TypeDataModule->Spawn(this, TypeDataOffset, SpawnTime, Particle);
+		LODLevel->TypeDataModule->Spawn({ *this, TypeDataOffset, SpawnTime, Particle });
 	}
 	PostSpawn(Particle, 1.f, SpawnTime);
 
diff --git a/Engine/Source/Runtime/Engine/Public/ParticleHelper.h b/Engine/Source/Runtime/Engine/Public/ParticleHelper.h
--- a/Engine/Source/Runtime/Engine/Public/ParticleHelper.h
+++ b/Engine/Source/Runtime/Engine/Public/ParticleHelper.h
@@ -76,13 +76,15 @@ inline const int32 PDM_DefaultValue = 0xFFFF;
 
 #define BEGIN_UPDATE_LOOP																								\
 	{																													\
-		check((Owner != NULL) && (Owner->Component != NULL));															\
-		int32&			ActiveParticles = Owner->ActiveParticles;														\
+		check(Context.Owner.Component != NULL);																			\
+		int32&			ActiveParticles = Context.Owner.ActiveParticles;												\
+		int32			Offset	= Context.Offset;																		\
 		uint32			CurrentOffset	= Offset;																		\
-		const uint8*		ParticleData	= Owner->ParticleData;															\
-		const uint32		ParticleStride	= Owner->ParticleStride;														\
-		uint16*			ParticleIndices	= Owner->ParticleIndices;														\
-		for(int32 i=ActiveParticles-1; i>=0; i--)																			\
+		float			DeltaTime = Context.DeltaTime;																	\
+		const uint8*	ParticleData	= Context.Owner.ParticleData;													\
+		const uint32	ParticleStride	= Context.Owner.ParticleStride;													\
+		uint16*			ParticleIndices	= Context.Owner.ParticleIndices;												\
+		for(int32 i=ActiveParticles-1; i>=0; i--)																		\
 		{																												\
 			const int32	CurrentIndex	= ParticleIndices[i];															\
 			const uint8* ParticleBase	= ParticleData + CurrentIndex * ParticleStride;									\
@@ -101,10 +103,11 @@ inline const int32 PDM_DefaultValue = 0xFFFF;
 		continue;
 
 #define SPAWN_INIT																										\
-	check((Owner != NULL) && (Owner->Component != NULL));																\
-	const int32		ActiveParticles	= Owner->ActiveParticles;															\
-	const uint32		ParticleStride	= Owner->ParticleStride;															\
-	uint32			CurrentOffset	= Offset;																			\
+	check(Context.Owner.Component != NULL);																				\
+	const int32		ActiveParticles	= Context.Owner.ActiveParticles;													\
+	const uint32	ParticleStride	= Context.Owner.ParticleStride;														\
+	uint32			CurrentOffset	= Context.Offset;																	\
+	FBaseParticle*	ParticleBase	= Context.ParticleBase;																\
 	FBaseParticle&	Particle		= *(ParticleBase);
 
 #define PARTICLE_ELEMENT(Type,Name)																						\
