[Engine]
* Attempt at fixing static analysis error in ParticleModules.cpp

diff --git a/Engine/Source/Runtime/Engine/Private/Particles/ParticleModules.cpp b/Engine/Source/Runtime/Engine/Private/Particles/ParticleModules.cpp
--- a/Engine/Source/Runtime/Engine/Private/Particles/ParticleModules.cpp
+++ b/Engine/Source/Runtime/Engine/Private/Particles/ParticleModules.cpp
@@ -2853,16 +2853,16 @@ UObject* UParticleModuleAcceleration::FContext::GetDistributionData() const
 
 void UParticleModuleAcceleration::Spawn(const FSpawnContext& Context)
 {
-	FParticleEmitterInstance* Owner = &Context.Owner;
+	FParticleEmitterInstance& Owner = Context.Owner;
 	SPAWN_INIT;
 	PARTICLE_ELEMENT(FVector3f, UsedAcceleration);
-	UsedAcceleration = FVector3f(Acceleration.GetValue(Owner->EmitterTime, Context.GetDistributionData()));
-	if ((bApplyOwnerScale == true) && Owner && Owner->Component)
+	UsedAcceleration = FVector3f(Acceleration.GetValue(Owner.EmitterTime, Context.GetDistributionData()));
+	if (bApplyOwnerScale == true && Owner.Component)
 	{
 		FVector3f Scale(Context.GetTransform().GetScale3D());
 		UsedAcceleration *= Scale;
 	}
-	UParticleLODLevel* LODLevel	= Owner->SpriteTemplate->GetCurrentLODLevel(Owner);
+	UParticleLODLevel* LODLevel	= Owner.SpriteTemplate->GetCurrentLODLevel(&Owner);
 	check(LODLevel);
 	if (bAlwaysInWorldSpace && LODLevel->RequiredModule->bUseLocalSpace)
 	{
@@ -2874,7 +2874,7 @@ void UParticleModuleAcceleration::Spawn(const FSpawnContext& Context)
 	{
 		if (LODLevel->RequiredModule->bUseLocalSpace)
 		{
-			UsedAcceleration = (FVector4f)Owner->EmitterToSimulation.TransformVector(FVector(UsedAcceleration));
+			UsedAcceleration = (FVector4f)Owner.EmitterToSimulation.TransformVector(FVector(UsedAcceleration));
 		}
 		Particle.Velocity		+= UsedAcceleration * Context.SpawnTime;
 		Particle.BaseVelocity	+= UsedAcceleration * Context.SpawnTime;
