Fixed a case where shape sweeps of zero length were incorrectly trying to convert collision hit info. The distance was uninitialized, causing it to sometimes treat the hit as sweep vs. an initial hit which asserted on a zero normal.
Also fixed a typo in an ensure.

diff --git a/Engine/Source/Runtime/Engine/Private/Collision/CollisionConversions.cpp b/Engine/Source/Runtime/Engine/Private/Collision/CollisionConversions.cpp
--- a/Engine/Source/Runtime/Engine/Private/Collision/CollisionConversions.cpp
+++ b/Engine/Source/Runtime/Engine/Private/Collision/CollisionConversions.cpp
@@ -624,7 +624,7 @@ static bool ConvertOverlappedShapeToImpactHit(const UWorld* World, const THitLoc
 			{
 				if (bHasNormal)
 				{
-					ensureMsgf(false, TEXT("Unable to comptue valid normal from start/end when normal was requested. Start: %s, End: %s."), 
+					ensureMsgf(false, TEXT("Unable to compute valid normal from start/end when normal was requested. Start: %s, End: %s."), 
 						*StartLoc.ToString(), *EndLoc.ToString());
 				}
 				OutResult.ImpactNormal = FVector::UpVector;
diff --git a/Engine/Source/Runtime/Experimental/Chaos/Public/Chaos/PhysicsObjectCollisionInterface.h b/Engine/Source/Runtime/Experimental/Chaos/Public/Chaos/PhysicsObjectCollisionInterface.h
--- a/Engine/Source/Runtime/Experimental/Chaos/Public/Chaos/PhysicsObjectCollisionInterface.h
+++ b/Engine/Source/Runtime/Experimental/Chaos/Public/Chaos/PhysicsObjectCollisionInterface.h
@@ -203,6 +203,7 @@ namespace Chaos
 					ChaosInterface::TThreadQueryHit<Id>& OutQueryHit = OutBestHit;
 					OutQueryHit = OverlapQueryHit;
 					OutBestHit.WorldNormal = MTDInfo.Normal * MTDInfo.Penetration;
+					OutBestHit.Distance = -MTDInfo.Penetration;
 				}
 				return bOverlap;
 			}
