[Engine]
* Changed so ParticleEmitterInstance always require component (this will be replaced with some sort of context)
* Removed more direct access of component in ParticleModule

diff --git a/Engine/Source/Runtime/Engine/Classes/Particles/ParticleEmitter.h b/Engine/Source/Runtime/Engine/Classes/Particles/ParticleEmitter.h
--- a/Engine/Source/Runtime/Engine/Classes/Particles/ParticleEmitter.h
+++ b/Engine/Source/Runtime/Engine/Classes/Particles/ParticleEmitter.h
@@ -263,7 +263,7 @@ class UParticleEmitter : public UObject
 	//~ End UObject Interface
 
 	// @todo document
-	virtual FParticleEmitterInstance* CreateInstance(UParticleSystemComponent* InComponent);
+	virtual FParticleEmitterInstance* CreateInstance(UParticleSystemComponent& InComponent);
 
 	// Sets up this emitter with sensible defaults so we can see some particles as soon as its created.
 	virtual void SetToSensibleDefaults() {}
diff --git a/Engine/Source/Runtime/Engine/Classes/Particles/ParticleModule.h b/Engine/Source/Runtime/Engine/Classes/Particles/ParticleModule.h
--- a/Engine/Source/Runtime/Engine/Classes/Particles/ParticleModule.h
+++ b/Engine/Source/Runtime/Engine/Classes/Particles/ParticleModule.h
@@ -238,6 +238,8 @@ class UParticleModule : public UObject
 		FParticleEmitterInstance& Owner;
 		const FTransform& GetTransform() const;
 		UObject* GetDistributionData() const;
+		FString GetTemplateName() const;
+		FString GetInstanceName() const;
 		FContext(FParticleEmitterInstance& Ow) : Owner(Ow) {}
 	};
 
diff --git a/Engine/Source/Runtime/Engine/Classes/Particles/ParticleSpriteEmitter.h b/Engine/Source/Runtime/Engine/Classes/Particles/ParticleSpriteEmitter.h
--- a/Engine/Source/Runtime/Engine/Classes/Particles/ParticleSpriteEmitter.h
+++ b/Engine/Source/Runtime/Engine/Classes/Particles/ParticleSpriteEmitter.h
@@ -34,7 +34,7 @@ class UParticleSpriteEmitter : public UParticleEmitter
 	//~ End UObject Interface
 
 	//~ Begin UParticleEmitter Interface
-	virtual FParticleEmitterInstance* CreateInstance(UParticleSystemComponent* InComponent) override;
+	virtual FParticleEmitterInstance* CreateInstance(UParticleSystemComponent& InComponent) override;
 	virtual void SetToSensibleDefaults() override;
 	//~ End UParticleEmitter Interface
 };
diff --git a/Engine/Source/Runtime/Engine/Classes/Particles/TypeData/ParticleModuleTypeDataAnimTrail.h b/Engine/Source/Runtime/Engine/Classes/Particles/TypeData/ParticleModuleTypeDataAnimTrail.h
--- a/Engine/Source/Runtime/Engine/Classes/Particles/TypeData/ParticleModuleTypeDataAnimTrail.h
+++ b/Engine/Source/Runtime/Engine/Classes/Particles/TypeData/ParticleModuleTypeDataAnimTrail.h
@@ -80,7 +80,7 @@ class UParticleModuleTypeDataAnimTrail : public UParticleModuleTypeDataBase
 	//~ End UParticleModule Interface
 
 	//~ Begin UParticleModuleTypeDataBase Interface
-	virtual FParticleEmitterInstance* CreateInstance(UParticleEmitter* InEmitterParent, UParticleSystemComponent* InComponent) override;
+	virtual FParticleEmitterInstance* CreateInstance(UParticleEmitter* InEmitterParent, UParticleSystemComponent& InComponent) override;
 	virtual const FVertexFactoryType* GetVertexFactoryType() const override;
 	virtual EPrimitiveType GetPrimitiveType() const override { return PT_TriangleStrip; }
 	virtual void CollectPSOPrecacheData(const UParticleEmitter* Emitter, FPSOPrecacheParams& OutParams) override;
diff --git a/Engine/Source/Runtime/Engine/Classes/Particles/TypeData/ParticleModuleTypeDataBase.h b/Engine/Source/Runtime/Engine/Classes/Particles/TypeData/ParticleModuleTypeDataBase.h
--- a/Engine/Source/Runtime/Engine/Classes/Particles/TypeData/ParticleModuleTypeDataBase.h
+++ b/Engine/Source/Runtime/Engine/Classes/Particles/TypeData/ParticleModuleTypeDataBase.h
@@ -35,7 +35,7 @@ class UParticleModuleTypeDataBase : public UParticleModule
 	virtual bool RequiresBuild() const { return false; }
 
 	// @todo document
-	virtual FParticleEmitterInstance* CreateInstance(UParticleEmitter* InEmitterParent, UParticleSystemComponent* InComponent);
+	virtual FParticleEmitterInstance* CreateInstance(UParticleEmitter* InEmitterParent, UParticleSystemComponent& InComponent);
 
 	virtual const FVertexFactoryType* GetVertexFactoryType() const { return nullptr; }
 	virtual EPrimitiveType GetPrimitiveType() const { return PT_TriangleList; }
diff --git a/Engine/Source/Runtime/Engine/Classes/Particles/TypeData/ParticleModuleTypeDataBeam2.h b/Engine/Source/Runtime/Engine/Classes/Particles/TypeData/ParticleModuleTypeDataBeam2.h
--- a/Engine/Source/Runtime/Engine/Classes/Particles/TypeData/ParticleModuleTypeDataBeam2.h
+++ b/Engine/Source/Runtime/Engine/Classes/Particles/TypeData/ParticleModuleTypeDataBeam2.h
@@ -222,7 +222,7 @@ class UParticleModuleTypeDataBeam2 : public UParticleModuleTypeDataBase
 
 
 	//~ Begin UParticleModuleTypeDataBase Interface
-	virtual FParticleEmitterInstance* CreateInstance(UParticleEmitter* InEmitterParent, UParticleSystemComponent* InComponent) override;
+	virtual FParticleEmitterInstance* CreateInstance(UParticleEmitter* InEmitterParent, UParticleSystemComponent& InComponent) override;
 	virtual void CacheModuleInfo(UParticleEmitter* Emitter) override;
 	virtual const FVertexFactoryType* GetVertexFactoryType() const override;
 	virtual EPrimitiveType GetPrimitiveType() const { return PT_TriangleStrip; }
diff --git a/Engine/Source/Runtime/Engine/Classes/Particles/TypeData/ParticleModuleTypeDataGpu.h b/Engine/Source/Runtime/Engine/Classes/Particles/TypeData/ParticleModuleTypeDataGpu.h
--- a/Engine/Source/Runtime/Engine/Classes/Particles/TypeData/ParticleModuleTypeDataGpu.h
+++ b/Engine/Source/Runtime/Engine/Classes/Particles/TypeData/ParticleModuleTypeDataGpu.h
@@ -482,7 +482,7 @@ class UParticleModuleTypeDataGpu : public UParticleModuleTypeDataBase
 	//~ Begin UParticleModuleTypeDataBase Interface
 	virtual void Build( struct FParticleEmitterBuildInfo& EmitterBuildInfo ) override;
 	virtual bool RequiresBuild() const override { return true; }
-	virtual FParticleEmitterInstance* CreateInstance(UParticleEmitter* InEmitterParent, UParticleSystemComponent* InComponent) override;
+	virtual FParticleEmitterInstance* CreateInstance(UParticleEmitter* InEmitterParent, UParticleSystemComponent& InComponent) override;
 	virtual const FVertexFactoryType* GetVertexFactoryType() const override;
 	virtual void CollectPSOPrecacheData(const UParticleEmitter* Emitter, FPSOPrecacheParams& OutParams) override;
 	//~ End UParticleModuleTypeDataBase Interface
diff --git a/Engine/Source/Runtime/Engine/Classes/Particles/TypeData/ParticleModuleTypeDataMesh.h b/Engine/Source/Runtime/Engine/Classes/Particles/TypeData/ParticleModuleTypeDataMesh.h
--- a/Engine/Source/Runtime/Engine/Classes/Particles/TypeData/ParticleModuleTypeDataMesh.h
+++ b/Engine/Source/Runtime/Engine/Classes/Particles/TypeData/ParticleModuleTypeDataMesh.h
@@ -224,7 +224,7 @@ class UParticleModuleTypeDataMesh : public UParticleModuleTypeDataBase
 	//~ End UParticleModule Interface
 
 	//~ Begin UParticleModuleTypeDataBase Interface
-	virtual FParticleEmitterInstance* CreateInstance(UParticleEmitter* InEmitterParent, UParticleSystemComponent* InComponent) override;
+	virtual FParticleEmitterInstance* CreateInstance(UParticleEmitter* InEmitterParent, UParticleSystemComponent& InComponent) override;
 	virtual const FVertexFactoryType* GetVertexFactoryType() const override;
 	virtual void CollectPSOPrecacheData(const UParticleEmitter* Emitter, FPSOPrecacheParams& OutParams) override;
 	virtual bool	SupportsSpecificScreenAlignmentFlags() const override {	return true;	}	
diff --git a/Engine/Source/Runtime/Engine/Classes/Particles/TypeData/ParticleModuleTypeDataRibbon.h b/Engine/Source/Runtime/Engine/Classes/Particles/TypeData/ParticleModuleTypeDataRibbon.h
--- a/Engine/Source/Runtime/Engine/Classes/Particles/TypeData/ParticleModuleTypeDataRibbon.h
+++ b/Engine/Source/Runtime/Engine/Classes/Particles/TypeData/ParticleModuleTypeDataRibbon.h
@@ -177,7 +177,7 @@ class UParticleModuleTypeDataRibbon : public UParticleModuleTypeDataBase
 	//~ End UParticleModule Interface
 
 	//~ Begin UParticleModuleTypeDataBase Interface
-	virtual FParticleEmitterInstance* CreateInstance(UParticleEmitter* InEmitterParent, UParticleSystemComponent* InComponent) override;
+	virtual FParticleEmitterInstance* CreateInstance(UParticleEmitter* InEmitterParent, UParticleSystemComponent& InComponent) override;
 	virtual const FVertexFactoryType* GetVertexFactoryType() const override;
 	virtual EPrimitiveType GetPrimitiveType() const override { return PT_TriangleStrip; }
 	virtual void CollectPSOPrecacheData(const UParticleEmitter* Emitter, FPSOPrecacheParams& OutParams) override;
diff --git a/Engine/Source/Runtime/Engine/Private/Particles/FXSystemPrivate.h b/Engine/Source/Runtime/Engine/Private/Particles/FXSystemPrivate.h
--- a/Engine/Source/Runtime/Engine/Private/Particles/FXSystemPrivate.h
+++ b/Engine/Source/Runtime/Engine/Private/Particles/FXSystemPrivate.h
@@ -17,6 +17,7 @@ class FCanvas;
 class FGlobalDistanceFieldParameterData;
 class FParticleSimulationGPU;
 class FParticleSimulationResources;
+class UParticleSystemComponent;
 class UVectorFieldComponent;
 struct FGPUSpriteEmitterInfo;
 struct FParticleEmitterInstance;
@@ -204,7 +205,7 @@ public:
 	virtual void AddVectorField(UVectorFieldComponent* VectorFieldComponent) override;
 	virtual void RemoveVectorField(UVectorFieldComponent* VectorFieldComponent) override;
 	virtual void UpdateVectorField(UVectorFieldComponent* VectorFieldComponent) override;
-	FParticleEmitterInstance* CreateGPUSpriteEmitterInstance(FGPUSpriteEmitterInfo& EmitterInfo);
+	FParticleEmitterInstance* CreateGPUSpriteEmitterInstance(UParticleSystemComponent& InComponent, FGPUSpriteEmitterInfo& EmitterInfo);
 	virtual void PreInitViews(class FRDGBuilder& GraphBuilder, bool bAllowGPUParticleUpdate, const TArrayView<const FSceneViewFamily*>& ViewFamilies, const FSceneViewFamily* CurrentFamily) override;
 	virtual void PostInitViews(FRDGBuilder& GraphBuilder, TConstStridedView<FSceneView> Views, bool bAllowGPUParticleUpdate) override;
 	virtual bool UsesGlobalDistanceField() const override;
diff --git a/Engine/Source/Runtime/Engine/Private/Particles/ParticleBeam2EmitterInstance.cpp b/Engine/Source/Runtime/Engine/Private/Particles/ParticleBeam2EmitterInstance.cpp
--- a/Engine/Source/Runtime/Engine/Private/Particles/ParticleBeam2EmitterInstance.cpp
+++ b/Engine/Source/Runtime/Engine/Private/Particles/ParticleBeam2EmitterInstance.cpp
@@ -53,8 +53,8 @@ DECLARE_CYCLE_STAT(TEXT("BeamEmitterInstance Init"), STAT_BeamEmitterInstance_In
  *	Structure for beam emitter instances
  */
 /** Constructor	*/
-FParticleBeam2EmitterInstance::FParticleBeam2EmitterInstance() :
-	FParticleEmitterInstance()
+FParticleBeam2EmitterInstance::FParticleBeam2EmitterInstance(UParticleSystemComponent& InComponent) :
+	FParticleEmitterInstance(InComponent)
 	, BeamTypeData(NULL)
 	, BeamModule_Source(NULL)
 	, BeamModule_Target(NULL)
@@ -385,9 +385,9 @@ void FParticleBeam2EmitterInstance::ApplyWorldOffset(FVector InOffset, bool bWor
 }
 
 
-void FParticleBeam2EmitterInstance::InitParameters(UParticleEmitter* InTemplate, UParticleSystemComponent* InComponent)
+void FParticleBeam2EmitterInstance::InitParameters(UParticleEmitter* InTemplate)
 {
-	FParticleEmitterInstance::InitParameters(InTemplate, InComponent);
+	FParticleEmitterInstance::InitParameters(InTemplate);
 
 	UParticleLODLevel* LODLevel	= InTemplate->GetLODLevel(0);
 	check(LODLevel);
diff --git a/Engine/Source/Runtime/Engine/Private/Particles/ParticleBeamModules.cpp b/Engine/Source/Runtime/Engine/Private/Particles/ParticleBeamModules.cpp
--- a/Engine/Source/Runtime/Engine/Private/Particles/ParticleBeamModules.cpp
+++ b/Engine/Source/Runtime/Engine/Private/Particles/ParticleBeamModules.cpp
@@ -689,13 +689,13 @@ void UParticleModuleTypeDataBeam2::PostEditChangeProperty(FPropertyChangedEvent&
 }
 #endif // WITH_EDITOR
 
-FParticleEmitterInstance* UParticleModuleTypeDataBeam2::CreateInstance(UParticleEmitter* InEmitterParent, UParticleSystemComponent* InComponent)
+FParticleEmitterInstance* UParticleModuleTypeDataBeam2::CreateInstance(UParticleEmitter* InEmitterParent, UParticleSystemComponent& InComponent)
 {
 	SetToSensibleDefaults(InEmitterParent);
-	FParticleEmitterInstance* Instance = new FParticleBeam2EmitterInstance();
+	FParticleEmitterInstance* Instance = new FParticleBeam2EmitterInstance(InComponent);
 	check(Instance);
 
-	Instance->InitParameters(InEmitterParent, InComponent);
+	Instance->InitParameters(InEmitterParent);
 
 	return Instance;
 }
diff --git a/Engine/Source/Runtime/Engine/Private/Particles/ParticleComponents.cpp b/Engine/Source/Runtime/Engine/Private/Particles/ParticleComponents.cpp
--- a/Engine/Source/Runtime/Engine/Private/Particles/ParticleComponents.cpp
+++ b/Engine/Source/Runtime/Engine/Private/Particles/ParticleComponents.cpp
@@ -922,7 +922,7 @@ UParticleEmitter::UParticleEmitter(const FObjectInitializer& ObjectInitializer)
 
 }
 
-FParticleEmitterInstance* UParticleEmitter::CreateInstance(UParticleSystemComponent* InComponent)
+FParticleEmitterInstance* UParticleEmitter::CreateInstance(UParticleSystemComponent& InComponent)
 {
 	UE_LOG(LogParticles, Fatal,TEXT("UParticleEmitter::CreateInstance is pure virtual")); 
 	return NULL; 
@@ -2050,7 +2050,7 @@ void UParticleSpriteEmitter::PostLoad()
 	}
 }
 
-FParticleEmitterInstance* UParticleSpriteEmitter::CreateInstance(UParticleSystemComponent* InComponent)
+FParticleEmitterInstance* UParticleSpriteEmitter::CreateInstance(UParticleSystemComponent& InComponent)
 {
 	// If this emitter was cooked out or has no valid LOD levels don't create an instance for it.
 	if ((bCookedOut == true) || (LODLevels.Num() == 0))
@@ -2070,10 +2070,9 @@ FParticleEmitterInstance* UParticleSpriteEmitter::CreateInstance(UParticleSystem
 	}
 	else
 	{
-		check(InComponent);
-		Instance = new FParticleSpriteEmitterInstance();
+		Instance = new FParticleSpriteEmitterInstance(InComponent);
 		check(Instance);
-		Instance->InitParameters(this, InComponent);
+		Instance->InitParameters(this);
 	}
 
 	if (Instance)
@@ -5967,14 +5966,14 @@ void UParticleSystemComponent::InitParticles()
 					}
 					else
 					{
-						Instance = Emitter->CreateInstance(this);
+						Instance = Emitter->CreateInstance(*this);
 						EmitterInstances[Idx] = Instance;
 					}
 
 					if (Instance)
 					{
 						Instance->bEnabled = true;
-						Instance->InitParameters(Emitter, this);
+						Instance->InitParameters(Emitter);
 						Instance->Init();
 
 						PreferredLODLevel = FMath::Min(PreferredLODLevel, Emitter->LODLevels.Num());
diff --git a/Engine/Source/Runtime/Engine/Private/Particles/ParticleEmitterInstances.cpp b/Engine/Source/Runtime/Engine/Private/Particles/ParticleEmitterInstances.cpp
--- a/Engine/Source/Runtime/Engine/Private/Particles/ParticleEmitterInstances.cpp
+++ b/Engine/Source/Runtime/Engine/Private/Particles/ParticleEmitterInstances.cpp
@@ -391,9 +391,9 @@ FParticleEmitterBuildInfo::FParticleEmitterBuildInfo()
 const float FParticleEmitterInstance::PeakActiveParticleUpdateDelta = 0.05f;
 
 /** Constructor	*/
-FParticleEmitterInstance::FParticleEmitterInstance() :
+FParticleEmitterInstance::FParticleEmitterInstance(UParticleSystemComponent& InComponent) :
 	  SpriteTemplate(NULL)
-    , Component(NULL)
+    , Component(&InComponent)
     , CurrentLODLevel(NULL)
     , CurrentLODLevelIndex(0)
     , TypeDataOffset(0)
@@ -472,12 +472,11 @@ void FParticleEmitterInstance::PreDestructorCall()
 }
 #endif
 
-void FParticleEmitterInstance::InitParameters(UParticleEmitter* InTemplate, UParticleSystemComponent* InComponent)
+void FParticleEmitterInstance::InitParameters(UParticleEmitter* InTemplate)
 {
 	SCOPE_CYCLE_COUNTER(STAT_ParticleEmitterInstance_InitParameters);
 
 	SpriteTemplate = InTemplate;
-	Component = InComponent;
 	SetupEmitterDuration();
 }
 
@@ -3021,8 +3020,8 @@ void FParticleEmitterInstance::GetScreenAlignmentAndScale(int32& OutScreenAlign,
  *	The structure for a standard sprite emitter instance.
  */
 /** Constructor	*/
-FParticleSpriteEmitterInstance::FParticleSpriteEmitterInstance() :
-	FParticleEmitterInstance()
+FParticleSpriteEmitterInstance::FParticleSpriteEmitterInstance(UParticleSystemComponent& InComponent) :
+	FParticleEmitterInstance(InComponent)
 {
 }
 
@@ -3193,8 +3192,8 @@ UMaterialInterface* FParticleEmitterInstance::GetCurrentMaterial()
  */
 
 /** Constructor	*/
-FParticleMeshEmitterInstance::FParticleMeshEmitterInstance() :
-	  FParticleEmitterInstance()
+FParticleMeshEmitterInstance::FParticleMeshEmitterInstance(UParticleSystemComponent& InComponent) :
+	  FParticleEmitterInstance(InComponent)
 	, MeshTypeData(NULL)
 	, MeshRotationActive(false)
 	, MeshRotationOffset(0)
@@ -3203,11 +3202,11 @@ FParticleMeshEmitterInstance::FParticleMeshEmitterInstance() :
 }
 
 
-void FParticleMeshEmitterInstance::InitParameters(UParticleEmitter* InTemplate, UParticleSystemComponent* InComponent)
+void FParticleMeshEmitterInstance::InitParameters(UParticleEmitter* InTemplate)
 {
 	SCOPE_CYCLE_COUNTER(STAT_MeshEmitterInstance_InitParameters);
 
-	FParticleEmitterInstance::InitParameters(InTemplate, InComponent);
+	FParticleEmitterInstance::InitParameters(InTemplate);
 
 	// Get the type data module
 	UParticleLODLevel* LODLevel	= InTemplate->GetLODLevel(0);
diff --git a/Engine/Source/Runtime/Engine/Private/Particles/ParticleGpuSimulation.cpp b/Engine/Source/Runtime/Engine/Private/Particles/ParticleGpuSimulation.cpp
--- a/Engine/Source/Runtime/Engine/Private/Particles/ParticleGpuSimulation.cpp
+++ b/Engine/Source/Runtime/Engine/Private/Particles/ParticleGpuSimulation.cpp
@@ -3202,8 +3202,9 @@ class FGPUSpriteParticleEmitterInstance : public FParticleEmitterInstance
 public:
 
 /** Initialization constructor. */
-FGPUSpriteParticleEmitterInstance(FFXSystem* InFXSystem, FGPUSpriteEmitterInfo& InEmitterInfo)
-	: FXSystem(InFXSystem)
+FGPUSpriteParticleEmitterInstance(UParticleSystemComponent& InComponent, FFXSystem* InFXSystem, FGPUSpriteEmitterInfo& InEmitterInfo)
+	: FParticleEmitterInstance(InComponent)
+	, FXSystem(InFXSystem)
 	, EmitterInfo(InEmitterInfo)
 	, LocalVectorFieldRotation(FRotator::ZeroRotator)
 	, PointAttractorStrength(0.0f)
@@ -5236,9 +5237,9 @@ void FFXSystem::VisualizeGPUParticles(FCanvas* Canvas)
 	External interface.
 -----------------------------------------------------------------------------*/
 
-FParticleEmitterInstance* FFXSystem::CreateGPUSpriteEmitterInstance( FGPUSpriteEmitterInfo& EmitterInfo )
+FParticleEmitterInstance* FFXSystem::CreateGPUSpriteEmitterInstance(UParticleSystemComponent& InComponent, FGPUSpriteEmitterInfo& EmitterInfo )
 {
-	return new FGPUSpriteParticleEmitterInstance( this, EmitterInfo );
+	return new FGPUSpriteParticleEmitterInstance(InComponent, this, EmitterInfo );
 }
 
 /**
diff --git a/Engine/Source/Runtime/Engine/Private/Particles/ParticleModules.cpp b/Engine/Source/Runtime/Engine/Private/Particles/ParticleModules.cpp
--- a/Engine/Source/Runtime/Engine/Private/Particles/ParticleModules.cpp
+++ b/Engine/Source/Runtime/Engine/Private/Particles/ParticleModules.cpp
@@ -2851,6 +2851,20 @@ UObject* UParticleModuleAcceleration::FContext::GetDistributionData() const
 	return Owner.Component;
 }
 
+FString UParticleModuleAcceleration::FContext::GetTemplateName() const
+{
+	if (Owner.Component)
+	{
+		return GetNameSafe(Owner.Component->Template);
+	}
+	return TEXT("UNKNOWN");
+}
+
+FString UParticleModuleAcceleration::FContext::GetInstanceName() const
+{
+	return GetPathNameSafe(Owner.Component);
+}
+
 void UParticleModuleAcceleration::Spawn(const FSpawnContext& Context)
 {
 	FParticleEmitterInstance& Owner = Context.Owner;
@@ -3404,7 +3418,7 @@ UParticleModuleTypeDataBase::UParticleModuleTypeDataBase(const FObjectInitialize
 	bUpdateModule = false;
 }
 
-FParticleEmitterInstance* UParticleModuleTypeDataBase::CreateInstance(UParticleEmitter* InEmitterParent, UParticleSystemComponent* InComponent)
+FParticleEmitterInstance* UParticleModuleTypeDataBase::CreateInstance(UParticleEmitter* InEmitterParent, UParticleSystemComponent& InComponent)
 {
 	return NULL;
 }
@@ -3443,13 +3457,13 @@ UParticleModuleTypeDataMesh::UParticleModuleTypeDataMesh(const FObjectInitialize
 	LODSizeScale = 1.0f;
 }
 
-FParticleEmitterInstance* UParticleModuleTypeDataMesh::CreateInstance(UParticleEmitter* InEmitterParent, UParticleSystemComponent* InComponent)
+FParticleEmitterInstance* UParticleModuleTypeDataMesh::CreateInstance(UParticleEmitter* InEmitterParent, UParticleSystemComponent& InComponent)
 {
 	SetToSensibleDefaults(InEmitterParent);
-	FParticleEmitterInstance* Instance = new FParticleMeshEmitterInstance();
+	FParticleEmitterInstance* Instance = new FParticleMeshEmitterInstance(InComponent);
 	check(Instance);
 
-	Instance->InitParameters(InEmitterParent, InComponent);
+	Instance->InitParameters(InEmitterParent);
 
 	CreateDistribution();
 
@@ -5011,23 +5025,22 @@ void UParticleModuleTypeDataGpu::Build( FParticleEmitterBuildInfo& EmitterBuildI
 	}
 }
 
-FParticleEmitterInstance* UParticleModuleTypeDataGpu::CreateInstance(UParticleEmitter* InEmitterParent, UParticleSystemComponent* InComponent)
+FParticleEmitterInstance* UParticleModuleTypeDataGpu::CreateInstance(UParticleEmitter* InEmitterParent, UParticleSystemComponent& InComponent)
 {
-	check(InComponent);
-	UWorld* World = InComponent->GetWorld();
+	UWorld* World = InComponent.GetWorld();
 	check(World);
 	UE_LOG(LogParticles,Verbose,
 		TEXT("Create GPU Sprite Emitter @ %fs %s"), World->TimeSeconds,
-		InComponent->Template != NULL ? *InComponent->Template->GetName() : TEXT("NULL"));
+		InComponent.Template != NULL ? *InComponent.Template->GetName() : TEXT("NULL"));
 
 	FParticleEmitterInstance* Instance = NULL;
 	if (World->Scene && RHISupportsGPUParticles())
 	{
 		LLM_SCOPE(ELLMTag::Particles);
 
-		check( InComponent && InComponent->FXSystem );
-		Instance = InComponent->FXSystem->CreateGPUSpriteEmitterInstance( EmitterInfo );
-		Instance->InitParameters( InEmitterParent, InComponent );
+		check( InComponent.FXSystem );
+		Instance = InComponent.FXSystem->CreateGPUSpriteEmitterInstance(InComponent, EmitterInfo );
+		Instance->InitParameters( InEmitterParent );
 	}
 	return Instance;
 }
diff --git a/Engine/Source/Runtime/Engine/Private/Particles/ParticleModules_Collision.cpp b/Engine/Source/Runtime/Engine/Private/Particles/ParticleModules_Collision.cpp
--- a/Engine/Source/Runtime/Engine/Private/Particles/ParticleModules_Collision.cpp
+++ b/Engine/Source/Runtime/Engine/Private/Particles/ParticleModules_Collision.cpp
@@ -406,7 +406,7 @@ void UParticleModuleCollision::Update(const FUpdateContext& Context)
 				{
 					if ((Hit.Normal.IsNearlyZero() == false) && (FMath::Abs(Hit.Normal.Z) + VerticalFudgeFactor) < 1.0f)
 					{
-						//UE_LOG(LogParticles, Log, TEXT("Particle from %s had a non-vertical hit!"), *(Owner->Component->Template->GetPathName()));
+						//UE_LOG(LogParticles, Log, TEXT("Particle from %s had a non-vertical hit!"), *Context.GetTemplateName());
 						bDecrementMaxCount = false;
 					}
 				}
diff --git a/Engine/Source/Runtime/Engine/Private/Particles/ParticleModules_Location.cpp b/Engine/Source/Runtime/Engine/Private/Particles/ParticleModules_Location.cpp
--- a/Engine/Source/Runtime/Engine/Private/Particles/ParticleModules_Location.cpp
+++ b/Engine/Source/Runtime/Engine/Private/Particles/ParticleModules_Location.cpp
@@ -128,7 +128,7 @@ void UParticleModuleLocation::SpawnEx(const FSpawnContext& Context, struct FRand
 
 	LocationOffset = Owner->EmitterToSimulation.TransformVector(LocationOffset);
 	Particle.Location += LocationOffset;
-	ensureMsgf(!Particle.Location.ContainsNaN(), TEXT("NaN in Particle Location. Template: %s, Component: %s"), Owner->Component ? *GetNameSafe(Owner->Component->Template) : TEXT("UNKNOWN"), *GetPathNameSafe(Owner->Component));
+	ensureMsgf(!Particle.Location.ContainsNaN(), TEXT("NaN in Particle Location. Template: %s, Component: %s"), *Context.GetTemplateName(), *Context.GetInstanceName());
 }
 
 void UParticleModuleLocation::Render3DPreview(const FPreviewContext& Context)
@@ -229,7 +229,7 @@ void UParticleModuleLocationWorldOffset::SpawnEx(const FSpawnContext& Context, s
 		FVector StartLoc = StartLocation.GetValue(Owner->EmitterTime, Context.GetDistributionData(), 0, InRandomStream);
 		Particle.Location += InvMat.TransformVector(StartLoc);
 	}
-	ensureMsgf(!Particle.Location.ContainsNaN(), TEXT("NaN in Particle Location. Template: %s, Component: %s"), Owner->Component ? *GetNameSafe(Owner->Component->Template) : TEXT("UNKNOWN"), *GetPathNameSafe(Owner->Component));
+	ensureMsgf(!Particle.Location.ContainsNaN(), TEXT("NaN in Particle Location. Template: %s, Component: %s"), *Context.GetTemplateName(), *Context.GetInstanceName());
 }
 
 /*-----------------------------------------------------------------------------
@@ -323,7 +323,7 @@ void UParticleModuleLocationDirect::Spawn(const FSpawnContext& Context)
 	PARTICLE_ELEMENT(FVector, LocOffset);
 	LocOffset	= LocationOffset.GetValue(Owner->EmitterTime, Context.GetDistributionData());
 	Particle.Location += LocOffset;
-	ensureMsgf(!Particle.Location.ContainsNaN(), TEXT("NaN in Particle Location. Template: %s, Component: %s"), Owner->Component ? *GetNameSafe(Owner->Component->Template) : TEXT("UNKNOWN"), *GetPathNameSafe(Owner->Component));
+	ensureMsgf(!Particle.Location.ContainsNaN(), TEXT("NaN in Particle Location. Template: %s, Component: %s"), *Context.GetTemplateName(), *Context.GetInstanceName());
 }
 
 void UParticleModuleLocationDirect::Update(const FUpdateContext& Context)
@@ -357,7 +357,7 @@ void UParticleModuleLocationDirect::Update(const FUpdateContext& Context)
 		float InvDeltaTime = (DeltaTime > 0.0f) ? 1.0f / DeltaTime : 0.0f;
 		Particle.Velocity	 = (FVector3f)ScaleDiffA * InvDeltaTime;
 		Particle.Location	+= ScaleDiffB;
-		ensureMsgf(!Particle.Location.ContainsNaN(), TEXT("NaN in Particle Location. Template: %s, Component: %s"), Owner->Component ? *GetNameSafe(Owner->Component->Template) : TEXT("UNKNOWN"), *GetPathNameSafe(Owner->Component));
+		ensureMsgf(!Particle.Location.ContainsNaN(), TEXT("NaN in Particle Location. Template: %s, Component: %s"), *Context.GetTemplateName(), *Context.GetInstanceName());
 	}
 	END_UPDATE_LOOP;
 }
@@ -525,9 +525,9 @@ void UParticleModuleLocationEmitter::Spawn(const FSpawnContext& Context)
 				}
 			}
 		}
-		ensureMsgf(!Particle.Location.ContainsNaN(), TEXT("NaN in Particle Location. Template: %s, Component: %s"), Owner->Component ? *GetNameSafe(Owner->Component->Template) : TEXT("UNKNOWN"), *GetPathNameSafe(Owner->Component));
-		ensureMsgf(!Particle.Velocity.ContainsNaN(), TEXT("NaN in Particle Velocity. Template: %s, Component: %s"), Owner->Component ? *GetNameSafe(Owner->Component->Template) : TEXT("UNKNOWN"), *GetPathNameSafe(Owner->Component));
-		ensureMsgf(!Particle.BaseVelocity.ContainsNaN(), TEXT("NaN in Particle Base Velocity. Template: %s, Component: %s"), Owner->Component ? *GetNameSafe(Owner->Component->Template) : TEXT("UNKNOWN"), *GetPathNameSafe(Owner->Component));
+		ensureMsgf(!Particle.Location.ContainsNaN(), TEXT("NaN in Particle Location. Template: %s, Component: %s"), *Context.GetTemplateName(), *Context.GetInstanceName());
+		ensureMsgf(!Particle.Velocity.ContainsNaN(), TEXT("NaN in Particle Velocity. Template: %s, Component: %s"), *Context.GetTemplateName(), *Context.GetInstanceName());
+		ensureMsgf(!Particle.BaseVelocity.ContainsNaN(), TEXT("NaN in Particle Base Velocity. Template: %s, Component: %s"), *Context.GetTemplateName(), *Context.GetInstanceName());
 }
 
 uint32 UParticleModuleLocationEmitter::RequiredBytesPerInstance()
@@ -595,8 +595,8 @@ void UParticleModuleLocationEmitterDirect::Spawn(const FSpawnContext& Context)
 			Particle.OldLocation	= pkParticle->OldLocation;
 			Particle.Velocity		= pkParticle->Velocity;
 			Particle.RelativeTime	= pkParticle->RelativeTime;
-			ensureMsgf(!Particle.Location.ContainsNaN(), TEXT("NaN in Particle Location. Template: %s, Component: %s"), Owner->Component ? *GetNameSafe(Owner->Component->Template) : TEXT("UNKNOWN"), *GetPathNameSafe(Owner->Component));
-			ensureMsgf(!Particle.Velocity.ContainsNaN(), TEXT("NaN in Particle Velocity. Template: %s, Component: %s"), Owner->Component ? *GetNameSafe(Owner->Component->Template) : TEXT("UNKNOWN"), *GetPathNameSafe(Owner->Component));
+			ensureMsgf(!Particle.Location.ContainsNaN(), TEXT("NaN in Particle Location. Template: %s, Component: %s"), *Context.GetTemplateName(), *Context.GetInstanceName());
+			ensureMsgf(!Particle.Velocity.ContainsNaN(), TEXT("NaN in Particle Velocity. Template: %s, Component: %s"), *Context.GetTemplateName(), *Context.GetInstanceName());
 		}
 } 
 
@@ -637,8 +637,8 @@ void UParticleModuleLocationEmitterDirect::Update(const FUpdateContext& Context)
 				Particle.OldLocation	= pkParticle->OldLocation;
 				Particle.Velocity		= pkParticle->Velocity;
 				Particle.RelativeTime	= pkParticle->RelativeTime;
-				ensureMsgf(!Particle.Location.ContainsNaN(), TEXT("NaN in Particle Location. Template: %s, Component: %s"), Owner->Component ? *GetNameSafe(Owner->Component->Template) : TEXT("UNKNOWN"), *GetPathNameSafe(Owner->Component));
-				ensureMsgf(!Particle.Velocity.ContainsNaN(), TEXT("NaN in Particle Velocity. Template: %s, Component: %s"), Owner->Component ? *GetNameSafe(Owner->Component->Template) : TEXT("UNKNOWN"), *GetPathNameSafe(Owner->Component));
+				ensureMsgf(!Particle.Location.ContainsNaN(), TEXT("NaN in Particle Location. Template: %s, Component: %s"), *Context.GetTemplateName(), *Context.GetInstanceName());
+				ensureMsgf(!Particle.Velocity.ContainsNaN(), TEXT("NaN in Particle Velocity. Template: %s, Component: %s"), *Context.GetTemplateName(), *Context.GetInstanceName());
 			}
 		}
 	END_UPDATE_LOOP;
@@ -857,7 +857,7 @@ void UParticleModuleLocationPrimitiveTriangle::SpawnEx(const FSpawnContext& Cont
 	LocationOffset = Owner->EmitterToSimulation.TransformVector(LocationOffset);
 
 	Particle.Location += LocationOffset;
-	ensureMsgf(!Particle.Location.ContainsNaN(), TEXT("NaN in Particle Location. Template: %s, Component: %s"), Owner->Component ? *GetNameSafe(Owner->Component->Template) : TEXT("UNKNOWN"), *GetPathNameSafe(Owner->Component));
+	ensureMsgf(!Particle.Location.ContainsNaN(), TEXT("NaN in Particle Location. Template: %s, Component: %s"), *Context.GetTemplateName(), *Context.GetInstanceName());
 }
 
 void UParticleModuleLocationPrimitiveTriangle::Render3DPreview(const FPreviewContext& Context)
@@ -1102,8 +1102,8 @@ void UParticleModuleLocationPrimitiveCylinder::SpawnEx(const FSpawnContext& Cont
 		Particle.Velocity		+= (FVector3f)vVelocity;
 		Particle.BaseVelocity	+= (FVector3f)vVelocity;
 	}
-	ensureMsgf(!Particle.Location.ContainsNaN(), TEXT("NaN in Particle Location. Template: %s, Component: %s"), Owner->Component ? *GetNameSafe(Owner->Component->Template) : TEXT("UNKNOWN"), *GetPathNameSafe(Owner->Component));
-	ensureMsgf(!Particle.Velocity.ContainsNaN(), TEXT("NaN in Particle Velocity. Template: %s, Component: %s"), Owner->Component ? *GetNameSafe(Owner->Component->Template) : TEXT("UNKNOWN"), *GetPathNameSafe(Owner->Component));
+	ensureMsgf(!Particle.Location.ContainsNaN(), TEXT("NaN in Particle Location. Template: %s, Component: %s"), *Context.GetTemplateName(), *Context.GetInstanceName());
+	ensureMsgf(!Particle.Velocity.ContainsNaN(), TEXT("NaN in Particle Velocity. Template: %s, Component: %s"), *Context.GetTemplateName(), *Context.GetInstanceName());
 }
 
 void UParticleModuleLocationPrimitiveCylinder::Render3DPreview(const FPreviewContext& Context)
@@ -1332,8 +1332,8 @@ void UParticleModuleLocationPrimitiveSphere::SpawnEx(const FSpawnContext& Contex
 		Particle.Velocity		+= (FVector3f)vVelocity;
 		Particle.BaseVelocity	+= (FVector3f)vVelocity;
 	}
-	ensureMsgf(!Particle.Location.ContainsNaN(), TEXT("NaN in Particle Location. Template: %s, Component: %s"), Owner->Component ? *GetNameSafe(Owner->Component->Template) : TEXT("UNKNOWN"), *GetPathNameSafe(Owner->Component));
-	ensureMsgf(!Particle.Velocity.ContainsNaN(), TEXT("NaN in Particle Velocity. Template: %s, Component: %s"), Owner->Component ? *GetNameSafe(Owner->Component->Template) : TEXT("UNKNOWN"), *GetPathNameSafe(Owner->Component));
+	ensureMsgf(!Particle.Location.ContainsNaN(), TEXT("NaN in Particle Location. Template: %s, Component: %s"), *Context.GetTemplateName(), *Context.GetInstanceName());
+	ensureMsgf(!Particle.Velocity.ContainsNaN(), TEXT("NaN in Particle Velocity. Template: %s, Component: %s"), *Context.GetTemplateName(), *Context.GetInstanceName());
 }
 
 void UParticleModuleLocationPrimitiveSphere::Render3DPreview(const FPreviewContext& Context)
@@ -1615,12 +1615,12 @@ void UParticleModuleLocationBoneSocket::Spawn(const FSpawnContext& Context)
 			FModuleLocationBoneSocketParticlePayload* ParticlePayload = (FModuleLocationBoneSocketParticlePayload*)((uint8*)&Particle + Context.Offset);
 			ParticlePayload->SourceIndex = SourceIndex;
 			Particle.Location = SourceLocation;
-			ensureMsgf(!Particle.Location.ContainsNaN(), TEXT("NaN in Particle Location. Template: %s, Component: %s"), Owner->Component ? *GetNameSafe(Owner->Component->Template) : TEXT("UNKNOWN"), *GetPathNameSafe(Owner->Component));
+			ensureMsgf(!Particle.Location.ContainsNaN(), TEXT("NaN in Particle Location. Template: %s, Component: %s"), *Context.GetTemplateName(), *Context.GetInstanceName());
 			if (InheritingBoneVelocity())
 			{
 				// Set the base velocity for this particle.
 				Particle.BaseVelocity = FMath::Lerp<FVector3f>(Particle.BaseVelocity, InstancePayload->BoneSocketVelocities[SourceIndex], InheritVelocityScale);
-				ensureMsgf(!Particle.BaseVelocity.ContainsNaN(), TEXT("NaN in Particle Base Velocity. Template: %s, Component: %s"), Owner->Component ? *GetNameSafe(Owner->Component->Template) : TEXT("UNKNOWN"), *GetPathNameSafe(Owner->Component));
+				ensureMsgf(!Particle.BaseVelocity.ContainsNaN(), TEXT("NaN in Particle Base Velocity. Template: %s, Component: %s"), *Context.GetTemplateName(), *Context.GetInstanceName());
 			}
 			if (bMeshRotationActive) // Note that right now the rotation wil *always* be Identity (see comments above)
 			{
@@ -1703,7 +1703,7 @@ void UParticleModuleLocationBoneSocket::Update(const FUpdateContext& Context)
 		if (GetParticleLocation(InstancePayload, Owner, SourceComponent, ParticlePayload->SourceIndex, SourceLocation, SourceRotation) == true)
 		{
 			Particle.Location = SourceLocation;
-			ensureMsgf(!Particle.Location.ContainsNaN(), TEXT("NaN in Particle Location. Template: %s, Component: %s"), Owner->Component ? *GetNameSafe(Owner->Component->Template) : TEXT("UNKNOWN"), *GetPathNameSafe(Owner->Component));
+			ensureMsgf(!Particle.Location.ContainsNaN(), TEXT("NaN in Particle Location. Template: %s, Component: %s"), *Context.GetTemplateName(), *Context.GetInstanceName());
 			if (bMeshRotationActive) // Note that right now due to logic above, the rotation will always be identity
 			{
 				FMeshRotationPayloadData* PayloadData = (FMeshRotationPayloadData*)((uint8*)&Particle + MeshRotationOffset);
@@ -2380,7 +2380,7 @@ void UParticleModuleLocationSkelVertSurface::Spawn(const FSpawnContext& Context)
 			FModuleLocationVertSurfaceParticlePayload* ParticlePayload = (FModuleLocationVertSurfaceParticlePayload*)((uint8*)&Particle + Context.Offset);
 			ParticlePayload->SourceIndex = SourceIndex;
 			Particle.Location = SourceLocation;
-			ensureMsgf(!Particle.Location.ContainsNaN(), TEXT("NaN in Particle Location. Template: %s, Component: %s"), Owner->Component ? *GetNameSafe(Owner->Component->Template) : TEXT("UNKNOWN"), *GetPathNameSafe(Owner->Component));
+			ensureMsgf(!Particle.Location.ContainsNaN(), TEXT("NaN in Particle Location. Template: %s, Component: %s"), *Context.GetTemplateName(), *Context.GetInstanceName());
 
 			// Set the base velocity
 			if(bInheritBoneVelocity && ActiveBoneIndex != INDEX_NONE)
@@ -2389,7 +2389,7 @@ void UParticleModuleLocationSkelVertSurface::Spawn(const FSpawnContext& Context)
 				if(VelocityIndex != INDEX_NONE)
 				{
 					Particle.BaseVelocity = FMath::Lerp(Particle.BaseVelocity, InstancePayload->BoneVelocities[VelocityIndex], InheritVelocityScale);
-					ensureMsgf(!Particle.BaseVelocity.ContainsNaN(), TEXT("NaN in Particle Base Velocity. Template: %s, Component: %s"), Owner->Component ? *GetNameSafe(Owner->Component->Template) : TEXT("UNKNOWN"), *GetPathNameSafe(Owner->Component));
+					ensureMsgf(!Particle.BaseVelocity.ContainsNaN(), TEXT("NaN in Particle Base Velocity. Template: %s, Component: %s"), *Context.GetTemplateName(), *Context.GetInstanceName());
 				}
 			}
 			
@@ -2544,7 +2544,7 @@ void UParticleModuleLocationSkelVertSurface::Update(const FUpdateContext& Contex
 		if (GetParticleLocation(Context, InstancePayload, SourceComponent, ParticlePayload->SourceIndex, SourceLocation, SourceRotation) == true)
 		{
 			Particle.Location = SourceLocation;
-			ensureMsgf(!Particle.Location.ContainsNaN(), TEXT("NaN in Particle Location. Template: %s, Component: %s"), Owner->Component ? *GetNameSafe(Owner->Component->Template) : TEXT("UNKNOWN"), *GetPathNameSafe(Owner->Component));
+			ensureMsgf(!Particle.Location.ContainsNaN(), TEXT("NaN in Particle Location. Template: %s, Component: %s"), *Context.GetTemplateName(), *Context.GetInstanceName());
 			if (bMeshRotationActive)
 			{
 				if (bOrientMeshEmitters)
diff --git a/Engine/Source/Runtime/Engine/Private/Particles/ParticleTrail2EmitterInstance.cpp b/Engine/Source/Runtime/Engine/Private/Particles/ParticleTrail2EmitterInstance.cpp
--- a/Engine/Source/Runtime/Engine/Private/Particles/ParticleTrail2EmitterInstance.cpp
+++ b/Engine/Source/Runtime/Engine/Private/Particles/ParticleTrail2EmitterInstance.cpp
@@ -86,9 +86,9 @@ void FParticleTrailsEmitterInstance_Base::Init()
 	SetupTrailModules();
 }
 
-void FParticleTrailsEmitterInstance_Base::InitParameters(UParticleEmitter* InTemplate, UParticleSystemComponent* InComponent)
+void FParticleTrailsEmitterInstance_Base::InitParameters(UParticleEmitter* InTemplate)
 {
-	FParticleEmitterInstance::InitParameters(InTemplate, InComponent);
+	FParticleEmitterInstance::InitParameters(InTemplate);
 	if (GIsEditor)
 	{
 		UParticleLODLevel* LODLevel	= InTemplate->GetLODLevel(0);
@@ -778,8 +778,8 @@ UMaterialInterface* FParticleTrailsEmitterInstance_Base::GetCurrentMaterial()
 	FParticleRibbonEmitterInstance.
 -----------------------------------------------------------------------------*/
 /** Constructor	*/
-FParticleRibbonEmitterInstance::FParticleRibbonEmitterInstance() :
-	  FParticleTrailsEmitterInstance_Base()
+FParticleRibbonEmitterInstance::FParticleRibbonEmitterInstance(UParticleSystemComponent& InComponent) :
+	  FParticleTrailsEmitterInstance_Base(InComponent)
 	, TrailTypeData(NULL)
 	, SpawnPerUnitModule(NULL)
 	, SourceModule(NULL)
@@ -813,9 +813,9 @@ FParticleRibbonEmitterInstance::~FParticleRibbonEmitterInstance()
 {
 }
 
-void FParticleRibbonEmitterInstance::InitParameters(UParticleEmitter* InTemplate, UParticleSystemComponent* InComponent)
+void FParticleRibbonEmitterInstance::InitParameters(UParticleEmitter* InTemplate)
 {
-	FParticleTrailsEmitterInstance_Base::InitParameters(InTemplate, InComponent);
+	FParticleTrailsEmitterInstance_Base::InitParameters(InTemplate);
 
 	// We don't support LOD on trails
 	UParticleLODLevel* LODLevel	= InTemplate->GetLODLevel(0);
@@ -2886,8 +2886,8 @@ void FParticleRibbonEmitterInstance::ApplyWorldOffset(FVector InOffset, bool bWo
 //	FParticleAnimTrailEmitterInstance
 //
 /** Constructor	*/
-FParticleAnimTrailEmitterInstance::FParticleAnimTrailEmitterInstance() :
-	  FParticleTrailsEmitterInstance_Base()
+FParticleAnimTrailEmitterInstance::FParticleAnimTrailEmitterInstance(UParticleSystemComponent& InComponent) :
+	  FParticleTrailsEmitterInstance_Base(InComponent)
 	, TrailTypeData(NULL)
 	, SpawnPerUnitModule(NULL)
 	, FirstSocketName(NAME_None)
@@ -2912,9 +2912,9 @@ FParticleAnimTrailEmitterInstance::~FParticleAnimTrailEmitterInstance()
 {
 }
 
-void FParticleAnimTrailEmitterInstance::InitParameters(UParticleEmitter* InTemplate, UParticleSystemComponent* InComponent)
+void FParticleAnimTrailEmitterInstance::InitParameters(UParticleEmitter* InTemplate)
 {
-	FParticleTrailsEmitterInstance_Base::InitParameters(InTemplate, InComponent);
+	FParticleTrailsEmitterInstance_Base::InitParameters(InTemplate);
 
 	// We don't support LOD on trails
 	UParticleLODLevel* LODLevel	= InTemplate->GetLODLevel(0);
diff --git a/Engine/Source/Runtime/Engine/Private/Particles/ParticleTrailModules.cpp b/Engine/Source/Runtime/Engine/Private/Particles/ParticleTrailModules.cpp
--- a/Engine/Source/Runtime/Engine/Private/Particles/ParticleTrailModules.cpp
+++ b/Engine/Source/Runtime/Engine/Private/Particles/ParticleTrailModules.cpp
@@ -232,11 +232,11 @@ void UParticleModuleTypeDataRibbon::PostEditChangeProperty(FPropertyChangedEvent
 }
 #endif // WITH_EDITOR
 
-FParticleEmitterInstance* UParticleModuleTypeDataRibbon::CreateInstance(UParticleEmitter* InEmitterParent, UParticleSystemComponent* InComponent)
+FParticleEmitterInstance* UParticleModuleTypeDataRibbon::CreateInstance(UParticleEmitter* InEmitterParent, UParticleSystemComponent& InComponent)
 {
-	FParticleEmitterInstance* Instance = new FParticleRibbonEmitterInstance();
+	FParticleEmitterInstance* Instance = new FParticleRibbonEmitterInstance(InComponent);
 	check(Instance);
-	Instance->InitParameters(InEmitterParent, InComponent);
+	Instance->InitParameters(InEmitterParent);
 	return Instance;
 }
 
@@ -274,11 +274,11 @@ uint32 UParticleModuleTypeDataAnimTrail::RequiredBytes(UParticleModuleTypeDataBa
 	return sizeof(FAnimTrailTypeDataPayload);
 }
 
-FParticleEmitterInstance* UParticleModuleTypeDataAnimTrail::CreateInstance(UParticleEmitter* InEmitterParent, UParticleSystemComponent* InComponent)
+FParticleEmitterInstance* UParticleModuleTypeDataAnimTrail::CreateInstance(UParticleEmitter* InEmitterParent, UParticleSystemComponent& InComponent)
 {
-	FParticleEmitterInstance* Instance = new FParticleAnimTrailEmitterInstance();
+	FParticleEmitterInstance* Instance = new FParticleAnimTrailEmitterInstance(InComponent);
 	check(Instance);
-	Instance->InitParameters(InEmitterParent, InComponent);
+	Instance->InitParameters(InEmitterParent);
 	return Instance;
 }
 
diff --git a/Engine/Source/Runtime/Engine/Public/ParticleEmitterInstances.h b/Engine/Source/Runtime/Engine/Public/ParticleEmitterInstances.h
--- a/Engine/Source/Runtime/Engine/Public/ParticleEmitterInstances.h
+++ b/Engine/Source/Runtime/Engine/Public/ParticleEmitterInstances.h
@@ -357,7 +357,7 @@ public:
 	TArray<class UPointLightComponent*> HighQualityLights;
 
 	/** Constructor	*/
-	ENGINE_API FParticleEmitterInstance();
+	ENGINE_API FParticleEmitterInstance(UParticleSystemComponent& InComponent);
 
 	/** Destructor	*/
 	ENGINE_API virtual ~FParticleEmitterInstance();
@@ -367,7 +367,7 @@ public:
 #endif
 
 	//
-	ENGINE_API virtual void InitParameters(UParticleEmitter* InTemplate, UParticleSystemComponent* InComponent);
+	ENGINE_API virtual void InitParameters(UParticleEmitter* InTemplate);
 	ENGINE_API virtual void Init();
 
 	/** @return The world that the component that owns this instance is in */
@@ -835,7 +835,7 @@ struct FScopeCycleCounterEmitter
 struct FParticleSpriteEmitterInstance : public FParticleEmitterInstance
 {
 	/** Constructor	*/
-	FParticleSpriteEmitterInstance();
+	FParticleSpriteEmitterInstance(UParticleSystemComponent& InComponent);
 
 	/** Destructor	*/
 	virtual ~FParticleSpriteEmitterInstance();
@@ -895,9 +895,9 @@ struct FParticleMeshEmitterInstance : public FParticleEmitterInstance
 	TArray<UMaterialInterface*> CurrentMaterials;
 
 	/** Constructor	*/
-	ENGINE_API FParticleMeshEmitterInstance();
+	ENGINE_API FParticleMeshEmitterInstance(UParticleSystemComponent& InComponent);
 
-	ENGINE_API virtual void InitParameters(UParticleEmitter* InTemplate, UParticleSystemComponent* InComponent) override;
+	ENGINE_API virtual void InitParameters(UParticleEmitter* InTemplate) override;
 	ENGINE_API virtual void Init() override;
 	ENGINE_API virtual bool Resize(int32 NewMaxActiveParticles, bool bSetMaxActiveCount = true) override;
 	ENGINE_API virtual void Tick(float DeltaTime, bool bSuppressSpawning) override;
@@ -1037,7 +1037,7 @@ struct FParticleBeam2EmitterInstance : public FParticleEmitterInstance
 	TArray<int32>						BeamTrianglesPerSheet;
 
 	/** Constructor	*/
-	FParticleBeam2EmitterInstance();
+	FParticleBeam2EmitterInstance(UParticleSystemComponent& InComponent);
 
 	/** Destructor	*/
 	virtual ~FParticleBeam2EmitterInstance();
@@ -1060,7 +1060,7 @@ struct FParticleBeam2EmitterInstance : public FParticleEmitterInstance
 	virtual bool GetBeamTargetTangent(int32 TargetIndex, FVector& OutTangentPoint) const override;
 	virtual bool GetBeamTargetStrength(int32 TargetIndex, float& OutTargetStrength) const override;
 	//
-	virtual void InitParameters(UParticleEmitter* InTemplate, UParticleSystemComponent* InComponent) override;
+	virtual void InitParameters(UParticleEmitter* InTemplate) override;
 	virtual void Init() override;
 	virtual void Tick(float DeltaTime, bool bSuppressSpawning) override;
 	/**
@@ -1385,8 +1385,8 @@ void CheckAllIndices(){}
 #endif
 
 	/** Constructor	*/
-	FParticleTrailsEmitterInstance_Base() :
-		  FParticleEmitterInstance()
+	FParticleTrailsEmitterInstance_Base(UParticleSystemComponent& InComponent) :
+		  FParticleEmitterInstance(InComponent)
 		, VertexCount(0)
 		, TriangleCount(0)
 		, TrailCount(0)
@@ -1417,7 +1417,7 @@ void CheckAllIndices(){}
 	}
 
 	virtual void Init() override;
-	virtual void InitParameters(UParticleEmitter* InTemplate, UParticleSystemComponent* InComponent) override;
+	virtual void InitParameters(UParticleEmitter* InTemplate) override;
 	virtual void Tick(float DeltaTime, bool bSuppressSpawning) override;
 
 	bool AddParticleHelper(int32 InTrailIdx,
@@ -1564,12 +1564,12 @@ struct FParticleRibbonEmitterInstance : public FParticleTrailsEmitterInstance_Ba
 	int32 HeadOnlyParticles;
 
 	/** Constructor	*/
-	FParticleRibbonEmitterInstance();
+	FParticleRibbonEmitterInstance(UParticleSystemComponent& InComponent);
 
 	/** Destructor	*/
 	virtual ~FParticleRibbonEmitterInstance();
 
-	virtual void InitParameters(UParticleEmitter* InTemplate, UParticleSystemComponent* InComponent) override;
+	virtual void InitParameters(UParticleEmitter* InTemplate) override;
 
 	/**
 	 *	Tick sub-function that handles recalculation of tangents
@@ -1746,12 +1746,12 @@ struct FParticleAnimTrailEmitterInstance : public FParticleTrailsEmitterInstance
 	int32 HeadOnlyParticles;
 
 	/** Constructor	*/
-	FParticleAnimTrailEmitterInstance();
+	FParticleAnimTrailEmitterInstance(UParticleSystemComponent& InComponent);
 
 	/** Destructor	*/
 	virtual ~FParticleAnimTrailEmitterInstance();
 
-	virtual void InitParameters(UParticleEmitter* InTemplate, UParticleSystemComponent* InComponent) override;
+	virtual void InitParameters(UParticleEmitter* InTemplate) override;
 
 	/**
 	 *	Helper function for recalculating tangents and the spline interpolation parameter...
