Add an unchecked variant for subtypes to test the skeletal mesh component association on a UAnimInstance object.

diff --git a/Engine/Source/Runtime/Engine/Classes/Animation/AnimInstance.h b/Engine/Source/Runtime/Engine/Classes/Animation/AnimInstance.h
--- a/Engine/Source/Runtime/Engine/Classes/Animation/AnimInstance.h
+++ b/Engine/Source/Runtime/Engine/Classes/Animation/AnimInstance.h
@@ -1479,8 +1479,11 @@ public:
 	UE_DEPRECATED(5.3, "Please use RecalcRequiredCurves that takes a FCurveFilterSettings.")
 	ENGINE_API void RecalcRequiredCurves(const FCurveEvaluationOption& CurveEvalOption);
 
-	// @todo document
-	ENGINE_API USkeletalMeshComponent* GetSkelMeshComponent() const;
+	/** Returns the associated skeletal mesh component. This will always be non-NULL. */
+	FORCEINLINE USkeletalMeshComponent* GetSkelMeshComponent() const
+	{
+		return GetSkelMeshComponentChecked();
+	}
 
 	ENGINE_API virtual UWorld* GetWorld() const override;
 
@@ -1522,6 +1525,12 @@ protected:
 	/** Implementable custom function to handle notifies */
 	ENGINE_API virtual bool HandleNotify(const FAnimNotifyEvent& AnimNotifyEvent);
 
+	/** Returns the associated skeletal mesh component. Will assert if invalid. */
+	ENGINE_API USkeletalMeshComponent* GetSkelMeshComponentChecked() const;
+
+	/** Returns the associated skeletal mesh component, or NULL if invalid/unset. */
+	ENGINE_API USkeletalMeshComponent* GetSkelMeshComponentUnchecked() const;
+
 	// Root Motion
 public:
 	/** Get current RootMotion FAnimMontageInstance if any. NULL otherwise. */
diff --git a/Engine/Source/Runtime/Engine/Private/Animation/AnimInstance.cpp b/Engine/Source/Runtime/Engine/Private/Animation/AnimInstance.cpp
--- a/Engine/Source/Runtime/Engine/Private/Animation/AnimInstance.cpp
+++ b/Engine/Source/Runtime/Engine/Private/Animation/AnimInstance.cpp
@@ -1398,11 +1398,16 @@ void UAnimInstance::RecalcRequiredCurves(const FCurveEvaluationOption& CurveEval
 	PRAGMA_ENABLE_DEPRECATION_WARNINGS
 }
 
-USkeletalMeshComponent* UAnimInstance::GetSkelMeshComponent() const
+USkeletalMeshComponent* UAnimInstance::GetSkelMeshComponentChecked() const
 {
 	return CastChecked<USkeletalMeshComponent>(GetOuter());
 }
 
+USkeletalMeshComponent* UAnimInstance::GetSkelMeshComponentUnchecked() const
+{
+	return Cast<USkeletalMeshComponent>(GetOuter());
+}
+
 void UAnimInstance::Serialize(FArchive& Ar)
 {
 	Super::Serialize(Ar);
