Clean up COLLAPSED_LOGF macro.

We can use a StringBuilder to implement the string formatting instead of copy-pasting the GROWABLE_LOGF macro.

diff --git a/Engine/Source/Runtime/Engine/Private/VisualLogger/VisualLogger.cpp b/Engine/Source/Runtime/Engine/Private/VisualLogger/VisualLogger.cpp
--- a/Engine/Source/Runtime/Engine/Private/VisualLogger/VisualLogger.cpp
+++ b/Engine/Source/Runtime/Engine/Private/VisualLogger/VisualLogger.cpp
@@ -2,6 +2,7 @@
 
 #include "VisualLogger/VisualLogger.h"
 #include "EngineStats.h"
+#include "Containers/StringFwd.h"
 #include "Misc/CommandLine.h"
 #include "Modules/ModuleManager.h"
 #include "AI/NavDataGenerator.h" // IWYU pragma: keep
@@ -42,41 +43,24 @@ namespace UE::VisLog::Private
 
 DEFINE_STAT(STAT_VisualLog);
 
-// Unfortunately needs to be a #define since it uses GET_TYPED_VARARGS_RESULT which uses the va_list stuff which operates on the
-// current function, so we can't easily call a function
-#define COLLAPSED_LOGF(SerializeFunc) \
-	SCOPE_CYCLE_COUNTER(STAT_VisualLog); \
-	UWorld *World = nullptr; \
-	FVisualLogEntry *CurrentEntry = nullptr; \
-	if (CheckVisualLogInputInternal(Object, CategoryName, Verbosity, &World, &CurrentEntry) == false) \
-	{ \
-		return; \
-	} \
-	int32	BufferSize	= 1024; \
-	TCHAR*	Buffer		= nullptr; \
-	int32	Result		= -1; \
-	/* allocate some stack space to use on the first pass, which matches most strings */ \
-	TCHAR	StackBuffer[512]; \
-	TCHAR*	AllocatedBuffer = nullptr; \
-	\
-	/* first, try using the stack buffer */ \
-	Buffer = StackBuffer; \
-	GET_TYPED_VARARGS_RESULT( TCHAR, Buffer, UE_ARRAY_COUNT(StackBuffer), UE_ARRAY_COUNT(StackBuffer) - 1, Fmt, Fmt, Result ); \
-	\
-	/* if that fails, then use heap allocation to make enough space */ \
-			while(Result == -1) \
-						{ \
-		FMemory::SystemFree(AllocatedBuffer); \
-		/* We need to use malloc here directly as GMalloc might not be safe. */ \
-		Buffer = AllocatedBuffer = (TCHAR*) FMemory::SystemMalloc( BufferSize * sizeof(TCHAR) ); \
-		GET_TYPED_VARARGS_RESULT( TCHAR, Buffer, BufferSize, BufferSize-1, Fmt, Fmt, Result ); \
-		BufferSize *= 2; \
-						}; \
-	Buffer[Result] = 0; \
-	; \
-	\
-	SerializeFunc; \
-	FMemory::SystemFree(AllocatedBuffer);
+template <typename SerializeFuncType>
+static void CollapsedLogfV(const SerializeFuncType& SerializeFunc, const UObject* Object, const FName& CategoryName, ELogVerbosity::Type Verbosity, const TCHAR* Fmt, va_list Args)
+{
+	SCOPE_CYCLE_COUNTER(STAT_VisualLog);
+
+	FVisualLogEntry* CurrentEntry = nullptr;
+	if (FVisualLogger::CheckVisualLogInputInternal(Object, CategoryName, Verbosity, /*OutWorld=*/ nullptr, &CurrentEntry))
+	{
+		SerializeFunc(CurrentEntry, TStringBuilder<512>().AppendV(Fmt, Args).ToString());
+	}
+}
+
+#define COLLAPSED_LOGF(SerializeSnippet) \
+	va_list Args; \
+	va_start(Args, Fmt); \
+	auto SerializeFunc = [&](FVisualLogEntry* CurrentEntry, const TCHAR* Buffer) { SerializeSnippet; }; \
+	CollapsedLogfV(SerializeFunc, Object, CategoryName, Verbosity, Fmt, Args); \
+	va_end(Args); 
 
 void FVisualLogger::CategorizedLogfImpl(const UObject* Object, const FName& CategoryName, ELogVerbosity::Type Verbosity, const TCHAR* Fmt, ...)
 {
