NotifyActiveSoundOwner introduced, so subsystems can observe SetOwner calls.

#rnx

diff --git a/Engine/Source/Runtime/Engine/Private/AudioDevice.cpp b/Engine/Source/Runtime/Engine/Private/AudioDevice.cpp
--- a/Engine/Source/Runtime/Engine/Private/AudioDevice.cpp
+++ b/Engine/Source/Runtime/Engine/Private/AudioDevice.cpp
@@ -133,6 +133,14 @@ FAutoConsoleVariableRef CVarDisableLegacyReverb(
 	TEXT("0: Enabled, 1: Disabled"),
 	ECVF_Default);
 
+static int32 DisableOwnerNotify = 0;
+FAutoConsoleVariableRef CVarDisableOwnerNotify(
+	TEXT("au.DisableOwnerNotify"),
+	DisableOwnerNotify,
+	TEXT("Disables SetOwner notifications.\n")
+	TEXT("0: Enabled, 1: Disabled"),
+	ECVF_Default);
+
 static float SoundDistanceOptimizationLengthCVar = 1.0f;
 FAutoConsoleVariableRef CVarSoundDistanceOptimizationLength(
 	TEXT("au.SoundDistanceOptimizationLength"),
@@ -3660,6 +3668,15 @@ void FAudioDevice::NotifySubsystemsVirtualizedSoundDeleting(FActiveSound& Active
 	});
 }
 
+void FAudioDevice::NotifySubsystemsSetOwner(FActiveSound& ActiveSound, const UObject* Owner) const
+{
+	SubsystemCollection.ForEachSubsystem<IActiveSoundUpdateInterface>([&ActiveSound, Owner](IActiveSoundUpdateInterface* ActiveSoundUpdate)
+	{
+		ActiveSoundUpdate->NotifyActiveSoundOwner(ActiveSound, Owner);
+		return true;
+	});
+}
+
 void FAudioDevice::SetBaseSoundMix(USoundMix* NewMix)
 {
 	if (!IsInAudioThread())
@@ -6481,6 +6498,10 @@ void FAudioDevice::PlaySoundAtLocation(USoundBase* Sound, UWorld* World, float V
 		NewActiveSound.Priority = Sound->Priority;
 
 		NewActiveSound.SetOwner(OwningActor);
+		if (!DisableOwnerNotify)
+		{
+			NotifySubsystemsSetOwner(NewActiveSound, OwningActor);
+		}
 
 		AddNewActiveSound(NewActiveSound, Params);
 	}
diff --git a/Engine/Source/Runtime/Engine/Public/ActiveSoundUpdateInterface.h b/Engine/Source/Runtime/Engine/Public/ActiveSoundUpdateInterface.h
--- a/Engine/Source/Runtime/Engine/Public/ActiveSoundUpdateInterface.h
+++ b/Engine/Source/Runtime/Engine/Public/ActiveSoundUpdateInterface.h
@@ -40,6 +40,17 @@ public:
 	 */
 	virtual void ApplyInteriorSettings(const FActiveSound& ActiveSound, FSoundParseParameters& ParseParams) {}
 
+	/**
+	 * Called while an active sound is being constructed. Will be followed by either
+	 * NotifyActiveSoundCreated or NotifyVirtualizedSoundCreated.
+	 * 
+	 * NOTE! Called on the AudioThread
+	 *
+	 * @param ActiveSound   The active sound being associated
+	 * @param Owner         The owner it is associated with, or nullptr
+	 */
+	virtual void NotifyActiveSoundOwner(FActiveSound& ActiveSound, const UObject* Owner) {}
+
 	/**
 	 * Called when an active sound has just been added to the audio engine,
 	 * both for brand new sounds and for virtualized sounds that have just become active.
diff --git a/Engine/Source/Runtime/Engine/Public/AudioDevice.h b/Engine/Source/Runtime/Engine/Public/AudioDevice.h
--- a/Engine/Source/Runtime/Engine/Public/AudioDevice.h
+++ b/Engine/Source/Runtime/Engine/Public/AudioDevice.h
@@ -909,6 +909,11 @@ protected:
 		NotifySubsystemsActiveSoundDeleting(ActiveSound);
 	}
 
+	/**
+	 * Notifies subsystems an owning object has been associated with a sound.
+	 */
+	ENGINE_API void NotifySubsystemsSetOwner(FActiveSound& ActiveSound, const UObject* Owner) const;
+
 public:
 
 	/**
