Chaos::Fix crash in dragging physics assets in ImmediatePhysics

diff --git a/Engine/Source/Runtime/Engine/Private/PhysicsEngine/ImmediatePhysics/ImmediatePhysicsChaos/ImmediatePhysicsSimulation_Chaos.cpp b/Engine/Source/Runtime/Engine/Private/PhysicsEngine/ImmediatePhysics/ImmediatePhysicsChaos/ImmediatePhysicsSimulation_Chaos.cpp
--- a/Engine/Source/Runtime/Engine/Private/PhysicsEngine/ImmediatePhysics/ImmediatePhysicsChaos/ImmediatePhysicsSimulation_Chaos.cpp
+++ b/Engine/Source/Runtime/Engine/Private/PhysicsEngine/ImmediatePhysics/ImmediatePhysicsChaos/ImmediatePhysicsSimulation_Chaos.cpp
@@ -541,15 +541,15 @@ namespace ImmediatePhysics_Chaos
 
 	FJointHandle* FSimulation::CreateJoint(FConstraintInstance* ConstraintInstance, FActorHandle* Body1, FActorHandle* Body2)
 	{
-		if (ConstraintInstance->ProfileInstance.bUseLinearJointSolver) 
+		if (ConstraintInstance && !ConstraintInstance->ProfileInstance.bUseLinearJointSolver)
 		{
-			FJointHandle* JointHandle = new FJointHandle(&(Implementation->Joints.LinearConstraints), ConstraintInstance, Body1, Body2);
+			FJointHandle* JointHandle = new FJointHandle(&(Implementation->Joints.NonLinearConstraints), ConstraintInstance, Body1, Body2);
 			Implementation->JointHandles.Add(JointHandle);
 			return JointHandle;
 		}
 		else
 		{
-			FJointHandle* JointHandle = new FJointHandle(&(Implementation->Joints.NonLinearConstraints), ConstraintInstance, Body1, Body2);
+			FJointHandle* JointHandle = new FJointHandle(&(Implementation->Joints.LinearConstraints), ConstraintInstance, Body1, Body2);
 			Implementation->JointHandles.Add(JointHandle);
 			return JointHandle;
 		}
