Remove unecessary check in FRayTracingManager.
- caused development builds to crash when launching from Editor and project has Zen Store disabled.
- FScatterGatherBuilder is able to handle streaming bulk data using both IODispatcher and legacy IO path.

diff --git a/Engine/Source/Runtime/Engine/Private/Rendering/RayTracingGeometryManager.cpp b/Engine/Source/Runtime/Engine/Private/Rendering/RayTracingGeometryManager.cpp
--- a/Engine/Source/Runtime/Engine/Private/Rendering/RayTracingGeometryManager.cpp
+++ b/Engine/Source/Runtime/Engine/Private/Rendering/RayTracingGeometryManager.cpp
@@ -919,8 +919,6 @@ bool FRayTracingGeometryManager::RequestRayTracingGeometryStreamIn(FRHICommandLi
 		checkf(StreamableData->CanLoadFromDisk(), TEXT("Bulk data is not loaded and cannot be loaded from disk!"));
 		check(!StreamableData->IsStoredCompressedOnDisk()); // We do not support compressed Bulkdata for this system. Limitation of the streaming request/bulk data
 
-		check(StreamableData->IsUsingIODispatcher());
-
 		if (NumStreamingRequests >= GRayTracingStreamingMaxPendingRequests)
 		{
 			return false;
