Removed check containing questionable cast. Not sure we need to enforce that layout at this point.
Fixed max bounds being initialized to MIN_flt (that's not the minimum float value, it's the minimum *positive* float value).

#rnx

diff --git a/Engine/Source/Runtime/Engine/Classes/Components/HierarchicalInstancedStaticMeshComponent.h b/Engine/Source/Runtime/Engine/Classes/Components/HierarchicalInstancedStaticMeshComponent.h
--- a/Engine/Source/Runtime/Engine/Classes/Components/HierarchicalInstancedStaticMeshComponent.h
+++ b/Engine/Source/Runtime/Engine/Classes/Components/HierarchicalInstancedStaticMeshComponent.h
@@ -44,7 +44,7 @@ struct FClusterNode_DEPRECATED
 	FClusterNode_DEPRECATED()
 		: BoundMin(MAX_flt, MAX_flt, MAX_flt)
 		, FirstChild(-1)
-		, BoundMax(MIN_flt, MIN_flt, MIN_flt)
+		, BoundMax(-MAX_flt, -MAX_flt, -MAX_flt)
 		, LastChild(-1)
 		, FirstInstance(-1)
 		, LastInstance(-1)
@@ -93,7 +93,7 @@ struct FClusterNode
 	FClusterNode()
 		: BoundMin(MAX_flt, MAX_flt, MAX_flt)
 		, FirstChild(-1)
-		, BoundMax(MIN_flt, MIN_flt, MIN_flt)
+		, BoundMax(-MAX_flt, -MAX_flt, -MAX_flt)
 		, LastChild(-1)
 		, FirstInstance(-1)
 		, LastInstance(-1)
@@ -143,7 +143,7 @@ class UHierarchicalInstancedStaticMeshComponent : public UInstancedStaticMeshCom
 	TSharedPtr<TArray<FClusterNode>, ESPMode::ThreadSafe> ClusterTreePtr;
 
 	// If true then we allow a translated space when building the cluster tree.
-	// This can help for impementations (foliage) where we can have instances with offsets to large for single float precision.
+	// This can help for implementations (foliage) where we can have instances with offsets to large for single float precision.
 	UPROPERTY()
 	uint32 bUseTranslatedInstanceSpace : 1;
 
diff --git a/Engine/Source/Runtime/Engine/Private/HierarchicalInstancedStaticMesh.cpp b/Engine/Source/Runtime/Engine/Private/HierarchicalInstancedStaticMesh.cpp
--- a/Engine/Source/Runtime/Engine/Private/HierarchicalInstancedStaticMesh.cpp
+++ b/Engine/Source/Runtime/Engine/Private/HierarchicalInstancedStaticMesh.cpp
@@ -1739,11 +1739,6 @@ void FHierarchicalStaticMeshSceneProxy::GetDynamicMeshElements(const TArray<cons
 				{
 					SCOPE_CYCLE_COUNTER(STAT_FoliageTraversalTime);
 
-					// validate that the bounding box is layed out correctly in memory
-					check((const FVector4f*)&ClusterTree[0].BoundMin + 1 == (const FVector4f*)&ClusterTree[0].BoundMax); //-V594
-					//check(UPTRINT(&ClusterTree[0].BoundMin) % 16 == 0);
-					//check(UPTRINT(&ClusterTree[0].BoundMax) % 16 == 0);
-
 					int32 UseMinLOD = ClampedMinLOD;
 
 					int32 DebugMin = FMath::Min(CVarMinLOD.GetValueOnRenderThread(), InstanceParams.LODs - 1);
