Workaround for crash triggered by using the duplicate function in the ISM instance array in the editor by forcing a full update.

diff --git a/Engine/Source/Runtime/Engine/Private/InstancedStaticMesh.cpp b/Engine/Source/Runtime/Engine/Private/InstancedStaticMesh.cpp
--- a/Engine/Source/Runtime/Engine/Private/InstancedStaticMesh.cpp
+++ b/Engine/Source/Runtime/Engine/Private/InstancedStaticMesh.cpp
@@ -5611,6 +5611,9 @@ void UInstancedStaticMeshComponent::PostEditChangeChainProperty(FPropertyChanged
 
 				// added via the property editor, so we will want to interactively work with instances
 				bHasPerInstanceHitProxies = true;
+
+				// This path may reorder the instances by insertion and this is not handled by the incremental updates in the proxy, so we force a full update.
+				PrimitiveInstanceDataManager.ClearChangeTracking();
 			}
 			else if (PropertyChangedEvent.ChangeType == EPropertyChangeType::ArrayRemove)
 			{
