First Person: Primitives can now be tagged as being the world space/third person representation of a first person primitive.

#rnx

diff --git a/Engine/Source/Runtime/Engine/Classes/Components/PrimitiveComponent.h b/Engine/Source/Runtime/Engine/Classes/Components/PrimitiveComponent.h
--- a/Engine/Source/Runtime/Engine/Classes/Components/PrimitiveComponent.h
+++ b/Engine/Source/Runtime/Engine/Classes/Components/PrimitiveComponent.h
@@ -132,6 +132,8 @@ enum class EFirstPersonPrimitiveType : uint8
 	None UMETA(DisplayName = "None"),
 	/** Primitive is rendered as first person and affected by first person properties on the camera. */
 	FirstPerson UMETA(DisplayName = "First Person"),
+	/** Primitive represents a first person primitive in world space. This is the primitive that other players would see and is used to cast shadow onto the ground among other things. */
+	WorldSpaceRepresentation UMETA(DisplayName = "World Space Representation"),
 };
 
 /** Converts a stencil mask from the editor's USTRUCT version to the version the renderer uses. */
diff --git a/Engine/Source/Runtime/Engine/Private/Components/PrimitiveComponent.cpp b/Engine/Source/Runtime/Engine/Private/Components/PrimitiveComponent.cpp
--- a/Engine/Source/Runtime/Engine/Private/Components/PrimitiveComponent.cpp
+++ b/Engine/Source/Runtime/Engine/Private/Components/PrimitiveComponent.cpp
@@ -1227,7 +1227,7 @@ void UPrimitiveComponent::PostEditChangeProperty(FPropertyChangedEvent& Property
 			MarkChildPrimitiveComponentRenderStateDirty();
 		}
 
-		// bIsFirstPerson can be toggled at runtime and needs to propagate to the scene proxy.
+		// FirstPersonPrimitiveType can be changed at runtime and needs to propagate to the scene proxy.
 		if (PropertyName == GET_MEMBER_NAME_CHECKED(UPrimitiveComponent, FirstPersonPrimitiveType))
 		{
 			MarkRenderStateDirty();
diff --git a/Engine/Source/Runtime/Engine/Private/PrimitiveSceneProxy.cpp b/Engine/Source/Runtime/Engine/Private/PrimitiveSceneProxy.cpp
--- a/Engine/Source/Runtime/Engine/Private/PrimitiveSceneProxy.cpp
+++ b/Engine/Source/Runtime/Engine/Private/PrimitiveSceneProxy.cpp
@@ -309,6 +309,7 @@ void FPrimitiveSceneProxyDesc::InitializeFromPrimitiveComponent(const UPrimitive
 	bHoldout = InComponent->bHoldout;
 	bWantsEditorEffects = InComponent->bWantsEditorEffects;
 	bIsFirstPerson = InComponent->FirstPersonPrimitiveType == EFirstPersonPrimitiveType::FirstPerson;
+	bIsFirstPersonWorldSpaceRepresentation = InComponent->FirstPersonPrimitiveType == EFirstPersonPrimitiveType::WorldSpaceRepresentation;
 	
 	bIsVisible = InComponent->IsVisible();
 	bIsVisibleEditor = InComponent->GetVisibleFlag();
@@ -494,6 +495,7 @@ FPrimitiveSceneProxy::FPrimitiveSceneProxy(const FPrimitiveSceneProxyDesc& InPro
 ,	bSplineMesh(false)
 ,	bSkinnedMesh(false)
 ,	bIsFirstPerson(InProxyDesc.bIsFirstPerson)
+,	bIsFirstPersonWorldSpaceRepresentation(InProxyDesc.bIsFirstPersonWorldSpaceRepresentation)
 ,	bUseAsOccluder(InProxyDesc.bUseAsOccluder)
 ,	bSelectable(InProxyDesc.bSelectable)
 ,	bHasPerInstanceHitProxies(InProxyDesc.bHasPerInstanceHitProxies)
diff --git a/Engine/Source/Runtime/Engine/Public/PrimitiveSceneProxy.h b/Engine/Source/Runtime/Engine/Public/PrimitiveSceneProxy.h
--- a/Engine/Source/Runtime/Engine/Public/PrimitiveSceneProxy.h
+++ b/Engine/Source/Runtime/Engine/Public/PrimitiveSceneProxy.h
@@ -801,6 +801,7 @@ public:
 	inline bool IsSplineMesh() const { return bSplineMesh; }
 	inline bool IsSkinnedMesh() const { return bSkinnedMesh; }
 	inline bool IsFirstPerson() const { return bIsFirstPerson; }
+	inline bool IsFirstPersonWorldSpaceRepresentation() const { return bIsFirstPersonWorldSpaceRepresentation; }
 
 	inline FRHIUniformBuffer* GetUniformBuffer() const
 	{
@@ -1474,8 +1475,11 @@ protected:
 	/** If this is True, this primitive is a skinned mesh */
 	uint8 bSkinnedMesh : 1;
 
-	/** If this is True, this primitive is to be rendered as a first person view object */
+	/** If this is True, this primitive is to be rendered as a first person view object. Mutually exclusive with bIsFirstPersonWorldSpaceRepresentation. */
 	uint8 bIsFirstPerson : 1;
+
+	/** If this is True, this primitive represents another first person primitive in world space. Mutually exclusive with bIsFirstPerson. */
+	uint8 bIsFirstPersonWorldSpaceRepresentation : 1;
 	
 private:
 
diff --git a/Engine/Source/Runtime/Engine/Public/PrimitiveSceneProxyDesc.h b/Engine/Source/Runtime/Engine/Public/PrimitiveSceneProxyDesc.h
--- a/Engine/Source/Runtime/Engine/Public/PrimitiveSceneProxyDesc.h
+++ b/Engine/Source/Runtime/Engine/Public/PrimitiveSceneProxyDesc.h
@@ -62,6 +62,7 @@ struct FPrimitiveSceneProxyDesc
 		bIsInstancedStaticMesh = false;
 		bHoldout = false;
 		bIsFirstPerson = false;
+		bIsFirstPersonWorldSpaceRepresentation = false;
 
 		bHasStaticLighting = false;
 		bHasValidSettingsForStaticLighting = false;
@@ -124,6 +125,7 @@ struct FPrimitiveSceneProxyDesc
 	uint32 bRayTracingFarField : 1;
 	uint32 bHoldout : 1;
 	uint32 bIsFirstPerson : 1;
+	uint32 bIsFirstPersonWorldSpaceRepresentation : 1;
 
 	// not mirrored from UPrimitiveComponent
 	uint32 bIsVisible : 1;
