Platform screenshots, video recording, and video streaming can now be en- and disabled via Console commands

#rn minor Media

diff --git a/Engine/Source/Runtime/Engine/Private/PlatformFeatures.cpp b/Engine/Source/Runtime/Engine/Private/PlatformFeatures.cpp
--- a/Engine/Source/Runtime/Engine/Private/PlatformFeatures.cpp
+++ b/Engine/Source/Runtime/Engine/Private/PlatformFeatures.cpp
@@ -42,4 +42,42 @@ void IPlatformFeaturesModule::SetScreenshotEnableState(bool bEnabled)
 		AndroidThunkCpp_DisableScreenCapture(!bEnabled);
 	}
 #endif
-}
\ No newline at end of file
+}
+
+static FAutoConsoleCommand CVarPFMSetScreenshotEnableState(
+	TEXT("pf.SetScreenshotEnableState"),
+	TEXT("Enables or disables taking screenshots if the platform supports it."),
+	FConsoleCommandWithArgsDelegate::CreateLambda([](const TArray<FString>& Args)
+	{
+		if (Args.Num() > 0)
+		{
+			IPlatformFeaturesModule::Get().SetScreenshotEnableState(Args[0].ToBool());
+		}
+	})
+);
+
+static FAutoConsoleCommand CVarPFMEnableRecording(
+	TEXT("pf.EnableRecording"),
+	TEXT("Enables or disables recording if the platform supports it."),
+	FConsoleCommandWithArgsDelegate::CreateLambda([](const TArray<FString>& Args)
+	{
+		IVideoRecordingSystem* VideoRecordingSystem = IPlatformFeaturesModule::Get().GetVideoRecordingSystem();
+		if (VideoRecordingSystem && (Args.Num() > 0))
+		{
+			VideoRecordingSystem->EnableRecording(Args[0].ToBool());
+		}
+	})
+);
+
+static FAutoConsoleCommand CVarPFMEnableStreaming(
+	TEXT("pf.EnableStreaming"),
+	TEXT("Enables or disables streaming if the platform supports it."),
+	FConsoleCommandWithArgsDelegate::CreateLambda([](const TArray<FString>& Args)
+	{
+		IDVRStreamingSystem* DVRStreamingSystem = IPlatformFeaturesModule::Get().GetStreamingSystem();
+		if (DVRStreamingSystem && (Args.Num() > 0))
+		{
+			DVRStreamingSystem->EnableStreaming(Args[0].ToBool());
+		}
+	})
+);
\ No newline at end of file
