[Engine]
* Removed more direct component access from cascade system

diff --git a/Engine/Source/Runtime/Engine/Classes/Particles/Beam/ParticleModuleBeamSource.h b/Engine/Source/Runtime/Engine/Classes/Particles/Beam/ParticleModuleBeamSource.h
--- a/Engine/Source/Runtime/Engine/Classes/Particles/Beam/ParticleModuleBeamSource.h
+++ b/Engine/Source/Runtime/Engine/Classes/Particles/Beam/ParticleModuleBeamSource.h
@@ -97,7 +97,7 @@ class UParticleModuleBeamSource : public UParticleModuleBeamBase
 						
 
 	// @todo document
-	bool ResolveSourceData(FParticleBeam2EmitterInstance* BeamInst, 
+	bool ResolveSourceData(const FContext& Context, FParticleBeam2EmitterInstance* BeamInst, 
 				FBeam2TypeDataPayload* BeamData, const uint8* ParticleBase, 
 				int32& CurrentOffset, int32	ParticleIndex, bool bSpawning,
 				FBeamParticleModifierPayloadData* ModifierData);
diff --git a/Engine/Source/Runtime/Engine/Classes/Particles/Beam/ParticleModuleBeamTarget.h b/Engine/Source/Runtime/Engine/Classes/Particles/Beam/ParticleModuleBeamTarget.h
--- a/Engine/Source/Runtime/Engine/Classes/Particles/Beam/ParticleModuleBeamTarget.h
+++ b/Engine/Source/Runtime/Engine/Classes/Particles/Beam/ParticleModuleBeamTarget.h
@@ -92,7 +92,7 @@ class UParticleModuleBeamTarget : public UParticleModuleBeamBase
 				FBeamParticleSourceTargetPayloadData*& ParticleSource);
 						
 	// @todo document
-	bool	ResolveTargetData(FParticleBeam2EmitterInstance* BeamInst, 
+	bool	ResolveTargetData(const FContext& Context, FParticleBeam2EmitterInstance* BeamInst, 
 				FBeam2TypeDataPayload* BeamData, const uint8* ParticleBase, 
 				int32& CurrentOffset, int32	ParticleIndex, bool bSpawning,
 				FBeamParticleModifierPayloadData* ModifierData);
diff --git a/Engine/Source/Runtime/Engine/Classes/Particles/Parameter/ParticleModuleParameterDynamic.h b/Engine/Source/Runtime/Engine/Classes/Particles/Parameter/ParticleModuleParameterDynamic.h
--- a/Engine/Source/Runtime/Engine/Classes/Particles/Parameter/ParticleModuleParameterDynamic.h
+++ b/Engine/Source/Runtime/Engine/Classes/Particles/Parameter/ParticleModuleParameterDynamic.h
@@ -167,7 +167,7 @@ class UParticleModuleParameterDynamic : public UParticleModuleParameterBase
 	 *
 	 *	@return	float			The value for the parameter.
 	 */
-	FORCEINLINE float GetParameterValue(FEmitterDynamicParameter& InDynParams, FBaseParticle& Particle, FParticleEmitterInstance* Owner, struct FRandomStream* InRandomStream)
+	FORCEINLINE float GetParameterValue(const FContext& Context, FEmitterDynamicParameter& InDynParams, FBaseParticle& Particle, FParticleEmitterInstance* Owner, struct FRandomStream* InRandomStream)
 	{
 		float ScaleValue = 1.0f;
 		float DistributionValue = 1.0f;
@@ -194,7 +194,7 @@ class UParticleModuleParameterDynamic : public UParticleModuleParameterBase
 		if ((InDynParams.bScaleVelocityByParamValue == true) || (InDynParams.ValueMethod == EDPV_UserSet))
 		{
 			float TimeValue = InDynParams.bUseEmitterTime ? Owner->EmitterTime : Particle.RelativeTime;
-			DistributionValue = InDynParams.ParamValue.GetValue(TimeValue, Owner->Component, InRandomStream);
+			DistributionValue = InDynParams.ParamValue.GetValue(TimeValue, Context.GetDistributionData(), InRandomStream);
 		}
 
 		return DistributionValue * ScaleValue;
@@ -210,9 +210,9 @@ class UParticleModuleParameterDynamic : public UParticleModuleParameterBase
 	 *
 	 *	@return	float			The value for the parameter.
 	 */
-	FORCEINLINE float GetParameterValue_UserSet(FEmitterDynamicParameter& InDynParams, FBaseParticle& Particle, FParticleEmitterInstance* Owner, struct FRandomStream* InRandomStream)
+	FORCEINLINE float GetParameterValue_UserSet(const FContext& Context, FEmitterDynamicParameter& InDynParams, FBaseParticle& Particle, FParticleEmitterInstance* Owner, struct FRandomStream* InRandomStream)
 	{
-		return InDynParams.ParamValue.GetValue(InDynParams.bUseEmitterTime ? Owner->EmitterTime : Particle.RelativeTime, Owner->Component, InRandomStream);
+		return InDynParams.ParamValue.GetValue(InDynParams.bUseEmitterTime ? Owner->EmitterTime : Particle.RelativeTime, Context.GetDistributionData(), InRandomStream);
 	}
 
 	/**
diff --git a/Engine/Source/Runtime/Engine/Private/Particles/ParticleBeamModules.cpp b/Engine/Source/Runtime/Engine/Private/Particles/ParticleBeamModules.cpp
--- a/Engine/Source/Runtime/Engine/Private/Particles/ParticleBeamModules.cpp
+++ b/Engine/Source/Runtime/Engine/Private/Particles/ParticleBeamModules.cpp
@@ -102,8 +102,6 @@ void UParticleModuleTypeDataBeam2::Spawn(const FSpawnContext& Context)
 	FParticleBeam2EmitterInstance*	BeamInst	= (FParticleBeam2EmitterInstance*)Owner;
 	check(Owner->bIsBeam);
 
-	UParticleSystemComponent* Component = Owner->Component;
-
 	// Setup the particle data points with the SPAWN_INIT macro
 	SPAWN_INIT;
 
@@ -129,8 +127,8 @@ void UParticleModuleTypeDataBeam2::Spawn(const FSpawnContext& Context)
 	// If there is no Source module, use the emitter position as the source point
 	if (BeamInst->BeamModule_Source == NULL)
 	{
-		BeamData->SourcePoint	= Component->GetComponentLocation();
-		BeamData->SourceTangent = Component->GetComponentTransform().GetScaledAxis( EAxis::X );
+		BeamData->SourcePoint	= Owner->Component->GetComponentLocation();
+		BeamData->SourceTangent = Context.GetTransform().GetScaledAxis( EAxis::X );
 		BeamData->SourceStrength = 1.0f;
 	}
 
@@ -138,9 +136,9 @@ void UParticleModuleTypeDataBeam2::Spawn(const FSpawnContext& Context)
 	if ((BeamInst->BeamModule_Target == NULL) && (BeamInst->BeamMethod == PEB2M_Distance))
 	{
 		// Set the particle target based on the distance
-		double	TotalDistance	= Distance.GetValue(Particle.RelativeTime, Component);
+		double	TotalDistance	= Distance.GetValue(Particle.RelativeTime, Context.GetDistributionData());
 		// Always use the X-axis of the component as the direction
-		FVector	Direction		= Component->GetComponentTransform().GetScaledAxis( EAxis::X );
+		FVector	Direction		= Context.GetTransform().GetScaledAxis( EAxis::X );
 		Direction.Normalize();
 		// Calculate the final target point
 		BeamData->TargetPoint	= BeamData->SourcePoint + Direction * TotalDistance;
@@ -198,7 +196,7 @@ void UParticleModuleTypeDataBeam2::Spawn(const FSpawnContext& Context)
 		for (int32 TaperIndex = 0; TaperIndex < TaperCount; TaperIndex++)
 		{
 			CurrStep	= TaperIndex * Increment;
-			TaperValues[TaperIndex] = TaperFactor.GetValue(CurrStep, Component) * TaperScale.GetValue(CurrStep, Component);
+			TaperValues[TaperIndex] = TaperFactor.GetValue(CurrStep, Context.GetDistributionData()) * TaperScale.GetValue(CurrStep, Context.GetDistributionData());
 		}
 	}
 }
@@ -207,7 +205,6 @@ void UParticleModuleTypeDataBeam2::Spawn(const FSpawnContext& Context)
 void UParticleModuleTypeDataBeam2::Update(const FUpdateContext& Context)
 {
 	FParticleEmitterInstance* Owner = &Context.Owner;
-	UParticleSystemComponent*		Component	= Owner->Component;
 	FParticleBeam2EmitterInstance*	BeamInst	= (FParticleBeam2EmitterInstance*)Owner;
 	check(Owner->bIsBeam);
 	UParticleModuleBeamNoise*		BeamNoise	= BeamInst->BeamModule_Noise;
@@ -249,16 +246,16 @@ void UParticleModuleTypeDataBeam2::Update(const FUpdateContext& Context)
 		// If there is no Source module, use the emitter location
 		if (BeamInst->BeamModule_Source == NULL)
 		{
-			BeamData->SourcePoint	= Component->GetComponentLocation();
-			BeamData->SourceTangent = Component->GetComponentTransform().GetScaledAxis( EAxis::X );
+			BeamData->SourcePoint	= Owner->Component->GetComponentLocation();
+			BeamData->SourceTangent = Context.GetTransform().GetScaledAxis( EAxis::X );
 		}
 
 		// If the method is set for distance, or there is no target, determine the target point
 		if ((BeamInst->BeamModule_Target == NULL) && (BeamInst->BeamMethod == PEB2M_Distance))
 		{
 			// Set the particle target based on the distance
-			double	TotalDistance	= Distance.GetValue(Particle.RelativeTime, Component);
-			FVector	Direction		= Component->GetComponentTransform().GetScaledAxis( EAxis::X );
+			double	TotalDistance	= Distance.GetValue(Particle.RelativeTime, Context.GetDistributionData());
+			FVector	Direction		= Context.GetTransform().GetScaledAxis( EAxis::X );
 			Direction.Normalize();
 			BeamData->TargetPoint	= BeamData->SourcePoint + Direction * TotalDistance;
 			BeamData->TargetTangent = -Direction;
@@ -501,7 +498,7 @@ void UParticleModuleTypeDataBeam2::Update(const FUpdateContext& Context)
 #endif	//#if defined(_BEAM2_TYPEDATA_NORMAL_TANGENTS_)
 			if (SourceTangent.IsNearlyZero())
 			{
-				SourceTangent	= Component->GetComponentTransform().GetScaledAxis( EAxis::X );
+				SourceTangent	= Context.GetTransform().GetScaledAxis( EAxis::X );
 			}
 			SourceTangent	*= BeamData->SourceStrength;
 
@@ -512,7 +509,7 @@ void UParticleModuleTypeDataBeam2::Update(const FUpdateContext& Context)
 #endif	//#if defined(_BEAM2_TYPEDATA_NORMAL_TANGENTS_)
 			if (TargetTangent.IsNearlyZero())
 			{
-				TargetTangent	= Component->GetComponentTransform().GetScaledAxis( EAxis::X );
+				TargetTangent	= Context.GetTransform().GetScaledAxis( EAxis::X );
 			}
 			TargetTangent	*= BeamData->TargetStrength;
 
@@ -1478,7 +1475,6 @@ void UParticleModuleBeamNoise::Update(const FUpdateContext& Context)
 		return;
 	}
 
-	UParticleSystemComponent*		Component	= Owner->Component;
 	UParticleModuleTypeDataBeam2*	BeamTD		= BeamInst->BeamTypeData;
 
 	bool bLocalOscillate = false;
@@ -1704,7 +1700,6 @@ void UParticleModuleBeamSource::Spawn(const FSpawnContext& Context)
 	{
 		return;
 	}
-	UParticleSystemComponent*		Component	= Owner->Component;
 	UParticleModuleTypeDataBeam2*	BeamTD		= BeamInst->BeamTypeData;
 
 	SPAWN_INIT;
@@ -1728,7 +1723,7 @@ void UParticleModuleBeamSource::Spawn(const FSpawnContext& Context)
 
 	// Resolve the source data
 	int32 Offset = Context.Offset;
-	ResolveSourceData(BeamInst, BeamData, (uint8*)ParticleBase, Offset, BeamInst->ActiveParticles, true, SourceModifier);
+	ResolveSourceData(Context, BeamInst, BeamData, (uint8*)ParticleBase, Offset, BeamInst->ActiveParticles, true, SourceModifier);
 
 	// Set the location and clear the initial data flags
 	Particle.Location					= BeamData->SourcePoint - BeamInst->PositionOffsetThisTick;
@@ -1755,7 +1750,6 @@ void UParticleModuleBeamSource::Update(const FUpdateContext& Context)
 	{
 		return;
 	}
-	UParticleSystemComponent*		Component	= Owner->Component;
 	UParticleModuleTypeDataBeam2*	BeamTD		= BeamInst->BeamTypeData;
 
 	BEGIN_UPDATE_LOOP;
@@ -1778,7 +1772,7 @@ void UParticleModuleBeamSource::Update(const FUpdateContext& Context)
 			NoiseDistanceScale, SourceModifier, TargetModifier);
 
 		// Resolve the source data
-		ResolveSourceData(BeamInst, BeamData, ParticleBase, Offset, i, false, SourceModifier);
+		ResolveSourceData(Context, BeamInst, BeamData, ParticleBase, Offset, i, false, SourceModifier);
 	}
 	END_UPDATE_LOOP;
 }
@@ -1876,7 +1870,7 @@ void UParticleModuleBeamSource::GetDataPointers(FParticleEmitterInstance* Owner,
 	}
 }
 
-bool UParticleModuleBeamSource::ResolveSourceData(FParticleBeam2EmitterInstance* BeamInst, 
+bool UParticleModuleBeamSource::ResolveSourceData(const FContext& Context, FParticleBeam2EmitterInstance* BeamInst, 
 	FBeam2TypeDataPayload* BeamData, const uint8* ParticleBase, int32& Offset, int32	ParticleIndex,
 	bool bSpawning, FBeamParticleModifierPayloadData* ModifierData)
 {
@@ -2033,15 +2027,15 @@ bool UParticleModuleBeamSource::ResolveSourceData(FParticleBeam2EmitterInstance*
 			if (bSourceAbsolute)
 			{
 				// Use the value as a world space position
-				BeamData->SourcePoint	= Source.GetValue(BeamInst->EmitterTime, BeamInst->Component);
+				BeamData->SourcePoint	= Source.GetValue(BeamInst->EmitterTime, Context.GetDistributionData());
 				// Take into account current world origin offset
 				BeamData->SourcePoint  -= (FVector)BeamInst->Component->GetWorld()->OriginLocation;
 			}
 			else
 			{
 				// Use the value as a local space position.
-				BeamData->SourcePoint	= BeamInst->Component->GetComponentTransform().TransformPosition(
-					Source.GetValue(BeamInst->EmitterTime, BeamInst->Component));
+				BeamData->SourcePoint	= Context.GetTransform().TransformPosition(
+					Source.GetValue(BeamInst->EmitterTime, Context.GetDistributionData()));
 			}
 		}
 	}
@@ -2054,7 +2048,7 @@ bool UParticleModuleBeamSource::ResolveSourceData(FParticleBeam2EmitterInstance*
 		{
 		case PEB2STTM_Direct:
 			// Use the emitter direction as the tangent
-			BeamData->SourceTangent	= BeamInst->Component->GetComponentTransform().GetScaledAxis( EAxis::X );
+			BeamData->SourceTangent	= Context.GetTransform().GetScaledAxis( EAxis::X );
 			bSetSourceTangent		= true;
 			break;
 		case PEB2STTM_UserSet:
@@ -2075,12 +2069,12 @@ bool UParticleModuleBeamSource::ResolveSourceData(FParticleBeam2EmitterInstance*
 			break;
 		case PEB2STTM_Distribution:
 			// Use the tangent contained in the distribution
-			BeamData->SourceTangent	= SourceTangent.GetValue(Particle.RelativeTime, BeamInst->Component);
+			BeamData->SourceTangent	= SourceTangent.GetValue(Particle.RelativeTime, Context.GetDistributionData());
 			bSetSourceTangent		= true;
 			break;
 		case PEB2STTM_Emitter:
 			// Use the emitter direction as the tangent
-			BeamData->SourceTangent	= BeamInst->Component->GetComponentTransform().GetScaledAxis( EAxis::X );
+			BeamData->SourceTangent	= Context.GetTransform().GetScaledAxis( EAxis::X );
 			bSetSourceTangent		= true;
 			break;
 		}
@@ -2089,11 +2083,11 @@ bool UParticleModuleBeamSource::ResolveSourceData(FParticleBeam2EmitterInstance*
 		{
 			// By default, use the distribution. This will allow artists an easier setup phase...
 //			BeamData->SourceTangent	= BeamInst->Component->LocalToWorld.GetScaledAxis( EAxis::X );
-			BeamData->SourceTangent	= SourceTangent.GetValue(Particle.RelativeTime, BeamInst->Component);
+			BeamData->SourceTangent	= SourceTangent.GetValue(Particle.RelativeTime, Context.GetDistributionData());
 			if (bSourceAbsolute == false)
 			{
 				// If not tagged as absolute, transform it to world space
-				BeamData->SourceTangent	= BeamInst->Component->GetComponentTransform().TransformVector(BeamData->SourceTangent);
+				BeamData->SourceTangent	= Context.GetTransform().TransformVector(BeamData->SourceTangent);
 			}
 		}
 	}
@@ -2121,7 +2115,7 @@ bool UParticleModuleBeamSource::ResolveSourceData(FParticleBeam2EmitterInstance*
 
 		if (!bSetSourceStrength)
 		{
-			BeamData->SourceStrength	= SourceStrength.GetValue(Particle.RelativeTime, BeamInst->Component);
+			BeamData->SourceStrength	= SourceStrength.GetValue(Particle.RelativeTime, Context.GetDistributionData());
 		}
 	}
 
@@ -2197,7 +2191,6 @@ void UParticleModuleBeamTarget::Spawn(const FSpawnContext& Context)
 	{
 		return;
 	}
-	UParticleSystemComponent*		Component	= Owner->Component;
 	UParticleModuleTypeDataBeam2*	BeamTD		= BeamInst->BeamTypeData;
 
 	SPAWN_INIT;
@@ -2220,7 +2213,7 @@ void UParticleModuleBeamTarget::Spawn(const FSpawnContext& Context)
 	CurrentOffset	= TempOffset;
 
 	int32 Offset = Context.Offset;
-	ResolveTargetData(BeamInst, BeamData, (uint8*)ParticleBase, Offset, BeamInst->ActiveParticles, true, TargetModifier);
+	ResolveTargetData(Context, BeamInst, BeamData, (uint8*)ParticleBase, Offset, BeamInst->ActiveParticles, true, TargetModifier);
 }
 
 void UParticleModuleBeamTarget::Update(const FUpdateContext& Context)
@@ -2236,7 +2229,6 @@ void UParticleModuleBeamTarget::Update(const FUpdateContext& Context)
 	{
 		return;
 	}
-	UParticleSystemComponent*		Component	= Owner->Component;
 	UParticleModuleTypeDataBeam2*	BeamTD		= BeamInst->BeamTypeData;
 
 	BEGIN_UPDATE_LOOP;
@@ -2257,7 +2249,7 @@ void UParticleModuleBeamTarget::Update(const FUpdateContext& Context)
 			NoiseRate, NoiseDelta, TargetNoisePoints, NextNoisePoints, TaperValues,
 			NoiseDistanceScale, SourceModifier, TargetModifier);
 
-		ResolveTargetData(BeamInst, BeamData, ParticleBase, Offset, i, false, TargetModifier);
+		ResolveTargetData(Context, BeamInst, BeamData, ParticleBase, Offset, i, false, TargetModifier);
 	}
 	END_UPDATE_LOOP;
 }
@@ -2328,7 +2320,7 @@ void UParticleModuleBeamTarget::GetDataPointers(FParticleEmitterInstance* Owner,
 	}
 }
 						
-bool UParticleModuleBeamTarget::ResolveTargetData(FParticleBeam2EmitterInstance* BeamInst, 
+bool UParticleModuleBeamTarget::ResolveTargetData(const FContext& Context, FParticleBeam2EmitterInstance* BeamInst, 
 	FBeam2TypeDataPayload* BeamData, const uint8* ParticleBase, int32& CurrentOffset, int32	ParticleIndex, bool bSpawning, 
 	FBeamParticleModifierPayloadData* ModifierData)
 {
@@ -2347,12 +2339,12 @@ bool UParticleModuleBeamTarget::ResolveTargetData(FParticleBeam2EmitterInstance*
 		if (BeamInst->BeamTypeData->BeamMethod	== PEB2M_Distance)
 		{
 			// Set the particle target based on the distance
-			float	Distance		= BeamInst->BeamTypeData->Distance.GetValue(Particle.RelativeTime, BeamInst->Component);
+			float	Distance		= BeamInst->BeamTypeData->Distance.GetValue(Particle.RelativeTime, Context.GetDistributionData());
 			if (FMath::Abs(Distance) < UE_KINDA_SMALL_NUMBER)
 			{
 				Distance	= 0.001f;
 			}
-			FVector	Direction		= BeamInst->Component->GetComponentTransform().GetScaledAxis( EAxis::X );
+			FVector	Direction		= Context.GetTransform().GetScaledAxis( EAxis::X );
 			Direction.Normalize();
 			BeamData->TargetPoint	= BeamData->SourcePoint + Direction * Distance;
 			bSetTarget				= true;
@@ -2499,14 +2491,14 @@ bool UParticleModuleBeamTarget::ResolveTargetData(FParticleBeam2EmitterInstance*
 		{
 			if (bTargetAbsolute)
 			{
-				BeamData->TargetPoint	= Target.GetValue(BeamInst->EmitterTime, BeamInst->Component);
+				BeamData->TargetPoint	= Target.GetValue(BeamInst->EmitterTime, Context.GetDistributionData());
 				// Take into account current world origin offset
 				BeamData->TargetPoint  -= (FVector)BeamInst->Component->GetWorld()->OriginLocation;
 			}
 			else
 			{
-				BeamData->TargetPoint	= BeamInst->Component->GetComponentTransform().TransformPosition(
-					Target.GetValue(BeamInst->EmitterTime, BeamInst->Component));
+				BeamData->TargetPoint	= Context.GetTransform().TransformPosition(
+					Target.GetValue(BeamInst->EmitterTime, Context.GetDistributionData()));
 			}
 		}
 	}
@@ -2518,7 +2510,7 @@ bool UParticleModuleBeamTarget::ResolveTargetData(FParticleBeam2EmitterInstance*
 		switch (TargetTangentMethod)
 		{
 		case PEB2STTM_Direct:
-			BeamData->TargetTangent	= BeamInst->Component->GetComponentTransform().GetScaledAxis( EAxis::X );
+			BeamData->TargetTangent	= Context.GetTransform().GetScaledAxis( EAxis::X );
 			bSetTargetTangent		= true;
 			break;
 		case PEB2STTM_UserSet:
@@ -2537,11 +2529,11 @@ bool UParticleModuleBeamTarget::ResolveTargetData(FParticleBeam2EmitterInstance*
 			}
 			break;
 		case PEB2STTM_Distribution:
-			BeamData->TargetTangent	= TargetTangent.GetValue(Particle.RelativeTime, BeamInst->Component);
+			BeamData->TargetTangent	= TargetTangent.GetValue(Particle.RelativeTime, Context.GetDistributionData());
 			bSetTargetTangent		= true;
 			break;
 		case PEB2STTM_Emitter:
-			BeamData->TargetTangent	= BeamInst->Component->GetComponentTransform().GetScaledAxis( EAxis::X );
+			BeamData->TargetTangent	= Context.GetTransform().GetScaledAxis( EAxis::X );
 			bSetTargetTangent		= true;
 			break;
 		}
@@ -2549,10 +2541,10 @@ bool UParticleModuleBeamTarget::ResolveTargetData(FParticleBeam2EmitterInstance*
 		if (bSetTargetTangent == false)
 		{
 //			BeamData->TargetTangent	= BeamInst->Component->LocalToWorld.GetScaledAxis( EAxis::X );
-			BeamData->TargetTangent	= TargetTangent.GetValue(Particle.RelativeTime, BeamInst->Component);
+			BeamData->TargetTangent	= TargetTangent.GetValue(Particle.RelativeTime, Context.GetDistributionData());
 			if (bTargetAbsolute == false)
 			{
-				BeamData->TargetTangent	= BeamInst->Component->GetComponentTransform().TransformVector(BeamData->TargetTangent);
+				BeamData->TargetTangent	= Context.GetTransform().TransformVector(BeamData->TargetTangent);
 			}
 		}
 	}
@@ -2580,7 +2572,7 @@ bool UParticleModuleBeamTarget::ResolveTargetData(FParticleBeam2EmitterInstance*
 
 		if (!bSetTargetStrength)
 		{
-			BeamData->TargetStrength	= TargetStrength.GetValue(Particle.RelativeTime, BeamInst->Component);
+			BeamData->TargetStrength	= TargetStrength.GetValue(Particle.RelativeTime, Context.GetDistributionData());
 		}
 	}
 
diff --git a/Engine/Source/Runtime/Engine/Private/Particles/ParticleModules.cpp b/Engine/Source/Runtime/Engine/Private/Particles/ParticleModules.cpp
--- a/Engine/Source/Runtime/Engine/Private/Particles/ParticleModules.cpp
+++ b/Engine/Source/Runtime/Engine/Private/Particles/ParticleModules.cpp
@@ -4494,8 +4494,6 @@ void UParticleModuleAttractorPoint::PostEditChangeProperty(FPropertyChangedEvent
 void UParticleModuleAttractorPoint::Update(const FUpdateContext& Context)
 {
 	FParticleEmitterInstance* Owner = &Context.Owner;
-	check(Owner);
-	UParticleSystemComponent* Component = Owner->Component;
 
 	// Grab the position of the attractor in Emitter time???
 	FVector AttractorPosition = Position.GetValue(Owner->EmitterTime, Context.GetDistributionData());
@@ -4509,9 +4507,9 @@ void UParticleModuleAttractorPoint::Update(const FUpdateContext& Context)
 	if ( (LODLevel->RequiredModule->bUseLocalSpace == false) && ( bUseWorldSpacePosition == false ) )
 	{
 		// Transform the attractor into world space
-		AttractorPosition = Component->GetComponentTransform().TransformPosition(AttractorPosition);
+		AttractorPosition = Context.GetTransform().TransformPosition(AttractorPosition);
 
-		Scale *= Component->GetComponentTransform().GetScale3D();
+		Scale *= Context.GetTransform().GetScale3D();
 	}
 	float ScaleSize = Scale.Size();
 
diff --git a/Engine/Source/Runtime/Engine/Private/Particles/ParticleModules_Camera.cpp b/Engine/Source/Runtime/Engine/Private/Particles/ParticleModules_Camera.cpp
--- a/Engine/Source/Runtime/Engine/Private/Particles/ParticleModules_Camera.cpp
+++ b/Engine/Source/Runtime/Engine/Private/Particles/ParticleModules_Camera.cpp
@@ -73,16 +73,12 @@ void UParticleModuleCameraOffset::Spawn(const FSpawnContext& Context)
 {
 	FParticleEmitterInstance* Owner = &Context.Owner;
 
-	check(Owner && Owner->Component)
 	float ScaleFactor = 1.0f;
 
 	UParticleLODLevel* LODLevel = Owner ? Owner->SpriteTemplate->GetCurrentLODLevel(Owner) : NULL;
 	if (LODLevel && LODLevel->RequiredModule->bUseLocalSpace == false)
 	{
-		if (Owner != NULL && Owner->Component != NULL)
-		{
-			ScaleFactor = Owner->Component->GetAsyncComponentToWorld().GetMaximumAxisScale();
-		}
+		ScaleFactor = Owner->Component->GetAsyncComponentToWorld().GetMaximumAxisScale();
 	}
 	SPAWN_INIT;
 	{
diff --git a/Engine/Source/Runtime/Engine/Private/Particles/ParticleModules_Collision.cpp b/Engine/Source/Runtime/Engine/Private/Particles/ParticleModules_Collision.cpp
--- a/Engine/Source/Runtime/Engine/Private/Particles/ParticleModules_Collision.cpp
+++ b/Engine/Source/Runtime/Engine/Private/Particles/ParticleModules_Collision.cpp
@@ -181,8 +181,6 @@ void UParticleModuleCollision::Update(const FUpdateContext& Context)
 {
 	SCOPE_CYCLE_COUNTER(STAT_ParticleCollisionTime);
 	FParticleEmitterInstance* Owner = &Context.Owner;
-	check(Owner);
-	check(Owner->Component);
 	UWorld* World = Owner->Component->GetWorld();
 	if (Owner->ActiveParticles == 0 || (bDropDetail && World && World->bDropDetail))
 	{
@@ -363,8 +361,6 @@ void UParticleModuleCollision::Update(const FUpdateContext& Context)
 		Hit.Normal.Y = 0.0f;
 		Hit.Normal.Z = 0.0f;
 
-		check( Owner->Component );
-
 		FVector End = Location + Direction * Size / DirScalar;
 
 		if ((World->IsGameWorld() == true) && (MaxCollisionDistance < UE_OLD_WORLD_MAX))	// LWC_TODO: Fix needed to handle WORLD_MAX increase.
diff --git a/Engine/Source/Runtime/Engine/Private/Particles/ParticleModules_Parameter.cpp b/Engine/Source/Runtime/Engine/Private/Particles/ParticleModules_Parameter.cpp
--- a/Engine/Source/Runtime/Engine/Private/Particles/ParticleModules_Parameter.cpp
+++ b/Engine/Source/Runtime/Engine/Private/Particles/ParticleModules_Parameter.cpp
@@ -148,19 +148,19 @@ void UParticleModuleParameterDynamic::SpawnEx(const FSpawnContext& Context, stru
 		PARTICLE_ELEMENT(FEmitterDynamicParameterPayload, DynamicPayload);
 		if (DynamicParams[0].ValueMethod != EDPV_AutoSet)
 		{
-			DynamicPayload.DynamicParameterValue[0] = GetParameterValue(DynamicParams[0], Particle, Owner, InRandomStream);
+			DynamicPayload.DynamicParameterValue[0] = GetParameterValue(Context, DynamicParams[0], Particle, Owner, InRandomStream);
 		}
 		if (DynamicParams[1].ValueMethod != EDPV_AutoSet)
 		{
-			DynamicPayload.DynamicParameterValue[1] = GetParameterValue(DynamicParams[1], Particle, Owner, InRandomStream);
+			DynamicPayload.DynamicParameterValue[1] = GetParameterValue(Context, DynamicParams[1], Particle, Owner, InRandomStream);
 		}
 		if (DynamicParams[2].ValueMethod != EDPV_AutoSet)
 		{
-			DynamicPayload.DynamicParameterValue[2] = GetParameterValue(DynamicParams[2], Particle, Owner, InRandomStream);
+			DynamicPayload.DynamicParameterValue[2] = GetParameterValue(Context, DynamicParams[2], Particle, Owner, InRandomStream);
 		}
 		if (DynamicParams[3].ValueMethod != EDPV_AutoSet)
 		{
-			DynamicPayload.DynamicParameterValue[3] = GetParameterValue(DynamicParams[3], Particle, Owner, InRandomStream);
+			DynamicPayload.DynamicParameterValue[3] = GetParameterValue(Context, DynamicParams[3], Particle, Owner, InRandomStream);
 		}
 	}
 }
@@ -205,7 +205,7 @@ void UParticleModuleParameterDynamic::Update(const FUpdateContext& Context)
 					FEmitterDynamicParameterPayload& DynamicPayload = *((FEmitterDynamicParameterPayload*)(ParticleBase + CurrentOffset));
 					FPlatformMisc::Prefetch(ParticleData, (ParticleIndices[i+1] * ParticleStride));
 					FPlatformMisc::Prefetch(ParticleData, (ParticleIndices[i+1] * ParticleStride) + PLATFORM_CACHE_LINE_SIZE);
-					DynamicPayload.DynamicParameterValue[ParameterIndex] = GetParameterValue_UserSet(DynParam, Particle, Owner, &RandomStream);
+					DynamicPayload.DynamicParameterValue[ParameterIndex] = GetParameterValue_UserSet(Context, DynParam, Particle, Owner, &RandomStream);
 				}
 				END_UPDATE_LOOP;
 			}
@@ -216,7 +216,7 @@ void UParticleModuleParameterDynamic::Update(const FUpdateContext& Context)
 					FEmitterDynamicParameterPayload& DynamicPayload = *((FEmitterDynamicParameterPayload*)(ParticleBase + CurrentOffset));
 					FPlatformMisc::Prefetch(ParticleData, (ParticleIndices[i+1] * ParticleStride));
 					FPlatformMisc::Prefetch(ParticleData, (ParticleIndices[i+1] * ParticleStride) + PLATFORM_CACHE_LINE_SIZE);
-					DynamicPayload.DynamicParameterValue[ParameterIndex] = GetParameterValue(DynParam, Particle, Owner, &RandomStream);
+					DynamicPayload.DynamicParameterValue[ParameterIndex] = GetParameterValue(Context, DynParam, Particle, Owner, &RandomStream);
 				}
 				END_UPDATE_LOOP;
 			}
@@ -234,8 +234,8 @@ void UParticleModuleParameterDynamic::Update(const FUpdateContext& Context)
 					FEmitterDynamicParameterPayload& DynamicPayload = *((FEmitterDynamicParameterPayload*)(ParticleBase + CurrentOffset));
 					FPlatformMisc::Prefetch(ParticleData, (ParticleIndices[i+1] * ParticleStride));
 					FPlatformMisc::Prefetch(ParticleData, (ParticleIndices[i+1] * ParticleStride) + PLATFORM_CACHE_LINE_SIZE);
-					DynamicPayload.DynamicParameterValue[0] = GetParameterValue_UserSet(DynParam0, Particle, Owner, &RandomStream);
-					DynamicPayload.DynamicParameterValue[1] = GetParameterValue_UserSet(DynParam1, Particle, Owner, &RandomStream);
+					DynamicPayload.DynamicParameterValue[0] = GetParameterValue_UserSet(Context, DynParam0, Particle, Owner, &RandomStream);
+					DynamicPayload.DynamicParameterValue[1] = GetParameterValue_UserSet(Context, DynParam1, Particle, Owner, &RandomStream);
 				}
 				END_UPDATE_LOOP;
 			}
@@ -246,8 +246,8 @@ void UParticleModuleParameterDynamic::Update(const FUpdateContext& Context)
 					FEmitterDynamicParameterPayload& DynamicPayload = *((FEmitterDynamicParameterPayload*)(ParticleBase + CurrentOffset));
 					FPlatformMisc::Prefetch(ParticleData, (ParticleIndices[i+1] * ParticleStride));
 					FPlatformMisc::Prefetch(ParticleData, (ParticleIndices[i+1] * ParticleStride) + PLATFORM_CACHE_LINE_SIZE);
-					DynamicPayload.DynamicParameterValue[0] = GetParameterValue(DynParam0, Particle, Owner, &RandomStream);
-					DynamicPayload.DynamicParameterValue[1] = GetParameterValue(DynParam1, Particle, Owner, &RandomStream);
+					DynamicPayload.DynamicParameterValue[0] = GetParameterValue(Context, DynParam0, Particle, Owner, &RandomStream);
+					DynamicPayload.DynamicParameterValue[1] = GetParameterValue(Context, DynParam1, Particle, Owner, &RandomStream);
 				}
 				END_UPDATE_LOOP;
 			}
@@ -266,9 +266,9 @@ void UParticleModuleParameterDynamic::Update(const FUpdateContext& Context)
 					FEmitterDynamicParameterPayload& DynamicPayload = *((FEmitterDynamicParameterPayload*)(ParticleBase + CurrentOffset));
 					FPlatformMisc::Prefetch(ParticleData, (ParticleIndices[i+1] * ParticleStride));
 					FPlatformMisc::Prefetch(ParticleData, (ParticleIndices[i+1] * ParticleStride) + PLATFORM_CACHE_LINE_SIZE);
-					DynamicPayload.DynamicParameterValue[0] = GetParameterValue_UserSet(DynParam0, Particle, Owner, &RandomStream);
-					DynamicPayload.DynamicParameterValue[1] = GetParameterValue_UserSet(DynParam1, Particle, Owner, &RandomStream);
-					DynamicPayload.DynamicParameterValue[2] = GetParameterValue_UserSet(DynParam2, Particle, Owner, &RandomStream);
+					DynamicPayload.DynamicParameterValue[0] = GetParameterValue_UserSet(Context, DynParam0, Particle, Owner, &RandomStream);
+					DynamicPayload.DynamicParameterValue[1] = GetParameterValue_UserSet(Context, DynParam1, Particle, Owner, &RandomStream);
+					DynamicPayload.DynamicParameterValue[2] = GetParameterValue_UserSet(Context, DynParam2, Particle, Owner, &RandomStream);
 				}
 				END_UPDATE_LOOP;
 			}
@@ -279,9 +279,9 @@ void UParticleModuleParameterDynamic::Update(const FUpdateContext& Context)
 					FEmitterDynamicParameterPayload& DynamicPayload = *((FEmitterDynamicParameterPayload*)(ParticleBase + CurrentOffset));
 					FPlatformMisc::Prefetch(ParticleData, (ParticleIndices[i+1] * ParticleStride));
 					FPlatformMisc::Prefetch(ParticleData, (ParticleIndices[i+1] * ParticleStride) + PLATFORM_CACHE_LINE_SIZE);
-					DynamicPayload.DynamicParameterValue[0] = GetParameterValue(DynParam0, Particle, Owner, &RandomStream);
-					DynamicPayload.DynamicParameterValue[1] = GetParameterValue(DynParam1, Particle, Owner, &RandomStream);
-					DynamicPayload.DynamicParameterValue[2] = GetParameterValue(DynParam2, Particle, Owner, &RandomStream);
+					DynamicPayload.DynamicParameterValue[0] = GetParameterValue(Context, DynParam0, Particle, Owner, &RandomStream);
+					DynamicPayload.DynamicParameterValue[1] = GetParameterValue(Context, DynParam1, Particle, Owner, &RandomStream);
+					DynamicPayload.DynamicParameterValue[2] = GetParameterValue(Context, DynParam2, Particle, Owner, &RandomStream);
 				}
 				END_UPDATE_LOOP;
 			}
@@ -300,10 +300,10 @@ void UParticleModuleParameterDynamic::Update(const FUpdateContext& Context)
 					FEmitterDynamicParameterPayload& DynamicPayload = *((FEmitterDynamicParameterPayload*)(ParticleBase + CurrentOffset));
 					FPlatformMisc::Prefetch(ParticleData, (ParticleIndices[i+1] * ParticleStride));
 					FPlatformMisc::Prefetch(ParticleData, (ParticleIndices[i+1] * ParticleStride) + PLATFORM_CACHE_LINE_SIZE);
-					DynamicPayload.DynamicParameterValue[0] = GetParameterValue_UserSet(DynParam0, Particle, Owner, &RandomStream);
-					DynamicPayload.DynamicParameterValue[1] = GetParameterValue_UserSet(DynParam1, Particle, Owner, &RandomStream);
-					DynamicPayload.DynamicParameterValue[2] = GetParameterValue_UserSet(DynParam2, Particle, Owner, &RandomStream);
-					DynamicPayload.DynamicParameterValue[3] = GetParameterValue_UserSet(DynParam3, Particle, Owner, &RandomStream);
+					DynamicPayload.DynamicParameterValue[0] = GetParameterValue_UserSet(Context, DynParam0, Particle, Owner, &RandomStream);
+					DynamicPayload.DynamicParameterValue[1] = GetParameterValue_UserSet(Context, DynParam1, Particle, Owner, &RandomStream);
+					DynamicPayload.DynamicParameterValue[2] = GetParameterValue_UserSet(Context, DynParam2, Particle, Owner, &RandomStream);
+					DynamicPayload.DynamicParameterValue[3] = GetParameterValue_UserSet(Context, DynParam3, Particle, Owner, &RandomStream);
 				}
 				END_UPDATE_LOOP;
 			}
@@ -314,10 +314,10 @@ void UParticleModuleParameterDynamic::Update(const FUpdateContext& Context)
 					FEmitterDynamicParameterPayload& DynamicPayload = *((FEmitterDynamicParameterPayload*)(ParticleBase + CurrentOffset));
 					FPlatformMisc::Prefetch(ParticleData, (ParticleIndices[i+1] * ParticleStride));
 					FPlatformMisc::Prefetch(ParticleData, (ParticleIndices[i+1] * ParticleStride) + PLATFORM_CACHE_LINE_SIZE);
-					DynamicPayload.DynamicParameterValue[0] = GetParameterValue(DynParam0, Particle, Owner, &RandomStream);
-					DynamicPayload.DynamicParameterValue[1] = GetParameterValue(DynParam1, Particle, Owner, &RandomStream);
-					DynamicPayload.DynamicParameterValue[2] = GetParameterValue(DynParam2, Particle, Owner, &RandomStream);
-					DynamicPayload.DynamicParameterValue[3] = GetParameterValue(DynParam3, Particle, Owner, &RandomStream);
+					DynamicPayload.DynamicParameterValue[0] = GetParameterValue(Context, DynParam0, Particle, Owner, &RandomStream);
+					DynamicPayload.DynamicParameterValue[1] = GetParameterValue(Context, DynParam1, Particle, Owner, &RandomStream);
+					DynamicPayload.DynamicParameterValue[2] = GetParameterValue(Context, DynParam2, Particle, Owner, &RandomStream);
+					DynamicPayload.DynamicParameterValue[3] = GetParameterValue(Context, DynParam3, Particle, Owner, &RandomStream);
 				}
 				END_UPDATE_LOOP;
 			}
@@ -336,10 +336,10 @@ void UParticleModuleParameterDynamic::Update(const FUpdateContext& Context)
 					FEmitterDynamicParameterPayload& DynamicPayload = *((FEmitterDynamicParameterPayload*)(ParticleBase + CurrentOffset));
 					FPlatformMisc::Prefetch(ParticleData, (ParticleIndices[i+1] * ParticleStride));
 					FPlatformMisc::Prefetch(ParticleData, (ParticleIndices[i+1] * ParticleStride) + PLATFORM_CACHE_LINE_SIZE);
-					DynamicPayload.DynamicParameterValue[0] = (UpdateFlags & EDPU_UPDATE_0) ? GetParameterValue_UserSet(DynParam0, Particle, Owner, &RandomStream) : DynamicPayload.DynamicParameterValue[0];
-					DynamicPayload.DynamicParameterValue[1] = (UpdateFlags & EDPU_UPDATE_1) ? GetParameterValue_UserSet(DynParam1, Particle, Owner, &RandomStream) : DynamicPayload.DynamicParameterValue[1];
-					DynamicPayload.DynamicParameterValue[2] = (UpdateFlags & EDPU_UPDATE_2) ? GetParameterValue_UserSet(DynParam2, Particle, Owner, &RandomStream) : DynamicPayload.DynamicParameterValue[2];
-					DynamicPayload.DynamicParameterValue[3] = (UpdateFlags & EDPU_UPDATE_3) ? GetParameterValue_UserSet(DynParam3, Particle, Owner, &RandomStream) : DynamicPayload.DynamicParameterValue[3];
+					DynamicPayload.DynamicParameterValue[0] = (UpdateFlags & EDPU_UPDATE_0) ? GetParameterValue_UserSet(Context, DynParam0, Particle, Owner, &RandomStream) : DynamicPayload.DynamicParameterValue[0];
+					DynamicPayload.DynamicParameterValue[1] = (UpdateFlags & EDPU_UPDATE_1) ? GetParameterValue_UserSet(Context, DynParam1, Particle, Owner, &RandomStream) : DynamicPayload.DynamicParameterValue[1];
+					DynamicPayload.DynamicParameterValue[2] = (UpdateFlags & EDPU_UPDATE_2) ? GetParameterValue_UserSet(Context, DynParam2, Particle, Owner, &RandomStream) : DynamicPayload.DynamicParameterValue[2];
+					DynamicPayload.DynamicParameterValue[3] = (UpdateFlags & EDPU_UPDATE_3) ? GetParameterValue_UserSet(Context, DynParam3, Particle, Owner, &RandomStream) : DynamicPayload.DynamicParameterValue[3];
 				}
 				END_UPDATE_LOOP;
 			}
@@ -350,10 +350,10 @@ void UParticleModuleParameterDynamic::Update(const FUpdateContext& Context)
 					FEmitterDynamicParameterPayload& DynamicPayload = *((FEmitterDynamicParameterPayload*)(ParticleBase + CurrentOffset));
 					FPlatformMisc::Prefetch(ParticleData, (ParticleIndices[i+1] * ParticleStride));
 					FPlatformMisc::Prefetch(ParticleData, (ParticleIndices[i+1] * ParticleStride) + PLATFORM_CACHE_LINE_SIZE);
-					DynamicPayload.DynamicParameterValue[0] = (UpdateFlags & EDPU_UPDATE_0) ? GetParameterValue(DynParam0, Particle, Owner, &RandomStream) : DynamicPayload.DynamicParameterValue[0];
-					DynamicPayload.DynamicParameterValue[1] = (UpdateFlags & EDPU_UPDATE_1) ? GetParameterValue(DynParam1, Particle, Owner, &RandomStream) : DynamicPayload.DynamicParameterValue[1];
-					DynamicPayload.DynamicParameterValue[2] = (UpdateFlags & EDPU_UPDATE_2) ? GetParameterValue(DynParam2, Particle, Owner, &RandomStream) : DynamicPayload.DynamicParameterValue[2];
-					DynamicPayload.DynamicParameterValue[3] = (UpdateFlags & EDPU_UPDATE_3) ? GetParameterValue(DynParam3, Particle, Owner, &RandomStream) : DynamicPayload.DynamicParameterValue[3];
+					DynamicPayload.DynamicParameterValue[0] = (UpdateFlags & EDPU_UPDATE_0) ? GetParameterValue(Context, DynParam0, Particle, Owner, &RandomStream) : DynamicPayload.DynamicParameterValue[0];
+					DynamicPayload.DynamicParameterValue[1] = (UpdateFlags & EDPU_UPDATE_1) ? GetParameterValue(Context, DynParam1, Particle, Owner, &RandomStream) : DynamicPayload.DynamicParameterValue[1];
+					DynamicPayload.DynamicParameterValue[2] = (UpdateFlags & EDPU_UPDATE_2) ? GetParameterValue(Context, DynParam2, Particle, Owner, &RandomStream) : DynamicPayload.DynamicParameterValue[2];
+					DynamicPayload.DynamicParameterValue[3] = (UpdateFlags & EDPU_UPDATE_3) ? GetParameterValue(Context, DynParam3, Particle, Owner, &RandomStream) : DynamicPayload.DynamicParameterValue[3];
 				}
 				END_UPDATE_LOOP;
 			}
diff --git a/Engine/Source/Runtime/Engine/Private/Particles/ParticleModules_Velocity.cpp b/Engine/Source/Runtime/Engine/Private/Particles/ParticleModules_Velocity.cpp
--- a/Engine/Source/Runtime/Engine/Private/Particles/ParticleModules_Velocity.cpp
+++ b/Engine/Source/Runtime/Engine/Private/Particles/ParticleModules_Velocity.cpp
@@ -88,7 +88,7 @@ void UParticleModuleVelocity::SpawnEx(const FSpawnContext& Context, struct FRand
 		FVector FromOrigin = (Particle.Location - Owner->EmitterToSimulation.GetOrigin()).GetSafeNormal();
 
 		FVector OwnerScale(1.0f);
-		if ((bApplyOwnerScale == true) && Owner->Component)
+		if (bApplyOwnerScale == true)
 		{
 			OwnerScale = Context.GetTransform().GetScale3D();
 		}
@@ -241,11 +241,9 @@ void UParticleModuleVelocityOverLifetime::Spawn(const FSpawnContext& Context)
 {
 	if (Absolute)
 	{
-		FParticleEmitterInstance* Owner = &Context.Owner;
-
 		SPAWN_INIT;
 		FVector OwnerScale(1.0f);
-		if ((bApplyOwnerScale == true) && Owner && Owner->Component)
+		if (bApplyOwnerScale == true)
 		{
 			OwnerScale = Context.GetTransform().GetScale3D();
 		}
@@ -258,7 +256,6 @@ void UParticleModuleVelocityOverLifetime::Spawn(const FSpawnContext& Context)
 void UParticleModuleVelocityOverLifetime::Update(const FUpdateContext& Context)
 {
 	FParticleEmitterInstance* Owner = &Context.Owner;
-	check(Owner && Owner->Component);
 	UParticleLODLevel* LODLevel	= Owner->SpriteTemplate->GetCurrentLODLevel(Owner);
 	check(LODLevel);
 	FVector OwnerScale(1.0f);
@@ -424,7 +421,7 @@ void UParticleModuleVelocityCone::SpawnEx(const FSpawnContext& Context, struct F
 	UParticleLODLevel* LODLevel	= Owner->SpriteTemplate->GetCurrentLODLevel(Owner);
 	check(LODLevel);
 	FVector OwnerScale(1.0f);
-	if ((bApplyOwnerScale == true) && Owner->Component)
+	if (bApplyOwnerScale == true)
 	{
 		OwnerScale = Context.GetTransform().GetScale3D();
 	}
diff --git a/Engine/Source/Runtime/Engine/Public/Engine.h b/Engine/Source/Runtime/Engine/Public/Engine.h
--- a/Engine/Source/Runtime/Engine/Public/Engine.h
+++ b/Engine/Source/Runtime/Engine/Public/Engine.h
@@ -73,7 +73,6 @@ MONOLITHIC_HEADER_BOILERPLATE()
 #include "MeshBatch.h"
 #include "SceneUtils.h"
 #include "MeshParticleVertexFactory.h"
-#include "ParticleHelper.h"
 #include "Distributions.h"
 #include "ParticleEmitterInstances.h"
 #include "Scalability.h"
@@ -332,4 +331,8 @@ MONOLITHIC_HEADER_BOILERPLATE()
 #include "SlateCore.h"
 #include "SlateBasics.h"
 
+#if UE_ENABLE_INCLUDE_ORDER_DEPRECATED_IN_5_6
+#include "ParticleHelper.h"
+#endif
+
 // IWYU pragma: end_keep
diff --git a/Engine/Source/Runtime/Engine/Public/EngineMinimal.h b/Engine/Source/Runtime/Engine/Public/EngineMinimal.h
--- a/Engine/Source/Runtime/Engine/Public/EngineMinimal.h
+++ b/Engine/Source/Runtime/Engine/Public/EngineMinimal.h
@@ -91,7 +91,6 @@ MONOLITHIC_HEADER_BOILERPLATE()
 #include "SceneUtils.h"
 #include "SceneManagement.h"
 #include "MeshParticleVertexFactory.h"
-#include "ParticleHelper.h"
 #include "Distributions.h"
 #include "ParticleEmitterInstances.h"
 #include "Scalability.h"
@@ -115,4 +114,8 @@ MONOLITHIC_HEADER_BOILERPLATE()
 #include "Sound/SoundBase.h"
 #include "Sound/SoundWave.h"
 
+#if UE_ENABLE_INCLUDE_ORDER_DEPRECATED_IN_5_6
+#include "ParticleHelper.h"
+#endif
+
 // IWYU pragma: end_keep
diff --git a/Engine/Source/Runtime/Engine/Public/ParticleHelper.h b/Engine/Source/Runtime/Engine/Public/ParticleHelper.h
--- a/Engine/Source/Runtime/Engine/Public/ParticleHelper.h
+++ b/Engine/Source/Runtime/Engine/Public/ParticleHelper.h
@@ -76,7 +76,6 @@ inline const int32 PDM_DefaultValue = 0xFFFF;
 
 #define BEGIN_UPDATE_LOOP																								\
 	{																													\
-		check(Context.Owner.Component != NULL);																			\
 		int32&			ActiveParticles = Context.Owner.ActiveParticles;												\
 		int32			Offset	= Context.Offset;																		\
 		uint32			CurrentOffset	= Offset;																		\
@@ -103,7 +102,6 @@ inline const int32 PDM_DefaultValue = 0xFFFF;
 		continue;
 
 #define SPAWN_INIT																										\
-	check(Context.Owner.Component != NULL);																				\
 	const int32		ActiveParticles	= Context.Owner.ActiveParticles;													\
 	const uint32	ParticleStride	= Context.Owner.ParticleStride;														\
 	uint32			CurrentOffset	= Context.Offset;																	\
