World Partition - HLOD: Fixed issue with invalid HLOD configurations that could create meshes with sections that have no materials assigned
* Some types of HLOD layers aren't meant to be used as parent layers of others. Instancing for example is a good example - using instancing for HLOD1 when HLOD0 is built out of merged meshes makes no sense. The merged assets are private inside the OFPA packages (and the HLOD0 meshes are all unique so wouldn't benefit from instancing anyway).
* Some configuration of HLOD layers are invalid. For example, using merged meshes without merging materials for an HLOD1 would mean that the generated meshes for that layer would try to use the HLOD0 materials directly.
* Added data validation/errors and safeguards to alert users and prevent these scenarios from causing issues.

diff --git a/Engine/Plugins/Editor/WorldPartitionHLODUtilities/Source/Public/WorldPartition/HLOD/Builders/HLODBuilderInstancing.h b/Engine/Plugins/Editor/WorldPartitionHLODUtilities/Source/Public/WorldPartition/HLOD/Builders/HLODBuilderInstancing.h
--- a/Engine/Plugins/Editor/WorldPartitionHLODUtilities/Source/Public/WorldPartition/HLOD/Builders/HLODBuilderInstancing.h
+++ b/Engine/Plugins/Editor/WorldPartitionHLODUtilities/Source/Public/WorldPartition/HLOD/Builders/HLODBuilderInstancing.h
@@ -14,6 +14,11 @@ class WORLDPARTITIONHLODUTILITIES_API UHLODBuilderInstancingSettings : public UH
 
 	virtual uint32 GetCRC() const override;
 
+	/**
+     * By definition, instancing means reusing the source assets.
+     */
+	virtual bool IsReusingSourceMaterials() const override { return true; }
+
 	/**
 	 * If enabled, the components created for the HLODs will not use Nanite.
 	 * Necessary if you want to use the last LOD & the mesh is Nanite enabled, as forced LODs are ignored by Nanite
diff --git a/Engine/Plugins/Editor/WorldPartitionHLODUtilities/Source/Public/WorldPartition/HLOD/Builders/HLODBuilderMeshMerge.h b/Engine/Plugins/Editor/WorldPartitionHLODUtilities/Source/Public/WorldPartition/HLOD/Builders/HLODBuilderMeshMerge.h
--- a/Engine/Plugins/Editor/WorldPartitionHLODUtilities/Source/Public/WorldPartition/HLOD/Builders/HLODBuilderMeshMerge.h
+++ b/Engine/Plugins/Editor/WorldPartitionHLODUtilities/Source/Public/WorldPartition/HLOD/Builders/HLODBuilderMeshMerge.h
@@ -18,6 +18,12 @@ class WORLDPARTITIONHLODUTILITIES_API UHLODBuilderMeshMergeSettings : public UHL
 
 	virtual uint32 GetCRC() const override;
 
+	/**
+     * Merge mesh will only reuse the source materials when not merging the materials. In this case, the created mesh
+     * will have multiple sections, with each of them directly using the source materials.
+     */
+    virtual bool IsReusingSourceMaterials() const { return !MeshMergeSettings.bMergeMaterials; }
+
 	/** Merged mesh generation settings */
 	UPROPERTY(EditAnywhere, Category = HLOD)
 	FMeshMergingSettings MeshMergeSettings;
diff --git a/Engine/Source/Developer/MeshMergeUtilities/Private/MeshMergeUtilities.cpp b/Engine/Source/Developer/MeshMergeUtilities/Private/MeshMergeUtilities.cpp
--- a/Engine/Source/Developer/MeshMergeUtilities/Private/MeshMergeUtilities.cpp
+++ b/Engine/Source/Developer/MeshMergeUtilities/Private/MeshMergeUtilities.cpp
@@ -2642,6 +2642,7 @@ void FMeshMergeUtilities::MergeComponentsToStaticMesh(const TArray<UPrimitiveCom
 				}
 				else
 				{
+					UE_LOG(LogMeshMerging, Error, TEXT("Merged mesh is referencing a non-asset material (\"%s\"). This material can't be referenced outside it's package as it is not public. Clearing out material entry in the final mesh."), *Material->GetPathName());
 					Material = nullptr; // do not save non-asset materials
 				}
 			}
diff --git a/Engine/Source/Runtime/Engine/Private/WorldPartition/HLOD/HLODLayer.cpp b/Engine/Source/Runtime/Engine/Private/WorldPartition/HLOD/HLODLayer.cpp
--- a/Engine/Source/Runtime/Engine/Private/WorldPartition/HLOD/HLODLayer.cpp
+++ b/Engine/Source/Runtime/Engine/Private/WorldPartition/HLOD/HLODLayer.cpp
@@ -9,10 +9,13 @@
 
 #if WITH_EDITOR
 #include "Engine/World.h"
+#include "Framework/Notifications/NotificationManager.h"
 #include "Misc/ConfigCacheIni.h"
+#include "Misc/DataValidation.h"
 #include "Misc/StringFormatArg.h"
 #include "Modules/ModuleManager.h"
 #include "UObject/UnrealType.h"
+#include "Widgets/Notifications/SNotificationList.h"
 #include "WorldPartition/HLOD/IWorldPartitionHLODUtilities.h"
 #include "WorldPartition/HLOD/IWorldPartitionHLODUtilitiesModule.h"
 #include "WorldPartition/WorldPartition.h"
@@ -22,6 +25,8 @@
 
 DEFINE_LOG_CATEGORY_STATIC(LogHLODLayer, Log, All);
 
+#define LOCTEXT_NAMESPACE "HLODLayer"
+
 UHLODLayer::UHLODLayer(const FObjectInitializer& ObjectInitializer)
 	: Super(ObjectInitializer)
 	, bIsSpatiallyLoaded(true)
@@ -105,21 +110,24 @@ void UHLODLayer::PostLoad()
 {
 	Super::PostLoad();
 
-	IWorldPartitionHLODUtilitiesModule* WPHLODUtilitiesModule = FModuleManager::Get().LoadModulePtr<IWorldPartitionHLODUtilitiesModule>("WorldPartitionHLODUtilities");
-	if (IWorldPartitionHLODUtilities* WPHLODUtilities = WPHLODUtilitiesModule != nullptr ? WPHLODUtilitiesModule->GetUtilities() : nullptr)
+	if (!IsTemplate())
 	{
-		const UClass* BuilderClass = WPHLODUtilities->GetHLODBuilderClass(this);
-		const UClass* BuilderSettingsClass = BuilderClass ? BuilderClass->GetDefaultObject<UHLODBuilder>()->GetSettingsClass() : nullptr;
-
-		if (!HLODBuilderSettings || (BuilderSettingsClass && !HLODBuilderSettings->IsA(BuilderSettingsClass)))
+		IWorldPartitionHLODUtilitiesModule* WPHLODUtilitiesModule = FModuleManager::Get().LoadModulePtr<IWorldPartitionHLODUtilitiesModule>("WorldPartitionHLODUtilities");
+		if (IWorldPartitionHLODUtilities* WPHLODUtilities = WPHLODUtilitiesModule != nullptr ? WPHLODUtilitiesModule->GetUtilities() : nullptr)
 		{
-			HLODBuilderSettings = WPHLODUtilities->CreateHLODBuilderSettings(this);
+			const UClass* BuilderClass = WPHLODUtilities->GetHLODBuilderClass(this);
+			const UClass* BuilderSettingsClass = BuilderClass ? BuilderClass->GetDefaultObject<UHLODBuilder>()->GetSettingsClass() : nullptr;
+
+			if (!HLODBuilderSettings || (BuilderSettingsClass && !HLODBuilderSettings->IsA(BuilderSettingsClass)))
+			{
+				HLODBuilderSettings = WPHLODUtilities->CreateHLODBuilderSettings(this);
+			}
 		}
-	}
 
-	if (bAlwaysLoaded_DEPRECATED)
-	{
-		bIsSpatiallyLoaded = false;
+		if (bAlwaysLoaded_DEPRECATED)
+		{
+			bIsSpatiallyLoaded = false;
+		}
 	}
 }
 
@@ -134,6 +142,24 @@ void UHLODLayer::DeclareConstructClasses(TArray<FTopLevelAssetPath>& OutConstruc
 }
 #endif
 
+// Some types of HLOD layers aren't meant to be used as parent layers of others.
+// Instancing for example is a good example - using instancing for HLOD1 when HLOD0 is built out of merged meshes makes no sense. 
+// The merged assets are private inside the OFPA packages (and the HLOD0 meshes are all unique so wouldn't benefit from instancing anyway).
+// Some configuration of HLOD layers are invalid. For example, using merged meshes without merging materials for an HLOD1 would mean
+// that the generated meshes for that layer would try to use the HLOD0 materials directly.
+static bool IsInvalidSourceMaterialReuse(const UHLODLayer* ParentLayer, const UHLODLayer* Layer)
+{
+	if (ParentLayer && ParentLayer->GetHLODBuilderSettings() && Layer->GetHLODBuilderSettings())
+	{
+		if (ParentLayer->GetHLODBuilderSettings()->IsReusingSourceMaterials() && !Layer->GetHLODBuilderSettings()->IsReusingSourceMaterials())
+		{
+			return true;
+		}
+	}
+
+	return false;
+}
+
 void UHLODLayer::PostEditChangeProperty(FPropertyChangedEvent& PropertyChangedEvent)
 {
 	Super::PostEditChangeProperty(PropertyChangedEvent);
@@ -153,6 +179,8 @@ void UHLODLayer::PostEditChangeProperty(FPropertyChangedEvent& PropertyChangedEv
 	}
 	else if (PropertyName == GET_MEMBER_NAME_CHECKED(UHLODLayer, ParentLayer))
 	{
+		bool bInvalidParentLayer = false;
+
 		TSet<const UHLODLayer*> VisitedHLODLayers;
 		const UHLODLayer* CurHLODLayer = ParentLayer;
 		while (CurHLODLayer)
@@ -161,13 +189,46 @@ void UHLODLayer::PostEditChangeProperty(FPropertyChangedEvent& PropertyChangedEv
 			VisitedHLODLayers.Add(CurHLODLayer, &bHLODLayerWasAlreadyInSet);
 			if (bHLODLayerWasAlreadyInSet)
 			{
-				UE_LOG(LogHLODLayer, Warning, TEXT("Circular HLOD parent chain detected: HLODLayer=%s ParentLayer=%s"), *GetName(), *ParentLayer->GetName());
-				ParentLayer = nullptr;
+				bInvalidParentLayer = true;
+				UE_LOG(LogHLODLayer, Error, TEXT("Circular HLOD parent chain detected: HLODLayer=%s ParentLayer=%s"), *GetName(), *ParentLayer->GetName());
 				break;
 			}
 			CurHLODLayer = CurHLODLayer->GetParentLayer();
 		}
+
+		if (IsInvalidSourceMaterialReuse(ParentLayer, this))
+		{
+			UE_LOG(LogHLODLayer, Error, TEXT("Invalid HLOD settings. Parent layer %s will reuse private materials created for this HLOD layer. Common error is to use an \"instancing\" parent layer, or a \"merged\" parent layer set to not merge materials. Change the type of HLOD generated by the parent layer or assign a new parent layer."), *ParentLayer->GetPathName());
+			bInvalidParentLayer = true;
+		}
+
+		if (bInvalidParentLayer)
+		{
+			ParentLayer = nullptr;
+
+			FText FormattedMessage = LOCTEXT("NotifyInvalidHLODParentLayer", "Invalid Parent HLOD Layer specified, see log for more details.");
+
+			// Show toast.
+			FNotificationInfo Info(FormattedMessage);
+			Info.ExpireDuration = 5.0f;
+			Info.bFireAndForget = true;
+			Info.Image = FCoreStyle::Get().GetBrush(TEXT("MessageLog.Error"));
+			FSlateNotificationManager::Get().AddNotification(Info);
+		}
+	}
+}
+
+EDataValidationResult UHLODLayer::IsDataValid(class FDataValidationContext& Context) const
+{
+	EDataValidationResult Result = CombineDataValidationResults(Super::IsDataValid(Context), EDataValidationResult::Valid);
+
+	if (IsInvalidSourceMaterialReuse(ParentLayer, this))
+	{
+		Result = EDataValidationResult::Invalid;
+		Context.AddError(LOCTEXT("ParentReusingSourceMaterial", "Invalid HLOD settings. Assigned parent layer will reuse private materials created for this HLOD layer. Common error is to use an \"instancing\" parent layer, or a \"merged\" parent layer set to not merge materials. Change the type of HLOD generated by the parent layer or assign a new parent layer."));
 	}
+
+	return Result;
 }
 
 FName UHLODLayer::GetRuntimeGridName(uint32 InLODLevel, int32 InCellSize, double InLoadingRange)
@@ -180,3 +241,5 @@ FName UHLODLayer::GetRuntimeGrid(uint32 InHLODLevel) const
 	return !IsSpatiallyLoaded() ? NAME_None : GetRuntimeGridName(InHLODLevel, CellSize, LoadingRange);
 }
 #endif // WITH_EDITOR
+
+#undef LOCTEXT_NAMESPACE
\ No newline at end of file
diff --git a/Engine/Source/Runtime/Engine/Private/WorldPartition/HLOD/HLODRuntimeSubsystem.cpp b/Engine/Source/Runtime/Engine/Private/WorldPartition/HLOD/HLODRuntimeSubsystem.cpp
--- a/Engine/Source/Runtime/Engine/Private/WorldPartition/HLOD/HLODRuntimeSubsystem.cpp
+++ b/Engine/Source/Runtime/Engine/Private/WorldPartition/HLOD/HLODRuntimeSubsystem.cpp
@@ -415,10 +415,13 @@ static void PrepareVTRequests(TMap<UMaterialInterface*, float>& InOutVTRequests,
 
 	for (UMaterialInterface* MaterialInterface : InStaticMeshComponent->GetMaterials())
 	{
-		// We have a VT we'd like to prefetch, add or update a request in our request map.
-		// If the texture was already requested by another component, fetch the highest required resolution only.
-		float& CurrentMaxPixel = InOutVTRequests.FindOrAdd(MaterialInterface);
-		CurrentMaxPixel = FMath::Max(CurrentMaxPixel, PixelSize);
+		if (MaterialInterface)
+		{
+			// We have a VT we'd like to prefetch, add or update a request in our request map.
+			// If the texture was already requested by another component, fetch the highest required resolution only.
+			float& CurrentMaxPixel = InOutVTRequests.FindOrAdd(MaterialInterface);
+			CurrentMaxPixel = FMath::Max(CurrentMaxPixel, PixelSize);
+		}
 	}
 }
 
diff --git a/Engine/Source/Runtime/Engine/Public/WorldPartition/HLOD/HLODBuilder.h b/Engine/Source/Runtime/Engine/Public/WorldPartition/HLOD/HLODBuilder.h
--- a/Engine/Source/Runtime/Engine/Public/WorldPartition/HLOD/HLODBuilder.h
+++ b/Engine/Source/Runtime/Engine/Public/WorldPartition/HLOD/HLODBuilder.h
@@ -29,6 +29,13 @@ public:
 	 * This is used to detect settings changes and trigger an HLOD rebuild if necessarry.
 	 */
 	virtual uint32 GetCRC() const { return 0; }
+
+	/**
+	 * Whether these HLOD settings are going to result in the created mesh referencing external materials.
+	 * This is used to prevent HLOD[N] from referencing private materials from an HLOD[N-1]
+	 * By default most HLOD methods are baking down their textures and uses their own unique material instance.
+	 */
+	virtual bool IsReusingSourceMaterials() const { return false; }
 #endif // WITH_EDITOR
 };
 
diff --git a/Engine/Source/Runtime/Engine/Public/WorldPartition/HLOD/HLODLayer.h b/Engine/Source/Runtime/Engine/Public/WorldPartition/HLOD/HLODLayer.h
--- a/Engine/Source/Runtime/Engine/Public/WorldPartition/HLOD/HLODLayer.h
+++ b/Engine/Source/Runtime/Engine/Public/WorldPartition/HLOD/HLODLayer.h
@@ -78,6 +78,7 @@ private:
 	static ENGINE_API void DeclareConstructClasses(TArray<FTopLevelAssetPath>& OutConstructClasses, const UClass* SpecificSubclass);
 #endif
 	ENGINE_API virtual void PostEditChangeProperty(FPropertyChangedEvent& PropertyChangedEvent) override;
+	virtual EDataValidationResult IsDataValid(class FDataValidationContext& Context) const override;
 	//~ End UObject Interface.
 #endif
 
