Fix/Workaround for crash in chaos physics

diff --git a/Engine/Source/Runtime/Engine/Private/PhysicsEngine/Experimental/PhysScene_Chaos.cpp b/Engine/Source/Runtime/Engine/Private/PhysicsEngine/Experimental/PhysScene_Chaos.cpp
--- a/Engine/Source/Runtime/Engine/Private/PhysicsEngine/Experimental/PhysScene_Chaos.cpp
+++ b/Engine/Source/Runtime/Engine/Private/PhysicsEngine/Experimental/PhysScene_Chaos.cpp
@@ -1984,27 +1984,34 @@ void FPhysScene_Chaos::UpdateKinematicsOnDeferredSkelMeshes()
 						const int32 BoneIndex = BodyInst->InstanceBoneIndex;
 						if (BoneIndex != INDEX_NONE)
 						{
-							const FTransform BoneTransform = ComponentSpaceTransforms[BoneIndex] * CurrentLocalToWorld;
+							if (ComponentSpaceTransforms.IsValidIndex(BoneIndex))
+							{
+								const FTransform BoneTransform = ComponentSpaceTransforms[BoneIndex] * CurrentLocalToWorld;
 
-							TeleportActorsPool[ActorPoolStartIndex + i] = ActorHandle;
+								TeleportActorsPool[ActorPoolStartIndex + i] = ActorHandle;
 
-							// TODO: Kinematic targets. Check Teleport type on FDeferredKinematicUpdateInfo and don't always teleport.
-							Body_External.SetX(BoneTransform.GetLocation(), false);	// only set dirty once in SetR
-							Body_External.SetR(BoneTransform.GetRotation());
-							Body_External.UpdateShapeBounds(BoneTransform);
+								// TODO: Kinematic targets. Check Teleport type on FDeferredKinematicUpdateInfo and don't always teleport.
+								Body_External.SetX(BoneTransform.GetLocation(), false);	// only set dirty once in SetR
+								Body_External.SetR(BoneTransform.GetRotation());
+								Body_External.UpdateShapeBounds(BoneTransform);
 
-							if (!PhysicsAsset->SkeletalBodySetups[i]->bSkipScaleFromAnimation)
-							{
-								const FVector& MeshScale3D = CurrentLocalToWorld.GetScale3D();
-								if (MeshScale3D.IsUniform())
+								if (!PhysicsAsset->SkeletalBodySetups[i]->bSkipScaleFromAnimation)
 								{
-									BodiesUpdatingScale.Enqueue(BodyInstanceScalePair({ BodyInst, BoneTransform.GetScale3D() }));
-								}
-								else
-								{
-									BodiesUpdatingScale.Enqueue(BodyInstanceScalePair({ BodyInst, MeshScale3D }));
+									const FVector& MeshScale3D = CurrentLocalToWorld.GetScale3D();
+									if (MeshScale3D.IsUniform())
+									{
+										BodiesUpdatingScale.Enqueue(BodyInstanceScalePair({ BodyInst, BoneTransform.GetScale3D() }));
+									}
+									else
+									{
+										BodiesUpdatingScale.Enqueue(BodyInstanceScalePair({ BodyInst, MeshScale3D }));
+									}
 								}
 							}
+							else
+							{
+								UE_LOG(LogChaos,Log,TEXT("Bone Index %d out of bounds (%d Bones)"), BoneIndex, ComponentSpaceTransforms.Num());
+							}
 						}
 					}
 				}
