Integrated changes from licensee to custom show flag SetEnabled

#rnx

diff --git a/Engine/Source/Runtime/Engine/Public/ShowFlags.h b/Engine/Source/Runtime/Engine/Public/ShowFlags.h
--- a/Engine/Source/Runtime/Engine/Public/ShowFlags.h
+++ b/Engine/Source/Runtime/Engine/Public/ShowFlags.h
@@ -645,6 +645,29 @@ struct TCustomShowFlag
 		}
 	}
 
+	bool SetEnabled(FEngineShowFlags& ShowFlags, bool bSet)
+	{
+#if UE_BUILD_SHIPPING
+		if constexpr (ShippingValue == EShowFlagShippingValue::ForceDisabled ||
+			ShippingValue == EShowFlagShippingValue::ForceEnabled)
+		{
+			return false;
+		}
+		else
+#endif
+		{
+			if (FlagIndex == FEngineShowFlags::ECustomShowFlag::None)
+			{
+				return false;
+			}
+			else
+			{
+				ShowFlags.SetSingleFlag((uint32)FlagIndex + FEngineShowFlags::SF_FirstCustom, bSet);
+				return true;
+			}
+		}
+	}
+
 private:
 	FEngineShowFlags::ECustomShowFlag FlagIndex = FEngineShowFlags::ECustomShowFlag::None;
 };
