Fix warnings appearing after creating a new Empty Open World Level
Don't try to scan actor paths when ActorDescContainer package is a temp package

diff --git a/Engine/Source/Runtime/Engine/Private/WorldPartition/ActorDescContainer.cpp b/Engine/Source/Runtime/Engine/Private/WorldPartition/ActorDescContainer.cpp
--- a/Engine/Source/Runtime/Engine/Private/WorldPartition/ActorDescContainer.cpp
+++ b/Engine/Source/Runtime/Engine/Private/WorldPartition/ActorDescContainer.cpp
@@ -71,7 +71,7 @@ void UActorDescContainer::Initialize(const FInitializeParams& InitParams)
 
 	TArray<FAssetData> ExternalAssets;
 	TArray<FString> InternalAssets;
-	if (!ContainerPackageName.IsNone())
+	if (!ContainerPackageName.IsNone() && !FPackageName::IsTempPackage(ContainerPackageName.ToString()))
 	{
 		const FString ContainerExternalActorsPath = GetExternalActorPath();
 
