Substrate - fixed SSS prpofile not working with parameter blending.
BSDF within parameter blended substree now receive correctin formation such as bIsBottom.
Improved the blending of SSSProfile to select one of two has this has been made possible in Substrate now.

#rnx
#tests PC substrate levels and such.

diff --git a/Engine/Shaders/Private/Substrate/Substrate.ush b/Engine/Shaders/Private/Substrate/Substrate.ush
--- a/Engine/Shaders/Private/Substrate/Substrate.ush
+++ b/Engine/Shaders/Private/Substrate/Substrate.ush
@@ -2549,8 +2549,17 @@ FSubstrateData SubstrateHorizontalMixingParameterBlending(FSubstrateData Backgro
 		const bool bForegroundHasSSSProfile = BSDF_GETSSSTYPE(ForegroundBSDF) == SSS_TYPE_DIFFUSION_PROFILE;
 		const bool bBackgroundHasSSSProfile = BSDF_GETSSSTYPE(BackgroundBSDF) == SSS_TYPE_DIFFUSION_PROFILE;
 
-		// As mentioned above, we cannot blend SSSProfiles. So we select the one that triggered the profile, foreground having priority.
-		SLAB_SSSPROFILEID(ResultBSDF) = bForegroundHasSSSProfile ? SLAB_SSSPROFILEID(ForegroundBSDF) : SLAB_SSSPROFILEID(BackgroundBSDF);
+		// As mentioned above, we cannot blend SSSProfiles.
+		if (bForegroundHasSSSProfile && bBackgroundHasSSSProfile)
+		{
+			// Select the profile that has the highest mixing weight.
+			SLAB_SSSPROFILEID(ResultBSDF) = ForegroundLightingMixFactor > 0.5f ? SLAB_SSSPROFILEID(ForegroundBSDF) : SLAB_SSSPROFILEID(BackgroundBSDF);
+		}
+		else
+		{
+			// Select the only profile available.
+			SLAB_SSSPROFILEID(ResultBSDF) = bForegroundHasSSSProfile ? SLAB_SSSPROFILEID(ForegroundBSDF) : SLAB_SSSPROFILEID(BackgroundBSDF);
+		}
 
 		// We also lerp the radius scale, while taking into account whether or not the one or the other have SSSProfile.
 		SLAB_SSSPROFILERADIUSSCALE(ResultBSDF) = lerp(
diff --git a/Engine/Source/Runtime/Engine/Private/Materials/HLSLMaterialTranslator.cpp b/Engine/Source/Runtime/Engine/Private/Materials/HLSLMaterialTranslator.cpp
--- a/Engine/Source/Runtime/Engine/Private/Materials/HLSLMaterialTranslator.cpp
+++ b/Engine/Source/Runtime/Engine/Private/Materials/HLSLMaterialTranslator.cpp
@@ -12853,8 +12853,6 @@ bool FHLSLMaterialTranslator::FSubstrateCompilationContext::SubstrateGenerateDer
 		//
 		for (auto& It : SubstrateMaterialExpressionRegisteredOperators)
 		{
-			bool bMustHaveLeftChild = false;
-			bool bMustHaveRightChild = false;
 			switch (It.OperatorType)
 			{
 				// Operators without any child
@@ -13161,11 +13159,39 @@ bool FHLSLMaterialTranslator::FSubstrateCompilationContext::SubstrateGenerateDer
 
 		//
 		// Compute IsTop or IsBottom layer using a depth first tree visit while counting vertical right and left branches taken.
+		// When a BSDF is encoutered, and it is the root of a parameter blend subtree, we continue to parse, setting the same information as the root.
 		//
 		{
 			int VOpTopBranchCountTaken = 0;
 			int VOpBottomBranchCountTaken = 0;
 
+			std::function<void(FSubstrateOperator& , const FSubstrateOperator&)> PopulateParameterBlendedSubtree = [&](FSubstrateOperator& CurrentOperator, const FSubstrateOperator& RootOperator) -> void
+			{
+				switch (CurrentOperator.OperatorType)
+				{
+				case SUBSTRATE_OPERATOR_VERTICAL:
+				case SUBSTRATE_OPERATOR_HORIZONTAL:
+				case SUBSTRATE_OPERATOR_ADD:
+				{
+					PopulateParameterBlendedSubtree(SubstrateMaterialExpressionRegisteredOperators[CurrentOperator.LeftIndex], RootOperator);
+					PopulateParameterBlendedSubtree(SubstrateMaterialExpressionRegisteredOperators[CurrentOperator.RightIndex], RootOperator);
+					break;
+				}
+				case SUBSTRATE_OPERATOR_WEIGHT:
+				{
+					PopulateParameterBlendedSubtree(SubstrateMaterialExpressionRegisteredOperators[CurrentOperator.LeftIndex], RootOperator);
+					break;
+				}
+				case SUBSTRATE_OPERATOR_BSDF:
+				{
+					CurrentOperator.LayerDepth = RootOperator.LayerDepth;
+					CurrentOperator.bIsTop = RootOperator.bIsTop;
+					CurrentOperator.bIsBottom = RootOperator.bIsBottom;
+					break;
+				}
+				}
+			};
+
 			std::function<void(FSubstrateOperator&)> WalkOperators = [&](FSubstrateOperator& CurrentOperator) -> void
 			{
 				switch (CurrentOperator.OperatorType)
@@ -13198,6 +13224,18 @@ bool FHLSLMaterialTranslator::FSubstrateCompilationContext::SubstrateGenerateDer
 					CurrentOperator.LayerDepth = VOpBottomBranchCountTaken;
 					CurrentOperator.bIsTop = VopCount == 0 || VOpTopBranchCountTaken == VopCount;
 					CurrentOperator.bIsBottom = VopCount == 0 || VOpBottomBranchCountTaken == VopCount;
+					if (CurrentOperator.bRootOfParameterBlendingSubTree)
+					{
+						// Make sure to also set up those values onto the BSDF within parameter blend sub trees.
+						if (CurrentOperator.LeftIndex != INDEX_NONE)
+						{
+							PopulateParameterBlendedSubtree(SubstrateMaterialExpressionRegisteredOperators[CurrentOperator.LeftIndex], CurrentOperator);
+						}
+						if (CurrentOperator.RightIndex != INDEX_NONE)
+						{
+							PopulateParameterBlendedSubtree(SubstrateMaterialExpressionRegisteredOperators[CurrentOperator.RightIndex], CurrentOperator);
+						}
+					}
 					break;
 				}
 				}
