New material translator push.

- Converted BinaryOperator to Operator to model both binary and ternary operators.

diff --git a/Engine/Source/Runtime/Engine/Private/Materials/MaterialExpressionsIR.cpp b/Engine/Source/Runtime/Engine/Private/Materials/MaterialExpressionsIR.cpp
--- a/Engine/Source/Runtime/Engine/Private/Materials/MaterialExpressionsIR.cpp
+++ b/Engine/Source/Runtime/Engine/Private/Materials/MaterialExpressionsIR.cpp
@@ -390,10 +390,7 @@ void UMaterialExpressionAppendVector::Build(MIR::FEmitter& Emitter)
 	MIR::FValue* AValue = Emitter.InputPrimitive(&A);
 	MIR::FValue* BValue = Emitter.TryInput(&B) ? Emitter.InputPrimitive(&B) : nullptr;
 
-	if (Emitter.IsInvalid())
-	{
-		return;
-	}
+	UE_MIR_CHECKPOINT();
 
 	MIR::FPrimitiveTypePtr AType = AValue->Type->AsPrimitive();
 	if (AType->IsMatrix())
@@ -463,7 +460,7 @@ void UMaterialExpressionAppendVector::Build(MIR::FEmitter& Emitter)
 
 static void BuildBinaryArithmeticOperator(
 	MIR::FEmitter& Emitter,
-	MIR::EBinaryOperator Op,
+	MIR::EOperator Op,
 	FExpressionInput* LhsInput,
 	float LhsConst,
 	FExpressionInput* RhsInput,
@@ -487,7 +484,7 @@ static void BuildBinaryArithmeticOperator(
 	UE_MIR_CHECKPOINT();
 
 	// Finally emit the binary operator.
-	MIR::FValue* Value = Emitter.EmitBinaryOperator(Op, LhsValue, RhsValue);
+	MIR::FValue* Value = Emitter.EmitOperator(Op, LhsValue, RhsValue);
 
 	// And flow it out of the expression's only output.
 	Emitter.Output(0, Value);
@@ -591,8 +588,8 @@ void UMaterialExpressionIf::Build(MIR::FEmitter& Emitter)
 	UE_MIR_CHECKPOINT();
 
 	// Emit the comparison expressions.
-	MIR::FValue* ALessThanBConditionValue = Emitter.EmitBinaryOperator(MIR::BO_LowerThan, AValue, BValue);
-	MIR::FValue* AEqualsBConditionValue = Emitter.EmitBinaryOperator(MIR::BO_Equals, AValue, BValue);
+	MIR::FValue* ALessThanBConditionValue = Emitter.EmitOperator(MIR::BO_LowerThan, AValue, BValue);
+	MIR::FValue* AEqualsBConditionValue = Emitter.EmitOperator(MIR::BO_Equals, AValue, BValue);
 
 	// And finally emit the full conditional expression.
 	MIR::FValue* OutputValue = Emitter.EmitBranch(AEqualsBConditionValue, AEqualsBValue, AGreaterThanBValue);
@@ -684,7 +681,7 @@ void UMaterialExpressionTextureCoordinate::Build(UE::MIR::FEmitter& Emitter)
 	MIR::FValue* OutputValue = Emitter.EmitExternalInput(MIR::TexCoordIndexToExternalInput(CoordinateIndex));
 	
 	// Multiply the UV input by the UV tiling constants
-	OutputValue = Emitter.EmitBinaryOperator(MIR::BO_Multiply, OutputValue, Emitter.EmitConstantFloat2({ UTiling, VTiling }));
+	OutputValue = Emitter.EmitOperator(MIR::BO_Multiply, OutputValue, Emitter.EmitConstantFloat2({ UTiling, VTiling }));
 	
 	Emitter.Output(GetOutput(0), OutputValue);
 }
@@ -695,5 +692,44 @@ void UMaterialExpressionReroute::Build(UE::MIR::FEmitter& Emitter)
 	Emitter.Output(0, Value);
 }
 
+void BuildTernaryArithmeticOperator(
+	MIR::FEmitter& Emitter,
+	MIR::EOperator Op,
+	FExpressionInput* A,
+	float ConstA,
+	FExpressionInput* B,
+	float ConstB,
+	FExpressionInput* C,
+	float ConstC)
+{
+	Emitter.DefaultInputToFloatScalar(A, ConstA);
+	Emitter.DefaultInputToFloatScalar(B, ConstB);
+	Emitter.DefaultInputToFloatScalar(C, ConstC);
+
+	UE_MIR_CHECKPOINT();
+
+	MIR::FValue* ValueA = Emitter.InputPrimitive(A);
+	MIR::FValue* ValueB = Emitter.Input(B);
+	MIR::FValue* ValueC = Emitter.Input(C);
+
+	UE_MIR_CHECKPOINT();
+
+	ValueB = Emitter.EmitConstruct(ValueA->Type, ValueB);
+	ValueC = Emitter.EmitConstruct(ValueA->Type, ValueC);
+
+	UE_MIR_CHECKPOINT();
+
+	Emitter.Output(0, Emitter.EmitOperator(Op, ValueA, ValueB, ValueC));
+}
+
+void UMaterialExpressionLinearInterpolate::Build(UE::MIR::FEmitter& Emitter)
+{
+	BuildTernaryArithmeticOperator(Emitter, MIR::TO_Lerp, &A, ConstA, &B, ConstB, &Alpha, ConstAlpha);
+}
+
+void UMaterialExpressionSmoothStep::Build(UE::MIR::FEmitter& Emitter)
+{
+	BuildTernaryArithmeticOperator(Emitter, MIR::TO_Smoothstep, &Min, ConstMin, &Max, ConstMax, &Value, ConstValue);
+}
 
 #endif // WITH_EDITOR
diff --git a/Engine/Source/Runtime/Engine/Private/Materials/MaterialIR.cpp b/Engine/Source/Runtime/Engine/Private/Materials/MaterialIR.cpp
--- a/Engine/Source/Runtime/Engine/Private/Materials/MaterialIR.cpp
+++ b/Engine/Source/Runtime/Engine/Private/Materials/MaterialIR.cpp
@@ -18,7 +18,7 @@ const TCHAR* ValueKindToString(EValueKind Kind)
 		case VK_TextureObject: return TEXT("TextureObject");
 		case VK_Dimensional: return TEXT("Dimensional");
 		case VK_SetMaterialOutput: return TEXT("SetMaterialOutput");
-		case VK_BinaryOperator: return TEXT("BinaryOperator");
+		case VK_Operator: return TEXT("Operator");
 		case VK_Branch: return TEXT("Branch");
 		case VK_Subscript: return TEXT("Subscript");
 		case VK_Cast: return TEXT("Cast");
@@ -35,7 +35,7 @@ uint32 FValue::GetSizeInBytes() const
 		case VK_ExternalInput: return sizeof(FExternalInput);
 		case VK_TextureObject: return sizeof(FTextureObject);
 		case VK_SetMaterialOutput: return sizeof(FSetMaterialOutput);
-		case VK_BinaryOperator: return sizeof(FBinaryOperator);
+		case VK_Operator: return sizeof(FOperator);
 		case VK_Dimensional: return sizeof(FDimensional) + sizeof(FValue*) * static_cast<const FDimensional*>(this)->GetComponents().Num();
 		case VK_Branch: return sizeof(FBranch);
 		case VK_Subscript: return sizeof(FSubscript);
@@ -116,10 +116,10 @@ TArrayView<FValue*> FValue::GetUses()
 			return { &This->Arg, 1 };
 		}
 
-		case VK_BinaryOperator:
+		case VK_Operator:
 		{
-			auto This = static_cast<FBinaryOperator*>(this);
-			return { &This->LhsArg, 2 };
+			auto This = static_cast<FOperator*>(this);
+			return { &This->AArg, 3 };
 		}
 
 		case VK_Branch:
@@ -374,32 +374,54 @@ FBlock* FInstruction::GetDesiredBlockForUse(int32 UseIndex)
 	return Block;
 }
 
-bool IsArithmeticOperator(EBinaryOperator Op)
+bool OperatorIsComparison(EOperator Op)
 {
-	return Op >= BO_Add && Op <= BO_Divide;
+	return Op >= BO_GreaterThan && Op <= BO_NotEquals;
 }
 
-bool IsComparisonOperator(EBinaryOperator Op)
+bool OperatorIsArithmetic(EOperator Op)
 {
-	return Op >= BO_GreaterThan && Op <= BO_NotEquals;
+	return Op >= BO_Add;
+}
+
+bool OperatorIsTernary(EOperator Op)
+{
+	return Op >= TO_Lerp;
 }
 
-const TCHAR* BinaryOperatorToString(EBinaryOperator Op)
+bool OperatorIsBinaryAlgebraic(EOperator Op)
+{
+	return Op >= BO_Add && Op <= BO_Divide;
+}
+
+const TCHAR* OperatorToString(EOperator Op)
 {
 	switch (Op)
 	{
-		case BO_Add: return TEXT("Add");
-		case BO_Subtract: return TEXT("Subtract");
-		case BO_Multiply: return TEXT("Multiply");
-		case BO_Divide: return TEXT("Divide");
+		case O_Invalid: return TEXT("Invalid");
+
+		/* Binary Comparisons */
 		case BO_GreaterThan: return TEXT("GreaterThan");
 		case BO_GreaterThanOrEquals: return TEXT("GreaterThanOrEquals");
 		case BO_LowerThan: return TEXT("LowerThan");
 		case BO_LowerThanOrEquals: return TEXT("LowerThanOrEquals");
 		case BO_Equals: return TEXT("Equals");
 		case BO_NotEquals: return TEXT("NotEquals");
+	
+		/* Binary Arithmetic */
+		case BO_Add: return TEXT("Add");
+		case BO_Subtract: return TEXT("Subtract");
+		case BO_Multiply: return TEXT("Multiply");
+		case BO_Divide: return TEXT("Divide");
+		case BO_Min: return TEXT("Min");
+		case BO_Max: return TEXT("Max");
+		
+		/* Ternary Arithmetic */
+		case TO_Lerp: return TEXT("Lerp");
+		case TO_Smoothstep: return TEXT("Smoothstep");
+
 		default: UE_MIR_UNREACHABLE();
-	};
+	}
 }
 
 } // namespace UE::MIR
diff --git a/Engine/Source/Runtime/Engine/Private/Materials/MaterialIRDebug.cpp b/Engine/Source/Runtime/Engine/Private/Materials/MaterialIRDebug.cpp
--- a/Engine/Source/Runtime/Engine/Private/Materials/MaterialIRDebug.cpp
+++ b/Engine/Source/Runtime/Engine/Private/Materials/MaterialIRDebug.cpp
@@ -64,9 +64,9 @@ static void DumpValueInfo(const MIR::FValue* Value, FString& Out)
 			Out.Appendf(TEXT("Index: %d"), Subscript->Index);
 		}
 	}
-	else if (auto BinaryOperator = Value->As<MIR::FBinaryOperator>())
+	else if (auto Operator = Value->As<MIR::FOperator>())
 	{
-		Out.Append(MIR::BinaryOperatorToString(BinaryOperator->Operator));
+		Out.Append(MIR::OperatorToString(Operator->Op));
 	}
 }				
 
@@ -89,9 +89,9 @@ static void DumpUseInfo(const MIR::FValue* Value, const MIR::FValue* Use, int Us
 		static const TCHAR* Uses[] = { TEXT("condition"), TEXT("true"), TEXT("false") };
 		Out.Append(Uses[UseIndex]);
 	}
-	else if (auto BinaryOperator = Value->As<MIR::FBinaryOperator>())
+	else if (auto Operator = Value->As<MIR::FOperator>())
 	{
-		static const TCHAR* Uses[] = { TEXT("lhs"), TEXT("rhs") };
+		static const TCHAR* Uses[] = { TEXT("a"), TEXT("b"), TEXT("c") };
 		Out.Append(Uses[UseIndex]);
 	}
 }
diff --git a/Engine/Source/Runtime/Engine/Private/Materials/MaterialIREmitter.cpp b/Engine/Source/Runtime/Engine/Private/Materials/MaterialIREmitter.cpp
--- a/Engine/Source/Runtime/Engine/Private/Materials/MaterialIREmitter.cpp
+++ b/Engine/Source/Runtime/Engine/Private/Materials/MaterialIREmitter.cpp
@@ -652,58 +652,82 @@ FSetMaterialOutput* FEmitter::EmitSetMaterialOutput(EMaterialProperty InProperty
 }
 
 template <typename T>
-static T FoldScalarArithmeticOperator(EBinaryOperator Operator, T Lhs, T Rhs)
+static bool FoldComparisonOperatorScalar(EOperator Operator, T A, T B)
 {
 	switch (Operator)
 	{
-		case BO_Add: return Lhs + Rhs;
-		case BO_Subtract: return Lhs - Rhs;
-		case BO_Multiply: return Lhs * Rhs;
-		case BO_Divide: return Lhs / Rhs;
+		case BO_GreaterThan: return A > B;
+		case BO_GreaterThanOrEquals: return A >= B;
+		case BO_LowerThan: return A < B;
+		case BO_LowerThanOrEquals: return A <= B;
+		case BO_Equals: return A == B;
+		case BO_NotEquals: return A != B;
 		default: UE_MIR_UNREACHABLE();
 	}
 }
 
 template <typename T>
-static bool FoldComparisonOperatorScalar(EBinaryOperator Operator, T Lhs, T Rhs)
+static T FoldScalarArithmeticOperator(EOperator Operator, T A, T B, T C)
 {
 	switch (Operator)
 	{
-		case BO_GreaterThan: return Lhs > Rhs;
-		case BO_GreaterThanOrEquals: return Lhs >= Rhs;
-		case BO_LowerThan: return Lhs < Rhs;
-		case BO_LowerThanOrEquals: return Lhs <= Rhs;
-		case BO_Equals: return Lhs == Rhs;
-		case BO_NotEquals: return Lhs != Rhs;
+		case BO_Add: return A + B;
+		case BO_Subtract: return A - B;
+		case BO_Multiply: return A * B;
+		case BO_Divide: return A / B;
+		case BO_Min: return FMath::Min<T>(A, B);
+		case BO_Max: return FMath::Max<T>(A, B);
+		case TO_Lerp: return FMath::Lerp<T>(A, B, C);
+		case TO_Smoothstep: return FMath::SmoothStep<T>(A, B, C);
 		default: UE_MIR_UNREACHABLE();
 	}
 }
 
-static FValue* TryFoldBinaryOperatorScalar(FEmitter* Emitter, EBinaryOperator Operator, FValue* Lhs, FValue* Rhs)
+static FValue* TryFoldOperatorScalar(FEmitter* Emitter, EOperator Op, FValue* A, FValue* B, FValue* C)
 {
-	FPrimitiveTypePtr PrimitiveType = Lhs->Type->AsPrimitive();
+	FPrimitiveTypePtr PrimitiveType = A->Type->AsPrimitive();
 
 	// Check whether we can fold the operation based on identies (e.g. "x + 0 = x ∀ x ∈ R").
 	// These simplifications don't requre that both lhs and rhs are constant.
-	switch (Operator)
+	switch (Op)
 	{
+		/* Comparisons */
+		case BO_GreaterThan:
+		case BO_LowerThan:
+		case BO_NotEquals:
+			if (A->Equals(B))
+			{
+				return Emitter->EmitConstantFalse();
+			}
+			break;
+
+		case BO_GreaterThanOrEquals:
+		case BO_LowerThanOrEquals:
+		case BO_Equals:
+			if (A->Equals(B))
+			{
+				return Emitter->EmitConstantTrue();
+			}
+			break;
+
+		/* Arithmetic */
 		case BO_Add:
-			if (Lhs->IsNearlyZero())
+			if (A->IsNearlyZero())
 			{
-				return Rhs;
+				return B;
 			}
-			else if (Rhs->IsNearlyZero())
+			else if (B->IsNearlyZero())
 			{
-				return Lhs;
+				return A;
 			}
 			break;
 
 		case BO_Subtract:
-			if (Rhs->IsNearlyZero())
+			if (B->IsNearlyZero())
 			{
-				return Lhs;
+				return A;
 			}
-			else if (Lhs->IsNearlyZero())
+			else if (A->IsNearlyZero())
 			{
 				// return Emitter->EmitUnaryOperator(UO_Minus, Rhs);
 				UE_MIR_TODO();
@@ -711,114 +735,108 @@ static FValue* TryFoldBinaryOperatorScalar(FEmitter* Emitter, EBinaryOperator Op
 			break;
 
 		case BO_Multiply:
-			if (Lhs->IsNearlyZero() || Rhs->IsNearlyOne())
+			if (A->IsNearlyZero() || B->IsNearlyOne())
 			{
-				return Lhs;
+				return A;
 			}
-			else if (Lhs->IsNearlyOne() || Rhs->IsNearlyZero())
+			else if (A->IsNearlyOne() || B->IsNearlyZero())
 			{
-				return Rhs;
+				return B;
 			}
 			break;
 
 		case BO_Divide:
-			if (Lhs->IsNearlyZero() || Rhs->IsNearlyOne())
+			if (A->IsNearlyZero() || B->IsNearlyOne())
 			{
-				return Lhs;
+				return A;
 			}
 			break;
 
-		case BO_GreaterThan:
-		case BO_LowerThan:
-		case BO_NotEquals:
-			if (Lhs->Equals(Rhs))
-			{
-				return Emitter->EmitConstantFalse();
-			}
-			break;
+		case TO_Lerp:
+			break; // todo
 
-		case BO_GreaterThanOrEquals:
-		case BO_LowerThanOrEquals:
-		case BO_Equals:
-			if (Lhs->Equals(Rhs))
-			{
-				return Emitter->EmitConstantTrue();
-			}
-			break;
+		case BO_Min:
+		case BO_Max:
+		case TO_Smoothstep:
+			break; // no simplifications for these operations
+
+		default:
+			UE_MIR_UNREACHABLE();
 	}
 
 	// Verify that both lhs and rhs are constants, otherwise we cannot fold the operation.
-	const FConstant* LhsConstant = Lhs->As<FConstant>();
-	const FConstant* RhsConstant = Rhs->As<FConstant>();
-	if (!LhsConstant || !RhsConstant)
+	const FConstant* AConstant = A->As<FConstant>();
+	const FConstant* BConstant = B->As<FConstant>();
+	const FConstant* CConstant = C->As<FConstant>();
+	if (!AConstant || !BConstant || (OperatorIsTernary(Op) && !CConstant))
 	{
-		return MIR::FPoison::Get();
+		return nullptr;
 	}
 
-	if (IsArithmeticOperator(Operator))
+	if (OperatorIsComparison(Op))
 	{
+		bool Result;
 		switch (PrimitiveType->ScalarKind)
 		{
-			case SK_Int:
-			{
-				int Result = FoldScalarArithmeticOperator<TInteger>(Operator, LhsConstant->Integer, RhsConstant->Integer);
-				return Emitter->EmitConstantIntScalar(Result);
-			}
+		case SK_Int:
+			Result = FoldComparisonOperatorScalar<TInteger>(Op, AConstant->Integer, BConstant->Integer);
+			break;
 
-			case SK_Float:
-			{
-				TFloat Result = FoldScalarArithmeticOperator<TFloat>(Operator, LhsConstant->Float, RhsConstant->Float);
-				return Emitter->EmitConstantFloatScalar(Result);
-			}
+		case SK_Float:
+			Result = FoldComparisonOperatorScalar<TFloat>(Op, AConstant->Float, BConstant->Float);
+			break;
 
-			default:
-				UE_MIR_UNREACHABLE();
+		default:
+			UE_MIR_UNREACHABLE();
 		}
+		return Emitter->EmitConstantBoolScalar(Result);
 	}
-	else if (IsComparisonOperator(Operator))
+	else
 	{
-		bool Result;
+		check(OperatorIsArithmetic(Op));
 		switch (PrimitiveType->ScalarKind)
 		{
 			case SK_Int:
-				Result = FoldComparisonOperatorScalar<TInteger>(Operator, LhsConstant->Integer, RhsConstant->Integer);
-				break;
+			{
+				int Result = FoldScalarArithmeticOperator<TInteger>(Op, AConstant->Integer, BConstant->Integer, CConstant ? CConstant->Integer : 0);
+				return Emitter->EmitConstantIntScalar(Result);
+			}
 
 			case SK_Float:
-				Result = FoldComparisonOperatorScalar<TFloat>(Operator, LhsConstant->Float, RhsConstant->Float);
-				break;
+			{
+				TFloat Result = FoldScalarArithmeticOperator<TFloat>(Op, AConstant->Float, BConstant->Float, CConstant ? CConstant->Float : 0);
+				return Emitter->EmitConstantFloatScalar(Result);
+			}
 
 			default:
 				UE_MIR_UNREACHABLE();
 		}
-		return Emitter->EmitConstantBoolScalar(Result);
-	}
-	else
-	{
-		UE_MIR_UNREACHABLE();
 	}
 }
 
-FValue* FEmitter::EmitBinaryOperator(EBinaryOperator Operator, FValue* Lhs, FValue* Rhs)
+FValue* FEmitter::EmitOperator(EOperator Op, FValue* A, FValue* B, FValue* C)
 {
+	// Verify that C argument has been provided if operator is ternary.
+	check(!MIR::OperatorIsTernary(Op) || C);
+
 	// Argument value types must always match.
-	if (Lhs->Type != Rhs->Type)
+	if (A->Type != B->Type || (C && (A->Type != C->Type || B->Type != C->Type)))
 	{
-		Errorf(TEXT("operands of binary operator `%s` do not have the same type."), BinaryOperatorToString(Operator));
+		Errorf(TEXT("Operands of operator `%s` do not have the same type."), OperatorToString(Op));
 		return MIR::FPoison::Get();
 	}
 
 	// Get the operands primitive type.
-	FPrimitiveTypePtr PrimitiveType = Lhs->Type->AsPrimitive();
+	FPrimitiveTypePtr PrimitiveType = A->Type->AsPrimitive();
 	if (!PrimitiveType)
 	{
-		Errorf(TEXT("invalid binary operator on non primitive types."));
+		Errorf(TEXT("Invalid operator on non primitive types."));
 		return MIR::FPoison::Get();
 	}
 
 	// Determine the result type. If the operator is arithmetic, the result type will be the same
-	// as the operands type. Otherwise it will have the same number of components but bool.
-	FPrimitiveTypePtr ResultType = IsArithmeticOperator(Operator)
+	// as the arguments type. Otherwise it will have the same number of components but bool.
+	FPrimitiveTypePtr ResultType = OperatorIsArithmetic(Op)
 		? PrimitiveType
 		: FPrimitiveType::Get(SK_Bool, PrimitiveType->NumRows, PrimitiveType->NumColumns);
 
@@ -827,8 +845,9 @@ FValue* FEmitter::EmitBinaryOperator(EBinaryOperator Operator, FValue* Lhs, FVal
 	// an individual operator instruction for each component, simply emit a single binary operator
 	// instruction applied between lhs and rhs as a whole. (v1 + v2 rather than float2(v1.x + v2.x, v1.y + v2.y)
 	bool bSomeResultComponentWasFolded = false;
-	bool bResultIsIdenticalToLhs = true;
-	bool bResultIsIdenticalToRhs = true;
+	bool bResultIsIdenticalToA = true;
+	bool bResultIsIdenticalToB = true;
+	bool bResultIsIdenticalToC = true;
 
 	// Allocate the temporary array to store the folded component results
 	FMemMark Mark(FMemStack::Get());
@@ -837,36 +856,42 @@ FValue* FEmitter::EmitBinaryOperator(EBinaryOperator Operator, FValue* Lhs, FVal
 	for (int i = 0; i < ResultType->GetNumComponents(); ++i)
 	{
 		// Extract the arguments individual components
-		FValue* LhsComponent = EmitSubscript(Lhs, i);
-		FValue* RhsComponent = EmitSubscript(Rhs, i);
+		FValue* AComponent = EmitSubscript(A, i);
+		FValue* BComponent = EmitSubscript(B, i);
+		FValue* CComponent = C ? EmitSubscript(C, i) : nullptr;
 
 		// Try folding the operation, it may return null
-		FValue* ResultComponent = TryFoldBinaryOperatorScalar(this, Operator, LhsComponent, RhsComponent);
+		FValue* ResultComponent = TryFoldOperatorScalar(this, Op, AComponent, BComponent, CComponent);
 
 		// Update the flags
 		bSomeResultComponentWasFolded |= (bool)ResultComponent;
-		bResultIsIdenticalToLhs &= ResultComponent && ResultComponent->Equals(LhsComponent);
-		bResultIsIdenticalToRhs &= ResultComponent && ResultComponent->Equals(RhsComponent);
+		bResultIsIdenticalToA &= ResultComponent && ResultComponent->Equals(AComponent);
+		bResultIsIdenticalToB &= ResultComponent && ResultComponent->Equals(BComponent);
+		bResultIsIdenticalToC &= CComponent && ResultComponent && ResultComponent->Equals(CComponent);
 
 		// Cache the results
 		TempResultComponents[i] = ResultComponent;
 	}
 
 	// If result is identical to either lhs or rhs, simply return it
-	if (bResultIsIdenticalToLhs)
+	if (bResultIsIdenticalToA)
+	{
+		return A;
+	}
+	else if (bResultIsIdenticalToB)
 	{
-		return Lhs;
+		return B;
 	}
-	else if (bResultIsIdenticalToRhs)
+	else if (bResultIsIdenticalToC)
 	{
-		return Rhs;
+		return C;
 	}
 
 	// If some component was folded (it is either constant or the operation was a NOP), it is worth
-	// build the binary operation as a separate operation for each component, that is like
-	//    float2(lhs.x + rhs.x, lhs.y + rhs.y)
+	// build the operation as a separate operation for each component, that is like
+	//    float2(a.x + b.x, a.y + b.y)
 	// rather than
-	//    lhs + rhs
+	//    a + b
 	// so that we retain as much compile-time information as possible.
 	if (bSomeResultComponentWasFolded)
 	{
@@ -879,8 +904,13 @@ FValue* FEmitter::EmitBinaryOperator(EBinaryOperator Operator, FValue* Lhs, FVal
 
 		// Make the new dimensional value
 		FDimensional* Result = NewDimensionalValue(this, ResultType);
+
+		// Fetch the components array from the result dimensional
 		TArrayView<FValue*> ResultComponents = Result->GetComponents();
 
+		// Also cache the type of a single component
+		FPrimitiveTypePtr ComponentType = ResultType->ToScalar();
+
 		// Create the operator instruction for each component pair
 		for (int i = 0; i < ResultType->GetNumComponents(); ++i)
 		{
@@ -890,9 +920,12 @@ FValue* FEmitter::EmitBinaryOperator(EBinaryOperator Operator, FValue* Lhs, FVal
 			// Otherwise emit the binary operation between the two components (this will create a new instruction)
 			if (!ResultComponents[i])
 			{
-				FValue* LhsComponent = EmitSubscript(Lhs, i);
-				FValue* RhsComponent = EmitSubscript(Rhs, i);
-				ResultComponents[i] = EmitBinaryOperator(Operator, LhsComponent, RhsComponent);
+				FOperator Proto = MakePrototype<FOperator>(ComponentType);
+				Proto.Op = Op;
+				Proto.AArg = EmitSubscript(A, i);
+				Proto.BArg = EmitSubscript(B, i);
+				Proto.CArg = C ? EmitSubscript(C, i) : nullptr;
+				ResultComponents[i] = EmitPrototype(this, Proto);
 			}
 		}
 
@@ -900,11 +933,11 @@ FValue* FEmitter::EmitBinaryOperator(EBinaryOperator Operator, FValue* Lhs, FVal
 	}
 
 	// Finally, simply emit the binary operator instruction between lhs and rhs (i.e. "lhs + rhs")
-	FBinaryOperator Proto = MakePrototype<FBinaryOperator>(ResultType);
-	Proto.Operator = Operator;
-	Proto.LhsArg= Lhs;
-	Proto.RhsArg = Rhs;
-
+	FOperator Proto = MakePrototype<FOperator>(ResultType);
+	Proto.Op = Op;
+	Proto.AArg = A;
+	Proto.BArg = B;
+	Proto.CArg = C;
 	return EmitPrototype(this, Proto);
 }
 
diff --git a/Engine/Source/Runtime/Engine/Private/Materials/MaterialIRToHLSLTranslator.cpp b/Engine/Source/Runtime/Engine/Private/Materials/MaterialIRToHLSLTranslator.cpp
--- a/Engine/Source/Runtime/Engine/Private/Materials/MaterialIRToHLSLTranslator.cpp
+++ b/Engine/Source/Runtime/Engine/Private/Materials/MaterialIRToHLSLTranslator.cpp
@@ -314,7 +314,7 @@ struct FTranslator : FMaterialIRToHLSLTranslation
 		{
 			case MIR::VK_Dimensional: LowerDimensional(static_cast<MIR::FDimensional*>(Instr)); break;
 			case MIR::VK_SetMaterialOutput: LowerSetMaterialOutput(static_cast<MIR::FSetMaterialOutput*>(Instr)); break;
-			case MIR::VK_BinaryOperator: LowerBinaryOperator(static_cast<MIR::FBinaryOperator*>(Instr)); break;
+			case MIR::VK_Operator: LowerOperator(static_cast<MIR::FOperator*>(Instr)); break;
 			case MIR::VK_Branch: LowerBranch(static_cast<MIR::FBranch*>(Instr)); break;
 			case MIR::VK_Subscript: LowerSubscript(static_cast<MIR::FSubscript*>(Instr)); break;
 			case MIR::VK_TextureSample: LowerTextureSample(static_cast<MIR::FTextureSample*>(Instr)); break;
@@ -386,25 +386,52 @@ struct FTranslator : FMaterialIRToHLSLTranslation
 		Printer << TEXT("PixelMaterialInputs.") << PropertyName << TEXT(" = ") << LowerValue(Output->Arg);
 	}
 
-	void LowerBinaryOperator(const MIR::FBinaryOperator* BinaryOperator)
+	void LowerOperator(const MIR::FOperator* Operator)
 	{
-		LowerValue(BinaryOperator->LhsArg);
-
-		const TCHAR* OpString;
-		switch (BinaryOperator->Operator)
-		{
-			case MIR::BO_Add: OpString = TEXT(" + "); break;
-			case MIR::BO_Subtract: OpString = TEXT(" - "); break;
-			case MIR::BO_Multiply: OpString = TEXT(" * "); break;
-			case MIR::BO_Divide: OpString = TEXT(" / "); break;
-			case MIR::BO_GreaterThan: OpString = TEXT(" > "); break;
-			case MIR::BO_LowerThan: OpString = TEXT(" < "); break;
-			case MIR::BO_Equals: OpString = TEXT(" == "); break;
-			default: UE_MIR_UNREACHABLE();
+		if (MIR::OperatorIsComparison(Operator->Op) || MIR::OperatorIsBinaryAlgebraic(Operator->Op))
+		{
+			LowerValue(Operator->AArg);
+
+			const TCHAR* OpString;
+			switch (Operator->Op)
+			{
+				case MIR::BO_GreaterThan: OpString = TEXT(">"); break;
+				case MIR::BO_GreaterThanOrEquals: OpString = TEXT(">="); break;
+				case MIR::BO_LowerThan: OpString = TEXT("<"); break;
+				case MIR::BO_LowerThanOrEquals: OpString = TEXT("<="); break;
+				case MIR::BO_Equals: OpString = TEXT("=="); break;
+				case MIR::BO_NotEquals: OpString = TEXT("!="); break;
+				
+				case MIR::BO_Add: OpString = TEXT(" + "); break;
+				case MIR::BO_Subtract: OpString = TEXT(" - "); break;
+				case MIR::BO_Multiply: OpString = TEXT(" * "); break;
+				case MIR::BO_Divide: OpString = TEXT(" / "); break;
+
+				default: UE_MIR_UNREACHABLE();
+			}
+			Printer << OpString;
+
+			LowerValue(Operator->BArg);
 		}
-		Printer << OpString;
+		else
+		{
+			const TCHAR* OpString;
+			switch (Operator->Op)
+			{
+				case MIR::BO_Min: OpString = TEXT("min"); break;
+				case MIR::BO_Max: OpString = TEXT("max"); break;
+				case MIR::TO_Lerp: OpString = TEXT("lerp"); break;
+				case MIR::TO_Smoothstep: OpString = TEXT("smoothstep"); break;
+				default: UE_MIR_UNREACHABLE();
+			}
 
-		LowerValue(BinaryOperator->RhsArg);
+			Printer << OpString << "(" << LowerValue(Operator->AArg) << ", " << LowerValue(Operator->BArg);
+			if (Operator->CArg)
+			{
+				Printer << ", " << LowerValue(Operator->CArg);
+			}
+			Printer << ")";
+		}
 	}
 
 	void LowerBranch(const MIR::FBranch* Branch)
diff --git a/Engine/Source/Runtime/Engine/Public/Materials/MaterialExpressionLinearInterpolate.h b/Engine/Source/Runtime/Engine/Public/Materials/MaterialExpressionLinearInterpolate.h
--- a/Engine/Source/Runtime/Engine/Public/Materials/MaterialExpressionLinearInterpolate.h
+++ b/Engine/Source/Runtime/Engine/Public/Materials/MaterialExpressionLinearInterpolate.h
@@ -38,6 +38,7 @@ class UMaterialExpressionLinearInterpolate : public UMaterialExpression
 
 	//~ Begin UMaterialExpression Interface
 #if WITH_EDITOR
+	virtual void Build(UE::MIR::FEmitter& Emitter) override;
 	virtual bool GenerateHLSLExpression(FMaterialHLSLGenerator& Generator, UE::HLSLTree::FScope& Scope, int32 OutputIndex, UE::HLSLTree::FExpression const*& OutExpression) const override;
 	virtual int32 Compile(class FMaterialCompiler* Compiler, int32 OutputIndex) override;
 	virtual void GetCaption(TArray<FString>& OutCaptions) const override;
diff --git a/Engine/Source/Runtime/Engine/Public/Materials/MaterialExpressionSmoothStep.h b/Engine/Source/Runtime/Engine/Public/Materials/MaterialExpressionSmoothStep.h
--- a/Engine/Source/Runtime/Engine/Public/Materials/MaterialExpressionSmoothStep.h
+++ b/Engine/Source/Runtime/Engine/Public/Materials/MaterialExpressionSmoothStep.h
@@ -38,6 +38,7 @@ class UMaterialExpressionSmoothStep : public UMaterialExpression
 
 	//~ Begin UMaterialExpression Interface
 #if WITH_EDITOR
+	virtual void Build(UE::MIR::FEmitter& Emitter) override;
 	virtual int32 Compile(class FMaterialCompiler* Compiler, int32 OutputIndex) override;
 	virtual void GetCaption(TArray<FString>& OutCaptions) const override;
 	virtual FText GetKeywords() const override { return FText::FromString(TEXT("SmoothStep")); }
diff --git a/Engine/Source/Runtime/Engine/Public/Materials/MaterialIR.h b/Engine/Source/Runtime/Engine/Public/Materials/MaterialIR.h
--- a/Engine/Source/Runtime/Engine/Public/Materials/MaterialIR.h
+++ b/Engine/Source/Runtime/Engine/Public/Materials/MaterialIR.h
@@ -23,7 +23,7 @@ enum EValueKind
 
 	VK_Dimensional,
 	VK_SetMaterialOutput,
-	VK_BinaryOperator,
+	VK_Operator,
 	VK_Branch,
 	VK_Subscript,
 	VK_Cast,
@@ -285,40 +285,54 @@ struct FSetMaterialOutput : TInstruction<VK_SetMaterialOutput>
 	FValue* Arg;
 };
 
-enum EBinaryOperator
+//
+enum EOperator
 {
-	BO_Invalid,
-
-	/* Arithmetic */
-	BO_Add,
-	BO_Subtract,
-	BO_Multiply,
-	BO_Divide,
-
-	/* Comparison */
+	O_Invalid,
+	
+	/* Binary Comparison Operators */
 	BO_GreaterThan,
 	BO_GreaterThanOrEquals,
 	BO_LowerThan,
 	BO_LowerThanOrEquals,
 	BO_Equals,
 	BO_NotEquals,
+
+	/* Binary Arithmetic Operators */
+	BO_Add,
+	BO_Subtract,
+	BO_Multiply,
+	BO_Divide,
+	BO_Min,
+	BO_Max,
+
+	/* Ternary Arithmetic Operators */
+	TO_Lerp,
+	TO_Smoothstep,
 };
 
 //
-bool IsArithmeticOperator(EBinaryOperator Op);
+bool OperatorIsComparison(EOperator Op);
+
+//
+bool OperatorIsArithmetic(EOperator Op);
+
+//
+bool OperatorIsTernary(EOperator Op);
 
 //
-bool IsComparisonOperator(EBinaryOperator Op);
+bool OperatorIsBinaryAlgebraic(EOperator Op);
 
 //
-const TCHAR* BinaryOperatorToString(EBinaryOperator Op);
+const TCHAR* OperatorToString(EOperator Op);
 
 //
-struct FBinaryOperator : TInstruction<VK_BinaryOperator>
+struct FOperator : TInstruction<VK_Operator>
 {
-	EBinaryOperator Operator = BO_Invalid;
-	FValue* LhsArg{};
-	FValue* RhsArg{};
+	EOperator Op = O_Invalid;
+	FValue* AArg{};
+	FValue* BArg{};
+	FValue* CArg{};
 };
 
 //
diff --git a/Engine/Source/Runtime/Engine/Public/Materials/MaterialIREmitter.h b/Engine/Source/Runtime/Engine/Public/Materials/MaterialIREmitter.h
--- a/Engine/Source/Runtime/Engine/Public/Materials/MaterialIREmitter.h
+++ b/Engine/Source/Runtime/Engine/Public/Materials/MaterialIREmitter.h
@@ -153,7 +153,7 @@ public:
 	/* Instructions */
 
 	FSetMaterialOutput* EmitSetMaterialOutput(EMaterialProperty InProperty, FValue* InArgValue);
-	FValue* EmitBinaryOperator(EBinaryOperator Operator, FValue* Lhs, FValue* Rhs);
+	FValue* EmitOperator(EOperator Operator, FValue* A, FValue* B, FValue* C = nullptr);
 	FValue* EmitBranch(FValue* Condition, FValue* True, FValue* False);
 	FValue* EmitConstruct(FTypePtr Type, FValue* Initializer);
 	FValue* EmitTextureSample(UTexture* Texture, FValue* TexCoord, ESamplerSourceMode SamplerSourceMode, ETextureMipValueMode MipValueMode, EMaterialSamplerType SamplerType);
