Added `foliage.ApplyViewOcclusionCullToShadow`. When disabled, foliage shadows are more accurate, but rendering is likely slower (with traditional shadow maps)

#tests used the repro steps to create a test scene, and verified the issue does not reproduce when the new cvar is set to false

diff --git a/Engine/Source/Runtime/Engine/Private/HierarchicalInstancedStaticMesh.cpp b/Engine/Source/Runtime/Engine/Private/HierarchicalInstancedStaticMesh.cpp
--- a/Engine/Source/Runtime/Engine/Private/HierarchicalInstancedStaticMesh.cpp
+++ b/Engine/Source/Runtime/Engine/Private/HierarchicalInstancedStaticMesh.cpp
@@ -112,6 +112,11 @@ static TAutoConsoleVariable<int32> CVarMinVertsToSplitNode(
 	8192,
 	TEXT("Controls the accuracy between culling and LOD accuracy and culling and CPU performance."));
 
+static TAutoConsoleVariable<bool> CVarApplyViewOcclusionCullToShadow(
+	TEXT("foliage.ApplyViewOcclusionCullToShadow"),
+	true,
+	TEXT("Re-use main view occlusion culling results to speed up the shadow view. Shadows for fully occluded objects may be missing."));
+
 static TAutoConsoleVariable<int32> CVarMaxOcclusionQueriesPerComponent(
 	TEXT("foliage.MaxOcclusionQueriesPerComponent"),
 	16,
@@ -1506,6 +1511,8 @@ void FHierarchicalStaticMeshSceneProxy::GetDynamicMeshElements(const TArray<cons
 
 	const FMatrix WorldToLocal = GetLocalToWorld().Inverse();
 
+	const bool bApplyViewOcclusionCullToShadow = CVarApplyViewOcclusionCullToShadow.GetValueOnAnyThread();
+
 	for (int32 ViewIndex = 0; ViewIndex < Views.Num(); ViewIndex++)
 	{
 		if (VisibilityMap & (1 << ViewIndex))
@@ -1696,7 +1703,10 @@ void FHierarchicalStaticMeshSceneProxy::GetDynamicMeshElements(const TArray<cons
 					}
 				}
 
-				if (FirstOcclusionNode >= 0 && LastOcclusionNode >= 0 && FirstOcclusionNode <= LastOcclusionNode)
+				// Culling objects from the shadow view using main view occlusion culling results is technically wrong, but can give good perf improvements
+				// For VSM, it's likely better to turn this off as non-visible pages aren't rendered anyway and we care more about quality
+				bool bApplyOcclusionCull = !ElementParams.ShadowFrustum || bApplyViewOcclusionCullToShadow;
+				if (FirstOcclusionNode >= 0 && LastOcclusionNode >= 0 && FirstOcclusionNode <= LastOcclusionNode && bApplyOcclusionCull)
 				{
 					uint32 ViewId = View->GetViewKey();
 					const FFoliageOcclusionResults* OldResults = OcclusionResults.Find(ViewId);
