Substrate - Select node implementation beginning.

#tests Substrate no pc
#rnx

diff --git a/Engine/Shaders/Private/Substrate/Substrate.ush b/Engine/Shaders/Private/Substrate/Substrate.ush
--- a/Engine/Shaders/Private/Substrate/Substrate.ush
+++ b/Engine/Shaders/Private/Substrate/Substrate.ush
@@ -2884,6 +2884,37 @@ FSubstrateData SubstrateWeightParameterBlending(FSubstrateData A, float Weight)
 	return SubstrateData;
 }
 
+float SelectParameterBlendingToNormal(float SelectValue)
+{
+	if (SelectValue <= 0.0f)
+	{
+		return 0.0f;
+	}
+	return 1.0f;
+}
+
+FSubstrateData SubstrateSelectParameterBlending(FSubstrateData A, FSubstrateData B, float SelectValue, uint NewNormalIndex, inout uint OutSharedLocalBasisTypes)
+{
+	// It is fine to reuse A&B and morph them since they are here a copy of the original SubstrateData.
+
+	if (SelectValue <= 0.0f)
+	{
+		BSDF_SETSHAREDLOCALBASISID(A.InlinedBSDF, NewNormalIndex);
+		if (!SUBSTRATE_FASTPATH && BSDF_GETHASANISOTROPY(A.InlinedBSDF))
+		{
+			SubstrateRequestSharedLocalBasisTangent(OutSharedLocalBasisTypes, NewNormalIndex);
+		}
+		return A;
+	}
+
+	BSDF_SETSHAREDLOCALBASISID(B.InlinedBSDF, NewNormalIndex);
+	if (!SUBSTRATE_FASTPATH && BSDF_GETHASANISOTROPY(B.InlinedBSDF))
+	{
+		SubstrateRequestSharedLocalBasisTangent(OutSharedLocalBasisTypes, NewNormalIndex);
+	}
+	return B;
+}
+
 void SubstrateGetThinFilmF0F90(float NoV, float ThinFilmNormalizedThickness, float ThinFilmIOR, inout float3 OutF0, inout float3 OutF90)
 {
 	const float3 F0 = OutF0;
diff --git a/Engine/Shaders/Shared/SubstrateDefinitions.h b/Engine/Shaders/Shared/SubstrateDefinitions.h
--- a/Engine/Shaders/Shared/SubstrateDefinitions.h
+++ b/Engine/Shaders/Shared/SubstrateDefinitions.h
@@ -9,7 +9,7 @@
 
 #ifndef __cplusplus
 // Change this to force recompilation of all Substrate dependent shaders (for instance https://guidgenerator.com/online-guid-generator.aspx)
-#pragma message("UESHADERMETADATA_VERSION ADA0DF8F-BCE7-4921-89DF-E00C06DA6844")
+#pragma message("UESHADERMETADATA_VERSION 0C2A906D-907F-4184-A1D3-DF20C0802553")
 #endif
 
 // Closure offsets are packed into 32bits, each entry using SUBSTRATE_CLOSURE_OFFSET_BIT_COUNT bits
@@ -56,8 +56,9 @@
 #define SUBSTRATE_OPERATOR_VERTICAL						1
 #define SUBSTRATE_OPERATOR_HORIZONTAL					2
 #define SUBSTRATE_OPERATOR_ADD							3
-#define SUBSTRATE_OPERATOR_BSDF							4
-#define SUBSTRATE_OPERATOR_BSDF_LEGACY					5
+#define SUBSTRATE_OPERATOR_SELECT						4
+#define SUBSTRATE_OPERATOR_BSDF							5
+#define SUBSTRATE_OPERATOR_BSDF_LEGACY					6
 
 // Sometimes, Unlit or Weight operators are used to transport data for other meaning (e.g. Light Function or ConvertToDecal)
 #define SUBSTRATE_OPERATOR_SUBUSAGE_NONE				0
diff --git a/Engine/Source/Editor/GraphEditor/Private/MaterialNodes/SGraphSubstrateMaterial.cpp b/Engine/Source/Editor/GraphEditor/Private/MaterialNodes/SGraphSubstrateMaterial.cpp
--- a/Engine/Source/Editor/GraphEditor/Private/MaterialNodes/SGraphSubstrateMaterial.cpp
+++ b/Engine/Source/Editor/GraphEditor/Private/MaterialNodes/SGraphSubstrateMaterial.cpp
@@ -118,6 +118,7 @@ static const TSharedRef<SWidget> InternalProcessOperator(
 		}
 		break;
 		case SUBSTRATE_OPERATOR_HORIZONTAL:
+		case SUBSTRATE_OPERATOR_SELECT:
 		{
 			auto HorizontalOperator = SNew(SHorizontalBox)
 				+SHorizontalBox::Slot()
@@ -263,6 +264,19 @@ void FSubstrateWidget::GetPinColor(TSharedPtr<SGraphPin>& Out, const UMaterialGr
 			ColorModifier = Color0;
 		}
 	}
+	else if (InNode->MaterialExpression->IsA(UMaterialExpressionSubstrateSelect::StaticClass()))
+	{
+		if (PinName == InNode->MaterialExpression->GetInputName(1)) // B
+		{
+			bHasColorModifier = true;
+			ColorModifier = Color1;
+		}
+		else if (PinName == InNode->MaterialExpression->GetInputName(0)) // A
+		{
+			bHasColorModifier = true;
+			ColorModifier = Color0;
+		}
+	}
 	else if (InNode->MaterialExpression->IsA(UMaterialExpressionSubstrateAdd::StaticClass()))
 	{
 		if (PinName == InNode->MaterialExpression->GetInputName(0)) // A
diff --git a/Engine/Source/Runtime/Core/Private/UObject/FortniteMainBranchObjectVersion.cpp b/Engine/Source/Runtime/Core/Private/UObject/FortniteMainBranchObjectVersion.cpp
--- a/Engine/Source/Runtime/Core/Private/UObject/FortniteMainBranchObjectVersion.cpp
+++ b/Engine/Source/Runtime/Core/Private/UObject/FortniteMainBranchObjectVersion.cpp
@@ -14,6 +14,6 @@ TMap<FGuid, FGuid> FFortniteMainBranchObjectVersion::GetSystemGuids()
 	SystemGuids.Add(DevGuids.Niagara_LatestScriptCompileVersion, FGuid("F8D995224324470E97ED718D33939A0E"));
 	SystemGuids.Add(DevGuids.SkeletalMeshDerivedDataVersion, FGuid("B9AEAA3EE7AC4FCFB3860E8F3E06DD75"));
 	SystemGuids.Add(DevGuids.STATICMESH_DERIVEDDATA_VER, FGuid("CDFDB82FE0A83DB602AE2EFFD8BB7F42"));
-	SystemGuids.Add(DevGuids.MaterialTranslationDDCVersion, FGuid("38F4505C9CB84A8B9C9648D2E4CF0ED6"));
+	SystemGuids.Add(DevGuids.MaterialTranslationDDCVersion, FGuid("E779012A63684F46B85596E4E89D489A"));
 	return SystemGuids;
 }
diff --git a/Engine/Source/Runtime/Engine/Private/Materials/HLSLMaterialTranslator.cpp b/Engine/Source/Runtime/Engine/Private/Materials/HLSLMaterialTranslator.cpp
--- a/Engine/Source/Runtime/Engine/Private/Materials/HLSLMaterialTranslator.cpp
+++ b/Engine/Source/Runtime/Engine/Private/Materials/HLSLMaterialTranslator.cpp
@@ -605,6 +605,10 @@ static inline const TCHAR* GetSubstrateOperatorStr(int32 OperatorType)
 	{
 		return TEXT("ADD       ");
 	}
+	case SUBSTRATE_OPERATOR_SELECT:
+	{
+		return TEXT("SELECT    ");
+	}
 	case SUBSTRATE_OPERATOR_BSDF:
 	{
 		return TEXT("BSDF      ");
@@ -2367,6 +2371,10 @@ void FHLSLMaterialTranslator::TranslateMaterial()
 									{
 										break; // NOP
 									}
+									case SUBSTRATE_OPERATOR_SELECT:
+									{
+										break; // NOP
+									}
 									default:
 									case SUBSTRATE_OPERATOR_BSDF:
 									{
@@ -2467,6 +2475,7 @@ void FHLSLMaterialTranslator::TranslateMaterial()
 										break;
 									}
 									case SUBSTRATE_OPERATOR_ADD:
+									case SUBSTRATE_OPERATOR_SELECT:
 									{
 										break; // NOP
 									}
@@ -12755,7 +12764,6 @@ bool FHLSLMaterialTranslator::FSubstrateCompilationContext::SubstrateGenerateDer
 	//
 	// Make sure each and every path of the Substrate tree have valid children and path.
 	//
-	uint32 TotalBSDFCount = 0;
 	for (auto& It : SubstrateMaterialExpressionRegisteredOperators)
 	{
 		if (It.IsDiscarded())
@@ -12770,7 +12778,6 @@ bool FHLSLMaterialTranslator::FSubstrateCompilationContext::SubstrateGenerateDer
 		// Operators without any child
 		case SUBSTRATE_OPERATOR_BSDF:
 		{
-			TotalBSDFCount++;
 		}
 		break;
 
@@ -12778,6 +12785,7 @@ bool FHLSLMaterialTranslator::FSubstrateCompilationContext::SubstrateGenerateDer
 		case SUBSTRATE_OPERATOR_HORIZONTAL:
 		case SUBSTRATE_OPERATOR_VERTICAL:
 		case SUBSTRATE_OPERATOR_ADD:
+		case SUBSTRATE_OPERATOR_SELECT:
 		{
 			bMustHaveLeftChild = true;
 			bMustHaveRightChild = true;
@@ -12830,6 +12838,7 @@ bool FHLSLMaterialTranslator::FSubstrateCompilationContext::SubstrateGenerateDer
 			}
 			case SUBSTRATE_OPERATOR_HORIZONTAL:
 			case SUBSTRATE_OPERATOR_ADD:
+			case SUBSTRATE_OPERATOR_SELECT:
 			{
 				WalkOperatorsForDepth(SubstrateMaterialExpressionRegisteredOperators[CurrentOperator.LeftIndex]);
 				WalkOperatorsForDepth(SubstrateMaterialExpressionRegisteredOperators[CurrentOperator.RightIndex]);
@@ -12861,9 +12870,11 @@ bool FHLSLMaterialTranslator::FSubstrateCompilationContext::SubstrateGenerateDer
 
 				// Operators with a single child
 			case SUBSTRATE_OPERATOR_WEIGHT:
-			{
-			}
-			break;
+				break;
+
+				// Operators always has parameter blending enabled
+			case SUBSTRATE_OPERATOR_SELECT:
+				break;
 
 			// Operators with two children
 			case SUBSTRATE_OPERATOR_HORIZONTAL:
@@ -12977,6 +12988,7 @@ bool FHLSLMaterialTranslator::FSubstrateCompilationContext::SubstrateGenerateDer
 				}
 				case SUBSTRATE_OPERATOR_HORIZONTAL:
 				case SUBSTRATE_OPERATOR_ADD:
+				case SUBSTRATE_OPERATOR_SELECT:
 				{
 					WalkOperators(SubstrateMaterialExpressionRegisteredOperators[CurrentOperator.LeftIndex], bUseParameterBlending);
 					WalkOperators(SubstrateMaterialExpressionRegisteredOperators[CurrentOperator.RightIndex], bUseParameterBlending);
@@ -13093,6 +13105,7 @@ bool FHLSLMaterialTranslator::FSubstrateCompilationContext::SubstrateGenerateDer
 			case SUBSTRATE_OPERATOR_HORIZONTAL:
 			case SUBSTRATE_OPERATOR_VERTICAL:
 			case SUBSTRATE_OPERATOR_ADD:
+			case SUBSTRATE_OPERATOR_SELECT:
 			{
 				check(It.RightIndex != INDEX_NONE);
 			}
@@ -13123,6 +13136,7 @@ bool FHLSLMaterialTranslator::FSubstrateCompilationContext::SubstrateGenerateDer
 				case SUBSTRATE_OPERATOR_VERTICAL:
 				case SUBSTRATE_OPERATOR_HORIZONTAL:
 				case SUBSTRATE_OPERATOR_ADD:
+				case SUBSTRATE_OPERATOR_SELECT:
 				{
 					CurrentOperator.MaxDistanceFromLeaves = FMath::Max(
 						SubstrateMaterialExpressionRegisteredOperators[CurrentOperator.LeftIndex].MaxDistanceFromLeaves,
@@ -13172,6 +13186,7 @@ bool FHLSLMaterialTranslator::FSubstrateCompilationContext::SubstrateGenerateDer
 				case SUBSTRATE_OPERATOR_VERTICAL:
 				case SUBSTRATE_OPERATOR_HORIZONTAL:
 				case SUBSTRATE_OPERATOR_ADD:
+				case SUBSTRATE_OPERATOR_SELECT:
 				{
 					PopulateParameterBlendedSubtree(SubstrateMaterialExpressionRegisteredOperators[CurrentOperator.LeftIndex], RootOperator);
 					PopulateParameterBlendedSubtree(SubstrateMaterialExpressionRegisteredOperators[CurrentOperator.RightIndex], RootOperator);
@@ -13208,6 +13223,7 @@ bool FHLSLMaterialTranslator::FSubstrateCompilationContext::SubstrateGenerateDer
 				}
 				case SUBSTRATE_OPERATOR_HORIZONTAL:
 				case SUBSTRATE_OPERATOR_ADD:
+				case SUBSTRATE_OPERATOR_SELECT:
 				{
 					WalkOperators(SubstrateMaterialExpressionRegisteredOperators[CurrentOperator.LeftIndex]);
 					WalkOperators(SubstrateMaterialExpressionRegisteredOperators[CurrentOperator.RightIndex]);
@@ -14479,6 +14495,42 @@ int32 FHLSLMaterialTranslator::SubstrateWeightParameterBlending(int32 A, int32 W
 	);
 }
 
+int32 FHLSLMaterialTranslator::SubstrateSelectParameterBlending(int32 A, int32 B, int32 SelectValue, const FString& SharedLocalBasisIndexMacro, FSubstrateOperator* PromoteToOperator)
+{
+	if (A == INDEX_NONE || B == INDEX_NONE || SelectValue == INDEX_NONE)
+	{
+		return INDEX_NONE;
+	}
+
+	if (PromoteToOperator)
+	{
+		check(PromoteToOperator->Index != INDEX_NONE);
+		check(PromoteToOperator->BSDFIndex != INDEX_NONE);
+		return AddCodeChunk(
+			MCT_Substrate, TEXT("Parameters.%s.PromoteParameterBlendedBSDFToOperator(SubstrateSelectParameterBlending(%s, %s, %s, %s, Parameters.%s.Types), %u, %u, %u, %u)"),
+			*GetParametersSubstrateTreeName(CurrentSubstrateCompilationContext),
+			*GetParameterCode(A),
+			*GetParameterCode(B),
+			*GetParameterCode(SelectValue),
+			*SharedLocalBasisIndexMacro,
+			*GetParametersSharedLocalBasesName(CurrentSubstrateCompilationContext),
+			PromoteToOperator->Index,
+			PromoteToOperator->BSDFIndex,
+			PromoteToOperator->LayerDepth,
+			PromoteToOperator->bIsBottom ? 1 : 0
+		);
+	}
+
+	return AddCodeChunk(
+		MCT_Substrate, TEXT("SubstrateSelectParameterBlending(%s, %s, %s, %s, Parameters.%s.Types)"),
+		*GetParameterCode(A),
+		*GetParameterCode(B),
+		*GetParameterCode(SelectValue),
+		*SharedLocalBasisIndexMacro,
+		*GetParametersSharedLocalBasesName(CurrentSubstrateCompilationContext)
+	);
+}
+
 int32 FHLSLMaterialTranslator::SubstrateTransmittanceToMFP(int32 TransmittanceColor, int32 DesiredThickness, int32 OutputIndex)
 {
 	if (OutputIndex == INDEX_NONE)
diff --git a/Engine/Source/Runtime/Engine/Private/Materials/HLSLMaterialTranslator.h b/Engine/Source/Runtime/Engine/Private/Materials/HLSLMaterialTranslator.h
--- a/Engine/Source/Runtime/Engine/Private/Materials/HLSLMaterialTranslator.h
+++ b/Engine/Source/Runtime/Engine/Private/Materials/HLSLMaterialTranslator.h
@@ -1334,6 +1334,7 @@ protected:
 	virtual int32 SubstrateAddParameterBlending(int32 A, int32 B, int32 AMixWeight, const FString& SharedLocalBasisIndexMacro, FSubstrateOperator* PromoteToOperator) override;
 	virtual int32 SubstrateWeight(int32 A, int32 Weight, int OperatorIndex, uint32 MaxDistanceFromLeaves) override;
 	virtual int32 SubstrateWeightParameterBlending(int32 A, int32 Weight, FSubstrateOperator* PromoteToOperator) override;
+	virtual int32 SubstrateSelectParameterBlending(int32 A, int32 B, int32 SelectValue, const FString& SharedLocalBasisIndexMacro, FSubstrateOperator* PromoteToOperator) override;
 	virtual int32 SubstrateTransmittanceToMFP(int32 TransmittanceColor, int32 DesiredThickness, int32 OutputIndex) override;
 	virtual int32 SubstrateMetalnessToDiffuseAlbedoF0(int32 BaseColor, int32 Specular, int32 Metallic, int32 OutputIndex) override;
 	virtual int32 SubstrateHazinessToSecondaryRoughness(int32 BaseRoughness, int32 Haziness, int32 OutputIndex) override;
diff --git a/Engine/Source/Runtime/Engine/Private/Materials/MaterialExpressionSubstrate.cpp b/Engine/Source/Runtime/Engine/Private/Materials/MaterialExpressionSubstrate.cpp
--- a/Engine/Source/Runtime/Engine/Private/Materials/MaterialExpressionSubstrate.cpp
+++ b/Engine/Source/Runtime/Engine/Private/Materials/MaterialExpressionSubstrate.cpp
@@ -3026,6 +3026,173 @@ FSubstrateOperator* UMaterialExpressionSubstrateWeight::SubstrateGenerateMateria
 }
 #endif // WITH_EDITOR
 
+
+
+UMaterialExpressionSubstrateSelect::UMaterialExpressionSubstrateSelect(const FObjectInitializer& ObjectInitializer)
+	: Super(ObjectInitializer)
+{
+	struct FConstructorStatics
+	{
+		FText NAME_Substrate;
+		FConstructorStatics() : NAME_Substrate(LOCTEXT("Substrate Ops", "Substrate Operators")) { }
+	};
+	static FConstructorStatics ConstructorStatics;
+#if WITH_EDITORONLY_DATA
+	MenuCategories.Add(ConstructorStatics.NAME_Substrate);
+#endif
+}
+
+#if WITH_EDITOR
+int32 UMaterialExpressionSubstrateSelect::Compile(class FMaterialCompiler* Compiler, int32 OutputIndex)
+{
+	if (!A.GetTracedInput().Expression)
+	{
+		return Compiler->Errorf(TEXT("Missing A input"));
+	}
+	if (!B.GetTracedInput().Expression)
+	{
+		return Compiler->Errorf(TEXT("Missing B input"));
+	}
+
+	Compiler->SubstrateTreeStackPush(this, 0);
+	int32 ACodeChunk = A.Compile(Compiler);
+	Compiler->SubstrateTreeStackPop();
+	Compiler->SubstrateTreeStackPush(this, 1);
+	int32 BCodeChunk = B.Compile(Compiler);
+	Compiler->SubstrateTreeStackPop();
+
+	// if SelectValue is not pluggedin, need to be 0. Otherwise it must be a float value.
+	int32 SelectValueCodeChunk = CompileWithDefaultFloat1(Compiler, SelectValue, 0.0f);
+	SelectValueCodeChunk = Compiler->ValidCast(SelectValueCodeChunk, MCT_Float1);
+
+	int32 ZeroCodeChunk = Compiler->Constant(0.0f);
+	int32 OneCodeChunk = Compiler->Constant(1.0f);
+	int32 SelectLerpCodeChunk = Compiler->If(SelectValueCodeChunk, ZeroCodeChunk, OneCodeChunk, ZeroCodeChunk, ZeroCodeChunk, ZeroCodeChunk);
+
+	// Compute the new Normal and Tangent resulting from the selection using code chunk
+	FSubstrateOperator& SubstrateOperator = Compiler->SubstrateCompilationGetOperator(Compiler->SubstrateTreeStackGetPathUniqueId());
+	FSubstrateOperator* OperatorA = Compiler->SubstrateCompilationGetOperatorFromIndex(SubstrateOperator.LeftIndex);
+	FSubstrateOperator* OperatorB = Compiler->SubstrateCompilationGetOperatorFromIndex(SubstrateOperator.RightIndex);
+	const int32 NewNormalCodeChunk = SubstrateBlendNormal(Compiler, OperatorA->BSDFRegisteredSharedLocalBasis.NormalCodeChunk, OperatorB->BSDFRegisteredSharedLocalBasis.NormalCodeChunk, SelectLerpCodeChunk);
+	// The tangent is optional so we treat it differently if INDEX_NONE is specified
+	int32 NewTangentCodeChunk = INDEX_NONE;
+	if (OperatorA->BSDFRegisteredSharedLocalBasis.TangentCodeChunk != INDEX_NONE && OperatorB->BSDFRegisteredSharedLocalBasis.TangentCodeChunk != INDEX_NONE)
+	{
+		NewTangentCodeChunk = SubstrateBlendNormal(Compiler, OperatorB->BSDFRegisteredSharedLocalBasis.TangentCodeChunk, OperatorA->BSDFRegisteredSharedLocalBasis.TangentCodeChunk, SelectLerpCodeChunk);
+	}
+	else if (OperatorA->BSDFRegisteredSharedLocalBasis.TangentCodeChunk != INDEX_NONE)
+	{
+		NewTangentCodeChunk = OperatorA->BSDFRegisteredSharedLocalBasis.TangentCodeChunk;
+	}
+	else if (OperatorB->BSDFRegisteredSharedLocalBasis.TangentCodeChunk != INDEX_NONE)
+	{
+		NewTangentCodeChunk = OperatorB->BSDFRegisteredSharedLocalBasis.TangentCodeChunk;
+	}
+	const FSubstrateRegisteredSharedLocalBasis NewRegisteredSharedLocalBasis = SubstrateCompilationInfoCreateSharedLocalBasis(Compiler, NewNormalCodeChunk, NewTangentCodeChunk);
+
+	int32 OutputCodeChunk = INDEX_NONE;
+	if (SubstrateOperator.bUseParameterBlending)
+	{
+		OutputCodeChunk = Compiler->SubstrateSelectParameterBlending(ACodeChunk, BCodeChunk, SelectValueCodeChunk, 
+			Compiler->GetSubstrateSharedLocalBasisIndexMacro(NewRegisteredSharedLocalBasis), 
+			SubstrateOperator.bRootOfParameterBlendingSubTree ? &SubstrateOperator : nullptr);
+	}
+	else
+	{
+		Compiler->Errorf(TEXT("The Select node must be compiled with parameter blending to only select between one of two BSDF"));
+	}
+
+	// Propagate the parameter blended normal
+	SubstrateOperator.BSDFRegisteredSharedLocalBasis = NewRegisteredSharedLocalBasis;
+
+	return OutputCodeChunk;
+}
+
+void UMaterialExpressionSubstrateSelect::GetCaption(TArray<FString>& OutCaptions) const
+{
+	OutCaptions.Add(TEXT("Substrate Select"));
+}
+
+uint32 UMaterialExpressionSubstrateSelect::GetOutputType(int32 OutputIndex)
+{
+	return MCT_Substrate;
+}
+
+uint32 UMaterialExpressionSubstrateSelect::GetInputType(int32 InputIndex)
+{
+	if (InputIndex == 2)
+	{
+		return MCT_Float1;
+	}
+	return MCT_Substrate;
+}
+
+bool UMaterialExpressionSubstrateSelect::IsResultSubstrateMaterial(int32 OutputIndex)
+{
+	return true;
+}
+
+void UMaterialExpressionSubstrateSelect::GatherSubstrateMaterialInfo(FSubstrateMaterialInfo& SubstrateMaterialInfo, int32 OutputIndex)
+{
+	if (!SubstrateMaterialInfo.PushSubstrateTreeStack())
+	{
+		return;
+	}
+
+	FExpressionInput TracedInputA = A.GetTracedInput();
+	FExpressionInput TracedInputB = B.GetTracedInput();
+	if (TracedInputA.Expression)
+	{
+		TracedInputA.Expression->GatherSubstrateMaterialInfo(SubstrateMaterialInfo, TracedInputA.OutputIndex);
+	}
+	if (TracedInputB.Expression)
+	{
+		TracedInputB.Expression->GatherSubstrateMaterialInfo(SubstrateMaterialInfo, TracedInputB.OutputIndex);
+	}
+	SubstrateMaterialInfo.AddGuid(MaterialExpressionGuid);
+
+	SubstrateMaterialInfo.PopSubstrateTreeStack();
+}
+
+FSubstrateOperator* UMaterialExpressionSubstrateSelect::SubstrateGenerateMaterialTopologyTree(class FMaterialCompiler* Compiler, class UMaterialExpression* Parent, int32 OutputIndex)
+{
+	const bool bUseParameterBlending = true; // The select node only works with a single BSDF on each side.
+	FSubstrateOperator& SubstrateOperator = Compiler->SubstrateCompilationRegisterOperator(SUBSTRATE_OPERATOR_SELECT, Compiler->SubstrateTreeStackGetPathUniqueId(), this, Parent, Compiler->SubstrateTreeStackGetParentPathUniqueId(), bUseParameterBlending);
+	if (Compiler->GetSubstrateTreeOutOfStackDepthOccurred())
+	{
+		return &SubstrateOperator; // Out ot stack space, return now to fail the compilation
+	}
+
+	FExpressionInput TracedInputA = A.GetTracedInput();
+	FExpressionInput TracedInputB = B.GetTracedInput();
+	UMaterialExpression* ChildAExpression = TracedInputA.Expression;
+	UMaterialExpression* ChildBExpression = TracedInputB.Expression;
+	FSubstrateOperator* OpA = nullptr;
+	FSubstrateOperator* OpB = nullptr;
+	if (ChildAExpression)
+	{
+		Compiler->SubstrateTreeStackPush(this, 0);
+		OpA = ChildAExpression->SubstrateGenerateMaterialTopologyTree(Compiler, this, TracedInputA.OutputIndex);
+		Compiler->SubstrateTreeStackPop();
+		AssignOperatorIndexIfNotNull(SubstrateOperator.LeftIndex, OpA);
+	}
+	if (ChildBExpression)
+	{
+		Compiler->SubstrateTreeStackPush(this, 1);
+		OpB = ChildBExpression->SubstrateGenerateMaterialTopologyTree(Compiler, this, TracedInputB.OutputIndex);
+		Compiler->SubstrateTreeStackPop();
+		AssignOperatorIndexIfNotNull(SubstrateOperator.RightIndex, OpB);
+	}
+
+	// Since A or B can be used, we need to combine all their flag to support the most expenssive use case selected.
+	CombineFlagForParameterBlending(SubstrateOperator, OpA, OpB);
+
+	return &SubstrateOperator;
+}
+#endif // WITH_EDITOR
+
+
+
 UMaterialExpressionSubstrateUtilityBase::UMaterialExpressionSubstrateUtilityBase(const FObjectInitializer& ObjectInitializer)
 	: Super(ObjectInitializer)
 {
diff --git a/Engine/Source/Runtime/Engine/Public/MaterialCompiler.h b/Engine/Source/Runtime/Engine/Public/MaterialCompiler.h
--- a/Engine/Source/Runtime/Engine/Public/MaterialCompiler.h
+++ b/Engine/Source/Runtime/Engine/Public/MaterialCompiler.h
@@ -706,6 +706,7 @@ public:
 	virtual int32 SubstrateAddParameterBlending(int32 A, int32 B, int32 AMixWeight, const FString& SharedLocalBasisIndexMacro, FSubstrateOperator* PromoteToOperator) = 0;
 	virtual int32 SubstrateWeight(int32 A, int32 Weight, int OperatorIndex, uint32 MaxDistanceFromLeaves) = 0;
 	virtual int32 SubstrateWeightParameterBlending(int32 A, int32 Weight, FSubstrateOperator* PromoteToOperator) = 0;
+	virtual int32 SubstrateSelectParameterBlending(int32 A, int32 B, int32 SelectValue, const FString& SharedLocalBasisIndexMacro, FSubstrateOperator* PromoteToOperator) = 0;
 	virtual int32 SubstrateTransmittanceToMFP(int32 TransmittanceColor, int32 DesiredThickness, int32 OutputIndex) = 0;
 	virtual int32 SubstrateMetalnessToDiffuseAlbedoF0(int32 BaseColor, int32 Specular, int32 Metallic, int32 OutputIndex) = 0;
 	virtual int32 SubstrateHazinessToSecondaryRoughness(int32 BaseRoughness, int32 Haziness, int32 OutputIndex) = 0;
@@ -1492,6 +1493,11 @@ public:
 	{
 		return Compiler->SubstrateTransmittanceToMFP(TransmittanceColor, DesiredThickness, OutputIndex);
 	}
+
+	virtual int32 SubstrateSelectParameterBlending(int32 A, int32 B, int32 SelectValue, const FString& SharedLocalBasisIndexMacro, FSubstrateOperator* PromoteToOperator) override
+	{
+		return Compiler->SubstrateSelectParameterBlending(A, B, SelectValue, SharedLocalBasisIndexMacro, PromoteToOperator);
+	}
 	
 	virtual int32 SubstrateMetalnessToDiffuseAlbedoF0(int32 BaseColor, int32 Specular, int32 Metallic, int32 OutputIndex) override
 	{
diff --git a/Engine/Source/Runtime/Engine/Public/Materials/MaterialExpressionSubstrate.h b/Engine/Source/Runtime/Engine/Public/Materials/MaterialExpressionSubstrate.h
--- a/Engine/Source/Runtime/Engine/Public/Materials/MaterialExpressionSubstrate.h
+++ b/Engine/Source/Runtime/Engine/Public/Materials/MaterialExpressionSubstrate.h
@@ -1063,6 +1063,42 @@ class UMaterialExpressionSubstrateWeight : public UMaterialExpressionSubstrateBS
 	//~ End UMaterialExpression Interface
 };
 
+UCLASS(MinimalAPI, collapsecategories, hidecategories = Object, DisplayName = "Substrate Select")
+class UMaterialExpressionSubstrateSelect : public UMaterialExpressionSubstrateBSDF
+{
+	GENERATED_UCLASS_BODY()
+
+	/**
+	 * Substrate material
+	 */
+	UPROPERTY()
+	FExpressionInput A;
+
+	/**
+	 * Substrate material
+	 */
+	UPROPERTY()
+	FExpressionInput B;
+
+	/**
+	 * If <=0, A is selected, otherwise B is selected. (Default: 0))
+	 */
+	UPROPERTY()
+	FExpressionInput SelectValue;
+
+	//~ Begin UMaterialExpression Interface
+#if WITH_EDITOR
+	virtual int32 Compile(class FMaterialCompiler* Compiler, int32 OutputIndex) override;
+	virtual void GetCaption(TArray<FString>& OutCaptions) const override;
+	virtual uint32 GetOutputType(int32 OutputIndex) override;
+	virtual uint32 GetInputType(int32 InputIndex) override;
+	virtual bool IsResultSubstrateMaterial(int32 OutputIndex) override;
+	virtual void GatherSubstrateMaterialInfo(FSubstrateMaterialInfo& SubstrateMaterialInfo, int32 OutputIndex) override;
+	virtual FSubstrateOperator* SubstrateGenerateMaterialTopologyTree(class FMaterialCompiler* Compiler, class UMaterialExpression* Parent, int32 OutputIndex) override;
+#endif
+	//~ End UMaterialExpression Interface
+};
+
 
 
 ///////////////////////////////////////////////////////////////////////////////
