Improved logging when more than one WorldDataLayers actor is found in the same container.

diff --git a/Engine/Source/Runtime/Engine/Private/WorldPartition/WorldPartition.cpp b/Engine/Source/Runtime/Engine/Private/WorldPartition/WorldPartition.cpp
--- a/Engine/Source/Runtime/Engine/Private/WorldPartition/WorldPartition.cpp
+++ b/Engine/Source/Runtime/Engine/Private/WorldPartition/WorldPartition.cpp
@@ -2220,7 +2220,11 @@ UActorDescContainerInstance* UWorldPartition::RegisterActorDescContainerInstance
 				const FWorldDataLayersActorDesc* FoundWorldDataLayerActorsDesc = StaticCast<const FWorldDataLayersActorDesc*>(ActorDesc);
 				if (FoundWorldDataLayerActorsDesc != nullptr && WorldDataLayerActorsDesc != nullptr)
 				{
-					UE_LOG(LogWorldPartition, Warning, TEXT("Extra World Data Layer '%s' actor found. Clean up invalid actors to remove the error."), *ActorDesc->GetActorPackage().ToString());
+					UE_LOG(LogWorldPartition, Warning, TEXT("Found WorldDataLayers actor '%s' (%s), but another WorldDataLayers actor '%s' (%s) already exists. Actor will be ignored, clean up the invalid actor to remove the warning."),
+						*ActorDesc->GetActorLabelOrName().ToString(),
+						*ActorDesc->GetActorPackage().ToString(),
+						*WorldDataLayerActorsDesc->GetActorLabelOrName().ToString(),
+						*WorldDataLayerActorsDesc->GetActorPackage().ToString());
 					return false;
 				}
 
