First Person: Added a function to check if first person primitives are tagged in the GBuffer and also added a global define that should be used before checking the GBuffer/SubstrateHeader for IsFirstPerson().

#rnx

diff --git a/Engine/Source/Runtime/Engine/Private/ShaderCompiler/ShaderCompiler.cpp b/Engine/Source/Runtime/Engine/Private/ShaderCompiler/ShaderCompiler.cpp
--- a/Engine/Source/Runtime/Engine/Private/ShaderCompiler/ShaderCompiler.cpp
+++ b/Engine/Source/Runtime/Engine/Private/ShaderCompiler/ShaderCompiler.cpp
@@ -3797,6 +3797,11 @@ void GlobalBeginCompileShader(
 		SET_SHADER_DEFINE(Input.Environment, USE_HAIR_TRIANGLE_STRIP, bUseTriangleStrips ? 1 : 0);
 	}
 
+	{
+		const bool bHasFirstPersonGBufferBit = HasFirstPersonGBufferBit(Target.GetPlatform());
+		SET_SHADER_DEFINE(Input.Environment, HAS_FIRST_PERSON_GBUFFER_BIT, bHasFirstPersonGBufferBit ? 1 : 0);
+	}
+
 	const bool bSubstrate = Substrate::IsSubstrateEnabled() && IsSubstrateSupportForShaderPipeline(Input);
 	{
 		SET_SHADER_DEFINE(Input.Environment, SUBSTRATE_ENABLED, bSubstrate ? 1 : 0);
diff --git a/Engine/Source/Runtime/RenderCore/Private/RenderUtils.cpp b/Engine/Source/Runtime/RenderCore/Private/RenderUtils.cpp
--- a/Engine/Source/Runtime/RenderCore/Private/RenderUtils.cpp
+++ b/Engine/Source/Runtime/RenderCore/Private/RenderUtils.cpp
@@ -2078,6 +2078,12 @@ namespace Substrate
 	}
 }
 
+bool HasFirstPersonGBufferBit(EShaderPlatform ShaderPlatform)
+{
+	// FirstPerson uses the ZERO_PRECSHADOW_MASK GBuffer bit, which is free if static lighting is disabled.
+	return !IsMobilePlatform(ShaderPlatform) && !IsForwardShadingEnabled(ShaderPlatform) && !IsStaticLightingAllowed();
+}
+
 bool DoesPlatformSupportLumenGI(EShaderPlatform Platform, bool bSkipProjectCheck)
 {
 	static IConsoleVariable* CVar = IConsoleManager::Get().FindConsoleVariable(TEXT("r.Lumen.Supported"));
diff --git a/Engine/Source/Runtime/RenderCore/Public/RenderUtils.h b/Engine/Source/Runtime/RenderCore/Public/RenderUtils.h
--- a/Engine/Source/Runtime/RenderCore/Public/RenderUtils.h
+++ b/Engine/Source/Runtime/RenderCore/Public/RenderUtils.h
@@ -561,6 +561,9 @@ inline bool IsStaticLightingAllowed()
 	return FReadOnlyCVARCache::AllowStaticLighting();
 }
 
+// Returns 'true' when first person primitives are tagged in the GBuffer.
+RENDERCORE_API bool HasFirstPersonGBufferBit(EShaderPlatform ShaderPlatform);
+
 extern RENDERCORE_API bool DoesPlatformSupportLumenGI(EShaderPlatform Platform, bool bSkipProjectCheck = false);
 
 extern RENDERCORE_API bool DoesProjectSupportLumenRayTracedTranslucentRefraction();
