Substrate - added customisable threshold to Select node as a property.

#tests Tested PC raster and PT
#rnx

diff --git a/Engine/Source/Runtime/Engine/Private/Materials/MaterialExpressionSubstrate.cpp b/Engine/Source/Runtime/Engine/Private/Materials/MaterialExpressionSubstrate.cpp
--- a/Engine/Source/Runtime/Engine/Private/Materials/MaterialExpressionSubstrate.cpp
+++ b/Engine/Source/Runtime/Engine/Private/Materials/MaterialExpressionSubstrate.cpp
@@ -3030,6 +3030,7 @@ FSubstrateOperator* UMaterialExpressionSubstrateWeight::SubstrateGenerateMateria
 
 UMaterialExpressionSubstrateSelect::UMaterialExpressionSubstrateSelect(const FObjectInitializer& ObjectInitializer)
 	: Super(ObjectInitializer)
+	, Threshold(0.5f)
 {
 	struct FConstructorStatics
 	{
@@ -3067,18 +3068,25 @@ int32 UMaterialExpressionSubstrateSelect::Compile(class FMaterialCompiler* Compi
 
 	int32 ZeroCodeChunk = Compiler->Constant(0.0f);
 	int32 OneCodeChunk = Compiler->Constant(1.0f);
-	int32 SelectLerpCodeChunk = Compiler->If(SelectValueCodeChunk, ZeroCodeChunk, OneCodeChunk, ZeroCodeChunk, ZeroCodeChunk, ZeroCodeChunk);
+	int32 ThresholdCodeChunk = Compiler->Constant(Threshold);
+	SelectValueCodeChunk = Compiler->If(SelectValueCodeChunk, ThresholdCodeChunk, OneCodeChunk, ZeroCodeChunk, ZeroCodeChunk, ZeroCodeChunk);
+	// Now, SelectValueCodeChunk is 0 or 1 for any threshold.
+
+	auto SubstrateSelectNormal = [&](int32 NormalA, int32 NormalB)
+	{
+		return Compiler->If(SelectValueCodeChunk, ZeroCodeChunk, NormalB, NormalA, NormalA, ZeroCodeChunk);
+	};
 
 	// Compute the new Normal and Tangent resulting from the selection using code chunk
 	FSubstrateOperator& SubstrateOperator = Compiler->SubstrateCompilationGetOperator(Compiler->SubstrateTreeStackGetPathUniqueId());
 	FSubstrateOperator* OperatorA = Compiler->SubstrateCompilationGetOperatorFromIndex(SubstrateOperator.LeftIndex);
 	FSubstrateOperator* OperatorB = Compiler->SubstrateCompilationGetOperatorFromIndex(SubstrateOperator.RightIndex);
-	const int32 NewNormalCodeChunk = SubstrateBlendNormal(Compiler, OperatorA->BSDFRegisteredSharedLocalBasis.NormalCodeChunk, OperatorB->BSDFRegisteredSharedLocalBasis.NormalCodeChunk, SelectLerpCodeChunk);
+	const int32 NewNormalCodeChunk = SubstrateSelectNormal( OperatorA->BSDFRegisteredSharedLocalBasis.NormalCodeChunk, OperatorB->BSDFRegisteredSharedLocalBasis.NormalCodeChunk);
 	// The tangent is optional so we treat it differently if INDEX_NONE is specified
 	int32 NewTangentCodeChunk = INDEX_NONE;
 	if (OperatorA->BSDFRegisteredSharedLocalBasis.TangentCodeChunk != INDEX_NONE && OperatorB->BSDFRegisteredSharedLocalBasis.TangentCodeChunk != INDEX_NONE)
 	{
-		NewTangentCodeChunk = SubstrateBlendNormal(Compiler, OperatorB->BSDFRegisteredSharedLocalBasis.TangentCodeChunk, OperatorA->BSDFRegisteredSharedLocalBasis.TangentCodeChunk, SelectLerpCodeChunk);
+		NewTangentCodeChunk = SubstrateSelectNormal(OperatorB->BSDFRegisteredSharedLocalBasis.TangentCodeChunk, OperatorA->BSDFRegisteredSharedLocalBasis.TangentCodeChunk);
 	}
 	else if (OperatorA->BSDFRegisteredSharedLocalBasis.TangentCodeChunk != INDEX_NONE)
 	{
@@ -3093,7 +3101,7 @@ int32 UMaterialExpressionSubstrateSelect::Compile(class FMaterialCompiler* Compi
 	int32 OutputCodeChunk = INDEX_NONE;
 	if (SubstrateOperator.bUseParameterBlending)
 	{
-		OutputCodeChunk = Compiler->SubstrateSelectParameterBlending(ACodeChunk, BCodeChunk, SelectValueCodeChunk, 
+		OutputCodeChunk = Compiler->SubstrateSelectParameterBlending(ACodeChunk, BCodeChunk, SelectValueCodeChunk,
 			Compiler->GetSubstrateSharedLocalBasisIndexMacro(NewRegisteredSharedLocalBasis), 
 			SubstrateOperator.bRootOfParameterBlendingSubTree ? &SubstrateOperator : nullptr);
 	}
diff --git a/Engine/Source/Runtime/Engine/Public/Materials/MaterialExpressionSubstrate.h b/Engine/Source/Runtime/Engine/Public/Materials/MaterialExpressionSubstrate.h
--- a/Engine/Source/Runtime/Engine/Public/Materials/MaterialExpressionSubstrate.h
+++ b/Engine/Source/Runtime/Engine/Public/Materials/MaterialExpressionSubstrate.h
@@ -1086,6 +1086,12 @@ class UMaterialExpressionSubstrateSelect : public UMaterialExpressionSubstrateBS
 	UPROPERTY()
 	FExpressionInput SelectValue;
 
+	/**
+	 * The threshold to use to select between A or B (Default: 0.5)
+	 */
+	UPROPERTY(EditAnywhere, Category = Mode, meta = (UIMin = 0.0, UIMax = 1.0, ClampMin = 0.0, ClampMax = 1.0))
+	float Threshold;
+
 	//~ Begin UMaterialExpression Interface
 #if WITH_EDITOR
 	virtual int32 Compile(class FMaterialCompiler* Compiler, int32 OutputIndex) override;
