Lumen scene - add basic skeletal mesh surface cache support

#tests tested in editor, replay a/b test

diff --git a/Engine/Source/Runtime/Engine/Private/SkeletalMeshSceneProxy.cpp b/Engine/Source/Runtime/Engine/Private/SkeletalMeshSceneProxy.cpp
--- a/Engine/Source/Runtime/Engine/Private/SkeletalMeshSceneProxy.cpp
+++ b/Engine/Source/Runtime/Engine/Private/SkeletalMeshSceneProxy.cpp
@@ -16,6 +16,8 @@
 #include "SkeletalRenderPublic.h"
 #include "SkinnedMeshSceneProxyDesc.h"
 #include "UnrealEngine.h"
+#include "MeshCardBuild.h"
+#include "ComponentRecreateRenderStateContext.h"
 
 #if WITH_EDITOR
 #include "Components/BrushComponent.h"
@@ -48,8 +50,21 @@ static FAutoConsoleVariableRef CVarSkeletalMeshUseCachedMDCs(
 	TEXT("Whether skeletal meshes will take the cached MDC path."),
 	ECVF_RenderThreadSafe);
 
+static TAutoConsoleVariable<int32> CVarMeshCardRepresentationSkeletalMesh(
+	TEXT("r.MeshCardRepresentation.SkeletalMesh"),
+	1,
+	TEXT("Whether to allow generating mesh cards for skeletal meshes."),
+	FConsoleVariableDelegate::CreateLambda([](IConsoleVariable* InVariable) { FGlobalComponentRecreateRenderStateContext Context; }),
+	ECVF_RenderThreadSafe);
+
 const FQuat SphylBasis(FVector(1.0f / FMath::Sqrt(2.0f), 0.0f, 1.0f / FMath::Sqrt(2.0f)), UE_PI);
 
+static bool AllowLumenCardGenerationForSkeletalMeshes(EShaderPlatform Platform)
+{
+	return CVarMeshCardRepresentationSkeletalMesh.GetValueOnAnyThread() != 0
+		&& DoesPlatformSupportLumenGI(Platform);
+}
+
 /** 
  * Constructor. 
  * @param	Component - skeletal mesh primitive being added
@@ -265,6 +280,18 @@ FSkeletalMeshSceneProxy::FSkeletalMeshSceneProxy(const FSkinnedMeshSceneProxyDes
 
 	bStaticElementsAlwaysUseProxyPrimitiveUniformBuffer = true;
 	bOpaqueOrMasked = MaterialRelevance.bOpaque || MaterialRelevance.bMasked;
+
+	UpdateVisibleInLumenScene();
+	UpdateLumenCardsFromBounds();
+}
+
+FSkeletalMeshSceneProxy::~FSkeletalMeshSceneProxy()
+{
+	if (CardRepresentationData)
+	{
+		delete CardRepresentationData;
+		CardRepresentationData = nullptr;
+	}
 }
 
 
@@ -614,6 +641,70 @@ void FSkeletalMeshSceneProxy::GetMeshElementsConditionallySelectable(const TArra
 #endif
 }
 
+// Same as LumenMeshCards::GetAxisAlignedDirection
+static FVector3f GetLumenCardAxisAlignedDirection(uint32 AxisAlignedDirectionIndex)
+{
+	const uint32 AxisIndex = AxisAlignedDirectionIndex / 2;
+
+	FVector3f Direction(0.0f, 0.0f, 0.0f);
+	Direction[AxisIndex] = AxisAlignedDirectionIndex & 1 ? 1.0f : -1.0f;
+	return Direction;
+}
+
+void FSkeletalMeshSceneProxy::UpdateLumenCardsFromBounds()
+{
+	if (CardRepresentationData)
+	{
+		delete CardRepresentationData;
+		CardRepresentationData = nullptr;
+	}
+
+	if (!bVisibleInLumenScene || !AllowLumenCardGenerationForSkeletalMeshes(GetFeatureLevelShaderPlatform(FeatureLevel)))
+	{
+		return;
+	}
+
+	CardRepresentationData = new FCardRepresentationData;
+	FMeshCardsBuildData& CardData = CardRepresentationData->MeshCardsBuildData;
+
+	CardData.Bounds = PreSkinnedLocalBounds.GetBox();
+	// Skeletal meshes usually doesn't match their surface cache very well due to animation.
+	// Mark as two-sided so a high sampling bias is used and hits are accepted even if they don't match well
+	CardData.bMostlyTwoSided = true;
+	CardData.CardBuildData.SetNum(6);
+
+	const FVector3f MeshCenter(CardData.Bounds.GetCenter());
+	const FVector3f MeshExtent(CardData.Bounds.GetExtent());
+	const float CardZOffset = 5.f;
+
+	for (uint32 DirectionIndex = 0; DirectionIndex < 6u; ++DirectionIndex)
+	{
+		FLumenCardOBBf OBB;
+
+		OBB.AxisZ = GetLumenCardAxisAlignedDirection(DirectionIndex);
+		OBB.AxisZ.FindBestAxisVectors(OBB.AxisX, OBB.AxisY);
+		OBB.AxisX = FVector3f::CrossProduct(OBB.AxisZ, OBB.AxisY);
+		OBB.AxisX.Normalize();
+
+		// Project 3D mesh extent onto the specific axes of this CardOBB (this just reshuffles them but the combinatorics are messy).
+		OBB.Extent.X = FMath::Abs(OBB.AxisX.Dot(MeshExtent));
+		OBB.Extent.Y = FMath::Abs(OBB.AxisY.Dot(MeshExtent));
+		OBB.Extent.Z = FMath::Abs(OBB.AxisZ.Dot(MeshExtent));
+
+		// Shrink and shift the card box to cover half the mesh bounds along Z and then bump it forward a bit
+		OBB.Extent.Z = OBB.Extent.Z * 0.5f + CardZOffset;
+		OBB.Origin = MeshCenter + OBB.Extent.Z * OBB.AxisZ;
+
+		CardData.CardBuildData[DirectionIndex].OBB = OBB;
+		CardData.CardBuildData[DirectionIndex].AxisAlignedDirectionIndex = DirectionIndex;
+	}
+}
+
+const FCardRepresentationData* FSkeletalMeshSceneProxy::GetMeshCardRepresentation() const
+{
+	return CardRepresentationData;
+}
+
 void FSkeletalMeshSceneProxy::CreateBaseMeshBatch(const FSceneView* View, const FSkeletalMeshLODRenderData& LODData, const int32 LODIndex, const int32 SectionIndex, const FSectionElementInfo& SectionElementInfo, FMeshBatch& Mesh, ESkinVertexFactoryMode VFMode) const
 {
 	Mesh.VertexFactory = MeshObject->GetSkinVertexFactory(View, LODIndex, SectionIndex, VFMode);
@@ -800,7 +891,7 @@ void FSkeletalMeshSceneProxy::CreateRenderThreadResources(FRHICommandListBase& R
 #if RHI_RAYTRACING
 bool FSkeletalMeshSceneProxy::HasRayTracingRepresentation() const
 {
-	return bRenderStatic;
+	return bRenderStatic || (CVarRayTracingSupportSkeletalMeshes.GetValueOnAnyThread() != 0 && CVarRayTracingSkeletalMeshes.GetValueOnAnyThread() != 0);
 }
 
 RayTracing::GeometryGroupHandle FSkeletalMeshSceneProxy::GetRayTracingGeometryGroupHandle() const
diff --git a/Engine/Source/Runtime/Engine/Public/SkeletalMeshSceneProxy.h b/Engine/Source/Runtime/Engine/Public/SkeletalMeshSceneProxy.h
--- a/Engine/Source/Runtime/Engine/Public/SkeletalMeshSceneProxy.h
+++ b/Engine/Source/Runtime/Engine/Public/SkeletalMeshSceneProxy.h
@@ -28,12 +28,16 @@ public:
 	ENGINE_API FSkeletalMeshSceneProxy(const USkinnedMeshComponent* Component, FSkeletalMeshRenderData* InSkelMeshRenderData);
 	ENGINE_API FSkeletalMeshSceneProxy(const FSkinnedMeshSceneProxyDesc& InMeshDesc, FSkeletalMeshRenderData* InSkelMeshRenderData);
 
+	ENGINE_API virtual ~FSkeletalMeshSceneProxy();
+
 #if WITH_EDITOR
 	ENGINE_API virtual HHitProxy* CreateHitProxies(UPrimitiveComponent* Component, TArray<TRefCountPtr<HHitProxy> >& OutHitProxies) override;
 #endif
 	ENGINE_API virtual void DrawStaticElements(FStaticPrimitiveDrawInterface* PDI) override;
 	ENGINE_API virtual void GetDynamicMeshElements(const TArray<const FSceneView*>& Views, const FSceneViewFamily& ViewFamily, uint32 VisibilityMap, FMeshElementCollector& Collector) const override;
 
+	ENGINE_API virtual const FCardRepresentationData* GetMeshCardRepresentation() const override;
+
 	ENGINE_API virtual void CreateRenderThreadResources(FRHICommandListBase& RHICmdList) override;
 
 #if RHI_RAYTRACING
@@ -206,6 +210,9 @@ protected:
 	/** The primitive's pre-skinned local space bounds. */
 	FBoxSphereBounds PreSkinnedLocalBounds;
 
+	/** Lumen cards data */
+	class FCardRepresentationData* CardRepresentationData = nullptr;
+
 #if RHI_RAYTRACING
 	RayTracing::GeometryGroupHandle RayTracingGeometryGroupHandle = INDEX_NONE;
 #endif
@@ -228,6 +235,9 @@ protected:
 
 	/** Only call on render thread timeline */
 	ENGINE_API uint8 GetCurrentFirstLODIdx_Internal() const;
+
+	ENGINE_API void UpdateLumenCardsFromBounds();
+
 private:
 	ENGINE_API void CreateBaseMeshBatch(const FSceneView* View, const FSkeletalMeshLODRenderData& LODData, const int32 LODIndex, const int32 SectionIndex, const FSectionElementInfo& SectionElementInfo, FMeshBatch& Mesh, ESkinVertexFactoryMode VFMode = ESkinVertexFactoryMode::Default) const;
 
