[Engine]
* Added ParticleEmitterInstanceOwner

diff --git a/Engine/Source/Runtime/Engine/Classes/Particles/ParticleEmitter.h b/Engine/Source/Runtime/Engine/Classes/Particles/ParticleEmitter.h
--- a/Engine/Source/Runtime/Engine/Classes/Particles/ParticleEmitter.h
+++ b/Engine/Source/Runtime/Engine/Classes/Particles/ParticleEmitter.h
@@ -14,6 +14,7 @@
 #include "ParticleEmitterInstances.h"
 #include "ParticleEmitter.generated.h"
 
+class IParticleEmitterInstanceOwner;
 class UInterpCurveEdSetup;
 class UParticleLODLevel;
 class UParticleSystemComponent;
@@ -263,7 +264,7 @@ class UParticleEmitter : public UObject
 	//~ End UObject Interface
 
 	// @todo document
-	virtual FParticleEmitterInstance* CreateInstance(UParticleSystemComponent& InComponent);
+	virtual FParticleEmitterInstance* CreateInstance(IParticleEmitterInstanceOwner& InComponent);
 
 	// Sets up this emitter with sensible defaults so we can see some particles as soon as its created.
 	virtual void SetToSensibleDefaults() {}
diff --git a/Engine/Source/Runtime/Engine/Classes/Particles/ParticleSpriteEmitter.h b/Engine/Source/Runtime/Engine/Classes/Particles/ParticleSpriteEmitter.h
--- a/Engine/Source/Runtime/Engine/Classes/Particles/ParticleSpriteEmitter.h
+++ b/Engine/Source/Runtime/Engine/Classes/Particles/ParticleSpriteEmitter.h
@@ -34,7 +34,7 @@ class UParticleSpriteEmitter : public UParticleEmitter
 	//~ End UObject Interface
 
 	//~ Begin UParticleEmitter Interface
-	virtual FParticleEmitterInstance* CreateInstance(UParticleSystemComponent& InComponent) override;
+	virtual FParticleEmitterInstance* CreateInstance(IParticleEmitterInstanceOwner& InComponent) override;
 	virtual void SetToSensibleDefaults() override;
 	//~ End UParticleEmitter Interface
 };
diff --git a/Engine/Source/Runtime/Engine/Classes/Particles/ParticleSystemComponent.h b/Engine/Source/Runtime/Engine/Classes/Particles/ParticleSystemComponent.h
--- a/Engine/Source/Runtime/Engine/Classes/Particles/ParticleSystemComponent.h
+++ b/Engine/Source/Runtime/Engine/Classes/Particles/ParticleSystemComponent.h
@@ -1313,6 +1313,10 @@ public:
 	/** Stream of random values to use with this component */
 	FRandomStream RandomStream;
 
+	/** Owner for particle emitter instances. Reroutes calls to this component. (This is used to abstract away dependency on UParticleSystemComponent) */
+	class FInstanceOwner;
+	TUniquePtr<FInstanceOwner> InstanceOwner;
+
 private:
 	/** In some cases the async work for this PSC can be created externally by the manager. */
 	FORCEINLINE void SetAsyncWork(FGraphEventRef& InAsyncWork) { AsyncWork = InAsyncWork; }
diff --git a/Engine/Source/Runtime/Engine/Classes/Particles/TypeData/ParticleModuleTypeDataAnimTrail.h b/Engine/Source/Runtime/Engine/Classes/Particles/TypeData/ParticleModuleTypeDataAnimTrail.h
--- a/Engine/Source/Runtime/Engine/Classes/Particles/TypeData/ParticleModuleTypeDataAnimTrail.h
+++ b/Engine/Source/Runtime/Engine/Classes/Particles/TypeData/ParticleModuleTypeDataAnimTrail.h
@@ -80,7 +80,7 @@ class UParticleModuleTypeDataAnimTrail : public UParticleModuleTypeDataBase
 	//~ End UParticleModule Interface
 
 	//~ Begin UParticleModuleTypeDataBase Interface
-	virtual FParticleEmitterInstance* CreateInstance(UParticleEmitter* InEmitterParent, UParticleSystemComponent& InComponent) override;
+	virtual FParticleEmitterInstance* CreateInstance(UParticleEmitter* InEmitterParent, IParticleEmitterInstanceOwner& InComponent) override;
 	virtual const FVertexFactoryType* GetVertexFactoryType() const override;
 	virtual EPrimitiveType GetPrimitiveType() const override { return PT_TriangleStrip; }
 	virtual void CollectPSOPrecacheData(const UParticleEmitter* Emitter, FPSOPrecacheParams& OutParams) override;
diff --git a/Engine/Source/Runtime/Engine/Classes/Particles/TypeData/ParticleModuleTypeDataBase.h b/Engine/Source/Runtime/Engine/Classes/Particles/TypeData/ParticleModuleTypeDataBase.h
--- a/Engine/Source/Runtime/Engine/Classes/Particles/TypeData/ParticleModuleTypeDataBase.h
+++ b/Engine/Source/Runtime/Engine/Classes/Particles/TypeData/ParticleModuleTypeDataBase.h
@@ -9,8 +9,8 @@
 #include "Particles/ParticleModule.h"
 #include "ParticleModuleTypeDataBase.generated.h"
 
+class IParticleEmitterInstanceOwner;
 class UParticleEmitter;
-class UParticleSystemComponent;
 struct FParticleEmitterInstance;
 
 UCLASS(editinlinenew, hidecategories=Object, abstract, MinimalAPI)
@@ -35,7 +35,7 @@ class UParticleModuleTypeDataBase : public UParticleModule
 	virtual bool RequiresBuild() const { return false; }
 
 	// @todo document
-	virtual FParticleEmitterInstance* CreateInstance(UParticleEmitter* InEmitterParent, UParticleSystemComponent& InComponent);
+	virtual FParticleEmitterInstance* CreateInstance(UParticleEmitter* InEmitterParent, IParticleEmitterInstanceOwner& InComponent);
 
 	virtual const FVertexFactoryType* GetVertexFactoryType() const { return nullptr; }
 	virtual EPrimitiveType GetPrimitiveType() const { return PT_TriangleList; }
diff --git a/Engine/Source/Runtime/Engine/Classes/Particles/TypeData/ParticleModuleTypeDataBeam2.h b/Engine/Source/Runtime/Engine/Classes/Particles/TypeData/ParticleModuleTypeDataBeam2.h
--- a/Engine/Source/Runtime/Engine/Classes/Particles/TypeData/ParticleModuleTypeDataBeam2.h
+++ b/Engine/Source/Runtime/Engine/Classes/Particles/TypeData/ParticleModuleTypeDataBeam2.h
@@ -222,7 +222,7 @@ class UParticleModuleTypeDataBeam2 : public UParticleModuleTypeDataBase
 
 
 	//~ Begin UParticleModuleTypeDataBase Interface
-	virtual FParticleEmitterInstance* CreateInstance(UParticleEmitter* InEmitterParent, UParticleSystemComponent& InComponent) override;
+	virtual FParticleEmitterInstance* CreateInstance(UParticleEmitter* InEmitterParent, IParticleEmitterInstanceOwner& InComponent) override;
 	virtual void CacheModuleInfo(UParticleEmitter* Emitter) override;
 	virtual const FVertexFactoryType* GetVertexFactoryType() const override;
 	virtual EPrimitiveType GetPrimitiveType() const { return PT_TriangleStrip; }
diff --git a/Engine/Source/Runtime/Engine/Classes/Particles/TypeData/ParticleModuleTypeDataGpu.h b/Engine/Source/Runtime/Engine/Classes/Particles/TypeData/ParticleModuleTypeDataGpu.h
--- a/Engine/Source/Runtime/Engine/Classes/Particles/TypeData/ParticleModuleTypeDataGpu.h
+++ b/Engine/Source/Runtime/Engine/Classes/Particles/TypeData/ParticleModuleTypeDataGpu.h
@@ -482,7 +482,7 @@ class UParticleModuleTypeDataGpu : public UParticleModuleTypeDataBase
 	//~ Begin UParticleModuleTypeDataBase Interface
 	virtual void Build( struct FParticleEmitterBuildInfo& EmitterBuildInfo ) override;
 	virtual bool RequiresBuild() const override { return true; }
-	virtual FParticleEmitterInstance* CreateInstance(UParticleEmitter* InEmitterParent, UParticleSystemComponent& InComponent) override;
+	virtual FParticleEmitterInstance* CreateInstance(UParticleEmitter* InEmitterParent, IParticleEmitterInstanceOwner& InComponent) override;
 	virtual const FVertexFactoryType* GetVertexFactoryType() const override;
 	virtual void CollectPSOPrecacheData(const UParticleEmitter* Emitter, FPSOPrecacheParams& OutParams) override;
 	//~ End UParticleModuleTypeDataBase Interface
diff --git a/Engine/Source/Runtime/Engine/Classes/Particles/TypeData/ParticleModuleTypeDataMesh.h b/Engine/Source/Runtime/Engine/Classes/Particles/TypeData/ParticleModuleTypeDataMesh.h
--- a/Engine/Source/Runtime/Engine/Classes/Particles/TypeData/ParticleModuleTypeDataMesh.h
+++ b/Engine/Source/Runtime/Engine/Classes/Particles/TypeData/ParticleModuleTypeDataMesh.h
@@ -224,7 +224,7 @@ class UParticleModuleTypeDataMesh : public UParticleModuleTypeDataBase
 	//~ End UParticleModule Interface
 
 	//~ Begin UParticleModuleTypeDataBase Interface
-	virtual FParticleEmitterInstance* CreateInstance(UParticleEmitter* InEmitterParent, UParticleSystemComponent& InComponent) override;
+	virtual FParticleEmitterInstance* CreateInstance(UParticleEmitter* InEmitterParent, IParticleEmitterInstanceOwner& InComponent) override;
 	virtual const FVertexFactoryType* GetVertexFactoryType() const override;
 	virtual void CollectPSOPrecacheData(const UParticleEmitter* Emitter, FPSOPrecacheParams& OutParams) override;
 	virtual bool	SupportsSpecificScreenAlignmentFlags() const override {	return true;	}	
diff --git a/Engine/Source/Runtime/Engine/Classes/Particles/TypeData/ParticleModuleTypeDataRibbon.h b/Engine/Source/Runtime/Engine/Classes/Particles/TypeData/ParticleModuleTypeDataRibbon.h
--- a/Engine/Source/Runtime/Engine/Classes/Particles/TypeData/ParticleModuleTypeDataRibbon.h
+++ b/Engine/Source/Runtime/Engine/Classes/Particles/TypeData/ParticleModuleTypeDataRibbon.h
@@ -13,10 +13,6 @@
 #include "Particles/TypeData/ParticleModuleTypeDataBase.h"
 #include "ParticleModuleTypeDataRibbon.generated.h"
 
-class UParticleEmitter;
-class UParticleSystemComponent;
-struct FParticleEmitterInstance;
-
 UENUM()
 enum ETrailsRenderAxisOption : int
 {
@@ -177,7 +173,7 @@ class UParticleModuleTypeDataRibbon : public UParticleModuleTypeDataBase
 	//~ End UParticleModule Interface
 
 	//~ Begin UParticleModuleTypeDataBase Interface
-	virtual FParticleEmitterInstance* CreateInstance(UParticleEmitter* InEmitterParent, UParticleSystemComponent& InComponent) override;
+	virtual FParticleEmitterInstance* CreateInstance(UParticleEmitter* InEmitterParent, IParticleEmitterInstanceOwner& InComponent) override;
 	virtual const FVertexFactoryType* GetVertexFactoryType() const override;
 	virtual EPrimitiveType GetPrimitiveType() const override { return PT_TriangleStrip; }
 	virtual void CollectPSOPrecacheData(const UParticleEmitter* Emitter, FPSOPrecacheParams& OutParams) override;
diff --git a/Engine/Source/Runtime/Engine/Private/Particles/FXSystemPrivate.h b/Engine/Source/Runtime/Engine/Private/Particles/FXSystemPrivate.h
--- a/Engine/Source/Runtime/Engine/Private/Particles/FXSystemPrivate.h
+++ b/Engine/Source/Runtime/Engine/Private/Particles/FXSystemPrivate.h
@@ -17,6 +17,7 @@ class FCanvas;
 class FGlobalDistanceFieldParameterData;
 class FParticleSimulationGPU;
 class FParticleSimulationResources;
+class IParticleEmitterInstanceOwner;
 class UParticleSystemComponent;
 class UVectorFieldComponent;
 struct FGPUSpriteEmitterInfo;
@@ -205,7 +206,7 @@ public:
 	virtual void AddVectorField(UVectorFieldComponent* VectorFieldComponent) override;
 	virtual void RemoveVectorField(UVectorFieldComponent* VectorFieldComponent) override;
 	virtual void UpdateVectorField(UVectorFieldComponent* VectorFieldComponent) override;
-	FParticleEmitterInstance* CreateGPUSpriteEmitterInstance(UParticleSystemComponent& InComponent, FGPUSpriteEmitterInfo& EmitterInfo);
+	FParticleEmitterInstance* CreateGPUSpriteEmitterInstance(IParticleEmitterInstanceOwner& InComponent, FGPUSpriteEmitterInfo& EmitterInfo);
 	virtual void PreInitViews(class FRDGBuilder& GraphBuilder, bool bAllowGPUParticleUpdate, const TArrayView<const FSceneViewFamily*>& ViewFamilies, const FSceneViewFamily* CurrentFamily) override;
 	virtual void PostInitViews(FRDGBuilder& GraphBuilder, TConstStridedView<FSceneView> Views, bool bAllowGPUParticleUpdate) override;
 	virtual bool UsesGlobalDistanceField() const override;
diff --git a/Engine/Source/Runtime/Engine/Private/Particles/ParticleBeam2EmitterInstance.cpp b/Engine/Source/Runtime/Engine/Private/Particles/ParticleBeam2EmitterInstance.cpp
--- a/Engine/Source/Runtime/Engine/Private/Particles/ParticleBeam2EmitterInstance.cpp
+++ b/Engine/Source/Runtime/Engine/Private/Particles/ParticleBeam2EmitterInstance.cpp
@@ -10,6 +10,7 @@
 #include "Materials/Material.h"
 #include "MaterialDomain.h"
 #include "ParticleEmitterInstances.h"
+#include "ParticleEmitterInstanceOwner.h"
 #include "Particles/Beam/ParticleModuleBeamBase.h"
 #include "Particles/ParticleSystemComponent.h"
 #include "Particles/Beam/ParticleModuleBeamModifier.h"
@@ -53,7 +54,7 @@ DECLARE_CYCLE_STAT(TEXT("BeamEmitterInstance Init"), STAT_BeamEmitterInstance_In
  *	Structure for beam emitter instances
  */
 /** Constructor	*/
-FParticleBeam2EmitterInstance::FParticleBeam2EmitterInstance(UParticleSystemComponent& InComponent) :
+FParticleBeam2EmitterInstance::FParticleBeam2EmitterInstance(IParticleEmitterInstanceOwner& InComponent) :
 	FParticleEmitterInstance(InComponent)
 	, BeamTypeData(NULL)
 	, BeamModule_Source(NULL)
@@ -470,7 +471,6 @@ void FParticleBeam2EmitterInstance::Init()
 void FParticleBeam2EmitterInstance::Tick(float DeltaTime, bool bSuppressSpawning)
 {
 	SCOPE_CYCLE_COUNTER(STAT_BeamTickTime);
-	if (Component)
 	{
 		UParticleLODLevel* LODLevel = SpriteTemplate->GetCurrentLODLevel(this);
 		check(LODLevel);	// TTP #33141
@@ -494,7 +494,7 @@ void FParticleBeam2EmitterInstance::Tick(float DeltaTime, bool bSuppressSpawning
 					// and focus strictly on the BurstList...
 					float SpawnRate = 0.0f;
 					// Figure out spawn rate for this tick.
-					SpawnRate = LODLevel->SpawnModule->Rate.GetValue(EmitterTime, Component);
+					SpawnRate = LODLevel->SpawnModule->Rate.GetValue(EmitterTime, Component.GetDistributionData());
 					// Take Bursts into account as well...
 					int32		Burst = 0;
 					float	BurstTime = GetCurrentBurstRateOffset(DeltaTime, Burst);
@@ -635,19 +635,18 @@ void FParticleBeam2EmitterInstance::SetCurrentLODIndex(int32 InLODIndex, bool bI
 void FParticleBeam2EmitterInstance::UpdateBoundingBox(float DeltaTime)
 {
 	SCOPE_CYCLE_COUNTER(STAT_ParticleUpdateBounds);
-	if (Component)
 	{
-		bool bUpdateBox = ((Component->bWarmingUp == false) &&
-			(Component->Template != NULL) && (Component->Template->bUseFixedRelativeBoundingBox == false));
+		const UParticleSystem* Template = Component.GetTemplate();
+		bool bUpdateBox = ((Component.GetIsWarmingUp() == false) && (Template != NULL) && (Template->bUseFixedRelativeBoundingBox == false));
 		float MaxSizeScale	= 1.0f;
 		if (bUpdateBox)
 		{
 			ParticleBoundingBox.Init();
 
 			//@todo. Currently, we don't support UseLocalSpace for beams
-			//if (Template->UseLocalSpace == false) 
+			//if (Template.UseLocalSpace == false) 
 			{
-				ParticleBoundingBox += Component->GetComponentLocation();
+				ParticleBoundingBox += Component.GetComponentTransform().GetLocation();
 			}
 		}
 
@@ -660,7 +659,7 @@ void FParticleBeam2EmitterInstance::UpdateBoundingBox(float DeltaTime)
 		}
 
 		// Take scale into account as well
-		FVector Scale = Component->GetComponentTransform().GetScale3D();
+		FVector Scale = Component.GetComponentTransform().GetScale3D();
 
 		bool bSkipDoubleSpawnUpdate = !SpriteTemplate->bUseLegacySpawningBehavior;
 		// Take each particle into account
@@ -726,11 +725,10 @@ void FParticleBeam2EmitterInstance::UpdateBoundingBox(float DeltaTime)
  */
 void FParticleBeam2EmitterInstance::ForceUpdateBoundingBox()
 {
-	if (Component)
 	{
 		float MaxSizeScale	= 1.0f;
 		ParticleBoundingBox.Init();
-		ParticleBoundingBox += Component->GetComponentLocation();
+		ParticleBoundingBox += Component.GetComponentTransform().GetLocation();
 
 		FVector	NoiseMin(0.0f);
 		FVector NoiseMax(0.0f);
@@ -741,7 +739,7 @@ void FParticleBeam2EmitterInstance::ForceUpdateBoundingBox()
 		}
 
 		// Take scale into account as well
-		FVector Scale = Component->GetComponentTransform().GetScale3D();
+		FVector Scale = Component.GetComponentTransform().GetScale3D();
 
 		// Take each particle into account
 		for (int32 i=0; i<ActiveParticles; i++)
@@ -1014,10 +1012,11 @@ void FParticleBeam2EmitterInstance::ResolveSource()
 			{
 			case PEB2STM_Actor:
 				{
+					TArrayView<const FParticleSysParam> InstanceParameters = Component.GetInstanceParameters();
 					FParticleSysParam Param;
-					for (int32 i = 0; i < Component->InstanceParameters.Num(); i++)
+					for (int32 i = 0; i < InstanceParameters.Num(); i++)
 					{
-						Param = Component->InstanceParameters[i];
+						Param = InstanceParameters[i];
 						if (Param.Name == BeamModule_Source->SourceName)
 						{
 							SourceActor = Param.Actor;
@@ -1030,9 +1029,10 @@ void FParticleBeam2EmitterInstance::ResolveSource()
 			case PEB2STM_Particle:
 				if (SourceEmitter == NULL)
 				{
-					for (int32 ii = 0; ii < Component->EmitterInstances.Num(); ii++)
+					TArrayView<FParticleEmitterInstance*> EmitterInstances = Component.GetEmitterInstances();
+					for (int32 ii = 0; ii < EmitterInstances.Num(); ii++)
 					{
-						FParticleEmitterInstance* pkEmitInst = Component->EmitterInstances[ii];
+						FParticleEmitterInstance* pkEmitInst = EmitterInstances[ii];
 						if (pkEmitInst && (pkEmitInst->SpriteTemplate->EmitterName == BeamModule_Source->SourceName))
 						{
 							SourceEmitter = pkEmitInst;
@@ -1060,10 +1060,11 @@ void FParticleBeam2EmitterInstance::ResolveTarget()
 			{
 			case PEB2STM_Actor:
 				{
+					TArrayView<const FParticleSysParam> InstanceParameters = Component.GetInstanceParameters();
 					FParticleSysParam Param;
-					for (int32 i = 0; i < Component->InstanceParameters.Num(); i++)
+					for (int32 i = 0; i < InstanceParameters.Num(); i++)
 					{
-						Param = Component->InstanceParameters[i];
+						Param = InstanceParameters[i];
 						if (Param.Name == BeamModule_Target->TargetName)
 						{
 							TargetActor = Param.Actor;
@@ -1076,9 +1077,10 @@ void FParticleBeam2EmitterInstance::ResolveTarget()
 			case PEB2STM_Particle:
 				if (TargetEmitter == NULL)
 				{
-					for (int32 ii = 0; ii < Component->EmitterInstances.Num(); ii++)
+					TArrayView<FParticleEmitterInstance*> EmitterInstances = Component.GetEmitterInstances();
+					for (int32 ii = 0; ii < EmitterInstances.Num(); ii++)
 					{
-						FParticleEmitterInstance* pkEmitInst = Component->EmitterInstances[ii];
+						FParticleEmitterInstance* pkEmitInst = EmitterInstances[ii];
 						if (pkEmitInst && (pkEmitInst->SpriteTemplate->EmitterName == BeamModule_Target->TargetName))
 						{
 							TargetEmitter = pkEmitInst;
@@ -1181,11 +1183,11 @@ FDynamicEmitterDataBase* FParticleBeam2EmitterInstance::GetDynamicData(bool bSel
 		)
 	{
 #if !(UE_BUILD_SHIPPING || UE_BUILD_TEST)
-		if (Component && Component->GetWorld())
+		if (Component.GetWorld())
 		{
 			FString ErrorMessage = 
 				FString::Printf(TEXT("BeamEmitter with too much data: %s"),
-					Component->Template ? *Component->Template->GetName() : TEXT("No template"));
+					Component.GetTemplate() ? *Component.GetTemplate()->GetName() : TEXT("No template"));
 			FColor ErrorColor(255,0,0);
 			GEngine->AddOnScreenDebugMessage((uint64)((PTRINT)this), 5.0f, ErrorColor,ErrorMessage);
 			UE_LOG(LogParticles, Log, TEXT("%s"), *ErrorMessage);
@@ -1260,7 +1262,7 @@ void FParticleBeam2EmitterInstance::GetAllocatedSize(int32& OutNum, int32& OutMa
  */
 void FParticleBeam2EmitterInstance::GetResourceSizeEx(FResourceSizeEx& CumulativeResourceSize)
 {
-	if (CumulativeResourceSize.GetResourceSizeMode() == EResourceSizeMode::EstimatedTotal || (Component && Component->SceneProxy))
+	if (CumulativeResourceSize.GetResourceSizeMode() == EResourceSizeMode::EstimatedTotal || (Component.GetSceneProxy()))
 	{
 		int32 MaxActiveParticleDataSize = (ParticleData != NULL) ? (MaxActiveParticles * ParticleStride) : 0;
 		int32 MaxActiveParticleIndexSize = 0;
@@ -1398,14 +1400,14 @@ bool FParticleBeam2EmitterInstance::FillReplayData( FDynamicEmitterReplayDataBas
 		NewReplayData->NoiseTangentStrength= BeamModule_Noise->NoiseTangentStrength.GetValue(EmitterTime);
 		if (BeamModule_Noise->bNRScaleEmitterTime)
 		{
-			NewReplayData->NoiseRangeScale = BeamModule_Noise->NoiseRangeScale.GetValue(EmitterTime, Component);
+			NewReplayData->NoiseRangeScale = BeamModule_Noise->NoiseRangeScale.GetValue(EmitterTime, Component.GetDistributionData());
 		}
 		else
 		{	//-V523 Remove when todo will be implemented
 			[REDACTED_EMAIL]. Need to address this!!!!
 			//					check(0 && TEXT("NoiseRangeScale - No way to get per-particle setting at this time."));
 			//					NewReplayData->NoiseRangeScale	= BeamModule_Noise->NoiseRangeScale.GetValue(Particle->RelativeTime, Component);
-			NewReplayData->NoiseRangeScale = BeamModule_Noise->NoiseRangeScale.GetValue(EmitterTime, Component);
+			NewReplayData->NoiseRangeScale = BeamModule_Noise->NoiseRangeScale.GetValue(EmitterTime, Component.GetDistributionData());
 		}
 		NewReplayData->NoiseSpeed = (FVector3f)BeamModule_Noise->NoiseSpeed.GetValue(EmitterTime);
 		NewReplayData->NoiseLockTime = BeamModule_Noise->NoiseLockTime;
diff --git a/Engine/Source/Runtime/Engine/Private/Particles/ParticleBeamModules.cpp b/Engine/Source/Runtime/Engine/Private/Particles/ParticleBeamModules.cpp
--- a/Engine/Source/Runtime/Engine/Private/Particles/ParticleBeamModules.cpp
+++ b/Engine/Source/Runtime/Engine/Private/Particles/ParticleBeamModules.cpp
@@ -7,6 +7,7 @@
 #include "Particles/ParticleSystem.h"
 #include "Particles/ParticleModule.h"
 #include "ParticleEmitterInstances.h"
+#include "ParticleEmitterInstanceOwner.h"
 #include "Particles/ParticleSystemComponent.h"
 #include "Distributions/DistributionFloatConstant.h"
 #include "Distributions/DistributionVectorConstant.h"
@@ -127,7 +128,7 @@ void UParticleModuleTypeDataBeam2::Spawn(const FSpawnContext& Context)
 	// If there is no Source module, use the emitter position as the source point
 	if (BeamInst->BeamModule_Source == NULL)
 	{
-		BeamData->SourcePoint	= Owner->Component->GetComponentLocation();
+		BeamData->SourcePoint	= Owner->Component.GetComponentTransform().GetLocation();
 		BeamData->SourceTangent = Context.GetTransform().GetScaledAxis( EAxis::X );
 		BeamData->SourceStrength = 1.0f;
 	}
@@ -246,7 +247,7 @@ void UParticleModuleTypeDataBeam2::Update(const FUpdateContext& Context)
 		// If there is no Source module, use the emitter location
 		if (BeamInst->BeamModule_Source == NULL)
 		{
-			BeamData->SourcePoint	= Owner->Component->GetComponentLocation();
+			BeamData->SourcePoint	= Owner->Component.GetComponentTransform().GetLocation();
 			BeamData->SourceTangent = Context.GetTransform().GetScaledAxis( EAxis::X );
 		}
 
@@ -686,7 +687,7 @@ void UParticleModuleTypeDataBeam2::PostEditChangeProperty(FPropertyChangedEvent&
 }
 #endif // WITH_EDITOR
 
-FParticleEmitterInstance* UParticleModuleTypeDataBeam2::CreateInstance(UParticleEmitter* InEmitterParent, UParticleSystemComponent& InComponent)
+FParticleEmitterInstance* UParticleModuleTypeDataBeam2::CreateInstance(UParticleEmitter* InEmitterParent, IParticleEmitterInstanceOwner& InComponent)
 {
 	SetToSensibleDefaults(InEmitterParent);
 	FParticleEmitterInstance* Instance = new FParticleBeam2EmitterInstance(InComponent);
@@ -1906,7 +1907,7 @@ bool UParticleModuleBeamSource::ResolveSourceData(const FContext& Context, FPart
 			break;
 		case PEB2STM_Emitter:
 			// The position of the owner component is the source
-			BeamData->SourcePoint	= BeamInst->Component->GetComponentLocation();
+			BeamData->SourcePoint	= BeamInst->Component.GetComponentTransform().GetLocation();
 			bSetSource				= true;
 			break;
 		case PEB2STM_Particle:
@@ -1972,15 +1973,15 @@ bool UParticleModuleBeamSource::ResolveSourceData(const FContext& Context, FPart
 
 						// Grab the particle
 
-						UParticleSystemComponent* Comp = BeamInst->SourceEmitter->Component;
+						IParticleEmitterInstanceOwner& Comp = BeamInst->SourceEmitter->Component;
 						FBaseParticle* SourceParticle = (SourceIndex >= 0) ? BeamInst->SourceEmitter->GetParticleDirect(SourceIndex) : nullptr;
 						if (SourceParticle != nullptr)
 						{
 							const FVector WorldOrigin = BeamInst->SourceEmitter->SimulationToWorld.GetOrigin();
 
-							if (!ensureMsgf(!SourceParticle->Location.ContainsNaN(), TEXT("NaN in SourceParticle Location. Template: %s, Component: %s"), Comp ? *GetNameSafe(Comp->Template) : TEXT("UNKNOWN"), *GetPathNameSafe(Comp)) ||
-								!ensureMsgf(!SourceParticle->OldLocation.ContainsNaN(), TEXT("NaN in SourceParticle OldLocation. Template: %s, Component: %s"), Comp ? *GetNameSafe(Comp->Template) : TEXT("UNKNOWN"), *GetPathNameSafe(Comp)) ||
-								!ensureMsgf(!WorldOrigin.ContainsNaN(), TEXT("NaN in WorldOrigin. Template: %s, Component: %s"), Comp ? *GetNameSafe(Comp->Template) : TEXT("UNKNOWN"), *GetPathNameSafe(Comp))
+							if (!ensureMsgf(!SourceParticle->Location.ContainsNaN(), TEXT("NaN in SourceParticle Location. Template: %s, Component: %s"), *GetNameSafe(Comp.GetTemplate()), *Comp.GetPathName()) ||
+								!ensureMsgf(!SourceParticle->OldLocation.ContainsNaN(), TEXT("NaN in SourceParticle OldLocation. Template: %s, Component: %s"), *GetNameSafe(Comp.GetTemplate()), *Comp.GetPathName()) ||
+								!ensureMsgf(!WorldOrigin.ContainsNaN(), TEXT("NaN in WorldOrigin. Template: %s, Component: %s"), *GetNameSafe(Comp.GetTemplate()), *Comp.GetPathName())
 								)
 							{
 								// Contains NaN!
@@ -1995,7 +1996,7 @@ bool UParticleModuleBeamSource::ResolveSourceData(const FContext& Context, FPart
 						else
 						{
 							// Fall back to the emitter location??
-							BeamData->SourcePoint = Comp->GetComponentLocation();
+							BeamData->SourcePoint = Comp.GetComponentTransform().GetLocation();
 							bSetSource = true;
 
 							//@todo. How to handle this... can potentially cause a jump from the emitter to the
@@ -2029,7 +2030,7 @@ bool UParticleModuleBeamSource::ResolveSourceData(const FContext& Context, FPart
 				// Use the value as a world space position
 				BeamData->SourcePoint	= Source.GetValue(BeamInst->EmitterTime, Context.GetDistributionData());
 				// Take into account current world origin offset
-				BeamData->SourcePoint  -= (FVector)BeamInst->Component->GetWorld()->OriginLocation;
+				BeamData->SourcePoint  -= (FVector)BeamInst->Component.GetWorld()->OriginLocation;
 			}
 			else
 			{
@@ -2438,15 +2439,15 @@ bool UParticleModuleBeamTarget::ResolveTargetData(const FContext& Context, FPart
 
 							// Grab the particle
 
-							UParticleSystemComponent* Comp = BeamInst->TargetEmitter->Component;
+							IParticleEmitterInstanceOwner& Comp = BeamInst->TargetEmitter->Component;
 							FBaseParticle* TargetParticle = (TargetIndex >= 0) ? BeamInst->TargetEmitter->GetParticleDirect(TargetIndex) : nullptr;
 							if (TargetParticle != nullptr)
 							{
 								const FVector WorldOrigin = BeamInst->TargetEmitter->SimulationToWorld.GetOrigin();
 
-								if (!ensureMsgf(!TargetParticle->Location.ContainsNaN(), TEXT("NaN in TargetParticle Location. Template: %s, Component: %s"), Comp ? *GetNameSafe(Comp->Template) : TEXT("UNKNOWN"), *GetPathNameSafe(Comp)) ||
-									!ensureMsgf(!TargetParticle->OldLocation.ContainsNaN(), TEXT("NaN in TargetParticle OldLocation. Template: %s, Component: %s"), Comp ? *GetNameSafe(Comp->Template) : TEXT("UNKNOWN"), *GetPathNameSafe(Comp)) ||
-									!ensureMsgf(!WorldOrigin.ContainsNaN(), TEXT("NaN in WorldOrigin. Template: %s, Component: %s"), Comp ? *GetNameSafe(Comp->Template) : TEXT("UNKNOWN"), *GetPathNameSafe(Comp))
+								if (!ensureMsgf(!TargetParticle->Location.ContainsNaN(), TEXT("NaN in TargetParticle Location. Template: %s, Component: %s"), *GetNameSafe(Comp.GetTemplate()), *Comp.GetPathName()) ||
+									!ensureMsgf(!TargetParticle->OldLocation.ContainsNaN(), TEXT("NaN in TargetParticle OldLocation. Template: %s, Component: %s"), *GetNameSafe(Comp.GetTemplate()), *Comp.GetPathName()) ||
+									!ensureMsgf(!WorldOrigin.ContainsNaN(), TEXT("NaN in WorldOrigin. Template: %s, Component: %s"), *GetNameSafe(Comp.GetTemplate()), *Comp.GetPathName())
 									)
 								{
 									// Contains NaN!
@@ -2460,7 +2461,7 @@ bool UParticleModuleBeamTarget::ResolveTargetData(const FContext& Context, FPart
 							else
 							{
 								// Fall back to the emitter location??
-								LocalTargetPosition = BeamInst->TargetEmitter->Component->GetComponentLocation();
+								LocalTargetPosition = BeamInst->TargetEmitter->Component.GetComponentTransform().GetLocation();
 								//@todo. How to handle this... can potentially cause a jump from the emitter to the
 								// particle...
 								TargetIndex = -1;//No valid particle Target;
@@ -2493,7 +2494,7 @@ bool UParticleModuleBeamTarget::ResolveTargetData(const FContext& Context, FPart
 			{
 				BeamData->TargetPoint	= Target.GetValue(BeamInst->EmitterTime, Context.GetDistributionData());
 				// Take into account current world origin offset
-				BeamData->TargetPoint  -= (FVector)BeamInst->Component->GetWorld()->OriginLocation;
+				BeamData->TargetPoint  -= (FVector)BeamInst->Component.GetWorld()->OriginLocation;
 			}
 			else
 			{
diff --git a/Engine/Source/Runtime/Engine/Private/Particles/ParticleComponents.cpp b/Engine/Source/Runtime/Engine/Private/Particles/ParticleComponents.cpp
--- a/Engine/Source/Runtime/Engine/Private/Particles/ParticleComponents.cpp
+++ b/Engine/Source/Runtime/Engine/Private/Particles/ParticleComponents.cpp
@@ -22,6 +22,7 @@
 #include "Particles/Emitter.h"
 #include "Particles/Orientation/ParticleModuleOrientationAxisLock.h"
 #include "ParticleEmitterInstances.h"
+#include "ParticleEmitterInstanceOwner.h"
 #include "Particles/ParticleModule.h"
 #include "Particles/ParticleSystemComponent.h"
 #include "Distributions/DistributionFloatConstant.h"
@@ -87,6 +88,7 @@
 #include "Particles/SubUV/ParticleModuleSubUV.h"
 #include "Particles/ParticleSystemManager.h"
 
+#include UE_INLINE_GENERATED_CPP_BY_NAME(ParticleSystemComponent)
 
 DECLARE_CYCLE_STAT(TEXT("ParticleComponent InitParticles GT"), STAT_ParticleSystemComponent_InitParticles, STATGROUP_Particles);
 DECLARE_CYCLE_STAT(TEXT("ParticleComponent SendRenderDynamicData GT"), STAT_ParticleSystemComponent_SendRenderDynamicData_Concurrent, STATGROUP_Particles);
@@ -922,7 +924,7 @@ UParticleEmitter::UParticleEmitter(const FObjectInitializer& ObjectInitializer)
 
 }
 
-FParticleEmitterInstance* UParticleEmitter::CreateInstance(UParticleSystemComponent& InComponent)
+FParticleEmitterInstance* UParticleEmitter::CreateInstance(IParticleEmitterInstanceOwner& InComponent)
 {
 	UE_LOG(LogParticles, Fatal,TEXT("UParticleEmitter::CreateInstance is pure virtual")); 
 	return NULL; 
@@ -1564,12 +1566,12 @@ void UParticleEmitter::EditorUpdateCurrentLOD(FParticleEmitterInstance* Instance
 	UParticleLODLevel*	Higher			= NULL;
 
 	int32 SetLODLevel = -1;
-	if (Instance->Component && Instance->Component->Template)
+	if (const UParticleSystem* Template = Instance->Component.GetTemplate())
 	{
-		int32 DesiredLODLevel = Instance->Component->Template->EditorLODSetting;
+		int32 DesiredLODLevel = Template->EditorLODSetting;
 		if (GIsEditor && GEngine->bEnableEditorPSysRealtimeLOD)
 		{
-			DesiredLODLevel = Instance->Component->GetCurrentLODIndex();
+			DesiredLODLevel = Instance->Component.GetCurrentLODIndex();
 		}
 
 		for (int32 LevelIndex = 0; LevelIndex < LODLevels.Num(); LevelIndex++)
@@ -2050,7 +2052,7 @@ void UParticleSpriteEmitter::PostLoad()
 	}
 }
 
-FParticleEmitterInstance* UParticleSpriteEmitter::CreateInstance(UParticleSystemComponent& InComponent)
+FParticleEmitterInstance* UParticleSpriteEmitter::CreateInstance(IParticleEmitterInstanceOwner& InComponent)
 {
 	// If this emitter was cooked out or has no valid LOD levels don't create an instance for it.
 	if ((bCookedOut == true) || (LODLevels.Num() == 0))
@@ -3588,6 +3590,78 @@ void UFXSystemComponent::PrecacheAssetPSOs(UFXSystemAsset* FXSystemAsset)
 #endif // UE_WITH_PSO_PRECACHING
 }
 
+
+class UParticleSystemComponent::FInstanceOwner : public IParticleEmitterInstanceOwner
+{
+public:
+	FInstanceOwner(UParticleSystemComponent* InComponent) : Component(InComponent) {}
+
+	virtual const FTransform& GetAsyncComponentToWorld() const override { return Component->GetAsyncComponentToWorld(); }
+
+	virtual UObject* GetDistributionData() const override { return Component; }
+	virtual const FTransform& GetComponentTransform() const override { return Component->GetComponentTransform(); }
+	virtual FRotator GetComponentRotation() const override { return Component->GetComponentRotation(); }
+	virtual const FTransform& GetComponentToWorld() const override { return Component->GetComponentToWorld(); }
+	virtual const FBoxSphereBounds& GetBounds() const override { return Component->Bounds; }
+	virtual UWorld* GetWorld() const override { return Component->GetWorld(); }
+	virtual bool GetFloatParameter(const FName InName, float& OutFloat) override { return Component->GetFloatParameter(InName, OutFloat); }
+	virtual const FVector3f& GetLWCTile() const override { return Component->GetLWCTile(); }
+	virtual FString GetName() const override { return Component->GetName(); }
+	virtual FString GetFullName() const override { return Component->GetFullName(); }
+	virtual FString GetPathName() const override { return Component->GetPathName(); }
+	virtual bool IsActive() const override { return Component->IsActive(); }
+	virtual bool IsValidLowLevel() const override { return Component->IsValidLowLevel(); }
+	virtual TArrayView<const FParticleSysParam> GetAsyncInstanceParameters() override { return Component->GetAsyncInstanceParameters(); }
+	virtual int32 GetCurrentDetailMode() const override { return Component->GetCurrentDetailMode(); }
+	virtual int32 GetCurrentLODIndex() const override { return Component->GetCurrentLODIndex(); }
+	virtual const FVector& GetPartSysVelocity() const override { return Component->PartSysVelocity; }
+	virtual const FVector& GetOldPosition() const override { return Component->OldPosition; }
+	virtual FFXSystem* GetFXSystem() const override { return Component->FXSystem; }
+	virtual const UParticleSystem* GetTemplate() const override { return Component->Template; }
+	virtual TArrayView<const FParticleSysParam> GetInstanceParameters() const override { return Component->InstanceParameters; }
+	virtual TArrayView<FParticleEmitterInstance*> GetEmitterInstances() const override { return Component->EmitterInstances; }
+	virtual TArrayView<TObjectPtr<UMaterialInterface>> GetEmitterMaterials() const override { return Component->EmitterMaterials; }
+	virtual FPrimitiveSceneProxy* GetSceneProxy() const override { return Component->SceneProxy; }
+	virtual bool GetIsWarmingUp() const override { return Component->bWarmingUp; }
+	virtual bool GetJustRegistered() const override { return Component->bJustRegistered; }
+	virtual float GetWarmupTime() const override { return Component->WarmupTime; }
+	virtual float GetEmitterDelay() const override { return Component->EmitterDelay; }
+	virtual FRandomStream& GetRandomStream() const override { return Component->RandomStream; }
+
+	virtual void SetComponentToWorld(const FTransform& NewComponentToWorld) override { Component->SetComponentToWorld(NewComponentToWorld); }
+	virtual void DeactivateNextTick() override { Component->DeactivaateNextTick(); }
+
+	virtual UParticleSystemComponent* AsComponent() const override { return Component; }
+
+	virtual void ReportEventSpawn(const FName InEventName, const float InEmitterTime, const FVector InLocation, const FVector InVelocity, const TArray<class UParticleModuleEventSendToGame*>& InEventData) override
+	{
+		Component->ReportEventSpawn(InEventName, InEmitterTime, InLocation, InVelocity, InEventData);
+	}
+
+	virtual void ReportEventDeath(const FName InEventName, const float InEmitterTime, const FVector InLocation, const FVector InVelocity, const TArray<class UParticleModuleEventSendToGame*>& InEventData, const float InParticleTime) override
+	{
+		Component->ReportEventDeath(InEventName, InEmitterTime, InLocation, InVelocity, InEventData, InParticleTime);
+	}
+
+	virtual void ReportEventCollision(const FName InEventName, const float InEmitterTime, const FVector InLocation, const FVector InDirection, const FVector InVelocity, const TArray<class UParticleModuleEventSendToGame*>& InEventData, const float InParticleTime, const FVector InNormal, const float InTime, const int32 InItem, const FName InBoneName, UPhysicalMaterial* PhysMat) override
+	{
+		Component->ReportEventCollision(InEventName, InEmitterTime, InLocation, InDirection, InVelocity, InEventData, InParticleTime, InNormal, InTime, InItem, InBoneName, PhysMat);
+	}
+
+	virtual void ReportEventBurst(const FName InEventName, const float InEmitterTime, const int32 ParticleCount, const FVector InLocation, const TArray<class UParticleModuleEventSendToGame*>& InEventData) override
+	{
+		Component->ReportEventBurst(InEventName, InEmitterTime, ParticleCount, InLocation, InEventData);
+	}
+
+	virtual TArrayView<FParticleEventSpawnData> GetSpawnEvents() const override { return Component->SpawnEvents; }
+	virtual TArrayView<FParticleEventDeathData> GetDeathEvents() const override { return Component->DeathEvents; }
+	virtual TArrayView<FParticleEventCollideData> GetCollisionEvents() const override { return Component->CollisionEvents; }
+	virtual TArrayView<FParticleEventBurstData> GetBurstEvents() const override { return Component->BurstEvents; }
+	virtual TArrayView<FParticleEventKismetData> GetKismetEvents() const override { return Component->KismetEvents; }
+
+	UParticleSystemComponent* Component;
+};
+
 FOnSystemPreActivationChange UParticleSystemComponent::OnSystemPreActivationChange;
 
 UParticleSystemComponent::UParticleSystemComponent(const FObjectInitializer& ObjectInitializer)
@@ -5966,7 +6040,12 @@ void UParticleSystemComponent::InitParticles()
 					}
 					else
 					{
-						Instance = Emitter->CreateInstance(*this);
+						if (!InstanceOwner)
+						{
+							InstanceOwner = MakeUnique<FInstanceOwner>(this);
+						}
+
+						Instance = Emitter->CreateInstance(*InstanceOwner);
 						EmitterInstances[Idx] = Instance;
 					}
 
@@ -6059,7 +6138,6 @@ void UParticleSystemComponent::ResetParticles(bool bEmptyInstances)
 				if (!(!bIsGameWorld || bEmptyInstances))
 				{
 					EmitterInstance->SpriteTemplate	= NULL;
-					EmitterInstance->Component	= NULL;
 				}
 			}
 		}
diff --git a/Engine/Source/Runtime/Engine/Private/Particles/ParticleEmitterInstances.cpp b/Engine/Source/Runtime/Engine/Private/Particles/ParticleEmitterInstances.cpp
--- a/Engine/Source/Runtime/Engine/Private/Particles/ParticleEmitterInstances.cpp
+++ b/Engine/Source/Runtime/Engine/Private/Particles/ParticleEmitterInstances.cpp
@@ -5,6 +5,7 @@
 =============================================================================*/
 
 #include "ParticleEmitterInstances.h"
+#include "ParticleEmitterInstanceOwner.h"
 #include "Engine/Engine.h"
 #include "Engine/World.h"
 #include "Materials/Material.h"
@@ -391,9 +392,9 @@ FParticleEmitterBuildInfo::FParticleEmitterBuildInfo()
 const float FParticleEmitterInstance::PeakActiveParticleUpdateDelta = 0.05f;
 
 /** Constructor	*/
-FParticleEmitterInstance::FParticleEmitterInstance(UParticleSystemComponent& InComponent) :
+FParticleEmitterInstance::FParticleEmitterInstance(IParticleEmitterInstanceOwner& InComponent) :
 	  SpriteTemplate(NULL)
-    , Component(&InComponent)
+    , Component(InComponent)
     , CurrentLODLevel(NULL)
     , CurrentLODLevelIndex(0)
     , TypeDataOffset(0)
@@ -564,7 +565,7 @@ void FParticleEmitterInstance::Init()
 	ParticleCounter			= 0;
 
 	UpdateTransforms();	
-	Location				= Component->GetComponentLocation();
+	Location				= Component.GetComponentTransform().GetLocation();
 	OldLocation				= Location;
 	
 	TrianglesToRender		= 0;
@@ -588,7 +589,7 @@ void FParticleEmitterInstance::Init()
 
 	// Resize to sensible default.
 	if (bNeedsInit && 
-		Component->GetWorld()->IsGameWorld() == true &&
+		Component.GetWorld()->IsGameWorld() == true &&
 		// Only presize if any particles will be spawned 
 		SpriteTemplate->QualityLevelSpawnRateScale > 0)
 	{
@@ -676,7 +677,7 @@ void FParticleEmitterInstance::Init()
 
 UWorld* FParticleEmitterInstance::GetWorld() const
 {
-	return Component->GetWorld();
+	return Component.GetWorld();
 }
 
 void FParticleEmitterInstance::UpdateTransforms()
@@ -686,8 +687,7 @@ void FParticleEmitterInstance::UpdateTransforms()
 	check(SpriteTemplate != NULL);
 
 	UParticleLODLevel* LODLevel = GetCurrentLODLevelChecked();
-	FMatrix ComponentToWorld = Component != NULL ?
-		Component->GetComponentToWorld().ToMatrixNoScale() : FMatrix::Identity;
+	FMatrix ComponentToWorld = Component.GetComponentToWorld().ToMatrixNoScale();
 	FMatrix EmitterToComponent = FRotationTranslationMatrix(
 		LODLevel->RequiredModule->EmitterRotation,
 		LODLevel->RequiredModule->EmitterOrigin
@@ -731,11 +731,8 @@ bool FParticleEmitterInstance::Resize(int32 NewMaxActiveParticles, bool bSetMaxA
 			{
 				UE_LOG(LogParticles, Warning, TEXT("Emitter::Resize> Invalid NewMaxActive (%d) for Emitter in PSys %s"),
 					NewMaxActiveParticles, 
-					Component	? 
-								Component->Template ? *(Component->Template->GetPathName()) 
-													: *(Component->GetName()) 
-								:
-								TEXT("INVALID COMPONENT"));
+								Component.GetTemplate() ? *(Component.GetTemplate()->GetPathName()) 
+													: *(Component.GetName()));
 			}
 
 			return false;
@@ -930,16 +927,16 @@ void FParticleEmitterInstance::CheckEmitterFinished()
 float FParticleEmitterInstance::Tick_EmitterTimeSetup(float DeltaTime, UParticleLODLevel* InCurrentLODLevel)
 {
 	// Make sure we don't try and do any interpolation on the first frame we are attached (OldLocation is not valid in this circumstance)
-	if (Component->bJustRegistered)
+	if (Component.GetJustRegistered())
 	{
-		Location	= Component->GetComponentLocation();
+		Location	= Component.GetComponentTransform().GetLocation();
 		OldLocation	= Location;
 	}
 	else
 	{
 		// Keep track of location for world- space interpolation and other effects.
 		OldLocation	= Location;
-		Location	= Component->GetComponentLocation();
+		Location	= Component.GetComponentTransform().GetLocation();
 	}
 
 	UpdateTransforms();
@@ -1203,10 +1200,10 @@ void FParticleEmitterInstance::SetCurrentLODIndex(int32 InLODIndex, bool bInFull
  */
 void FParticleEmitterInstance::Rewind()
 {
-	if( Component && Component->GetWorld() )
+	if (Component.GetWorld())
 	{
 		UE_LOG(LogParticles,Verbose,
-			TEXT("FParticleEmitterInstance::Rewind @ %fs %s"), Component->GetWorld()->TimeSeconds,
+			TEXT("FParticleEmitterInstance::Rewind @ %fs %s"), Component.GetWorld()->TimeSeconds,
 			(SpriteTemplate != NULL && SpriteTemplate->GetOuter() != NULL) ? *SpriteTemplate->GetOuter()->GetName() : TEXT("NULL"));
 	}
 
@@ -1275,12 +1272,12 @@ FVector FParticleEmitterInstance::GetParticleLocationWithOrbitOffset(FBasePartic
 void FParticleEmitterInstance::UpdateBoundingBox(float DeltaTime)
 {
 	SCOPE_CYCLE_COUNTER(STAT_ParticleUpdateBounds);
-	if (Component)
 	{
-		bool bUpdateBox = ((Component->bWarmingUp == false) && (Component->Template != NULL) && (Component->Template->bUseFixedRelativeBoundingBox == false));
+		const UParticleSystem* Template = Component.GetTemplate();
+		bool bUpdateBox = ((Component.GetIsWarmingUp() == false) && (Template != NULL) && (Template->bUseFixedRelativeBoundingBox == false));
 
 		// Take component scale into account
-		FVector Scale = Component->GetComponentTransform().GetScale3D();
+		FVector Scale = Component.GetComponentTransform().GetScale3D();
 
 		UParticleLODLevel* LODLevel = GetCurrentLODLevelChecked();
 
@@ -1319,7 +1316,7 @@ void FParticleEmitterInstance::UpdateBoundingBox(float DeltaTime)
 		const bool bUseLocalSpace = LODLevel->RequiredModule->bUseLocalSpace;
 
 		const FMatrix ComponentToWorld = bUseLocalSpace 
-			? Component->GetComponentToWorld().ToMatrixWithScale() 
+			? Component.GetComponentToWorld().ToMatrixWithScale() 
 			: FMatrix::Identity;
 
 		bool bSkipDoubleSpawnUpdate = !SpriteTemplate->bUseLegacySpawningBehavior;
@@ -1420,10 +1417,9 @@ void FParticleEmitterInstance::UpdateBoundingBox(float DeltaTime)
  */
 void FParticleEmitterInstance::ForceUpdateBoundingBox()
 {
-	if ( Component )
 	{
 		// Take component scale into account
-		FVector Scale = Component->GetComponentTransform().GetScale3D();
+		FVector Scale = Component.GetComponentTransform().GetScale3D();
 
 		UParticleLODLevel* LODLevel = GetCurrentLODLevelChecked();
 		UParticleLODLevel* HighestLODLevel = SpriteTemplate->LODLevels[0];
@@ -1437,7 +1433,7 @@ void FParticleEmitterInstance::ForceUpdateBoundingBox()
 		const bool bUseLocalSpace = LODLevel->RequiredModule->bUseLocalSpace;
 
 		const FMatrix ComponentToWorld = bUseLocalSpace 
-			? Component->GetComponentToWorld().ToMatrixWithScale() 
+			? Component.GetComponentToWorld().ToMatrixWithScale() 
 			: FMatrix::Identity;
 
 		// For each particle, offset the box appropriately 
@@ -1507,21 +1503,13 @@ uint32 FParticleEmitterInstance::RequiredBytes()
 				if ((InterpolationMethod != PSUVIM_None) && (bHasSubUV == false))
 				{
 					UE_LOG(LogParticles, Warning, TEXT("Emitter w/ mismatched SubUV settings: %s"),
-						Component ? 
-							Component->Template ? 
-								*(Component->Template->GetPathName()) :
-								*(Component->GetFullName()) :
-							TEXT("INVALID PSYS!"));
+							Component.GetTemplate() ? *(Component.GetTemplate()->GetPathName()) : *(Component.GetPathName()));
 				}
 
 				if ((InterpolationMethod == PSUVIM_None) && (bHasSubUV == true))
 				{
 					UE_LOG(LogParticles, Warning, TEXT("Emitter w/ mismatched SubUV settings: %s"),
-						Component ? 
-						Component->Template ? 
-						*(Component->Template->GetPathName()) :
-					*(Component->GetFullName()) :
-					TEXT("INVALID PSYS!"));
+						Component.GetTemplate() ? *(Component.GetTemplate()->GetPathName()) : *(Component.GetPathName()));
 				}
 			}
 			// Check for SubUV utilization, and update the required bytes accordingly
@@ -1668,7 +1656,7 @@ float FParticleEmitterInstance::GetCurrentBurstRateOffset(float& DeltaTime, int3
 								Count = RandomStream.RandRange(BurstEntry->CountLow, BurstEntry->Count);
 							}
 							// Take in to account scale.
-							float Scale = LODLevel->SpawnModule->BurstScale.GetValue(EmitterTime, Component);
+							float Scale = LODLevel->SpawnModule->BurstScale.GetValue(EmitterTime, Component.GetDistributionData());
 							Count = FMath::CeilToInt(Count * Scale);
 							SpawnRateInc += Count / DeltaTime;
 							Burst += Count;
@@ -1932,17 +1920,17 @@ void FParticleEmitterInstance::ParticlePrefetch()
 void FParticleEmitterInstance::CheckSpawnCount(int32 InNewCount, int32 InMaxCount)
 {
 #if !(UE_BUILD_SHIPPING || UE_BUILD_TEST)
-	if( Component )
 	{
-		AWorldSettings* WorldSettings = Component->GetWorld() ? Component->GetWorld()->GetWorldSettings() : NULL;
+		AWorldSettings* WorldSettings = Component.GetWorld() ? Component.GetWorld()->GetWorldSettings() : NULL;
 		if (WorldSettings)
 		{
 			const int32 SizeScalar = sizeof(FParticleSpriteVertex);
 
 			int32 MyIndex = -1;
-			for (int32 CheckIdx = 0; CheckIdx < Component->EmitterInstances.Num(); CheckIdx++)
+			TArrayView<FParticleEmitterInstance*> EmitterInstances = Component.GetEmitterInstances();
+			for (int32 CheckIdx = 0; CheckIdx < EmitterInstances.Num(); CheckIdx++)
 			{
-				if (Component->EmitterInstances[CheckIdx] == this)
+				if (EmitterInstances[CheckIdx] == this)
 				{
 					MyIndex = CheckIdx;
 					break;
@@ -1957,7 +1945,7 @@ void FParticleEmitterInstance::CheckSpawnCount(int32 InNewCount, int32 InMaxCoun
 				InMaxCount, 
 				(float)(InMaxCount * 4 * SizeScalar) / 1024.0f,
 				InNewCount - ActiveParticles,
-				Component->Template ? *Component->Template->GetPathName() : TEXT("No template"));
+				Component.GetTemplate() ? *Component.GetTemplate()->GetPathName() : TEXT("No template"));
 			FColor ErrorColor(255,255,0);
 			if (GEngine->OnScreenDebugMessageExists((uint64)(0x8000000 | (PTRINT)this)) == false)
 			{
@@ -1993,7 +1981,7 @@ float FParticleEmitterInstance::Spawn(float DeltaTime)
 
 	bool bProcessSpawnRate = true;
 	bool bProcessBurstList = true;
-	int32 DetailMode = Component->GetCurrentDetailMode();
+	int32 DetailMode = Component.GetCurrentDetailMode();
 
 	if (SpriteTemplate->QualityLevelSpawnRateScale > 0.0f)
 	{
@@ -2033,8 +2021,8 @@ float FParticleEmitterInstance::Spawn(float DeltaTime)
 		// Figure out spawn rate for this tick.
 		if (bProcessSpawnRate)
 		{
-			float RateScale = LODLevel->SpawnModule->RateScale.GetValue(EmitterTime, Component) * LODLevel->SpawnModule->GetGlobalRateScale();
-			SpawnRate += LODLevel->SpawnModule->Rate.GetValue(EmitterTime, Component) * RateScale;
+			float RateScale = LODLevel->SpawnModule->RateScale.GetValue(EmitterTime, Component.GetDistributionData()) * LODLevel->SpawnModule->GetGlobalRateScale();
+			SpawnRate += LODLevel->SpawnModule->Rate.GetValue(EmitterTime, Component.GetDistributionData()) * RateScale;
 			SpawnRate = FMath::Max<float>(0.0f, SpawnRate);
 		}
 
@@ -2208,7 +2196,7 @@ void FParticleEmitterInstance::SpawnParticles( int32 Count, float StartTime, flo
 				static bool bErrorReported = false;
 				if (!bErrorReported)
 				{
-					UE_LOG(LogParticles, Error, TEXT("Detected null particles. Template : %s, Component %s"), *GetNameSafe(Component && Component->IsValidLowLevel() ? Component->Template : nullptr), *GetFullNameSafe(Component));
+					UE_LOG(LogParticles, Error, TEXT("Detected null particles. Template : %s, Component %s"), *GetNameSafe(Component.IsValidLowLevel() ? Component.GetTemplate() : nullptr), *Component.GetFullName());
 					bErrorReported = true;
 				}
 				ActiveParticles = 0;
@@ -2220,7 +2208,7 @@ void FParticleEmitterInstance::SpawnParticles( int32 Count, float StartTime, flo
 			uint16 NextFreeIndex = ParticleIndices[ActiveParticles];
 			if (!ensure(NextFreeIndex < MaxActiveParticles))
 			{
-				UE_LOG(LogParticles, Error, TEXT("Detected corrupted particle indices. Template : %s, Component %s"), *GetNameSafe(Component && Component->IsValidLowLevel() ? Component->Template : nullptr), *GetFullNameSafe(Component));
+				UE_LOG(LogParticles, Error, TEXT("Detected corrupted particle indices. Template : %s, Component %s"), *GetNameSafe(Component.IsValidLowLevel() ? Component.GetTemplate() : nullptr), *Component.GetFullName());
 				FixupParticleIndices();
 				NextFreeIndex = ParticleIndices[ActiveParticles];
 			}
@@ -2243,7 +2231,7 @@ void FParticleEmitterInstance::SpawnParticles( int32 Count, float StartTime, flo
 					UParticleModule* OffsetModule = HighestLODLevel->SpawnModules[ModuleIndex];
 					SpawnModule->Spawn({ *this, (int32)GetModuleDataOffset(OffsetModule), SpawnTime, Particle });
 
-					ensureMsgf(!Particle->Location.ContainsNaN(), TEXT("NaN in Particle Location. Template: %s, Component: %s"), Component ? *GetNameSafe(Component->Template) : TEXT("UNKNOWN"), *GetPathNameSafe(Component));
+					ensureMsgf(!Particle->Location.ContainsNaN(), TEXT("NaN in Particle Location. Template: %s, Component: %s"), *GetNameSafe(Component.GetTemplate()), *Component.GetPathName());
 				}
 			}
 			PostSpawn(Particle, Interp, SpawnTime);
@@ -2496,7 +2484,7 @@ void FParticleEmitterInstance::KillParticles()
 
 		if (bFoundCorruptIndices)
 		{
-			UE_LOG(LogParticles, Error, TEXT("Detected corrupted particle indices. Template : %s, Component %s"), *GetNameSafe(Component && Component->IsValidLowLevel() ? Component->Template : nullptr), *GetFullNameSafe(Component));			
+			UE_LOG(LogParticles, Error, TEXT("Detected corrupted particle indices. Template : %s, Component %s"), *GetNameSafe(Component.IsValidLowLevel() ? Component.GetTemplate() : nullptr), *Component.GetFullName());			
 			FixupParticleIndices();
 		}
 	}
@@ -2669,12 +2657,12 @@ void FParticleEmitterInstance::SetupEmitterDuration()
 
 		FRandomStream& RandomStream = RequiredModule->GetRandomStream({*this});
 
-		CurrentDelay = RequiredModule->EmitterDelay + Component->EmitterDelay;
+		CurrentDelay = RequiredModule->EmitterDelay + Component.GetEmitterDelay();
 		if (RequiredModule->bEmitterDelayUseRange)
 		{
 			const float	Rand	= RandomStream.FRand();
 			CurrentDelay	    = RequiredModule->EmitterDelayLow + 
-				((RequiredModule->EmitterDelay - RequiredModule->EmitterDelayLow) * Rand) + Component->EmitterDelay;
+				((RequiredModule->EmitterDelay - RequiredModule->EmitterDelayLow) * Rand) + Component.GetEmitterDelay();
 		}
 
 
@@ -2720,10 +2708,6 @@ bool FParticleEmitterInstance::IsDynamicDataRequired(UParticleLODLevel* InCurren
 		return false;
 	}
 
-	if (Component == NULL)
-	{
-		return false;
-	}
 	return true;
 }
 
@@ -2735,49 +2719,55 @@ void FParticleEmitterInstance::ProcessParticleEvents(float DeltaTime, bool bSupp
 	UParticleLODLevel* LODLevel = GetCurrentLODLevelChecked();
 	if (LODLevel->EventReceiverModules.Num() > 0)
 	{
+		TArrayView<FParticleEventSpawnData> SpawnEvents = Component.GetSpawnEvents();
+		TArrayView<FParticleEventDeathData> DeathEvents = Component.GetDeathEvents();
+		TArrayView<FParticleEventCollideData> CollisionEvents = Component.GetCollisionEvents();
+		TArrayView<FParticleEventBurstData> BurstEvents = Component.GetBurstEvents();
+		TArrayView<FParticleEventKismetData> KismetEvents = Component.GetKismetEvents();
+
 		for (int32 EventModIndex = 0; EventModIndex < LODLevel->EventReceiverModules.Num(); EventModIndex++)
 		{
 			int32 EventIndex;
 			UParticleModuleEventReceiverBase* EventRcvr = LODLevel->EventReceiverModules[EventModIndex];
 			check(EventRcvr);
 
-			if (EventRcvr->WillProcessParticleEvent(EPET_Spawn) && (Component->SpawnEvents.Num() > 0))
+			if (EventRcvr->WillProcessParticleEvent(EPET_Spawn) && (SpawnEvents.Num() > 0))
 			{
-				for (EventIndex = 0; EventIndex < Component->SpawnEvents.Num(); EventIndex++)
+				for (EventIndex = 0; EventIndex < SpawnEvents.Num(); EventIndex++)
 				{
-					EventRcvr->ProcessParticleEvent(this, Component->SpawnEvents[EventIndex], DeltaTime);
+					EventRcvr->ProcessParticleEvent(this, SpawnEvents[EventIndex], DeltaTime);
 				}
 			}
 
-			if (EventRcvr->WillProcessParticleEvent(EPET_Death) && (Component->DeathEvents.Num() > 0))
+			if (EventRcvr->WillProcessParticleEvent(EPET_Death) && (DeathEvents.Num() > 0))
 			{
-				for (EventIndex = 0; EventIndex < Component->DeathEvents.Num(); EventIndex++)
+				for (EventIndex = 0; EventIndex < DeathEvents.Num(); EventIndex++)
 				{
-					EventRcvr->ProcessParticleEvent(this, Component->DeathEvents[EventIndex], DeltaTime);
+					EventRcvr->ProcessParticleEvent(this, DeathEvents[EventIndex], DeltaTime);
 				}
 			}
 
-			if (EventRcvr->WillProcessParticleEvent(EPET_Collision) && (Component->CollisionEvents.Num() > 0))
+			if (EventRcvr->WillProcessParticleEvent(EPET_Collision) && (CollisionEvents.Num() > 0))
 			{
-				for (EventIndex = 0; EventIndex < Component->CollisionEvents.Num(); EventIndex++)
+				for (EventIndex = 0; EventIndex < CollisionEvents.Num(); EventIndex++)
 				{
-					EventRcvr->ProcessParticleEvent(this, Component->CollisionEvents[EventIndex], DeltaTime);
+					EventRcvr->ProcessParticleEvent(this, CollisionEvents[EventIndex], DeltaTime);
 				}
 			}
 
-			if (EventRcvr->WillProcessParticleEvent(EPET_Burst) && (Component->BurstEvents.Num() > 0))
+			if (EventRcvr->WillProcessParticleEvent(EPET_Burst) && (BurstEvents.Num() > 0))
 			{
-				for (EventIndex = 0; EventIndex < Component->BurstEvents.Num(); EventIndex++)
+				for (EventIndex = 0; EventIndex < BurstEvents.Num(); EventIndex++)
 				{
-					EventRcvr->ProcessParticleEvent(this, Component->BurstEvents[EventIndex], DeltaTime);
+					EventRcvr->ProcessParticleEvent(this, BurstEvents[EventIndex], DeltaTime);
 				}
 			}
 
-			if (EventRcvr->WillProcessParticleEvent(EPET_Blueprint) && (Component->KismetEvents.Num() > 0))
+			if (EventRcvr->WillProcessParticleEvent(EPET_Blueprint) && (KismetEvents.Num() > 0))
 			{
-				for (EventIndex = 0; EventIndex < Component->KismetEvents.Num(); EventIndex++)
+				for (EventIndex = 0; EventIndex < KismetEvents.Num(); EventIndex++)
 				{
-					EventRcvr->ProcessParticleEvent(this, Component->KismetEvents[EventIndex], DeltaTime);
+					EventRcvr->ProcessParticleEvent(this, KismetEvents[EventIndex], DeltaTime);
 				}
 			}
 		}
@@ -2829,10 +2819,7 @@ bool FParticleEmitterInstance::FillReplayData( FDynamicEmitterReplayDataBase& Ou
 
 	// Take scale into account
 	OutData.Scale = FVector3f::OneVector;
-	if (Component)
-	{
-		OutData.Scale = FVector3f(Component->GetComponentTransform().GetScale3D());
-	}
+	OutData.Scale = FVector3f(Component.GetComponentTransform().GetScale3D());
 
 	int32 ParticleMemSize = MaxActiveParticles * ParticleStride;
 
@@ -2852,7 +2839,7 @@ bool FParticleEmitterInstance::FillReplayData( FDynamicEmitterReplayDataBase& Ou
 		NewReplayData->RequiredModule = LODLevel->RequiredModule->CreateRendererResource();
 		NewReplayData->MaterialInterface = NULL;	// Must be set by derived implementation
 		NewReplayData->InvDeltaSeconds = (LastDeltaTime > UE_KINDA_SMALL_NUMBER) ? (1.0f / LastDeltaTime) : 0.0f;
-		NewReplayData->LWCTile = ((Component == nullptr) || LODLevel->RequiredModule->bUseLocalSpace) ? FVector3f::Zero() : Component->GetLWCTile();
+		NewReplayData->LWCTile = LODLevel->RequiredModule->bUseLocalSpace ? FVector3f::Zero() : Component.GetLWCTile();
 
 		NewReplayData->MaxDrawCount =
 			(LODLevel->RequiredModule->bUseMaxDrawCount == true) ? LODLevel->RequiredModule->MaxDrawCount : -1;
@@ -2953,12 +2940,13 @@ void FParticleEmitterInstance::ApplyWorldOffset(FVector InOffset, bool bWorldShi
 void FParticleEmitterInstance::Tick_MaterialOverrides(int32 EmitterIndex)
 {
 	UParticleLODLevel* LODLevel = SpriteTemplate->GetCurrentLODLevel(this);
+	TArrayView<TObjectPtr<UMaterialInterface>> EmitterMaterials = Component.GetEmitterMaterials();
+
 	bool bOverridden = false;
-	if( LODLevel && LODLevel->RequiredModule && Component && Component->Template )
+	if( LODLevel && LODLevel->RequiredModule && Component.GetTemplate() )
 	{
 	        TArray<FName>& NamedOverrides = LODLevel->RequiredModule->NamedMaterialOverrides;
-	        TArray<FNamedEmitterMaterial>& Slots = Component->Template->NamedMaterialSlots;
-	        TArray<TObjectPtr<UMaterialInterface>>& EmitterMaterials = Component->EmitterMaterials;
+	        const TArray<FNamedEmitterMaterial>& Slots = Component.GetTemplate()->NamedMaterialSlots;
 	        if (NamedOverrides.Num() > 0)
 	        {
 		        //If we have named material overrides then get it's index into the emitter materials array.
@@ -2982,13 +2970,13 @@ void FParticleEmitterInstance::Tick_MaterialOverrides(int32 EmitterIndex)
 	        }
 	}
 
-	if (bOverridden == false && Component)
+	if (bOverridden == false)
 	{
-		if (Component->EmitterMaterials.IsValidIndex(EmitterIndex))
+		if (EmitterMaterials.IsValidIndex(EmitterIndex))
 		{
-			if (Component->EmitterMaterials[EmitterIndex])
+			if (EmitterMaterials[EmitterIndex])
 			{
-				CurrentMaterial = Component->EmitterMaterials[EmitterIndex];
+				CurrentMaterial = EmitterMaterials[EmitterIndex];
 			}
 		}
 	}
@@ -3006,10 +2994,7 @@ void FParticleEmitterInstance::GetScreenAlignmentAndScale(int32& OutScreenAlign,
 	OutScreenAlign = (int32)LODLevel->RequiredModule->ScreenAlignment;
 
 	OutScale = FVector(1.0f, 1.0f, 1.0f);
-	if (Component)
-	{
-		OutScale = Component->GetComponentTransform().GetScale3D();
-	}
+	OutScale = Component.GetComponentTransform().GetScale3D();
 }
 
 /*-----------------------------------------------------------------------------
@@ -3020,7 +3005,7 @@ void FParticleEmitterInstance::GetScreenAlignmentAndScale(int32& OutScreenAlign,
  *	The structure for a standard sprite emitter instance.
  */
 /** Constructor	*/
-FParticleSpriteEmitterInstance::FParticleSpriteEmitterInstance(UParticleSystemComponent& InComponent) :
+FParticleSpriteEmitterInstance::FParticleSpriteEmitterInstance(IParticleEmitterInstanceOwner& InComponent) :
 	FParticleEmitterInstance(InComponent)
 {
 }
@@ -3122,7 +3107,7 @@ void FParticleSpriteEmitterInstance::GetAllocatedSize(int32& OutNum, int32& OutM
  */
 void FParticleSpriteEmitterInstance::GetResourceSizeEx(FResourceSizeEx& CumulativeResourceSize)
 {
-	if (CumulativeResourceSize.GetResourceSizeMode() == EResourceSizeMode::EstimatedTotal || (Component && Component->SceneProxy))
+	if (CumulativeResourceSize.GetResourceSizeMode() == EResourceSizeMode::EstimatedTotal || (Component.GetSceneProxy()))
 	{
 		int32 MaxActiveParticleDataSize = (ParticleData != NULL) ? (MaxActiveParticles * ParticleStride) : 0;
 		int32 MaxActiveParticleIndexSize = (ParticleIndices != NULL) ? (MaxActiveParticles * sizeof(uint16)) : 0;
@@ -3192,7 +3177,7 @@ UMaterialInterface* FParticleEmitterInstance::GetCurrentMaterial()
  */
 
 /** Constructor	*/
-FParticleMeshEmitterInstance::FParticleMeshEmitterInstance(UParticleSystemComponent& InComponent) :
+FParticleMeshEmitterInstance::FParticleMeshEmitterInstance(IParticleEmitterInstanceOwner& InComponent) :
 	  FParticleEmitterInstance(InComponent)
 	, MeshTypeData(NULL)
 	, MeshRotationActive(false)
@@ -3391,26 +3376,27 @@ void FParticleMeshEmitterInstance::Tick_MaterialOverrides(int32 EmitterIndex)
 	// we need to do this here, since CurrentMaterials is unique to mesh emitters, so this can't be done from the component. CurrentMaterials 
 	// may end up in MeshMaterials, so if this doesn't get updated, rendering may access a garbage collected material
 
+	TArrayView<TObjectPtr<UMaterialInterface>> EmitterMaterials = Component.GetEmitterMaterials();
+
 	// make sure currentmaterials are all set to the emitter material, so we don't end up with GC'd material pointers
 	// in MeshMaterials when GetMeshMaterials pushes CurrentMaterials to MeshMaterials
-	if (Component && Component->EmitterMaterials.IsValidIndex(EmitterIndex))
+	if (EmitterMaterials.IsValidIndex(EmitterIndex))
 	{
-		if (Component->EmitterMaterials[EmitterIndex])
+		if (EmitterMaterials[EmitterIndex])
 		{
 			for (UMaterialInterface *& CurMat : CurrentMaterials)
 			{
-				CurMat = Component->EmitterMaterials[EmitterIndex];
+				CurMat = EmitterMaterials[EmitterIndex];
 			}
 		}
 	}	
 
 	UParticleLODLevel* LODLevel = SpriteTemplate->GetCurrentLODLevel(this);
 	bool bOverridden = false;
-	if (LODLevel && LODLevel->RequiredModule && Component && Component->Template)
+	if (LODLevel && LODLevel->RequiredModule && Component.GetTemplate())
 	{
 		TArray<FName>& NamedOverrides = LODLevel->RequiredModule->NamedMaterialOverrides;
-		TArray<FNamedEmitterMaterial>& Slots = Component->Template->NamedMaterialSlots;
-		TArray<TObjectPtr<UMaterialInterface>>& EmitterMaterials = Component->EmitterMaterials;
+		const TArray<FNamedEmitterMaterial>& Slots = Component.GetTemplate()->NamedMaterialSlots;
 		if (NamedOverrides.Num() > 0)
 		{
 			CurrentMaterials.SetNumZeroed(NamedOverrides.Num());
@@ -3448,17 +3434,17 @@ void FParticleMeshEmitterInstance::UpdateBoundingBox(float DeltaTime)
 	SCOPE_CYCLE_COUNTER(STAT_ParticleUpdateBounds);
 	//@todo. Implement proper bound determination for mesh emitters.
 	// Currently, just 'forcing' the mesh size to be taken into account.
-	if ((Component != NULL) && (ActiveParticles > 0))
+	if (ActiveParticles > 0)
 	{
-		bool bUpdateBox = ((Component->bWarmingUp == false) &&
-			(Component->Template != NULL) && (Component->Template->bUseFixedRelativeBoundingBox == false));
+		const UParticleSystem* Template = Component.GetTemplate();
+		bool bUpdateBox = ((Component.GetIsWarmingUp() == false) && (Template != NULL) && (Template->bUseFixedRelativeBoundingBox == false));
 
 		// Take scale into account
-		FVector Scale = Component->GetComponentTransform().GetScale3D();
+		FVector Scale = Component.GetComponentTransform().GetScale3D();
 
 		// Get the static mesh bounds
 		FBoxSphereBounds MeshBound;
-		if (Component->bWarmingUp == false)
+		if (Component.GetIsWarmingUp() == false)
 		{	
 			if (MeshTypeData->Mesh)
 			{
@@ -3466,7 +3452,7 @@ void FParticleMeshEmitterInstance::UpdateBoundingBox(float DeltaTime)
 			}
 			else
 			{
-				//UE_LOG(LogParticles, Log, TEXT("MeshEmitter with no mesh set?? - %s"), Component->Template ? *(Component->Template->GetPathName()) : TEXT("??????"));
+				//UE_LOG(LogParticles, Log, TEXT("MeshEmitter with no mesh set?? - %s"), Component.GetTemplate() ? *(Component.GetTemplate()->GetPathName()) : TEXT("??????"));
 				MeshBound = FBoxSphereBounds(FVector(0, 0, 0), FVector(0, 0, 0), 0);
 			}
 		}
@@ -3481,7 +3467,7 @@ void FParticleMeshEmitterInstance::UpdateBoundingBox(float DeltaTime)
 		const bool bUseLocalSpace = LODLevel->RequiredModule->bUseLocalSpace;
 
 		const FMatrix ComponentToWorld = bUseLocalSpace 
-			? Component->GetComponentToWorld().ToMatrixWithScale() 
+			? Component.GetComponentToWorld().ToMatrixWithScale() 
 			: FMatrix::Identity;
 
 		FVector	NewLocation;
@@ -3775,7 +3761,7 @@ void FParticleMeshEmitterInstance::GetAllocatedSize(int32& OutNum, int32& OutMax
  */
 void FParticleMeshEmitterInstance::GetResourceSizeEx(FResourceSizeEx& CumulativeResourceSize)
 {
-	if (CumulativeResourceSize.GetResourceSizeMode() == EResourceSizeMode::EstimatedTotal || (Component && Component->SceneProxy))
+	if (CumulativeResourceSize.GetResourceSizeMode() == EResourceSizeMode::EstimatedTotal || (Component.GetSceneProxy()))
 	{
 		int32 MaxActiveParticleDataSize = (ParticleData != NULL) ? (MaxActiveParticles * ParticleStride) : 0;
 		int32 MaxActiveParticleIndexSize = (ParticleIndices != NULL) ? (MaxActiveParticles * sizeof(uint16)) : 0;
@@ -3924,7 +3910,6 @@ bool FParticleMeshEmitterInstance::FillReplayData( FDynamicEmitterReplayDataBase
 	// Scale needs to be handled in a special way for meshes.  The parent implementation set this
 	// itself, but we'll recompute it here.
 	NewReplayData->Scale = FVector3f::OneVector;
-	if (Component)
 	{
 		check(SpriteTemplate);
 		UParticleLODLevel* LODLevel2 = SpriteTemplate->GetCurrentLODLevel(this);
@@ -3935,7 +3920,7 @@ bool FParticleMeshEmitterInstance::FillReplayData( FDynamicEmitterReplayDataBase
 		{
 			if (!bIgnoreComponentScale)
 			{
-				NewReplayData->Scale = (FVector3f)Component->GetComponentTransform().GetScale3D();
+				NewReplayData->Scale = (FVector3f)Component.GetComponentTransform().GetScale3D();
 			}
 		}
 	}
diff --git a/Engine/Source/Runtime/Engine/Private/Particles/ParticleGpuSimulation.cpp b/Engine/Source/Runtime/Engine/Private/Particles/ParticleGpuSimulation.cpp
--- a/Engine/Source/Runtime/Engine/Private/Particles/ParticleGpuSimulation.cpp
+++ b/Engine/Source/Runtime/Engine/Private/Particles/ParticleGpuSimulation.cpp
@@ -15,6 +15,7 @@
 #include "SceneManagement.h"
 #include "Misc/DelayedAutoRegister.h"
 #include "ParticleEmitterInstances.h"
+#include "ParticleEmitterInstanceOwner.h"
 #include "Particles/Orientation/ParticleModuleOrientationAxisLock.h"
 #include "Particles/ParticleSystemComponent.h"
 #include "CanvasTypes.h"
@@ -3202,7 +3203,7 @@ class FGPUSpriteParticleEmitterInstance : public FParticleEmitterInstance
 public:
 
 /** Initialization constructor. */
-FGPUSpriteParticleEmitterInstance(UParticleSystemComponent& InComponent, FFXSystem* InFXSystem, FGPUSpriteEmitterInfo& InEmitterInfo)
+FGPUSpriteParticleEmitterInstance(IParticleEmitterInstanceOwner& InComponent, FFXSystem* InFXSystem, FGPUSpriteEmitterInfo& InEmitterInfo)
 	: FParticleEmitterInstance(InComponent)
 	, FXSystem(InFXSystem)
 	, EmitterInfo(InEmitterInfo)
@@ -3267,7 +3268,7 @@ FGPUSpriteParticleEmitterInstance(UParticleSystemComponent& InComponent, FFXSyst
 	virtual bool IsDynamicDataRequired(UParticleLODLevel* InCurrentLODLevel) override
 	{
 		bool bShouldRender = (ActiveParticles >= 0 || TilesToClear.Num() || NewParticles.Num());
-		bool bCanRender = (FXSystem != NULL) && (Component != NULL) && (Component->FXSystem == FXSystem);
+		bool bCanRender = (FXSystem != NULL) && (Component.GetFXSystem() == FXSystem);
 		return bShouldRender && bCanRender;
 	}
 
@@ -3277,10 +3278,9 @@ FGPUSpriteParticleEmitterInstance(UParticleSystemComponent& InComponent, FFXSyst
 	virtual FDynamicEmitterDataBase* GetDynamicData(bool bSelected, ERHIFeatureLevel::Type InFeatureLevel) override
 	{
 		QUICK_SCOPE_CYCLE_COUNTER(STAT_FDynamicEmitterDataBase_GetDynamicData);
-		check(Component);
 		check(SpriteTemplate);
 		check(FXSystem);
-		check(Component->FXSystem == FXSystem);
+		check(Component.GetFXSystem() == FXSystem);
 
 		// Grab the current LOD level
 		UParticleLODLevel* LODLevel = GetCurrentLODLevelChecked();
@@ -3289,14 +3289,14 @@ FGPUSpriteParticleEmitterInstance(UParticleSystemComponent& InComponent, FFXSyst
 			return NULL;
 		}
 
-		UParticleSystem *Template = Component->Template;
+		const UParticleSystem* Template = Component.GetTemplate();
 
 		const bool bLocalSpace = EmitterInfo.RequiredModule->bUseLocalSpace;
 		const bool bUseTileOffset = bLocalSpace == false && EmitterInfo.RequiredModule->bSupportLargeWorldCoordinates;
-		FTransform ComponentTransform = Component->GetComponentTransform();
+		FTransform ComponentTransform = Component.GetComponentTransform();
 		if (bUseTileOffset)
 		{
-			ComponentTransform.AddToTranslation(FVector(Component->GetLWCTile()) * -FLargeWorldRenderScalar::GetTileSize());	
+			ComponentTransform.AddToTranslation(FVector(Component.GetLWCTile()) * -FLargeWorldRenderScalar::GetTileSize());	
 		}
 		const FMatrix ComponentToWorldMatrix = ComponentTransform.ToMatrixWithScale();
 		const FMatrix ComponentToWorld = (bLocalSpace || EmitterInfo.LocalVectorField.bIgnoreComponentTransform) ? FMatrix::Identity : ComponentToWorldMatrix;
@@ -3307,13 +3307,13 @@ FGPUSpriteParticleEmitterInstance(UParticleSystemComponent& InComponent, FFXSyst
 		DynamicData->MaterialProxy = GetCurrentMaterial()->GetRenderProxy();
 		DynamicData->bIsMaterialTranslucent = IsTranslucentBlendMode(*GetCurrentMaterial());
 		DynamicData->Simulation = Simulation;
-		DynamicData->SimulationBounds = Template->bUseFixedRelativeBoundingBox ? Template->FixedRelativeBoundingBox.TransformBy(Component->GetComponentTransform()) : Component->Bounds.GetBox();
+		DynamicData->SimulationBounds = Template->bUseFixedRelativeBoundingBox ? Template->FixedRelativeBoundingBox.TransformBy(Component.GetComponentTransform()) : Component.GetBounds().GetBox();
 		DynamicData->SortMode = EmitterInfo.RequiredModule->SortMode;
 		DynamicData->bSelected = bSelected;
 		DynamicData->bUseLocalSpace = EmitterInfo.RequiredModule->bUseLocalSpace;
 
 		// set up vector field data
-		const FMatrix VectorFieldComponentToWorld = (bLocalSpace || EmitterInfo.LocalVectorField.bIgnoreComponentTransform) ? FMatrix::Identity : Component->GetComponentTransform().ToMatrixWithScale();
+		const FMatrix VectorFieldComponentToWorld = (bLocalSpace || EmitterInfo.LocalVectorField.bIgnoreComponentTransform) ? FMatrix::Identity : Component.GetComponentTransform().ToMatrixWithScale();
 		const FRotationMatrix VectorFieldTransform(LocalVectorFieldRotation);
 		const FMatrix VectorFieldToWorld = VectorFieldTransform * EmitterInfo.LocalVectorField.Transform.ToMatrixWithScale() * VectorFieldComponentToWorld;
 		DynamicData->LocalVectorFieldToWorld = VectorFieldToWorld;
@@ -3325,7 +3325,7 @@ FGPUSpriteParticleEmitterInstance(UParticleSystemComponent& InComponent, FFXSyst
 		DynamicData->bLocalVectorFieldUseFixDT = EmitterInfo.LocalVectorField.bUseFixDT;
 
 		// Get LWC tile
-		DynamicData->EmitterDynamicParameters.LWCTile = bUseTileOffset ? Component->GetLWCTile() : FVector3f::ZeroVector;
+		DynamicData->EmitterDynamicParameters.LWCTile = bUseTileOffset ? Component.GetLWCTile() : FVector3f::ZeroVector;
 
 		// Account for LocalToWorld scaling
 		FVector ComponentScale = ComponentTransform.GetScale3D();
@@ -3362,20 +3362,20 @@ FGPUSpriteParticleEmitterInstance(UParticleSystemComponent& InComponent, FFXSyst
 		FVector4 ColorScaleOverLife(1.0f, 1.0f, 1.0f, 1.0f);
 		if( EmitterInfo.DynamicColorScale.IsCreated() )
 		{
-			ColorScaleOverLife = EmitterInfo.DynamicColorScale.GetValue(0.0f,Component);
+			ColorScaleOverLife = EmitterInfo.DynamicColorScale.GetValue(0.0f,Component.GetDistributionData());
 		}
 		if( EmitterInfo.DynamicAlphaScale.IsCreated() )
 		{
-			ColorScaleOverLife.W = EmitterInfo.DynamicAlphaScale.GetValue(0.0f,Component);
+			ColorScaleOverLife.W = EmitterInfo.DynamicAlphaScale.GetValue(0.0f,Component.GetDistributionData());
 		}
 
 		if( EmitterInfo.DynamicColor.IsCreated() )
 		{
-			ColorOverLife = EmitterInfo.DynamicColor.GetValue(0.0f,Component);
+			ColorOverLife = EmitterInfo.DynamicColor.GetValue(0.0f,Component.GetDistributionData());
 		}
 		if( EmitterInfo.DynamicAlpha.IsCreated() )
 		{
-			ColorOverLife.W = EmitterInfo.DynamicAlpha.GetValue(0.0f,Component);
+			ColorOverLife.W = EmitterInfo.DynamicAlpha.GetValue(0.0f,Component.GetDistributionData());
 		}
 		DynamicData->EmitterDynamicParameters.DynamicColor = FVector4f(ColorOverLife * ColorScaleOverLife); // LWC_TODO: precision loss
 
@@ -3480,7 +3480,7 @@ FGPUSpriteParticleEmitterInstance(UParticleSystemComponent& InComponent, FFXSyst
 			DynamicData->PerFrameSimulationParameters.PositionOffsetAndAttractorStrength = FVector4f(FVector3f(PositionOffsetThisTick), PointAttractorStrength);
 			DynamicData->PerFrameSimulationParameters.LocalToWorldScale = DynamicData->EmitterDynamicParameters.LocalToWorldScale;
 			DynamicData->PerFrameSimulationParameters.DeltaSeconds = PendingDeltaSeconds; // This value is used when updating vector fields.
-			DynamicData->PerFrameSimulationParameters.LWCTile = bUseTileOffset ? Component->GetLWCTile() : FVector3f::ZeroVector;
+			DynamicData->PerFrameSimulationParameters.LWCTile = bUseTileOffset ? Component.GetLWCTile() : FVector3f::ZeroVector;
 			Exchange(DynamicData->TilesToClear, TilesToClear);
 			Exchange(DynamicData->NewParticles, NewParticles);
 		}
@@ -3523,12 +3523,12 @@ FGPUSpriteParticleEmitterInstance(UParticleSystemComponent& InComponent, FFXSyst
 		check(AllocatedTiles.Num() == TileTimeOfDeath.Num());
 		FreeParticlesInTile = 0;
 
-		RandomStream.Initialize(Component->RandomStream.FRand());
+		RandomStream.Initialize(Component.GetRandomStream().FRand());
 		EmitterInstRandom = RandomStream.GetFraction();
 
 		FParticleSimulationResources* ParticleSimulationResources = FXSystem->GetParticleSimulationResources();
 		const bool bUseTileOffset = EmitterInfo.RequiredModule->bUseLocalSpace == false && EmitterInfo.RequiredModule->bSupportLargeWorldCoordinates;
-		ParticleSimulationResources->LWCTile = bUseTileOffset ? Component->GetLWCTile() : FVector3f::ZeroVector;
+		ParticleSimulationResources->LWCTile = bUseTileOffset ? Component.GetLWCTile() : FVector3f::ZeroVector;
 		const int32 MinTileCount = GetMinTileCount();
 		int32 NumAllocated = 0;
 		{
@@ -3575,7 +3575,7 @@ FGPUSpriteParticleEmitterInstance(UParticleSystemComponent& InComponent, FFXSyst
 
 		InitLocalVectorField();
 
-		WeakWorld = Component->GetWorld();
+		WeakWorld = Component.GetWorld();
 	}
 
 	FORCENOINLINE void ReserveNewParticles(int32 Num)
@@ -3622,8 +3622,8 @@ FGPUSpriteParticleEmitterInstance(UParticleSystemComponent& InComponent, FFXSyst
 		{
 			// If the emitter is warming up but any particle spawned now will die
 			// anyway, suppress spawning.
-			if (Component && Component->bWarmingUp &&
-				Component->WarmupTime - SecondsSinceCreation > EmitterInfo.MaxLifetime)
+			if (Component.GetIsWarmingUp() &&
+				Component.GetWarmupTime() - SecondsSinceCreation > EmitterInfo.MaxLifetime)
 			{
 				bSuppressSpawning = true;
 			}
@@ -3708,12 +3708,12 @@ FGPUSpriteParticleEmitterInstance(UParticleSystemComponent& InComponent, FFXSyst
 				}
 	#endif
 
-				if (Component && Component->bWarmingUp)
+				if (Component.GetIsWarmingUp())
 				{
 					SimulateWarmupParticles(
 						NewParticles.GetData() + (NewParticles.Num() - NewParticleCount),
 						NewParticleCount,
-						Component->WarmupTime - SecondsSinceCreation );
+						Component.GetWarmupTime() - SecondsSinceCreation );
 				}
 			}
 			else if (bFakeBurstsWhenSpawningSupressed)
@@ -3931,7 +3931,7 @@ private:
 	 */
 	int32 GetMinTileCount()
 	{
-		if (AllowedLoopCount == 0 && Component->IsActive() && GetCurrentLODLevelChecked()->bEnabled)
+		if (AllowedLoopCount == 0 && Component.IsActive() && GetCurrentLODLevelChecked()->bEnabled)
 		{
 			const int32 EstMaxTiles = (EmitterInfo.MaxParticleCount + GParticlesPerTile - 1) / GParticlesPerTile;
 			const int32 SlackTiles = FMath::CeilToInt(FXConsoleVariables::ParticleSlackGPU * (float)EstMaxTiles);
@@ -4015,7 +4015,7 @@ private:
 #if TRACK_TILE_ALLOCATIONS
 					UE_LOG(LogParticles, VeryVerbose,
 						TEXT("%s|%s|0x%016p [ReleaseSimulationResources] %d tiles"),
-						*Component->GetName(), *Component->Template->GetName(), (PTRINT)this, AllocatedTiles.Num());
+						*Component.GetName(), *Component.Template->GetName(), (PTRINT)this, AllocatedTiles.Num());
 #endif // #if TRACK_TILE_ALLOCATIONS
 				}
 			}
@@ -4026,7 +4026,7 @@ private:
 			{
 				UE_LOG(LogParticles, Warning,
 					TEXT("%s|%s|0x%016p [ReleaseSimulationResources] LEAKING %d tiles FXSystem=0x%016x"),
-					*Component->GetName(), *Component->Template->GetName(), this, AllocatedTiles.Num(), (PTRINT)FXSystem);
+					*Component.GetName(), *Component.GetTemplate()->GetName(), this, AllocatedTiles.Num(), (PTRINT)FXSystem);
 			}
 		}
 
@@ -4086,7 +4086,7 @@ private:
 						// Out of particle tiles.
 						UE_LOG(LogParticles,Warning,
 							TEXT("Failed to allocate tiles for %s! %d new particles truncated to %d."),
-							*Component->Template->GetName(), NumNewParticles, ParticleIndex);
+							*Component.GetTemplate()->GetName(), NumNewParticles, ParticleIndex);
 						return ParticleIndex;
 					}
 
@@ -4130,8 +4130,8 @@ private:
 
 		// Determine spawn rate.
 		check(SpawnModule && RequiredModule);
-		const float RateScale = CurrentLODLevel->SpawnModule->RateScale.GetValue(EmitterTime, Component) * CurrentLODLevel->SpawnModule->GetGlobalRateScale();
-		float SpawnRate = CurrentLODLevel->SpawnModule->Rate.GetValue(EmitterTime, Component) * RateScale;
+		const float RateScale = CurrentLODLevel->SpawnModule->RateScale.GetValue(EmitterTime, Component.GetDistributionData()) * CurrentLODLevel->SpawnModule->GetGlobalRateScale();
+		float SpawnRate = CurrentLODLevel->SpawnModule->Rate.GetValue(EmitterTime, Component.GetDistributionData()) * RateScale;
 		SpawnRate = FMath::Max<float>(0.0f, SpawnRate);
 
 		if (EmitterInfo.SpawnPerUnitModule)
@@ -4247,7 +4247,7 @@ private:
 			FVector WSPosition = TempParticle->Location + InterpFraction * EmitterDelta + SpawnInfo.StartTime * (FVector)NewParticle->Velocity + EmitterInfo.OrbitOffsetBase + EmitterInfo.OrbitOffsetRange * RandomOrbit;
 			if (RequiredModule->bUseLocalSpace == false && RequiredModule->bSupportLargeWorldCoordinates)
 			{
-				NewParticle->Position = FVector3f(WSPosition - FLargeWorldRenderScalar::GetTileSize() * FVector(Component->GetLWCTile()));
+				NewParticle->Position = FVector3f(WSPosition - FLargeWorldRenderScalar::GetTileSize() * FVector(Component.GetLWCTile()));
 			}
 			else
 			{
@@ -4290,7 +4290,7 @@ private:
 	virtual void UpdateBoundingBox(float DeltaSeconds) override
 	{
 		// Setup a bogus bounding box at the origin. GPU emitters must use fixed bounds.
-		FVector Origin = Component ? Component->GetComponentToWorld().GetTranslation() : FVector::ZeroVector;
+		FVector Origin = Component.GetComponentToWorld().GetTranslation();
 		ParticleBoundingBox = FBox::BuildAABB(Origin, FVector(1.0f));
 	}
 
@@ -5237,7 +5237,7 @@ void FFXSystem::VisualizeGPUParticles(FCanvas* Canvas)
 	External interface.
 -----------------------------------------------------------------------------*/
 
-FParticleEmitterInstance* FFXSystem::CreateGPUSpriteEmitterInstance(UParticleSystemComponent& InComponent, FGPUSpriteEmitterInfo& EmitterInfo )
+FParticleEmitterInstance* FFXSystem::CreateGPUSpriteEmitterInstance(IParticleEmitterInstanceOwner& InComponent, FGPUSpriteEmitterInfo& EmitterInfo )
 {
 	return new FGPUSpriteParticleEmitterInstance(InComponent, this, EmitterInfo );
 }
diff --git a/Engine/Source/Runtime/Engine/Private/Particles/ParticleModules.cpp b/Engine/Source/Runtime/Engine/Private/Particles/ParticleModules.cpp
--- a/Engine/Source/Runtime/Engine/Private/Particles/ParticleModules.cpp
+++ b/Engine/Source/Runtime/Engine/Private/Particles/ParticleModules.cpp
@@ -18,6 +18,7 @@
 #include "Particles/Orientation/ParticleModuleOrientationBase.h"
 #include "Particles/Orientation/ParticleModuleOrientationAxisLock.h"
 #include "ParticleEmitterInstances.h"
+#include "ParticleEmitterInstanceOwner.h"
 #include "Particles/ParticleSpriteEmitter.h"
 #include "Particles/ParticleSystemComponent.h"
 #include "Distributions/DistributionFloatConstant.h"
@@ -645,7 +646,7 @@ void UParticleModule::GetParticleParametersUtilized(TArray<FString>& ParticlePar
 uint32 UParticleModule::PrepRandomSeedInstancePayload(FParticleEmitterInstance* Owner, FParticleRandomSeedInstancePayload* InRandSeedPayload, const FParticleRandomSeedInfo& InRandSeedInfo)
 {
     // These should never be null
-	if (!ensure( Owner != nullptr && Owner->Component != nullptr ))
+	if (!ensure(Owner != nullptr))
 	{
 		return 0xffffffff;
 	}
@@ -658,7 +659,7 @@ uint32 UParticleModule::PrepRandomSeedInstancePayload(FParticleEmitterInstance*
 		if (InRandSeedInfo.bGetSeedFromInstance == true)
 		{
 			float SeedValue;
-			if (Owner->Component->GetFloatParameter(InRandSeedInfo.ParameterName, SeedValue) == true)
+			if (Owner->Component.GetFloatParameter(InRandSeedInfo.ParameterName, SeedValue) == true)
 			{
 				if (InRandSeedInfo.bInstanceSeedIsIndex == false)
 				{
@@ -688,7 +689,7 @@ uint32 UParticleModule::PrepRandomSeedInstancePayload(FParticleEmitterInstance*
 
 			if (InRandSeedInfo.bRandomlySelectSeedArray)
 			{
-				Index = Owner->Component->RandomStream.RandHelper(InRandSeedInfo.RandomSeeds.Num());
+				Index = Owner->Component.GetRandomStream().RandHelper(InRandSeedInfo.RandomSeeds.Num());
 			}
 
 			InRandSeedPayload->RandomStream.Initialize(InRandSeedInfo.RandomSeeds[Index]);
@@ -780,7 +781,7 @@ bool UParticleModule::IsUsedInGPUEmitter()const
 FRandomStream& UParticleModule::GetRandomStream(const FContext& Context)
 {
 	FParticleRandomSeedInstancePayload* Payload = Context.Owner.GetModuleRandomSeedInstanceData(this);
-	FRandomStream& RandomStream = (Payload != nullptr) ? Payload->RandomStream : Context.Owner.Component->RandomStream;
+	FRandomStream& RandomStream = (Payload != nullptr) ? Payload->RandomStream : Context.Owner.Component.GetRandomStream();
 	return RandomStream;
 }
 
@@ -964,10 +965,8 @@ void UParticleModuleSourceMovement::FinalUpdate(const FUpdateContext& Context)
 {
 	Super::FinalUpdate(Context);
 	// If enabled, update the source offset translation for all particles...
-	FParticleEmitterInstance* Owner = &Context.Owner;
-	if (Owner && Owner->Component)
 	{
-		FVector FrameDiff = Context.GetTransform().GetLocation() - Owner->Component->OldPosition;
+		FVector FrameDiff = Context.GetTransform().GetLocation() - Context.Owner.Component.GetOldPosition();
 		BEGIN_UPDATE_LOOP;
 		{
 			// Rough estimation of the particle being alive for more than a frame
@@ -1426,7 +1425,7 @@ void UParticleModuleMeshRotation::SpawnEx(const FSpawnContext& Context, struct F
 			FVector Rotation = StartRotation.GetValue(Owner->EmitterTime, Context.GetDistributionData(), 0, InRandomStream);
 			if (bInheritParent)
 			{
-				FRotator	Rotator	= Owner->Component->GetComponentRotation();
+				FRotator	Rotator	= Owner->Component.GetComponentRotation();
 				FVector		ParentAffectedRotation	= Rotator.Euler();
 				Rotation.X	+= ParentAffectedRotation.X / 360.0f;
 				Rotation.Y	+= ParentAffectedRotation.Y / 360.0f;
@@ -2186,7 +2185,7 @@ float UParticleModuleSubUV::DetermineImageIndex(const FContext& Context, int32 O
 		}
 		else
 		{
-			UWorld* World = Owner->Component->GetWorld();
+			UWorld* World = Owner->Component.GetWorld();
 			if ((World != NULL) && (World->GetWorldSettings() != NULL))
 			{
 				ImageIndex = SubImageIndex.GetValue(Particle->RelativeTime / World->GetWorldSettings()->GetEffectiveTimeDilation(), Context.GetDistributionData());
@@ -2370,7 +2369,7 @@ float UParticleModuleSubUVMovie::DetermineImageIndex(const FContext& Context, in
 	}
 	else
 	{
-		UWorld* World = Owner->Component->GetWorld();
+		UWorld* World = Owner->Component.GetWorld();
 		if ((World != NULL) && (World->GetWorldSettings() != NULL))
 		{
 			MoviePayload.Time += DeltaTime / World->GetWorldSettings()->GetEffectiveTimeDilation();
@@ -2843,26 +2842,24 @@ void UParticleModuleAcceleration::PostEditChangeProperty(FPropertyChangedEvent&
 
 const FTransform& UParticleModuleAcceleration::FContext::GetTransform() const
 {
-	return Owner.Component->GetComponentTransform();
+	return Owner.Component.GetComponentTransform();
 }
 
 UObject* UParticleModuleAcceleration::FContext::GetDistributionData() const
 {
-	return Owner.Component;
+	return Owner.Component.GetDistributionData();
 }
 
 FString UParticleModuleAcceleration::FContext::GetTemplateName() const
 {
-	if (Owner.Component)
 	{
-		return GetNameSafe(Owner.Component->Template);
+		return GetNameSafe(Owner.Component.GetTemplate());
 	}
-	return TEXT("UNKNOWN");
 }
 
 FString UParticleModuleAcceleration::FContext::GetInstanceName() const
 {
-	return GetPathNameSafe(Owner.Component);
+	return Owner.Component.GetPathName();
 }
 
 void UParticleModuleAcceleration::Spawn(const FSpawnContext& Context)
@@ -2871,7 +2868,7 @@ void UParticleModuleAcceleration::Spawn(const FSpawnContext& Context)
 	SPAWN_INIT;
 	PARTICLE_ELEMENT(FVector3f, UsedAcceleration);
 	UsedAcceleration = FVector3f(Acceleration.GetValue(Owner.EmitterTime, Context.GetDistributionData()));
-	if (bApplyOwnerScale == true && Owner.Component)
+	if (bApplyOwnerScale == true)
 	{
 		FVector3f Scale(Context.GetTransform().GetScale3D());
 		UsedAcceleration *= Scale;
@@ -3142,7 +3139,7 @@ uint64 UParticleModuleLight::SpawnHQLight(const FLightParticlePayload& Payload,
 	{
 		return 0;
 	}
-	UParticleSystemComponent* ParticleSystem = Owner->Component;
+	UParticleSystemComponent* ParticleSystem = Owner->Component.AsComponent();
 	if (!ParticleSystem)
 	{
 		return 0;
@@ -3418,7 +3415,7 @@ UParticleModuleTypeDataBase::UParticleModuleTypeDataBase(const FObjectInitialize
 	bUpdateModule = false;
 }
 
-FParticleEmitterInstance* UParticleModuleTypeDataBase::CreateInstance(UParticleEmitter* InEmitterParent, UParticleSystemComponent& InComponent)
+FParticleEmitterInstance* UParticleModuleTypeDataBase::CreateInstance(UParticleEmitter* InEmitterParent, IParticleEmitterInstanceOwner& InComponent)
 {
 	return NULL;
 }
@@ -3457,7 +3454,7 @@ UParticleModuleTypeDataMesh::UParticleModuleTypeDataMesh(const FObjectInitialize
 	LODSizeScale = 1.0f;
 }
 
-FParticleEmitterInstance* UParticleModuleTypeDataMesh::CreateInstance(UParticleEmitter* InEmitterParent, UParticleSystemComponent& InComponent)
+FParticleEmitterInstance* UParticleModuleTypeDataMesh::CreateInstance(UParticleEmitter* InEmitterParent, IParticleEmitterInstanceOwner& InComponent)
 {
 	SetToSensibleDefaults(InEmitterParent);
 	FParticleEmitterInstance* Instance = new FParticleMeshEmitterInstance(InComponent);
@@ -3745,7 +3742,7 @@ void UParticleModuleKillBox::Render3DPreview(const FPreviewContext& Context)
 		KillboxVerts[7].Set(CheckUR.X, CheckLL.Y, CheckUR.Z);
 	}
 	
-	if ((bAbsolute == false) && (Owner->Component != NULL))
+	if (bAbsolute == false)
 	{
 		if (bAxisAlignedAndFixedSize == false)
 		{
@@ -3871,7 +3868,6 @@ void UParticleModuleKillHeight::Render3DPreview(const FPreviewContext& Context)
 #if WITH_EDITOR
 	FVector OwnerPosition(0.0f);
 	FParticleEmitterInstance* Owner = &Context.Owner;
-	if ((Owner != NULL) && (Owner->Component != NULL))
 	{
 		OwnerPosition = Context.GetTransform().GetLocation();
 
@@ -4249,9 +4245,10 @@ void UParticleModuleAttractorParticle::Spawn(const FSpawnContext& Context)
 	FParticleEmitterInstance* AttractorEmitterInst = NULL;
 	if (EmitterName != NAME_None)
 	{
-		for (int32 ii = 0; ii < Owner->Component->EmitterInstances.Num(); ii++)
+		TArrayView<FParticleEmitterInstance*> EmitterInstances = Owner->Component.GetEmitterInstances();
+		for (int32 ii = 0; ii < EmitterInstances.Num(); ii++)
 		{
-			FParticleEmitterInstance* pkEmitInst = Owner->Component->EmitterInstances[ii];
+			FParticleEmitterInstance* pkEmitInst = EmitterInstances[ii];
 			if (pkEmitInst && (pkEmitInst->SpriteTemplate->EmitterName == EmitterName))
 			{
 				AttractorEmitterInst = pkEmitInst;
@@ -4318,9 +4315,10 @@ void UParticleModuleAttractorParticle::Update(const FUpdateContext& Context)
 	FParticleEmitterInstance* AttractorEmitterInst = NULL;
 	if (EmitterName != NAME_None)
 	{
-		for (int32 ii = 0; ii < Owner->Component->EmitterInstances.Num(); ii++)
+		TArrayView<FParticleEmitterInstance*> EmitterInstances = Owner->Component.GetEmitterInstances();
+		for (int32 ii = 0; ii < EmitterInstances.Num(); ii++)
 		{
-			FParticleEmitterInstance* pkEmitInst = Owner->Component->EmitterInstances[ii];
+			FParticleEmitterInstance* pkEmitInst = EmitterInstances[ii];
 			if (pkEmitInst && (pkEmitInst->SpriteTemplate->EmitterName == EmitterName))
 			{
 				AttractorEmitterInst = pkEmitInst;
@@ -5023,21 +5021,21 @@ void UParticleModuleTypeDataGpu::Build( FParticleEmitterBuildInfo& EmitterBuildI
 	}
 }
 
-FParticleEmitterInstance* UParticleModuleTypeDataGpu::CreateInstance(UParticleEmitter* InEmitterParent, UParticleSystemComponent& InComponent)
+FParticleEmitterInstance* UParticleModuleTypeDataGpu::CreateInstance(UParticleEmitter* InEmitterParent, IParticleEmitterInstanceOwner& InComponent)
 {
 	UWorld* World = InComponent.GetWorld();
 	check(World);
 	UE_LOG(LogParticles,Verbose,
 		TEXT("Create GPU Sprite Emitter @ %fs %s"), World->TimeSeconds,
-		InComponent.Template != NULL ? *InComponent.Template->GetName() : TEXT("NULL"));
+		InComponent.GetTemplate() != NULL ? *InComponent.GetTemplate()->GetName() : TEXT("NULL"));
 
 	FParticleEmitterInstance* Instance = NULL;
 	if (World->Scene && RHISupportsGPUParticles())
 	{
 		LLM_SCOPE(ELLMTag::Particles);
 
-		check( InComponent.FXSystem );
-		Instance = InComponent.FXSystem->CreateGPUSpriteEmitterInstance(InComponent, EmitterInfo );
+		check(InComponent.GetFXSystem());
+		Instance = InComponent.GetFXSystem()->CreateGPUSpriteEmitterInstance(InComponent, EmitterInfo );
 		Instance->InitParameters( InEmitterParent );
 	}
 	return Instance;
diff --git a/Engine/Source/Runtime/Engine/Private/Particles/ParticleModules_Camera.cpp b/Engine/Source/Runtime/Engine/Private/Particles/ParticleModules_Camera.cpp
--- a/Engine/Source/Runtime/Engine/Private/Particles/ParticleModules_Camera.cpp
+++ b/Engine/Source/Runtime/Engine/Private/Particles/ParticleModules_Camera.cpp
@@ -6,6 +6,7 @@
 =============================================================================*/
 
 #include "ParticleEmitterInstances.h"
+#include "ParticleEmitterInstanceOwner.h"
 #include "Particles/ParticleSystemComponent.h"
 #include "Distributions/DistributionFloatConstant.h"
 #include "Particles/Camera/ParticleModuleCameraBase.h"
@@ -78,7 +79,7 @@ void UParticleModuleCameraOffset::Spawn(const FSpawnContext& Context)
 	UParticleLODLevel* LODLevel = Owner ? Owner->SpriteTemplate->GetCurrentLODLevel(Owner) : NULL;
 	if (LODLevel && LODLevel->RequiredModule->bUseLocalSpace == false)
 	{
-		ScaleFactor = Owner->Component->GetAsyncComponentToWorld().GetMaximumAxisScale();
+		ScaleFactor = Owner->Component.GetAsyncComponentToWorld().GetMaximumAxisScale();
 	}
 	SPAWN_INIT;
 	{
diff --git a/Engine/Source/Runtime/Engine/Private/Particles/ParticleModules_Collision.cpp b/Engine/Source/Runtime/Engine/Private/Particles/ParticleModules_Collision.cpp
--- a/Engine/Source/Runtime/Engine/Private/Particles/ParticleModules_Collision.cpp
+++ b/Engine/Source/Runtime/Engine/Private/Particles/ParticleModules_Collision.cpp
@@ -18,6 +18,7 @@
 #include "Engine/StaticMesh.h"
 #include "Engine/TriggerBase.h"
 #include "ParticleEmitterInstances.h"
+#include "ParticleEmitterInstanceOwner.h"
 #include "Particles/Collision/ParticleModuleCollisionBase.h"
 #include "Particles/Collision/ParticleModuleCollision.h"
 #include "Particles/Collision/ParticleModuleCollisionGPU.h"
@@ -181,14 +182,16 @@ void UParticleModuleCollision::Update(const FUpdateContext& Context)
 {
 	SCOPE_CYCLE_COUNTER(STAT_ParticleCollisionTime);
 	FParticleEmitterInstance* Owner = &Context.Owner;
-	UWorld* World = Owner->Component->GetWorld();
+	UWorld* World = Owner->Component.GetWorld();
 	if (Owner->ActiveParticles == 0 || (bDropDetail && World && World->bDropDetail))
 	{
 		return;
 	}
 
+	UParticleSystemComponent* Component = Owner->Component.AsComponent();
+
 	//Gets the owning actor of the component. Can be NULL if the component is spawned with the World as an Outer, e.g. in UGameplayStatics::SpawnEmitterAtLocation().
-	AActor* Actor = Owner->Component->GetOwner();
+	AActor* Actor = Component->GetOwner();
 
 	UParticleLODLevel* LODLevel	= Owner->SpriteTemplate->GetCurrentLODLevel(Owner);
 	check(LODLevel);
@@ -196,7 +199,7 @@ void UParticleModuleCollision::Update(const FUpdateContext& Context)
 	const int32 MeshRotationOffset = Owner->GetMeshRotationOffset();
 	const bool bMeshRotationActive = Owner->IsMeshRotationActive();
 
-	const FTransform& OwnerTM = Owner->Component->GetAsyncComponentToWorld();
+	const FTransform& OwnerTM = Component->GetAsyncComponentToWorld();
 	const FVector ParentScale = OwnerTM.GetScale3D();
 
 	FParticleEventInstancePayload* EventPayload = NULL;
@@ -213,8 +216,8 @@ void UParticleModuleCollision::Update(const FUpdateContext& Context)
 
 	FParticleCollisionInstancePayload* CollisionInstPayload = (FParticleCollisionInstancePayload*)(Owner->GetModuleInstanceData(this));
 
-	const TArray<FVector>& PlayerLocations = Owner->Component->GetPlayerLocations();
-	TArray<float> PlayerLODDistanceFactor = Owner->Component->GetPlayerLODDistanceFactor();	//Make a copy because we need to square it later
+	const TArray<FVector>& PlayerLocations = Component->GetPlayerLocations();
+	TArray<float> PlayerLODDistanceFactor = Component->GetPlayerLODDistanceFactor();	//Make a copy because we need to square it later
 	const int32 PlayerCount = PlayerLocations.Num();
 
 	if (World->IsGameWorld())
@@ -223,7 +226,7 @@ void UParticleModuleCollision::Update(const FUpdateContext& Context)
 
 		// LOD collision based on visibility
 		// This is at the 'emitter instance' level as it will be true or false for the whole instance...
-		if (bCollideOnlyIfVisible && ((World->TimeSeconds - Owner->Component->GetLastRenderTime()) > GParticleCollisionIgnoreInvisibleTime))
+		if (bCollideOnlyIfVisible && ((World->TimeSeconds - Component->GetLastRenderTime()) > GParticleCollisionIgnoreInvisibleTime))
 		{
 			// no collision if not recently rendered
 			bIgnoreAllCollision = true;
@@ -240,14 +243,14 @@ void UParticleModuleCollision::Update(const FUpdateContext& Context)
 					{
 						FBox BoundingBox;
 						BoundingBox.Init();
-						if (Owner->Component->Template && Owner->Component->Template->bUseFixedRelativeBoundingBox)
+						if (Component->Template && Component->Template->bUseFixedRelativeBoundingBox)
 						{
-							BoundingBox = Owner->Component->Template->FixedRelativeBoundingBox.TransformBy(OwnerTM);
+							BoundingBox = Component->Template->FixedRelativeBoundingBox.TransformBy(OwnerTM);
 						}
 						else
 						{
 							// A frame behind, but shouldn't be an issue...
-							BoundingBox = Owner->Component->GetAsyncBounds().GetBox();
+							BoundingBox = Component->GetAsyncBounds().GetBox();
 						}
 
 						// see if any player is within the extended bounds...
@@ -576,8 +579,9 @@ bool UParticleModuleCollision::GenerateLODModuleValues(UParticleModule* SourceMo
 bool UParticleModuleCollision::PerformCollisionCheck(FParticleEmitterInstance* Owner, FBaseParticle* InParticle, 
 	FHitResult& Hit, AActor* SourceActor, const FVector& End, const FVector& Start, const FVector& Extent)
 {
-	check(Owner && Owner->Component);
-	return Owner->Component->ParticleLineCheck(Hit, SourceActor, End, Start, Extent, ObjectParams);
+	UParticleSystemComponent* Component = Owner->Component.AsComponent();
+	check(Component);
+	return Component->ParticleLineCheck(Hit, SourceActor, End, Start, Extent, ObjectParams);
 }
 
 /*------------------------------------------------------------------------------
diff --git a/Engine/Source/Runtime/Engine/Private/Particles/ParticleModules_Event.cpp b/Engine/Source/Runtime/Engine/Private/Particles/ParticleModules_Event.cpp
--- a/Engine/Source/Runtime/Engine/Private/Particles/ParticleModules_Event.cpp
+++ b/Engine/Source/Runtime/Engine/Private/Particles/ParticleModules_Event.cpp
@@ -9,6 +9,7 @@
 #include "Distributions/DistributionFloatConstant.h"
 #include "Distributions/DistributionVectorConstant.h"
 #include "ParticleEmitterInstances.h"
+#include "ParticleEmitterInstanceOwner.h"
 #include "Particles/Event/ParticleModuleEventBase.h"
 #include "Particles/ParticleEmitter.h"
 #include "Particles/Event/ParticleModuleEventGenerator.h"
@@ -129,7 +130,7 @@ bool UParticleModuleEventGenerator::HandleParticleSpawned(FParticleEmitterInstan
 			{
 				FVector ParticleLocation = EventGenInfo.bUseOrbitOffset ? Owner->GetParticleLocationWithOrbitOffset(NewParticle) : FVector(NewParticle->Location);
 
-				Owner->Component->ReportEventSpawn(EventGenInfo.CustomName, Owner->EmitterTime, 
+				Owner->Component.ReportEventSpawn(EventGenInfo.CustomName, Owner->EmitterTime, 
 					ParticleLocation, FVector(NewParticle->Velocity), EventGenInfo.ParticleModuleEventsToSendToGame);
 				bProcessed = true;
 #if !(UE_BUILD_SHIPPING || UE_BUILD_TEST)
@@ -159,7 +160,7 @@ bool UParticleModuleEventGenerator::HandleParticleKilled(FParticleEmitterInstanc
 			{
 				FVector ParticleLocation = EventGenInfo.bUseOrbitOffset ? Owner->GetParticleLocationWithOrbitOffset(DeadParticle) : FVector(DeadParticle->Location);
 
-				Owner->Component->ReportEventDeath(EventGenInfo.CustomName, 
+				Owner->Component.ReportEventDeath(EventGenInfo.CustomName, 
 					Owner->EmitterTime, ParticleLocation, (FVector)DeadParticle->Velocity,
 					EventGenInfo.ParticleModuleEventsToSendToGame, DeadParticle->RelativeTime);
 				bProcessed = true;
@@ -205,7 +206,7 @@ bool UParticleModuleEventGenerator::HandleParticleCollision(FParticleEmitterInst
 
 			if (EventGenInfo.Frequency == 0 || (EventPayload->CollisionTrackingCount % EventGenInfo.Frequency) == 0)
 			{
-				Owner->Component->ReportEventCollision(
+				Owner->Component.ReportEventCollision(
 					EventGenInfo.CustomName, 
 					Owner->EmitterTime, 
 					Hit->Location,
@@ -244,7 +245,7 @@ bool UParticleModuleEventGenerator::HandleParticleBurst(FParticleEmitterInstance
 		{
 			if (EventGenInfo.Frequency == 0 || (EventPayload->BurstTrackingCount % EventGenInfo.Frequency) == 0)
 			{
-				Owner->Component->ReportEventBurst(EventGenInfo.CustomName, Owner->EmitterTime, ParticleCount, 
+				Owner->Component.ReportEventBurst(EventGenInfo.CustomName, Owner->EmitterTime, ParticleCount, 
 					Owner->Location, EventGenInfo.ParticleModuleEventsToSendToGame);
 				bProcessed = true;
 #if !(UE_BUILD_SHIPPING || UE_BUILD_TEST)
diff --git a/Engine/Source/Runtime/Engine/Private/Particles/ParticleModules_Location.cpp b/Engine/Source/Runtime/Engine/Private/Particles/ParticleModules_Location.cpp
--- a/Engine/Source/Runtime/Engine/Private/Particles/ParticleModules_Location.cpp
+++ b/Engine/Source/Runtime/Engine/Private/Particles/ParticleModules_Location.cpp
@@ -8,6 +8,7 @@
 #include "Misc/MessageDialog.h"
 #include "Components/SkeletalMeshComponent.h"
 #include "ParticleEmitterInstances.h"
+#include "ParticleEmitterInstanceOwner.h"
 #include "Particles/ParticleSystemComponent.h"
 #include "Distributions/DistributionFloatConstant.h"
 #include "Distributions/DistributionVectorConstant.h"
@@ -329,7 +330,7 @@ void UParticleModuleLocationDirect::Spawn(const FSpawnContext& Context)
 void UParticleModuleLocationDirect::Update(const FUpdateContext& Context)
 {
 	FParticleEmitterInstance* Owner = &Context.Owner;
-	const FTransform& OwnerTM = Owner->Component->GetAsyncComponentToWorld();
+	const FTransform& OwnerTM = Owner->Component.GetAsyncComponentToWorld();
 	BEGIN_UPDATE_LOOP;
 	{
 		FVector	NewLoc;
@@ -403,9 +404,10 @@ void UParticleModuleLocationEmitter::Spawn(const FSpawnContext& Context)
 	FParticleEmitterInstance* LocationEmitterInst = NULL;
 	if (EmitterName != NAME_None)
 	{
-		for (int32 ii = 0; ii < Owner->Component->EmitterInstances.Num(); ii++)
+		TArrayView<FParticleEmitterInstance*> EmitterInstances = Owner->Component.GetEmitterInstances();
+		for (int32 ii = 0; ii < EmitterInstances.Num(); ii++)
 		{
-			FParticleEmitterInstance* pkEmitInst = Owner->Component->EmitterInstances[ii];
+			FParticleEmitterInstance* pkEmitInst = EmitterInstances[ii];
 			if (pkEmitInst && (pkEmitInst->SpriteTemplate->EmitterName == EmitterName))
 			{
 				LocationEmitterInst = pkEmitInst;
@@ -470,7 +472,7 @@ void UParticleModuleLocationEmitter::Spawn(const FSpawnContext& Context)
 				{
 					if (bInLocalSpace == false)
 					{
-						Particle.Location = LocationEmitterInst->Component->GetComponentLocation();
+						Particle.Location = LocationEmitterInst->Component.GetComponentTransform().GetLocation();
 					}
 					else
 					{
@@ -487,12 +489,12 @@ void UParticleModuleLocationEmitter::Spawn(const FSpawnContext& Context)
 					else if ((bSourceIsInLocalSpace == true) && (bInLocalSpace == false))
 					{
 						// We need to transform it into world space
-						Particle.Location = LocationEmitterInst->Component->GetComponentTransform().TransformPosition(pkParticle->Location);
+						Particle.Location = LocationEmitterInst->Component.GetComponentTransform().TransformPosition(pkParticle->Location);
 					}
 					else //if ((bSourceIsInLocalSpace == false) && (bInLocalSpace == true))
 					{
 						// We need to transform it into local space
-						Particle.Location = LocationEmitterInst->Component->GetComponentTransform().InverseTransformPosition(pkParticle->Location);
+						Particle.Location = LocationEmitterInst->Component.GetComponentTransform().InverseTransformPosition(pkParticle->Location);
 					}
 				}
 				if (InheritSourceVelocity)
@@ -567,9 +569,10 @@ void UParticleModuleLocationEmitterDirect::Spawn(const FSpawnContext& Context)
 	FParticleEmitterInstance* LocationEmitterInst = NULL;
 	if (EmitterName != NAME_None)
 	{
-		for (int32 ii = 0; ii < Owner->Component->EmitterInstances.Num(); ii++)
+		TArrayView<FParticleEmitterInstance*> EmitterInstances = Owner->Component.GetEmitterInstances();
+		for (int32 ii = 0; ii < EmitterInstances.Num(); ii++)
 		{
-			FParticleEmitterInstance* pkEmitInst = Owner->Component->EmitterInstances[ii];
+			FParticleEmitterInstance* pkEmitInst = EmitterInstances[ii];
 			if (pkEmitInst && (pkEmitInst->SpriteTemplate->EmitterName == EmitterName))
 			{
 				LocationEmitterInst = pkEmitInst;
@@ -610,9 +613,10 @@ void UParticleModuleLocationEmitterDirect::Update(const FUpdateContext& Context)
 	FParticleEmitterInstance* LocationEmitterInst = NULL;
 	if (EmitterName != NAME_None)
 	{
-		for (int32 ii = 0; ii < Owner->Component->EmitterInstances.Num(); ii++)
+		TArrayView<FParticleEmitterInstance*> EmitterInstances = Owner->Component.GetEmitterInstances();
+		for (int32 ii = 0; ii < EmitterInstances.Num(); ii++)
 		{
-			FParticleEmitterInstance* pkEmitInst = Owner->Component->EmitterInstances[ii];
+			FParticleEmitterInstance* pkEmitInst = EmitterInstances[ii];
 			if (pkEmitInst && (pkEmitInst->SpriteTemplate->EmitterName == EmitterName))
 			{
 				LocationEmitterInst = pkEmitInst;
@@ -1159,7 +1163,7 @@ void UParticleModuleLocationPrimitiveCylinder::Render3DPreview(const FPreviewCon
 
 	float fStartRadius = 1.0f;
 	float fStartHeight = 1.0f;
-	if(Owner && Owner->Component)
+	if(Owner)
 	{
 		fStartRadius = StartRadius.GetValue(Owner->EmitterTime, Context.GetDistributionData());
 		fStartHeight = StartHeight.GetValue(Owner->EmitterTime, Context.GetDistributionData()) / 2.0f;
@@ -1339,14 +1343,10 @@ void UParticleModuleLocationPrimitiveSphere::SpawnEx(const FSpawnContext& Contex
 void UParticleModuleLocationPrimitiveSphere::Render3DPreview(const FPreviewContext& Context)
 {
 #if WITH_EDITOR
-	FVector Position(0.0f);
 
 	// Draw the location as a wire star
-	FParticleEmitterInstance* Owner = &Context.Owner;
-	if ((Owner != NULL) && (Owner->Component != NULL))
-	{
-		Position = Owner->SimulationToWorld.TransformPosition(Owner->EmitterToSimulation.GetOrigin());
-	}
+	FParticleEmitterInstance& Owner = Context.Owner;
+	FVector Position = Owner.SimulationToWorld.TransformPosition(Owner.EmitterToSimulation.GetOrigin());
 
 	FPrimitiveDrawInterface* PDI = Context.PDI;
 
@@ -1372,7 +1372,6 @@ void UParticleModuleLocationPrimitiveSphere::Render3DPreview(const FPreviewConte
 			// Draw at the avg. of the min/max extents
 			UDistributionVectorConstantCurve* pkCurve = CastChecked<UDistributionVectorConstantCurve>(StartLocation.Distribution);
 
-			if (Owner != NULL)
 			{
 				//pkCurve->
 				Position = StartLocation.GetValue(0.0f, Context.GetDistributionData());
@@ -1380,10 +1379,9 @@ void UParticleModuleLocationPrimitiveSphere::Render3DPreview(const FPreviewConte
 		}
 	}
 
-	if ((Owner != NULL) && (Owner->Component != NULL))
 	{
-		Position = Owner->EmitterToSimulation.TransformPosition(Position);
-		Position = Owner->SimulationToWorld.TransformPosition(Position);
+		Position = Owner.EmitterToSimulation.TransformPosition(Position);
+		Position = Owner.SimulationToWorld.TransformPosition(Position);
 	}
 
 	// Draw a wire start at the center position
@@ -1393,12 +1391,11 @@ void UParticleModuleLocationPrimitiveSphere::Render3DPreview(const FPreviewConte
 	int32		iNumSides	= 32;
 	FVector	vAxis[3];
 
-	if ((Owner != NULL) && (Owner->Component != NULL))
 	{
-		fRadius = StartRadius.GetValue(Owner->EmitterTime, Context.GetDistributionData());
-		vAxis[0]	= Owner->SimulationToWorld.TransformVector(Owner->EmitterToSimulation.GetScaledAxis( EAxis::X ));
-		vAxis[1]	= Owner->SimulationToWorld.TransformVector(Owner->EmitterToSimulation.GetScaledAxis( EAxis::Y ));
-		vAxis[2]	= Owner->SimulationToWorld.TransformVector(Owner->EmitterToSimulation.GetScaledAxis( EAxis::Z ));
+		fRadius = StartRadius.GetValue(Owner.EmitterTime, Context.GetDistributionData());
+		vAxis[0]	= Owner.SimulationToWorld.TransformVector(Owner.EmitterToSimulation.GetScaledAxis( EAxis::X ));
+		vAxis[1]	= Owner.SimulationToWorld.TransformVector(Owner.EmitterToSimulation.GetScaledAxis( EAxis::Y ));
+		vAxis[2]	= Owner.SimulationToWorld.TransformVector(Owner.EmitterToSimulation.GetScaledAxis( EAxis::Z ));
 	}
 
 	DrawCircle(PDI,Position, vAxis[0], vAxis[1], ModuleEditorColor, fRadius, iNumSides, SDPG_World);
@@ -1694,7 +1691,7 @@ void UParticleModuleLocationBoneSocket::Update(const FUpdateContext& Context)
 	const int32 MeshRotationOffset = Owner->GetMeshRotationOffset();
 	const bool bMeshRotationActive = MeshRotationOffset > 0 && Owner->IsMeshRotationActive();// Note that this will *never* be false because this module always reports as touching mesh rotation
 	FQuat* SourceRotation = (bMeshRotationActive) ? NULL : &RotationQuat; // We always will pass NULL down here due to the above condition.
-	const FTransform& OwnerTM = Owner->Component->GetAsyncComponentToWorld();
+	const FTransform& OwnerTM = Owner->Component.GetAsyncComponentToWorld();
 
 	//TODO: we have bone locations stored already if we're inheriting bone velocity, see if we can use those.
 	BEGIN_UPDATE_LOOP;
@@ -1976,7 +1973,7 @@ void UParticleModuleLocationBoneSocket::GetSkeletalMeshComponentSource(const FCo
 		return;
 	}
 
-	UParticleSystemComponent* PSysComp = Owner->Component;
+	UParticleSystemComponent* PSysComp = Owner->Component.AsComponent();
 	if (PSysComp == NULL)
 	{
 		InstancePayload->SourceComponent = nullptr;
@@ -2185,7 +2182,7 @@ bool UParticleModuleLocationBoneSocket::GetParticleLocation(FModuleLocationBoneS
 
 	if (Owner->CurrentLODLevel->RequiredModule->bUseLocalSpace == true)
 	{
-		OutPosition = Owner->Component->GetComponentTransform().InverseTransformPosition(OutPosition);
+		OutPosition = Owner->Component.GetComponentTransform().InverseTransformPosition(OutPosition);
 	}
 
 	return true;
@@ -2534,7 +2531,7 @@ void UParticleModuleLocationSkelVertSurface::Update(const FUpdateContext& Contex
 	FQuat SourceRotation;
 	const int32 MeshRotationOffset = Owner->GetMeshRotationOffset();
 	const bool bMeshRotationActive = MeshRotationOffset > 0 && Owner->IsMeshRotationActive();
-	const FTransform& OwnerTM = Owner->Component->GetAsyncComponentToWorld();
+	const FTransform& OwnerTM = Owner->Component.GetAsyncComponentToWorld();
 
 	FRandomStream& RandomStream = GetRandomStream(Context);
 
@@ -2750,7 +2747,7 @@ void UParticleModuleLocationSkelVertSurface::GetSkeletalMeshComponentSource(cons
 		return;
 	}
 
-	UParticleSystemComponent* PSysComp = Owner->Component;
+	UParticleSystemComponent* PSysComp = Owner->Component.AsComponent();
 	if (PSysComp == NULL)
 	{
 		return;
@@ -2826,11 +2823,11 @@ void UParticleModuleLocationSkelVertSurface::GetSkeletalMeshComponentSource(cons
 			UE_LOG(LogParticles, Warning, TEXT("Likely due to CPU side buffers being stripped with r.FreeSkeletalMeshBuffers=1."));
 			UE_LOG(LogParticles, Warning, TEXT("This emitter should probably be culled for this platform using detail mode."));
 			UE_LOG(LogParticles, Warning, TEXT("Mesh: %s"), NewSkelMeshComp ? *NewSkelMeshComp->GetSkeletalMeshAsset()->GetFullName() : TEXT("NULL"));
-			UE_LOG(LogParticles, Warning, TEXT("Comp: %s"), *Owner->Component->GetFullName());
-			UE_LOG(LogParticles, Warning, TEXT("System: %s"), *Owner->Component->Template->GetFullName());
+			UE_LOG(LogParticles, Warning, TEXT("Comp: %s"), *Owner->Component.GetFullName());
+			UE_LOG(LogParticles, Warning, TEXT("System: %s"), *Owner->Component.GetTemplate()->GetFullName());
 			UE_LOG(LogParticles, Warning, TEXT("----------------------------------------------------------------------"));
 
-			Owner->Component->DeactivaateNextTick();
+			Owner->Component.DeactivateNextTick();
 			return;
 		}
 
diff --git a/Engine/Source/Runtime/Engine/Private/Particles/ParticleModules_Velocity.cpp b/Engine/Source/Runtime/Engine/Private/Particles/ParticleModules_Velocity.cpp
--- a/Engine/Source/Runtime/Engine/Private/Particles/ParticleModules_Velocity.cpp
+++ b/Engine/Source/Runtime/Engine/Private/Particles/ParticleModules_Velocity.cpp
@@ -6,6 +6,7 @@
 =============================================================================*/
 
 #include "ParticleEmitterInstances.h"
+#include "ParticleEmitterInstanceOwner.h"
 #include "Particles/ParticleSystemComponent.h"
 #include "Distributions/DistributionFloatUniform.h"
 #include "Distributions/DistributionVectorConstant.h"
@@ -185,11 +186,11 @@ void UParticleModuleVelocityInheritParent::Spawn(const FSpawnContext& Context)
 	check(LODLevel);
 	if (LODLevel->RequiredModule->bUseLocalSpace)
 	{
-		Vel = Context.GetTransform().InverseTransformVector(Owner->Component->PartSysVelocity);
+		Vel = Context.GetTransform().InverseTransformVector(Owner->Component.GetPartSysVelocity());
 	}
 	else
 	{
-		Vel = Owner->Component->PartSysVelocity;
+		Vel = Owner->Component.GetPartSysVelocity();
 	}
 
 	FVector vScale = Scale.GetValue(Owner->EmitterTime, Context.GetDistributionData());
@@ -259,7 +260,7 @@ void UParticleModuleVelocityOverLifetime::Update(const FUpdateContext& Context)
 	UParticleLODLevel* LODLevel	= Owner->SpriteTemplate->GetCurrentLODLevel(Owner);
 	check(LODLevel);
 	FVector OwnerScale(1.0f);
-	const FTransform& OwnerTM = Owner->Component->GetAsyncComponentToWorld();
+	const FTransform& OwnerTM = Owner->Component.GetAsyncComponentToWorld();
 	if (bApplyOwnerScale == true)
 	{
 		OwnerScale = OwnerTM.GetScale3D();
@@ -538,10 +539,9 @@ void UParticleModuleVelocityCone::Render3DPreview(const FPreviewContext& Context
 	FMatrix OwnerRotation(FMatrix::Identity);
 	FVector LocalToWorldOrigin(0.0f);
 	FMatrix LocalToWorld(FMatrix::Identity);
-	if (Owner->Component)
+	if (UParticleSystemComponent* Component = Owner->Component.AsComponent())
 	{
-		AActor* Actor = Owner->Component->GetOwner();
-		if (Actor)
+		if (AActor* Actor = Component->GetOwner())
 		{
 			if (bApplyOwnerScale == true)
 			{
diff --git a/Engine/Source/Runtime/Engine/Private/Particles/ParticleTrail2EmitterInstance.cpp b/Engine/Source/Runtime/Engine/Private/Particles/ParticleTrail2EmitterInstance.cpp
--- a/Engine/Source/Runtime/Engine/Private/Particles/ParticleTrail2EmitterInstance.cpp
+++ b/Engine/Source/Runtime/Engine/Private/Particles/ParticleTrail2EmitterInstance.cpp
@@ -10,6 +10,7 @@
 #include "Materials/Material.h"
 #include "MaterialDomain.h"
 #include "ParticleEmitterInstances.h"
+#include "ParticleEmitterInstanceOwner.h"
 #include "Particles/ParticleEmitter.h"
 #include "Particles/ParticleSystemComponent.h"
 #include "Particles/Event/ParticleModuleEventGenerator.h"
@@ -100,7 +101,6 @@ void FParticleTrailsEmitterInstance_Base::InitParameters(UParticleEmitter* InTem
 void FParticleTrailsEmitterInstance_Base::Tick(float DeltaTime, bool bSuppressSpawning)
 {
 	SCOPE_CYCLE_COUNTER(STAT_TrailTickTime);
-	if (Component)
 	{
 
 #if defined(ULTRA_VERBOSE_TRAILS_DEBUG) && ULTRA_VERBOSE_TRAILS_DEBUG
@@ -163,15 +163,11 @@ void FParticleTrailsEmitterInstance_Base::Tick(float DeltaTime, bool bSuppressSp
 		// 'Reset' the emitter time so that the delay functions correctly
 		EmitterTime += CurrentDelay;
 		RunningTime += DeltaTime;
-		LastTickTime = Component->GetWorld() ? Component->GetWorld()->GetTimeSeconds() : 0.0f;
+		LastTickTime = Component.GetWorld() ? Component.GetWorld()->GetTimeSeconds() : 0.0f;
 		
 		// Reset particles position offset
 		PositionOffsetThisTick = FVector::ZeroVector;
 	}
-	else
-	{
-		LastTickTime = 0.0f;
-	}
 	INC_DWORD_STAT(STAT_TrailParticlesTickCalls);
 }
 
@@ -188,10 +184,9 @@ void FParticleTrailsEmitterInstance_Base::Tick_RecalculateTangents(float DeltaTi
 void FParticleTrailsEmitterInstance_Base::UpdateBoundingBox(float DeltaTime)
 {
 	SCOPE_CYCLE_COUNTER(STAT_ParticleUpdateBounds);
-	if (Component)
 	{
-		bool bUpdateBox = ((Component->bWarmingUp == false) &&
-			(Component->Template != NULL) && (Component->Template->bUseFixedRelativeBoundingBox == false));
+		const UParticleSystem* Template = Component.GetTemplate();
+		bool bUpdateBox = ((Component.GetIsWarmingUp() == false) && (Template != NULL) && (Template->bUseFixedRelativeBoundingBox == false));
 		// Handle local space usage
 		check(SpriteTemplate->LODLevels.Num() > 0);
 		UParticleLODLevel* LODLevel = SpriteTemplate->LODLevels[0];
@@ -202,8 +197,8 @@ void FParticleTrailsEmitterInstance_Base::UpdateBoundingBox(float DeltaTime)
 			// Set the min/max to the position of the trail
 			if (LODLevel->RequiredModule->bUseLocalSpace == false) 
 			{
-				ParticleBoundingBox.Max = Component->GetComponentLocation();
-				ParticleBoundingBox.Min = Component->GetComponentLocation();
+				ParticleBoundingBox.Max = Component.GetComponentTransform().GetLocation();
+				ParticleBoundingBox.Min = Component.GetComponentTransform().GetLocation();
 			}
 			else
 			{
@@ -214,7 +209,7 @@ void FParticleTrailsEmitterInstance_Base::UpdateBoundingBox(float DeltaTime)
 		ParticleBoundingBox.IsValid = true;
 
 		// Take scale into account
-		FVector Scale = Component->GetComponentTransform().GetScale3D();
+		FVector Scale = Component.GetComponentTransform().GetScale3D();
 
 		// As well as each particle
 		int32 LocalActiveParticles = ActiveParticles;
@@ -274,7 +269,7 @@ void FParticleTrailsEmitterInstance_Base::UpdateBoundingBox(float DeltaTime)
 		{
 			if (LODLevel->RequiredModule->bUseLocalSpace) 
 			{
-				ParticleBoundingBox = ParticleBoundingBox.TransformBy(Component->GetComponentTransform());
+				ParticleBoundingBox = ParticleBoundingBox.TransformBy(Component.GetComponentTransform());
 			}
 		}
 	}
@@ -285,7 +280,6 @@ void FParticleTrailsEmitterInstance_Base::UpdateBoundingBox(float DeltaTime)
  */
 void FParticleTrailsEmitterInstance_Base::ForceUpdateBoundingBox()
 {
-	if (Component)
 	{
 		// Handle local space usage
 		check(SpriteTemplate->LODLevels.Num() > 0);
@@ -295,8 +289,8 @@ void FParticleTrailsEmitterInstance_Base::ForceUpdateBoundingBox()
 		// Set the min/max to the position of the trail
 		if (LODLevel->RequiredModule->bUseLocalSpace == false) 
 		{
-			ParticleBoundingBox.Max = Component->GetComponentLocation();
-			ParticleBoundingBox.Min = Component->GetComponentLocation();
+			ParticleBoundingBox.Max = Component.GetComponentTransform().GetLocation();
+			ParticleBoundingBox.Min = Component.GetComponentTransform().GetLocation();
 		}
 		else
 		{
@@ -306,7 +300,7 @@ void FParticleTrailsEmitterInstance_Base::ForceUpdateBoundingBox()
 		ParticleBoundingBox.IsValid = true;
 
 		// Take scale into account
-		FVector Scale = Component->GetComponentTransform().GetScale3D();
+		FVector Scale = Component.GetComponentTransform().GetScale3D();
 
 		// As well as each particle
 		int32 LocalActiveParticles = ActiveParticles;
@@ -333,7 +327,7 @@ void FParticleTrailsEmitterInstance_Base::ForceUpdateBoundingBox()
 		// Transform bounding box into world space if the emitter uses a local space coordinate system.
 		if (LODLevel->RequiredModule->bUseLocalSpace) 
 		{
-			ParticleBoundingBox = ParticleBoundingBox.TransformBy(Component->GetComponentTransform());
+			ParticleBoundingBox = ParticleBoundingBox.TransformBy(Component.GetComponentTransform());
 		}
 	}
 }
@@ -347,9 +341,9 @@ void FParticleTrailsEmitterInstance_Base::KillParticles()
 	if (ActiveParticles > 0)
 	{
 		float CurrentTickTime = 0.0f;
-		if( Component && Component->GetWorld() )
+		if (Component.GetWorld())
 		{
-			CurrentTickTime = Component->GetWorld()->GetTimeSeconds();
+			CurrentTickTime = Component.GetWorld()->GetTimeSeconds();
 		}
 		bool bHasForceKillParticles = false;
 		// Loop over the active particles... If their RelativeTime is > 1.0f (indicating they are dead),
@@ -778,7 +772,7 @@ UMaterialInterface* FParticleTrailsEmitterInstance_Base::GetCurrentMaterial()
 	FParticleRibbonEmitterInstance.
 -----------------------------------------------------------------------------*/
 /** Constructor	*/
-FParticleRibbonEmitterInstance::FParticleRibbonEmitterInstance(UParticleSystemComponent& InComponent) :
+FParticleRibbonEmitterInstance::FParticleRibbonEmitterInstance(IParticleEmitterInstanceOwner& InComponent) :
 	  FParticleTrailsEmitterInstance_Base(InComponent)
 	, TrailTypeData(NULL)
 	, SpawnPerUnitModule(NULL)
@@ -991,7 +985,7 @@ bool FParticleRibbonEmitterInstance::GetSpawnPerUnitAmount(float DeltaTime, int3
 	{
 		bool bMoved = false;
 		float NewTravelLeftover = 0.0f;
-		float ParticlesPerUnit = SpawnPerUnitModule->SpawnPerUnit.GetValue(EmitterTime, Component) / SpawnPerUnitModule->UnitScalar;
+		float ParticlesPerUnit = SpawnPerUnitModule->SpawnPerUnit.GetValue(EmitterTime, Component.GetDistributionData()) / SpawnPerUnitModule->UnitScalar;
 		// Allow for PPU of 0.0f to allow for 'turning off' an emitter when moving
 		if (ParticlesPerUnit >= 0.0f)
 		{
@@ -1096,8 +1090,7 @@ bool FParticleRibbonEmitterInstance::GetSpawnPerUnitAmount(float DeltaTime, int3
 
 bool FParticleTrailsEmitterInstance_Base::AddParticleHelper(int32 InTrailIdx,
 	int32 StartParticleIndex, FTrailsBaseTypeDataPayload* StartTrailData,
-	int32 ParticleIndex, FTrailsBaseTypeDataPayload* TrailData,
-	UParticleSystemComponent* InPsysComp
+	int32 ParticleIndex, FTrailsBaseTypeDataPayload* TrailData
 	)
 {
 	bool bAddedParticle = false;
@@ -1121,11 +1114,8 @@ bool FParticleTrailsEmitterInstance_Base::AddParticleHelper(int32 InTrailIdx,
 				TrailData->bMovementSpawned ? TEXT("MOVESPAWN") : TEXT("SPAWNRATE"),
 				StartTrailData->SpawnTime,
 				StartTrailData->bMovementSpawned ? TEXT("MOVESPAWN") : TEXT("SPAWNRATE"),
-				InPsysComp ? 
-					(InPsysComp->Template ? *(InPsysComp->Template->GetPathName()) : TEXT("*** No Template")) :
-					TEXT("*** No Component"),
-				InPsysComp ? *(InPsysComp->GetPathName()) : TEXT("*** No Components")
-				);
+				(Component.GetTemplate() ? *(Component.GetTemplate()->GetPathName()) : TEXT("*** No Template")),
+				*Component.GetPathName());
 		}
 #endif
 
@@ -1156,10 +1146,8 @@ bool FParticleTrailsEmitterInstance_Base::AddParticleHelper(int32 InTrailIdx,
 				TrailData->bMovementSpawned ? TEXT("MOVESPAWN") : TEXT("SPAWNRATE"),
 				StartTrailData->SpawnTime,
 				StartTrailData->bMovementSpawned ? TEXT("MOVESPAWN") : TEXT("SPAWNRATE"),
-				InPsysComp ? 
-					(InPsysComp->Template ? *(InPsysComp->Template->GetPathName()) : TEXT("*** No Template")) :
-					TEXT("*** No Component"),
-				InPsysComp ? *(InPsysComp->GetPathName()) : TEXT("*** No Components"));
+				(Component.GetTemplate() ? *(Component.GetTemplate()->GetPathName()) : TEXT("*** No Template")),
+				*Component.GetPathName());
 		}
 #endif
 
@@ -1201,7 +1189,7 @@ void FParticleRibbonEmitterInstance::GetParticleLifetimeAndSize(int32 InTrailIdx
 			UParticleModuleLifetime* LifetimeModule = Cast<UParticleModuleLifetime>(LODLevel->SpawnModules[ModuleIdx]);
 			if (LifetimeModule != NULL)
 			{
-				float MaxLifetime = LifetimeModule->GetLifetimeValue({*this}, EmitterTime, Component);
+				float MaxLifetime = LifetimeModule->GetLifetimeValue({*this}, EmitterTime, Component.GetDistributionData());
 				if (CurrLifetime > 0.f)
 				{
 					// Another module already modified lifetime.
@@ -1267,9 +1255,9 @@ float FParticleRibbonEmitterInstance::Spawn(float DeltaTime)
 	// Figure out spawn rate for this tick.
 	if (bProcessSpawnRate)
 	{
-		float RateScale = LODLevel->SpawnModule->RateScale.GetValue(EmitterTime, Component) * LODLevel->SpawnModule->GetGlobalRateScale();
+		float RateScale = LODLevel->SpawnModule->RateScale.GetValue(EmitterTime, Component.GetDistributionData()) * LODLevel->SpawnModule->GetGlobalRateScale();
 		float QualityMult = SpriteTemplate->GetQualityLevelSpawnRateMult();
-		SpawnRate += LODLevel->SpawnModule->Rate.GetValue(EmitterTime, Component) * FMath::Clamp<float>(QualityMult, 0.0f, 1.0);
+		SpawnRate += LODLevel->SpawnModule->Rate.GetValue(EmitterTime, Component.GetDistributionData()) * FMath::Clamp<float>(QualityMult, 0.0f, 1.0);
 	}
 
 	// Take Bursts into account as well...
@@ -1355,7 +1343,7 @@ float FParticleRibbonEmitterInstance::Spawn(float DeltaTime)
 		FVector CurrentUp;
 		if (TrailTypeData->RenderAxis == Trails_SourceUp)
 		{
-			CurrentUp = Component->GetComponentTransform().GetScaledAxis( EAxis::Z );
+			CurrentUp = Component.GetComponentTransform().GetScaledAxis( EAxis::Z );
 		}
 		else
 		{
@@ -1439,9 +1427,6 @@ float FParticleRibbonEmitterInstance::Spawn(float DeltaTime)
 				bAddedParticle = AddParticleHelper(TrailIdx, 
 					StartIndex, (FTrailsBaseTypeDataPayload*)StartTrailData, 
 					ParticleIndex, (FTrailsBaseTypeDataPayload*)TrailData
-#if !(UE_BUILD_SHIPPING || UE_BUILD_TEST)
-					, Component
-#endif
 					);
 			}
 
@@ -1486,9 +1471,9 @@ float FParticleRibbonEmitterInstance::Spawn(float DeltaTime)
 #if !(UE_BUILD_SHIPPING || UE_BUILD_TEST)
 				if ((int32)ActiveParticles > LocalMaxParticleInTrailCount)
 				{
-					if (Component && Component->GetWorld())
+					if (Component.GetWorld())
 					{
-						const FString TemplateName = Component->Template ? *Component->Template->GetName() : TEXT("No template");
+						const FString TemplateName = Component.GetTemplate() ? *Component.GetTemplate()->GetName() : TEXT("No template");
 						GMaxParticleInTrailCountWarningCS.Lock();
 						static TMap<FString, int32> TemplateMaxCountWarnings;
 						int32& MaxCount = TemplateMaxCountWarnings.FindOrAdd(TemplateName);
@@ -1825,7 +1810,7 @@ bool FParticleRibbonEmitterInstance::Spawn_Source(float DeltaTime)
 			float InvCount = 1.0f / MovementSpawnCount;
 			float Increment = DeltaTime / MovementSpawnCount;
 
-			FTransform SavedComponentToWorld = Component->GetComponentTransform();
+			FTransform SavedComponentToWorld = Component.GetComponentTransform();
 
 			// Spawn the given number of particles, interpolating between the current and last position/tangent
 			float CurrTimeStep = InvCount;
@@ -1862,7 +1847,7 @@ bool FParticleRibbonEmitterInstance::Spawn_Source(float DeltaTime)
 				float SpawnTime = DeltaTime - (SpawnIdx * Increment);
 				float TrueSpawnTime = Diff * TimeStep;
 
-				Component->SetComponentToWorld(FTransform(CurrRotation, CurrPosition));
+				Component.SetComponentToWorld(FTransform(CurrRotation, CurrPosition));
 
 				// Standard spawn setup
 				PreSpawn(Particle, CurrPosition, FVector::ZeroVector);
@@ -1892,7 +1877,7 @@ bool FParticleRibbonEmitterInstance::Spawn_Source(float DeltaTime)
 				Particle->Size.Y = Particle->Size.X;
 				Particle->BaseSize = Particle->Size;
 
-				Component->SetComponentToWorld(SavedComponentToWorld);
+				Component.SetComponentToWorld(SavedComponentToWorld);
 
 				// Trail specific...
 				// Clear the next and previous - just to be safe
@@ -1938,9 +1923,6 @@ bool FParticleRibbonEmitterInstance::Spawn_Source(float DeltaTime)
 					bAddedParticle = AddParticleHelper(TrailIdx,
 						StartIndex, (FTrailsBaseTypeDataPayload*)StartTrailData, 
 						ParticleIndex, (FTrailsBaseTypeDataPayload*)TrailData
-#if !(UE_BUILD_SHIPPING || UE_BUILD_TEST)
-						, Component
-#endif
 						);
 				}
 
@@ -1970,12 +1952,12 @@ bool FParticleRibbonEmitterInstance::Spawn_Source(float DeltaTime)
 #if !(UE_BUILD_SHIPPING || UE_BUILD_TEST)
 					if ((int32)ActiveParticles > LocalMaxParticleInTrailCount)
 					{
-						if (Component->GetWorld())
+						if (Component.GetWorld())
 						{
 							FString ErrorMessage = 
 								FString::Printf(TEXT("Ribbon with too many particles: %5d vs. %5d, %s"), 
 								ActiveParticles, LocalMaxParticleInTrailCount,
-								Component ? Component->Template ? *(Component->Template->GetName()) : TEXT("No template") : TEXT("No component"));
+								Component.GetTemplate() ? *(Component.GetTemplate()->GetName()) : TEXT("No template"));
 							UE_LOG(LogParticles, Log, TEXT("%s"), *ErrorMessage);
 						}
 					}
@@ -2092,7 +2074,7 @@ void FParticleRibbonEmitterInstance::ResolveSource()
 		case PET2SRCM_Actor:
 			if (SourceActor == NULL)
 			{
-				const TArray<struct FParticleSysParam>& AsyncInstanceParameters = Component->GetAsyncInstanceParameters();
+				TArrayView<const FParticleSysParam> AsyncInstanceParameters = Component.GetAsyncInstanceParameters();
 				FParticleSysParam Param;
 				for (int32 ParamIdx = 0; ParamIdx < AsyncInstanceParameters.Num(); ParamIdx++)
 				{
@@ -2136,9 +2118,10 @@ void FParticleRibbonEmitterInstance::ResolveSource()
 		case PET2SRCM_Particle:
 			if (SourceEmitter == NULL)
 			{
-				for (int32 EmitterIdx = 0; EmitterIdx < Component->EmitterInstances.Num(); EmitterIdx++)
+				TArrayView<FParticleEmitterInstance*> EmitterInstances = Component.GetEmitterInstances();
+				for (int32 EmitterIdx = 0; EmitterIdx < EmitterInstances.Num(); EmitterIdx++)
 				{
-					FParticleEmitterInstance* EmitInst = Component->EmitterInstances[EmitterIdx];
+					FParticleEmitterInstance* EmitInst = EmitterInstances[EmitterIdx];
 					if (EmitInst && (EmitInst->SpriteTemplate->EmitterName == SourceModule->SourceName))
 					{
 						SourceEmitter = EmitInst;
@@ -2231,7 +2214,7 @@ void FParticleRibbonEmitterInstance::UpdateSourceData(float DeltaTime, bool bFir
 bool FParticleRibbonEmitterInstance::ResolveSourcePoint(int32 InTrailIdx, 
 	FVector& OutPosition, FQuat& OutRotation, FVector& OutUp, FVector& OutTangent, float& OutTangentStrength)
 {
-	const FTransform& AsyncComponentToWorld = Component->GetAsyncComponentToWorld();
+	const FTransform& AsyncComponentToWorld = Component.GetAsyncComponentToWorld();
 	bool bSourceWasSet = false;
 	// Resolve the source point...
 	if (SourceModule)
@@ -2317,10 +2300,10 @@ bool FParticleRibbonEmitterInstance::ResolveSourcePoint(int32 InTrailIdx,
 					if (SourceParticle != nullptr)
 					{
 						const FVector WorldOrigin = SourceEmitter->SimulationToWorld.GetOrigin();
-						UParticleSystemComponent* Comp = SourceEmitter->Component;
-						if (!ensureMsgf(!SourceParticle->Location.ContainsNaN(), TEXT("NaN in SourceParticle Location. Template: %s, Component: %s"), Comp ? *GetNameSafe(Comp->Template) : TEXT("UNKNOWN"), *GetPathNameSafe(Comp)) ||
-							!ensureMsgf(!SourceParticle->OldLocation.ContainsNaN(), TEXT("NaN in SourceParticle OldLocation. Template: %s, Component: %s"), Comp ? *GetNameSafe(Comp->Template) : TEXT("UNKNOWN"), *GetPathNameSafe(Comp)) ||
-							!ensureMsgf(!WorldOrigin.ContainsNaN(), TEXT("NaN in WorldOrigin. Template: %s, Component: %s"), Comp ? *GetNameSafe(Comp->Template) : TEXT("UNKNOWN"), *GetPathNameSafe(Comp))
+						IParticleEmitterInstanceOwner& Comp = SourceEmitter->Component;
+						if (!ensureMsgf(!SourceParticle->Location.ContainsNaN(), TEXT("NaN in SourceParticle Location. Template: %s, Component: %s"), *GetNameSafe(Comp.GetTemplate()), *Comp.GetPathName()) ||
+							!ensureMsgf(!SourceParticle->OldLocation.ContainsNaN(), TEXT("NaN in SourceParticle OldLocation. Template: %s, Component: %s"), *GetNameSafe(Comp.GetTemplate()), *Comp.GetPathName()) ||
+							!ensureMsgf(!WorldOrigin.ContainsNaN(), TEXT("NaN in WorldOrigin. Template: %s, Component: %s"), *GetNameSafe(Comp.GetTemplate()), *Comp.GetPathName())
 							)
 						{
 							UE_LOG(LogParticles, Warning, TEXT("TrailIdx: %d"), InTrailIdx);
@@ -2373,8 +2356,8 @@ bool FParticleRibbonEmitterInstance::ResolveSourcePoint(int32 InTrailIdx,
 					else
 					{
 						// Fall back to the emitter location??
-						OutPosition = SourceEmitter->Component->GetComponentLocation();
-						OutTangent = Component->PartSysVelocity;
+						OutPosition = SourceEmitter->Component.GetComponentTransform().GetLocation();
+						OutTangent = Component.GetPartSysVelocity();
 						//@todo. How to handle this... can potentially cause a jump from the emitter to the
 						// particle...
 						SourceTimes[InTrailIdx] = 0.0f;
@@ -2382,7 +2365,7 @@ bool FParticleRibbonEmitterInstance::ResolveSourcePoint(int32 InTrailIdx,
 					}
 					OutTangentStrength = OutTangent.SizeSquared();
 					//@todo. Allow particle rotation to define up??
-					OutUp = SourceEmitter->Component->GetComponentTransform().GetScaledAxis(EAxis::Z);
+					OutUp = SourceEmitter->Component.GetComponentTransform().GetScaledAxis(EAxis::Z);
 
 					//@todo. Where to get rotation from????
 					OutRotation = FQuat(0,0,0,1);
@@ -2421,7 +2404,7 @@ bool FParticleRibbonEmitterInstance::ResolveSourcePoint(int32 InTrailIdx,
 
 	if (bSourceWasSet == false)
 	{
-		OutPosition = Component->GetComponentLocation();
+		OutPosition = Component.GetComponentTransform().GetLocation();
 		if (SourceModule && (SourceModule->SourceOffsetCount > 0))
 		{
 			FVector SourceOffsetValue;
@@ -2430,15 +2413,15 @@ bool FParticleRibbonEmitterInstance::ResolveSourcePoint(int32 InTrailIdx,
 				if (CurrentLODLevel && (CurrentLODLevel->RequiredModule->bUseLocalSpace == false))
 				{
 					// Transform it
-					SourceOffsetValue = Component->GetComponentTransform().TransformVector(SourceOffsetValue);
+					SourceOffsetValue = Component.GetComponentTransform().TransformVector(SourceOffsetValue);
 				}
 				OutPosition += SourceOffsetValue;
 			}
 		}
-		OutRotation = Component->GetComponentQuat();
-		OutTangent = Component->PartSysVelocity;
+		OutRotation = Component.GetComponentTransform().GetRotation();
+		OutTangent = Component.GetPartSysVelocity();
 		OutTangentStrength = OutTangent.SizeSquared();
-		OutUp = Component->GetComponentTransform().GetScaledAxis(EAxis::Z);
+		OutUp = Component.GetComponentTransform().GetScaledAxis(EAxis::Z);
 
 		bSourceWasSet = true;
 	}
@@ -2747,7 +2730,7 @@ void FParticleRibbonEmitterInstance::GetAllocatedSize(int32& OutNum, int32& OutM
  */
 void FParticleRibbonEmitterInstance::GetResourceSizeEx(FResourceSizeEx& CumulativeResourceSize)
 {
-	if (CumulativeResourceSize.GetResourceSizeMode() == EResourceSizeMode::EstimatedTotal || (Component && Component->SceneProxy))
+	if (CumulativeResourceSize.GetResourceSizeMode() == EResourceSizeMode::EstimatedTotal || (Component.GetSceneProxy()))
 	{
 		int32 MaxActiveParticleDataSize = (ParticleData != NULL) ? (MaxActiveParticles * ParticleStride) : 0;
 		int32 MaxActiveParticleIndexSize = (ParticleIndices != NULL) ? (MaxActiveParticles * sizeof(uint16)) : 0;
@@ -2803,8 +2786,8 @@ bool FParticleRibbonEmitterInstance::FillReplayData(FDynamicEmitterReplayDataBas
 		UE_LOG(LogParticles, Warning, TEXT("\tActiveParticleCount = %d."), ActiveParticles);
 		UE_LOG(LogParticles, Warning, TEXT("\tTriangleCount = %d."), TriangleCount);
 		UE_LOG(LogParticles, Warning, TEXT("\tTrailCount = %d."), TrailCount);
-		UE_LOG(LogParticles, Warning, TEXT("\t%s"), Component ? Component->Template ? 
-			*(Component->Template->GetPathName()) : *(Component->GetName()) : TEXT("NO COMPONENT"));
+		UE_LOG(LogParticles, Warning, TEXT("\t%s"), Component.GetTemplate() ? 
+			*(Component.GetTemplate()->GetPathName()) : *(Component.GetName()));
 #endif
 		return false;
 	}
@@ -2826,12 +2809,12 @@ bool FParticleRibbonEmitterInstance::FillReplayData(FDynamicEmitterReplayDataBas
 // 					Component ? (Component->Template ? *Component->Template->GetName() : 
 // 					TEXT("No Template")) : TEXT("No Component"));
 #if !(UE_BUILD_SHIPPING || UE_BUILD_TEST)
-				if (Component && Component->GetWorld())
+				if (Component.GetWorld())
 				{
 					FString ErrorMessage = 
 						FString::Printf(TEXT("RIBBON: GetDynamicData -- TriangleCount == %d (APC = %4d) for PSys %s"),
 							TriangleCount, ActiveParticles, 
-							Component->Template ? *Component->Template->GetName() : TEXT("No Template"));
+							Component.GetTemplate() ? *Component.GetTemplate()->GetName() : TEXT("No Template"));
 					UE_LOG(LogParticles, Log, TEXT("%s"), *ErrorMessage);
 				}
 #endif	//#if !(UE_BUILD_SHIPPING || UE_BUILD_TEST)
@@ -2886,7 +2869,7 @@ void FParticleRibbonEmitterInstance::ApplyWorldOffset(FVector InOffset, bool bWo
 //	FParticleAnimTrailEmitterInstance
 //
 /** Constructor	*/
-FParticleAnimTrailEmitterInstance::FParticleAnimTrailEmitterInstance(UParticleSystemComponent& InComponent) :
+FParticleAnimTrailEmitterInstance::FParticleAnimTrailEmitterInstance(IParticleEmitterInstanceOwner& InComponent) :
 	  FParticleTrailsEmitterInstance_Base(InComponent)
 	, TrailTypeData(NULL)
 	, SpawnPerUnitModule(NULL)
@@ -3147,7 +3130,7 @@ void FParticleAnimTrailEmitterInstance::SpawnParticle( int32& StartParticleIndex
 //	}
 //	else
 	{
-		UMeshComponent* MeshComp = Cast<UMeshComponent>(Component->GetAttachParent());
+		UMeshComponent* MeshComp = Cast<UMeshComponent>(Component.AsComponent()->GetAttachParent());
 		check(TrailTypeData);
 		check(MeshComp);
 
@@ -3363,10 +3346,10 @@ float FParticleAnimTrailEmitterInstance::Spawn(float DeltaTime)
 	if (bProcessSpawnRate)
 	{
 		int32 EffectsQuality = Scalability::GetEffectsQualityDirect(true);
-		float RateScale = LODLevel->SpawnModule->RateScale.GetValue(EmitterTime, Component) * LODLevel->SpawnModule->GetGlobalRateScale();
+		float RateScale = LODLevel->SpawnModule->RateScale.GetValue(EmitterTime, Component.GetDistributionData()) * LODLevel->SpawnModule->GetGlobalRateScale();
 		float QualityMult = 0.25f * (1 << EffectsQuality);
 		RateScale *= SpriteTemplate->QualityLevelSpawnRateScale*QualityMult;
-		SpawnRate += LODLevel->SpawnModule->Rate.GetValue(EmitterTime, Component) * FMath::Clamp<float>(RateScale, 0.0f, RateScale);
+		SpawnRate += LODLevel->SpawnModule->Rate.GetValue(EmitterTime, Component.GetDistributionData()) * FMath::Clamp<float>(RateScale, 0.0f, RateScale);
 	}
 
 	// Take Bursts into account as well...
@@ -3438,7 +3421,7 @@ float FParticleAnimTrailEmitterInstance::Spawn(float DeltaTime)
 	bool bTilingTrail = !FMath::IsNearlyZero(TrailTypeData->TilingDistance);
 
 	//The mesh we're sampline socket locations from.
-	UMeshComponent* MeshComp = Cast<UMeshComponent>(Component->GetAttachParent());
+	UMeshComponent* MeshComp = Cast<UMeshComponent>(Component.AsComponent()->GetAttachParent());
 
 	check(TrailTypeData);
 	// Don't allow new spawning if the emitter is finished
@@ -3563,10 +3546,9 @@ void FParticleAnimTrailEmitterInstance::UpdateSourceData(float DeltaTime, bool b
 void FParticleAnimTrailEmitterInstance::UpdateBoundingBox(float DeltaTime)
 {
 	SCOPE_CYCLE_COUNTER(STAT_ParticleUpdateBounds);
-	if (Component)
 	{
-		bool bUpdateBox = ((Component->bWarmingUp == false) &&
-			(Component->Template != NULL) && (Component->Template->bUseFixedRelativeBoundingBox == false));
+		const UParticleSystem* Template = Component.GetTemplate();
+		bool bUpdateBox = ((Component.GetIsWarmingUp() == false) && (Template != NULL) && (Template->bUseFixedRelativeBoundingBox == false));
 		// Handle local space usage
 		check(SpriteTemplate->LODLevels.Num() > 0);
 		UParticleLODLevel* LODLevel = SpriteTemplate->LODLevels[0];
@@ -3577,8 +3559,8 @@ void FParticleAnimTrailEmitterInstance::UpdateBoundingBox(float DeltaTime)
 			// Set the min/max to the position of the trail
 			if (LODLevel->RequiredModule->bUseLocalSpace == false) 
 			{
-				ParticleBoundingBox.Max = Component->GetComponentLocation();
-				ParticleBoundingBox.Min = Component->GetComponentLocation();
+				ParticleBoundingBox.Max = Component.GetComponentTransform().GetLocation();
+				ParticleBoundingBox.Min = Component.GetComponentTransform().GetLocation();
 			}
 			else
 			{
@@ -3589,7 +3571,7 @@ void FParticleAnimTrailEmitterInstance::UpdateBoundingBox(float DeltaTime)
 		ParticleBoundingBox.IsValid = true;
 
 		// Take scale into account
-		FVector Scale = Component->GetComponentTransform().GetScale3D();
+		FVector Scale = Component.GetComponentTransform().GetScale3D();
 
 		// As well as each particle
 		int32 LocalActiveParticles = ActiveParticles;
@@ -3666,7 +3648,7 @@ void FParticleAnimTrailEmitterInstance::UpdateBoundingBox(float DeltaTime)
 		{
 			if (LODLevel->RequiredModule->bUseLocalSpace) 
 			{
-				ParticleBoundingBox = ParticleBoundingBox.TransformBy(Component->GetComponentTransform());
+				ParticleBoundingBox = ParticleBoundingBox.TransformBy(Component.GetComponentTransform());
 			}
 		}
 	}
@@ -3677,7 +3659,6 @@ void FParticleAnimTrailEmitterInstance::UpdateBoundingBox(float DeltaTime)
  */
 void FParticleAnimTrailEmitterInstance::ForceUpdateBoundingBox()
 {
-	if (Component)
 	{
 		// Handle local space usage
 		check(SpriteTemplate->LODLevels.Num() > 0);
@@ -3687,8 +3668,8 @@ void FParticleAnimTrailEmitterInstance::ForceUpdateBoundingBox()
 		// Set the min/max to the position of the trail
 		if (LODLevel->RequiredModule->bUseLocalSpace == false) 
 		{
-			ParticleBoundingBox.Max = Component->GetComponentLocation();
-			ParticleBoundingBox.Min = Component->GetComponentLocation();
+			ParticleBoundingBox.Max = Component.GetComponentTransform().GetLocation();
+			ParticleBoundingBox.Min = Component.GetComponentTransform().GetLocation();
 		}
 		else
 		{
@@ -3698,7 +3679,7 @@ void FParticleAnimTrailEmitterInstance::ForceUpdateBoundingBox()
 		ParticleBoundingBox.IsValid = true;
 
 		// Take scale into account
-		FVector Scale = Component->GetComponentTransform().GetScale3D();
+		FVector Scale = Component.GetComponentTransform().GetScale3D();
 
 		// As well as each particle
 		int32 LocalActiveParticles = ActiveParticles;
@@ -3750,7 +3731,7 @@ void FParticleAnimTrailEmitterInstance::ForceUpdateBoundingBox()
 		// Transform bounding box into world space if the emitter uses a local space coordinate system.
 		if (LODLevel->RequiredModule->bUseLocalSpace) 
 		{
-			ParticleBoundingBox = ParticleBoundingBox.TransformBy(Component->GetComponentTransform());
+			ParticleBoundingBox = ParticleBoundingBox.TransformBy(Component.GetComponentTransform());
 		}
 	}
 }
@@ -4065,7 +4046,7 @@ void FParticleAnimTrailEmitterInstance::GetAllocatedSize(int32& OutNum, int32& O
  */
 void FParticleAnimTrailEmitterInstance::GetResourceSizeEx(FResourceSizeEx& CumulativeResourceSize)
 {
-	if (CumulativeResourceSize.GetResourceSizeMode() == EResourceSizeMode::EstimatedTotal || (Component && Component->SceneProxy))
+	if (CumulativeResourceSize.GetResourceSizeMode() == EResourceSizeMode::EstimatedTotal || (Component.GetSceneProxy()))
 	{
 		int32 MaxActiveParticleDataSize = (ParticleData != NULL) ? (MaxActiveParticles * ParticleStride) : 0;
 		int32 MaxActiveParticleIndexSize = (ParticleIndices != NULL) ? (MaxActiveParticles * sizeof(uint16)) : 0;
@@ -4176,8 +4157,8 @@ bool FParticleAnimTrailEmitterInstance::FillReplayData( FDynamicEmitterReplayDat
 		UE_LOG(LogParticles, Warning, TEXT("\tActiveParticleCount = %d."), ActiveParticles);
 		UE_LOG(LogParticles, Warning, TEXT("\tTriangleCount = %d."), TriangleCount);
 		UE_LOG(LogParticles, Warning, TEXT("\tTrailCount = %d."), TrailCount);
-		UE_LOG(LogParticles, Warning, TEXT("\t%s"), Component ? Component->Template ? 
-			*(Component->Template->GetPathName()) : *(Component->GetName()) : TEXT("NO COMPONENT"));
+		UE_LOG(LogParticles, Warning, TEXT("\t%s"), Component.GetTemplate() ? 
+			*(Component.GetTemplate()->GetPathName()) : *(Component.GetName()));
 #endif
 		return false;
 	}
@@ -4197,12 +4178,12 @@ bool FParticleAnimTrailEmitterInstance::FillReplayData( FDynamicEmitterReplayDat
 // 					Component ? (Component->Template ? *Component->Template->GetName() : 
 // 					TEXT("No Template")) : TEXT("No Component"));
 #if !(UE_BUILD_SHIPPING || UE_BUILD_TEST)
-			if (Component && Component->GetWorld())
+			if (Component.GetWorld())
 			{
 				FString ErrorMessage = 
 					FString::Printf(TEXT("ANIMTRAIL: GetDynamicData -- TriangleCount == 0 (APC = %4d) for PSys %s"),
 					ActiveParticles, 
-					Component->Template ? *Component->Template->GetName() : TEXT("No Template"));
+					Component.GetTemplate() ? *Component.GetTemplate()->GetName() : TEXT("No Template"));
 				UE_LOG(LogParticles, Log, TEXT("%s"), *ErrorMessage);
 			}
 
@@ -4337,7 +4318,7 @@ void FParticleTrailsEmitterInstance_Base::DumpCircularTrailsSpam()
 {
 	if (GbEnableCircularAnimTrailDump > 0)
 	{
-		UE_LOG(LogParticles, Warning, TEXT("Circular links in Anim Trail discovered. \nPSys: %s\nEmitter: %s"), *Component->Template->GetFullName(), *SpriteTemplate->GetEmitterName().ToString());
+		UE_LOG(LogParticles, Warning, TEXT("Circular links in Anim Trail discovered. \nPSys: %s\nEmitter: %s"), *Component.GetTemplate()->GetFullName(), *SpriteTemplate->GetEmitterName().ToString());
 		if (GbEnableCircularAnimTrailDump > 1)
 		{
 			PrintAllActiveParticles();
diff --git a/Engine/Source/Runtime/Engine/Private/Particles/ParticleTrailModules.cpp b/Engine/Source/Runtime/Engine/Private/Particles/ParticleTrailModules.cpp
--- a/Engine/Source/Runtime/Engine/Private/Particles/ParticleTrailModules.cpp
+++ b/Engine/Source/Runtime/Engine/Private/Particles/ParticleTrailModules.cpp
@@ -232,7 +232,7 @@ void UParticleModuleTypeDataRibbon::PostEditChangeProperty(FPropertyChangedEvent
 }
 #endif // WITH_EDITOR
 
-FParticleEmitterInstance* UParticleModuleTypeDataRibbon::CreateInstance(UParticleEmitter* InEmitterParent, UParticleSystemComponent& InComponent)
+FParticleEmitterInstance* UParticleModuleTypeDataRibbon::CreateInstance(UParticleEmitter* InEmitterParent, IParticleEmitterInstanceOwner& InComponent)
 {
 	FParticleEmitterInstance* Instance = new FParticleRibbonEmitterInstance(InComponent);
 	check(Instance);
@@ -274,7 +274,7 @@ uint32 UParticleModuleTypeDataAnimTrail::RequiredBytes(UParticleModuleTypeDataBa
 	return sizeof(FAnimTrailTypeDataPayload);
 }
 
-FParticleEmitterInstance* UParticleModuleTypeDataAnimTrail::CreateInstance(UParticleEmitter* InEmitterParent, UParticleSystemComponent& InComponent)
+FParticleEmitterInstance* UParticleModuleTypeDataAnimTrail::CreateInstance(UParticleEmitter* InEmitterParent, IParticleEmitterInstanceOwner& InComponent)
 {
 	FParticleEmitterInstance* Instance = new FParticleAnimTrailEmitterInstance(InComponent);
 	check(Instance);
diff --git a/Engine/Source/Runtime/Engine/Public/ParticleEmitterInstanceOwner.h b/Engine/Source/Runtime/Engine/Public/ParticleEmitterInstanceOwner.h
new file mode 100644
--- /dev/null
+++ b/Engine/Source/Runtime/Engine/Public/ParticleEmitterInstanceOwner.h
@@ -0,0 +1,79 @@
+// Copyright Epic Games, Inc. All Rights Reserved.
+
+#pragma once
+
+#include "Containers/UnrealString.h"
+#include "Math/MathFwd.h"
+
+class FFXSystem;
+class FPrimitiveSceneProxy;
+class UMaterialInterface;
+class UParticleSystem;
+class UParticleSystemComponent;
+class UPhysicalMaterial;
+class UObject;
+class UWorld;
+struct FParticleSysParam;
+struct FParticleEmitterInstance;
+struct FParticleEventBurstData;
+struct FParticleEventDeathData;
+struct FParticleEventCollideData;
+struct FParticleEventKismetData;
+struct FParticleEventSpawnData;
+struct FRandomStream;
+
+
+class IParticleEmitterInstanceOwner
+{
+public:
+	virtual const FTransform& GetAsyncComponentToWorld() const = 0;
+
+	virtual UObject* GetDistributionData() const = 0;
+	virtual const FTransform& GetComponentTransform() const = 0;
+	virtual FRotator GetComponentRotation() const = 0;
+	virtual const FTransform& GetComponentToWorld() const = 0;
+	virtual const FBoxSphereBounds& GetBounds() const = 0;
+	virtual UWorld* GetWorld() const = 0;
+	virtual bool GetFloatParameter(const FName InName, float& OutFloat) = 0;
+	virtual const FVector3f& GetLWCTile() const = 0;
+	virtual FString GetName() const = 0;
+	virtual FString GetFullName() const = 0;
+	virtual FString GetPathName() const = 0;
+	virtual bool IsActive() const = 0;
+	virtual bool IsValidLowLevel() const = 0;
+	virtual TArrayView<const FParticleSysParam> GetAsyncInstanceParameters() = 0;
+	virtual int32 GetCurrentDetailMode() const = 0;
+	virtual int32 GetCurrentLODIndex() const = 0;
+	virtual const FVector& GetPartSysVelocity() const = 0;
+	virtual const FVector& GetOldPosition() const = 0;
+	virtual FFXSystem* GetFXSystem() const = 0;
+	virtual const UParticleSystem* GetTemplate() const = 0;
+	virtual TArrayView<const FParticleSysParam> GetInstanceParameters() const = 0;
+	virtual TArrayView<FParticleEmitterInstance*> GetEmitterInstances() const = 0;
+	virtual TArrayView<TObjectPtr<UMaterialInterface>> GetEmitterMaterials() const = 0;
+	virtual FPrimitiveSceneProxy* GetSceneProxy() const = 0;
+	virtual bool GetIsWarmingUp() const = 0;
+	virtual bool GetJustRegistered() const = 0;
+	virtual float GetWarmupTime() const = 0;
+	virtual float GetEmitterDelay() const = 0;
+	virtual FRandomStream& GetRandomStream() const = 0;
+
+	virtual void SetComponentToWorld(const FTransform& NewComponentToWorld) = 0;
+	virtual void DeactivateNextTick() = 0;
+
+	virtual UParticleSystemComponent* AsComponent() const = 0;
+
+	virtual void ReportEventSpawn(const FName InEventName, const float InEmitterTime, const FVector InLocation, const FVector InVelocity, const TArray<class UParticleModuleEventSendToGame*>& InEventData) = 0;
+	virtual void ReportEventDeath(const FName InEventName, const float InEmitterTime, const FVector InLocation, const FVector InVelocity, const TArray<class UParticleModuleEventSendToGame*>& InEventData, const float InParticleTime) = 0;
+	virtual void ReportEventCollision(const FName InEventName, const float InEmitterTime, const FVector InLocation, const FVector InDirection, const FVector InVelocity, const TArray<class UParticleModuleEventSendToGame*>& InEventData, const float InParticleTime, const FVector InNormal, const float InTime, const int32 InItem, const FName InBoneName, UPhysicalMaterial* PhysMat) = 0;
+	virtual void ReportEventBurst(const FName InEventName, const float InEmitterTime, const int32 ParticleCount, const FVector InLocation, const TArray<class UParticleModuleEventSendToGame*>& InEventData) = 0;
+
+	virtual TArrayView<FParticleEventSpawnData> GetSpawnEvents() const = 0;
+	virtual TArrayView<FParticleEventDeathData> GetDeathEvents() const = 0;
+	virtual TArrayView<FParticleEventCollideData> GetCollisionEvents() const = 0;
+	virtual TArrayView<FParticleEventBurstData> GetBurstEvents() const = 0;
+	virtual TArrayView<FParticleEventKismetData> GetKismetEvents() const = 0;
+
+protected:
+	virtual ~IParticleEmitterInstanceOwner() {}
+};
diff --git a/Engine/Source/Runtime/Engine/Public/ParticleEmitterInstances.h b/Engine/Source/Runtime/Engine/Public/ParticleEmitterInstances.h
--- a/Engine/Source/Runtime/Engine/Public/ParticleEmitterInstances.h
+++ b/Engine/Source/Runtime/Engine/Public/ParticleEmitterInstances.h
@@ -37,6 +37,7 @@ class UParticleSystemComponent;
 /*-----------------------------------------------------------------------------
 	Forward declarations
 -----------------------------------------------------------------------------*/
+class IParticleEmitterInstanceOwner;
 //	Emitter and module types
 class UParticleEmitter;
 class UParticleSpriteEmitter;
@@ -215,6 +216,7 @@ struct FParticleEmitterInstanceFixLayout
 	virtual ~FParticleEmitterInstanceFixLayout() = default;
 };
 
+
 /*-----------------------------------------------------------------------------
 	FParticleEmitterInstance
 -----------------------------------------------------------------------------*/
@@ -229,7 +231,7 @@ public:
 	/** The template this instance is based on.							*/
 	UParticleEmitter* SpriteTemplate;
 	/** The component who owns it.										*/
-	UParticleSystemComponent* Component;
+	IParticleEmitterInstanceOwner& Component;
 	/** The currently set LOD level.									*/
 	UParticleLODLevel* CurrentLODLevel;
 	/** The index of the currently set LOD level.						*/
@@ -357,7 +359,7 @@ public:
 	TArray<class UPointLightComponent*> HighQualityLights;
 
 	/** Constructor	*/
-	ENGINE_API FParticleEmitterInstance(UParticleSystemComponent& InComponent);
+	ENGINE_API FParticleEmitterInstance(IParticleEmitterInstanceOwner& InComponent);
 
 	/** Destructor	*/
 	ENGINE_API virtual ~FParticleEmitterInstance();
@@ -835,7 +837,7 @@ struct FScopeCycleCounterEmitter
 struct FParticleSpriteEmitterInstance : public FParticleEmitterInstance
 {
 	/** Constructor	*/
-	FParticleSpriteEmitterInstance(UParticleSystemComponent& InComponent);
+	FParticleSpriteEmitterInstance(IParticleEmitterInstanceOwner& InComponent);
 
 	/** Destructor	*/
 	virtual ~FParticleSpriteEmitterInstance();
@@ -895,7 +897,7 @@ struct FParticleMeshEmitterInstance : public FParticleEmitterInstance
 	TArray<UMaterialInterface*> CurrentMaterials;
 
 	/** Constructor	*/
-	ENGINE_API FParticleMeshEmitterInstance(UParticleSystemComponent& InComponent);
+	ENGINE_API FParticleMeshEmitterInstance(IParticleEmitterInstanceOwner& InComponent);
 
 	ENGINE_API virtual void InitParameters(UParticleEmitter* InTemplate) override;
 	ENGINE_API virtual void Init() override;
@@ -1037,7 +1039,7 @@ struct FParticleBeam2EmitterInstance : public FParticleEmitterInstance
 	TArray<int32>						BeamTrianglesPerSheet;
 
 	/** Constructor	*/
-	FParticleBeam2EmitterInstance(UParticleSystemComponent& InComponent);
+	FParticleBeam2EmitterInstance(IParticleEmitterInstanceOwner& InComponent);
 
 	/** Destructor	*/
 	virtual ~FParticleBeam2EmitterInstance();
@@ -1385,7 +1387,7 @@ void CheckAllIndices(){}
 #endif
 
 	/** Constructor	*/
-	FParticleTrailsEmitterInstance_Base(UParticleSystemComponent& InComponent) :
+	FParticleTrailsEmitterInstance_Base(IParticleEmitterInstanceOwner& InComponent) :
 		  FParticleEmitterInstance(InComponent)
 		, VertexCount(0)
 		, TriangleCount(0)
@@ -1422,8 +1424,7 @@ void CheckAllIndices(){}
 
 	bool AddParticleHelper(int32 InTrailIdx,
 		int32 StartParticleIndex, FTrailsBaseTypeDataPayload* StartTrailData,
-		int32 ParticleIndex, FTrailsBaseTypeDataPayload* TrailData,
-		UParticleSystemComponent* InPsysComp = nullptr);
+		int32 ParticleIndex, FTrailsBaseTypeDataPayload* TrailData);
 
 	/**
 	 *	Tick sub-function that handles recalculation of tangents
@@ -1564,7 +1565,7 @@ struct FParticleRibbonEmitterInstance : public FParticleTrailsEmitterInstance_Ba
 	int32 HeadOnlyParticles;
 
 	/** Constructor	*/
-	FParticleRibbonEmitterInstance(UParticleSystemComponent& InComponent);
+	FParticleRibbonEmitterInstance(IParticleEmitterInstanceOwner& InComponent);
 
 	/** Destructor	*/
 	virtual ~FParticleRibbonEmitterInstance();
@@ -1746,7 +1747,7 @@ struct FParticleAnimTrailEmitterInstance : public FParticleTrailsEmitterInstance
 	int32 HeadOnlyParticles;
 
 	/** Constructor	*/
-	FParticleAnimTrailEmitterInstance(UParticleSystemComponent& InComponent);
+	FParticleAnimTrailEmitterInstance(IParticleEmitterInstanceOwner& InComponent);
 
 	/** Destructor	*/
 	virtual ~FParticleAnimTrailEmitterInstance();
