Fixed an issue with a server crash w/ a Composite Table during GC saying an object flagged as disregard for GC is referencing an object not in the root set.

#rnx

diff --git a/Engine/Source/Runtime/Engine/Private/CompositeDataTable.cpp b/Engine/Source/Runtime/Engine/Private/CompositeDataTable.cpp
--- a/Engine/Source/Runtime/Engine/Private/CompositeDataTable.cpp
+++ b/Engine/Source/Runtime/Engine/Private/CompositeDataTable.cpp
@@ -280,7 +280,20 @@ void UCompositeDataTable::PostEditUndo()
 
 void UCompositeDataTable::AppendParentTables(const TArray<UDataTable*>& NewTables)
 {
-	ParentTables.Append(NewTables);
+	const bool bIsDisregardForGC = GUObjectArray.IsDisregardForGC(this);
+	for (UDataTable* DataTable : NewTables)
+	{
+		if (DataTable)
+		{
+			if (bIsDisregardForGC)
+			{
+				DataTable->AddToRoot();
+			}
+
+			ParentTables.Add(DataTable);
+		}
+	}
+
 	OnParentTablesUpdated(EPropertyChangeType::ValueSet);
 }
 
@@ -290,6 +303,7 @@ void UCompositeDataTable::RemoveParentTables(const TArray<UDataTable*>& TablesTo
 	{
 		if(DataTable)
 		{
+			DataTable->RemoveFromRoot();
 			ParentTables.Remove(DataTable);
 		}
 	}
@@ -298,10 +312,16 @@ void UCompositeDataTable::RemoveParentTables(const TArray<UDataTable*>& TablesTo
 
 void UCompositeDataTable::AddParentTable(const TObjectPtr<UDataTable>& TableToAdd)
 {
-	if(TableToAdd == nullptr)
+	if (TableToAdd == nullptr)
 	{
 		return;
 	}
+
+	if (GUObjectArray.IsDisregardForGC(this))
+	{
+		TableToAdd->AddToRoot();
+	}
+
 	ParentTables.Add(TableToAdd);
 	OnParentTablesUpdated(EPropertyChangeType::ValueSet);
 }
@@ -312,6 +332,8 @@ void UCompositeDataTable::RemoveParentTable(const TObjectPtr<UDataTable>& TableT
 	{
 		return;
 	}
+
+	TableToRemove->RemoveFromRoot();
 	ParentTables.Remove(TableToRemove);
 	OnParentTablesUpdated(EPropertyChangeType::ValueSet);
 }
