Add RowName to datatable rows data validation message

[RN] rnx

diff --git a/Engine/Source/Runtime/Engine/Private/DataTable.cpp b/Engine/Source/Runtime/Engine/Private/DataTable.cpp
--- a/Engine/Source/Runtime/Engine/Private/DataTable.cpp
+++ b/Engine/Source/Runtime/Engine/Private/DataTable.cpp
@@ -3,6 +3,7 @@
 #include "Engine/DataTable.h"
 #include "AssetRegistry/AssetData.h"
 #include "Internationalization/StabilizeLocalizationKeys.h"
+#include "Misc/DataValidation.h"
 #include "Misc/PackageName.h"
 #include "Serialization/PropertyLocalizationDataGathering.h"
 #include "Serialization/ObjectWriter.h"
@@ -16,6 +17,8 @@
 
 #include UE_INLINE_GENERATED_CPP_BY_NAME(DataTable)
 
+#define LOCTEXT_NAMESPACE "DataTable"
+
 LLM_DEFINE_TAG(DataTable);
 
 namespace
@@ -403,10 +406,40 @@ EDataValidationResult UDataTable::IsDataValid(FDataValidationContext& Context) c
 
 	if (RowStruct && RowStruct->IsChildOf(FTableRowBase::StaticStruct()))
 	{
+		FDataValidationContext RowContext(Context.WasAssetLoadedForValidation(), Context.GetValidationUsecase(), Context.GetAssociatedExternalObjects());
+		int32 LastNumIssues = 0;
 		for (const TPair<FName, uint8*>& TableRowPair : RowMap)
 		{
 			const FTableRowBase* CurRow = reinterpret_cast<FTableRowBase*>(TableRowPair.Value);
-			Result = CombineDataValidationResults(Result, CurRow->IsDataValid(Context));
+			Result = CombineDataValidationResults(Result, CurRow->IsDataValid(RowContext));
+
+			const TArray<FDataValidationContext::FIssue>& Issues = RowContext.GetIssues();
+			const int32 NewNumIssues = Issues.Num();
+			for (int32 IssueIdx = LastNumIssues; IssueIdx < NewNumIssues; ++IssueIdx)
+			{
+				const FDataValidationContext::FIssue& Issue = Issues[IssueIdx];
+				if (Issue.TokenizedMessage.IsValid())
+				{
+					TSharedRef<FTokenizedMessage> TokenizedMessage = Issue.TokenizedMessage->Clone();
+					TokenizedMessage->AddText(FText::Format(LOCTEXT("DataTableValidation.RowName", "[Row Name: {0}]"), FText::FromName(TableRowPair.Key)));
+
+					Context.AddMessage(TokenizedMessage);
+				}
+				else
+				{
+					const FText Message = FText::Format(LOCTEXT("DataTableValidation.RowNameWithMsg", "{0} [Row Name: {1}]"), Issue.Message, FText::FromName(TableRowPair.Key));
+					if (Issue.Severity == EMessageSeverity::Error)
+					{
+						Context.AddError(Message);
+					}
+					else
+					{
+						Context.AddWarning(Message);
+					}
+				}
+			}
+
+			LastNumIssues = NewNumIssues;
 		}
 	}
 	
@@ -1030,3 +1063,4 @@ bool FDataTableCategoryHandle::operator != (FDataTableCategoryHandle const& Othe
 	return DataTable != Other.DataTable || ColumnName != Other.ColumnName || RowContents != Other.RowContents;
 }
 
+#undef LOCTEXT_NAMESPACE
\ No newline at end of file
