Simplify implementation of FHLSLMaterialTranslator::AddCodeChunkInner/WithHash/UniformExpression.

The existing logic was effectively a clone of StringBuilder's Appendf method. We now reuse Appendf here.

diff --git a/Engine/Source/Runtime/Engine/Private/Materials/HLSLMaterialTranslator.cpp b/Engine/Source/Runtime/Engine/Private/Materials/HLSLMaterialTranslator.cpp
--- a/Engine/Source/Runtime/Engine/Private/Materials/HLSLMaterialTranslator.cpp
+++ b/Engine/Source/Runtime/Engine/Private/Materials/HLSLMaterialTranslator.cpp
@@ -52,6 +52,7 @@
 #include "ShaderCompiler.h"
 #include "Misc/FileHelper.h"
 #include "Misc/ScopeLock.h"
+#include "Misc/StringBuilder.h"
 #include "Serialization/ObjectWriter.h"
 #include "Serialization/ObjectReader.h"
 #include "Serialization/BufferArchive.h"
@@ -3814,45 +3815,29 @@ int32 FHLSLMaterialTranslator::AddCodeChunkInnerDeriv(const TCHAR* FormattedCode
 	*/
 int32 FHLSLMaterialTranslator::AddCodeChunkInner(EMaterialValueType Type, EDerivativeStatus DerivativeStatus, bool bInlined, const TCHAR* Format, ...)
 {
-	int32	BufferSize		= 256;
-	TCHAR*	FormattedCode	= NULL;
-	int32	Result			= -1;
+	TStringBuilder<256> Builder;
 
-	while(Result == -1)
-	{
-		FormattedCode = (TCHAR*) FMemory::Realloc( FormattedCode, BufferSize * sizeof(TCHAR) );
-		GET_TYPED_VARARGS_RESULT( TCHAR, FormattedCode, BufferSize, BufferSize-1, Format, Format, Result );
-		BufferSize *= 2;
-	};
-	FormattedCode[Result] = 0;
-
-	const uint64 Hash = CityHash64((char*)FormattedCode, Result * sizeof(TCHAR));
-	const int32 CodeIndex = AddCodeChunkInner(Hash, FormattedCode, Type, DerivativeStatus, bInlined);
-	FMemory::Free(FormattedCode);
+	va_list Args;
+	va_start(Args, Format);
+	const TCHAR* FormattedCode = Builder.AppendV(Format, Args).GetData();
+	va_end(Args);
 
-	return CodeIndex;
+	const uint64 Hash = CityHash64((const char*)FormattedCode, Builder.Len() * sizeof(TCHAR));
+	return AddCodeChunkInner(Hash, FormattedCode, Type, DerivativeStatus, bInlined);
 }
 
 int32 FHLSLMaterialTranslator::AddCodeChunkWithHash(uint64 BaseHash, EMaterialValueType Type, const TCHAR* Format, ...)
 {
-	int32	BufferSize = 256;
-	TCHAR*	FormattedCode = NULL;
-	int32	Result = -1;
+	TStringBuilder<256> Builder;
 
-	while (Result == -1)
-	{
-		FormattedCode = (TCHAR*)FMemory::Realloc(FormattedCode, BufferSize * sizeof(TCHAR));
-		GET_TYPED_VARARGS_RESULT(TCHAR, FormattedCode, BufferSize, BufferSize - 1, Format, Format, Result);
-		BufferSize *= 2;
-	};
-	FormattedCode[Result] = 0;
+	va_list Args;
+	va_start(Args, Format);
+	const TCHAR* FormattedCode = Builder.AppendV(Format, Args).GetData();
+	va_end(Args);
 
-	uint64 Hash = CityHash64((char*)FormattedCode, Result * sizeof(TCHAR));
+	uint64 Hash = CityHash64((const char*)FormattedCode, Builder.Len() * sizeof(TCHAR));
 	Hash = CityHash128to64({ BaseHash, Hash });
-	const int32 CodeIndex = AddCodeChunkInner(Hash, FormattedCode, Type, EDerivativeStatus::NotAware, false);
-	FMemory::Free(FormattedCode);
-
-	return CodeIndex;
+	return AddCodeChunkInner(Hash, FormattedCode, Type, EDerivativeStatus::NotAware, false);
 }
 
 int32 FHLSLMaterialTranslator::AddUniformExpressionInner(uint64 Hash, FMaterialUniformExpression* UniformExpression, EMaterialValueType Type, const TCHAR* FormattedCode)
@@ -3971,24 +3956,17 @@ int32 FHLSLMaterialTranslator::AddUniformExpressionInner(uint64 Hash, FMaterialU
 }
 
 // AddUniformExpression - Adds an input to the Code array and returns its index.
-int32 FHLSLMaterialTranslator::AddUniformExpression(FAddUniformExpressionScope& Scope, FMaterialUniformExpression* UniformExpression,EMaterialValueType Type, const TCHAR* Format,...)
+int32 FHLSLMaterialTranslator::AddUniformExpression(FAddUniformExpressionScope& Scope, FMaterialUniformExpression* UniformExpression, EMaterialValueType Type, const TCHAR* Format, ...)
 {
-	int32	BufferSize = 256;
-	TCHAR*	FormattedCode = NULL;
-	int32	Result = -1;
+	TStringBuilder<256> Builder;
 
-	while (Result == -1)
-	{
-		FormattedCode = (TCHAR*)FMemory::Realloc(FormattedCode, BufferSize * sizeof(TCHAR));
-		GET_TYPED_VARARGS_RESULT(TCHAR, FormattedCode, BufferSize, BufferSize - 1, Format, Format, Result);
-		BufferSize *= 2;
-	};
-	FormattedCode[Result] = 0;
+	va_list Args;
+	va_start(Args, Format);
+	const TCHAR* FormattedCode = Builder.AppendV(Format, Args).GetData();
+	va_end(Args);
 
-	const uint64 Hash = CityHash64((char*)FormattedCode, Result * sizeof(TCHAR));
-	const int32 CodeIndex = AddUniformExpressionInner(Hash, UniformExpression, Type, FormattedCode);
-	FMemory::Free(FormattedCode);
-	return CodeIndex;
+	const uint64 Hash = CityHash64((const char*)FormattedCode, Builder.Len() * sizeof(TCHAR));
+	return AddUniformExpressionInner(Hash, UniformExpression, Type, FormattedCode);
 }
 
 // From MaterialUniformExpressions.cpp
