[Engine]
* Iterating on IParticleEmitterInstanceOwner to reduce more dependencies directly on game thread state

diff --git a/Engine/Source/Runtime/Engine/Private/Particles/ParticleBeam2EmitterInstance.cpp b/Engine/Source/Runtime/Engine/Private/Particles/ParticleBeam2EmitterInstance.cpp
--- a/Engine/Source/Runtime/Engine/Private/Particles/ParticleBeam2EmitterInstance.cpp
+++ b/Engine/Source/Runtime/Engine/Private/Particles/ParticleBeam2EmitterInstance.cpp
@@ -1183,7 +1183,7 @@ FDynamicEmitterDataBase* FParticleBeam2EmitterInstance::GetDynamicData(bool bSel
 		)
 	{
 #if !(UE_BUILD_SHIPPING || UE_BUILD_TEST)
-		if (Component.GetWorld())
+		if (Component.HasWorld())
 		{
 			FString ErrorMessage = 
 				FString::Printf(TEXT("BeamEmitter with too much data: %s"),
diff --git a/Engine/Source/Runtime/Engine/Private/Particles/ParticleBeamModules.cpp b/Engine/Source/Runtime/Engine/Private/Particles/ParticleBeamModules.cpp
--- a/Engine/Source/Runtime/Engine/Private/Particles/ParticleBeamModules.cpp
+++ b/Engine/Source/Runtime/Engine/Private/Particles/ParticleBeamModules.cpp
@@ -2030,7 +2030,7 @@ bool UParticleModuleBeamSource::ResolveSourceData(const FContext& Context, FPart
 				// Use the value as a world space position
 				BeamData->SourcePoint	= Source.GetValue(BeamInst->EmitterTime, Context.GetDistributionData());
 				// Take into account current world origin offset
-				BeamData->SourcePoint  -= (FVector)BeamInst->Component.GetWorld()->OriginLocation;
+				BeamData->SourcePoint  -= (FVector)BeamInst->Component.GetWorldOriginLocation();
 			}
 			else
 			{
@@ -2494,7 +2494,7 @@ bool UParticleModuleBeamTarget::ResolveTargetData(const FContext& Context, FPart
 			{
 				BeamData->TargetPoint	= Target.GetValue(BeamInst->EmitterTime, Context.GetDistributionData());
 				// Take into account current world origin offset
-				BeamData->TargetPoint  -= (FVector)BeamInst->Component.GetWorld()->OriginLocation;
+				BeamData->TargetPoint  -= (FVector)BeamInst->Component.GetWorldOriginLocation();
 			}
 			else
 			{
diff --git a/Engine/Source/Runtime/Engine/Private/Particles/ParticleComponents.cpp b/Engine/Source/Runtime/Engine/Private/Particles/ParticleComponents.cpp
--- a/Engine/Source/Runtime/Engine/Private/Particles/ParticleComponents.cpp
+++ b/Engine/Source/Runtime/Engine/Private/Particles/ParticleComponents.cpp
@@ -7,6 +7,7 @@
 #include "Camera/CameraLensEffectInterface.h"
 #include "Engine/GameInstance.h"
 #include "Distributions/Distribution.h"
+#include "GameFramework/WorldSettings.h"
 #include "UObject/FrameworkObjectVersion.h"
 #include "Engine/Level.h"
 #include "Engine/CollisionProfile.h"
@@ -1547,7 +1548,7 @@ UParticleLODLevel* UParticleEmitter::GetCurrentLODLevel(FParticleEmitterInstance
 	else
 	{
 		// for the game (where we care about perf) we don't branch
-		if (Instance->GetWorld()->IsGameWorld() )
+		if (Instance->Component.IsGameWorld() )
 		{
 			return Instance->CurrentLODLevel;
 		}
@@ -3603,7 +3604,37 @@ public:
 	virtual FRotator GetComponentRotation() const override { return Component->GetComponentRotation(); }
 	virtual const FTransform& GetComponentToWorld() const override { return Component->GetComponentToWorld(); }
 	virtual const FBoxSphereBounds& GetBounds() const override { return Component->Bounds; }
-	virtual UWorld* GetWorld() const override { return Component->GetWorld(); }
+	virtual TWeakObjectPtr<UWorld> GetWeakWorld() const override { return Component->GetWorld(); }
+	virtual bool HasWorld() const override { return Component->GetWorld() != nullptr; }
+	virtual bool HasWorldSettings() const override { UWorld* World = Component->GetWorld(); return World ? World->GetWorldSettings() != nullptr : false; };
+	virtual bool IsGameWorld() const { UWorld* World = Component->GetWorld(); return World ? World->IsGameWorld() : false; }
+	virtual float GetWorldTimeSeconds() const override { UWorld* World = Component->GetWorld(); return World ? World->GetTimeSeconds() : 0.0f; }
+	virtual float GetWorldEffectiveTimeDilation() const override
+	{
+		if (UWorld* World = Component->GetWorld())
+		{
+			AWorldSettings* Settings = World->GetWorldSettings();
+			return Settings ? Settings->GetEffectiveTimeDilation() : 1.0f;
+		}
+		return 1.0f;
+	}
+	virtual FIntVector GetWorldOriginLocation() const override
+	{
+		if (UWorld* World = Component->GetWorld())
+		{
+			return World->OriginLocation;
+		}
+		return FIntVector(ForceInitToZero);
+	}
+	virtual FSceneInterface* GetScene() const override
+	{
+		if (UWorld* World = Component->GetWorld())
+		{
+			return World->Scene;
+		}
+		return nullptr;
+	}
+
 	virtual bool GetFloatParameter(const FName InName, float& OutFloat) override { return Component->GetFloatParameter(InName, OutFloat); }
 	virtual const FVector3f& GetLWCTile() const override { return Component->GetLWCTile(); }
 	virtual FString GetName() const override { return Component->GetName(); }
diff --git a/Engine/Source/Runtime/Engine/Private/Particles/ParticleEmitterInstances.cpp b/Engine/Source/Runtime/Engine/Private/Particles/ParticleEmitterInstances.cpp
--- a/Engine/Source/Runtime/Engine/Private/Particles/ParticleEmitterInstances.cpp
+++ b/Engine/Source/Runtime/Engine/Private/Particles/ParticleEmitterInstances.cpp
@@ -589,7 +589,7 @@ void FParticleEmitterInstance::Init()
 
 	// Resize to sensible default.
 	if (bNeedsInit && 
-		Component.GetWorld()->IsGameWorld() == true &&
+		Component.IsGameWorld() == true &&
 		// Only presize if any particles will be spawned 
 		SpriteTemplate->QualityLevelSpawnRateScale > 0)
 	{
@@ -675,11 +675,6 @@ void FParticleEmitterInstance::Init()
 	bEmitterIsDone = false;
 }
 
-UWorld* FParticleEmitterInstance::GetWorld() const
-{
-	return Component.GetWorld();
-}
-
 void FParticleEmitterInstance::UpdateTransforms()
 {
 	//QUICK_SCOPE_CYCLE_COUNTER(STAT_EmitterInstance_UpdateTransforms);
@@ -1176,7 +1171,7 @@ void FParticleEmitterInstance::SetCurrentLODIndex(int32 InLODIndex, bool bInFull
 			}
 		}
 
-		if ((GetWorld()->IsGameWorld() == true) && (CurrentLODLevel->bEnabled == false))
+		if ((Component.IsGameWorld() == true) && (CurrentLODLevel->bEnabled == false))
 		{
 			// Kill active particles...
 			KillParticlesForced();
@@ -1200,10 +1195,10 @@ void FParticleEmitterInstance::SetCurrentLODIndex(int32 InLODIndex, bool bInFull
  */
 void FParticleEmitterInstance::Rewind()
 {
-	if (Component.GetWorld())
+	if (Component.HasWorld())
 	{
 		UE_LOG(LogParticles,Verbose,
-			TEXT("FParticleEmitterInstance::Rewind @ %fs %s"), Component.GetWorld()->TimeSeconds,
+			TEXT("FParticleEmitterInstance::Rewind @ %fs %s"), Component.GetWorldTimeSeconds(),
 			(SpriteTemplate != NULL && SpriteTemplate->GetOuter() != NULL) ? *SpriteTemplate->GetOuter()->GetName() : TEXT("NULL"));
 	}
 
@@ -1921,8 +1916,7 @@ void FParticleEmitterInstance::CheckSpawnCount(int32 InNewCount, int32 InMaxCoun
 {
 #if !(UE_BUILD_SHIPPING || UE_BUILD_TEST)
 	{
-		AWorldSettings* WorldSettings = Component.GetWorld() ? Component.GetWorld()->GetWorldSettings() : NULL;
-		if (WorldSettings)
+		if (Component.HasWorldSettings())
 		{
 			const int32 SizeScalar = sizeof(FParticleSpriteVertex);
 
diff --git a/Engine/Source/Runtime/Engine/Private/Particles/ParticleGpuSimulation.cpp b/Engine/Source/Runtime/Engine/Private/Particles/ParticleGpuSimulation.cpp
--- a/Engine/Source/Runtime/Engine/Private/Particles/ParticleGpuSimulation.cpp
+++ b/Engine/Source/Runtime/Engine/Private/Particles/ParticleGpuSimulation.cpp
@@ -3575,7 +3575,7 @@ FGPUSpriteParticleEmitterInstance(IParticleEmitterInstanceOwner& InComponent, FF
 
 		InitLocalVectorField();
 
-		WeakWorld = Component.GetWorld();
+		WeakWorld = Component.GetWeakWorld();
 	}
 
 	FORCENOINLINE void ReserveNewParticles(int32 Num)
diff --git a/Engine/Source/Runtime/Engine/Private/Particles/ParticleModules.cpp b/Engine/Source/Runtime/Engine/Private/Particles/ParticleModules.cpp
--- a/Engine/Source/Runtime/Engine/Private/Particles/ParticleModules.cpp
+++ b/Engine/Source/Runtime/Engine/Private/Particles/ParticleModules.cpp
@@ -2185,15 +2185,7 @@ float UParticleModuleSubUV::DetermineImageIndex(const FContext& Context, int32 O
 		}
 		else
 		{
-			UWorld* World = Owner->Component.GetWorld();
-			if ((World != NULL) && (World->GetWorldSettings() != NULL))
-			{
-				ImageIndex = SubImageIndex.GetValue(Particle->RelativeTime / World->GetWorldSettings()->GetEffectiveTimeDilation(), Context.GetDistributionData());
-			}
-			else
-			{
-				ImageIndex = SubImageIndex.GetValue(Particle->RelativeTime, Context.GetDistributionData());
-			}
+			ImageIndex = SubImageIndex.GetValue(Particle->RelativeTime / Owner->Component.GetWorldEffectiveTimeDilation(), Context.GetDistributionData());
 		}
 
 		if (InterpMethod == PSUVIM_Linear)
@@ -2369,15 +2361,7 @@ float UParticleModuleSubUVMovie::DetermineImageIndex(const FContext& Context, in
 	}
 	else
 	{
-		UWorld* World = Owner->Component.GetWorld();
-		if ((World != NULL) && (World->GetWorldSettings() != NULL))
-		{
-			MoviePayload.Time += DeltaTime / World->GetWorldSettings()->GetEffectiveTimeDilation();
-		}
-		else
-		{
-			MoviePayload.Time += DeltaTime;
-		}
+		MoviePayload.Time += DeltaTime / Owner->Component.GetWorldEffectiveTimeDilation();
 	}
 	
 	float ImageIndex = MoviePayload.Time * UserSetFrameRate;
@@ -3255,12 +3239,7 @@ void UParticleModuleLight::Update(const FUpdateContext& Context)
 		return;
 	}
 
-	UWorld* OwnerWorld = Owner->GetWorld();
-	FSceneInterface* OwnerScene = nullptr;
-	if (OwnerWorld)
-	{
-		OwnerScene = OwnerWorld->Scene;
-	}
+	FSceneInterface* OwnerScene = Owner->Component.GetScene();;
 	
 	TSet<uint64> ActiveLights;
 	UParticleLODLevel* LODLevel	= Owner->SpriteTemplate->GetCurrentLODLevel(Owner);
@@ -5023,14 +5002,12 @@ void UParticleModuleTypeDataGpu::Build( FParticleEmitterBuildInfo& EmitterBuildI
 
 FParticleEmitterInstance* UParticleModuleTypeDataGpu::CreateInstance(UParticleEmitter* InEmitterParent, IParticleEmitterInstanceOwner& InComponent)
 {
-	UWorld* World = InComponent.GetWorld();
-	check(World);
 	UE_LOG(LogParticles,Verbose,
-		TEXT("Create GPU Sprite Emitter @ %fs %s"), World->TimeSeconds,
+		TEXT("Create GPU Sprite Emitter @ %fs %s"), InComponent.GetWorldTimeSeconds(),
 		InComponent.GetTemplate() != NULL ? *InComponent.GetTemplate()->GetName() : TEXT("NULL"));
 
 	FParticleEmitterInstance* Instance = NULL;
-	if (World->Scene && RHISupportsGPUParticles())
+	if (InComponent.GetScene() && RHISupportsGPUParticles())
 	{
 		LLM_SCOPE(ELLMTag::Particles);
 
diff --git a/Engine/Source/Runtime/Engine/Private/Particles/ParticleModules_Collision.cpp b/Engine/Source/Runtime/Engine/Private/Particles/ParticleModules_Collision.cpp
--- a/Engine/Source/Runtime/Engine/Private/Particles/ParticleModules_Collision.cpp
+++ b/Engine/Source/Runtime/Engine/Private/Particles/ParticleModules_Collision.cpp
@@ -182,14 +182,13 @@ void UParticleModuleCollision::Update(const FUpdateContext& Context)
 {
 	SCOPE_CYCLE_COUNTER(STAT_ParticleCollisionTime);
 	FParticleEmitterInstance* Owner = &Context.Owner;
-	UWorld* World = Owner->Component.GetWorld();
+	UParticleSystemComponent* Component = Owner->Component.AsComponent();
+	UWorld* World = Component->GetWorld();
 	if (Owner->ActiveParticles == 0 || (bDropDetail && World && World->bDropDetail))
 	{
 		return;
 	}
 
-	UParticleSystemComponent* Component = Owner->Component.AsComponent();
-
 	//Gets the owning actor of the component. Can be NULL if the component is spawned with the World as an Outer, e.g. in UGameplayStatics::SpawnEmitterAtLocation().
 	AActor* Actor = Component->GetOwner();
 
diff --git a/Engine/Source/Runtime/Engine/Private/Particles/ParticleTrail2EmitterInstance.cpp b/Engine/Source/Runtime/Engine/Private/Particles/ParticleTrail2EmitterInstance.cpp
--- a/Engine/Source/Runtime/Engine/Private/Particles/ParticleTrail2EmitterInstance.cpp
+++ b/Engine/Source/Runtime/Engine/Private/Particles/ParticleTrail2EmitterInstance.cpp
@@ -163,7 +163,7 @@ void FParticleTrailsEmitterInstance_Base::Tick(float DeltaTime, bool bSuppressSp
 		// 'Reset' the emitter time so that the delay functions correctly
 		EmitterTime += CurrentDelay;
 		RunningTime += DeltaTime;
-		LastTickTime = Component.GetWorld() ? Component.GetWorld()->GetTimeSeconds() : 0.0f;
+		LastTickTime = Component.GetWorldTimeSeconds();
 		
 		// Reset particles position offset
 		PositionOffsetThisTick = FVector::ZeroVector;
@@ -340,11 +340,7 @@ void FParticleTrailsEmitterInstance_Base::KillParticles()
 {
 	if (ActiveParticles > 0)
 	{
-		float CurrentTickTime = 0.0f;
-		if (Component.GetWorld())
-		{
-			CurrentTickTime = Component.GetWorld()->GetTimeSeconds();
-		}
+		float CurrentTickTime = Component.GetWorldTimeSeconds();
 		bool bHasForceKillParticles = false;
 		// Loop over the active particles... If their RelativeTime is > 1.0f (indicating they are dead),
 		// move them to the 'end' of the active particle list.
@@ -1471,7 +1467,7 @@ float FParticleRibbonEmitterInstance::Spawn(float DeltaTime)
 #if !(UE_BUILD_SHIPPING || UE_BUILD_TEST)
 				if ((int32)ActiveParticles > LocalMaxParticleInTrailCount)
 				{
-					if (Component.GetWorld())
+					if (Component.HasWorld())
 					{
 						const FString TemplateName = Component.GetTemplate() ? *Component.GetTemplate()->GetName() : TEXT("No template");
 						GMaxParticleInTrailCountWarningCS.Lock();
@@ -1952,7 +1948,7 @@ bool FParticleRibbonEmitterInstance::Spawn_Source(float DeltaTime)
 #if !(UE_BUILD_SHIPPING || UE_BUILD_TEST)
 					if ((int32)ActiveParticles > LocalMaxParticleInTrailCount)
 					{
-						if (Component.GetWorld())
+						if (Component.HasWorld())
 						{
 							FString ErrorMessage = 
 								FString::Printf(TEXT("Ribbon with too many particles: %5d vs. %5d, %s"), 
@@ -2809,7 +2805,7 @@ bool FParticleRibbonEmitterInstance::FillReplayData(FDynamicEmitterReplayDataBas
 // 					Component ? (Component->Template ? *Component->Template->GetName() : 
 // 					TEXT("No Template")) : TEXT("No Component"));
 #if !(UE_BUILD_SHIPPING || UE_BUILD_TEST)
-				if (Component.GetWorld())
+				if (Component.HasWorld())
 				{
 					FString ErrorMessage = 
 						FString::Printf(TEXT("RIBBON: GetDynamicData -- TriangleCount == %d (APC = %4d) for PSys %s"),
@@ -4178,7 +4174,7 @@ bool FParticleAnimTrailEmitterInstance::FillReplayData( FDynamicEmitterReplayDat
 // 					Component ? (Component->Template ? *Component->Template->GetName() : 
 // 					TEXT("No Template")) : TEXT("No Component"));
 #if !(UE_BUILD_SHIPPING || UE_BUILD_TEST)
-			if (Component.GetWorld())
+			if (Component.HasWorld())
 			{
 				FString ErrorMessage = 
 					FString::Printf(TEXT("ANIMTRAIL: GetDynamicData -- TriangleCount == 0 (APC = %4d) for PSys %s"),
diff --git a/Engine/Source/Runtime/Engine/Public/ParticleEmitterInstanceOwner.h b/Engine/Source/Runtime/Engine/Public/ParticleEmitterInstanceOwner.h
--- a/Engine/Source/Runtime/Engine/Public/ParticleEmitterInstanceOwner.h
+++ b/Engine/Source/Runtime/Engine/Public/ParticleEmitterInstanceOwner.h
@@ -4,9 +4,11 @@
 
 #include "Containers/UnrealString.h"
 #include "Math/MathFwd.h"
+#include "UObject/WeakObjectPtr.h"
 
 class FFXSystem;
 class FPrimitiveSceneProxy;
+class FSceneInterface;
 class UMaterialInterface;
 class UParticleSystem;
 class UParticleSystemComponent;
@@ -33,7 +35,14 @@ public:
 	virtual FRotator GetComponentRotation() const = 0;
 	virtual const FTransform& GetComponentToWorld() const = 0;
 	virtual const FBoxSphereBounds& GetBounds() const = 0;
-	virtual UWorld* GetWorld() const = 0;
+	virtual TWeakObjectPtr<UWorld> GetWeakWorld() const = 0;
+	virtual bool HasWorld() const = 0;
+	virtual bool HasWorldSettings() const = 0;
+	virtual bool IsGameWorld() const = 0;
+	virtual float GetWorldTimeSeconds() const = 0;
+	virtual float GetWorldEffectiveTimeDilation() const = 0;
+	virtual FIntVector GetWorldOriginLocation() const = 0;
+	virtual FSceneInterface* GetScene() const = 0;
 	virtual bool GetFloatParameter(const FName InName, float& OutFloat) = 0;
 	virtual const FVector3f& GetLWCTile() const = 0;
 	virtual FString GetName() const = 0;
diff --git a/Engine/Source/Runtime/Engine/Public/ParticleEmitterInstances.h b/Engine/Source/Runtime/Engine/Public/ParticleEmitterInstances.h
--- a/Engine/Source/Runtime/Engine/Public/ParticleEmitterInstances.h
+++ b/Engine/Source/Runtime/Engine/Public/ParticleEmitterInstances.h
@@ -372,9 +372,6 @@ public:
 	ENGINE_API virtual void InitParameters(UParticleEmitter* InTemplate);
 	ENGINE_API virtual void Init();
 
-	/** @return The world that the component that owns this instance is in */
-	ENGINE_API UWorld* GetWorld() const;
-
 	/**
 	 * Ensures enough memory is allocated for the requested number of particles.
 	 *
