Decouple spline mesh rendering from the spline mesh component

#tests confirmed landscape splines work in editor

diff --git a/Engine/Source/Runtime/Engine/Classes/Components/SplineMeshComponent.h b/Engine/Source/Runtime/Engine/Classes/Components/SplineMeshComponent.h
--- a/Engine/Source/Runtime/Engine/Classes/Components/SplineMeshComponent.h
+++ b/Engine/Source/Runtime/Engine/Classes/Components/SplineMeshComponent.h
@@ -30,6 +30,11 @@ namespace ESplineMeshAxis
 	};
 }
 
+namespace UE::SplineMesh
+{
+	ENGINE_API float RealToFloatChecked(const double Value);
+}
+
 /** 
  * Structure that holds info about spline, passed to renderer to deform UStaticMesh.
  * Also used by Lightmass, so be sure to update Lightmass::FSplineMeshParams and the static lighting code if this changes!
@@ -226,6 +231,7 @@ protected:
 
 public:
 	ENGINE_API virtual class UBodySetup* GetBodySetup() override;
+	ENGINE_API class UBodySetup* GetBodySetup() const;
 #if WITH_EDITOR
 	virtual bool ShouldRenderSelected() const override
 	{
@@ -249,6 +255,7 @@ public:
 	ENGINE_API virtual bool ContainsPhysicsTriMeshData(bool InUseAllTriData) const override;
 	virtual bool WantsNegXTriMesh() override { return false; }
 	ENGINE_API virtual void GetMeshId(FString& OutMeshId) override;
+	ENGINE_API void GetMeshId(FString& OutMeshId) const;
 	ENGINE_API virtual bool GetTriMeshSizeEstimates(struct FTriMeshCollisionDataEstimates& OutTriMeshEstimates, bool bInUseAllTriData) const override;
 	//~ End Interface_CollisionDataProvider Interface
 
@@ -423,6 +430,7 @@ public:
 
 	/** Returns a vector which, when componentwise-multiplied by another vector, will zero all the components not corresponding to the supplied ESplineMeshAxis */
 	inline static FVector GetAxisMask(ESplineMeshAxis::Type InAxis);
+	ENGINE_API static bool ShouldRenderNaniteSplineMeshes();
 
 	ENGINE_API virtual float GetTextureStreamingTransformScale() const override;
 
diff --git a/Engine/Source/Runtime/Engine/Private/Components/SplineMeshComponent.cpp b/Engine/Source/Runtime/Engine/Private/Components/SplineMeshComponent.cpp
--- a/Engine/Source/Runtime/Engine/Private/Components/SplineMeshComponent.cpp
+++ b/Engine/Source/Runtime/Engine/Private/Components/SplineMeshComponent.cpp
@@ -25,6 +25,7 @@
 #include "Engine/World.h"
 #include "NaniteVertexFactory.h"
 #include "StaticMeshComponentLODInfo.h"
+#include "SplineMeshSceneProxyDesc.h"
 
 #if WITH_EDITOR
 #include "IHierarchicalLODUtilities.h"
@@ -61,7 +62,7 @@ static FAutoConsoleVariableRef CVarSplineMeshRenderNanite(
 	)
 );
 
-static bool ShouldRenderNaniteSplineMeshes()
+bool USplineMeshComponent::ShouldRenderNaniteSplineMeshes()
 {
 	return NaniteSplineMeshesSupported() && GSplineMeshRenderNanite != 0;
 }
@@ -106,69 +107,6 @@ void PackSplineMeshParams(const FSplineMeshShaderParams& Params, const TArrayVie
 	Output[6].W = FMath::AsFloat(PackSNorm16(MeshRot.Z) | PackSNorm16(MeshRot.W, 16u));
 }
 
-/**
-* Functions used for transforming a static mesh component based on a spline.
-* This needs to be updated if the spline functionality changes!
-*/
-static float SmoothStep(float A, float B, float X)
-{
-	if (X < A)
-	{
-		return 0.0f;
-	}
-	else if (X >= B)
-	{
-		return 1.0f;
-	}
-	const float InterpFraction = (X - A) / (B - A);
-	return InterpFraction * InterpFraction * (3.0f - 2.0f * InterpFraction);
-}
-
-static FVector3f SplineEvalPos(const FVector3f& StartPos, const FVector3f& StartTangent, const FVector3f& EndPos, const FVector3f& EndTangent, float A)
-{
-	const float A2 = A * A;
-	const float A3 = A2 * A;
-
-	return (((2 * A3) - (3 * A2) + 1) * StartPos) + ((A3 - (2 * A2) + A) * StartTangent) + ((A3 - A2) * EndTangent) + (((-2 * A3) + (3 * A2)) * EndPos);
-}
-
-static FVector3f SplineEvalPos(const FSplineMeshParams& Params, float A)
-{
-	// TODO: these don't need to be doubles!
-	const FVector3f StartPos = FVector3f(Params.StartPos);
-	const FVector3f StartTangent = FVector3f(Params.StartTangent);
-	const FVector3f EndPos = FVector3f(Params.EndPos);
-	const FVector3f EndTangent = FVector3f(Params.EndTangent);
-
-	return SplineEvalPos(StartPos, StartTangent, EndPos, EndTangent, A);
-}
-
-static FVector3f SplineEvalTangent(const FVector3f& StartPos, const FVector3f& StartTangent, const FVector3f& EndPos, const FVector3f& EndTangent, const float A)
-{
-	const FVector3f C = (6 * StartPos) + (3 * StartTangent) + (3 * EndTangent) - (6 * EndPos);
-	const FVector3f D = (-6 * StartPos) - (4 * StartTangent) - (2 * EndTangent) + (6 * EndPos);
-	const FVector3f E = StartTangent;
-
-	const float A2 = A * A;
-
-	return (C * A2) + (D * A) + E;
-}
-
-static FVector3f SplineEvalTangent(const FSplineMeshParams& Params, const float A)
-{
-	// TODO: these don't need to be doubles!
-	const FVector3f StartPos = FVector3f(Params.StartPos);
-	const FVector3f StartTangent = FVector3f(Params.StartTangent);
-	const FVector3f EndPos = FVector3f(Params.EndPos);
-	const FVector3f EndTangent = FVector3f(Params.EndTangent);
-
-	return SplineEvalTangent(StartPos, StartTangent, EndPos, EndTangent, A);
-}
-
-static FVector3f SplineEvalDir(const FSplineMeshParams& Params, const float A)
-{
-	return SplineEvalTangent(Params, A).GetSafeNormal();
-}
 
 FSplineMeshInstanceData::FSplineMeshInstanceData(const USplineMeshComponent* SourceComponent)
 	: FStaticMeshComponentInstanceData(SourceComponent)
@@ -311,50 +249,7 @@ USplineMeshComponent::USplineMeshComponent(const FObjectInitializer& ObjectIniti
 
 void USplineMeshComponent::InitVertexFactory(int32 InLODIndex, FColorVertexBuffer* InOverrideColorVertexBuffer)
 {
-	UStaticMesh* Mesh = GetStaticMesh();
-	if (Mesh == nullptr)
-	{
-		return;
-	}
-
-	FStaticMeshLODResources& LODRenderData = Mesh->GetRenderData()->LODResources[InLODIndex];
-	FStaticMeshVertexFactories& VertexFactories = Mesh->GetRenderData()->LODVertexFactories[InLODIndex];
-
-	// Skip LODs that have their render data stripped (eg. platform MinLod settings)
-	if (LODRenderData.VertexBuffers.StaticMeshVertexBuffer.GetNumVertices() == 0)
-	{
-		return;
-	}
-
-	bool bOverrideColorVertexBuffer = !!InOverrideColorVertexBuffer;
-	const ERHIFeatureLevel::Type FeatureLevel = GetWorld()->GetFeatureLevel();
-
-	if ((VertexFactories.SplineVertexFactory && !bOverrideColorVertexBuffer) || (VertexFactories.SplineVertexFactoryOverrideColorVertexBuffer && bOverrideColorVertexBuffer))
-	{
-		// we already have it
-		return;
-	}
-
-	FSplineMeshVertexFactory* VertexFactory = new FSplineMeshVertexFactory(FeatureLevel);
-	if (bOverrideColorVertexBuffer)
-	{
-		VertexFactories.SplineVertexFactoryOverrideColorVertexBuffer = VertexFactory;
-	}
-	else
-	{
-		VertexFactories.SplineVertexFactory = VertexFactory;
-	}
-
-	int32 LightMapCoordinateIndex = Mesh->GetLightMapCoordinateIndex();
-	// Initialize the static mesh's vertex factory.
-	ENQUEUE_RENDER_COMMAND(InitSplineMeshVertexFactory)(
-		[VertexFactory, &LODRenderData, bOverrideColorVertexBuffer, LightMapCoordinateIndex](FRHICommandListBase& RHICmdList)
-	{
-		FLocalVertexFactory::FDataType Data;
-		InitSplineMeshVertexFactoryComponents(LODRenderData.VertexBuffers, VertexFactory, LightMapCoordinateIndex, bOverrideColorVertexBuffer, Data);
-		VertexFactory->SetData(RHICmdList, Data);
-		VertexFactory->InitResource(RHICmdList);
-	});
+	FSplineMeshSceneProxyDesc::InitVertexFactory(GetStaticMesh(), GetWorld()->GetFeatureLevel(), InLODIndex, InOverrideColorVertexBuffer);
 }
 
 FVector USplineMeshComponent::GetStartPosition() const
@@ -642,99 +537,7 @@ void USplineMeshComponent::UpdateMesh_Concurrent()
 
 FSplineMeshShaderParams USplineMeshComponent::CalculateShaderParams() const
 {
-	FSplineMeshShaderParams Output;
-
-	Output.StartPos 				= FVector3f(SplineParams.StartPos);
-	Output.EndPos 					= FVector3f(SplineParams.EndPos);
-	Output.StartTangent 			= FVector3f(SplineParams.StartTangent);
-	Output.EndTangent 				= FVector3f(SplineParams.EndTangent);
-	Output.StartScale 				= FVector2f(SplineParams.StartScale);
-	Output.EndScale 				= FVector2f(SplineParams.EndScale);
-	Output.StartOffset 				= FVector2f(SplineParams.StartOffset);
-	Output.EndOffset 				= FVector2f(SplineParams.EndOffset);
-	Output.StartRoll 				= SplineParams.StartRoll;
-	Output.EndRoll 					= SplineParams.EndRoll;
-	Output.NaniteClusterBoundsScale	= SplineParams.NaniteClusterBoundsScale;
-	Output.bSmoothInterpRollScale 	= bSmoothInterpRollScale;
-	Output.SplineUpDir 				= FVector3f(SplineUpDir);
-	Output.TextureCoord 			= FUintVector2(INDEX_NONE, INDEX_NONE); // either unused or assigned later
-
-	const uint32 MeshXAxis = (ForwardAxis + 1) % 3;
-	const uint32 MeshYAxis = (ForwardAxis + 2) % 3;
-	Output.MeshDir = Output.MeshX = Output.MeshY = FVector3f::ZeroVector;
-	Output.MeshDir[ForwardAxis] = 1.0f;
-	Output.MeshX[MeshXAxis] = 1.0f;
-	Output.MeshY[MeshYAxis] = 1.0f;
-
-	Output.MeshZScale = 1.0f;
-	Output.MeshZOffset = 0.0f;
-
-	if (GetStaticMesh())
-	{
-		const FBoxSphereBounds StaticMeshBounds = GetStaticMesh()->GetBounds();
-		const float BoundsXYRadius = FVector3f(StaticMeshBounds.BoxExtent).Dot((Output.MeshX + Output.MeshY).GetUnsafeNormal());
-
-		const float MeshMinZ = UE::SplineMesh::RealToFloatChecked(GetAxisValueRef(StaticMeshBounds.Origin - StaticMeshBounds.BoxExtent, ForwardAxis));
-		const float MeshZLen = UE::SplineMesh::RealToFloatChecked(2 * GetAxisValueRef(StaticMeshBounds.BoxExtent, ForwardAxis));
-		const float InvMeshZLen = (MeshZLen <= 0.0f) ? 1.0f : 1.0f / MeshZLen;
-		constexpr float MeshTexelLen = float(SPLINE_MESH_TEXEL_WIDTH - 1);
-
-		if (FMath::IsNearlyEqual(SplineBoundaryMin, SplineBoundaryMax))
-		{
-			Output.MeshZScale = InvMeshZLen;
-			Output.MeshZOffset = -MeshMinZ * InvMeshZLen;
-			Output.SplineDistToTexelScale = MeshTexelLen;
-			Output.SplineDistToTexelOffset = 0.0f;
-		}
-		else
-		{
-			const float BoundaryLen = SplineBoundaryMax - SplineBoundaryMin;
-			const float InvBoundaryLen = 1.0f / BoundaryLen;
-
-			Output.MeshZScale = InvBoundaryLen;
-			Output.MeshZOffset = -SplineBoundaryMin * InvBoundaryLen;
-			Output.SplineDistToTexelScale = BoundaryLen * InvMeshZLen * MeshTexelLen;
-			Output.SplineDistToTexelOffset = (SplineBoundaryMin - MeshMinZ) * InvMeshZLen * MeshTexelLen;
-		}
-
-		// Iteratively solve for an approximation of spline length
-		float SplineLength = 0.0f;
-		{
-			static const uint32 NumIterations = 63; // 64 sampled points
-			static const float IterStep = 1.0f / float(NumIterations);
-			float A = 0.0f;
-			FVector3f PrevPoint = SplineEvalPos(SplineParams, A);
-			for (uint32 i = 0; i < NumIterations; ++i)
-			{
-				FVector3f Point = SplineEvalPos(SplineParams, A);
-				SplineLength += (Point - PrevPoint).Length();
-				PrevPoint = Point;
-				A += IterStep;
-			}
-		}
-
-		// Calculate an approximation of how much the mesh gets scaled in each local axis as a result of spline
-		// deformation and take the smallest of the axes. This is important for LOD selection of Nanite spline
-		// meshes.
-		{
-			// Estimate length added due to twisting as well
-			const float XYRadius = BoundsXYRadius * FMath::Max(Output.StartScale.GetAbsMax(), Output.EndScale.GetAbsMax());
-			const float TwistRadians = FMath::Abs(Output.StartRoll - Output.EndRoll);
-			SplineLength += TwistRadians * XYRadius;
-
-			// Take the mid-point scale in X/Y to balance out LOD selection in case either of them are extreme.
-			auto AvgAbs = [](float A, float B) { return (FMath::Abs(A) + FMath::Abs(B)) * 0.5f; };
-			const FVector3f DeformScale = FVector3f(
-				SplineLength * Output.MeshZScale,
-				AvgAbs(Output.StartScale.X, Output.EndScale.X),
-				AvgAbs(Output.StartScale.Y, Output.EndScale.Y)
-			);
-			
-			Output.MeshDeformScaleMinMax = FVector2f(DeformScale.GetMin(), DeformScale.GetMax());
-		}
-	}
-
-	return Output;
+	return FSplineMeshSceneProxyDesc(this).CalculateShaderParams(GetStaticMesh());
 }
 
 void USplineMeshComponent::UpdateRenderStateAndCollision()
@@ -949,121 +752,17 @@ void USplineMeshComponent::UpdateBounds()
 
 float USplineMeshComponent::ComputeRatioAlongSpline(const float DistanceAlong) const 
 {
-	// Find how far 'along' mesh (or custom boundaries) we are
-	float Alpha = 0.f;
-
-	const bool bHasCustomBoundary = !FMath::IsNearlyEqual(SplineBoundaryMin, SplineBoundaryMax);
-	if (bHasCustomBoundary)
-	{
-		Alpha = (DistanceAlong - SplineBoundaryMin) / (SplineBoundaryMax - SplineBoundaryMin);
-	}
-	else if (GetStaticMesh())
-	{
-		const FBoxSphereBounds StaticMeshBounds = GetStaticMesh()->GetBounds();
-		const double MeshMinZ = GetAxisValueRef(StaticMeshBounds.Origin, ForwardAxis) - GetAxisValueRef(StaticMeshBounds.BoxExtent, ForwardAxis);
-		const double MeshRangeZ = 2 * GetAxisValueRef(StaticMeshBounds.BoxExtent, ForwardAxis);
-		if (MeshRangeZ > UE_SMALL_NUMBER)
-		{
-			Alpha = UE::SplineMesh::RealToFloatChecked((DistanceAlong - MeshMinZ) / MeshRangeZ);
-		}
-	}
-	return Alpha;
+	return FSplineMeshSceneProxyDesc(this).ComputeRatioAlongSpline(GetStaticMesh(), DistanceAlong);
 }
 
 void USplineMeshComponent::ComputeVisualMeshSplineTRange(float& MinT, float& MaxT) const
 {
-    MinT = 0.0;
-    MaxT = 1.0;
-
-    const bool bHasCustomBoundary = !FMath::IsNearlyEqual(SplineBoundaryMin, SplineBoundaryMax);
-    if (bHasCustomBoundary)
-    {
-        const FBoxSphereBounds& MeshBounds = GetStaticMesh()->GetBounds();
-        // If there's a custom boundary, alter the min/max of the spline we need to evaluate
-        const float BoundsMin = UE::SplineMesh::RealToFloatChecked(GetAxisValueRef(MeshBounds.Origin - MeshBounds.BoxExtent, ForwardAxis));
-        const float BoundsMax = UE::SplineMesh::RealToFloatChecked(GetAxisValueRef(MeshBounds.Origin + MeshBounds.BoxExtent, ForwardAxis));
-        const float BoundsMinT = (BoundsMin - SplineBoundaryMin) / (SplineBoundaryMax - SplineBoundaryMin);
-        const float BoundsMaxT = (BoundsMax - SplineBoundaryMin) / (SplineBoundaryMax - SplineBoundaryMin);
-        // Disallow extrapolation beyond a certain value; enormous bounding boxes cause the render thread to crash
-        constexpr float MaxSplineExtrapolation = 4.0f;
-        MinT = FMath::Max(-MaxSplineExtrapolation, BoundsMinT);
-        MaxT = FMath::Min(BoundsMaxT, MaxSplineExtrapolation);
-    }
+	return FSplineMeshSceneProxyDesc(this).ComputeVisualMeshSplineTRange(GetStaticMesh(), MinT, MaxT);
 }
 
 FBox USplineMeshComponent::ComputeDistortedBounds(const FTransform& InLocalToWorld, const FBoxSphereBounds& InMeshBounds, const FBoxSphereBounds* InBoundsToDistort) const
 {
-	float MinT = 0.0f;
-	float MaxT = 1.0f;
-	ComputeVisualMeshSplineTRange(MinT, MaxT);
-	const FBoxSphereBounds& BoundsToDistort = InBoundsToDistort ? *InBoundsToDistort : InMeshBounds;
-
-
-	const FVector AxisMask = GetAxisMask(ForwardAxis);
-	const FVector FlattenedBoundsOrigin = BoundsToDistort.Origin * AxisMask;
-	const FVector FlattenedBoundsExtent = BoundsToDistort.BoxExtent * AxisMask;
-	const FBox FlattenedBounds = FBox(FlattenedBoundsOrigin - FlattenedBoundsExtent, FlattenedBoundsOrigin + FlattenedBoundsExtent);
-
-	FBox BoundingBox(ForceInit);
-	BoundingBox += FlattenedBounds.TransformBy(CalcSliceTransformAtSplineOffset(MinT, MinT, MaxT));
-	BoundingBox += FlattenedBounds.TransformBy(CalcSliceTransformAtSplineOffset(MaxT, MinT, MaxT));
-
-	// Work out coefficients of the cubic spline derivative equation dx/dt
-	const FVector A(6 * SplineParams.StartPos + 3 * SplineParams.StartTangent + 3 * SplineParams.EndTangent - 6 * SplineParams.EndPos);
-	const FVector B(-6 * SplineParams.StartPos - 4 * SplineParams.StartTangent - 2 * SplineParams.EndTangent + 6 * SplineParams.EndPos);
-	const FVector C(SplineParams.StartTangent);
-
-	auto AppendAxisExtrema = [&BoundingBox, &FlattenedBounds, MinT, MaxT, this](const double Discriminant, const double A, const double B)
-		{
-			// Negative discriminant means no solution; A == 0 implies coincident start/end points
-			if (Discriminant > 0 && !FMath::IsNearlyZero(A))
-			{
-				const double SqrtDiscriminant = FMath::Sqrt(Discriminant);
-				const double Denominator = 0.5 / A;
-				const double T0 = (-B + SqrtDiscriminant) * Denominator;
-				const double T1 = (-B - SqrtDiscriminant) * Denominator;
-
-				if (T0 >= MinT && T0 <= MaxT)
-				{
-					BoundingBox += FlattenedBounds.TransformBy(CalcSliceTransformAtSplineOffset(UE::SplineMesh::RealToFloatChecked(T0), MinT, MaxT));
-				}
-
-				if (T1 >= MinT && T1 <= MaxT)
-				{
-					BoundingBox += FlattenedBounds.TransformBy(CalcSliceTransformAtSplineOffset(UE::SplineMesh::RealToFloatChecked(T1), MinT, MaxT));
-				}
-			}
-		};
-
-	// Minima/maxima happen where dx/dt == 0, calculate t values
-	const FVector Discriminant = B * B - 4 * A * C;
-
-	// Work out minima/maxima component-by-component.
-	AppendAxisExtrema(Discriminant.X, A.X, B.X);
-	AppendAxisExtrema(Discriminant.Y, A.Y, B.Y);
-	AppendAxisExtrema(Discriminant.Z, A.Z, B.Z);
-
-	// Applying extrapolation if bounds to apply on spline are different than the mesh bounds used
-	// to define the spline range [0,1]
-	if (InBoundsToDistort != nullptr && InBoundsToDistort != &InMeshBounds)
-	{
-		const double BoundsMin = GetAxisValueRef(BoundsToDistort.Origin - BoundsToDistort.BoxExtent, ForwardAxis);
-		const double BoundsMax = GetAxisValueRef(BoundsToDistort.Origin + BoundsToDistort.BoxExtent, ForwardAxis);
-
-		float Alpha = ComputeRatioAlongSpline(UE::SplineMesh::RealToFloatChecked(BoundsMin));
-		if (Alpha < MinT)
-		{
-			BoundingBox += FlattenedBounds.TransformBy(CalcSliceTransformAtSplineOffset(Alpha, MinT, MaxT));
-		}
-
-		Alpha = ComputeRatioAlongSpline(UE::SplineMesh::RealToFloatChecked(BoundsMax));
-		if (Alpha > MaxT)
-		{
-			BoundingBox += FlattenedBounds.TransformBy(CalcSliceTransformAtSplineOffset(Alpha, MinT, MaxT));
-		}
-	}
-
-	return BoundingBox.TransformBy(InLocalToWorld);
+	return FSplineMeshSceneProxyDesc(this).ComputeDistortedBounds(GetStaticMesh(), InLocalToWorld, InMeshBounds, InBoundsToDistort);
 }
 
 FTransform USplineMeshComponent::GetSocketTransform(FName InSocketName, ERelativeTransformSpace TransformSpace) const
@@ -1101,90 +800,16 @@ FTransform USplineMeshComponent::GetSocketTransform(FName InSocketName, ERelativ
 	return Super::GetSocketTransform(InSocketName, TransformSpace);
 }
 
-
 FTransform USplineMeshComponent::CalcSliceTransform(const float DistanceAlong) const
 {
-	const float Alpha = ComputeRatioAlongSpline(DistanceAlong);
-
-	float MinT = 0.f;
-	float MaxT = 1.f;
-	ComputeVisualMeshSplineTRange(MinT, MaxT);
-
-	return CalcSliceTransformAtSplineOffset(Alpha, MinT, MaxT);
+	return FSplineMeshSceneProxyDesc(this).CalcSliceTransform(GetStaticMesh(), DistanceAlong);
 }
 
 FTransform USplineMeshComponent::CalcSliceTransformAtSplineOffset(const float Alpha, const float MinT, const float MaxT) const
 {
-	// Apply hermite interp to Alpha if desired
-	const float HermiteAlpha = bSmoothInterpRollScale ? SmoothStep(0.0, 1.0, Alpha) : Alpha;
-
-	// Then find the point and direction of the spline at this point along
-	FVector3f SplinePos;
-	FVector3f SplineDir;
-
-	// Use linear extrapolation
-	if (Alpha < MinT)
-	{
-		const FVector3f StartTangent(SplineEvalTangent(SplineParams, MinT));
-		SplinePos = SplineEvalPos(SplineParams, MinT) + (StartTangent * (Alpha - MinT));
-		SplineDir = StartTangent.GetSafeNormal();
-	}
-	else if (Alpha > MaxT)
-	{
-		const FVector3f EndTangent(SplineEvalTangent(SplineParams, MaxT));
-		SplinePos = SplineEvalPos(SplineParams, MaxT) + (EndTangent * (Alpha - MaxT));
-		SplineDir = EndTangent.GetSafeNormal();
-	}
-	else
-	{
-		SplinePos = SplineEvalPos(SplineParams, Alpha);
-		SplineDir = SplineEvalDir(SplineParams, Alpha);
-	}
-
-	// Find base frenet frame
-	const FVector3f BaseXVec = (FVector3f(SplineUpDir) ^ SplineDir).GetSafeNormal();
-	const FVector3f BaseYVec = (FVector3f(SplineDir) ^ BaseXVec).GetSafeNormal();
-
-	// Offset the spline by the desired amount
-	const FVector2f SliceOffset = FMath::Lerp(FVector2f(SplineParams.StartOffset), FVector2f(SplineParams.EndOffset), HermiteAlpha);
-	SplinePos += SliceOffset.X * BaseXVec;
-	SplinePos += SliceOffset.Y * BaseYVec;
-
-	// Apply roll to frame around spline
-	const float UseRoll = FMath::Lerp(SplineParams.StartRoll, SplineParams.EndRoll, HermiteAlpha);
-	const float CosAng = FMath::Cos(UseRoll);
-	const float SinAng = FMath::Sin(UseRoll);
-	const FVector3f XVec = (CosAng * BaseXVec) - (SinAng * BaseYVec);
-	const FVector3f YVec = (CosAng * BaseYVec) + (SinAng * BaseXVec);
-
-	// Find scale at this point along spline
-	const FVector2f UseScale = FMath::Lerp(FVector2f(SplineParams.StartScale), FVector2f(SplineParams.EndScale), HermiteAlpha);
-
-	// Build overall transform
-	FTransform SliceTransform;
-	switch (ForwardAxis)
-	{
-	case ESplineMeshAxis::X:
-		SliceTransform = FTransform(FVector(SplineDir), FVector(XVec), FVector(YVec), FVector(SplinePos));
-		SliceTransform.SetScale3D(FVector(1, UseScale.X, UseScale.Y));
-		break;
-	case ESplineMeshAxis::Y:
-		SliceTransform = FTransform(FVector(YVec), FVector(SplineDir), FVector(XVec), FVector(SplinePos));
-		SliceTransform.SetScale3D(FVector(UseScale.Y, 1, UseScale.X));
-		break;
-	case ESplineMeshAxis::Z:
-		SliceTransform = FTransform(FVector(XVec), FVector(YVec), FVector(SplineDir), FVector(SplinePos));
-		SliceTransform.SetScale3D(FVector(UseScale.X, UseScale.Y, 1));
-		break;
-	default:
-		check(0);
-		break;
-	}
-
-	return SliceTransform;
+	return FSplineMeshSceneProxyDesc(this).CalcSliceTransformAtSplineOffset(Alpha, MinT, MaxT);
 }
 
-
 bool USplineMeshComponent::GetPhysicsTriMeshData(struct FTriMeshCollisionData* CollisionData, bool InUseAllTriData)
 {
 	if (GetStaticMesh())
@@ -1233,6 +858,12 @@ FBox USplineMeshComponent::GetNavigationBounds() const
 }
 
 void USplineMeshComponent::GetMeshId(FString& OutMeshId)
+{
+	//call the const version
+	return const_cast<const USplineMeshComponent*>(this)->GetMeshId(OutMeshId);
+}
+
+void USplineMeshComponent::GetMeshId(FString& OutMeshId) const
 {
 	// First get the base mesh id from the static mesh
 	if (GetStaticMesh())
@@ -1249,26 +880,32 @@ void USplineMeshComponent::GetMeshId(FString& OutMeshId)
 
 	if (BodySetup != nullptr && BodySetup->BodySetupGuid == CachedMeshBodySetupGuid)
 	{
+		FSplineMeshParams MutableSplineParams = SplineParams;
+		FVector MutableSplineUpDir = SplineUpDir;
+		float MutableSplineBoundaryMin = SplineBoundaryMin;
+		float MutableSplineBoundaryMax = SplineBoundaryMax;
+		TEnumAsByte<ESplineMeshAxis::Type> MutableForwardAxis = ForwardAxis;
+		bool bMutableSmoothInterp = bSmoothInterpRollScale; 
+
 		TArray<uint8> TempBytes;
 		TempBytes.Reserve(256);
 
 		FMemoryWriter Ar(TempBytes);
-		Ar << SplineParams.StartPos;
-		Ar << SplineParams.StartTangent;
-		Ar << SplineParams.StartScale;
-		Ar << SplineParams.StartRoll;
-		Ar << SplineParams.StartOffset;
-		Ar << SplineParams.EndPos;
-		Ar << SplineParams.EndTangent;
-		Ar << SplineParams.EndScale;
-		Ar << SplineParams.EndRoll;
-		Ar << SplineParams.EndOffset;
-		Ar << SplineUpDir;
-		bool bSmoothInterp = bSmoothInterpRollScale;
-		Ar << bSmoothInterp; // can't write a bitfield member into an archive
-		Ar << ForwardAxis;
-		Ar << SplineBoundaryMin;
-		Ar << SplineBoundaryMax;
+		Ar << MutableSplineParams.StartPos;
+		Ar << MutableSplineParams.StartTangent;
+		Ar << MutableSplineParams.StartScale;
+		Ar << MutableSplineParams.StartRoll;
+		Ar << MutableSplineParams.StartOffset;
+		Ar << MutableSplineParams.EndPos;
+		Ar << MutableSplineParams.EndTangent;
+		Ar << MutableSplineParams.EndScale;
+		Ar << MutableSplineParams.EndRoll;
+		Ar << MutableSplineParams.EndOffset;
+		Ar << MutableSplineUpDir;
+		Ar << bMutableSmoothInterp;
+		Ar << MutableForwardAxis;
+		Ar << MutableSplineBoundaryMin;
+		Ar << MutableSplineBoundaryMax;
 
 		// Now convert the raw bytes to a string.
 		const uint8* SettingsAsBytes = TempBytes.GetData();
@@ -1293,6 +930,12 @@ void USplineMeshComponent::OnCreatePhysicsState()
 }
 
 UBodySetup* USplineMeshComponent::GetBodySetup()
+{
+	//call the const version
+	return const_cast<const USplineMeshComponent*>(this)->GetBodySetup();
+}
+
+UBodySetup* USplineMeshComponent::GetBodySetup() const
 {
 	// Don't return a body setup that has no collision, it means we are interactively moving the spline and don't want to build collision.
 	// Instead we explicitly build collision with USplineMeshComponent::RecreateCollision()
diff --git a/Engine/Source/Runtime/Engine/Private/SplineMeshSceneProxy.cpp b/Engine/Source/Runtime/Engine/Private/SplineMeshSceneProxy.cpp
--- a/Engine/Source/Runtime/Engine/Private/SplineMeshSceneProxy.cpp
+++ b/Engine/Source/Runtime/Engine/Private/SplineMeshSceneProxy.cpp
@@ -10,6 +10,7 @@
 #include "StaticMeshResources.h"
 #include "RayTracingInstance.h"
 #include "SceneInterface.h"
+#include "SplineMeshSceneProxyDesc.h"
 
 IMPLEMENT_TYPE_LAYOUT(FSplineMeshVertexFactoryShaderParameters);
 
@@ -44,7 +45,11 @@ void FSplineMeshVertexFactory::GetPSOPrecacheVertexFetchElements(EVertexInputStr
 }
 
 FSplineMeshSceneProxy::FSplineMeshSceneProxy(USplineMeshComponent* InComponent) :
-	FStaticMeshSceneProxy(InComponent, false)
+	FSplineMeshSceneProxy(FStaticMeshSceneProxyDesc(InComponent), FSplineMeshSceneProxyDesc(InComponent))
+{}
+
+FSplineMeshSceneProxy::FSplineMeshSceneProxy(const FStaticMeshSceneProxyDesc& InMeshDesc, const FSplineMeshSceneProxyDesc& InSplineDesc) :
+	FStaticMeshSceneProxy(InMeshDesc, false)
 {
 	bSupportsDistanceFieldRepresentation = false;
 
@@ -72,7 +77,7 @@ FSplineMeshSceneProxy::FSplineMeshSceneProxy(USplineMeshComponent* InComponent)
 	}
 
 	// Copy spline params from component
-	SplineParams = InComponent->CalculateShaderParams();
+	SplineParams = InSplineDesc.CalculateShaderParams(InMeshDesc.GetStaticMesh());
 
 	// If we're using GPU Scene, we place the spline mesh parameters in the instance data buffer, with the
 	// exception of mobile platforms that are unable to pull this data from the structured buffer in the VS
@@ -85,10 +90,10 @@ FSplineMeshSceneProxy::FSplineMeshSceneProxy(USplineMeshComponent* InComponent)
 
 	for (int32 LODIndex = 0; LODIndex < LODs.Num(); LODIndex++)
 	{
-		InComponent->InitVertexFactory(LODIndex, nullptr); // we always need this one for shadows etc
-		if (InComponent->LODData.IsValidIndex(LODIndex) && InComponent->LODData[LODIndex].OverrideVertexColors)
+		FSplineMeshSceneProxyDesc::InitVertexFactory(InMeshDesc.GetStaticMesh(), FeatureLevel, LODIndex, nullptr); // we always need this one for shadows etc
+		if (InMeshDesc.LODData.IsValidIndex(LODIndex) && InMeshDesc.LODData[LODIndex].OverrideVertexColors)
 		{
-			InComponent->InitVertexFactory(LODIndex, InComponent->LODData[LODIndex].OverrideVertexColors);
+			FSplineMeshSceneProxyDesc::InitVertexFactory(InMeshDesc.GetStaticMesh(), FeatureLevel, LODIndex, InMeshDesc.LODData[LODIndex].OverrideVertexColors);
 		}
 	}
 }
@@ -273,7 +278,11 @@ bool FSplineMeshSceneInstanceDataBuffers::Update(const FSplineMeshShaderParams&
 }
 
 FNaniteSplineMeshSceneProxy::FNaniteSplineMeshSceneProxy(const Nanite::FMaterialAudit& NaniteMaterials, USplineMeshComponent* InComponent) :
-	Nanite::FSceneProxy(NaniteMaterials, InComponent)
+	FNaniteSplineMeshSceneProxy(NaniteMaterials, FStaticMeshSceneProxyDesc(InComponent), FSplineMeshSceneProxyDesc(InComponent))
+{}
+
+FNaniteSplineMeshSceneProxy::FNaniteSplineMeshSceneProxy(const Nanite::FMaterialAudit& NaniteMaterials, const FStaticMeshSceneProxyDesc& InMeshDesc, const FSplineMeshSceneProxyDesc& InSplineDesc) :
+	Nanite::FSceneProxy(NaniteMaterials, InMeshDesc)
 {
 	bSupportsDistanceFieldRepresentation = false;
 	
@@ -302,7 +311,7 @@ FNaniteSplineMeshSceneProxy::FNaniteSplineMeshSceneProxy(const Nanite::FMaterial
 	}
 
 	// Copy spline params from component
-	SplineParams = InComponent->CalculateShaderParams();
+	SplineParams = InSplineDesc.CalculateShaderParams(InMeshDesc.GetStaticMesh());
 	SplineMeshInstanceData.Setup(SplineParams);
 	SetupInstanceSceneDataBuffers(&SplineMeshInstanceData);
 
@@ -312,12 +321,13 @@ FNaniteSplineMeshSceneProxy::FNaniteSplineMeshSceneProxy(const Nanite::FMaterial
 	// We only need to init vertex factories for Nanite spline meshes if they can be ray traced
 	if (IsRayTracingAllowed())
 	{
+		const ERHIFeatureLevel::Type FeatureLevel = GetScene().GetFeatureLevel();
 		for (int32 LODIndex = 0; LODIndex < RenderData->LODResources.Num(); LODIndex++)
 		{
-			InComponent->InitVertexFactory(LODIndex, nullptr); // we always need this one for shadows etc
-			if (InComponent->LODData.IsValidIndex(LODIndex) && InComponent->LODData[LODIndex].OverrideVertexColors)
+			FSplineMeshSceneProxyDesc::InitVertexFactory(InMeshDesc.GetStaticMesh(), FeatureLevel, LODIndex, nullptr); // we always need this one for shadows etc
+			if (InMeshDesc.LODData.IsValidIndex(LODIndex) && InMeshDesc.LODData[LODIndex].OverrideVertexColors)
 			{
-				InComponent->InitVertexFactory(LODIndex, InComponent->LODData[LODIndex].OverrideVertexColors);
+				FSplineMeshSceneProxyDesc::InitVertexFactory(InMeshDesc.GetStaticMesh(), FeatureLevel, LODIndex, InMeshDesc.LODData[LODIndex].OverrideVertexColors);
 			}
 		}
 	}
diff --git a/Engine/Source/Runtime/Engine/Private/SplineMeshSceneProxyDesc.cpp b/Engine/Source/Runtime/Engine/Private/SplineMeshSceneProxyDesc.cpp
new file mode 100644
--- /dev/null
+++ b/Engine/Source/Runtime/Engine/Private/SplineMeshSceneProxyDesc.cpp
@@ -0,0 +1,442 @@
+// Copyright Epic Games, Inc. All Rights Reserved.
+
+#include "SplineMeshSceneProxyDesc.h"
+
+#include "SplineMeshSceneProxy.h"
+#include "StaticMeshResources.h"
+#include "Components/SplineMeshComponent.h"
+
+static FVector3f SplineEvalPos(const FVector3f& StartPos, const FVector3f& StartTangent, const FVector3f& EndPos, const FVector3f& EndTangent, float A)
+{
+	const float A2 = A * A;
+	const float A3 = A2 * A;
+
+	return (((2 * A3) - (3 * A2) + 1) * StartPos) + ((A3 - (2 * A2) + A) * StartTangent) + ((A3 - A2) * EndTangent) + (((-2 * A3) + (3 * A2)) * EndPos);
+}
+
+static FVector3f SplineEvalPos(const FSplineMeshParams& Params, float A)
+{
+	// TODO: these don't need to be doubles!
+	const FVector3f StartPos = FVector3f(Params.StartPos);
+	const FVector3f StartTangent = FVector3f(Params.StartTangent);
+	const FVector3f EndPos = FVector3f(Params.EndPos);
+	const FVector3f EndTangent = FVector3f(Params.EndTangent);
+
+	return SplineEvalPos(StartPos, StartTangent, EndPos, EndTangent, A);
+}
+
+static FVector3f SplineEvalTangent(const FVector3f& StartPos, const FVector3f& StartTangent, const FVector3f& EndPos, const FVector3f& EndTangent, const float A)
+{
+	const FVector3f C = (6 * StartPos) + (3 * StartTangent) + (3 * EndTangent) - (6 * EndPos);
+	const FVector3f D = (-6 * StartPos) - (4 * StartTangent) - (2 * EndTangent) + (6 * EndPos);
+	const FVector3f E = StartTangent;
+
+	const float A2 = A * A;
+
+	return (C * A2) + (D * A) + E;
+}
+
+static FVector3f SplineEvalTangent(const FSplineMeshParams& Params, const float A)
+{
+	// TODO: these don't need to be doubles!
+	const FVector3f StartPos = FVector3f(Params.StartPos);
+	const FVector3f StartTangent = FVector3f(Params.StartTangent);
+	const FVector3f EndPos = FVector3f(Params.EndPos);
+	const FVector3f EndTangent = FVector3f(Params.EndTangent);
+
+	return SplineEvalTangent(StartPos, StartTangent, EndPos, EndTangent, A);
+}
+
+static FVector3f SplineEvalDir(const FSplineMeshParams& Params, const float A)
+{
+	return SplineEvalTangent(Params, A).GetSafeNormal();
+}
+
+
+FSplineMeshSceneProxyDesc::FSplineMeshSceneProxyDesc(const USplineMeshComponent* InComponent)
+{
+	InitializeFrom(InComponent);
+}
+
+void FSplineMeshSceneProxyDesc::InitializeFrom(const USplineMeshComponent* InComponent)
+{
+	SplineParams = InComponent->SplineParams;
+	SplineUpDir = InComponent->SplineUpDir;
+	SplineBoundaryMin = InComponent->SplineBoundaryMin;
+	SplineBoundaryMax = InComponent->SplineBoundaryMax;
+	bSmoothInterpRollScale = InComponent->bSmoothInterpRollScale;
+	ForwardAxis = InComponent->ForwardAxis;
+}
+
+FSplineMeshShaderParams FSplineMeshSceneProxyDesc::CalculateShaderParams(UStaticMesh* StaticMesh) const
+{
+	FSplineMeshShaderParams Output;
+
+	Output.StartPos 				= FVector3f(SplineParams.StartPos);
+	Output.EndPos 					= FVector3f(SplineParams.EndPos);
+	Output.StartTangent 			= FVector3f(SplineParams.StartTangent);
+	Output.EndTangent 				= FVector3f(SplineParams.EndTangent);
+	Output.StartScale 				= FVector2f(SplineParams.StartScale);
+	Output.EndScale 				= FVector2f(SplineParams.EndScale);
+	Output.StartOffset 				= FVector2f(SplineParams.StartOffset);
+	Output.EndOffset 				= FVector2f(SplineParams.EndOffset);
+	Output.StartRoll 				= SplineParams.StartRoll;
+	Output.EndRoll 					= SplineParams.EndRoll;
+	Output.NaniteClusterBoundsScale	= SplineParams.NaniteClusterBoundsScale;
+	Output.bSmoothInterpRollScale 	= bSmoothInterpRollScale;
+	Output.SplineUpDir 				= FVector3f(SplineUpDir);
+	Output.TextureCoord 			= FUintVector2(INDEX_NONE, INDEX_NONE); // either unused or assigned later
+
+	const uint32 MeshXAxis = (ForwardAxis + 1) % 3;
+	const uint32 MeshYAxis = (ForwardAxis + 2) % 3;
+	Output.MeshDir = Output.MeshX = Output.MeshY = FVector3f::ZeroVector;
+	Output.MeshDir[ForwardAxis] = 1.0f;
+	Output.MeshX[MeshXAxis] = 1.0f;
+	Output.MeshY[MeshYAxis] = 1.0f;
+
+	Output.MeshZScale = 1.0f;
+	Output.MeshZOffset = 0.0f;
+
+	if (StaticMesh)
+	{
+		const FBoxSphereBounds StaticMeshBounds = StaticMesh->GetBounds();
+		const float BoundsXYRadius = FVector3f(StaticMeshBounds.BoxExtent).Dot((Output.MeshX + Output.MeshY).GetUnsafeNormal());
+
+		const float MeshMinZ = UE::SplineMesh::RealToFloatChecked(USplineMeshComponent::GetAxisValueRef(StaticMeshBounds.Origin - StaticMeshBounds.BoxExtent, ForwardAxis));
+		const float MeshZLen = UE::SplineMesh::RealToFloatChecked(2 * USplineMeshComponent::GetAxisValueRef(StaticMeshBounds.BoxExtent, ForwardAxis));
+		const float InvMeshZLen = (MeshZLen <= 0.0f) ? 1.0f : 1.0f / MeshZLen;
+		constexpr float MeshTexelLen = float(SPLINE_MESH_TEXEL_WIDTH - 1);
+
+		if (FMath::IsNearlyEqual(SplineBoundaryMin, SplineBoundaryMax))
+		{
+			Output.MeshZScale = InvMeshZLen;
+			Output.MeshZOffset = -MeshMinZ * InvMeshZLen;
+			Output.SplineDistToTexelScale = MeshTexelLen;
+			Output.SplineDistToTexelOffset = 0.0f;
+		}
+		else
+		{
+			const float BoundaryLen = SplineBoundaryMax - SplineBoundaryMin;
+			const float InvBoundaryLen = 1.0f / BoundaryLen;
+
+			Output.MeshZScale = InvBoundaryLen;
+			Output.MeshZOffset = -SplineBoundaryMin * InvBoundaryLen;
+			Output.SplineDistToTexelScale = BoundaryLen * InvMeshZLen * MeshTexelLen;
+			Output.SplineDistToTexelOffset = (SplineBoundaryMin - MeshMinZ) * InvMeshZLen * MeshTexelLen;
+		}
+
+		// Iteratively solve for an approximation of spline length
+		float SplineLength = 0.0f;
+		{
+			static const uint32 NumIterations = 63; // 64 sampled points
+			static const float IterStep = 1.0f / float(NumIterations);
+			float A = 0.0f;
+			FVector3f PrevPoint = SplineEvalPos(SplineParams, A);
+			for (uint32 i = 0; i < NumIterations; ++i)
+			{
+				FVector3f Point = SplineEvalPos(SplineParams, A);
+				SplineLength += (Point - PrevPoint).Length();
+				PrevPoint = Point;
+				A += IterStep;
+			}
+		}
+
+		// Calculate an approximation of how much the mesh gets scaled in each local axis as a result of spline
+		// deformation and take the smallest of the axes. This is important for LOD selection of Nanite spline
+		// meshes.
+		{
+			// Estimate length added due to twisting as well
+			const float XYRadius = BoundsXYRadius * FMath::Max(Output.StartScale.GetAbsMax(), Output.EndScale.GetAbsMax());
+			const float TwistRadians = FMath::Abs(Output.StartRoll - Output.EndRoll);
+			SplineLength += TwistRadians * XYRadius;
+
+			// Take the mid-point scale in X/Y to balance out LOD selection in case either of them are extreme.
+			auto AvgAbs = [](float A, float B) { return (FMath::Abs(A) + FMath::Abs(B)) * 0.5f; };
+			const FVector3f DeformScale = FVector3f(
+				SplineLength * Output.MeshZScale,
+				AvgAbs(Output.StartScale.X, Output.EndScale.X),
+				AvgAbs(Output.StartScale.Y, Output.EndScale.Y)
+			);
+			
+			Output.MeshDeformScaleMinMax = FVector2f(DeformScale.GetMin(), DeformScale.GetMax());
+		}
+	}
+	return Output;
+}
+
+extern void InitSplineMeshVertexFactoryComponents(
+	const FStaticMeshVertexBuffers& VertexBuffers, 
+	const FSplineMeshVertexFactory* VertexFactory, 
+	int32 LightMapCoordinateIndex, 
+	bool bOverrideColorVertexBuffer,
+	FLocalVertexFactory::FDataType& OutData);
+
+void FSplineMeshSceneProxyDesc::InitVertexFactory(UStaticMesh* Mesh, const ERHIFeatureLevel::Type FeatureLevel, int32 InLODIndex, FColorVertexBuffer* InOverrideColorVertexBuffer)
+{	
+	if (Mesh == nullptr)
+	{
+		return;
+	}
+
+	FStaticMeshLODResources& LODRenderData = Mesh->GetRenderData()->LODResources[InLODIndex];
+	FStaticMeshVertexFactories& VertexFactories = Mesh->GetRenderData()->LODVertexFactories[InLODIndex];
+
+	// Skip LODs that have their render data stripped (eg. platform MinLod settings)
+	if (LODRenderData.VertexBuffers.StaticMeshVertexBuffer.GetNumVertices() == 0)
+	{
+		return;
+	}
+
+	bool bOverrideColorVertexBuffer = !!InOverrideColorVertexBuffer;
+	
+	if ((VertexFactories.SplineVertexFactory && !bOverrideColorVertexBuffer) || (VertexFactories.SplineVertexFactoryOverrideColorVertexBuffer && bOverrideColorVertexBuffer))
+	{
+		// we already have it
+		return;
+	}
+
+	FSplineMeshVertexFactory* VertexFactory = new FSplineMeshVertexFactory(FeatureLevel);
+	if (bOverrideColorVertexBuffer)
+	{
+		VertexFactories.SplineVertexFactoryOverrideColorVertexBuffer = VertexFactory;
+	}
+	else
+	{
+		VertexFactories.SplineVertexFactory = VertexFactory;
+	}
+
+	int32 LightMapCoordinateIndex = Mesh->GetLightMapCoordinateIndex();
+	// Initialize the static mesh's vertex factory.
+	ENQUEUE_RENDER_COMMAND(InitSplineMeshVertexFactory)(
+		[VertexFactory, &LODRenderData, bOverrideColorVertexBuffer, LightMapCoordinateIndex](FRHICommandListBase& RHICmdList)
+	{
+		FLocalVertexFactory::FDataType Data;
+		InitSplineMeshVertexFactoryComponents(LODRenderData.VertexBuffers, VertexFactory, LightMapCoordinateIndex, bOverrideColorVertexBuffer, Data);
+		VertexFactory->SetData(RHICmdList, Data);
+		VertexFactory->InitResource(RHICmdList);
+	});
+}
+
+FBox FSplineMeshSceneProxyDesc::ComputeDistortedBounds(UStaticMesh* StaticMesh, const FTransform& InLocalToWorld, const FBoxSphereBounds& InMeshBounds, const FBoxSphereBounds* InBoundsToDistort) const
+{
+	float MinT = 0.0f;
+	float MaxT = 1.0f;
+	ComputeVisualMeshSplineTRange(StaticMesh, MinT, MaxT);
+	const FBoxSphereBounds& BoundsToDistort = InBoundsToDistort ? *InBoundsToDistort : InMeshBounds;
+
+	auto ForwardAxisType = static_cast<ESplineMeshAxis::Type>(ForwardAxis);
+	const FVector AxisMask = USplineMeshComponent::GetAxisMask(ForwardAxisType);
+	const FVector FlattenedBoundsOrigin = BoundsToDistort.Origin * AxisMask;
+	const FVector FlattenedBoundsExtent = BoundsToDistort.BoxExtent * AxisMask;
+	const FBox FlattenedBounds = FBox(FlattenedBoundsOrigin - FlattenedBoundsExtent, FlattenedBoundsOrigin + FlattenedBoundsExtent);
+
+	FBox BoundingBox(ForceInit);
+	BoundingBox += FlattenedBounds.TransformBy(CalcSliceTransformAtSplineOffset(MinT, MinT, MaxT));
+	BoundingBox += FlattenedBounds.TransformBy(CalcSliceTransformAtSplineOffset(MaxT, MinT, MaxT));
+
+	// Work out coefficients of the cubic spline derivative equation dx/dt
+	const FVector A(6 * SplineParams.StartPos + 3 * SplineParams.StartTangent + 3 * SplineParams.EndTangent - 6 * SplineParams.EndPos);
+	const FVector B(-6 * SplineParams.StartPos - 4 * SplineParams.StartTangent - 2 * SplineParams.EndTangent + 6 * SplineParams.EndPos);
+	const FVector C(SplineParams.StartTangent);
+
+	auto AppendAxisExtrema = [this, &StaticMesh,&BoundingBox, &FlattenedBounds, MinT, MaxT](const double Discriminant, const double A, const double B)
+		{
+			// Negative discriminant means no solution; A == 0 implies coincident start/end points
+			if (Discriminant > 0 && !FMath::IsNearlyZero(A))
+			{
+				const double SqrtDiscriminant = FMath::Sqrt(Discriminant);
+				const double Denominator = 0.5 / A;
+				const double T0 = (-B + SqrtDiscriminant) * Denominator;
+				const double T1 = (-B - SqrtDiscriminant) * Denominator;
+
+				if (T0 >= MinT && T0 <= MaxT)
+				{
+					BoundingBox += FlattenedBounds.TransformBy(CalcSliceTransformAtSplineOffset(UE::SplineMesh::RealToFloatChecked(T0), MinT, MaxT));
+				}
+
+				if (T1 >= MinT && T1 <= MaxT)
+				{
+					BoundingBox += FlattenedBounds.TransformBy(CalcSliceTransformAtSplineOffset(UE::SplineMesh::RealToFloatChecked(T1), MinT, MaxT));
+				}
+			}
+		};
+
+
+	// Minima/maxima happen where dx/dt == 0, calculate t values
+	const FVector Discriminant = B * B - 4 * A * C;
+
+	// Work out minima/maxima component-by-component.
+	AppendAxisExtrema(Discriminant.X, A.X, B.X);
+	AppendAxisExtrema(Discriminant.Y, A.Y, B.Y);
+	AppendAxisExtrema(Discriminant.Z, A.Z, B.Z);
+
+	// Applying extrapolation if bounds to apply on spline are different than the mesh bounds used
+	// to define the spline range [0,1]
+	if (InBoundsToDistort != nullptr && InBoundsToDistort != &InMeshBounds)
+	{
+		const double BoundsMin = USplineMeshComponent::GetAxisValueRef(BoundsToDistort.Origin - BoundsToDistort.BoxExtent, ForwardAxisType);
+		const double BoundsMax = USplineMeshComponent::GetAxisValueRef(BoundsToDistort.Origin + BoundsToDistort.BoxExtent, ForwardAxisType);
+
+		float Alpha = ComputeRatioAlongSpline(StaticMesh, UE::SplineMesh::RealToFloatChecked(BoundsMin));
+		if (Alpha < MinT)
+		{
+			BoundingBox += FlattenedBounds.TransformBy(CalcSliceTransformAtSplineOffset(Alpha, MinT, MaxT));
+		}
+
+		Alpha = ComputeRatioAlongSpline(StaticMesh, UE::SplineMesh::RealToFloatChecked(BoundsMax));
+		if (Alpha > MaxT)
+		{
+			BoundingBox += FlattenedBounds.TransformBy(CalcSliceTransformAtSplineOffset(Alpha, MinT, MaxT));
+		}
+	}
+
+	return BoundingBox.TransformBy(InLocalToWorld);
+}
+
+FTransform FSplineMeshSceneProxyDesc::CalcSliceTransform(UStaticMesh* StaticMesh, const float DistanceAlong) const
+{
+	const float Alpha = ComputeRatioAlongSpline(StaticMesh, DistanceAlong);
+	
+
+	float MinT = 0.f;
+	float MaxT = 1.f;
+	ComputeVisualMeshSplineTRange(StaticMesh, MinT, MaxT);
+
+	return CalcSliceTransformAtSplineOffset(Alpha, MinT, MaxT);
+}
+
+/**
+* Functions used for transforming a static mesh component based on a spline.
+* This needs to be updated if the spline functionality changes!
+*/
+static float SmoothStep(float A, float B, float X)
+{
+	if (X < A)
+	{
+		return 0.0f;
+	}
+	else if (X >= B)
+	{
+		return 1.0f;
+	}
+	const float InterpFraction = (X - A) / (B - A);
+	return InterpFraction * InterpFraction * (3.0f - 2.0f * InterpFraction);
+}
+
+FTransform FSplineMeshSceneProxyDesc::CalcSliceTransformAtSplineOffset(const float Alpha, const float MinT, const float MaxT) const
+{
+	// Apply hermite interp to Alpha if desired
+	const float HermiteAlpha = bSmoothInterpRollScale ? SmoothStep(0.0, 1.0, Alpha) : Alpha;
+
+
+	// Then find the point and direction of the spline at this point along
+	FVector3f SplinePos;
+	FVector3f SplineDir;
+
+	// Use linear extrapolation
+	if (Alpha < MinT)
+	{
+		const FVector3f StartTangent(SplineEvalTangent(SplineParams, MinT));
+		SplinePos = SplineEvalPos(SplineParams, MinT) + (StartTangent * (Alpha - MinT));
+		SplineDir = StartTangent.GetSafeNormal();
+	}
+	else if (Alpha > MaxT)
+	{
+		const FVector3f EndTangent(SplineEvalTangent(SplineParams, MaxT));
+		SplinePos = SplineEvalPos(SplineParams, MaxT) + (EndTangent * (Alpha - MaxT));
+		SplineDir = EndTangent.GetSafeNormal();
+	}
+	else
+	{
+		SplinePos = SplineEvalPos(SplineParams, Alpha);
+		SplineDir = SplineEvalDir(SplineParams, Alpha);
+	}
+
+	// Find base frenet frame
+	const FVector3f BaseXVec = (FVector3f(SplineUpDir) ^ SplineDir).GetSafeNormal();
+	const FVector3f BaseYVec = (FVector3f(SplineDir) ^ BaseXVec).GetSafeNormal();
+
+	// Offset the spline by the desired amount
+	const FVector2f SliceOffset = FMath::Lerp(FVector2f(SplineParams.StartOffset), FVector2f(SplineParams.EndOffset), HermiteAlpha);
+	SplinePos += SliceOffset.X * BaseXVec;
+	SplinePos += SliceOffset.Y * BaseYVec;
+
+	// Apply roll to frame around spline
+	const float UseRoll = FMath::Lerp(SplineParams.StartRoll, SplineParams.EndRoll, HermiteAlpha);
+	const float CosAng = FMath::Cos(UseRoll);
+	const float SinAng = FMath::Sin(UseRoll);
+	const FVector3f XVec = (CosAng * BaseXVec) - (SinAng * BaseYVec);
+	const FVector3f YVec = (CosAng * BaseYVec) + (SinAng * BaseXVec);
+
+	// Find scale at this point along spline
+	const FVector2f UseScale = FMath::Lerp(FVector2f(SplineParams.StartScale), FVector2f(SplineParams.EndScale), HermiteAlpha);
+
+	// Build overall transform
+	FTransform SliceTransform;
+	
+	switch (static_cast<ESplineMeshAxis::Type>(ForwardAxis))
+	{
+	case ESplineMeshAxis::X:
+		SliceTransform = FTransform(FVector(SplineDir), FVector(XVec), FVector(YVec), FVector(SplinePos));
+		SliceTransform.SetScale3D(FVector(1, UseScale.X, UseScale.Y));
+		break;
+	case ESplineMeshAxis::Y:
+		SliceTransform = FTransform(FVector(YVec), FVector(SplineDir), FVector(XVec), FVector(SplinePos));
+		SliceTransform.SetScale3D(FVector(UseScale.Y, 1, UseScale.X));
+		break;
+	case ESplineMeshAxis::Z:
+		SliceTransform = FTransform(FVector(XVec), FVector(YVec), FVector(SplineDir), FVector(SplinePos));
+		SliceTransform.SetScale3D(FVector(UseScale.X, UseScale.Y, 1));
+		break;
+	default:
+		check(0);
+		break;
+	}
+
+	return SliceTransform;
+}
+
+float FSplineMeshSceneProxyDesc::ComputeRatioAlongSpline(UStaticMesh* StaticMesh, float DistanceAlong) const
+{
+	// Find how far 'along' mesh (or custom boundaries) we are
+	float Alpha = 0.f;
+
+	const bool bHasCustomBoundary = !FMath::IsNearlyEqual(SplineBoundaryMin, SplineBoundaryMax);
+	if (bHasCustomBoundary)
+	{
+		Alpha = (DistanceAlong - SplineBoundaryMin) / (SplineBoundaryMax - SplineBoundaryMin);
+	}
+	else if (StaticMesh)
+	{
+		const FBoxSphereBounds StaticMeshBounds = StaticMesh->GetBounds();
+		const double MeshMinZ = USplineMeshComponent::GetAxisValueRef(StaticMeshBounds.Origin, ForwardAxis) - USplineMeshComponent::GetAxisValueRef(StaticMeshBounds.BoxExtent, ForwardAxis);
+		const double MeshRangeZ = 2 * USplineMeshComponent::GetAxisValueRef(StaticMeshBounds.BoxExtent, ForwardAxis);
+		if (MeshRangeZ > UE_SMALL_NUMBER)
+		{
+			Alpha = UE::SplineMesh::RealToFloatChecked((DistanceAlong - MeshMinZ) / MeshRangeZ);
+		}
+	}
+	return Alpha;
+}
+
+void FSplineMeshSceneProxyDesc::ComputeVisualMeshSplineTRange(UStaticMesh* StaticMesh, float& MinT, float& MaxT) const
+{
+	MinT = 0.0;
+	MaxT = 1.0;	
+	const bool bHasCustomBoundary = !FMath::IsNearlyEqual(SplineBoundaryMin, SplineBoundaryMax);
+	if (bHasCustomBoundary)
+	{
+		const FBoxSphereBounds& MeshBounds = StaticMesh->GetBounds();
+		// If there's a custom boundary, alter the min/max of the spline we need to evaluate
+		ESplineMeshAxis::Type ForwardAxisType = static_cast<ESplineMeshAxis::Type>(ForwardAxis);
+		const float BoundsMin = UE::SplineMesh::RealToFloatChecked(USplineMeshComponent::GetAxisValueRef(MeshBounds.Origin - MeshBounds.BoxExtent, ForwardAxisType));
+		const float BoundsMax = UE::SplineMesh::RealToFloatChecked(USplineMeshComponent::GetAxisValueRef(MeshBounds.Origin + MeshBounds.BoxExtent, ForwardAxisType));
+		const float BoundsMinT = (BoundsMin - SplineBoundaryMin) / (SplineBoundaryMax - SplineBoundaryMin);
+		const float BoundsMaxT = (BoundsMax - SplineBoundaryMin) / (SplineBoundaryMax - SplineBoundaryMin);
+
+		// Disallow extrapolation beyond a certain value; enormous bounding boxes cause the render thread to crash
+		constexpr float MaxSplineExtrapolation = 4.0f;
+		MinT = FMath::Max(-MaxSplineExtrapolation, BoundsMinT);
+		MaxT = FMath::Min(BoundsMaxT, MaxSplineExtrapolation);
+	}
+}
\ No newline at end of file
diff --git a/Engine/Source/Runtime/Engine/Private/SplineMeshSceneProxyDesc.h b/Engine/Source/Runtime/Engine/Private/SplineMeshSceneProxyDesc.h
new file mode 100644
--- /dev/null
+++ b/Engine/Source/Runtime/Engine/Private/SplineMeshSceneProxyDesc.h
@@ -0,0 +1,33 @@
+﻿// Copyright Epic Games, Inc. All Rights Reserved.
+#pragma once
+#include "CoreMinimal.h"
+#include "StaticMeshSceneProxyDesc.h"
+#include "Components/SplineMeshComponent.h"
+
+class USplineMeshComponent;
+
+struct FSplineMeshSceneProxyDesc
+{
+	FSplineMeshSceneProxyDesc() = default;
+	ENGINE_API FSplineMeshSceneProxyDesc(const USplineMeshComponent* InComponent);
+	
+	void InitializeFrom(const USplineMeshComponent* InComponent);
+	
+	FSplineMeshParams SplineParams{};
+	FVector SplineUpDir = FVector::UpVector;
+	float SplineBoundaryMin = 0.0f;
+	float SplineBoundaryMax = 0.0f;
+	TEnumAsByte<ESplineMeshAxis::Type> ForwardAxis = ESplineMeshAxis::X;
+	bool bSmoothInterpRollScale = false;
+
+	FSplineMeshParams GetSplineParams() const { return SplineParams; }
+
+	FSplineMeshShaderParams CalculateShaderParams(UStaticMesh* StaticMesh) const;
+	float ComputeRatioAlongSpline(UStaticMesh* StaticMesh, float DistanceAlong) const;
+	void ComputeVisualMeshSplineTRange(UStaticMesh* StaticMesh, float& MinT, float& MaxT) const;
+	ENGINE_API FBox ComputeDistortedBounds(UStaticMesh* StaticMesh, const FTransform& InLocalToWorld, const FBoxSphereBounds& InMeshBounds, const FBoxSphereBounds* InBoundsToDistort = nullptr) const;
+	FTransform CalcSliceTransform(UStaticMesh* StaticMesh, const float DistanceAlong) const;
+	FTransform CalcSliceTransformAtSplineOffset(const float Alpha, const float MinT, const float MaxT) const;
+	
+	static void InitVertexFactory(UStaticMesh* Mesh, const ERHIFeatureLevel::Type FeatureLevel, int32 InLODIndex, FColorVertexBuffer* InOverrideColorVertexBuffer);
+};
\ No newline at end of file
diff --git a/Engine/Source/Runtime/Engine/Public/SplineMeshSceneProxy.h b/Engine/Source/Runtime/Engine/Public/SplineMeshSceneProxy.h
--- a/Engine/Source/Runtime/Engine/Public/SplineMeshSceneProxy.h
+++ b/Engine/Source/Runtime/Engine/Public/SplineMeshSceneProxy.h
@@ -15,6 +15,8 @@
 #include "NaniteSceneProxy.h"
 #include "InstanceDataSceneProxy.h"
 
+struct FSplineMeshSceneProxyDesc;
+
 //////////////////////////////////////////////////////////////////////////
 // SplineMeshVertexFactory
 
@@ -161,6 +163,7 @@ class FSplineMeshSceneProxy final : public FStaticMeshSceneProxy, public TSpline
 public:
 	FSplineMeshSceneProxy(USplineMeshComponent* InComponent);
 
+	ENGINE_API FSplineMeshSceneProxy(const FStaticMeshSceneProxyDesc& InMeshDesc, const FSplineMeshSceneProxyDesc& InSplineDesc);
 	// FPrimitiveSceneProxy interface
 	virtual SIZE_T GetTypeHash() const override;
 	virtual bool GetShadowMeshElement(int32 LODIndex, int32 BatchIndex, uint8 InDepthPriorityGroup, FMeshBatch& OutMeshBatch, bool bDitheredLODTransition) const override;
@@ -199,6 +202,7 @@ class FNaniteSplineMeshSceneProxy final : public Nanite::FSceneProxy, public TSp
 
 public:
 	FNaniteSplineMeshSceneProxy(const Nanite::FMaterialAudit& MaterialAudit, USplineMeshComponent* InComponent);
+	ENGINE_API FNaniteSplineMeshSceneProxy(const Nanite::FMaterialAudit& MaterialAudit, const FStaticMeshSceneProxyDesc& InMeshDesc, const FSplineMeshSceneProxyDesc& InSplineDesc);
 
 	// FPrimitiveSceneProxy interface
 	virtual SIZE_T GetTypeHash() const override;
