Integrated changes from a licensee to increase max size of render targets to 16k

#rnx

diff --git a/Engine/Source/Runtime/Engine/Private/TextureRenderTarget2D.cpp b/Engine/Source/Runtime/Engine/Private/TextureRenderTarget2D.cpp
--- a/Engine/Source/Runtime/Engine/Private/TextureRenderTarget2D.cpp
+++ b/Engine/Source/Runtime/Engine/Private/TextureRenderTarget2D.cpp
@@ -242,7 +242,7 @@ void UTextureRenderTarget2D::PostEditChangeProperty(FPropertyChangedEvent& Prope
 		}
 	}
 
-	const int32 MaxSize = 8192; 
+	const int32 MaxSize = 16384;
 	
 	SizeX = FMath::Clamp<int32>(SizeX - (SizeX % GPixelFormats[Format].BlockSizeX),1,MaxSize);
 	SizeY = FMath::Clamp<int32>(SizeY - (SizeY % GPixelFormats[Format].BlockSizeY),1,MaxSize);
