Initial implementation of instanced UAnimBankMeshComponent

diff --git a/Engine/Shaders/Private/Skinning/AnimBankMeshDebug.usf b/Engine/Shaders/Private/Skinning/AnimBankMeshDebug.usf
new file mode 100644
--- /dev/null
+++ b/Engine/Shaders/Private/Skinning/AnimBankMeshDebug.usf
@@ -0,0 +1,100 @@
+// Copyright Epic Games, Inc. All Rights Reserved.
+
+#include "/Engine/Public/Platform.ush"
+#include "/Engine/Private/Common.ush"
+#include "/Engine/Private/SceneData.ush"
+#include "/Engine/Private/BoneTransform.ush"
+#include "/Engine/Private/Visualization.ush"
+#include "/Engine/Private/Nanite/NaniteDataDecode.ush"
+
+uint PersistentPrimitiveIndex;
+
+float3 GetBoneColor(uint BoneIndex, uint BoneDepth, uint BoneCount)
+{
+#if 0
+	float Scale = 20.0f;
+	return ColorMapTurbo((float(BoneDepth) / float(BoneCount)) * Scale);
+#else
+	return IntToColor(BoneIndex);
+#endif
+}
+
+struct FPackedBone
+{
+	uint BoneParent : 16;
+	uint BoneDepth  : 16;
+};
+
+FPackedBone GetPackedBone(FNaniteSkinningHeader SkinningHeader, uint BoneIndex)
+{
+	const uint BaseOffset = SkinningHeader.HierarchyBufferOffset * (uint)sizeof(uint);
+	const uint BoneOffset = BoneIndex * (uint)sizeof(uint);
+	return Scene.NaniteSkinning.BoneHierarchy.Load<FPackedBone>(BaseOffset + BoneOffset);
+}
+
+float4x4 MatrixFrom4x3(float4x3 M)
+{
+	return float4x4(
+		float4(M[0], 0.0f),
+		float4(M[1], 0.0f),
+		float4(M[2], 0.0f),
+		float4(M[3], 1.0f)
+	);
+}
+
+void LineVS(
+	in uint InstanceID		: SV_InstanceID,
+	in uint VertexID		: SV_VertexID,
+	out float4 OutPosition	: SV_POSITION,
+	out float4 OutColor		: TEXCOORD0
+)
+{
+	ResolvedView = ResolveView();
+
+	uint BoneIndex = VertexID >> 1u;
+
+	const FPrimitiveSceneData PrimitiveData = GetPrimitiveData(PersistentPrimitiveIndex);
+	const FInstanceSceneData InstanceData = GetInstanceSceneData(PrimitiveData.InstanceSceneDataOffset + InstanceID, Scene.GPUScene.InstanceDataSOAStride);
+	const FNaniteSkinningHeader SkinningHeader = LoadNaniteSkinningHeader(PersistentPrimitiveIndex);
+	const FPackedBone PackedBone = GetPackedBone(SkinningHeader, BoneIndex);
+
+	uint BoneDepth = PackedBone.BoneDepth;
+	if ((VertexID % 2u) == 0)
+	{
+		const uint ParentBoneIndex = PackedBone.BoneParent;
+		if (ParentBoneIndex < SkinningHeader.MaxTransformCount)
+		{
+			BoneIndex = ParentBoneIndex;
+			BoneDepth--;
+		}
+	}
+
+	const float4x4 BoneTransform = MatrixFrom4x3(LoadNaniteBoneTransform(SkinningHeader.TransformBufferOffset + InstanceData.SkinningData + BoneIndex));
+
+	float4x4 BoneObjectSpace;
+
+	BRANCH
+	if (SkinningHeader.bHasScale)
+	{
+		BoneObjectSpace = MatrixFromBoneTransformWithScale(LoadNaniteBoneObjectSpaceWithScale(SkinningHeader.ObjectSpaceBufferOffset, BoneIndex));
+	}
+	else
+	{
+		BoneObjectSpace = MatrixFromBoneTransform(LoadNaniteBoneObjectSpace(SkinningHeader.ObjectSpaceBufferOffset, BoneIndex));
+	}
+
+	const float3 LocalPosition = mul(BoneObjectSpace, BoneTransform)._m30_m31_m32;
+	const float4 TranslatedWorldPosition = float4(DFTransformLocalToTranslatedWorld(LocalPosition, InstanceData.LocalToWorld, ResolvedView.PreViewTranslation).xyz, 1.0f);
+
+	OutPosition = mul(TranslatedWorldPosition, ResolvedView.TranslatedWorldToClip);
+	OutColor = float4(GetBoneColor(VertexID >> 1u, BoneDepth, SkinningHeader.MaxTransformCount), 1.0f);
+}
+
+void LinePS(
+	 in float4 Position	: SV_POSITION,
+	 in float4 Color	: TEXCOORD0,
+	out float4 OutColor	: SV_Target0
+)
+{
+	OutColor = Color;
+}
\ No newline at end of file
diff --git a/Engine/Shaders/Shared/SkinningDefinitions.h b/Engine/Shaders/Shared/SkinningDefinitions.h
--- a/Engine/Shaders/Shared/SkinningDefinitions.h
+++ b/Engine/Shaders/Shared/SkinningDefinitions.h
@@ -2,6 +2,8 @@
 
 #pragma once
 
+#if defined(__cplusplus) || COMPILER_SUPPORTS_HLSL2021
+
 #ifdef __cplusplus
 
 #include "HLSLTypeAliases.h"
@@ -52,6 +54,10 @@ struct FBankScatterHeader
 } // namespace UE::HLSL
 
 #define REF_POSE_TRANSFORM_PROVIDER_GUID 0x665207E7, 0x449A4FB1, 0xA298F7AD, 0x8F989B11
+#define ANIM_BANK_GPU_TRANSFORM_PROVIDER_GUID 0xA5C0027A, 0x8F884C7C, 0x9312F138, 0x71A9300F
+#define ANIM_BANK_CPU_TRANSFORM_PROVIDER_GUID 0xE7D6173D, 0x246F431A, 0x912D384E, 0x156C0D2C
 #define ANIM_RUNTIME_TRANSFORM_PROVIDER_GUID 0xF1508490, 0xFCC24BB9, 0xA9F277B3, 0x1AF766F0
 
 #endif
+
+#endif // defined(__cplusplus) || COMPILER_SUPPORTS_HLSL2021
\ No newline at end of file
diff --git a/Engine/Source/Runtime/Engine/Classes/Components/AnimBankMeshComponent.h b/Engine/Source/Runtime/Engine/Classes/Components/AnimBankMeshComponent.h
new file mode 100644
--- /dev/null
+++ b/Engine/Source/Runtime/Engine/Classes/Components/AnimBankMeshComponent.h
@@ -0,0 +1,280 @@
+// Copyright Epic Games, Inc. All Rights Reserved.
+
+#pragma once
+
+#include "Components/StaticMeshComponent.h"
+#include "Components/SkeletalMeshComponent.h"
+#include "Components/InstancedStaticMeshComponent.h"
+#include "Rendering/NaniteInterface.h"
+#include "Animation/AnimBank.h"
+#include "InstanceDataSceneProxy.h"
+#include "AnimBankMeshComponent.generated.h"
+
+class FPrimitiveSceneProxy;
+class USkinnedMeshComponent;
+class UMaterialInterface;
+class UTexture;
+class UAnimBank;
+struct FAnimBankRecordHandle;
+
+USTRUCT()
+struct FAnimBankMeshInstanceData
+{
+	GENERATED_USTRUCT_BODY()
+
+	UPROPERTY(EditAnywhere, Category=Instances)
+	FMatrix Transform;
+
+	UPROPERTY(EditAnywhere, Category = Animation)
+	uint32 BankIndex;
+
+	uint32 Padding[3]; // Need to respect 16 byte alignment for serialization
+
+	FAnimBankMeshInstanceData()
+	: Transform(FMatrix::Identity)
+	, BankIndex(0)
+	{
+		Padding[0] = 0;
+		Padding[1] = 0;
+		Padding[2] = 0;
+	}
+
+	FAnimBankMeshInstanceData(const FMatrix& InTransform, uint32 InBankIndex)
+	: Transform(InTransform)
+	, BankIndex(InBankIndex)
+	{
+		Padding[0] = 0;
+		Padding[1] = 0;
+		Padding[2] = 0;
+	}
+
+	friend FArchive& operator<<(FArchive& Ar, FAnimBankMeshInstanceData& InstanceData)
+	{
+		// @warning BulkSerialize: FAnimBankMeshInstanceData is serialized as memory dump
+		// See TArray::BulkSerialize for detailed description of implied limitations.
+		Ar << InstanceData.Transform;
+		Ar << InstanceData.BankIndex;
+		Ar << InstanceData.Padding[0];
+		Ar << InstanceData.Padding[1];
+		Ar << InstanceData.Padding[2];
+		return Ar;
+	}
+};
+
+UCLASS(ClassGroup=Rendering, hidecategories=(Object,Activation,Collision,"Components|Activation",Physics), editinlinenew, meta=(BlueprintSpawnableComponent))
+class ENGINE_API UAnimBankMeshComponent : public USkinnedMeshComponent
+{
+	GENERATED_BODY()
+
+public:
+	/** Array of instances, bulk serialized. */
+	UPROPERTY(EditAnywhere, SkipSerialization, DisplayName="Instances", Category=Instances, meta=(MakeEditWidget=true, EditFixedOrder))
+	TArray<FAnimBankMeshInstanceData> InstanceData;
+
+	/** Defines the number of floats that will be available per instance for custom data */
+	UPROPERTY(EditAnywhere, Category=Instances, AdvancedDisplay)
+	int32 NumCustomDataFloats;
+
+	/** Array of custom data for instances. This will contains NumCustomDataFloats*InstanceCount entries. The entries are represented sequantially, in instance order. Can be read in a material and manipulated through Blueprints.
+	*	Example: If NumCustomDataFloats is 1, then each entry will belong to an instance. Custom data 0 will belong to Instance 0. Custom data 1 will belong to Instance 1 etc.
+	*	Example: If NumCustomDataFloats is 2, then each pair of sequential entries belong to an instance. Custom data 0 and 1 will belong to Instance 0. Custom data 2 and 3 will belong to Instance 2 etc.
+	*/
+	UPROPERTY(EditAnywhere, EditFixedSize, SkipSerialization, DisplayName="Custom Data", Category=Instances, AdvancedDisplay, meta=(EditFixedOrder))
+	TArray<float> InstanceCustomData;
+
+	UPROPERTY(EditAnywhere, BlueprintReadWrite, Category = Animation, meta = (ShowOnlyInnerProperties))
+	TArray<struct FAnimBankItem> AnimBankItems;
+
+	/** If true, this component will avoid serializing its per instance data / those properties will also not be editable */
+	UPROPERTY()
+	uint8 bInheritPerInstanceData : 1;
+
+public:
+	UAnimBankMeshComponent(FVTableHelper& Helper);
+	virtual ~UAnimBankMeshComponent();
+
+	virtual void Serialize(FArchive& Ar) override;
+
+	virtual void PostLoad() override;
+
+	virtual void BeginDestroy() override;
+
+	virtual void TickComponent(float DeltaTime, enum ELevelTick TickType, FActorComponentTickFunction* ThisTickFunction) override;
+	virtual void OnRegister() override;
+
+	virtual const Nanite::FResources* GetNaniteResources() const;
+
+	virtual void CreateRenderState_Concurrent(FRegisterComponentContext* Context) override;
+	virtual void DestroyRenderState_Concurrent() override;
+
+	virtual FPrimitiveSceneProxy* CreateSceneProxy() override;
+
+#if WITH_EDITOR
+	virtual void PostEditChangeChainProperty(FPropertyChangedChainEvent& PropertyChangedEvent) override;
+	virtual void PostEditUndo() override;
+#endif
+
+	//~ Begin UActorComponent Interface
+	virtual TStructOnScope<FActorComponentInstanceData> GetComponentInstanceData() const override;
+	//~ End UActorComponent Interface
+
+	/** Handle changes that must happen before the proxy is recreated. */
+	void PreApplyComponentInstanceData(struct FAnimBankMeshComponentInstanceData* ComponentInstanceData);
+
+	/** Applies the cached component instance data to a newly blueprint constructed component. */
+	virtual void ApplyComponentInstanceData(struct FAnimBankMeshComponentInstanceData* ComponentInstanceData);
+
+	virtual FBoxSphereBounds CalcBounds(const FTransform& LocalToWorld) const override;
+
+	virtual void RefreshBoneTransforms(FActorComponentTickFunction* TickFunction = nullptr) override;
+
+	/** Add an instance to this component. Transform is given in local space of this component unless bWorldSpace is set. */
+	UFUNCTION(BlueprintCallable, Category = "Components|AnimBankMesh")
+	virtual int32 AddInstance(const FTransform& InstanceTransform, int32 BankIndex, bool bWorldSpace = false);
+
+	/** Add multiple instances to this component. Transform is given in local space of this component unless bWorldSpace is set. */
+	UFUNCTION(BlueprintCallable, Category = "Components|AnimBankMesh")
+	virtual TArray<int32> AddInstances(const TArray<FTransform>& InstanceTransforms, bool bShouldReturnIndices, bool bWorldSpace = false, bool bUpdateNavigation = true);
+
+	/** Update custom data for specific instance */
+	UFUNCTION(BlueprintCallable, Category = "Components|AnimBankMesh")
+	virtual bool SetCustomDataValue(int32 InstanceIndex, int32 CustomDataIndex, float CustomDataValue, bool bMarkRenderStateDirty = false);
+
+	/** Update number of custom data entries per instance. This applies to all instances and will reallocate the full custom data buffer and reset all values to 0 */
+	UFUNCTION(BlueprintCallable, Category = "Components|AnimBankMesh")
+	virtual void SetNumCustomDataFloats(int32 InNumCustomDataFloats);
+
+	/** Per Instance Custom Data */
+	virtual bool SetCustomData(int32 InstanceIndex, TArrayView<const float> CustomDataFloats, bool bMarkRenderStateDirty = false);
+
+	/** Preallocated memory to include the new added instances count, to prevent reallloc during the add operation. */
+	virtual void PreAllocateInstancesMemory(int32 AddedInstanceCount);
+
+	/** Get the transform for the instance specified. Instance is returned in local space of this component unless bWorldSpace is set.  Returns True on success. */
+	UFUNCTION(BlueprintCallable, Category = "Components|AnimBankMesh")
+	bool GetInstanceTransform(int32 InstanceIndex, FTransform& OutInstanceTransform, bool bWorldSpace = false) const;
+
+	/** Remove the instance specified. Returns True on success. */
+	UFUNCTION(BlueprintCallable, Category = "Components|AnimBankMesh")
+	virtual bool RemoveInstance(int32 InstanceIndex);
+
+	/** Remove the instances specified. Returns True on success. */
+	UFUNCTION(BlueprintCallable, Category = "Components|AnimBankMesh")
+	virtual bool RemoveInstances(const TArray<int32>& InstancesToRemove);
+
+	/** Remove the instances specified. Returns True on success.
+	* @param InstancesToRemove array of the instance indices to remove ( can be sorted @see bInstanceArrayAlreadySortedInReverseOrder )
+	* @param bInstanceArrayAlreadySortedInReverseOrder true is the array is already sorted in reverse order
+	*/
+	virtual bool RemoveInstances(const TArray<int32>& InstancesToRemove, bool bInstanceArrayAlreadySortedInReverseOrder);
+
+	/** Clear all instances being rendered by this component. */
+	UFUNCTION(BlueprintCallable, Category = "Components|AnimBankMesh")
+	virtual void ClearInstances();
+
+	UFUNCTION(BlueprintCallable, Category = "Components|AnimBankMesh")
+	virtual int32 AddAnimBankItem(const FAnimBankItem& BankItem);
+
+	bool IsEnabled() const;
+	int32 GetNumInstances() const;
+
+#if WITH_EDITOR
+	/** One bit per instance if the instance is selected. */
+	TBitArray<> SelectedInstances;
+
+	/** Indicates that the user has purposely chosen to show the instance list in the details panel, despite the performance warning. */
+	bool bForceShowAllInstancesDetails = false;
+
+	/** The reason why a deletion operation is currently happening. */
+	enum class EInstanceDeletionReason : uint8
+	{
+		NotDeleting, /** There is currently no deletion in progress. */
+		EntryAlreadyRemoved, /** The instance has been deleted externally. Data synchronization in progress. */
+		EntryRemoval, /** The instance is being removed. */
+		Clearing /** All instances are being removed. */
+	};
+
+	/** This will be set to the appropriate state when one or more instances are in the process of being
+	 *  deleted. This is primarily used for functions that round trip to this class, such as callbacks
+	 *  for deselecting instances. */
+	EInstanceDeletionReason DeletionState = EInstanceDeletionReason::NotDeleting;
+#endif
+
+protected:
+	bool bIsInstanceDataApplyCompleted = true;
+
+	void CreateHitProxyData(TArray<TRefCountPtr<HHitProxy>>& HitProxies);
+
+	/** Internal version of AddInstance */
+	int32 AddInstanceInternal(int32 InstanceIndex, FAnimBankMeshInstanceData* InNewInstanceData, const FTransform& InstanceTransform, int32 BankIndex, bool bWorldSpace);
+
+	/** Internal implementation of AddInstances */
+	TArray<int32> AddInstancesInternal(TConstArrayView<FTransform> InstanceTransforms, bool bShouldReturnIndices, bool bWorldSpace, bool bUpdateNavigation = true);
+
+	/** Internal version of RemoveInstance */
+	bool RemoveInstanceInternal(int32 InstanceIndex, bool InstanceAlreadyRemoved, bool bForceRemoveAtSwap = false, bool bUpdateNavigation = true);
+
+private:
+	void ApplyInheritedPerInstanceData(const UAnimBankMeshComponent* InArchetype);
+	bool ShouldInheritPerInstanceData(const UAnimBankMeshComponent* InArchetype) const;
+	bool ShouldInheritPerInstanceData() const;
+
+	/** Sets up new instance data to sensible defaults, creates physics counterparts if possible. */
+	void SetupNewInstanceData(FAnimBankMeshInstanceData& InOutNewInstanceData, int32 InInstanceIndex, const FTransform& InInstanceTransform, int32 BankIndex);
+
+private:
+	void SetSkinnedAssetCallback();
+};
+
+/** Helper class used to preserve state across blueprint re-instancing */
+USTRUCT()
+struct FAnimBankMeshComponentInstanceData : public FSceneComponentInstanceData
+{
+	GENERATED_BODY()
+
+public:
+	FAnimBankMeshComponentInstanceData() = default;
+	
+	FAnimBankMeshComponentInstanceData(const UAnimBankMeshComponent* InComponent)
+	: FSceneComponentInstanceData(InComponent)
+	, SkinnedAsset(InComponent->GetSkinnedAsset())
+	{
+	}
+
+	virtual ~FAnimBankMeshComponentInstanceData() = default;
+
+	virtual bool ContainsData() const override
+	{
+		return true;
+	}
+
+	virtual void ApplyToComponent(UActorComponent* Component, const ECacheApplyPhase CacheApplyPhase) override
+	{
+		// The Super::ApplyToComponent will cause the scene proxy to be recreated, so we must do what we can to make sure the state is ok before that.
+		CastChecked<UAnimBankMeshComponent>(Component)->PreApplyComponentInstanceData(this);
+		Super::ApplyToComponent(Component, CacheApplyPhase);
+		CastChecked<UAnimBankMeshComponent>(Component)->ApplyComponentInstanceData(this);
+	}
+
+	virtual void AddReferencedObjects(FReferenceCollector& Collector) override
+	{
+		Super::AddReferencedObjects(Collector);
+		Collector.AddReferencedObject(SkinnedAsset);
+	}
+
+public:
+	UPROPERTY()
+	TObjectPtr<USkinnedAsset> SkinnedAsset = nullptr;
+
+	UPROPERTY()
+	TArray<FAnimBankMeshInstanceData> InstanceData;
+
+	TBitArray<> SelectedInstances;
+
+	UPROPERTY()
+	bool bHasPerInstanceHitProxies = false;
+};
+
+#if UE_ENABLE_INCLUDE_ORDER_DEPRECATED_IN_5_2
+#include "CoreMinimal.h"
+#endif
diff --git a/Engine/Source/Runtime/Engine/Private/Components/AnimBankMeshComponent.cpp b/Engine/Source/Runtime/Engine/Private/Components/AnimBankMeshComponent.cpp
new file mode 100644
--- /dev/null
+++ b/Engine/Source/Runtime/Engine/Private/Components/AnimBankMeshComponent.cpp
@@ -0,0 +1,1533 @@
+// Copyright Epic Games, Inc. All Rights Reserved.
+
+#include "Components/AnimBankMeshComponent.h"
+#include "Animation/AnimSequence.h"
+#include "Animation/Skeleton.h"
+#include "ComponentRecreateRenderStateContext.h"
+#include "Elements/Framework/EngineElementsLibrary.h"
+#include "Elements/Interfaces/TypedElementWorldInterface.h"
+#include "Elements/SMInstance/SMInstanceElementData.h" // For SMInstanceElementDataUtil::SMInstanceElementsEnabled
+#include "Elements/SMInstance/SMInstanceElementId.h"
+#include "Engine/StaticMesh.h"
+#include "HitProxies.h"
+#include "NaniteSceneProxy.h"
+#include "PrimitiveSceneInfo.h"
+#include "Rendering/RenderCommandPipes.h"
+#include "Rendering/SkeletalMeshRenderData.h"
+#include "SceneInterface.h"
+#include "SkeletalRenderPublic.h"
+#include "SkinningDefinitions.h"
+#include "SkinnedMeshSceneProxyDesc.h"
+
+#include UE_INLINE_GENERATED_CPP_BY_NAME(AnimBankMeshComponent)
+
+static TAutoConsoleVariable<int32> CVarAnimBankForceRefPose(
+	TEXT("r.AnimBank.ForceRefPose"),
+	0,
+	TEXT("Whether to force ref pose for instanced skinned meshes"),
+	ECVF_RenderThreadSafe
+);
+
+static TAutoConsoleVariable<int32> CVarAnimBankSampledBounds(
+	TEXT("r.AnimBank.SampledBounds"),
+	1,
+	TEXT("Whether to use sampled bounds for anim bank meshes"),
+	FConsoleVariableDelegate::CreateLambda([](IConsoleVariable* InVariable)
+	{
+		FGlobalComponentRecreateRenderStateContext Context;
+	}),
+	ECVF_RenderThreadSafe
+);
+
+class FAnimBankMeshSceneProxy : public Nanite::FSkinnedSceneProxy
+{
+public:
+	using Super = Nanite::FSkinnedSceneProxy;
+
+	FAnimBankMeshSceneProxy(const Nanite::FMaterialAudit& MaterialAudit, UAnimBankMeshComponent* InComponent, FSkeletalMeshRenderData* InRenderData)
+	: Super(MaterialAudit, InComponent, InRenderData, false /* bAllowScale */)
+	{
+	#if WITH_EDITOR
+		const bool bSupportInstancePicking = HasPerInstanceHitProxies() && SMInstanceElementDataUtil::SMInstanceElementsEnabled();
+		HitProxyMode = bSupportInstancePicking ? EHitProxyMode::PerInstance : EHitProxyMode::MaterialSection;
+
+		if (HitProxyMode == EHitProxyMode::PerInstance)
+		{
+			bHasSelectedInstances = InComponent->SelectedInstances.Find(true) != INDEX_NONE;
+			if (bHasSelectedInstances)
+			{
+				// If we have selected indices, mark scene proxy as selected.
+				SetSelection_GameThread(true);
+			}
+		}
+	#endif
+
+		MeshBounds = InComponent->GetSkinnedAsset()->GetBounds();
+
+		FInstanceSceneDataBuffers::FAccessTag AccessTag(PointerHash(this));
+		FInstanceSceneDataBuffers::FWriteView ProxyData = InstanceSceneDataBuffersImpl.BeginWriteAccess(AccessTag);
+
+		const bool bForceRefPose = CVarAnimBankForceRefPose.GetValueOnAnyThread() != 0;
+		const bool bUseAnimBank = !bForceRefPose && InComponent->AnimBankItems.Num() > 0;
+
+		ProxyData.Flags.bHasPerInstanceHierarchyOffset = false;
+		ProxyData.Flags.bHasPerInstanceLocalBounds = bUseAnimBank && InComponent->AnimBankItems.Num() > 1;
+		ProxyData.Flags.bHasPerInstanceDynamicData = false;
+		ProxyData.Flags.bHasPerInstanceSkinningData = true;
+	#if WITH_EDITOR
+		ProxyData.Flags.bHasPerInstanceEditorData = false;// GIsEditor != 0 && HasPerInstanceHitProxies();
+	#endif
+
+		const int32 InstanceCount = InComponent->GetNumInstances();
+
+		InstanceTransforms.Reset(InstanceCount);
+		if (ProxyData.Flags.bHasPerInstanceLocalBounds)
+		{
+			InstanceBounds.Reset(InstanceCount);
+		}
+		else if (bUseAnimBank && InComponent->AnimBankItems.Num() == 1)
+		{
+			// No use wasting memory/perf with per-instance local bounds if they all point to the same anim bank sequence
+			const FAnimBankItem& BankItem = InComponent->AnimBankItems[0];
+			if (BankItem.BankAsset != nullptr && CVarAnimBankSampledBounds.GetValueOnAnyThread() != 0)
+			{
+				const FAnimBankData& BankData = BankItem.BankAsset->GetData();
+				if (BankItem.SequenceIndex < BankData.Entries.Num())
+				{
+					MeshBounds = BankData.Entries[BankItem.SequenceIndex].SampledBounds;
+				}
+			}
+		}
+
+		ProxyData.InstanceSkinningData.SetNumUninitialized(InstanceCount);
+
+		for (int32 InstanceIndex = 0; InstanceIndex < InstanceCount; ++InstanceIndex)
+		{
+			const FAnimBankMeshInstanceData& Instance = InComponent->InstanceData[InstanceIndex];
+
+			if (ProxyData.Flags.bHasPerInstanceLocalBounds)
+			{
+				FRenderBounds& SampledBounds = InstanceBounds.Emplace_GetRef(MeshBounds);
+
+				if (Instance.BankIndex < uint32(InComponent->AnimBankItems.Num()))
+				{
+					const FAnimBankItem& BankItem = InComponent->AnimBankItems[Instance.BankIndex];
+					if (BankItem.BankAsset != nullptr)
+					{
+						const FAnimBankData& BankData = BankItem.BankAsset->GetData();
+						if (BankItem.SequenceIndex < BankData.Entries.Num())
+						{
+							SampledBounds = BankData.Entries[BankItem.SequenceIndex].SampledBounds;
+						}
+					}
+				}
+			}
+
+			InstanceTransforms.Emplace(FRenderTransform(Instance.Transform));
+			ProxyData.InstanceSkinningData[InstanceIndex] = Instance.BankIndex * MaxBoneTransformCount * 2u /* Current and Previous */;
+		}
+
+		InstanceSceneDataBuffersImpl.EndWriteAccess(AccessTag);
+
+		// Nanite supports the GPUScene instance data buffer.
+		SetupInstanceSceneDataBuffers(&InstanceSceneDataBuffersImpl);
+
+		// Note: ideally this would be picked up from the Flags.bHasPerInstanceDynamicData above, but that path is not great at the moment.
+		bAlwaysHasVelocity = true;
+
+		if (bUseAnimBank)
+		{
+			static FGuid AnimBankGPUProviderId(ANIM_BANK_GPU_TRANSFORM_PROVIDER_GUID);
+			static FGuid AnimBankCPUProviderId(ANIM_BANK_CPU_TRANSFORM_PROVIDER_GUID);
+
+			static const auto AnimBankGPUVar = IConsoleManager::Get().FindTConsoleVariableDataInt(TEXT("r.AnimBank.GPU"));
+			if (AnimBankGPUVar && AnimBankGPUVar->GetValueOnAnyThread() == 1)
+			{
+				TransformProviderId = AnimBankGPUProviderId;
+			}
+			else
+			{
+				TransformProviderId = AnimBankCPUProviderId;
+			}
+
+			AnimBankItems = InComponent->AnimBankItems;
+			AnimBankIds.Reset(AnimBankItems.Num());
+			AnimBankHandles.Reset(AnimBankItems.Num());
+			UniqueAnimationCount = AnimBankItems.Num();
+		}
+		else
+		{
+			UniqueAnimationCount = 1; // Ref Pose
+			static FGuid RefPoseProviderId(REF_POSE_TRANSFORM_PROVIDER_GUID);
+			TransformProviderId = RefPoseProviderId;
+		}
+	}
+
+	virtual ~FAnimBankMeshSceneProxy()
+	{
+	}
+
+	virtual void OnTransformChanged(FRHICommandListBase& RHICmdList) override
+	{
+		Super::OnTransformChanged(RHICmdList);
+
+		FInstanceSceneDataBuffers::FAccessTag AccessTag(PointerHash(this));
+		FInstanceSceneDataBuffers::FWriteView ProxyData = InstanceSceneDataBuffersImpl.BeginWriteAccess(AccessTag);
+
+		InstanceSceneDataBuffersImpl.SetPrimitiveLocalToWorld(GetLocalToWorld(), AccessTag);
+
+		const FRenderTransform& PrimitiveToRelativeWorld = InstanceSceneDataBuffersImpl.GetPrimitiveToRelativeWorld(AccessTag);
+
+		const int32 InstanceCount = InstanceTransforms.Num();
+		ProxyData.InstanceToPrimitiveRelative.Reset(InstanceCount);
+		if (ProxyData.Flags.bHasPerInstanceLocalBounds)
+		{
+			ProxyData.InstanceLocalBounds.Reset(InstanceCount);
+		}
+
+		for (int32 InstanceIndex = 0; InstanceIndex < InstanceCount; ++InstanceIndex)
+		{
+			FRenderTransform InstanceToPrimitiveRelativeWorld = InstanceTransforms[InstanceIndex] * PrimitiveToRelativeWorld;
+			
+			// Remove shear
+			InstanceToPrimitiveRelativeWorld.Orthogonalize();
+
+			//FRenderTransform& PrevInstanceToPrimitiveRelative = ProxyData.PrevInstanceToPrimitiveRelative.Emplace_GetRef();
+			//PrevInstanceToPrimitiveRelative = InstanceToPrimitiveRelative;
+
+			ProxyData.InstanceToPrimitiveRelative.Emplace(InstanceToPrimitiveRelativeWorld);
+
+			if (ProxyData.Flags.bHasPerInstanceLocalBounds)
+			{
+				ProxyData.InstanceLocalBounds.Emplace(InstanceBounds[InstanceIndex]);
+			}
+		}
+
+		if (!ProxyData.Flags.bHasPerInstanceLocalBounds)
+		{
+			if (InstanceCount > 0)
+			{
+				ProxyData.InstanceLocalBounds.SetNum(1);
+				ProxyData.InstanceLocalBounds[0] = MeshBounds;
+			}
+			else
+			{
+				ProxyData.InstanceLocalBounds.SetNum(0);
+			}
+		}
+
+		InstanceSceneDataBuffersImpl.EndWriteAccess(AccessTag);
+	}
+
+	virtual void CreateRenderThreadResources(FRHICommandListBase& RHICmdList) override
+	{
+		Super::CreateRenderThreadResources(RHICmdList);
+
+		TArray<FAnimBankDesc> Descs;
+		Descs.Reserve(AnimBankItems.Num());
+
+		for (const FAnimBankItem& Item : AnimBankItems)
+		{
+			if (Item.BankAsset == nullptr || SkinnedAsset == nullptr)
+			{
+				Descs.Emplace(FAnimBankDesc());
+				continue;
+			}
+
+			const FAnimBankData& BankData = Item.BankAsset->GetData();
+
+			if (Item.SequenceIndex >= BankData.Entries.Num())
+			{
+				Descs.Emplace(FAnimBankDesc());
+				continue;
+			}
+
+			const FAnimBankEntry& BankEntry = BankData.Entries[Item.SequenceIndex];
+
+			FAnimBankDesc& Desc = Descs.Emplace_GetRef();
+
+			Desc.BankAsset		= MakeWeakObjectPtr(Item.BankAsset);
+			Desc.SequenceIndex	= Item.SequenceIndex;
+			Desc.Asset			= MakeWeakObjectPtr(SkinnedAsset);
+			Desc.Position		= BankEntry.Position;
+			Desc.PlayRate		= BankEntry.PlayRate;
+			Desc.bLooping		= BankEntry.IsLooping()   ? 1 : 0;
+			Desc.bAutoStart		= BankEntry.IsAutoStart() ? 1 : 0;
+		}
+
+		AnimBankHandles = GetScene().RegisterAnimBank(Descs);
+		AnimBankIds.Reset(AnimBankHandles.Num());
+		for (const FAnimBankRecordHandle& Handle : AnimBankHandles)
+		{
+			AnimBankIds.Emplace((uint64)Handle.Id);
+		}
+	}
+
+	virtual void DestroyRenderThreadResources() override
+	{
+		GetScene().UnregisterAnimBank(AnimBankHandles);
+		AnimBankHandles.Reset();
+		AnimBankIds.Reset();
+		Super::DestroyRenderThreadResources();
+	}
+
+	virtual const TConstArrayView<uint64> GetAnimationProviderData(bool& bOutValid) const override
+	{
+		bOutValid = (AnimBankIds.Num() == UniqueAnimationCount);
+		return AnimBankIds;
+	}
+
+public:
+	// FPrimitiveSceneProxy interface.
+	virtual SIZE_T GetTypeHash() const override
+	{
+		static size_t UniquePointer;
+		return reinterpret_cast<size_t>(&UniquePointer);
+	}
+
+private:
+	FInstanceSceneDataBuffers InstanceSceneDataBuffersImpl;
+
+	TArray<FRenderTransform> InstanceTransforms;
+	TArray<FRenderBounds> InstanceBounds;
+	FRenderBounds MeshBounds;
+
+	TArray<uint64> AnimBankIds;
+	TArray<FAnimBankRecordHandle> AnimBankHandles;
+	TArray<FAnimBankItem> AnimBankItems;
+	uint64 BankRegistration = 0;
+};
+
+/** 
+* Stores the updated matrices needed to skin the verts.
+* Created by the game thread and sent to the rendering thread as an update 
+*/
+class FAnimBankMeshData
+{
+public:
+	FAnimBankMeshData(
+		const USkinnedAsset* InSkinnedAsset,
+		TConstArrayView<FAnimBankItem> AnimBankItems,
+		int32 InLODIndex,
+		EPreviousBoneTransformUpdateMode InPreviousBoneTransformUpdateMode
+	);
+
+	virtual ~FAnimBankMeshData();
+
+	// Current reference pose to local space transforms
+	TArray<FMatrix44f> ReferenceToLocal;
+
+	// Previous reference pose to local space transforms
+	TArray<FMatrix44f> PrevReferenceToLocal;
+
+	TArray<FMatrix3x4> CurrentBoneTransforms;
+	TArray<FMatrix3x4> PreviousBoneTransforms;
+
+#if !(UE_BUILD_SHIPPING || UE_BUILD_TEST) 
+	// Component space bone transforms
+	TArray<FTransform> ComponentSpaceTransforms;
+#endif
+
+	// Current LOD for bones being updated
+	int32 LODIndex;
+
+	// Returns the size of memory allocated by render data
+	void GetResourceSizeEx(FResourceSizeEx& CumulativeResourceSize);
+
+private:
+
+	enum class ETransformsToUpdate
+	{
+		Current,
+		Previous,
+	};
+};
+
+FAnimBankMeshData::FAnimBankMeshData(
+	const USkinnedAsset* InSkinnedAsset,
+	TConstArrayView<FAnimBankItem> AnimBankItems,
+	int32 InLODIndex,
+	EPreviousBoneTransformUpdateMode InPreviousBoneTransformUpdateMode
+)
+:	LODIndex(InLODIndex)
+{
+	int32 BoneCount = InSkinnedAsset->GetSkeleton()->GetReferenceSkeleton().GetRawBoneNum();
+	ReferenceToLocal.SetNumUninitialized(BoneCount);
+	for (int32 TransformIndex = 0; TransformIndex < BoneCount; ++TransformIndex)
+	{
+		ReferenceToLocal[TransformIndex].SetIdentity();
+	}
+
+	const int64 ReferenceToLocalCount = int64(ReferenceToLocal.Num());
+	const FMatrix44f* ReferenceToLocalPtr = ReferenceToLocal.GetData();
+
+	const int32 UniqueAnimationCount = FMath::Max(AnimBankItems.Num(), 1 /* Ref Pose */);
+
+	CurrentBoneTransforms.SetNumUninitialized(BoneCount * UniqueAnimationCount);
+	FMatrix3x4* CurrentBoneTransformsPtr = CurrentBoneTransforms.GetData();
+	TransposeTransforms(CurrentBoneTransformsPtr, ReferenceToLocalPtr, ReferenceToLocalCount);
+
+	// TODO: Temp, optimize out
+	for (int32 AnimationIndex = 1; AnimationIndex < UniqueAnimationCount; ++AnimationIndex)
+	{
+		FMemory::Memcpy(CurrentBoneTransformsPtr + (BoneCount * AnimationIndex), CurrentBoneTransformsPtr, sizeof(FMatrix3x4) * BoneCount);
+	}
+
+	PrevReferenceToLocal = ReferenceToLocal;
+	PreviousBoneTransforms = CurrentBoneTransforms;
+
+#if 0
+	UpdateRefToLocalMatrices(ReferenceToLocal, InComponent, InRenderData, LODIndex);
+	UpdateBonesRemovedByLOD(ReferenceToLocal, InComponent, ETransformsToUpdate::Current);
+
+	switch (InPreviousBoneTransformUpdateMode)
+	{
+	case EPreviousBoneTransformUpdateMode::None:
+		// Use previously uploaded buffer
+		// TODO: Nanite-Skinning, optimize scene extension upload to keep cached GPU representation using PreviousBoneTransformRevisionNumber
+		// For now we'll just redundantly update and upload previous transforms
+		UpdatePreviousRefToLocalMatrices(PrevReferenceToLocal, InComponent, InRenderData, LODIndex);
+		UpdateBonesRemovedByLOD(PrevReferenceToLocal, InComponent, ETransformsToUpdate::Previous);
+		break;
+
+	case EPreviousBoneTransformUpdateMode::UpdatePrevious:
+		UpdatePreviousRefToLocalMatrices(PrevReferenceToLocal, InComponent, InRenderData, LODIndex);
+		UpdateBonesRemovedByLOD(PrevReferenceToLocal, InComponent, ETransformsToUpdate::Previous);
+		break;
+
+	case EPreviousBoneTransformUpdateMode::DuplicateCurrentToPrevious:
+		// TODO: Nanite-Skinning likely possible we can just return ReferenceToLocal here rather than cloning it into previous
+		// Need to make sure it's safe when next update mode = None
+		PrevReferenceToLocal = ReferenceToLocal;
+		break;
+	}
+
+#if !(UE_BUILD_SHIPPING || UE_BUILD_TEST)
+	ComponentSpaceTransforms = InComponent->GetComponentSpaceTransforms();
+#endif
+#endif
+}
+
+FAnimBankMeshData::~FAnimBankMeshData() = default;
+
+void FAnimBankMeshData::GetResourceSizeEx(FResourceSizeEx& CumulativeResourceSize)
+{
+	CumulativeResourceSize.AddDedicatedSystemMemoryBytes(sizeof(*this));
+	CumulativeResourceSize.AddDedicatedSystemMemoryBytes(ReferenceToLocal.GetAllocatedSize());
+}
+
+class FAnimBankMeshObject : public FSkeletalMeshObject
+{
+public:
+	FAnimBankMeshObject(const FSkinnedMeshSceneProxyDesc& InMeshDesc, TConstArrayView<FAnimBankItem> InAnimBankItems, FSkeletalMeshRenderData* InRenderData, ERHIFeatureLevel::Type InFeatureLevel)
+		: FSkeletalMeshObject(InMeshDesc, InRenderData, InFeatureLevel)
+		, AnimBankItems(InAnimBankItems)
+	{
+		for (int32 LODIndex = 0; LODIndex < InRenderData->LODRenderData.Num(); ++LODIndex)
+		{
+			new(LODs) FAnimBankMeshObjectLOD(InFeatureLevel, InRenderData, LODIndex);
+		}
+
+		InitResources(InMeshDesc);
+	}
+
+	FAnimBankMeshObject(USkinnedMeshComponent* InComponent, FSkeletalMeshRenderData* InRenderData, ERHIFeatureLevel::Type InFeatureLevel)
+	: FAnimBankMeshObject(FSkinnedMeshSceneProxyDesc(InComponent), CastChecked<UAnimBankMeshComponent>(InComponent)->AnimBankItems, InRenderData, InFeatureLevel)
+	{
+	}
+
+	virtual ~FAnimBankMeshObject()
+	{
+		delete DynamicData;
+	}
+
+	virtual void InitResources(const FSkinnedMeshSceneProxyDesc& InMeshDesc) override
+	{
+		for (int32 LODIndex = 0; LODIndex < LODs.Num(); ++LODIndex)
+		{
+			FAnimBankMeshObjectLOD& LOD = LODs[LODIndex];
+
+			// Skip LODs that have their render data stripped
+			if (LOD.RenderData->LODRenderData[LODIndex].GetNumVertices() > 0)
+			{
+				const FSkelMeshComponentLODInfo* InitLODInfo = nullptr;
+				if (InMeshDesc.LODInfo.IsValidIndex(LODIndex))
+				{
+					InitLODInfo = &InMeshDesc.LODInfo[LODIndex];
+				}
+
+				LOD.InitResources(InitLODInfo);
+			}
+		}
+	}
+
+	virtual void ReleaseResources() override
+	{
+		for (int32 LODIndex = 0; LODIndex < LODs.Num(); ++LODIndex)
+		{
+			FAnimBankMeshObjectLOD& LOD = LODs[LODIndex];
+			LOD.ReleaseResources();
+		}
+	}
+	
+	virtual void Update(
+		int32 LODIndex,
+		const FSkinnedMeshSceneProxyDynamicData& InDynamicData,
+		const FPrimitiveSceneProxy* InSceneProxy,
+		const USkinnedAsset* InSkinnedAsset,
+		const FMorphTargetWeightMap& InActiveMorphTargets,
+		const TArray<float>& MorphTargetWeights,
+		EPreviousBoneTransformUpdateMode PreviousBoneTransformUpdateMode,
+		const FExternalMorphWeightData& InExternalMorphWeightData) override
+	{
+		// Create the new dynamic data for use by the rendering thread
+		// this data is only deleted when another update is sent
+		FAnimBankMeshData* NewDynamicData = new FAnimBankMeshData(InSkinnedAsset, AnimBankItems, LODIndex, PreviousBoneTransformUpdateMode);
+
+		uint64 FrameNumberToPrepare = GFrameCounter;
+		uint32 RevisionNumber = 0;
+
+		if (InSceneProxy)
+		{
+			RevisionNumber = InDynamicData.GetBoneTransformRevisionNumber();
+		}
+
+		// Queue a call to update this data
+		{
+			FAnimBankMeshObject* MeshObject = this;
+			ENQUEUE_RENDER_COMMAND(SkelMeshObjectUpdateDataCommand)(UE::RenderCommandPipe::SkeletalMesh,
+				[MeshObject, FrameNumberToPrepare, RevisionNumber, NewDynamicData](FRHICommandList& RHICmdList)
+				{
+					FScopeCycleCounter Context(MeshObject->GetStatId());
+					MeshObject->UpdateDynamicData_RenderThread(RHICmdList, NewDynamicData, FrameNumberToPrepare, RevisionNumber);
+				}
+			);
+		}
+	}
+
+	void UpdateDynamicData_RenderThread(
+		FRHICommandList& RHICmdList,
+		FAnimBankMeshData* InDynamicData,
+		uint64 FrameNumberToPrepare,
+		uint32 RevisionNumber
+	)
+	{
+		// We should be done with the old data at this point
+		delete DynamicData;
+
+		// Update with new data
+		DynamicData = InDynamicData;
+		check(DynamicData);
+
+		check(IsInParallelRenderingThread());
+	}
+
+	virtual const FVertexFactory* GetSkinVertexFactory(const FSceneView* View, int32 LODIndex, int32 ChunkIdx, ESkinVertexFactoryMode VFMode = ESkinVertexFactoryMode::Default) const override
+	{
+		check(LODs.IsValidIndex(LODIndex));
+		return nullptr;
+		//return &LODs[LODIndex].VertexFactory;
+	}
+
+	virtual const FVertexFactory* GetStaticSkinVertexFactory(int32 LODIndex, int32 ChunkIdx, ESkinVertexFactoryMode VFMode = ESkinVertexFactoryMode::Default) const override
+	{
+		check(LODs.IsValidIndex(LODIndex));
+		return nullptr;
+		//return &LODs[LODIndex].VertexFactory;
+	}
+
+	virtual TArray<FTransform>* GetComponentSpaceTransforms() const override
+	{
+	#if !(UE_BUILD_SHIPPING || UE_BUILD_TEST)
+		if (DynamicData)
+		{
+			return &(DynamicData->ComponentSpaceTransforms);
+		}
+		else
+	#endif
+		{
+			return nullptr;
+		}
+	}
+
+	virtual const TArray<FMatrix44f>& GetReferenceToLocalMatrices() const override
+	{
+		return DynamicData->ReferenceToLocal;
+	}
+
+	virtual const TArray<FMatrix44f>& GetPrevReferenceToLocalMatrices() const override
+	{
+		return DynamicData->PrevReferenceToLocal;
+	}
+
+	const TArray<FMatrix3x4>* GetCurrentBoneTransforms() const override
+	{
+		return &DynamicData->CurrentBoneTransforms;
+	}
+
+	const TArray<FMatrix3x4>* GetPreviousBoneTransforms() const override
+	{
+		return &DynamicData->PreviousBoneTransforms;
+	}
+
+	virtual int32 GetLOD() const override
+	{
+		// WorkingMinDesiredLODLevel can be a LOD that's not loaded, so need to clamp it to the first loaded LOD
+		return 0;
+		//return FMath::Max<int32>(WorkingMinDesiredLODLevel, SkeletalMeshRenderData->CurrentFirstLODIdx);
+		/*if (DynamicData)
+		{
+			return DynamicData->LODIndex;
+		}
+		else
+		{
+			return 0;
+		}*/
+	}
+
+	virtual bool HaveValidDynamicData() const override
+	{
+		return (DynamicData != nullptr);
+	}
+
+	virtual void GetResourceSizeEx(FResourceSizeEx& CumulativeResourceSize) override
+	{
+		CumulativeResourceSize.AddDedicatedSystemMemoryBytes(sizeof(*this));
+
+		if (DynamicData)
+		{
+			DynamicData->GetResourceSizeEx(CumulativeResourceSize);
+		}
+
+		CumulativeResourceSize.AddDedicatedSystemMemoryBytes(LODs.GetAllocatedSize());
+
+		for (int32 Index = 0; Index < LODs.Num(); ++Index)
+		{
+			LODs[Index].GetResourceSizeEx(CumulativeResourceSize);
+		}
+	}
+
+	virtual void UpdateSkinWeightBuffer(const TArrayView<const FSkelMeshComponentLODInfo> InLODInfo) override
+	{
+		for (int32 LODIndex = 0; LODIndex < LODs.Num(); ++LODIndex)
+		{
+			FAnimBankMeshObjectLOD& LOD = LODs[LODIndex];
+
+			// Skip LODs that have their render data stripped
+			if (LOD.RenderData->LODRenderData[LODIndex].GetNumVertices() > 0)
+			{
+				const FSkelMeshComponentLODInfo* UpdateLODInfo = nullptr;
+				if (InLODInfo.IsValidIndex(LODIndex))
+				{
+					UpdateLODInfo = &InLODInfo[LODIndex];
+				}
+
+				LOD.UpdateSkinWeights(UpdateLODInfo);
+			}
+		}
+	}
+
+	virtual bool IsNaniteMesh() const override { return true; }
+
+private:
+	FAnimBankMeshData* DynamicData = nullptr;
+
+	struct FAnimBankMeshObjectLOD
+	{
+		const FSkeletalMeshRenderData* RenderData;
+		int32 LODIndex;
+		bool bInitialized;
+
+		FAnimBankMeshObjectLOD(ERHIFeatureLevel::Type InFeatureLevel, const FSkeletalMeshRenderData* InRenderData, int32 InLOD)
+		: RenderData(InRenderData)
+		, LODIndex(InLOD)
+		, bInitialized(false)
+		{
+		}
+
+		void InitResources(const FSkelMeshComponentLODInfo* InLODInfo)
+		{
+			check(RenderData);
+			check(RenderData->LODRenderData.IsValidIndex(LODIndex));
+
+			const FSkeletalMeshLODRenderData& LODData = RenderData->LODRenderData[LODIndex];
+			(void)LODData;
+
+		#if RHI_RAYTRACING
+			if (IsRayTracingEnabled() && RenderData->bSupportRayTracing)
+			{
+			}
+		#endif
+
+			bInitialized = true;
+		}
+
+		void ReleaseResources()
+		{
+			bInitialized = false;
+		}
+
+		void GetResourceSizeEx(FResourceSizeEx& CumulativeResourceSize)
+		{
+		}
+
+		void UpdateSkinWeights(const FSkelMeshComponentLODInfo* InLODInfo)
+		{
+			check(RenderData);
+			check(RenderData->LODRenderData.IsValidIndex(LODIndex));
+
+			//FSkeletalMeshLODRenderData& LODData = RenderData->LODRenderData[LODIndex];
+			//MeshObjectWeightBuffer = FSkeletalMeshObject::GetSkinWeightVertexBuffer(LODData, InLODInfo);
+		}
+	};
+
+	TConstArrayView<FAnimBankItem> AnimBankItems;
+	TArray<FAnimBankMeshObjectLOD> LODs;
+
+	mutable int32 CachedLOD;
+};
+
+static FSkeletalMeshObject* CreateAnimBankMeshObjectFn(void* UserData, USkinnedMeshComponent* InComponent, FSkeletalMeshRenderData* InRenderData, ERHIFeatureLevel::Type InFeatureLevel)
+{
+	if (InComponent->ShouldNaniteSkin())
+	{
+		return ::new FAnimBankMeshObject(InComponent, InRenderData, InFeatureLevel);
+	}
+
+	return nullptr;
+}
+
+UAnimBankMeshComponent::UAnimBankMeshComponent(FVTableHelper& Helper)
+: Super(Helper)
+, NumCustomDataFloats(0)
+, bInheritPerInstanceData(false)
+{
+}
+
+UAnimBankMeshComponent::~UAnimBankMeshComponent()
+{
+
+}
+
+void UAnimBankMeshComponent::Serialize(FArchive& Ar)
+{
+	LLM_SCOPE(ELLMTag::SkeletalMesh);
+	Super::Serialize(Ar);
+
+	bool bCooked = Ar.IsCooking();
+	Ar << bCooked;
+
+	// Inherit properties when bEditableWhenInherited == false || bInheritPerInstanceData == true (when the component isn't a template and we are persisting data)
+	const UAnimBankMeshComponent* Archetype = Cast<UAnimBankMeshComponent>(GetArchetype());
+	const bool bInheritSkipSerializationProperties = ShouldInheritPerInstanceData(Archetype) && Ar.IsPersistent();
+	
+	// Check if we need have SkipSerialization property data to load/save
+	bool bHasSkipSerializationPropertiesData = !bInheritSkipSerializationProperties;
+	Ar << bHasSkipSerializationPropertiesData;
+
+	if (Ar.IsLoading())
+	{
+		// Read existing data if it was serialized
+		TArray<FAnimBankMeshInstanceData> TempInstanceData;
+		TArray<float> TempInstanceCustomData;
+		if (bHasSkipSerializationPropertiesData)
+		{
+			TempInstanceData.BulkSerialize(Ar, false /* force per element serialization */);
+			TempInstanceCustomData.BulkSerialize(Ar);
+		}
+
+		// If we should inherit use Archetype Data
+		if (bInheritSkipSerializationProperties)
+		{
+			ApplyInheritedPerInstanceData(Archetype);
+		} 
+		// It is possible for a component to lose its BP archetype between a save / load so in this case we have no per instance data (usually this component gets deleted through construction script)
+		else if (bHasSkipSerializationPropertiesData)
+		{
+			InstanceData = MoveTemp(TempInstanceData);
+			InstanceCustomData = MoveTemp(TempInstanceCustomData);
+		}
+	}
+	else if (bHasSkipSerializationPropertiesData)
+	{
+		InstanceData.BulkSerialize(Ar, false /* force per element serialization */);
+		InstanceCustomData.BulkSerialize(Ar);
+	}
+
+#if WITH_EDITOR
+	if (Ar.IsTransacting())
+	{
+		Ar << SelectedInstances;
+	}
+#endif
+}
+
+void UAnimBankMeshComponent::PostLoad()
+{
+	Super::PostLoad();
+
+#if WITH_EDITOR
+	if (GIsEditor)
+	{
+		SetSkinnedAssetCallback();
+	}
+#endif
+}
+
+void UAnimBankMeshComponent::BeginDestroy()
+{
+	Super::BeginDestroy();
+}
+
+void UAnimBankMeshComponent::OnRegister()
+{
+	Super::OnRegister();
+}
+
+void UAnimBankMeshComponent::TickComponent(float DeltaTime, ELevelTick TickType, FActorComponentTickFunction* ThisTickFunction)
+{
+	Super::TickComponent(DeltaTime, TickType, ThisTickFunction);
+}
+
+bool UAnimBankMeshComponent::IsEnabled() const
+{
+	USkeletalMesh* SkeletalMeshPtr = Cast<USkeletalMesh>(GetSkinnedAsset());
+	if (SkeletalMeshPtr && SkeletalMeshPtr->GetResourceForRendering())
+	{
+		return GetNumInstances() > 0;
+	}
+
+	return false;
+}
+
+int32 UAnimBankMeshComponent::GetNumInstances() const
+{
+	return InstanceData.Num();
+}
+
+void UAnimBankMeshComponent::ApplyInheritedPerInstanceData(const UAnimBankMeshComponent* InArchetype)
+{
+	check(InArchetype);
+	InstanceData = InArchetype->InstanceData;
+	InstanceCustomData = InArchetype->InstanceCustomData;
+	NumCustomDataFloats = InArchetype->NumCustomDataFloats;
+}
+
+bool UAnimBankMeshComponent::ShouldInheritPerInstanceData() const
+{
+	return ShouldInheritPerInstanceData(Cast<UAnimBankMeshComponent>(GetArchetype()));
+}
+
+bool UAnimBankMeshComponent::ShouldInheritPerInstanceData(const UAnimBankMeshComponent* InArchetype) const
+{
+	return (bInheritPerInstanceData || !bEditableWhenInherited) && InArchetype && InArchetype->IsInBlueprint() && !IsTemplate();
+}
+
+void UAnimBankMeshComponent::SetupNewInstanceData(FAnimBankMeshInstanceData& InOutNewInstanceData, int32 InInstanceIndex, const FTransform& InInstanceTransform, int32 InBankIndex)
+{
+	InOutNewInstanceData.Transform = InInstanceTransform.ToMatrixWithScale();
+	InOutNewInstanceData.BankIndex = InBankIndex;
+
+	if (bPhysicsStateCreated)
+	{
+		// ..
+	}
+}
+
+const Nanite::FResources* UAnimBankMeshComponent::GetNaniteResources() const
+{
+	return Super::GetNaniteResources();
+}
+
+void UAnimBankMeshComponent::CreateRenderState_Concurrent(FRegisterComponentContext* Context)
+{
+	MeshObjectFactory = &CreateAnimBankMeshObjectFn;
+	Super::CreateRenderState_Concurrent(Context);
+}
+
+void UAnimBankMeshComponent::DestroyRenderState_Concurrent()
+{
+	Super::DestroyRenderState_Concurrent();
+}
+
+FPrimitiveSceneProxy* UAnimBankMeshComponent::CreateSceneProxy()
+{
+	LLM_SCOPE(ELLMTag::SkeletalMesh);
+
+	ERHIFeatureLevel::Type SceneFeatureLevel = GetWorld()->GetFeatureLevel();
+	FPrimitiveSceneProxy* Result = nullptr;
+	FSkeletalMeshRenderData* SkelMeshRenderData = GetSkeletalMeshRenderData();
+
+#if WITH_EDITOR
+	if (!bIsInstanceDataApplyCompleted)
+	{
+		return nullptr;
+	}
+#endif
+
+	const USkinnedAsset* SkinnedAssetPtr = GetSkinnedAsset();
+	if (InstanceData.Num() == 0 || SkinnedAssetPtr == nullptr || SkinnedAssetPtr->IsCompiling())
+	{
+		return nullptr;
+	}
+
+	if (CheckPSOPrecachingAndBoostPriority() && GetPSOPrecacheProxyCreationStrategy() == EPSOPrecacheProxyCreationStrategy::DelayUntilPSOPrecached)
+	{
+		UE_LOG(LogAnimBank, Verbose, TEXT("Skipping CreateSceneProxy for UAnimBankMeshComponent %s (UAnimBankMeshComponent PSOs are still compiling)"), *GetFullName());
+		return nullptr;
+	}
+
+	// Only create a scene proxy for rendering if properly initialized
+	if (SkelMeshRenderData &&
+		SkelMeshRenderData->LODRenderData.IsValidIndex(GetPredictedLODLevel()) &&
+		!bHideSkin &&
+		MeshObject)
+	{
+		// Only create a scene proxy if the bone count being used is supported, or if we don't have a skeleton (this is the case with destructibles)
+		int32 MinLODIndex = ComputeMinLOD();
+		int32 MaxBonesPerChunk = SkelMeshRenderData->GetMaxBonesPerSection(MinLODIndex);
+		int32 MaxSupportedNumBones = MeshObject->IsCPUSkinned() ? MAX_int32 : FGPUBaseSkinVertexFactory::GetMaxGPUSkinBones();
+		if (MaxBonesPerChunk <= MaxSupportedNumBones)
+		{
+			if (ShouldNaniteSkin() && IsEnabled())
+			{
+				Nanite::FMaterialAudit NaniteMaterials{};
+				AuditMaterials(this, NaniteMaterials, true /* Set material usage flags */);
+
+				const bool bIsMaskingAllowed = Nanite::IsMaskingAllowed(GetWorld(), false /* force Nanite for masked */);
+				if (NaniteMaterials.IsValid(bIsMaskingAllowed))
+				{
+					Result = ::new FAnimBankMeshSceneProxy(NaniteMaterials, this, SkelMeshRenderData);
+				}
+			}
+
+			if (Result == nullptr)
+			{
+				Result = Super::CreateSceneProxy();
+			}
+		}
+	}
+
+#if !(UE_BUILD_SHIPPING || UE_BUILD_TEST)
+	SendRenderDebugPhysics(Result);
+#endif
+
+	return Result;
+}
+
+#if WITH_EDITOR
+
+void UAnimBankMeshComponent::PostEditChangeChainProperty(FPropertyChangedChainEvent& PropertyChangedEvent)
+{
+	if (PropertyChangedEvent.Property != nullptr)
+	{
+		if (PropertyChangedEvent.Property->GetFName() == GET_MEMBER_NAME_CHECKED(UAnimBankMeshComponent, InstanceData))
+		{
+			if (PropertyChangedEvent.ChangeType == EPropertyChangeType::ArrayAdd
+				|| PropertyChangedEvent.ChangeType == EPropertyChangeType::Duplicate)
+			{
+				int32 AddedAtIndex = PropertyChangedEvent.GetArrayIndex(PropertyChangedEvent.Property->GetFName().ToString());
+				check(AddedAtIndex != INDEX_NONE);
+
+				AddInstanceInternal(
+					AddedAtIndex,
+					&InstanceData[AddedAtIndex],
+					PropertyChangedEvent.ChangeType == EPropertyChangeType::ArrayAdd ? FTransform::Identity : FTransform(InstanceData[AddedAtIndex].Transform),
+					PropertyChangedEvent.ChangeType == EPropertyChangeType::ArrayAdd ? 0 : InstanceData[AddedAtIndex].BankIndex,
+					/*bWorldSpace*/false
+				);
+
+				// added via the property editor, so we will want to interactively work with instances
+				//bHasPerInstanceHitProxies = true;
+			}
+			else if (PropertyChangedEvent.ChangeType == EPropertyChangeType::ArrayRemove)
+			{
+				int32 RemovedAtIndex = PropertyChangedEvent.GetArrayIndex(PropertyChangedEvent.Property->GetFName().ToString());
+				check(RemovedAtIndex != INDEX_NONE);
+
+				RemoveInstanceInternal(RemovedAtIndex, true);
+			}
+			else if (PropertyChangedEvent.ChangeType == EPropertyChangeType::ArrayClear)
+			{
+				ClearInstances();
+			}
+			else if (PropertyChangedEvent.ChangeType == EPropertyChangeType::ValueSet)
+			{
+			}
+			
+			MarkRenderStateDirty();
+		}
+		else if (PropertyChangedEvent.Property->GetFName() == GET_MEMBER_NAME_CHECKED(FAnimBankMeshInstanceData, Transform))
+		{
+
+		}
+		else if (PropertyChangedEvent.Property->GetFName() == "NumCustomDataFloats")
+		{
+			// Can't just call SetNumCustomDataFloats because it doesn't do anything if the value is the same, and the editor has already modified the value...
+			NumCustomDataFloats = FMath::Max(NumCustomDataFloats, 0);
+
+			// Clear out and reinit to 0
+			InstanceCustomData.Empty(InstanceData.Num() * NumCustomDataFloats);
+			InstanceCustomData.SetNumZeroed(InstanceData.Num() * NumCustomDataFloats);
+		}
+		else if (PropertyChangedEvent.PropertyChain.GetActiveMemberNode()->GetValue()->GetFName() == "InstanceCustomData")
+		{
+			int32 ChangedCustomValueIndex = PropertyChangedEvent.GetArrayIndex(PropertyChangedEvent.Property->GetFName().ToString());
+			if (ensure(NumCustomDataFloats > 0))
+			{
+				int32 InstanceIndex = ChangedCustomValueIndex / NumCustomDataFloats;
+			}
+		}
+		else if (
+			(PropertyChangedEvent.Property->GetFName() == GET_MEMBER_NAME_CHECKED(UAnimBankMeshComponent, AnimBankItems)) ||
+			(PropertyChangedEvent.Property->GetFName() == GET_MEMBER_NAME_CHECKED(FAnimBankItem, BankAsset))
+			)
+		{
+			for (FAnimBankItem& BankItem : AnimBankItems)
+			{
+				// TODO: BankItem.ValidatePosition();
+
+				// Make sure the animation skeleton is valid
+				bool bInvalid = true;
+				if (BankItem.BankAsset && BankItem.BankAsset->Asset && BankItem.BankAsset->Asset->GetSkeleton())
+				{
+					USkeletalMesh* SkeletalMeshPtr = Cast<USkeletalMesh>(GetSkinnedAsset());
+					if (SkeletalMeshPtr)
+					{
+						// Make sure the skeletons match!
+						if (BankItem.BankAsset->Asset->GetSkeleton() == SkeletalMeshPtr->GetSkeleton())
+						{
+							bInvalid = false;
+						}
+					}
+				}
+
+				if (bInvalid)
+				{
+					UE_LOG(LogAnimation, Warning, TEXT("Invalid animation skeleton"));
+					BankItem.BankAsset = nullptr;
+				}
+			}
+
+			MarkRenderStateDirty();
+		}
+	}
+	Super::PostEditChangeChainProperty(PropertyChangedEvent);
+}
+
+void UAnimBankMeshComponent::PostEditUndo()
+{
+	Super::PostEditUndo();
+	MarkRenderStateDirty();
+}
+
+#endif
+
+TStructOnScope<FActorComponentInstanceData> UAnimBankMeshComponent::GetComponentInstanceData() const
+{
+	TStructOnScope<FActorComponentInstanceData> ComponentInstanceData;
+#if WITH_EDITOR
+	ComponentInstanceData.InitializeAs<FAnimBankMeshComponentInstanceData>(this);
+	FAnimBankMeshComponentInstanceData* SkinnedMeshInstanceData = ComponentInstanceData.Cast<FAnimBankMeshComponentInstanceData>();
+
+	// Back up per-instance info (this is strictly for Comparison in UAnimBankMeshComponent::ApplyComponentInstanceData 
+	// as this Property will get serialized by base class FActorComponentInstanceData through FComponentPropertyWriter which uses the PPF_ForceTaggedSerialization to backup all properties even the custom serialized ones
+	SkinnedMeshInstanceData->InstanceData = InstanceData;
+
+	// Back up instance selection
+	SkinnedMeshInstanceData->SelectedInstances = SelectedInstances;
+
+	// Back up per-instance hit proxies
+	SkinnedMeshInstanceData->bHasPerInstanceHitProxies = bHasPerInstanceHitProxies;
+#endif
+	return ComponentInstanceData;
+}
+
+void UAnimBankMeshComponent::PreApplyComponentInstanceData(struct FAnimBankMeshComponentInstanceData* InstancedMeshData)
+{
+#if WITH_EDITOR
+	// Prevent proxy recreate while traversing the ::ApplyToComponent stack
+	bIsInstanceDataApplyCompleted = false;
+#endif
+}
+
+void UAnimBankMeshComponent::ApplyComponentInstanceData(struct FAnimBankMeshComponentInstanceData* InstancedMeshData)
+{
+#if WITH_EDITOR
+	check(InstancedMeshData);
+
+	ON_SCOPE_EXIT
+	{
+		bIsInstanceDataApplyCompleted = true;
+	};
+
+	if (GetSkinnedAsset() != InstancedMeshData->SkinnedAsset)
+	{
+		return;
+	}
+
+	// If we should inherit from archetype do it here after data was applied and before comparing (RerunConstructionScript will serialize SkipSerialization properties and reapply them even if we want to inherit them)
+	const UAnimBankMeshComponent* Archetype = Cast<UAnimBankMeshComponent>(GetArchetype());
+	if (ShouldInheritPerInstanceData(Archetype))
+	{
+		ApplyInheritedPerInstanceData(Archetype);
+	}
+
+	SelectedInstances = InstancedMeshData->SelectedInstances;
+	bHasPerInstanceHitProxies = InstancedMeshData->bHasPerInstanceHitProxies;
+	bIsInstanceDataApplyCompleted = true;
+#endif
+}
+
+FBoxSphereBounds UAnimBankMeshComponent::CalcBounds(const FTransform& LocalToWorld) const
+{
+	//SCOPE_CYCLE_COUNTER(STAT_CalcSkelMeshBounds);
+
+	const USkinnedAsset* SkinnedAssetPtr = GetSkinnedAsset();
+	if (SkinnedAssetPtr && InstanceData.Num() > 0)
+	{
+		const FMatrix BoundTransformMatrix = LocalToWorld.ToMatrixWithScale();
+
+		FBoxSphereBounds::Builder BoundsBuilder;
+
+		const bool bUseSampledBounds = CVarAnimBankSampledBounds.GetValueOnAnyThread() != 0;
+		if (bUseSampledBounds && AnimBankItems.Num() > 0)
+		{
+			// Trade per sequence bounds (tigher fitting) for faster builds with high instance counts.
+			const bool bFastBuild = false;
+			if (bFastBuild)
+			{
+				FBox MergedBounds;
+
+				for (const FAnimBankItem& BankItem : AnimBankItems)
+				{
+					if (BankItem.BankAsset != nullptr)
+					{
+						const FAnimBankData& BankData = BankItem.BankAsset->GetData();
+						for (const FAnimBankEntry& BankEntry : BankData.Entries)
+						{
+							MergedBounds += BankEntry.SampledBounds.GetBox();
+						}
+					}
+				}
+
+				if (MergedBounds.IsValid)
+				{
+					for (int32 InstanceIndex = 0; InstanceIndex < InstanceData.Num(); InstanceIndex++)
+					{
+						const FAnimBankMeshInstanceData& Instance = InstanceData[InstanceIndex];
+						BoundsBuilder += MergedBounds.TransformBy(Instance.Transform * BoundTransformMatrix);
+					}
+				}
+			}
+			else
+			{
+				for (int32 InstanceIndex = 0; InstanceIndex < InstanceData.Num(); InstanceIndex++)
+				{
+					const FAnimBankMeshInstanceData& Instance = InstanceData[InstanceIndex];
+					if (Instance.BankIndex < uint32(AnimBankItems.Num()))
+					{
+						const FAnimBankItem& BankItem = AnimBankItems[Instance.BankIndex];
+						if (BankItem.BankAsset != nullptr)
+						{
+							const FAnimBankData& BankData = BankItem.BankAsset->GetData();
+							if (BankItem.SequenceIndex < BankData.Entries.Num())
+							{
+								const FBox BankBounds = BankData.Entries[BankItem.SequenceIndex].SampledBounds.GetBox();
+								BoundsBuilder += BankBounds.TransformBy(Instance.Transform * BoundTransformMatrix);
+							}
+						}
+					}
+				}
+			}
+
+			return BoundsBuilder;
+		}
+		else
+		{
+			const FBox InstanceBounds = SkinnedAssetPtr->GetBounds().GetBox();
+			if (InstanceBounds.IsValid)
+			{
+				for (int32 InstanceIndex = 0; InstanceIndex < InstanceData.Num(); InstanceIndex++)
+				{
+					BoundsBuilder += InstanceBounds.TransformBy(InstanceData[InstanceIndex].Transform * BoundTransformMatrix);
+				}
+
+				return BoundsBuilder;
+			}
+		}
+	}
+
+	return CalcMeshBound(FVector3f::ZeroVector, false, LocalToWorld);
+}
+
+void UAnimBankMeshComponent::SetSkinnedAssetCallback()
+{
+	MarkRenderStateDirty();
+}
+
+void UAnimBankMeshComponent::RefreshBoneTransforms(FActorComponentTickFunction* TickFunction)
+{
+	//SCOPE_CYCLE_COUNTER(STAT_RefreshBoneTransforms);
+
+	// Can't do anything without a SkinnedAsset
+	if (!GetSkinnedAsset())
+	{
+		return;
+	}
+
+	// Do nothing more if no bones in skeleton.
+	if (GetNumComponentSpaceTransforms() == 0)
+	{
+		return;
+	}
+
+	// We need the mesh space bone transforms now for renderer to get delta from ref pose:
+	//FillComponentSpaceTransforms();
+	//FinalizeBoneTransform();
+
+	//UpdateChildTransforms();
+	UpdateBounds();
+	MarkRenderTransformDirty();
+	MarkRenderDynamicDataDirty();
+
+	//bNeedsRefreshTransform = false;
+}
+
+int32 UAnimBankMeshComponent::AddInstance(const FTransform& InstanceTransform, int32 BankIndex, bool bWorldSpace)
+{
+	return AddInstanceInternal(InstanceData.Num(), nullptr, InstanceTransform, BankIndex, bWorldSpace);
+}
+
+TArray<int32> UAnimBankMeshComponent::AddInstances(const TArray<FTransform>& InstanceTransforms, bool bShouldReturnIndices, bool bWorldSpace, bool bUpdateNavigation)
+{
+	return AddInstancesInternal(InstanceTransforms, bShouldReturnIndices, bWorldSpace, bUpdateNavigation);
+}
+
+bool UAnimBankMeshComponent::SetCustomDataValue(int32 InstanceIndex, int32 CustomDataIndex, float CustomDataValue, bool bMarkRenderStateDirty)
+{
+	if (!InstanceData.IsValidIndex(InstanceIndex) || CustomDataIndex < 0 || CustomDataIndex >= NumCustomDataFloats)
+	{
+		return false;
+	}
+
+	Modify();
+
+	InstanceCustomData[InstanceIndex * NumCustomDataFloats + CustomDataIndex] = CustomDataValue;
+
+	if (bMarkRenderStateDirty)
+	{
+		MarkRenderStateDirty();
+	}
+
+	return true;
+}
+
+void UAnimBankMeshComponent::SetNumCustomDataFloats(int32 InNumCustomDataFloats)
+{
+	if (FMath::Max(InNumCustomDataFloats, 0) != NumCustomDataFloats)
+	{
+		NumCustomDataFloats = FMath::Max(InNumCustomDataFloats, 0);
+
+		// Clear out and reinit to 0
+		InstanceCustomData.Empty(InstanceData.Num() * NumCustomDataFloats);
+		InstanceCustomData.SetNumZeroed(InstanceData.Num() * NumCustomDataFloats);
+	}
+}
+
+bool UAnimBankMeshComponent::SetCustomData(int32 InstanceIndex, TArrayView<const float> CustomDataFloats, bool bMarkRenderStateDirty)
+{
+	if (!InstanceData.IsValidIndex(InstanceIndex) || CustomDataFloats.Num() == 0)
+	{
+		return false;
+	}
+
+	if (bMarkRenderStateDirty)
+	{
+		Modify();
+	}
+
+	const int32 NumToCopy = FMath::Min(CustomDataFloats.Num(), NumCustomDataFloats);
+	FMemory::Memcpy(&InstanceCustomData[InstanceIndex * NumCustomDataFloats], CustomDataFloats.GetData(), NumToCopy * CustomDataFloats.GetTypeSize());
+
+	if (bMarkRenderStateDirty)
+	{
+		MarkRenderStateDirty();
+	}
+
+	return true;
+}
+
+void UAnimBankMeshComponent::PreAllocateInstancesMemory(int32 AddedInstanceCount)
+{
+	InstanceData.Reserve(InstanceData.Num() + AddedInstanceCount);
+	InstanceCustomData.Reserve(InstanceCustomData.Num() + AddedInstanceCount * NumCustomDataFloats);
+}
+
+bool UAnimBankMeshComponent::GetInstanceTransform(int32 InstanceIndex, FTransform& OutInstanceTransform, bool bWorldSpace) const
+{
+	if (!InstanceData.IsValidIndex(InstanceIndex))
+	{
+		return false;
+	}
+
+	const FAnimBankMeshInstanceData& Instance = InstanceData[InstanceIndex];
+
+	OutInstanceTransform = FTransform(Instance.Transform);
+	if (bWorldSpace)
+	{
+		OutInstanceTransform = OutInstanceTransform * GetComponentTransform();
+	}
+
+	return true;
+}
+
+bool UAnimBankMeshComponent::RemoveInstance(int32 InstanceIndex)
+{
+	return RemoveInstanceInternal(InstanceIndex, false);
+}
+
+bool UAnimBankMeshComponent::RemoveInstances(const TArray<int32>& InstancesToRemove)
+{
+	return RemoveInstances(InstancesToRemove, false /*bInstanceArrayAlreadySortedInReverseOrder*/);
+}
+
+bool UAnimBankMeshComponent::RemoveInstances(const TArray<int32>& InstancesToRemove, bool bInstanceArrayAlreadySortedInReverseOrder)
+{
+	if (InstancesToRemove.IsEmpty())
+	{
+		return false;
+	}
+	
+	auto RemoveInstancedWithSortedArray = [this](const TArray<int32>& SortedInstancesToRemove) -> bool
+	{
+		if (!InstanceData.IsValidIndex(SortedInstancesToRemove[0]) || !InstanceData.IsValidIndex(SortedInstancesToRemove.Last()))
+		{
+			return false;
+		}
+
+		for (const int32 InstanceIndex : SortedInstancesToRemove)
+		{
+			RemoveInstanceInternal(InstanceIndex, false);
+		}
+
+		return true;
+	};
+
+	bool bSuccess = false;
+	if (bInstanceArrayAlreadySortedInReverseOrder)
+	{
+		bSuccess = RemoveInstancedWithSortedArray(InstancesToRemove);
+	}
+	else
+	{
+		// Sort so Remove doesn't alter the indices of items still to remove
+		TArray<int32> SortedInstancesToRemove = InstancesToRemove;
+		SortedInstancesToRemove.Sort(TGreater<int32>());
+		bSuccess = RemoveInstancedWithSortedArray(SortedInstancesToRemove);
+	}
+
+	return bSuccess;
+}
+
+void UAnimBankMeshComponent::ClearInstances()
+{
+#if WITH_EDITOR
+	DeletionState = EInstanceDeletionReason::Clearing;
+#endif
+
+	const int32 PrevNumInstances = GetNumInstances();
+
+	// Clear all the per-instance data
+	InstanceData.Empty();
+	InstanceCustomData.Empty();
+
+	//ProxySize = 0;
+
+#if WITH_EDITOR
+	DeletionState = EInstanceDeletionReason::NotDeleting;
+#endif
+}
+
+int32 UAnimBankMeshComponent::AddAnimBankItem(const FAnimBankItem& BankItem)
+{
+	int32 ItemIndex = AnimBankItems.Num();
+	AnimBankItems.Emplace(BankItem);
+	return ItemIndex;
+}
+
+struct HAnimBankMeshInstance : public HHitProxy
+{
+	TObjectPtr<UAnimBankMeshComponent> Component;
+	int32 InstanceIndex;
+
+	DECLARE_HIT_PROXY(ENGINE_API);
+	HAnimBankMeshInstance(UAnimBankMeshComponent* InComponent, int32 InInstanceIndex) : HHitProxy(HPP_World), Component(InComponent), InstanceIndex(InInstanceIndex) {}
+
+	virtual void AddReferencedObjects(FReferenceCollector& Collector) override
+	{
+		Collector.AddReferencedObject(Component);
+	}
+
+	virtual FTypedElementHandle GetElementHandle() const override
+	{
+	#if WITH_EDITOR
+		if (Component)
+		{
+		#if 0
+			if (true)//if (CVarEnableViewportSMInstanceSelection.GetValueOnAnyThread() != 0)
+			{
+				// Prefer per-instance selection if available
+				// This may fail to return a handle if the feature is disabled, or if per-instance editing is disabled for this component
+				if (FTypedElementHandle ElementHandle = UEngineElementsLibrary::AcquireEditorSMInstanceElementHandle(Component, InstanceIndex))
+				{
+					return ElementHandle;
+				}
+			}
+		#endif
+
+			// If per-instance selection isn't possible, fallback to general per-component selection (which may choose to select the owner actor instead)
+			return UEngineElementsLibrary::AcquireEditorComponentElementHandle(Component);
+		}
+	#endif	// WITH_EDITOR
+		return FTypedElementHandle();
+	}
+
+	virtual EMouseCursor::Type GetMouseCursor() override
+	{
+		return EMouseCursor::Crosshairs;
+	}
+};
+
+IMPLEMENT_HIT_PROXY(HAnimBankMeshInstance, HHitProxy);
+
+void UAnimBankMeshComponent::CreateHitProxyData(TArray<TRefCountPtr<HHitProxy>>& HitProxies)
+{
+	if (GIsEditor && bHasPerInstanceHitProxies && (CreationMethod == EComponentCreationMethod::Instance))
+	{
+		int32 NumProxies = InstanceData.Num();
+		HitProxies.Empty(NumProxies);
+
+		for (int32 InstanceIdx = 0; InstanceIdx < NumProxies; ++InstanceIdx)
+		{
+			HitProxies.Add(new HAnimBankMeshInstance(this, InstanceIdx));
+		}
+	}
+	else
+	{
+		HitProxies.Empty();
+	}
+}
+
+int32 UAnimBankMeshComponent::AddInstanceInternal(int32 InstanceIndex, FAnimBankMeshInstanceData* InNewInstanceData, const FTransform& InstanceTransform, int32 BankIndex, bool bWorldSpace)
+{
+	FAnimBankMeshInstanceData* NewInstanceData = InNewInstanceData;
+
+	if (NewInstanceData == nullptr)
+	{
+		NewInstanceData = &InstanceData.AddDefaulted_GetRef();
+	}
+
+	const FTransform LocalTransform = bWorldSpace ? InstanceTransform.GetRelativeTransform(GetComponentTransform()) : InstanceTransform;
+	SetupNewInstanceData(*NewInstanceData, InstanceIndex, LocalTransform, BankIndex);
+
+	// Add custom data to instance
+	InstanceCustomData.AddZeroed(NumCustomDataFloats);
+
+#if WITH_EDITOR
+	if (SelectedInstances.Num())
+	{
+		SelectedInstances.Add(false);
+	}
+#endif
+
+	return InstanceIndex;
+}
+
+TArray<int32> UAnimBankMeshComponent::AddInstancesInternal(TConstArrayView<FTransform> InstanceTransforms, bool bShouldReturnIndices, bool bWorldSpace, bool bUpdateNavigation)
+{
+	const int32 Count = InstanceTransforms.Num();
+
+	TArray<int32> NewInstanceIndices;
+
+	if (bShouldReturnIndices)
+	{
+		NewInstanceIndices.Reserve(Count);
+	}
+
+	int32 InstanceIndex = InstanceData.Num();
+
+	InstanceCustomData.AddZeroed(NumCustomDataFloats * Count);
+
+#if WITH_EDITOR
+	SelectedInstances.Add(false, Count);
+#endif
+
+	const int32 NumInstances = InstanceIndex + Count;
+	InstanceData.Reserve(NumInstances);
+
+	for (const FTransform& InstanceTransform : InstanceTransforms)
+	{
+		FAnimBankMeshInstanceData& NewInstanceData = InstanceData.AddDefaulted_GetRef();
+
+		int32 BankIndex = 0; // TODO
+
+		const FTransform LocalTransform = bWorldSpace ? InstanceTransform.GetRelativeTransform(GetComponentTransform()) : InstanceTransform;
+		SetupNewInstanceData(NewInstanceData, InstanceIndex, LocalTransform, BankIndex);
+
+		if (bShouldReturnIndices)
+		{
+			NewInstanceIndices.Add(InstanceIndex);
+		}
+
+		++InstanceIndex;
+	}
+
+	return NewInstanceIndices;
+}
+
+bool UAnimBankMeshComponent::RemoveInstanceInternal(int32 InstanceIndex, bool InstanceAlreadyRemoved, bool bForceRemoveAtSwap, bool bUpdateNavigation)
+{
+#if WITH_EDITOR
+	DeletionState = InstanceAlreadyRemoved ? EInstanceDeletionReason::EntryAlreadyRemoved : EInstanceDeletionReason::EntryRemoval;
+#endif
+
+	const bool bUseRemoveAtSwap = true;
+
+	// remove instance
+	if (!InstanceAlreadyRemoved && InstanceData.IsValidIndex(InstanceIndex))
+	{
+		if (bUseRemoveAtSwap)
+		{
+			InstanceData.RemoveAtSwap(InstanceIndex, EAllowShrinking::No);
+			InstanceCustomData.RemoveAtSwap(InstanceIndex * NumCustomDataFloats, NumCustomDataFloats, EAllowShrinking::No);
+		}
+		else
+		{
+			InstanceData.RemoveAt(InstanceIndex);
+			InstanceCustomData.RemoveAt(InstanceIndex * NumCustomDataFloats, NumCustomDataFloats);
+		}
+	}
+
+#if WITH_EDITOR
+	// remove selection flag if array is filled in
+	if (SelectedInstances.IsValidIndex(InstanceIndex))
+	{
+		if (bUseRemoveAtSwap)
+		{
+			SelectedInstances.RemoveAtSwap(InstanceIndex);
+		}
+		else
+		{
+			SelectedInstances.RemoveAt(InstanceIndex);
+		}
+	}
+#endif
+
+	const int32 LastInstanceIndex = InstanceData.Num();
+
+#if WITH_EDITOR
+	DeletionState = EInstanceDeletionReason::NotDeleting;
+#endif
+	return true;
+}
diff --git a/Engine/Source/Runtime/Engine/Public/SceneInterface.h b/Engine/Source/Runtime/Engine/Public/SceneInterface.h
--- a/Engine/Source/Runtime/Engine/Public/SceneInterface.h
+++ b/Engine/Source/Runtime/Engine/Public/SceneInterface.h
@@ -13,6 +13,8 @@
 #include "ShowFlags.h"
 
 class AWorldSettings;
+struct FAnimBankDesc;
+struct FAnimBankRecordHandle;
 class FArchive;
 class FFloat16Color;
 class FInstanceCullingManager;
@@ -391,6 +393,9 @@ public:
 	virtual void AddSparseVolumeTextureViewer(FSparseVolumeTextureViewerSceneProxy* SVTV) = 0;
 	virtual void RemoveSparseVolumeTextureViewer(FSparseVolumeTextureViewerSceneProxy* SVTV) = 0;
 
+	virtual TArray<FAnimBankRecordHandle> RegisterAnimBank(const TConstArrayView<FAnimBankDesc> Descs) = 0;
+	virtual void UnregisterAnimBank(const TConstArrayView<FAnimBankRecordHandle> Handles) = 0;
+
 	/**
 	 * Adds the unique volumetric cloud component to the scene
 	 *
diff --git a/Engine/Source/Runtime/Renderer/Private/RendererScene.cpp b/Engine/Source/Runtime/Renderer/Private/RendererScene.cpp
--- a/Engine/Source/Runtime/Renderer/Private/RendererScene.cpp
+++ b/Engine/Source/Runtime/Renderer/Private/RendererScene.cpp
@@ -90,6 +90,9 @@
 #include "Nanite/NaniteSkinningSceneExtension.h"
 #include "ObjectCacheContext.h"
 #include "SceneComputeUpdates.h"
+#include "Animation/AnimBank.h"
+#include "Skinning/AnimBankTransformProvider.h"
+#include "Skinning/SkinningTransformProvider.h"
 
 #if RHI_RAYTRACING
 #include "Nanite/NaniteRayTracing.h"
@@ -3038,6 +3041,57 @@ void FScene::RemoveHairStrands(FHairStrandsInstance* Proxy)
 	}
 }
 
+TArray<FAnimBankRecordHandle> FScene::RegisterAnimBank(const TConstArrayView<FAnimBankDesc> Descs)
+{
+	check(IsInRenderingThread());
+
+	TArray<FAnimBankRecordHandle> Handles;
+	Handles.Reserve(Descs.Num());
+
+	FAnimBankTransformProvider* BankProvider = GetExtensionPtr<FAnimBankTransformProvider>();
+	if (BankProvider != nullptr)
+	{
+		for (const FAnimBankDesc& Desc : Descs)
+		{
+			if (!Desc.BankAsset.IsValid() || !Desc.Asset.IsValid())
+			{
+				Handles.Emplace(FAnimBankRecordHandle());
+				continue;
+			}
+
+			const FAnimBankData& BankData = Desc.BankAsset->GetData();
+
+			if (Desc.SequenceIndex >= uint32(BankData.Entries.Num()))
+			{
+				Handles.Emplace(FAnimBankRecordHandle());
+				continue;
+			}
+
+			Handles.Emplace(BankProvider->RegisterBank(Desc));
+		}
+	}
+
+	return Handles;
+}
+
+void FScene::UnregisterAnimBank(const TConstArrayView<FAnimBankRecordHandle> Handles)
+{
+	check(IsInRenderingThread());
+
+	FAnimBankTransformProvider* BankProvider = GetExtensionPtr<FAnimBankTransformProvider>();
+	if (BankProvider != nullptr)
+	{
+		for (const FAnimBankRecordHandle& Handle : Handles)
+		{
+			if (!Handle.IsValid())
+			{
+				continue;
+			}
+
+			BankProvider->UnregisterBank(Handle);
+		}
+	}
+}
 
 void FScene::GetLightIESAtlasSlot(const FLightSceneProxy* Proxy, FLightRenderParameters* Out)
 {
@@ -6804,7 +6858,7 @@ public:
 	FNULLSceneInterface(UWorld* InWorld, bool bCreateFXSystem )
 		:	FSceneInterface(GMaxRHIFeatureLevel)
 		,	World( InWorld )
-		,	FXSystem( NULL )
+		,	FXSystem( nullptr )
 	{
 		World->Scene = this;
 
@@ -6814,8 +6868,8 @@ public:
 		}
 		else
 		{
-			World->FXSystem = NULL;
-			SetFXSystem(NULL);
+			World->FXSystem = nullptr;
+			SetFXSystem(nullptr);
 		}
 	}
 
@@ -6868,12 +6922,15 @@ public:
 
 	virtual void AddSkyAtmosphere(FSkyAtmosphereSceneProxy* SkyAtmosphereSceneProxy, bool bStaticLightingBuilt) override {}
 	virtual void RemoveSkyAtmosphere(FSkyAtmosphereSceneProxy* SkyAtmosphereSceneProxy) override {}
-	virtual FSkyAtmosphereRenderSceneInfo* GetSkyAtmosphereSceneInfo() override { return NULL; }
-	virtual const FSkyAtmosphereRenderSceneInfo* GetSkyAtmosphereSceneInfo() const override { return NULL; }
+	virtual FSkyAtmosphereRenderSceneInfo* GetSkyAtmosphereSceneInfo() override { return nullptr; }
+	virtual const FSkyAtmosphereRenderSceneInfo* GetSkyAtmosphereSceneInfo() const override { return nullptr; }
 
 	virtual void AddSparseVolumeTextureViewer(FSparseVolumeTextureViewerSceneProxy* SVTV) override {}
 	virtual void RemoveSparseVolumeTextureViewer(FSparseVolumeTextureViewerSceneProxy* SVTV) override {}
 
+	virtual TArray<FAnimBankRecordHandle> RegisterAnimBank(const TConstArrayView<FAnimBankDesc> Descs) { return TArray<FAnimBankRecordHandle>(); }
+	virtual void UnregisterAnimBank(const TConstArrayView<FAnimBankRecordHandle> Handles) override {}
+
 	virtual void AddHairStrands(FHairStrandsInstance* Proxy) override {}
 	virtual void RemoveHairStrands(FHairStrandsInstance* Proxy) override {}
 	virtual void GetLightIESAtlasSlot(const FLightSceneProxy* Proxy, FLightRenderParameters* Out) override {}
diff --git a/Engine/Source/Runtime/Renderer/Private/ScenePrivate.h b/Engine/Source/Runtime/Renderer/Private/ScenePrivate.h
--- a/Engine/Source/Runtime/Renderer/Private/ScenePrivate.h
+++ b/Engine/Source/Runtime/Renderer/Private/ScenePrivate.h
@@ -3439,6 +3439,9 @@ public:
 	virtual void AddSparseVolumeTextureViewer(FSparseVolumeTextureViewerSceneProxy* SVTV) override;
 	virtual void RemoveSparseVolumeTextureViewer(FSparseVolumeTextureViewerSceneProxy* SVTV) override;
 
+	virtual TArray<FAnimBankRecordHandle> RegisterAnimBank(const TConstArrayView<struct FAnimBankDesc> Descs) override;
+	virtual void UnregisterAnimBank(const TConstArrayView<FAnimBankRecordHandle> Handles) override;
+
 	virtual void SetPhysicsField(class FPhysicsFieldSceneProxy* PhysicsFieldSceneProxy) override;
 	virtual void ResetPhysicsField() override;
 	virtual void ShowPhysicsField() override;
diff --git a/Engine/Source/Runtime/Renderer/Private/Skinning/AnimBankTransformProvider.cpp b/Engine/Source/Runtime/Renderer/Private/Skinning/AnimBankTransformProvider.cpp
--- a/Engine/Source/Runtime/Renderer/Private/Skinning/AnimBankTransformProvider.cpp
+++ b/Engine/Source/Runtime/Renderer/Private/Skinning/AnimBankTransformProvider.cpp
@@ -21,34 +21,31 @@
 #include "SkeletalRenderPublic.h"
 #include "SkinningDefinitions.h"
 
-static FGuid AnimBankGPUProviderId(0xA5C0027A, 0x8F884C7C, 0x9312F138, 0x71A9300F);
-static FGuid AnimBankCPUProviderId(0xE7D6173D, 0x246F431A, 0x912D384E, 0x156C0D2C);
+static FGuid AnimBankGPUProviderId(ANIM_BANK_GPU_TRANSFORM_PROVIDER_GUID);
+static FGuid AnimBankCPUProviderId(ANIM_BANK_CPU_TRANSFORM_PROVIDER_GUID);
 
 IMPLEMENT_SCENE_EXTENSION(FAnimBankTransformProvider);
 
 // Animation is always sampled at 30hz
 #define ANIM_BANK_SAMPLE_RATE 30.0f
 
-int32 GAnimBankInterp = 1;
-FAutoConsoleVariableRef CVarAnimBankInterp(
+static TAutoConsoleVariable<bool> CVarAnimBankInterp(
 	TEXT("r.AnimBank.Interpolation"),
-	GAnimBankInterp,
+	true,
 	TEXT(""),
 	ECVF_Scalability | ECVF_RenderThreadSafe
 );
 
-float GAnimBankTimeScale = 1.0f;
-FAutoConsoleVariableRef CVarAnimBankTimeScale(
+static TAutoConsoleVariable<float> CVarAnimBankTimeScale(
 	TEXT("r.AnimBank.TimeScale"),
-	GAnimBankTimeScale,
+	1.0f,
 	TEXT(""),
 	ECVF_Scalability | ECVF_RenderThreadSafe
 );
 
-int32 GAnimBankGPU = 1;
-FAutoConsoleVariableRef CVarAnimBankGPU(
+static TAutoConsoleVariable<bool> CVarAnimBankGPU(
 	TEXT("r.AnimBank.GPU"),
-	GAnimBankGPU,
+	true,
 	TEXT(""),
 	FConsoleVariableDelegate::CreateLambda([](IConsoleVariable* InVariable)
 	{
@@ -123,8 +120,6 @@ private:
 
 IMPLEMENT_GLOBAL_SHADER(FAnimBankScatterCS, "/Engine/Private/Skinning/AnimBankEval.usf", "BankScatterCS", SF_Compute);
 
-extern int32 GAnimBankInterp;
-
 bool FAnimBankTransformProvider::ShouldCreateExtension(FScene& InScene)
 {
 	return NaniteSkinnedMeshesSupported() && DoesRuntimeSupportNanite(GetFeatureLevelShaderPlatform(InScene.GetFeatureLevel()), true, true);
@@ -272,6 +267,8 @@ static FAnimBankGPUData BuildAnimBankGPUData(const FAnimBankRecordMap& BankRecor
 	BankData.bValid = BankData.BlockCount > 0;
 	if (BankData.bValid)
 	{
+		const bool bInterpolating = CVarAnimBankInterp.GetValueOnRenderThread();
+
 		// Transform Buffer
 		{
 			const uint32 TransformSize = sizeof(FMatrix3x4) * BankData.TransformCount;
@@ -319,7 +316,7 @@ static FAnimBankGPUData BuildAnimBankGPUData(const FAnimBankRecordMap& BankRecor
 				Header.PreviousTime		= BankRecord.PreviousTime;
 				Header.TransformOffset	= TransformOffset;
 				Header.Playing			= BankRecord.Playing;
-				Header.Interpolating	= !!GAnimBankInterp;
+				Header.Interpolating	= bInterpolating;
 				Header.HasScale			= false;
 				checkSlow(Header.TransformOffset == TransformOffset);
 
@@ -464,6 +461,8 @@ static FAnimBankCPUData BuildAnimBankCPUData(const FAnimBankRecordMap& BankRecor
 	BankData.bValid = BankData.RecordCount > 0;
 	if (BankData.bValid)
 	{
+		const bool bInterpolating = CVarAnimBankInterp.GetValueOnRenderThread();
+
 		FMatrix3x4* Transforms = GraphBuilder.AllocPODArray<FMatrix3x4>(BankData.TransformCount);
 
 		uint32 TransformOffset = 0;
@@ -493,7 +492,7 @@ static FAnimBankCPUData BuildAnimBankCPUData(const FAnimBankRecordMap& BankRecor
 			float Alpha;
 			FAnimationRuntime::GetKeyIndicesFromTime(KeyIndex0, KeyIndex1, Alpha, Time, BankRecord.FrameCount, TrackLength);
 
-			if (GAnimBankInterp == 0)
+			if (!bInterpolating)
 			{
 				// Forcing alpha to zero disables pose interpolation (interpolation method is "step")
 				Alpha = 0.0f;
@@ -559,7 +558,7 @@ static FAnimBankCPUData BuildAnimBankCPUData(const FAnimBankRecordMap& BankRecor
 
 void FAnimBankTransformProvider::AdvanceAnimation(FSkinningTransformProvider::FProviderContext& Context)
 {
-	const float GlobalTimeScale = GAnimBankTimeScale;
+	const float GlobalTimeScale = CVarAnimBankTimeScale.GetValueOnRenderThread();
 
 	for (auto& RecordPair : BankRecordMap)
 	{
@@ -615,7 +614,12 @@ void FAnimBankTransformProvider::ScatterAnimation(
 
 		for (const uint64 BankId : BankIds)
 		{
-			check(int32(BankId) < IdToOffsetMapping.Num());
+			if (BankId == ~uint64(0))
+			{
+				continue;
+			}
+
+			check(BankId < uint64(IdToOffsetMapping.Num()));
 			uint32 SrcTransformOffset = IdToOffsetMapping[int32(BankId)];
 
 			const bool bIsPending = (SrcTransformOffset == ~uint32(0));
@@ -740,7 +744,7 @@ void FAnimBankTransformProvider::ProvideCPUBankTransforms(FSkinningTransformProv
 
 const FSkinningTransformProvider::FProviderId& GetAnimBankProviderId()
 {
-	if (GAnimBankGPU != 0)
+	if (CVarAnimBankGPU.GetValueOnRenderThread())
 	{
 		return AnimBankGPUProviderId;
 	}
