Manually apply the backout diff from CL 37184770

diff --git a/Engine/Source/Runtime/Engine/Private/Components/AnimBankMeshComponent.cpp b/Engine/Source/Runtime/Engine/Private/Components/AnimBankMeshComponent.cpp
--- a/Engine/Source/Runtime/Engine/Private/Components/AnimBankMeshComponent.cpp
+++ b/Engine/Source/Runtime/Engine/Private/Components/AnimBankMeshComponent.cpp
@@ -1403,7 +1403,7 @@ IMPLEMENT_HIT_PROXY(HAnimBankMeshInstance, HHitProxy);
 
 void UAnimBankMeshComponent::CreateHitProxyData(TArray<TRefCountPtr<HHitProxy>>& HitProxies)
 {
-	if (GIsEditor && bHasPerInstanceHitProxies && (CreationMethod == EComponentCreationMethod::Instance))
+	if (GIsEditor && bHasPerInstanceHitProxies)
 	{
 		int32 NumProxies = InstanceData.Num();
 		HitProxies.Empty(NumProxies);
