Fixed race condition in calling GetAllocatedScenes from the render thread (in order to broadcast ReloadNaniteFixedFunctionBins), and prevented reoccurrence by checking it is only called from the render thread..

diff --git a/Engine/Source/Runtime/Engine/Private/Materials/MaterialShared.cpp b/Engine/Source/Runtime/Engine/Private/Materials/MaterialShared.cpp
--- a/Engine/Source/Runtime/Engine/Private/Materials/MaterialShared.cpp
+++ b/Engine/Source/Runtime/Engine/Private/Materials/MaterialShared.cpp
@@ -4846,15 +4846,14 @@ FMaterialUpdateContext::~FMaterialUpdateContext()
 		UpdateInstance(InstancesToUpdate.Last());
 	}
 	
-	ENQUEUE_RENDER_COMMAND(ReloadNaniteFixedFunctionBins)(
-		[](FRHICommandListImmediate& RHICmdList) mutable
+	for (FSceneInterface* Scene : GetRendererModule().GetAllocatedScenes())
+	{
+		ENQUEUE_RENDER_COMMAND(ReloadNaniteFixedFunctionBins)(
+		[Scene](FRHICommandListImmediate& RHICmdList) mutable
 		{
-			for (FSceneInterface* Scene : GetRendererModule().GetAllocatedScenes())
-			{
-				Scene->ReloadNaniteFixedFunctionBins();
-			}
-		}
-	);
+			Scene->ReloadNaniteFixedFunctionBins();
+		});
+	}
 
 	if (bUpdateStaticDrawLists)
 	{
diff --git a/Engine/Source/Runtime/Renderer/Private/RendererModule.h b/Engine/Source/Runtime/Renderer/Private/RendererModule.h
--- a/Engine/Source/Runtime/Renderer/Private/RendererModule.h
+++ b/Engine/Source/Runtime/Renderer/Private/RendererModule.h
@@ -76,6 +76,7 @@ public:
 
 	virtual const TSet<FSceneInterface*>& GetAllocatedScenes() override
 	{
+		check(IsInGameThread());
 		return AllocatedScenes;
 	}
 
