Optimized NamedValueArray union operations
This is heavily used in animation curve blending

diff --git a/Engine/Source/Runtime/Engine/Public/Animation/NamedValueArray.h b/Engine/Source/Runtime/Engine/Public/Animation/NamedValueArray.h
--- a/Engine/Source/Runtime/Engine/Public/Animation/NamedValueArray.h
+++ b/Engine/Source/Runtime/Engine/Public/Animation/NamedValueArray.h
@@ -263,17 +263,17 @@ struct FNamedValueArrayUtils
 		CURVE_PROFILE_CYCLE_COUNTER(FNamedValueArrayUtils_Union2Params);
 
 		// Check arrays are not overlapping
-		checkSlow((void*)&InOutValueArray0 != (void*)&InValueArray1);
+		check((void*)&InOutValueArray0 != (void*)&InValueArray1);
 
-		int32 NumElements0 = InOutValueArray0.Num();	// ValueArray1 elements remain constant, but ValueArray0 can have entries added.
+		const int32 NumElements0 = InOutValueArray0.Num();	// ValueArray1 elements remain constant, but ValueArray0 can have entries added.
 		const int32 NumElements1 = InValueArray1.Num();
 
 		// Early out if we have no elements to union
-		if(NumElements1 == 0)
+		if (NumElements1 == 0)
 		{
 			return;
 		}
-		
+
 		// Sort both input arrays if required
 		InOutValueArray0.SortElementsIfRequired();
 		InValueArray1.SortElementsIfRequired();
@@ -286,59 +286,86 @@ struct FNamedValueArrayUtils
 		ReserveSize += ReserveSize / 2;
 		InOutValueArray0.Reserve(ReserveSize);
 
-		int32 ElementIndex0 = 0;
-		int32 ElementIndex1 = 0;
+		// Use pointers to iterate as this uses fewer registers and this code is very hot
+		ElementTypeResult* RESTRICT ElementPtr0 = InOutValueArray0.Elements.GetData();
+		const ElementTypeParam* RESTRICT ElementPtr1 = InValueArray1.Elements.GetData();
+
+		const ElementTypeResult* RESTRICT ElementEndPtr0 = ElementPtr0 + NumElements0;
+		const ElementTypeParam* RESTRICT ElementEndPtr1 = ElementPtr1 + NumElements1;
+
+		// A default element we re-use when an element from one of the two inputs is missing
+		ElementTypeParam DefaultElement;
+
+		// When we reach the end of either input arrays, we stop the tape merge and copy what remains
+		bool bIsDone = ElementPtr0 == ElementEndPtr0 || ElementPtr1 == ElementEndPtr1;
 
 		// Perform dual-iteration on the two sorted arrays
-		while(true)
+		while (!bIsDone)
 		{
-			if(ElementIndex0 == NumElements0 && ElementIndex1 < NumElements1)
+			if (ElementPtr0->Name == ElementPtr1->Name)
 			{
-				// Reached end of ValueArray0 with remaining in ValueArray1, we can just copy the remainder of ValueArray1
-				InOutValueArray0.Elements.Reserve(InOutValueArray0.Elements.Num() + (NumElements1 - ElementIndex1));
-				for( ; ElementIndex1 < NumElements1; ++ElementIndex1)
-				{
-					const ElementTypeParam* RESTRICT Element1 = &InValueArray1.Elements[ElementIndex1];
-					ElementTypeResult* RESTRICT Element0 = &InOutValueArray0.Elements.AddDefaulted_GetRef();
-					Element0->Name = Element1->Name;
+				// Elements match, run predicate and increment both indices
+				InPredicate(*ElementPtr0, *ElementPtr1, UE::Anim::ENamedValueUnionFlags::BothArgsValid);
 
-					InPredicate(*Element0, *Element1, UE::Anim::ENamedValueUnionFlags::ValidArg1);
-				}
-				break;
-			}
+				++ElementPtr0;
+				++ElementPtr1;
 
-			if(ElementIndex1 == NumElements1 && ElementIndex0 <= NumElements0)
-			{
-				// Reached end of ValueArray1 with remaining in ValueArray0 (or reached the end of both), we can just early out
-				break;
+				bIsDone = ElementPtr0 == ElementEndPtr0 || ElementPtr1 == ElementEndPtr1;
 			}
-
-			ElementTypeResult* RESTRICT Element0 = &InOutValueArray0.Elements[ElementIndex0];
-			const ElementTypeParam* RESTRICT Element1 = &InValueArray1.Elements[ElementIndex1];
-			
-			if(Element0->Name == Element1->Name)
-			{
-				// Elements match, run predicate and increment both indices
-				InPredicate(*Element0, *Element1, UE::Anim::ENamedValueUnionFlags::BothArgsValid);
-				++ElementIndex0;
-				++ElementIndex1;
-			}
-			else if(Element0->Name.FastLess(Element1->Name))
+			else if (ElementPtr0->Name.FastLess(ElementPtr1->Name))
 			{
 				// ValueArray0 element is earlier, so run predicate with only ValueArray0 and increment ValueArray0
-				ElementTypeParam DefaultElement;
-				DefaultElement.Name = Element0->Name;
-				InPredicate(*Element0, DefaultElement, UE::Anim::ENamedValueUnionFlags::ValidArg0);
-				++ElementIndex0;
+				DefaultElement.Name = ElementPtr0->Name;
+
+				InPredicate(*ElementPtr0, DefaultElement, UE::Anim::ENamedValueUnionFlags::ValidArg0);
+
+				++ElementPtr0;
+
+				bIsDone = ElementPtr0 == ElementEndPtr0;
 			}
 			else
 			{
 				// ValueArray1 element is earlier, so add to ValueArray0, run predicate with only second and increment ValueArray1
-				ElementTypeResult* RESTRICT NewElement = &InOutValueArray0.Elements.InsertDefaulted_GetRef(ElementIndex0++); // increment beyond new element
-				NumElements0 = InOutValueArray0.Num();
-				NewElement->Name = Element1->Name;
-				InPredicate(*NewElement, *Element1, UE::Anim::ENamedValueUnionFlags::ValidArg1);
-				++ElementIndex1;
+				const int32 ElementIndex0 = UE_PTRDIFF_TO_INT32(ElementPtr0 - InOutValueArray0.Elements.GetData());
+				InOutValueArray0.Elements.InsertUninitialized(ElementIndex0);
+
+				// Refresh pointers since they might have changed
+				ElementPtr0 = InOutValueArray0.Elements.GetData() + ElementIndex0;
+				ElementEndPtr0 = InOutValueArray0.Elements.GetData() + InOutValueArray0.Elements.Num();
+
+				// We use placement new to make sure the constructor is inlined to reduce redundant work
+				new(ElementPtr0) ElementTypeResult();
+				ElementPtr0->Name = ElementPtr1->Name;
+
+				InPredicate(*ElementPtr0, *ElementPtr1, UE::Anim::ENamedValueUnionFlags::ValidArg1);
+
+				++ElementPtr0;	// Increment this as well since we've inserted
+				++ElementPtr1;
+
+				bIsDone = ElementPtr1 == ElementEndPtr1;
+			}
+		}
+
+		// Tape merge is done, copy anything that might be remaining
+		if (ElementPtr1 < ElementEndPtr1)
+		{
+			// Reached end of ValueArray0 with remaining in ValueArray1, we can just copy the remainder of ValueArray1
+			const int32 NumResults = InOutValueArray0.Elements.Num();
+			const int32 NumNewElements = (ElementEndPtr1 - ElementPtr1);
+			InOutValueArray0.Elements.Reserve(NumResults + NumNewElements);
+			InOutValueArray0.Elements.AddUninitialized(NumNewElements);
+
+			// Refresh pointers since they might have changed
+			ElementPtr0 = InOutValueArray0.Elements.GetData() + NumResults;
+			ElementEndPtr0 = ElementPtr0 + NumNewElements;
+
+			for (; ElementPtr1 < ElementEndPtr1; ++ElementPtr0, ++ElementPtr1)
+			{
+				// We use placement new to make sure the constructor is inlined to reduce redundant work
+				new(ElementPtr0) ElementTypeResult();
+				ElementPtr0->Name = ElementPtr1->Name;
+
+				InPredicate(*ElementPtr0, *ElementPtr1, UE::Anim::ENamedValueUnionFlags::ValidArg1);
 			}
 		}
 
@@ -386,9 +413,9 @@ struct FNamedValueArrayUtils
 		CURVE_PROFILE_CYCLE_COUNTER(FNamedValueArrayUtils_Union3Params);
 
 		// Check arrays are not overlapping
-		checkSlow((void*)&OutResultValueArray != (void*)&InValueArray0);
-		checkSlow((void*)&OutResultValueArray != (void*)&InValueArray1);
-		checkSlow((void*)&InValueArray0 != (void*)&InValueArray1);
+		check((void*)&OutResultValueArray != (void*)&InValueArray0);
+		check((void*)&OutResultValueArray != (void*)&InValueArray1);
+		check((void*)&InValueArray0 != (void*)&InValueArray1);
 
 		// Make sure result is clear
 		OutResultValueArray.Elements.Reset();
@@ -408,92 +435,112 @@ struct FNamedValueArrayUtils
 		ReserveSize += ReserveSize / 2;
 		OutResultValueArray.Reserve(ReserveSize);
 
-		int32 ElementIndex0 = 0;
-		int32 ElementIndex1 = 0;
+		// Use pointers to iterate as this uses fewer registers and this code is very hot
+		const ElementType0* RESTRICT ElementPtr0 = InValueArray0.Elements.GetData();
+		const ElementType1* RESTRICT ElementPtr1 = InValueArray1.Elements.GetData();
 
-		// Perform dual-iteration on the two sorted arrays
-		while(true)
-		{
-			if(ElementIndex0 == NumElements0 && ElementIndex1 < NumElements1)
-			{
-				// Reached end of ValueArray0 with remaining in ValueArray1, we can just copy the remainder of ValueArray1
-				const int32 NumResults = OutResultValueArray.Elements.Num();
-				OutResultValueArray.Elements.Reserve(OutResultValueArray.Elements.Num() + (NumElements1 - ElementIndex1));
-				for(int32 ResultIndex = NumResults; ElementIndex1 < NumElements1; ++ResultIndex, ++ElementIndex1)
-				{
-					const ElementType1* RESTRICT Element1 = &InValueArray1.Elements[ElementIndex1];
-					ElementType0 DefaultElement0;
-					DefaultElement0.Name = Element1->Name;
-					ElementTypeResult* RESTRICT ResultElement = &OutResultValueArray.Elements.AddDefaulted_GetRef();
-					ResultElement->Name = Element1->Name;
+		const ElementType0* RESTRICT ElementEndPtr0 = ElementPtr0 + NumElements0;
+		const ElementType1* RESTRICT ElementEndPtr1 = ElementPtr1 + NumElements1;
 
-					InPredicate(*ResultElement, DefaultElement0, *Element1, UE::Anim::ENamedValueUnionFlags::ValidArg1);
-				}
-				break;
-			}
+		// A default element we re-use when an element from one of the two inputs is missing
+		ElementType0 DefaultElement0;
+		ElementType1 DefaultElement1;
 
-			if(ElementIndex1 == NumElements1)
-			{
-				if(ElementIndex0 < NumElements0)
-				{
-					// Reached end of ValueArray0 with remaining in ValueArray0, we can just copy the remainder of ValueArray0
-					const int32 NumResults = OutResultValueArray.Elements.Num();
-					OutResultValueArray.Elements.Reserve(OutResultValueArray.Elements.Num() + (NumElements0 - ElementIndex0));
-					for(int32 ResultIndex = NumResults; ElementIndex0 < NumElements0; ++ResultIndex, ++ElementIndex0)
-					{
-						const ElementType0* RESTRICT Element0 = &InValueArray0.Elements[ElementIndex0];
-						ElementType1 DefaultElement1;
-						DefaultElement1.Name = Element0->Name;
-						ElementTypeResult* RESTRICT ResultElement = &OutResultValueArray.Elements.AddDefaulted_GetRef();
-						ResultElement->Name = Element0->Name;
-						
-						InPredicate(OutResultValueArray.Elements[ResultIndex], *Element0, DefaultElement1, UE::Anim::ENamedValueUnionFlags::ValidArg0);
-					}
-					break;
-				}
-				else if(ElementIndex0 == NumElements0)
-				{
-					// Reached end of both, exit
-					break;
-				}
-			}
+		// When we reach the end of either input arrays, we stop the tape merge and copy what remains
+		bool bIsDone = ElementPtr0 == ElementEndPtr0 || ElementPtr1 == ElementEndPtr1;
 
-			const ElementType0* RESTRICT Element0 = &InValueArray0.Elements[ElementIndex0];
-			const ElementType1* RESTRICT Element1 = &InValueArray1.Elements[ElementIndex1];
-			
-			if(Element0->Name == Element1->Name)
+		// Perform dual-iteration on the two sorted arrays
+		while (!bIsDone)
+		{
+			ElementTypeResult NewResultElement;
+
+			if (ElementPtr0->Name == ElementPtr1->Name)
 			{
 				// Elements match, run predicate and increment both indices
-				ElementTypeResult* RESTRICT NewResultElement = &OutResultValueArray.Elements.AddDefaulted_GetRef();
-				NewResultElement->Name = Element0->Name;
-				InPredicate(*NewResultElement, *Element0, *Element1, UE::Anim::ENamedValueUnionFlags::BothArgsValid);
+				NewResultElement.Name = ElementPtr0->Name;
 
-				++ElementIndex0;
-				++ElementIndex1;
+				InPredicate(NewResultElement, *ElementPtr0, *ElementPtr1, UE::Anim::ENamedValueUnionFlags::BothArgsValid);
+
+				++ElementPtr0;
+				++ElementPtr1;
+
+				bIsDone = ElementPtr0 == ElementEndPtr0 || ElementPtr1 == ElementEndPtr1;
 			}
-			else if(Element0->Name.FastLess(Element1->Name))
+			else if (ElementPtr0->Name.FastLess(ElementPtr1->Name))
 			{
 				// ValueArray0 element is earlier, so run predicate with only ValueArray0 and increment ValueArray0
-				ElementTypeResult* RESTRICT NewResultElement = &OutResultValueArray.Elements.AddDefaulted_GetRef();
-				NewResultElement->Name = Element0->Name;
+				NewResultElement.Name = ElementPtr0->Name;
 
-				ElementType1 DefaultElement;
-				DefaultElement.Name = Element0->Name;
+				// Element 1 is missing, use stub
+				DefaultElement1.Name = ElementPtr0->Name;
 
-				InPredicate(*NewResultElement, *Element0, DefaultElement, UE::Anim::ENamedValueUnionFlags::ValidArg0);
-				++ElementIndex0;
+				InPredicate(NewResultElement, *ElementPtr0, DefaultElement1, UE::Anim::ENamedValueUnionFlags::ValidArg0);
+
+				++ElementPtr0;
+
+				bIsDone = ElementPtr0 == ElementEndPtr0;
 			}
 			else
 			{
 				// ValueArray1 element is earlier, so so run predicate with only ValueArray1 and increment ValueArray1
-				ElementTypeResult* RESTRICT NewResultElement = &OutResultValueArray.Elements.AddDefaulted_GetRef();
-				NewResultElement->Name = Element1->Name;
-				
-				ElementType0 DefaultElement;
-				DefaultElement.Name = Element1->Name;
-				
-				InPredicate(*NewResultElement, DefaultElement, *Element1, UE::Anim::ENamedValueUnionFlags::ValidArg1);
-				++ElementIndex1;
+				NewResultElement.Name = ElementPtr1->Name;
+
+				// Element 0 is missing, use stub
+				DefaultElement0.Name = ElementPtr1->Name;
+
+				InPredicate(NewResultElement, DefaultElement0, *ElementPtr1, UE::Anim::ENamedValueUnionFlags::ValidArg1);
+
+				++ElementPtr1;
+
+				bIsDone = ElementPtr1 == ElementEndPtr1;
+			}
+
+			OutResultValueArray.Elements.Add(MoveTemp(NewResultElement));
+		}
+
+		// Tape merge is done, copy anything that might be remaining
+		if (ElementPtr0 < ElementEndPtr0)
+		{
+			// Reached end of ValueArray1 with remaining elements in ValueArray0, we can just copy the remainder of ValueArray0
+			const int32 NumResults = OutResultValueArray.Elements.Num();
+			const int32 NumNewElements = (ElementEndPtr0 - ElementPtr0);
+			OutResultValueArray.Elements.Reserve(NumResults + NumNewElements);
+			OutResultValueArray.Elements.AddUninitialized(NumNewElements);
+
+			ElementTypeResult* RESTRICT ResultElementPtr = OutResultValueArray.Elements.GetData() + NumResults;
+
+			for (; ElementPtr0 < ElementEndPtr0; ++ResultElementPtr, ++ElementPtr0)
+			{
+				// Element 1 is missing, use stub
+				DefaultElement1.Name = ElementPtr0->Name;
+
+				// We use placement new to make sure the constructor is inlined to reduce redundant work
+				new(ResultElementPtr) ElementTypeResult();
+				ResultElementPtr->Name = ElementPtr0->Name;
+
+				InPredicate(*ResultElementPtr, *ElementPtr0, DefaultElement1, UE::Anim::ENamedValueUnionFlags::ValidArg0);
+			}
+		}
+		else if (ElementPtr1 < ElementEndPtr1)
+		{
+			// Reached end of ValueArray0 with remaining elements in ValueArray1, we can just copy the remainder of ValueArray1
+			const int32 NumResults = OutResultValueArray.Elements.Num();
+			const int32 NumNewElements = (ElementEndPtr1 - ElementPtr1);
+			OutResultValueArray.Elements.Reserve(NumResults + NumNewElements);
+			OutResultValueArray.Elements.AddUninitialized(NumNewElements);
+
+			ElementTypeResult* RESTRICT ResultElementPtr = OutResultValueArray.Elements.GetData() + NumResults;
+
+			for (; ElementPtr1 < ElementEndPtr1; ++ResultElementPtr, ++ElementPtr1)
+			{
+				// Element 0 is missing, use stub
+				DefaultElement0.Name = ElementPtr1->Name;
+
+				// We use placement new to make sure the constructor is inlined to reduce redundant work
+				new(ResultElementPtr) ElementTypeResult();
+				ResultElementPtr->Name = ElementPtr1->Name;
+
+				InPredicate(*ResultElementPtr, DefaultElement0, *ElementPtr1, UE::Anim::ENamedValueUnionFlags::ValidArg1);
 			}
 		}
 
