Make FSplineMeshSceneProxyDesc public

diff --git a/Engine/Source/Runtime/Engine/Private/SplineMeshSceneProxyDesc.h b/Engine/Source/Runtime/Engine/Private/SplineMeshSceneProxyDesc.h
deleted file mode 100644
diff --git a/Engine/Source/Runtime/Engine/Public/SplineMeshSceneProxyDesc.h b/Engine/Source/Runtime/Engine/Public/SplineMeshSceneProxyDesc.h
new file mode 100644
--- /dev/null
+++ b/Engine/Source/Runtime/Engine/Public/SplineMeshSceneProxyDesc.h
@@ -0,0 +1,33 @@
+﻿// Copyright Epic Games, Inc. All Rights Reserved.
+#pragma once
+#include "CoreMinimal.h"
+#include "StaticMeshSceneProxyDesc.h"
+#include "Components/SplineMeshComponent.h"
+
+class USplineMeshComponent;
+
+struct FSplineMeshSceneProxyDesc
+{
+	FSplineMeshSceneProxyDesc() = default;
+	ENGINE_API FSplineMeshSceneProxyDesc(const USplineMeshComponent* InComponent);
+	
+	void InitializeFrom(const USplineMeshComponent* InComponent);
+	
+	FSplineMeshParams SplineParams{};
+	FVector SplineUpDir = FVector::UpVector;
+	float SplineBoundaryMin = 0.0f;
+	float SplineBoundaryMax = 0.0f;
+	TEnumAsByte<ESplineMeshAxis::Type> ForwardAxis = ESplineMeshAxis::X;
+	bool bSmoothInterpRollScale = false;
+
+	FSplineMeshParams GetSplineParams() const { return SplineParams; }
+
+	FSplineMeshShaderParams CalculateShaderParams(UStaticMesh* StaticMesh) const;
+	float ComputeRatioAlongSpline(UStaticMesh* StaticMesh, float DistanceAlong) const;
+	void ComputeVisualMeshSplineTRange(UStaticMesh* StaticMesh, float& MinT, float& MaxT) const;
+	ENGINE_API FBox ComputeDistortedBounds(UStaticMesh* StaticMesh, const FTransform& InLocalToWorld, const FBoxSphereBounds& InMeshBounds, const FBoxSphereBounds* InBoundsToDistort = nullptr) const;
+	FTransform CalcSliceTransform(UStaticMesh* StaticMesh, const float DistanceAlong) const;
+	FTransform CalcSliceTransformAtSplineOffset(const float Alpha, const float MinT, const float MaxT) const;
+	
+	static void InitVertexFactory(UStaticMesh* Mesh, const ERHIFeatureLevel::Type FeatureLevel, int32 InLODIndex, FColorVertexBuffer* InOverrideColorVertexBuffer);
+};
\ No newline at end of file
