Removed unnecessary call to AWorldDataLayers::ResetDataLayerRuntimeStates when uninitializing a partition world as it was causing a crash in actor replication.

diff --git a/Engine/Source/Runtime/Engine/Private/WorldPartition/DataLayer/DataLayerManager.cpp b/Engine/Source/Runtime/Engine/Private/WorldPartition/DataLayer/DataLayerManager.cpp
--- a/Engine/Source/Runtime/Engine/Private/WorldPartition/DataLayer/DataLayerManager.cpp
+++ b/Engine/Source/Runtime/Engine/Private/WorldPartition/DataLayer/DataLayerManager.cpp
@@ -237,11 +237,6 @@ void UDataLayerManager::Initialize()
 
 void UDataLayerManager::DeInitialize()
 {
-	if (AWorldDataLayers* WorldDataLayers = GetWorldDataLayers())
-	{
-		WorldDataLayers->OnDataLayerManagerDeinitialized();
-	}
-
 #if WITH_EDITOR
 	UActorDescContainerInstance::OnActorDescContainerInstanceInitialized.RemoveAll(this);
 
diff --git a/Engine/Source/Runtime/Engine/Private/WorldPartition/DataLayer/WorldDataLayers.cpp b/Engine/Source/Runtime/Engine/Private/WorldPartition/DataLayer/WorldDataLayers.cpp
--- a/Engine/Source/Runtime/Engine/Private/WorldPartition/DataLayer/WorldDataLayers.cpp
+++ b/Engine/Source/Runtime/Engine/Private/WorldPartition/DataLayer/WorldDataLayers.cpp
@@ -1247,11 +1247,6 @@ void AWorldDataLayers::OnDataLayerManagerInitialized()
 	InitializeDataLayerRuntimeStates();
 }
 
-void AWorldDataLayers::OnDataLayerManagerDeinitialized()
-{
-	ResetDataLayerRuntimeStates();
-}
-
 void AWorldDataLayers::Serialize(FArchive& Ar)
 {
 	Super::Serialize(Ar);
diff --git a/Engine/Source/Runtime/Engine/Public/WorldPartition/DataLayer/WorldDataLayers.h b/Engine/Source/Runtime/Engine/Public/WorldPartition/DataLayer/WorldDataLayers.h
--- a/Engine/Source/Runtime/Engine/Public/WorldPartition/DataLayer/WorldDataLayers.h
+++ b/Engine/Source/Runtime/Engine/Public/WorldPartition/DataLayer/WorldDataLayers.h
@@ -266,7 +266,6 @@ private:
 #endif
 
 	ENGINE_API void OnDataLayerManagerInitialized();
-	ENGINE_API void OnDataLayerManagerDeinitialized();
 	ENGINE_API void ResolveEffectiveRuntimeState(const UDataLayerInstance* InDataLayer, bool bInNotifyChange = true);
 	ENGINE_API void DumpDataLayerRecursively(const UDataLayerInstance* DataLayer, FString Prefix, FOutputDevice& OutputDevice) const;
 	bool CanChangeDataLayerRuntimeState(const UDataLayerInstance* InDataLayerInstance, ESetDataLayerRuntimeStateError* OutReason = nullptr) const;
