World Partition: ActorDescContainerSubsystem now implements the Asset Registry duplicated asset resolver to resolve asset collisions and gather invalid assets for repair.

diff --git a/Engine/Source/Runtime/Engine/Private/Level.cpp b/Engine/Source/Runtime/Engine/Private/Level.cpp
--- a/Engine/Source/Runtime/Engine/Private/Level.cpp
+++ b/Engine/Source/Runtime/Engine/Private/Level.cpp
@@ -3809,8 +3809,12 @@ FString ULevel::GetActorPackageName(const FString& InBaseDir, EActorPackagingSch
 	TStringBuilderWithBuffer<TCHAR, NAME_SIZE> ActorPackageName;
 
 	uint32 FilenameOffset = 0;
-	ActorPackageName.Append(InBaseDir);
-	ActorPackageName.Append(TEXT("/"));
+
+	if (InBaseDir.Len())
+	{
+		ActorPackageName.Append(InBaseDir);
+		ActorPackageName.Append(TEXT("/"));
+	}
 
 	switch (ActorPackagingScheme)
 	{
diff --git a/Engine/Source/Runtime/Engine/Private/WorldPartition/ActorDescContainer.cpp b/Engine/Source/Runtime/Engine/Private/WorldPartition/ActorDescContainer.cpp
--- a/Engine/Source/Runtime/Engine/Private/WorldPartition/ActorDescContainer.cpp
+++ b/Engine/Source/Runtime/Engine/Private/WorldPartition/ActorDescContainer.cpp
@@ -269,6 +269,11 @@ void UActorDescContainer::Initialize(const FInitializeParams& InitParams)
 				}
 			}
 		}
+
+		if (UActorDescContainerSubsystem* ActorDescContainerSubsystem = UActorDescContainerSubsystem::Get())
+		{
+			InvalidActors.Append(ActorDescContainerSubsystem->InvalidMapAssets.FindRef(ContainerPackageName).Array());
+		}
 	}
 
 	{
diff --git a/Engine/Source/Runtime/Engine/Private/WorldPartition/ActorDescContainerSubsystem.cpp b/Engine/Source/Runtime/Engine/Private/WorldPartition/ActorDescContainerSubsystem.cpp
--- a/Engine/Source/Runtime/Engine/Private/WorldPartition/ActorDescContainerSubsystem.cpp
+++ b/Engine/Source/Runtime/Engine/Private/WorldPartition/ActorDescContainerSubsystem.cpp
@@ -1,7 +1,9 @@
 // Copyright Epic Games, Inc. All Rights Reserved.
 
 #include "WorldPartition/ActorDescContainerSubsystem.h"
+#include "AssetRegistry/AssetRegistryModule.h"
 #include "Engine/Engine.h"
+#include "Engine/Level.h"
 
 #include UE_INLINE_GENERATED_CPP_BY_NAME(ActorDescContainerSubsystem)
 
@@ -23,6 +25,64 @@ UActorDescContainerSubsystem& UActorDescContainerSubsystem::GetChecked()
 	return *Get();
 }
 
+void UActorDescContainerSubsystem::Initialize(FSubsystemCollectionBase& Collection)
+{
+	Super::Initialize(Collection);
+
+	if (FAssetRegistryModule* AssetRegistryModule = FModuleManager::GetModulePtr<FAssetRegistryModule>(FName("AssetRegistry")))
+	{
+		IAssetRegistry* AssetRegistry = AssetRegistryModule->TryGet();
+		if (AssetRegistry)
+		{
+			AssetRegistry->OnAssetCollision_Private().AddUObject(this, &UActorDescContainerSubsystem::OnAssetCollision);
+		}		
+	}
+}
+
+void UActorDescContainerSubsystem::Deinitialize()
+{
+	Super::Deinitialize();
+
+	if (FAssetRegistryModule* AssetRegistryModule = FModuleManager::GetModulePtr<FAssetRegistryModule>(FName("AssetRegistry")))
+	{
+		IAssetRegistry* AssetRegistry = AssetRegistryModule->TryGet();
+		if (AssetRegistry)
+		{
+			AssetRegistry->OnAssetCollision_Private().RemoveAll(this);
+		}
+	}
+}
+
+void UActorDescContainerSubsystem::OnAssetCollision(FAssetData& A, FAssetData& B, FAssetData*& Keep)
+{
+	if (A.PackagePath.ToString().Contains(ULevel::GetExternalActorsFolderName()))
+	{
+		auto IsMatching = [](const FAssetData& Asset, EActorPackagingScheme PackagingScheme)
+		{
+			const FString AssetPath = ULevel::GetActorPackageName(TEXT(""), PackagingScheme, Asset.GetSoftObjectPath().ToString());
+			return Asset.PackageName.ToString().EndsWith(AssetPath);
+		};
+
+		if (IsMatching(A, EActorPackagingScheme::Reduced) || IsMatching(A, EActorPackagingScheme::Original))
+		{
+			Keep = &A;
+			FSoftObjectPath ObjectPathB(B.GetOptionalOuterPathName().ToString());
+			InvalidMapAssets.FindOrAdd(ObjectPathB.GetAssetPath().GetPackageName()).Add(B);
+			return;
+		}
+
+		if (IsMatching(B, EActorPackagingScheme::Reduced) || IsMatching(A, EActorPackagingScheme::Original))
+		{
+			Keep = &B;
+			FSoftObjectPath ObjectPathA(A.GetOptionalOuterPathName().ToString());
+			InvalidMapAssets.FindOrAdd(ObjectPathA.GetAssetPath().GetPackageName()).Add(A);
+			return;
+		}
+	}
+
+	Keep = nullptr;
+}
+
 void UActorDescContainerSubsystem::AddReferencedObjects(UObject* InThis, FReferenceCollector& Collector)
 {
 	Super::AddReferencedObjects(InThis, Collector);
@@ -140,4 +200,4 @@ void UActorDescContainerSubsystem::FContainerManager::SetContainerPackage(UActor
 		}
 	}
 }
-#endif
\ No newline at end of file
+#endif
diff --git a/Engine/Source/Runtime/Engine/Public/WorldPartition/ActorDescContainerSubsystem.h b/Engine/Source/Runtime/Engine/Public/WorldPartition/ActorDescContainerSubsystem.h
--- a/Engine/Source/Runtime/Engine/Public/WorldPartition/ActorDescContainerSubsystem.h
+++ b/Engine/Source/Runtime/Engine/Public/WorldPartition/ActorDescContainerSubsystem.h
@@ -12,6 +12,8 @@ class UActorDescContainerSubsystem : public UEngineSubsystem
 {
 	GENERATED_BODY()
 
+	friend class UActorDescContainer;
+
 public:
 	UActorDescContainerSubsystem()
 #if WITH_EDITOR
@@ -33,6 +35,11 @@ public:
 	ENGINE_API static UActorDescContainerSubsystem* Get();
 	ENGINE_API static UActorDescContainerSubsystem& GetChecked();
 
+	virtual void Initialize(FSubsystemCollectionBase& Collection) override;
+	virtual void Deinitialize() override;
+
+	void OnAssetCollision(FAssetData& A, FAssetData& B, FAssetData*& Keep);
+
 	DECLARE_EVENT_OneParam(UActorDescContainerSubsystem, FContainerUpdatedEvent, FName);
 	DECLARE_EVENT_TwoParams(UActorDescContainerSubsystem, FContainerReplacedEvent, UActorDescContainer*, UActorDescContainer*);
 
@@ -159,5 +166,7 @@ private:
 	};
 
 	FContainerManager ContainerManager;
+
+	TMap<FName, TSet<FAssetData>> InvalidMapAssets;
 #endif
 };
\ No newline at end of file
