Level Instance: ensure that a level instance has a valid id before attempting to load it.
This is a temporary fix until we figure out what's the reason for it.

diff --git a/Engine/Source/Runtime/Engine/Private/LevelInstance/LevelInstanceSubsystem.cpp b/Engine/Source/Runtime/Engine/Private/LevelInstance/LevelInstanceSubsystem.cpp
--- a/Engine/Source/Runtime/Engine/Private/LevelInstance/LevelInstanceSubsystem.cpp
+++ b/Engine/Source/Runtime/Engine/Private/LevelInstance/LevelInstanceSubsystem.cpp
@@ -566,6 +566,11 @@ void ULevelInstanceSubsystem::LoadLevelInstance(ILevelInstanceInterface* LevelIn
 		return;
 	}
 
+	if (!ensure(LevelInstance->HasValidLevelInstanceID()))
+	{
+		return;
+	}
+
 	const FLevelInstanceID& LevelInstanceID = LevelInstance->GetLevelInstanceID();
 	check(!LoadedLevelInstances.Contains(LevelInstanceID));
 	check(!LoadingLevelInstances.Contains(LevelInstanceID));
