Fix crash on Mobile due to read/write on UAV when Substrate is enabled.

Disable specular profile convolution on unsupported platforms.

diff --git a/Engine/Source/Runtime/Engine/Private/Rendering/SpecularProfile.cpp b/Engine/Source/Runtime/Engine/Private/Rendering/SpecularProfile.cpp
--- a/Engine/Source/Runtime/Engine/Private/Rendering/SpecularProfile.cpp
+++ b/Engine/Source/Runtime/Engine/Private/Rendering/SpecularProfile.cpp
@@ -48,6 +48,11 @@ static bool ForceUpdateSpecularProfile()
 
 ///////////////////////////////////////////////////////////////////////////////////////////////////
 
+static bool IsSpecularProfileSupport(EShaderPlatform InShaderPlatform)
+{
+	return !IsMobilePlatform(InShaderPlatform) && IsFeatureLevelSupported(InShaderPlatform, ERHIFeatureLevel::SM5);
+}
+
 class FSpecularProfileCopyCS : public FGlobalShader
 {
 public:
@@ -71,7 +76,7 @@ public:
 
 	static bool ShouldCompilePermutation(const FGlobalShaderPermutationParameters& Parameters)
 	{
-		return true;
+		return IsSpecularProfileSupport(Parameters.Platform);
 	}
 	static void ModifyCompilationEnvironment(const FGlobalShaderPermutationParameters& Parameters, FShaderCompilerEnvironment& OutEnvironment)
 	{
@@ -103,7 +108,7 @@ public:
 
 	static bool ShouldCompilePermutation(const FGlobalShaderPermutationParameters& Parameters)
 	{
-		return true;
+		return IsSpecularProfileSupport(Parameters.Platform);
 	}
 	static void ModifyCompilationEnvironment(const FGlobalShaderPermutationParameters& Parameters, FShaderCompilerEnvironment& OutEnvironment)
 	{
@@ -317,7 +322,7 @@ IPooledRenderTarget* FSpecularProfileTextureManager::GetAtlasTexture()
 
 IPooledRenderTarget* FSpecularProfileTextureManager::GetAtlasTexture(FRDGBuilder& GraphBuilder, EShaderPlatform ShaderPlatform)
 {
-	if (!Substrate::IsSubstrateEnabled())
+	if (!Substrate::IsSubstrateEnabled() || !Substrate::IsSpecularProfileEnabled() || !IsSpecularProfileSupport(ShaderPlatform))
 	{
 		return nullptr;
 	}
