Expose FIntVector2 to BP

[RN] rnx

diff --git a/Engine/Source/Runtime/CoreUObject/Public/UObject/NoExportTypes.h b/Engine/Source/Runtime/CoreUObject/Public/UObject/NoExportTypes.h
--- a/Engine/Source/Runtime/CoreUObject/Public/UObject/NoExportTypes.h
+++ b/Engine/Source/Runtime/CoreUObject/Public/UObject/NoExportTypes.h
@@ -1174,13 +1174,13 @@ struct FInt64Vector2
  * An integer vector in 4D space.
  * @note The full C++ class is located here: Engine\Source\Runtime\Core\Public\Math\IntVector.h
  */
-USTRUCT(immutable, noexport, IsAlwaysAccessible, HasDefaults, HasNoOpConstructor, IsCoreType)
+USTRUCT(immutable, noexport, BlueprintType, IsAlwaysAccessible, HasDefaults, HasNoOpConstructor, IsCoreType)
 struct FIntVector2
 {
-	UPROPERTY(EditAnywhere, Category = IntVector, SaveGame)
+	UPROPERTY(EditAnywhere, BlueprintReadWrite, Category = IntVector, SaveGame)
 	int32 X;
 
-	UPROPERTY(EditAnywhere, Category = IntVector, SaveGame)
+	UPROPERTY(EditAnywhere, BlueprintReadWrite, Category = IntVector, SaveGame)
 	int32 Y;
 };
 
diff --git a/Engine/Source/Runtime/Engine/Classes/Kismet/KismetMathLibrary.h b/Engine/Source/Runtime/Engine/Classes/Kismet/KismetMathLibrary.h
--- a/Engine/Source/Runtime/Engine/Classes/Kismet/KismetMathLibrary.h
+++ b/Engine/Source/Runtime/Engine/Classes/Kismet/KismetMathLibrary.h
@@ -1584,6 +1584,10 @@ class UKismetMathLibrary : public UBlueprintFunctionLibrary
 	UFUNCTION(BlueprintPure, meta = (DisplayName = "Truncate (Vector)", ScriptMethod = "Truncated", BlueprintAutocast), Category = "Math|Float")
 	static ENGINE_API FIntVector FTruncVector(const FVector& InVector);
 
+	/** Rounds A to an integer with truncation towards zero for each element in a vector2D.  (e.g. -1.7 truncated to -1, 2.8 truncated to 2) */
+	UFUNCTION(BlueprintPure, meta = (DisplayName = "Truncate (Vector2D)", ScriptMethod = "Truncated", BlueprintAutocast), Category = "Math|Float")
+	static ENGINE_API FIntVector2 FTruncVector2D(const FVector2D& InVector2D);
+
 	/**
 	 * Distance between two points.
 	 *
@@ -3616,6 +3620,10 @@ class UKismetMathLibrary : public UBlueprintFunctionLibrary
 	UFUNCTION(BlueprintPure, meta = (DisplayName = "To IntVector (Integer)", CompactNodeTitle = "->", Keywords = "cast convert", BlueprintAutocast), Category = "Math|Conversions")
 	static ENGINE_API FIntVector Conv_IntToIntVector(int32 InInt);
 
+	/** Converts an integer to an IntVector2*/
+	UFUNCTION(BlueprintPure, meta = (DisplayName = "To IntVector2 (Integer)", CompactNodeTitle = "->", Keywords = "cast convert", BlueprintAutocast), Category = "Math|Conversions")
+	static ENGINE_API FIntVector2 Conv_IntToIntVector2(int32 InInt);
+
 	/** Converts an integer to a FVector */
 	UFUNCTION(BlueprintPure, meta = (DisplayName = "To Vector (Integer)", CompactNodeTitle = "->", Keywords = "cast convert", BlueprintAutocast), Category = "Math|Conversions")
 	static ENGINE_API FVector Conv_IntToVector(int32 InInt);
@@ -3656,6 +3664,10 @@ class UKismetMathLibrary : public UBlueprintFunctionLibrary
 	UFUNCTION(BlueprintPure, meta = (DisplayName = "To Vector (IntVector)", CompactNodeTitle = "->", Keywords = "cast convert", BlueprintAutocast), Category = "Math|Conversions")
 	static ENGINE_API FVector Conv_IntVectorToVector(const FIntVector& InIntVector);
 
+	/** Converts an IntVector2 to a vector2D */
+	UFUNCTION(BlueprintPure, meta = (DisplayName = "To Vector2D (IntVector2)", CompactNodeTitle = "->", Keywords = "cast convert", BlueprintAutocast), Category = "Math|Conversions")
+	static ENGINE_API FVector2D Conv_IntVector2ToVector2D(const FIntVector2& InIntVector2);
+
 	/** Converts a float into a vector, where each element is that float */
 	UE_DEPRECATED(5.2, "This method has been deprecated and will be removed. Use the double version instead.")
 	static ENGINE_API FVector Conv_FloatToVector(float InFloat);
diff --git a/Engine/Source/Runtime/Engine/Classes/Kismet/KismetMathLibrary.inl b/Engine/Source/Runtime/Engine/Classes/Kismet/KismetMathLibrary.inl
--- a/Engine/Source/Runtime/Engine/Classes/Kismet/KismetMathLibrary.inl
+++ b/Engine/Source/Runtime/Engine/Classes/Kismet/KismetMathLibrary.inl
@@ -1984,6 +1984,12 @@ FIntVector UKismetMathLibrary::FTruncVector(const FVector& InVector)
 		static_cast<int32>(FMath::TruncToInt(InVector.Z)));
 }
 
+KISMET_MATH_FORCEINLINE
+FIntVector2 UKismetMathLibrary::FTruncVector2D(const FVector2D& InVector2D)
+{
+	return FIntVector2(static_cast<int32>(FMath::TruncToInt(InVector2D.X)),
+		static_cast<int32>(FMath::TruncToInt(InVector2D.Y)));
+}
 
 /* FVector4
 *****************************************************************************/
@@ -2701,6 +2707,12 @@ FIntVector UKismetMathLibrary::Conv_IntToIntVector(int32 InInt)
 	return FIntVector(InInt, InInt, InInt);
 }
 
+KISMET_MATH_FORCEINLINE
+FIntVector2 UKismetMathLibrary::Conv_IntToIntVector2(int32 InInt)
+{
+	return FIntVector2(InInt, InInt);
+}
+
 KISMET_MATH_FORCEINLINE
 FVector UKismetMathLibrary::Conv_IntToVector(int32 InInt)
 {
@@ -2767,6 +2779,12 @@ FVector UKismetMathLibrary::Conv_IntVectorToVector(const FIntVector& InIntVector
 	return FVector(InIntVector);
 }
 
+KISMET_MATH_FORCEINLINE
+FVector2D UKismetMathLibrary::Conv_IntVector2ToVector2D(const FIntVector2& InIntVector2)
+{
+	return FVector2D(InIntVector2.X, InIntVector2.Y);
+}
+
 KISMET_MATH_FORCEINLINE
 FLinearColor UKismetMathLibrary::Conv_ColorToLinearColor(FColor InColor)
 {
diff --git a/Engine/Source/Runtime/Engine/Classes/Kismet/KismetStringLibrary.h b/Engine/Source/Runtime/Engine/Classes/Kismet/KismetStringLibrary.h
--- a/Engine/Source/Runtime/Engine/Classes/Kismet/KismetStringLibrary.h
+++ b/Engine/Source/Runtime/Engine/Classes/Kismet/KismetStringLibrary.h
@@ -48,6 +48,10 @@ class UKismetStringLibrary : public UBlueprintFunctionLibrary
 	UFUNCTION(BlueprintPure, meta=(DisplayName = "To String (IntVector)", CompactNodeTitle = "->", BlueprintAutocast), Category="Utilities|String")
 	static ENGINE_API FString Conv_IntVectorToString(FIntVector InIntVec);
 
+	/** Converts an IntVector2 value to a string, in the form 'X= Y=' */
+	UFUNCTION(BlueprintPure, meta=(DisplayName = "To String (IntVector2)", CompactNodeTitle = "->", BlueprintAutocast), Category="Utilities|String")
+	static ENGINE_API FString Conv_IntVector2ToString(FIntVector2 InIntVec2);
+
 	/** Converts an IntPoint value to a string, in the form 'X= Y=' */
 	UFUNCTION(BlueprintPure, meta = (DisplayName = "To String (IntPoint)", CompactNodeTitle = "->", BlueprintAutocast), Category = "Utilities|String")
 	static ENGINE_API FString Conv_IntPointToString(FIntPoint InIntPoint);
@@ -200,7 +204,19 @@ class UKismetStringLibrary : public UBlueprintFunctionLibrary
 	 * @return A new string built from the passed parameters
 	 */
 	UFUNCTION(BlueprintPure, meta = (DisplayName = "Build String (IntVector)"), Category = "Utilities|String")
-	static ENGINE_API FString BuildString_IntVector(const FString& AppendTo, const FString& Prefix, FIntVector InIntVector, const FString& Suffix);
+	static ENGINE_API FString BuildString_IntVector(const FString& AppendTo, const FString& Prefix, FIntVector InIntVector2, const FString& Suffix);
+
+	/**
+	 * Converts an IntVector2->string, creating a new string in the form AppendTo+Prefix+InIntVector+Suffix
+	 * @param AppendTo - An existing string to use as the start of the conversion string
+	 * @param Prefix - A string to use as a prefix, after the AppendTo string
+	 * @param InIntVector - The intVector2 value to convert. Uses the standard FVector::ToString conversion
+	 * @param Suffix - A suffix to append to the end of the conversion string
+	 * @return A new string built from the passed parameters
+	 */
+	UFUNCTION(BlueprintPure, meta = (DisplayName = "Build String (IntVector2)"), Category = "Utilities|String")
+	static ENGINE_API FString BuildString_IntVector2(const FString& AppendTo, const FString& Prefix, FIntVector2 InIntVector, const FString& Suffix);
+
 	/** 
 	 * Converts a vector2d->string, creating a new string in the form AppendTo+Prefix+InVector2d+Suffix
 	 * @param AppendTo - An existing string to use as the start of the conversion string
diff --git a/Engine/Source/Runtime/Engine/Private/KismetStringLibrary.cpp b/Engine/Source/Runtime/Engine/Private/KismetStringLibrary.cpp
--- a/Engine/Source/Runtime/Engine/Private/KismetStringLibrary.cpp
+++ b/Engine/Source/Runtime/Engine/Private/KismetStringLibrary.cpp
@@ -99,6 +99,11 @@ FString UKismetStringLibrary::Conv_IntVectorToString(FIntVector InIntVec)
 	return InIntVec.ToString();
 }
 
+FString UKismetStringLibrary::Conv_IntVector2ToString(FIntVector2 InIntVec2)
+{
+	return InIntVec2.ToString();
+}
+
 FString UKismetStringLibrary::Conv_IntPointToString(FIntPoint InIntPoint)
 {
 	return InIntPoint.ToString();
@@ -293,6 +298,21 @@ FString UKismetStringLibrary::BuildString_IntVector(const FString& AppendTo, con
 	return StringResult;
 }
 
+FString UKismetStringLibrary::BuildString_IntVector2(const FString& AppendTo, const FString& Prefix, FIntVector2 InIntVector2, const FString& Suffix)
+{
+	// faster, preallocating method
+	const FString VecStr = InIntVector2.ToString();
+
+	FString StringResult;
+	StringResult.Empty(AppendTo.Len() + Prefix.Len() + VecStr.Len() + Suffix.Len() + 1); // adding one for the string terminator
+	StringResult += AppendTo;
+	StringResult += Prefix;
+	StringResult += VecStr;
+	StringResult += Suffix;
+
+	return StringResult;
+}
+
 FString UKismetStringLibrary::BuildString_Vector2d(const FString& AppendTo, const FString& Prefix, FVector2D InVector2d, const FString& Suffix)
 {
 	// faster, preallocating method
