Holdout: Adjust new allow primitive alpha holdout flag so that it's not visible everywhere in the primary common section.

Also fixed missing spaces in string name.

#rnx

diff --git a/Engine/Source/Runtime/Engine/Public/ShowFlagsValues.inl b/Engine/Source/Runtime/Engine/Public/ShowFlagsValues.inl
--- a/Engine/Source/Runtime/Engine/Public/ShowFlagsValues.inl
+++ b/Engine/Source/Runtime/Engine/Public/ShowFlagsValues.inl
@@ -467,7 +467,7 @@ SHOWFLAG_ALWAYS_ACCESSIBLE(Cameras, SFG_Normal, NSLOCTEXT("UnrealEd", "CamerasSF
 SHOWFLAG_ALWAYS_ACCESSIBLE(Hair, SFG_Normal, NSLOCTEXT("UnrealEd", "HairSF", "Hair"))
 
 /** Allows primitive alpha holdout rendering. Note that the "Support Primitive Alpha Holdout" project setting first needs to be enabled for the deferred renderer. */
-SHOWFLAG_ALWAYS_ACCESSIBLE(AllowPrimitiveAlphaHoldout, SFG_Normal, NSLOCTEXT("UnrealEd", "AllowPrimitiveAlphaHoldoutSF", "AllowPrimitiveAlphaHoldout"))
+SHOWFLAG_ALWAYS_ACCESSIBLE(AllowPrimitiveAlphaHoldout, SFG_Advanced, NSLOCTEXT("UnrealEd", "AllowPrimitiveAlphaHoldoutSF", "Allow Primitive Alpha Holdout"))
 
 #undef SHOWFLAG_ALWAYS_ACCESSIBLE
 #undef SHOWFLAG_FIXED_IN_SHIPPING
