IncrementalCook: BodySetup: Make the initial value for a UBodySetup's BodySetupGuid deterministic, so that empty BodySetups created during cook will not cause cook indeterminism.
#rnx

diff --git a/Engine/Source/Runtime/Engine/Private/PhysicsEngine/BodySetup.cpp b/Engine/Source/Runtime/Engine/Private/PhysicsEngine/BodySetup.cpp
--- a/Engine/Source/Runtime/Engine/Private/PhysicsEngine/BodySetup.cpp
+++ b/Engine/Source/Runtime/Engine/Private/PhysicsEngine/BodySetup.cpp
@@ -794,7 +794,10 @@ void UBodySetup::RescaleSimpleCollision( FVector BuildScale )
 void UBodySetup::InvalidatePhysicsData()
 {
 	ClearPhysicsMeshes();
-	BodySetupGuid = FGuid::NewGuid(); // change the guid
+	// We use the BodySetupGuid as a change indicator for DDC. It would be better to create it deterministically
+	// from the physics data but that is not yet implemented. As a workaround we assign it an arbitrary unique value
+	// every time the physics data changes.
+	BodySetupGuid = FGuid::NewGuid();
 	if (!bSharedCookedData)
 	{
 		CookedFormatData.FlushData();
@@ -828,11 +831,18 @@ void UBodySetup::Serialize(FArchive& Ar)
 	Ar << BodySetupGuid;
 
 	// If we loaded a ZERO Guid, fix that
-	if(Ar.IsLoading() && !BodySetupGuid.IsValid())
+	if (Ar.IsLoading() && !BodySetupGuid.IsValid())
 	{
 		MarkPackageDirty();
-		UE_LOG(LogPhysics, Log, TEXT("FIX GUID FOR: %s"), *GetPathName());
-		BodySetupGuid = FGuid::NewGuid();
+		FString PathName = GetPathName();
+		UE_LOG(LogPhysics, Log, TEXT("FIX GUID FOR: %s"), *PathName);
+		// For the autocreated guid for an old version, we construct it to a constant value; this is important
+		// for determinism during cooking if a legacy BodySetup is encountered during cooking. The BodySetupGuid is used
+		// as a change-indicator for the BodySetup's DDC key, so this constant value would ordinarily be a problem, but
+		// it is okay in this legacy-version handling because the physics data stored in the package will always be the
+		// same, unless the package is resaved. And when the package is resaved, the new code will calculate a unique
+		// BodySetupGuid for it and store that BodySetupGuid in the serialized data.
+		BodySetupGuid = FGuid::NewDeterministicGuid(PathName);
 	}
 
 	bool bCooked = Ar.IsCooking();
@@ -1501,9 +1511,14 @@ void UBodySetup::PostInitProperties()
 {
 	Super::PostInitProperties();
 
-	if(!IsTemplate())
+	if (!IsTemplate())
 	{
-		BodySetupGuid = FGuid::NewGuid();
+		// When constructing a BodySetup, we initialize its guid to a constant value; this is important
+		// for determinism during cooking if a new BodySetup is created during cooking. The BodySetupGuid is used
+		// as a change-indicator for the BodySetup's DDC key, so it is not possible to use a constant value for it
+		// for changes to the BodySetup, but since a new BodySetup is created into an initial state that never varies,
+		// making the initial value of BodySetupGuid a constant is okay.
+		BodySetupGuid = FGuid::NewDeterministicGuid(GetPathName());
 	}
 }
 
