[HWRT] Include fallback materials in GetUsedMaterials otherwise cached ray tracing draw commands will not be invalidated leading to a crash.

#rnx

diff --git a/Engine/Source/Runtime/Engine/Private/Components/StaticMeshComponent.cpp b/Engine/Source/Runtime/Engine/Private/Components/StaticMeshComponent.cpp
--- a/Engine/Source/Runtime/Engine/Private/Components/StaticMeshComponent.cpp
+++ b/Engine/Source/Runtime/Engine/Private/Components/StaticMeshComponent.cpp
@@ -2943,6 +2943,45 @@ UMaterialInterface* UStaticMeshComponent::GetEditorMaterial(int32 MaterialIndex)
 }
 #endif
 
+template<typename T>
+void GetUsedRayTracingOnlyMaterials(TArray<UMaterialInterface*>& OutMaterials, const T* Component)
+{
+	TArray<UMaterialInterface*> RayTracingMaterials;
+	Component->GetStaticMesh()->GetUsedMaterials(RayTracingMaterials, [Component](int32 MaterialIndex) {
+		UMaterialInterface* OutMaterial = nullptr;
+
+		// If we have a base materials array, use that
+		if (Component->OverrideMaterials.IsValidIndex(MaterialIndex) && Component->OverrideMaterials[MaterialIndex])
+		{
+			OutMaterial = Component->OverrideMaterials[MaterialIndex];
+		}
+		// Otherwise get from static mesh
+		else if (Component->GetStaticMesh())
+		{
+			OutMaterial = Component->GetStaticMesh()->GetMaterial(MaterialIndex);
+		}
+
+		if (OutMaterial)
+		{
+			//@note FH: temporary preemptive PostLoad until zenloader load ordering improvements
+			OutMaterial->ConditionalPostLoad();
+
+			if (UMaterialInterface* NaniteOverride = OutMaterial->GetNaniteOverride())
+			{
+				return OutMaterial;
+			}
+		}
+
+		return (UMaterialInterface*)nullptr;
+	});
+
+	if (!RayTracingMaterials.IsEmpty())
+	{
+		RayTracingMaterials.RemoveAll([](UMaterialInterface* Material) { return Material == nullptr; });
+		OutMaterials.Append(RayTracingMaterials);
+	}
+}
+
 void UStaticMeshComponent::GetUsedMaterials(TArray<UMaterialInterface*>& OutMaterials, bool bGetDebugMaterials) const
 {
 	//TRACE_CPUPROFILER_EVENT_SCOPE(UStaticMeshComponent::GetUsedMaterials);
@@ -2950,6 +2989,13 @@ void UStaticMeshComponent::GetUsedMaterials(TArray<UMaterialInterface*>& OutMate
 	if (GetStaticMesh())
 	{
 		GetStaticMesh()->GetUsedMaterials(OutMaterials, [this](int32 Index) { return GetMaterial(Index); });
+
+		// For ray tracing if the mesh is using nanite override materials we need to include fallback materials as well
+		if (IsRayTracingAllowed() && GetStaticMesh()->bSupportRayTracing && UseNaniteOverrideMaterials(false))
+		{
+			GetUsedRayTracingOnlyMaterials(OutMaterials, this);
+		}
+
 		if (OutMaterials.Num() > 0)
 		{
 			UMaterialInterface* OverlayMaterialInterface = GetOverlayMaterial();
@@ -2961,6 +3007,28 @@ void UStaticMeshComponent::GetUsedMaterials(TArray<UMaterialInterface*>& OutMate
 	}
 }
 
+void FStaticMeshSceneProxyDesc::GetUsedMaterials(TArray<UMaterialInterface*>& OutMaterials, bool bGetDebugMaterials) const
+{
+	if (GetStaticMesh())
+	{
+		GetStaticMesh()->GetUsedMaterials(OutMaterials, [this](int32 Index) { return GetMaterial(Index); });
+
+		// For ray tracing if the mesh is using nanite override materials we need to include fallback materials as well
+		if (IsRayTracingAllowed() && GetStaticMesh()->bSupportRayTracing && UseNaniteOverrideMaterials(false))
+		{
+			GetUsedRayTracingOnlyMaterials(OutMaterials, this);
+		}
+
+		if (OutMaterials.Num() > 0)
+		{
+			if (OverlayMaterial != nullptr)
+			{
+				OutMaterials.Add(OverlayMaterial);
+			}
+		}
+	}
+}
+
 #if WITH_EDITOR
 bool UStaticMeshComponent::GetMaterialPropertyPath(int32 ElementIndex, UObject*& OutOwner, FString& OutPropertyPath, FProperty*& OutProperty)
 {
diff --git a/Engine/Source/Runtime/Engine/Public/StaticMeshSceneProxyDesc.h b/Engine/Source/Runtime/Engine/Public/StaticMeshSceneProxyDesc.h
--- a/Engine/Source/Runtime/Engine/Public/StaticMeshSceneProxyDesc.h
+++ b/Engine/Source/Runtime/Engine/Public/StaticMeshSceneProxyDesc.h
@@ -199,20 +199,7 @@ struct FStaticMeshSceneProxyDesc : public FPrimitiveSceneProxyDesc
 	FColor	WireframeColor = FColor(0, 255, 255, 255);
 	FColor GetWireframeColor() const { return WireframeColor; }
 
-	virtual void GetUsedMaterials(TArray<UMaterialInterface*>& OutMaterials, bool bGetDebugMaterials = false) const override
-	{
-		if (GetStaticMesh())
-		{
-			GetStaticMesh()->GetUsedMaterials(OutMaterials, [this](int32 Index) { return GetMaterial(Index); });
-			if (OutMaterials.Num() > 0)
-			{
-				if (OverlayMaterial != nullptr)
-				{
-					OutMaterials.Add(OverlayMaterial);
-				}
-			}
-		}
-	}
+	ENGINE_API virtual void GetUsedMaterials(TArray<UMaterialInterface*>& OutMaterials, bool bGetDebugMaterials = false) const override;
 
 	const UStaticMeshComponent* GetUStaticMeshComponent() const { return Cast<UStaticMeshComponent>(Component); }
 };
\ No newline at end of file
