NotifyActiveSoundOwner docu-comment corrected.

#rnx

diff --git a/Engine/Source/Runtime/Engine/Public/ActiveSoundUpdateInterface.h b/Engine/Source/Runtime/Engine/Public/ActiveSoundUpdateInterface.h
--- a/Engine/Source/Runtime/Engine/Public/ActiveSoundUpdateInterface.h
+++ b/Engine/Source/Runtime/Engine/Public/ActiveSoundUpdateInterface.h
@@ -44,7 +44,7 @@ public:
 	 * Called while an active sound is being constructed. Will be followed by either
 	 * NotifyActiveSoundCreated or NotifyVirtualizedSoundCreated.
 	 * 
-	 * NOTE! Called on the AudioThread
+	 * NOTE! Called on the GameThread
 	 *
 	 * @param ActiveSound   The active sound being associated
 	 * @param Owner         The owner it is associated with, or nullptr
