[Engine]
* Moved larger classes out of ParticleComponents.cpp and in to their own files (since we inline .gen.cpp it doesn't end up being more files to compile in non-unity configs anyway)

diff --git a/Engine/Source/Runtime/Engine/Private/Particles/ParticleComponents.cpp b/Engine/Source/Runtime/Engine/Private/Particles/ParticleComponents.cpp
--- a/Engine/Source/Runtime/Engine/Private/Particles/ParticleComponents.cpp
+++ b/Engine/Source/Runtime/Engine/Private/Particles/ParticleComponents.cpp
@@ -4,90 +4,42 @@
 	UnParticleComponent.cpp: Particle component implementation.
 =============================================================================*/
 
-#include "Camera/CameraLensEffectInterface.h"
-#include "Engine/GameInstance.h"
-#include "Distributions/Distribution.h"
-#include "GameFramework/WorldSettings.h"
-#include "UObject/FrameworkObjectVersion.h"
-#include "Engine/Level.h"
-#include "Engine/CollisionProfile.h"
-#include "GameFramework/PlayerController.h"
-#include "UObject/AssetRegistryTagsContext.h"
-#include "UObject/ObjectSaveContext.h"
-#include "Materials/MaterialRelevance.h"
-#include "UObject/UObjectIterator.h"
-#include "Misc/LargeWorldRenderPosition.h"
-#include "UObject/Package.h"
-#include "Particles/Color/ParticleModuleColorBase.h"
-#include "Particles/ParticleSystem.h"
-#include "Particles/Emitter.h"
-#include "Particles/Orientation/ParticleModuleOrientationAxisLock.h"
-#include "ParticleEmitterInstances.h"
-#include "ParticleEmitterInstanceOwner.h"
-#include "Particles/ParticleModule.h"
 #include "Particles/ParticleSystemComponent.h"
+#include "Distributions/DistributionFloat.h"
 #include "Distributions/DistributionFloatConstant.h"
+#include "Distributions/DistributionFloatConstantCurve.h"
 #include "Distributions/DistributionFloatUniform.h"
+#include "Distributions/DistributionVector.h"
 #include "Distributions/DistributionVectorConstant.h"
-#include "Distributions/DistributionVectorUniform.h"
-#include "Materials/MaterialInstanceDynamic.h"
-#include "Particles/ParticlePerfStats.h"
-#include "UnrealEngine.h"
 #include "Distributions/DistributionVectorConstantCurve.h"
-#include "Particles/SubUVAnimation.h"
-#include "StaticMeshResources.h"
-#include "Particles/EmitterCameraLensEffectBase.h"
-#include "ParticleSystemSceneProxy.h"
+#include "Distributions/DistributionVectorUniform.h"
+#include "Engine/Engine.h"
+#include "Engine/GameInstance.h"
+#include "Engine/World.h"
+#include "GameFramework/WorldSettings.h"
 #include "Logging/MessageLog.h"
-#include "Misc/UObjectToken.h"
+#include "Materials/MaterialInstanceDynamic.h"
+#include "MeshUVChannelInfo.h"
+#include "Misc/LargeWorldRenderPosition.h"
 #include "Misc/MapErrors.h"
-#include "Engine/StaticMesh.h"
-#include "Interfaces/ITargetPlatform.h"
-#include "DeviceProfiles/DeviceProfile.h"
-#include "DeviceProfiles/DeviceProfileManager.h"
-#include "PSOPrecache.h"
-#include "PSOPrecacheMaterial.h"
-#include "PrimitiveSceneProxy.h"
-#include "RenderingThread.h"
-#include "SceneInterface.h"
-#include "Stats/StatsTrace.h"
-
-#if WITH_EDITOR
-#include "Engine/InterpCurveEdSetup.h"
-#include "ObjectEditorUtils.h"
-#endif
-
+#include "Misc/UObjectToken.h"
+#include "ParticleEmitterInstanceOwner.h"
+#include "ParticleEmitterInstances.h"
+#include "ParticleHelper.h"
+#include "ParticleSystemSceneProxy.h"
+#include "Particles/EmitterCameraLensEffectBase.h"
 #include "Particles/FXSystemPrivate.h"
-#include "Particles/Camera/ParticleModuleCameraOffset.h"
-#include "Particles/Collision/ParticleModuleCollision.h"
-#include "Particles/Color/ParticleModuleColorOverLife.h"
 #include "Particles/ParticleEmitter.h"
-#include "Particles/Event/ParticleModuleEventGenerator.h"
-#include "Particles/Event/ParticleModuleEventReceiverBase.h"
-#include "Particles/Lifetime/ParticleModuleLifetimeBase.h"
-#include "Particles/Lifetime/ParticleModuleLifetime.h"
-#include "Particles/Light/ParticleModuleLight.h"
-#include "Particles/Location/ParticleModuleLocationBoneSocket.h"
-#include "Particles/Material/ParticleModuleMeshMaterial.h"
-#include "Particles/Modules/Location/ParticleModulePivotOffset.h"
-#include "Particles/Orbit/ParticleModuleOrbit.h"
-#include "Particles/Parameter/ParticleModuleParameterDynamic.h"
-#include "Particles/Size/ParticleModuleSize.h"
-#include "Particles/Spawn/ParticleModuleSpawnBase.h"
-#include "Particles/Spawn/ParticleModuleSpawn.h"
-#include "Particles/TypeData/ParticleModuleTypeDataBase.h"
-#include "Particles/TypeData/ParticleModuleTypeDataBeam2.h"
-#include "Particles/ParticleSpriteEmitter.h"
-#include "Particles/TypeData/ParticleModuleTypeDataGpu.h"
-#include "Particles/TypeData/ParticleModuleTypeDataMesh.h"
-#include "Particles/Velocity/ParticleModuleVelocity.h"
 #include "Particles/ParticleEventManager.h"
 #include "Particles/ParticleLODLevel.h"
 #include "Particles/ParticleModuleRequired.h"
-#include "Particles/ParticleSystemReplay.h"
-#include "Distributions/DistributionFloatConstantCurve.h"
-#include "Particles/SubUV/ParticleModuleSubUV.h"
 #include "Particles/ParticleSystemManager.h"
+#include "Particles/ParticleSystemReplay.h"
+#include "Particles/TypeData/ParticleModuleTypeDataMesh.h"
+#include "SceneInterface.h"
+#include "UObject/FrameworkObjectVersion.h"
+#include "UObject/UObjectIterator.h"
+#include "UnrealEngine.h"
 
 #include UE_INLINE_GENERATED_CPP_BY_NAME(ParticleSystemComponent)
 
@@ -132,50 +84,6 @@ FCustomVersionRegistration GRegisterParticleSystemCustomVersion(FParticleSystemC
 
 //////////////////////////////////////////////////////////////////////////
 
-void UFXSystemAsset::LaunchPSOPrecaching(const FMaterialInterfacePSOPrecacheParamsList& PSOPrecacheParamsList)
-{
-	FGraphEventArray PrecachePSOsEvents;
-	if (IsComponentPSOPrecachingEnabled())
-	{
-		PrecacheMaterialPSOs(PSOPrecacheParamsList, MaterialPSOPrecacheRequestIDs, PrecachePSOsEvents);
-	} 
-
-	// Create task to signal that the PSO precache events are done by adding them as prerequisite to the task.
-	if (PrecachePSOsEvents.Num() > 0)
-	{
-		struct FReleasePrecachePSOsEventTask
-		{
-			explicit FReleasePrecachePSOsEventTask(UFXSystemAsset* OwnerAsset)
-				: WeakOwnerAsset(OwnerAsset)
-			{
-			}
-
-			static TStatId GetStatId() { return TStatId(); }
-			static ENamedThreads::Type GetDesiredThread() { return ENamedThreads::GameThread; }
-			static ESubsequentsMode::Type GetSubsequentsMode() { return ESubsequentsMode::TrackSubsequents; }
-
-			void DoTask(ENamedThreads::Type CurrentThread, const FGraphEventRef& MyCompletionGraphEvent)
-			{
-				if (UFXSystemAsset* Asset = WeakOwnerAsset.Get())
-				{
-					Asset->PrecachePSOsEvent = nullptr;
-				}
-			}
-
-			TWeakObjectPtr<UFXSystemAsset> WeakOwnerAsset;
-		};
-
-		// need to set `PrecachePSOsEvent` before the task is launched to not race with its execution
-		TGraphTask<FReleasePrecachePSOsEventTask>* ReleasePrecachePSOsEventTask = TGraphTask<FReleasePrecachePSOsEventTask>::CreateTask(&PrecachePSOsEvents).ConstructAndHold(this);
-		PrecachePSOsEvent = ReleasePrecachePSOsEventTask->GetCompletionEvent();
-		ReleasePrecachePSOsEventTask->Unlock();
-	}
-
-	PSOPrecachingLaunched = true;
-}
-
-//////////////////////////////////////////////////////////////////////////
-
 int32 GParticleLODBias = 0;
 FAutoConsoleVariableRef CVarParticleLODBias(
 	TEXT("r.ParticleLODBias"),
@@ -184,16 +92,6 @@ FAutoConsoleVariableRef CVarParticleLODBias(
 	ECVF_Scalability
 	);
 
-static TAutoConsoleVariable<float> CVarQLSpawnRateReferenceLevel(
-	TEXT("fx.QualityLevelSpawnRateScaleReferenceLevel"),
-	2,
-	TEXT("Controls the reference level for quality level based spawn rate scaling. This is the FX quality level\n")
-	TEXT("at which spawn rate is not scaled down; Spawn rate scaling will happen by each emitter's\n")
-	TEXT("QualityLevelSpawnRateScale value for each reduction in level below the reference level.\n")
-	TEXT("\n")
-	TEXT("Default = 2. Value should range from 0 to the maximum FX quality level."),
-	ECVF_Scalability);
-
 static TAutoConsoleVariable<float> CVarPruneEmittersOnCookByDetailMode(
 	TEXT("fx.PruneEmittersOnCookByDetailMode"),
 	0,
@@ -236,3311 +134,162 @@ ENGINE_API bool GIsAllowingParticles = true;
 bool GbEnableGameThreadLODCalculation = true;
 
 namespace CascadeLocal
-{
-	bool			bUseTemplateDenyList = false;
-	TSet<FName>		TemplateDenyList;
-	FString			TemplateDenyListString;
-
-	static void UpdateTemplateDenyList(IConsoleVariable*)
-	{
-		TArray<FString> Names;
-		TemplateDenyListString.ParseIntoArray(Names, TEXT(","));
-
-		TemplateDenyList.Empty();
-		for (const FString& Name : Names)
-		{
-			TemplateDenyList.Emplace(Name);
-		}
-
-		bUseTemplateDenyList = TemplateDenyList.Num() > 0;
-	}
-
-	bool AllowTemplate(UParticleSystem* Template)
-	{
-		return !bUseTemplateDenyList || (Template && !TemplateDenyList.Contains(Template->GetFName()));
-	}
-
-	static FAutoConsoleVariableRef CVarCascadeSetTemplateDenyList(
-		TEXT("fx.Cascade.SetTemplateDenyList"),
-		TemplateDenyListString,
-		TEXT("Set the template deny List to use. (i.e. P_SystemA,P_SystemB)"),
-		FConsoleVariableDelegate::CreateStatic(UpdateTemplateDenyList),
-		ECVF_Scalability | ECVF_Default
-	);
-}
-
-// Comment this in to debug empty emitter instance templates...
-//#define _PSYSCOMP_DEBUG_INVALID_EMITTER_INSTANCE_TEMPLATES_
-
-/*-----------------------------------------------------------------------------
-	Particle scene view
------------------------------------------------------------------------------*/
-FSceneView*			GParticleView = NULL;
-
-/*-----------------------------------------------------------------------------
-	Conversion functions
------------------------------------------------------------------------------*/
-void Particle_ModifyFloatDistribution(UDistributionFloat* pkDistribution, float fScale)
-{
-	if (pkDistribution->IsA(UDistributionFloatConstant::StaticClass()))
-	{
-		UDistributionFloatConstant* pkDistConstant = Cast<UDistributionFloatConstant>(pkDistribution);
-		pkDistConstant->Constant *= fScale;
-	}
-	else
-	if (pkDistribution->IsA(UDistributionFloatUniform::StaticClass()))
-	{
-		UDistributionFloatUniform* pkDistUniform = Cast<UDistributionFloatUniform>(pkDistribution);
-		pkDistUniform->Min *= fScale;
-		pkDistUniform->Max *= fScale;
-	}
-	else
-	if (pkDistribution->IsA(UDistributionFloatConstantCurve::StaticClass()))
-	{
-		UDistributionFloatConstantCurve* pkDistCurve = Cast<UDistributionFloatConstantCurve>(pkDistribution);
-
-		int32 iKeys = pkDistCurve->GetNumKeys();
-		int32 iCurves = pkDistCurve->GetNumSubCurves();
-
-		for (int32 KeyIndex = 0; KeyIndex < iKeys; KeyIndex++)
-		{
-			float fKeyIn = pkDistCurve->GetKeyIn(KeyIndex);
-			for (int32 SubIndex = 0; SubIndex < iCurves; SubIndex++)
-			{
-				float fKeyOut = pkDistCurve->GetKeyOut(SubIndex, KeyIndex);
-				float ArriveTangent;
-				float LeaveTangent;
-				pkDistCurve->GetTangents(SubIndex, KeyIndex, ArriveTangent, LeaveTangent);
-
-				pkDistCurve->SetKeyOut(SubIndex, KeyIndex, fKeyOut * fScale);
-				pkDistCurve->SetTangents(SubIndex, KeyIndex, ArriveTangent * fScale, LeaveTangent * fScale);
-			}
-		}
-	}
-}
-
-void Particle_ModifyVectorDistribution(UDistributionVector* pkDistribution, FVector& vScale)
-{
-	if (pkDistribution->IsA(UDistributionVectorConstant::StaticClass()))
-	{
-		UDistributionVectorConstant* pkDistConstant = Cast<UDistributionVectorConstant>(pkDistribution);
-		pkDistConstant->Constant *= vScale;
-	}
-	else
-	if (pkDistribution->IsA(UDistributionVectorUniform::StaticClass()))
-	{
-		UDistributionVectorUniform* pkDistUniform = Cast<UDistributionVectorUniform>(pkDistribution);
-		pkDistUniform->Min *= vScale;
-		pkDistUniform->Max *= vScale;
-	}
-	else
-	if (pkDistribution->IsA(UDistributionVectorConstantCurve::StaticClass()))
-	{
-		UDistributionVectorConstantCurve* pkDistCurve = Cast<UDistributionVectorConstantCurve>(pkDistribution);
-
-		int32 iKeys = pkDistCurve->GetNumKeys();
-		int32 iCurves = pkDistCurve->GetNumSubCurves();
-
-		for (int32 KeyIndex = 0; KeyIndex < iKeys; KeyIndex++)
-		{
-			float fKeyIn = pkDistCurve->GetKeyIn(KeyIndex);
-			for (int32 SubIndex = 0; SubIndex < iCurves; SubIndex++)
-			{
-				float fKeyOut = pkDistCurve->GetKeyOut(SubIndex, KeyIndex);
-				float ArriveTangent;
-				float LeaveTangent;
-				pkDistCurve->GetTangents(SubIndex, KeyIndex, ArriveTangent, LeaveTangent);
-
-				switch (SubIndex)
-				{
-				case 1:
-					pkDistCurve->SetKeyOut(SubIndex, KeyIndex, fKeyOut * vScale.Y);
-					pkDistCurve->SetTangents(SubIndex, KeyIndex, ArriveTangent * vScale.Y, LeaveTangent * vScale.Y);
-					break;
-				case 2:
-					pkDistCurve->SetKeyOut(SubIndex, KeyIndex, fKeyOut * vScale.Z);
-					pkDistCurve->SetTangents(SubIndex, KeyIndex, ArriveTangent * vScale.Z, LeaveTangent * vScale.Z);
-					break;
-				case 0:
-				default:
-					pkDistCurve->SetKeyOut(SubIndex, KeyIndex, fKeyOut * vScale.X);
-					pkDistCurve->SetTangents(SubIndex, KeyIndex, ArriveTangent * vScale.X, LeaveTangent * vScale.X);
-					break;
-				}
-			}
-		}
-	}
-}
-
-/** Console command to reset all particle components. */
-static void ResetAllParticleComponents()
-{
-	for (TObjectIterator<UParticleSystemComponent> It; It; ++It)
-	{
-		UParticleSystemComponent* ParticleSystemComponent = *It;
-		ParticleSystemComponent->ResetParticles();
-		ParticleSystemComponent->ActivateSystem(true);
-		ParticleSystemComponent->bIsViewRelevanceDirty = true;
-		ParticleSystemComponent->CachedViewRelevanceFlags.Empty();
-		ParticleSystemComponent->ConditionalCacheViewRelevanceFlags();
-		ParticleSystemComponent->ReregisterComponent();
-	}
-}
-static FAutoConsoleCommand GResetAllParticleComponentsCmd(
-	TEXT("FX.RestartAll"),
-	TEXT("Restarts all particle system components"),
-	FConsoleCommandDelegate::CreateStatic(ResetAllParticleComponents)
-	);
-
-/*-----------------------------------------------------------------------------
-	UParticleLODLevel implementation.
------------------------------------------------------------------------------*/
-UParticleLODLevel::UParticleLODLevel(const FObjectInitializer& ObjectInitializer)
-	: Super(ObjectInitializer)
-{
-	bEnabled = true;
-	ConvertedModules = true;
-	PeakActiveParticles = 0;
-}
-
-
-void UParticleLODLevel::CompileModules( FParticleEmitterBuildInfo& EmitterBuildInfo )
-{
-	check( RequiredModule );
-	check( SpawnModule );
-
-	// Store a few special modules.
-	EmitterBuildInfo.RequiredModule = RequiredModule;
-	EmitterBuildInfo.SpawnModule = SpawnModule;
-
-	// Compile those special modules.
-	RequiredModule->CompileModule( EmitterBuildInfo );
-	if ( SpawnModule->bEnabled )
-	{
-		SpawnModule->CompileModule( EmitterBuildInfo );
-	}
-
-	// Compile all remaining modules.
-	const int32 ModuleCount = Modules.Num();
-	for ( int32 ModuleIndex = 0; ModuleIndex < ModuleCount; ++ModuleIndex )
-	{
-		UParticleModule* Module = Modules[ModuleIndex];
-		if ( Module && Module->bEnabled )
-		{
-			Module->CompileModule( EmitterBuildInfo );
-		}
-	}
-
-	// Estimate the maximum number of active particles.
-	EmitterBuildInfo.EstimatedMaxActiveParticleCount = CalculateMaxActiveParticleCount();
-}
-
-bool UParticleLODLevel::IsPostLoadThreadSafe() const
-{
-	return false;
-}
-
-void UParticleLODLevel::PostLoad()
-{
-	Super::PostLoad();
-
-#if WITH_EDITORONLY_DATA
-	checkf(SpawnModule, TEXT("Missing spawn module on %s (%s)"), *(GetPathName()), 
-		(GetOuter() ? (GetOuter()->GetOuter() ? *(GetOuter()->GetOuter()->GetPathName()) : *(GetOuter()->GetPathName())) : TEXT("???")));
-#endif // WITH_EDITORONLY_DATA
-
-	{
-		RequiredModule->ConditionalPostLoad();
-	}
-	if ( SpawnModule )
-	{
-		SpawnModule->ConditionalPostLoad();
-	}
-
-	for (UParticleModule* ParticleModule : Modules)
-	{
-		ParticleModule->ConditionalPostLoad();
-	}
-
-	// shouldn't ever set another UObjects serialized variable in post load
-	// this causes determinisitc cooking issues due to load order being different
-	/*if (RequiredModule)
-	{
-		RequiredModule->ConditionalPostLoad();
-		if (RequiredModule->bEnabled != bEnabled)
-		{
-			RequiredModule->bEnabled = bEnabled;
-		}
-	}*/
-}
-
-void UParticleLODLevel::UpdateModuleLists()
-{
-	LLM_SCOPE(ELLMTag::Particles);
-
-	SpawningModules.Empty();
-	SpawnModules.Empty();
-	UpdateModules.Empty();
-	OrbitModules.Empty();
-	EventReceiverModules.Empty();
-	EventGenerator = NULL;
-
-	UParticleModule* Module;
-	int32 TypeDataModuleIndex = -1;
-
-	for (int32 i = 0; i < Modules.Num(); i++)
-	{
-		Module = Modules[i];
-		if (!Module)
-		{
-			continue;
-		}
-
-		if (Module->bSpawnModule)
-		{
-			SpawnModules.Add(Module);
-		}
-		if (Module->bUpdateModule || Module->bFinalUpdateModule)
-		{
-			UpdateModules.Add(Module);
-		}
-
-		if (Module->IsA(UParticleModuleTypeDataBase::StaticClass()))
-		{
-			TypeDataModule = CastChecked<UParticleModuleTypeDataBase>(Module);
-			if (!Module->bSpawnModule && !Module->bUpdateModule)
-			{
-				// For now, remove it from the list and set it as the TypeDataModule
-				TypeDataModuleIndex = i;
-			}
-		}
-		else
-		if (Module->IsA(UParticleModuleSpawnBase::StaticClass()))
-		{
-			UParticleModuleSpawnBase* SpawnBase = CastChecked<UParticleModuleSpawnBase>(Module);
-			SpawningModules.Add(SpawnBase);
-		}
-		else
-		if (Module->IsA(UParticleModuleOrbit::StaticClass()))
-		{
-			UParticleModuleOrbit* Orbit = CastChecked<UParticleModuleOrbit>(Module);
-			OrbitModules.Add(Orbit);
-		}
-		else
-		if (Module->IsA(UParticleModuleEventGenerator::StaticClass()))
-		{
-			EventGenerator = CastChecked<UParticleModuleEventGenerator>(Module);
-		}
-		else
-		if (Module->IsA(UParticleModuleEventReceiverBase::StaticClass()))
-		{
-			UParticleModuleEventReceiverBase* Event = CastChecked<UParticleModuleEventReceiverBase>(Module);
-			EventReceiverModules.Add(Event);
-		}
-	}
-
-	if (EventGenerator)
-	{
-		// Force the event generator module to the top of the module stack...
-		Modules.RemoveSingle(EventGenerator);
-		Modules.Insert(EventGenerator, 0);
-	}
-
-	if (TypeDataModuleIndex != -1)
-	{
-		Modules.RemoveAt(TypeDataModuleIndex);
-	}
-
-	if (TypeDataModule /**&& (Level == 0)**/)
-	{
-		UParticleModuleTypeDataMesh* MeshTD = Cast<UParticleModuleTypeDataMesh>(TypeDataModule);
-		if (MeshTD
-			&& MeshTD->Mesh
-			&& MeshTD->Mesh->HasValidRenderData(false))
-		{
-			UParticleSpriteEmitter* SpriteEmitter = Cast<UParticleSpriteEmitter>(GetOuter());
-			if (SpriteEmitter && (MeshTD->bOverrideMaterial == false))
-			{
-				FStaticMeshSection& Section = MeshTD->Mesh->GetRenderData()->LODResources[0].Sections[0];
-				UMaterialInterface* Material = MeshTD->Mesh->GetMaterial(Section.MaterialIndex);
-				if (Material)
-				{
-					RequiredModule->Material = Material;
-				}
-			}
-		}
-	}
-}
-
-
-bool UParticleLODLevel::GenerateFromLODLevel(UParticleLODLevel* SourceLODLevel, float Percentage, bool bGenerateModuleData)
-{
-	// See if there are already modules in place
-	if (Modules.Num() > 0)
-	{
-		UE_LOG(LogParticles, Log, TEXT("ERROR? - GenerateFromLODLevel - modules already present!"));
-		return false;
-	}
-
-	bool	bResult	= true;
-
-	// Allocate slots in the array...
-	Modules.InsertZeroed(0, SourceLODLevel->Modules.Num());
-
-	// Set the enabled flag
-	bEnabled = SourceLODLevel->bEnabled;
-
-	// Set up for undo/redo!
-	SetFlags(RF_Transactional);
-
-	// Required module...
-	RequiredModule = CastChecked<UParticleModuleRequired>(
-		SourceLODLevel->RequiredModule->GenerateLODModule(SourceLODLevel, this, Percentage, bGenerateModuleData));
-
-	// Spawn module...
-	SpawnModule = CastChecked<UParticleModuleSpawn>(
-		SourceLODLevel->SpawnModule->GenerateLODModule(SourceLODLevel, this, Percentage, bGenerateModuleData));
-
-	// TypeData module, if present...
-	if (SourceLODLevel->TypeDataModule)
-	{
-		TypeDataModule = 
-			CastChecked<UParticleModuleTypeDataBase>(
-			SourceLODLevel->TypeDataModule->GenerateLODModule(SourceLODLevel, this, Percentage, bGenerateModuleData));
-		check(TypeDataModule == SourceLODLevel->TypeDataModule); // Code expects typedata to be the same across LODs
-	}
-
-	// The remaining modules...
-	for (int32 ModuleIndex = 0; ModuleIndex < SourceLODLevel->Modules.Num(); ModuleIndex++)
-	{
-		if (SourceLODLevel->Modules[ModuleIndex])
-		{
-			Modules[ModuleIndex] = SourceLODLevel->Modules[ModuleIndex]->GenerateLODModule(SourceLODLevel, this, Percentage, bGenerateModuleData);
-		}
-		else
-		{
-			Modules[ModuleIndex] = NULL;
-		}
-	}
-
-	return bResult;
-}
-
-
-int32	UParticleLODLevel::CalculateMaxActiveParticleCount()
-{
-	check(RequiredModule != NULL);
-
-	// Determine the lifetime for particles coming from the emitter
-	float ParticleLifetime = 0.0f;
-	float MaxSpawnRate = SpawnModule->GetEstimatedSpawnRate();
-	int32 MaxBurstCount = SpawnModule->GetMaximumBurstCount();
-	for (int32 ModuleIndex = 0; ModuleIndex < Modules.Num(); ModuleIndex++)
-	{
-		UParticleModuleLifetimeBase* LifetimeMod = Cast<UParticleModuleLifetimeBase>(Modules[ModuleIndex]);
-		if (LifetimeMod != NULL)
-		{
-			ParticleLifetime += LifetimeMod->GetMaxLifetime();
-		}
-
-		UParticleModuleSpawnBase* SpawnMod = Cast<UParticleModuleSpawnBase>(Modules[ModuleIndex]);
-		if (SpawnMod != NULL)
-		{
-			MaxSpawnRate += SpawnMod->GetEstimatedSpawnRate();
-			MaxBurstCount += SpawnMod->GetMaximumBurstCount();
-		}
-	}
-
-	// Determine the maximum duration for this particle system
-	float MaxDuration = 0.0f;
-	float TotalDuration = 0.0f;
-	int32 TotalLoops = 0;
-	if (RequiredModule != NULL)
-	{
-		// We don't care about delay wrt spawning...
-		MaxDuration = FMath::Max<float>(RequiredModule->EmitterDuration, RequiredModule->EmitterDurationLow);
-		TotalLoops = RequiredModule->EmitterLoops;
-		TotalDuration = MaxDuration * TotalLoops;
-	}
-
-	// Determine the max
-	int32 MaxAPC = 0;
-
-	if (TotalDuration != 0.0f)
-	{
-		if (TotalLoops == 1)
-		{
-			// Special case for one loop... 
-			if (ParticleLifetime < MaxDuration)
-			{
-				MaxAPC += FMath::CeilToInt(ParticleLifetime * MaxSpawnRate);
-			}
-			else
-			{
-				MaxAPC += FMath::CeilToInt(MaxDuration * MaxSpawnRate);
-			}
-			// Safety zone...
-			MaxAPC += 1;
-			// Add in the bursts...
-			MaxAPC += MaxBurstCount;
-		}
-		else
-		{
-			if (ParticleLifetime < MaxDuration)
-			{
-				MaxAPC += FMath::CeilToInt(ParticleLifetime * MaxSpawnRate);
-			}
-			else
-			{
-				MaxAPC += (FMath::CeilToInt(FMath::CeilToInt(MaxDuration * MaxSpawnRate) * ParticleLifetime));
-			}
-			// Safety zone...
-			MaxAPC += 1;
-			// Add in the bursts...
-			MaxAPC += MaxBurstCount;
-			if (ParticleLifetime > MaxDuration)
-			{
-				MaxAPC += MaxBurstCount * FMath::CeilToInt(ParticleLifetime - MaxDuration);
-			}
-		}
-	}
-	else
-	{
-		// We are infinite looping... 
-		// Single loop case is all we will worry about. Safer base estimate - but not ideal.
-		if (ParticleLifetime < MaxDuration)
-		{
-			MaxAPC += FMath::CeilToInt(ParticleLifetime * FMath::CeilToInt(MaxSpawnRate));
-		}
-		else
-		{
-			if (ParticleLifetime != 0.0f)
-			{
-				if (ParticleLifetime <= MaxDuration)
-				{
-					MaxAPC += FMath::CeilToInt(MaxDuration * MaxSpawnRate);
-				}
-				else //if (ParticleLifetime > MaxDuration)
-				{
-					MaxAPC += FMath::CeilToInt(MaxDuration * MaxSpawnRate) * ParticleLifetime;
-				}
-			}
-			else
-			{
-				// No lifetime, no duration...
-				MaxAPC += FMath::CeilToInt(MaxSpawnRate);
-			}
-		}
-		// Safety zone...
-		MaxAPC += FMath::Max<int32>(FMath::CeilToInt(MaxSpawnRate * 0.032f), 2);
-		// Burst
-		MaxAPC += MaxBurstCount;
-	}
-
-	PeakActiveParticles = MaxAPC;
-
-	return MaxAPC;
-}
-
-
-void UParticleLODLevel::ConvertToSpawnModule()
-{
-#if WITH_EDITOR
-	// Move the required module SpawnRate and Burst information to a new SpawnModule.
-	if (SpawnModule)
-	{
-//		UE_LOG(LogParticles, Warning, TEXT("LOD Level already has a spawn module!"));
-		return;
-	}
-
-	UParticleEmitter* EmitterOuter = CastChecked<UParticleEmitter>(GetOuter());
-	SpawnModule = NewObject<UParticleModuleSpawn>(EmitterOuter->GetOuter());
-	check(SpawnModule);
-
-	UDistributionFloat* SourceDist = RequiredModule->SpawnRate.Distribution;
-	if (SourceDist)
-	{
-		SpawnModule->Rate.Distribution = Cast<UDistributionFloat>(StaticDuplicateObject(SourceDist, SpawnModule));
-		SpawnModule->Rate.Distribution->bIsDirty = true;
-		SpawnModule->Rate.Initialize();
-	}
-
-	// Now the burst list.
-	int32 BurstCount = RequiredModule->BurstList.Num();
-	if (BurstCount > 0)
-	{
-		SpawnModule->BurstList.AddZeroed(BurstCount);
-		for (int32 BurstIndex = 0; BurstIndex < BurstCount; BurstIndex++)
-		{
-			SpawnModule->BurstList[BurstIndex].Count = RequiredModule->BurstList[BurstIndex].Count;
-			SpawnModule->BurstList[BurstIndex].CountLow = RequiredModule->BurstList[BurstIndex].CountLow;
-			SpawnModule->BurstList[BurstIndex].Time = RequiredModule->BurstList[BurstIndex].Time;
-		}
-	}
-
-	MarkPackageDirty();
-#endif	//#if WITH_EDITOR
-}
-
-
-int32 UParticleLODLevel::GetModuleIndex(UParticleModule* InModule)
-{
-	if (InModule)
-	{
-		if (InModule == RequiredModule)
-		{
-			return INDEX_REQUIREDMODULE;
-		}
-		else if (InModule == SpawnModule)
-		{
-			return INDEX_SPAWNMODULE;
-		}
-		else if (InModule == TypeDataModule)
-		{
-			return INDEX_TYPEDATAMODULE;
-		}
-		else
-		{
-			for (int32 ModuleIndex = 0; ModuleIndex < Modules.Num(); ModuleIndex++)
-			{
-				if (InModule == Modules[ModuleIndex])
-				{
-					return ModuleIndex;
-				}
-			}
-		}
-	}
-
-	return INDEX_NONE;
-}
-
-
-UParticleModule* UParticleLODLevel::GetModuleAtIndex(int32 InIndex)
-{
-	// 'Normal' modules
-	if (InIndex > INDEX_NONE)
-	{
-		if (InIndex < Modules.Num())
-		{
-			return Modules[InIndex];
-		}
-
-		return NULL;
-	}
-
-	switch (InIndex)
-	{
-	case INDEX_REQUIREDMODULE:		return RequiredModule;
-	case INDEX_SPAWNMODULE:			return SpawnModule;
-	case INDEX_TYPEDATAMODULE:		return TypeDataModule;
-	}
-
-	return NULL;
-}
-
-
-void UParticleLODLevel::SetLevelIndex(int32 InLevelIndex)
-{
-	// Remove the 'current' index from the validity flags and set the new one.
-	RequiredModule->LODValidity &= ~(1 << Level);
-	RequiredModule->LODValidity |= (1 << InLevelIndex);
-	SpawnModule->LODValidity &= ~(1 << Level);
-	SpawnModule->LODValidity |= (1 << InLevelIndex);
-	if (TypeDataModule)
-	{
-		TypeDataModule->LODValidity &= ~(1 << Level);
-		TypeDataModule->LODValidity |= (1 << InLevelIndex);
-	}
-	for (int32 ModuleIndex = 0; ModuleIndex < Modules.Num(); ModuleIndex++)
-	{
-		UParticleModule* CheckModule = Modules[ModuleIndex];
-		if (CheckModule)
-		{
-			CheckModule->LODValidity &= ~(1 << Level);
-			CheckModule->LODValidity |= (1 << InLevelIndex);
-		}
-	}
-
-	Level = InLevelIndex;
-}
-
-
-bool UParticleLODLevel::IsModuleEditable(UParticleModule* InModule)
-{
-	// If the module validity flag is not set for this level, it is not editable.
-	if ((InModule->LODValidity & (1 << Level)) == 0)
-	{
-		return false;
-	}
-
-	// If the module is shared w/ higher LOD levels, then it is not editable...
-	int32 Validity = 0;
-	if (Level > 0)
-	{
-		int32 Check = Level - 1;
-		while (Check >= 0)
-		{
-			Validity |= (1 << Check);
-			Check--;
-		}
-
-		if ((Validity & InModule->LODValidity) != 0)
-		{
-			return false;
-		}
-	}
-
-	return true;
-}
-
-/*-----------------------------------------------------------------------------
-	UParticleEmitter implementation.
------------------------------------------------------------------------------*/
-UParticleEmitter::UParticleEmitter(const FObjectInitializer& ObjectInitializer)
-	: Super(ObjectInitializer)
-	, SignificanceLevel(EParticleSignificanceLevel::Critical)
-	, bUseLegacySpawningBehavior(false)
-	, bDisabledLODsKeepEmitterAlive(false)
-	, bDisableWhenInsignficant(0)
-	, QualityLevelSpawnRateScale(1.0f)
-	, DetailModeBitmask(PDM_DefaultValue)
-{
-	// Structure to hold one-time initialization
-	struct FConstructorStatics
-	{
-		FName NAME_Particle_Emitter;
-		FConstructorStatics()
-			: NAME_Particle_Emitter(TEXT("Particle Emitter"))
-		{
-		}
-	};
-	static FConstructorStatics ConstructorStatics;
-
-	EmitterName = ConstructorStatics.NAME_Particle_Emitter;
-	ConvertedModules = true;
-	PeakActiveParticles = 0;
-#if WITH_EDITORONLY_DATA
-	EmitterEditorColor = FColor(0, 150, 150, 255);
-#endif // WITH_EDITORONLY_DATA
-
-}
-
-FParticleEmitterInstance* UParticleEmitter::CreateInstance(IParticleEmitterInstanceOwner& InComponent)
-{
-	UE_LOG(LogParticles, Fatal,TEXT("UParticleEmitter::CreateInstance is pure virtual")); 
-	return NULL; 
-}
-
-void UParticleEmitter::UpdateModuleLists()
-{
-	for (int32 LODIndex = 0; LODIndex < LODLevels.Num(); LODIndex++)
-	{
-		UParticleLODLevel* LODLevel = LODLevels[LODIndex];
-		if (LODLevel)
-		{
-			LODLevel->UpdateModuleLists();
-		}
-	}
-	Build();
-}
-
-/**
- *	Helper function for fixing up LODValidity issues on particle modules...
- *
- *	@param	LODIndex		The index of the LODLevel the module is from.
- *	@param	ModuleIndex		The index of the module being checked.
- *	@param	Emitter			The emitter owner.
- *	@param	CurrModule		The module being checked.
- *
- *	@return	 0		If there was no problem.
- *			 1		If there was a problem and it was fixed.
- *			-1		If there was a problem that couldn't be fixed.
- */
-int32 ParticleEmitterHelper_FixupModuleLODErrors( int32 LODIndex, int32 ModuleIndex, 
-	const UParticleEmitter* Emitter, UParticleModule* CurrModule )
-{
-	int32 Result = 1;
-	bool bIsDirty = false;
-
-	UObject* ModuleOuter = CurrModule->GetOuter();
-	UObject* EmitterOuter = Emitter->GetOuter();
-	if (ModuleOuter != EmitterOuter)
-	{
-		// Module has an incorrect outer
-		CurrModule->Rename(NULL, EmitterOuter, REN_DoNotDirty);
-		bIsDirty = true;
-	}
-
-	if (CurrModule->LODValidity == 0)
-	{
-		// Immediately tag it for this lod level...
-		CurrModule->LODValidity = (1 << LODIndex);
-		bIsDirty = true;
-	}
-	else
-	if (CurrModule->IsUsedInLODLevel(LODIndex) == false)
-	{
-		// Why was this even called here?? 
-		// The assumption is that it should be called for the module in the given lod level...
-		// So, we will tag it with this index.
-		CurrModule->LODValidity |= (1 << LODIndex);
-		bIsDirty = true;
-	}
-
-	if (LODIndex > 0)
-	{
-		int32 CheckIndex = LODIndex - 1;
-		while (CheckIndex >= 0)
-		{
-			if (CurrModule->IsUsedInLODLevel(CheckIndex))
-			{
-				// Ensure that it is the same as the one it THINKS it is shared with...
-				UParticleLODLevel* CheckLODLevel = Emitter->LODLevels[CheckIndex];
-
-				if (CurrModule->IsA(UParticleModuleSpawn::StaticClass()))
-				{
-					if (CheckLODLevel->SpawnModule != CurrModule)
-					{
-						// Fix it up... Turn off the higher LOD flag
-						CurrModule->LODValidity &= ~(1 << CheckIndex);
-						bIsDirty = true;
-					}
-				}
-				else
-				if (CurrModule->IsA(UParticleModuleRequired::StaticClass()))
-				{
-					if (CheckLODLevel->RequiredModule != CurrModule)
-					{
-						// Fix it up... Turn off the higher LOD flag
-						CurrModule->LODValidity &= ~(1 << CheckIndex);
-						bIsDirty = true;
-					}
-				}
-				else
-				if (CurrModule->IsA(UParticleModuleTypeDataBase::StaticClass()))
-				{
-					if (CheckLODLevel->TypeDataModule != CurrModule)
-					{
-						// Fix it up... Turn off the higher LOD flag
-						CurrModule->LODValidity &= ~(1 << CheckIndex);
-						bIsDirty = true;
-					}
-				}
-				else
-				{
-					if (ModuleIndex >= CheckLODLevel->Modules.Num())
-					{
-						UE_LOG(LogParticles, Warning, TEXT("\t\tMismatched module count at %2d in %s"), LODIndex, *(Emitter->GetPathName()));
-						Result = -1;
-					}
-					else
-					{
-						UParticleModule* CheckModule = CheckLODLevel->Modules[ModuleIndex];
-						if (CheckModule != CurrModule)
-						{
-							// Fix it up... Turn off the higher LOD flag
-							CurrModule->LODValidity &= ~(1 << CheckIndex);
-							bIsDirty = true;
-						}
-					}
-				}
-			}
-
-			CheckIndex--;
-		}
-	}
-
-	if ((bIsDirty == true) && IsRunningCommandlet())
-	{
-		CurrModule->MarkPackageDirty();
-		Emitter->MarkPackageDirty();
-	}
-
-	return Result;
-}
-
-bool UParticleEmitter::IsPostLoadThreadSafe() const
-{
-	return false;
-}
-
-void UParticleEmitter::Serialize(FArchive& Ar)
-{
-	Super::Serialize(Ar);
-
-	Ar.UsingCustomVersion(FParticleSystemCustomVersion::GUID);
-}
-
-void UParticleEmitter::PostLoad()
-{
-	Super::PostLoad();
-
-	const int32 PSysVer = GetLinkerCustomVersion(FParticleSystemCustomVersion::GUID);
-	if (PSysVer < FParticleSystemCustomVersion::FixLegacySpawningBugs)
-	{
-		bUseLegacySpawningBehavior = true;
-	}
-
-	if (PSysVer < FParticleSystemCustomVersion::AddEpicDetailMode)
-	{
-		// Init epic detail mode to enabled if high is set
-		if (DetailModeBitmask & (1 << EParticleDetailMode::PDM_High))
-		{
-			DetailModeBitmask |= (1 << EParticleDetailMode::PDM_Epic);
-		}
-	}
-
-	for (int32 LODIndex = 0; LODIndex < LODLevels.Num(); LODIndex++)
-	{
-		UParticleLODLevel* LODLevel = LODLevels[LODIndex];
-		if (LODLevel)
-		{
-			LODLevel->ConditionalPostLoad();
-
-			FLinkerLoad* LODLevelLinker = LODLevel->GetLinker();
-			if (LODLevel->SpawnModule == NULL)
-			{
-				// Force the conversion to SpawnModule
-				UParticleSystem* PSys = Cast<UParticleSystem>(GetOuter());
-				if (PSys)
-				{
-					UE_LOG(LogParticles, Warning, TEXT("LODLevel %d was not converted to spawn module - forcing: %s"), 
-						LODLevel->Level, *(PSys->GetPathName()));
-				}
-				LODLevel->ConvertToSpawnModule();
-			}
-
-			check(LODLevel->SpawnModule);
-
-		}
-	}
-
-#if WITH_EDITORONLY_DATA
-	UpdateDetailModeDisplayString();
-#endif
-
-#if WITH_EDITOR
-	if ((GIsEditor == true) && 1)//(IsRunningCommandlet() == false))
-	{
-		ConvertedModules = false;
-		PeakActiveParticles = 0;
-
-		// Check for improper outers...
-		UObject* EmitterOuter = GetOuter();
-		bool bWarned = false;
-		for (int32 LODIndex = 0; (LODIndex < LODLevels.Num()) && !bWarned; LODIndex++)
-		{
-			UParticleLODLevel* LODLevel = LODLevels[LODIndex];
-			if (LODLevel)
-			{
-				LODLevel->ConditionalPostLoad();
-
-				UParticleModule* Module = LODLevel->TypeDataModule;
-				if (Module)
-				{
-					Module->ConditionalPostLoad();
-
-					UObject* OuterObj = Module->GetOuter();
-					check(OuterObj);
-					if (OuterObj != EmitterOuter)
-					{
-						UE_LOG(LogParticles, Warning, TEXT("UParticleModule %s has an incorrect outer on %s... run FixupEmitters on package %s (%s)"),
-							*(Module->GetPathName()), 
-							*(EmitterOuter->GetPathName()),
-							*(OuterObj->GetOutermost()->GetPathName()),
-							*(GetOutermost()->GetPathName()));
-						UE_LOG(LogParticles, Warning, TEXT("\tModule Outer..............%s"), *(OuterObj->GetPathName()));
-						UE_LOG(LogParticles, Warning, TEXT("\tModule Outermost..........%s"), *(Module->GetOutermost()->GetPathName()));
-						UE_LOG(LogParticles, Warning, TEXT("\tEmitter Outer.............%s"), *(EmitterOuter->GetPathName()));
-						UE_LOG(LogParticles, Warning, TEXT("\tEmitter Outermost.........%s"), *(GetOutermost()->GetPathName()));
-						bWarned = true;
-					}
-				}
-
-				if (!bWarned)
-				{
-					for (int32 ModuleIndex = 0; (ModuleIndex < LODLevel->Modules.Num()) && !bWarned; ModuleIndex++)
-					{
-						Module = LODLevel->Modules[ModuleIndex];
-						if (Module)
-						{
-							Module->ConditionalPostLoad();
-
-							UObject* OuterObj = Module->GetOuter();
-							check(OuterObj);
-							if (OuterObj != EmitterOuter)
-							{
-								UE_LOG(LogParticles, Warning, TEXT("UParticleModule %s has an incorrect outer on %s... run FixupEmitters on package %s (%s)"),
-									*(Module->GetPathName()), 
-									*(EmitterOuter->GetPathName()),
-									*(OuterObj->GetOutermost()->GetPathName()),
-									*(GetOutermost()->GetPathName()));
-								UE_LOG(LogParticles, Warning, TEXT("\tModule Outer..............%s"), *(OuterObj->GetPathName()));
-								UE_LOG(LogParticles, Warning, TEXT("\tModule Outermost..........%s"), *(Module->GetOutermost()->GetPathName()));
-								UE_LOG(LogParticles, Warning, TEXT("\tEmitter Outer.............%s"), *(EmitterOuter->GetPathName()));
-								UE_LOG(LogParticles, Warning, TEXT("\tEmitter Outermost.........%s"), *(GetOutermost()->GetPathName()));
-								bWarned = true;
-							}
-						}
-					}
-				}
-			}
-		}
-	}
-	else
-#endif
-	{
-		for (int32 LODIndex = 0; LODIndex < LODLevels.Num(); LODIndex++)
-		{
-			UParticleLODLevel* LODLevel = LODLevels[LODIndex];
-			if (LODLevel)
-			{
-				LODLevel->ConditionalPostLoad();
-			}
-		}
-	}
-   
-
-	ConvertedModules = true;
-
-	// this will look at all of the emitters and then remove ones that some how have become NULL (e.g. from a removal of an Emitter where content
-	// is still referencing it)
-	for (int32 LODIndex = 0; LODIndex < LODLevels.Num(); LODIndex++)
-	{
-		UParticleLODLevel* LODLevel = LODLevels[LODIndex];
-		if (LODLevel)
-		{
-			for (int32 ModuleIndex = LODLevel->Modules.Num()-1; ModuleIndex >= 0; ModuleIndex--)
-			{
-				UParticleModule* ParticleModule = LODLevel->Modules[ModuleIndex];
-				if( ParticleModule == NULL )
-				{
-					LODLevel->Modules.RemoveAt(ModuleIndex);
-					MarkPackageDirty();
-				}
-			}
-		}
-	}
-
-
-	UObject* MyOuter = GetOuter();
-	UParticleSystem* PSysOuter = Cast<UParticleSystem>(MyOuter);
-	bool bRegenDup = false;
-	if (PSysOuter)
-	{
-		bRegenDup = PSysOuter->bRegenerateLODDuplicate;
-	}
-
-	// Clamp the detail spawn rate scale...
-	QualityLevelSpawnRateScale = FMath::Clamp<float>(QualityLevelSpawnRateScale, 0.0f, 1.0f);
-
-	UpdateModuleLists();
-}
-
-#if WITH_EDITOR
-void UParticleEmitter::PostEditChangeProperty(FPropertyChangedEvent& PropertyChangedEvent)
-{
-	check(GIsEditor);
-
-	// Reset the peak active particle counts.
-	// This could check for changes to SpawnRate and Burst and only reset then,
-	// but since we reset the particle system after any edited property, it
-	// may as well just autoreset the peak counts.
-	for (int32 LODIndex = 0; LODIndex < LODLevels.Num(); LODIndex++)
-	{
-		UParticleLODLevel* LODLevel = LODLevels[LODIndex];
-		if (LODLevel)
-		{
-			LODLevel->PeakActiveParticles	= 1;
-		}
-	}
-
-	UpdateModuleLists();
-
-	for (TObjectIterator<UParticleSystemComponent> It;It;++It)
-	{
-		if (It->Template)
-		{
-			int32 i;
-
-			for (i=0; i<It->Template->Emitters.Num(); i++)
-			{
-				if (It->Template->Emitters[i] == this)
-				{
-					It->UpdateInstances();
-				}
-			}
-		}
-	}
-
-	Super::PostEditChangeProperty(PropertyChangedEvent);
-
-	if (CalculateMaxActiveParticleCount() == false)
-	{
-		//
-	}
-
-	// Clamp the detail spawn rate scale...
-	QualityLevelSpawnRateScale = FMath::Clamp<float>(QualityLevelSpawnRateScale, 0.0f, 1.0f);
-
-#if	WITH_EDITORONLY_DATA
-	UpdateDetailModeDisplayString();
-#endif
-}
-#endif // WITH_EDITOR
-
-void UParticleEmitter::SetEmitterName(FName Name)
-{
-	EmitterName = Name;
-}
-
-FName& UParticleEmitter::GetEmitterName()
-{
-	return EmitterName;
-}
-
-void UParticleEmitter::SetLODCount(int32 LODCount)
-{
-	// 
-}
-
-void UParticleEmitter::AddEmitterCurvesToEditor(UInterpCurveEdSetup* EdSetup)
-{
-	UE_LOG(LogParticles, Log, TEXT("UParticleEmitter::AddEmitterCurvesToEditor> Should no longer be called..."));
-	return;
-}
-
-void UParticleEmitter::RemoveEmitterCurvesFromEditor(UInterpCurveEdSetup* EdSetup)
-{
-	for (int32 LODIndex = 0; LODIndex < LODLevels.Num(); LODIndex++)
-	{
-		UParticleLODLevel* LODLevel = LODLevels[LODIndex];
-		// Remove the typedata curves...
-		if (LODLevel->TypeDataModule && LODLevel->TypeDataModule->IsDisplayedInCurveEd(EdSetup))
-		{
-			LODLevel->TypeDataModule->RemoveModuleCurvesFromEditor(EdSetup);
-		}
-
-		// Remove the spawn module curves...
-		if (LODLevel->SpawnModule && LODLevel->SpawnModule->IsDisplayedInCurveEd(EdSetup))
-		{
-			LODLevel->SpawnModule->RemoveModuleCurvesFromEditor(EdSetup);
-		}
-
-		// Remove each modules curves as well.
-		for (int32 ii = 0; ii < LODLevel->Modules.Num(); ii++)
-		{
-			if (LODLevel->Modules[ii]->IsDisplayedInCurveEd(EdSetup))
-			{
-				// Remove it from the curve editor!
-				LODLevel->Modules[ii]->RemoveModuleCurvesFromEditor(EdSetup);
-			}
-		}
-	}
-}
-
-void UParticleEmitter::ChangeEditorColor(FColor& Color, UInterpCurveEdSetup* EdSetup)
-{
-#if WITH_EDITORONLY_DATA
-	UParticleLODLevel* LODLevel = LODLevels[0];
-	EmitterEditorColor = Color;
-	for (int32 TabIndex = 0; TabIndex < EdSetup->Tabs.Num(); TabIndex++)
-	{
-		FCurveEdTab*	Tab = &(EdSetup->Tabs[TabIndex]);
-		for (int32 CurveIndex = 0; CurveIndex < Tab->Curves.Num(); CurveIndex++)
-		{
-			FCurveEdEntry* Entry	= &(Tab->Curves[CurveIndex]);
-			if (LODLevel->SpawnModule->Rate.Distribution == Entry->CurveObject)
-			{
-				Entry->CurveColor	= Color;
-			}
-		}
-	}
-#endif // WITH_EDITORONLY_DATA
-}
-
-void UParticleEmitter::AutoPopulateInstanceProperties(UParticleSystemComponent* PSysComp)
-{
-	for (int32 LODIndex = 0; LODIndex < LODLevels.Num(); LODIndex++)
-	{
-		UParticleLODLevel* LODLevel	= LODLevels[LODIndex];
-		for (int32 ModuleIndex = 0; ModuleIndex < LODLevel->Modules.Num(); ModuleIndex++)
-		{
-			UParticleModule* Module = LODLevel->Modules[ModuleIndex];	
-			LODLevel->SpawnModule->AutoPopulateInstanceProperties(PSysComp);
-			LODLevel->RequiredModule->AutoPopulateInstanceProperties(PSysComp);
-			if (LODLevel->TypeDataModule)
-			{
-				LODLevel->TypeDataModule->AutoPopulateInstanceProperties(PSysComp);
-			}
-
-			Module->AutoPopulateInstanceProperties(PSysComp);
-		}
-	}
-}
-
-
-int32 UParticleEmitter::CreateLODLevel(int32 LODLevel, bool bGenerateModuleData)
-{
-	int32					LevelIndex		= -1;
-	UParticleLODLevel*	CreatedLODLevel	= NULL;
-
-	if (LODLevels.Num() == 0)
-	{
-		LODLevel = 0;
-	}
-
-	// Is the requested index outside a viable range?
-	if ((LODLevel < 0) || (LODLevel > LODLevels.Num()))
-	{
-		return -1;
-	}
-
-	// NextHighestLODLevel is the one that will be 'copied'
-	UParticleLODLevel*	NextHighestLODLevel	= NULL;
-	int32 NextHighIndex = -1;
-	// NextLowestLODLevel is the one (and all ones lower than it) that will have their LOD indices updated
-	UParticleLODLevel*	NextLowestLODLevel	= NULL;
-	int32 NextLowIndex = -1;
-
-	// Grab the two surrounding LOD levels...
-	if (LODLevel == 0)
-	{
-		// It is being added at the front of the list... (highest)
-		if (LODLevels.Num() > 0)
-		{
-			NextHighestLODLevel = LODLevels[0];
-			NextHighIndex = 0;
-			NextLowestLODLevel = NextHighestLODLevel;
-			NextLowIndex = 0;
-		}
-	}
-	else
-	if (LODLevel > 0)
-	{
-		NextHighestLODLevel = LODLevels[LODLevel - 1];
-		NextHighIndex = LODLevel - 1;
-		if (LODLevel < LODLevels.Num())
-		{
-			NextLowestLODLevel = LODLevels[LODLevel];
-			NextLowIndex = LODLevel;
-		}
-	}
-	
-	// Update the LODLevel index for the lower levels and
-	// offset the LOD validity flags for the modules...
-	if (NextLowestLODLevel)
-	{
-		NextLowestLODLevel->ConditionalPostLoad();
-		for (int32 LowIndex = LODLevels.Num() - 1; LowIndex >= NextLowIndex; LowIndex--)
-		{
-			UParticleLODLevel* LowRemapLevel = LODLevels[LowIndex];
-			if (LowRemapLevel)
-			{
-				LowRemapLevel->SetLevelIndex(LowIndex + 1);
-			}
-		}
-	}
-
-	// Create a ParticleLODLevel
-	CreatedLODLevel = NewObject<UParticleLODLevel>(this);
-	check(CreatedLODLevel);
-
-	CreatedLODLevel->Level = LODLevel;
-	CreatedLODLevel->bEnabled = true;
-	CreatedLODLevel->ConvertedModules = true;
-	CreatedLODLevel->PeakActiveParticles = 0;
-
-	// Determine where to place it...
-	if (LODLevels.Num() == 0)
-	{
-		LODLevels.InsertZeroed(0, 1);
-		LODLevels[0] = CreatedLODLevel;
-		CreatedLODLevel->Level	= 0;
-	}
-	else
-	{
-		LODLevels.InsertZeroed(LODLevel, 1);
-		LODLevels[LODLevel] = CreatedLODLevel;
-		CreatedLODLevel->Level = LODLevel;
-	}
-
-	if (NextHighestLODLevel)
-	{
-		NextHighestLODLevel->ConditionalPostLoad();
-
-		// Generate from the higher LOD level
-		if (CreatedLODLevel->GenerateFromLODLevel(NextHighestLODLevel, 100.0, bGenerateModuleData) == false)
-		{
-			UE_LOG(LogParticles, Warning, TEXT("Failed to generate LOD level %d from level %d"), LODLevel, NextHighestLODLevel->Level);
-		}
-	}
-	else
-	{
-		// Create the RequiredModule
-		UParticleModuleRequired* RequiredModule = NewObject<UParticleModuleRequired>(GetOuter());
-		check(RequiredModule);
-		RequiredModule->SetToSensibleDefaults(this);
-		CreatedLODLevel->RequiredModule	= RequiredModule;
-
-		// The SpawnRate for the required module
-		RequiredModule->bUseLocalSpace			= false;
-		RequiredModule->bKillOnDeactivate		= false;
-		RequiredModule->bKillOnCompleted		= false;
-		RequiredModule->EmitterDuration			= 1.0f;
-		RequiredModule->EmitterLoops			= 0;
-		RequiredModule->ParticleBurstMethod		= EPBM_Instant;
-#if WITH_EDITORONLY_DATA
-		RequiredModule->ModuleEditorColor		= FColor::MakeRandomColor();
-#endif // WITH_EDITORONLY_DATA
-		RequiredModule->InterpolationMethod		= PSUVIM_None;
-		RequiredModule->SubImages_Horizontal	= 1;
-		RequiredModule->SubImages_Vertical		= 1;
-		RequiredModule->bScaleUV				= false;
-		RequiredModule->RandomImageTime			= 0.0f;
-		RequiredModule->RandomImageChanges		= 0;
-		RequiredModule->bEnabled				= true;
-
-		RequiredModule->LODValidity = (1 << LODLevel);
-
-		// There must be a spawn module as well...
-		UParticleModuleSpawn* SpawnModule = NewObject<UParticleModuleSpawn>(GetOuter());
-		check(SpawnModule);
-		CreatedLODLevel->SpawnModule = SpawnModule;
-		SpawnModule->LODValidity = (1 << LODLevel);
-		UDistributionFloatConstant* ConstantSpawn	= Cast<UDistributionFloatConstant>(SpawnModule->Rate.Distribution);
-		ConstantSpawn->Constant					= 10;
-		ConstantSpawn->bIsDirty					= true;
-		SpawnModule->BurstList.Empty();
-
-		// Copy the TypeData module
-		CreatedLODLevel->TypeDataModule			= NULL;
-	}
-
-	LevelIndex	= CreatedLODLevel->Level;
-
-	MarkPackageDirty();
-
-	return LevelIndex;
-}
-
-
-bool UParticleEmitter::IsLODLevelValid(int32 LODLevel)
-{
-	for (int32 LODIndex = 0; LODIndex < LODLevels.Num(); LODIndex++)
-	{
-		UParticleLODLevel* CheckLODLevel	= LODLevels[LODIndex];
-		if (CheckLODLevel->Level == LODLevel)
-		{
-			return true;
-		}
-	}
-
-	return false;
-}
-
-UParticleLODLevel* UParticleEmitter::GetCurrentLODLevel(FParticleEmitterInstance* Instance)
-{
-	if (!FPlatformProperties::HasEditorOnlyData())
-	{
-		return Instance->CurrentLODLevel;
-	}
-	else
-	{
-		// for the game (where we care about perf) we don't branch
-		if (Instance->Component.IsGameWorld() )
-		{
-			return Instance->CurrentLODLevel;
-		}
-		else
-		{
-			EditorUpdateCurrentLOD( Instance );
-			return Instance->CurrentLODLevel;
-		}
-	}
-}
-
-void UParticleEmitter::EditorUpdateCurrentLOD(FParticleEmitterInstance* Instance)
-{
-#if WITH_EDITORONLY_DATA
-	UParticleLODLevel*	CurrentLODLevel	= NULL;
-	UParticleLODLevel*	Higher			= NULL;
-
-	int32 SetLODLevel = -1;
-	if (const UParticleSystem* Template = Instance->Component.GetTemplate())
-	{
-		int32 DesiredLODLevel = Template->EditorLODSetting;
-		if (GIsEditor && GEngine->bEnableEditorPSysRealtimeLOD)
-		{
-			DesiredLODLevel = Instance->Component.GetCurrentLODIndex();
-		}
-
-		for (int32 LevelIndex = 0; LevelIndex < LODLevels.Num(); LevelIndex++)
-		{
-			Higher	= LODLevels[LevelIndex];
-			if (Higher && (Higher->Level == DesiredLODLevel))
-			{
-				SetLODLevel = LevelIndex;
-				break;
-			}
-		}
-	}
-
-	if (SetLODLevel == -1)
-	{
-		SetLODLevel = 0;
-	}
-	Instance->SetCurrentLODIndex(SetLODLevel, false);
-#endif // WITH_EDITORONLY_DATA
-}
-
-
-
-UParticleLODLevel* UParticleEmitter::GetLODLevel(int32 LODLevel)
-{
-	if (LODLevel >= LODLevels.Num())
-	{
-		return NULL;
-	}
-
-	return LODLevels[LODLevel];
-}
-
-
-bool UParticleEmitter::AutogenerateLowestLODLevel(bool bDuplicateHighest)
-{
-	// Didn't find it?
-	if (LODLevels.Num() == 1)
-	{
-		// We need to generate it...
-		LODLevels.InsertZeroed(1, 1);
-		UParticleLODLevel* LODLevel = NewObject<UParticleLODLevel>(this);
-		check(LODLevel);
-		LODLevels[1]					= LODLevel;
-		LODLevel->Level					= 1;
-		LODLevel->ConvertedModules		= true;
-		LODLevel->PeakActiveParticles	= 0;
-
-		// Grab LODLevel 0 for creation
-		UParticleLODLevel* SourceLODLevel	= LODLevels[0];
-
-		LODLevel->bEnabled				= SourceLODLevel->bEnabled;
-
-		float Percentage	= 10.0f;
-		if (SourceLODLevel->TypeDataModule)
-		{
-			UParticleModuleTypeDataBeam2*	Beam2TD		= Cast<UParticleModuleTypeDataBeam2>(SourceLODLevel->TypeDataModule);
-
-			if (Beam2TD)
-			{
-				// For now, don't support LOD on beams and trails
-				Percentage	= 100.0f;
-			}
-		}
-
-		if (bDuplicateHighest == true)
-		{
-			Percentage = 100.0f;
-		}
-
-		if (LODLevel->GenerateFromLODLevel(SourceLODLevel, Percentage) == false)
-		{
-			UE_LOG(LogParticles, Warning, TEXT("Failed to generate LOD level %d from LOD level 0"), 1);
-			return false;
-		}
-
-		MarkPackageDirty();
-		return true;
-	}
-
-	return true;
-}
-
-
-bool UParticleEmitter::CalculateMaxActiveParticleCount()
-{
-	int32	CurrMaxAPC = 0;
-
-	int32 MaxCount = 0;
-	
-	for (int32 LODIndex = 0; LODIndex < LODLevels.Num(); LODIndex++)
-	{
-		UParticleLODLevel* LODLevel = LODLevels[LODIndex];
-		if (LODLevel && LODLevel->bEnabled)
-		{
-			bool bForceMaxCount = false;
-			// Check for beams or trails
-			if ((LODLevel->Level == 0) && (LODLevel->TypeDataModule != NULL))
-			{
-				UParticleModuleTypeDataBeam2* BeamTD = Cast<UParticleModuleTypeDataBeam2>(LODLevel->TypeDataModule);
-				if (BeamTD)
-				{
-					bForceMaxCount = true;
-					MaxCount = BeamTD->MaxBeamCount + 2;
-				}
-			}
-
-			int32 LODMaxAPC = LODLevel->CalculateMaxActiveParticleCount();
-			if (bForceMaxCount == true)
-			{
-				LODLevel->PeakActiveParticles = MaxCount;
-				LODMaxAPC = MaxCount;
-			}
-
-			if (LODMaxAPC > CurrMaxAPC)
-			{
-				if (LODIndex > 0)
-				{
-					// Check for a ridiculous difference in counts...
-					if ((CurrMaxAPC > 0) && (LODMaxAPC / CurrMaxAPC) > 2)
-					{
-						//UE_LOG(LogParticles, Log, TEXT("MaxActiveParticleCount Discrepancy?\n\tLOD %2d, Emitter %16s"), LODIndex, *GetName());
-					}
-				}
-				CurrMaxAPC = LODMaxAPC;
-			}
-		}
-	}
-
-#if WITH_EDITOR
-	if ((GIsEditor == true) && (CurrMaxAPC > 500))
-	{
-		//@todo. Added an option to the emitter to disable this warning - for 
-		// the RARE cases where it is really required to render that many.
-		UE_LOG(LogParticles, Warning, TEXT("MaxCount = %4d for Emitter %s (%s)"),
-			CurrMaxAPC, *(GetName()), GetOuter() ? *(GetOuter()->GetPathName()) : TEXT("????"));
-	}
-#endif
-	return true;
-}
-
-
-void UParticleEmitter::GetParametersUtilized(TArray<FString>& ParticleSysParamList,
-											 TArray<FString>& ParticleParameterList)
-{
-	// Clear the lists
-	ParticleSysParamList.Empty();
-	ParticleParameterList.Empty();
-
-	TArray<UParticleModule*> ProcessedModules;
-	ProcessedModules.Empty();
-
-	for (int32 LODIndex = 0; LODIndex < LODLevels.Num(); LODIndex++)
-	{
-		UParticleLODLevel* LODLevel = LODLevels[LODIndex];
-		if (LODLevel)
-		{
-			int32 FindIndex;
-			// Grab that parameters from each module...
-			check(LODLevel->RequiredModule);
-			if (ProcessedModules.Find(LODLevel->RequiredModule, FindIndex) == false)
-			{
-				LODLevel->RequiredModule->GetParticleSysParamsUtilized(ParticleSysParamList);
-				LODLevel->RequiredModule->GetParticleParametersUtilized(ParticleParameterList);
-				ProcessedModules.AddUnique(LODLevel->RequiredModule);
-			}
-
-			check(LODLevel->SpawnModule);
-			if (ProcessedModules.Find(LODLevel->SpawnModule, FindIndex) == false)
-			{
-				LODLevel->SpawnModule->GetParticleSysParamsUtilized(ParticleSysParamList);
-				LODLevel->SpawnModule->GetParticleParametersUtilized(ParticleParameterList);
-				ProcessedModules.AddUnique(LODLevel->SpawnModule);
-			}
-
-			if (LODLevel->TypeDataModule)
-			{
-				if (ProcessedModules.Find(LODLevel->TypeDataModule, FindIndex) == false)
-				{
-					LODLevel->TypeDataModule->GetParticleSysParamsUtilized(ParticleSysParamList);
-					LODLevel->TypeDataModule->GetParticleParametersUtilized(ParticleParameterList);
-					ProcessedModules.AddUnique(LODLevel->TypeDataModule);
-				}
-			}
-			
-			for (int32 ModuleIndex = 0; ModuleIndex < LODLevel->Modules.Num(); ModuleIndex++)
-			{
-				UParticleModule* Module = LODLevel->Modules[ModuleIndex];
-				if (Module)
-				{
-					if (ProcessedModules.Find(Module, FindIndex) == false)
-					{
-						Module->GetParticleSysParamsUtilized(ParticleSysParamList);
-						Module->GetParticleParametersUtilized(ParticleParameterList);
-						ProcessedModules.AddUnique(Module);
-					}
-				}
-			}
-		}
-	}
-}
-
-
-void UParticleEmitter::Build()
-{
-	const int32 LODCount = LODLevels.Num();
-	if ( LODCount > 0 )
-	{
-		UParticleLODLevel* HighLODLevel = LODLevels[0];
-		check(HighLODLevel);
-		if (HighLODLevel->TypeDataModule != nullptr)
-		{
-			if(HighLODLevel->TypeDataModule->RequiresBuild())
-			{
-				FParticleEmitterBuildInfo EmitterBuildInfo;
-#if WITH_EDITOR
-				if(!GetOutermost()->bIsCookedForEditor)
-				{
-					HighLODLevel->CompileModules( EmitterBuildInfo );
-				}
-#endif
-				HighLODLevel->TypeDataModule->Build( EmitterBuildInfo );
-			}
-
-			// Allow TypeData module to cache pointers to modules
-			HighLODLevel->TypeDataModule->CacheModuleInfo(this);
-		}
-
-		// Cache particle size/offset data for all LOD Levels
-		CacheEmitterModuleInfo();
-	}
-}
-
-void UParticleEmitter::CacheEmitterModuleInfo()
-{
-	// This assert makes sure that packing is as expected.
-	// Added FBaseColor...
-	// Linear color change
-	// Added Flags field
-
-	bRequiresLoopNotification = false;
-	bAxisLockEnabled = false;
-	bMeshRotationActive = false;
-	LockAxisFlags = EPAL_NONE;
-	ModuleOffsetMap.Empty();
-	ModuleInstanceOffsetMap.Empty();
-	ModuleRandomSeedInstanceOffsetMap.Empty();
-	ModulesNeedingInstanceData.Empty();
-	ModulesNeedingRandomSeedInstanceData.Empty();
-	MeshMaterials.Empty();
-	DynamicParameterDataOffset = 0;
-	LightDataOffset = 0;
-	LightVolumetricScatteringIntensity = 0;
-	CameraPayloadOffset = 0;
-	ParticleSize = sizeof(FBaseParticle);
-	ReqInstanceBytes = 0;
-	PivotOffset = FVector2D(-0.5f, -0.5f);
-	TypeDataOffset = 0;
-	TypeDataInstanceOffset = -1;
-	SubUVAnimation = nullptr;
-
-	UParticleLODLevel* HighLODLevel = GetLODLevel(0);
-	check(HighLODLevel);
-
-	UParticleModuleTypeDataBase* HighTypeData = HighLODLevel->TypeDataModule;
-	if (HighTypeData)
-	{
-		int32 ReqBytes = HighTypeData->RequiredBytes(static_cast<UParticleModuleTypeDataBase*>(nullptr));
-		if (ReqBytes)
-		{
-			TypeDataOffset = ParticleSize;
-			ParticleSize += ReqBytes;
-		}
-
-		int32 TempInstanceBytes = HighTypeData->RequiredBytesPerInstance();
-		if (TempInstanceBytes)
-		{
-			TypeDataInstanceOffset = ReqInstanceBytes;
-			ReqInstanceBytes += TempInstanceBytes;
-		}
-	}
-
-	// Grab required module
-	UParticleModuleRequired* RequiredModule = HighLODLevel->RequiredModule;
-	check(RequiredModule);
-	// mesh rotation active if alignment is set
-	bMeshRotationActive = (RequiredModule->ScreenAlignment == PSA_Velocity || RequiredModule->ScreenAlignment == PSA_AwayFromCenter);
-
-	// NOTE: This code assumes that the same module order occurs in all LOD levels
-
-	for (int32 ModuleIdx = 0; ModuleIdx < HighLODLevel->Modules.Num(); ModuleIdx++)
-	{
-		UParticleModule* ParticleModule = HighLODLevel->Modules[ModuleIdx];
-		check(ParticleModule);
-
-		// Loop notification?
-		bRequiresLoopNotification |= (ParticleModule->bEnabled && ParticleModule->RequiresLoopingNotification());
-
-		if (ParticleModule->IsA(UParticleModuleTypeDataBase::StaticClass()) == false)
-		{
-			int32 ReqBytes = ParticleModule->RequiredBytes(HighTypeData);
-			if (ReqBytes)
-			{
-				ModuleOffsetMap.Add(ParticleModule, ParticleSize);
-				if (ParticleModule->IsA(UParticleModuleParameterDynamic::StaticClass()) && (DynamicParameterDataOffset == 0))
-				{
-					DynamicParameterDataOffset = ParticleSize;
-				}
-				if (ParticleModule->IsA(UParticleModuleLight::StaticClass()) && (LightDataOffset == 0))
-				{
-					UParticleModuleLight* ParticleModuleLight = Cast<UParticleModuleLight>(ParticleModule);
-					LightVolumetricScatteringIntensity = ParticleModuleLight->VolumetricScatteringIntensity;
-					LightDataOffset = ParticleSize;
-				}
-				if (ParticleModule->IsA(UParticleModuleCameraOffset::StaticClass()) && (CameraPayloadOffset == 0))
-				{
-					CameraPayloadOffset = ParticleSize;
-				}
-				ParticleSize += ReqBytes;
-			}
-
-			int32 TempInstanceBytes = ParticleModule->RequiredBytesPerInstance();
-			if (TempInstanceBytes > 0)
-			{
-				// Add the high-lodlevel offset to the lookup map
-				ModuleInstanceOffsetMap.Add(ParticleModule, ReqInstanceBytes);
-				// Remember that this module has emitter-instance data
-				ModulesNeedingInstanceData.Add(ParticleModule);
-
-				// Add all the other LODLevel modules, using the same offset.
-				// This removes the need to always also grab the HighestLODLevel pointer.
-				for (int32 LODIdx = 1; LODIdx < LODLevels.Num(); LODIdx++)
-				{
-					UParticleLODLevel* CurLODLevel = LODLevels[LODIdx];
-					ModuleInstanceOffsetMap.Add(CurLODLevel->Modules[ModuleIdx], ReqInstanceBytes);
-				}
-				ReqInstanceBytes += TempInstanceBytes;
-			}
-
-			// Add space for per instance random seed value if required
-			if (FApp::bUseFixedSeed || ParticleModule->bSupportsRandomSeed)
-			{
-				// Add the high-lodlevel offset to the lookup map
-				ModuleRandomSeedInstanceOffsetMap.Add(ParticleModule, ReqInstanceBytes);
-				// Remember that this module has emitter-instance data
-				ModulesNeedingRandomSeedInstanceData.Add(ParticleModule);
-
-				// Add all the other LODLevel modules, using the same offset.
-				// This removes the need to always also grab the HighestLODLevel pointer.
-				for (int32 LODIdx = 1; LODIdx < LODLevels.Num(); LODIdx++)
-				{
-					UParticleLODLevel* CurLODLevel = LODLevels[LODIdx];
-					ModuleRandomSeedInstanceOffsetMap.Add(CurLODLevel->Modules[ModuleIdx], ReqInstanceBytes);
-				}
-
-				ReqInstanceBytes += sizeof(FParticleRandomSeedInstancePayload);
-			}
-		}
-
-		if (ParticleModule->IsA(UParticleModuleOrientationAxisLock::StaticClass()))
-		{
-			UParticleModuleOrientationAxisLock* Module_AxisLock = CastChecked<UParticleModuleOrientationAxisLock>(ParticleModule);
-			bAxisLockEnabled = Module_AxisLock->bEnabled;
-			LockAxisFlags = Module_AxisLock->LockAxisFlags;
-		}
-		else if (ParticleModule->IsA(UParticleModulePivotOffset::StaticClass()))
-		{
-			PivotOffset += Cast<UParticleModulePivotOffset>(ParticleModule)->PivotOffset;
-		}
-		else if (ParticleModule->IsA(UParticleModuleMeshMaterial::StaticClass()))
-		{
-			UParticleModuleMeshMaterial* MeshMaterialModule = CastChecked<UParticleModuleMeshMaterial>(ParticleModule);
-			if (MeshMaterialModule->bEnabled)
-			{
-				MeshMaterials = MeshMaterialModule->MeshMaterials;
-			}
-		}
-		else if (ParticleModule->IsA(UParticleModuleSubUV::StaticClass()))
-		{
-			USubUVAnimation* ModuleSubUVAnimation = Cast<UParticleModuleSubUV>(ParticleModule)->Animation;
-			SubUVAnimation = ModuleSubUVAnimation && ModuleSubUVAnimation->SubUVTexture && ModuleSubUVAnimation->IsBoundingGeometryValid()
-				? ModuleSubUVAnimation
-				: NULL;
-		}
-		// Perform validation / fixup on some modules that can cause crashes if LODs / Modules are out of sync
-		// This should only be applied on uncooked builds to avoid wasting cycles
-		else if ( !FPlatformProperties::RequiresCookedData() )
-		{
-			if (ParticleModule->IsA(UParticleModuleLocationBoneSocket::StaticClass()))
-			{
-				UParticleModuleLocationBoneSocket::ValidateLODLevels(this, ModuleIdx);
-			}
-		}
-
-		// Set bMeshRotationActive if module says so
-		if(!bMeshRotationActive && ParticleModule->TouchesMeshRotation())
-		{
-			bMeshRotationActive = true;
-		}
-	}
-}
-
-float UParticleEmitter::GetQualityLevelSpawnRateMult()
-{
-	int32 EffectsQuality = Scalability::GetEffectsQualityDirect(IsInGameThread() || IsInParallelGameThread());
-	int32 ReferenceLevel = CVarQLSpawnRateReferenceLevel.GetValueOnAnyThread(true);
-	float Level = (ReferenceLevel - EffectsQuality);
-	float Q = FMath::Pow(QualityLevelSpawnRateScale, Level);
-	return FMath::Min(1.0f, Q);
-}
-
-bool UParticleEmitter::HasAnyEnabledLODs()const
-{
-	for (UParticleLODLevel* LodLevel : LODLevels)
-	{
-		if (LodLevel && LodLevel->bEnabled)
-		{
-			return true;
-		}
-	}
-	
-	return false;
-}
-
-#if STATS
-
-DEFINE_STAT(STAT_EmittersStatGroupTester);
-DEFINE_STAT(STAT_EmittersRTStatGroupTester);
-
-void UParticleEmitter::CreateStatID() const
-{
-	QUICK_SCOPE_CYCLE_COUNTER(STAT_UParticleEmitterCreateStatID);
-
-	UObject* Outer = GetOuter();
-	FName OuterName = Outer ? Outer->GetFName() : NAME_None;
-	FString LongName;
-	LongName = FString(TEXT("Emitter")) / OuterName.ToString() / EmitterName.ToString();
-	StatID = FDynamicStats::CreateStatId<FStatGroup_STATGROUP_Emitters>(LongName);
-	StatIDRT = FDynamicStats::CreateStatId<FStatGroup_STATGROUP_EmittersRT>(LongName / TEXT("RT"));
-}
-#endif
-
-bool UParticleEmitter::IsSignificant(EParticleSignificanceLevel RequiredSignificance)
-{
-	UParticleSystem* PSysOuter = CastChecked<UParticleSystem>(GetOuter());
-	EParticleSignificanceLevel Significance = FMath::Min(PSysOuter->MaxSignificanceLevel, SignificanceLevel);
-	return Significance >= RequiredSignificance;
-}
-
-/*-----------------------------------------------------------------------------
-	UParticleSpriteEmitter implementation.
------------------------------------------------------------------------------*/
-UParticleSpriteEmitter::UParticleSpriteEmitter(const FObjectInitializer& ObjectInitializer)
-	: Super(ObjectInitializer)
-{
-}
-
-
-void UParticleSpriteEmitter::PostLoad()
-{
-	Super::PostLoad();
-
-	// Postload the materials
-	for (int32 LODIndex = 0; LODIndex < LODLevels.Num(); LODIndex++)
-	{
-		UParticleLODLevel* LODLevel = LODLevels[LODIndex];
-		if (LODLevel)
-		{
-			UParticleModuleRequired* RequiredModule = LODLevel->RequiredModule;
-			if (RequiredModule)
-			{
-				if (RequiredModule->Material)
-				{
-					RequiredModule->Material->ConditionalPostLoad();
-				}
-			}
-		}
-	}
-}
-
-FParticleEmitterInstance* UParticleSpriteEmitter::CreateInstance(IParticleEmitterInstanceOwner& InComponent)
-{
-	// If this emitter was cooked out or has no valid LOD levels don't create an instance for it.
-	if ((bCookedOut == true) || (LODLevels.Num() == 0))
-	{
-		return NULL;
-	}
-
-	FParticleEmitterInstance* Instance = 0;
-
-	UParticleLODLevel* LODLevel	= GetLODLevel(0);
-	check(LODLevel);
-
-	if (LODLevel->TypeDataModule)
-	{
-		//@todo. This will NOT work for trails/beams!
-		Instance = LODLevel->TypeDataModule->CreateInstance(this, InComponent);
-	}
-	else
-	{
-		Instance = new FParticleSpriteEmitterInstance(InComponent);
-		check(Instance);
-		Instance->InitParameters(this);
-	}
-
-	if (Instance)
-	{
-		Instance->CurrentLODLevelIndex	= 0;
-		Instance->CurrentLODLevel		= LODLevels[Instance->CurrentLODLevelIndex];
-		Instance->Init();
-	}
-
-	return Instance;
-}
-
-void UParticleSpriteEmitter::SetToSensibleDefaults()
-{
-#if WITH_EDITOR
-	PreEditChange(NULL);
-#endif // WITH_EDITOR
-
-	UParticleLODLevel* LODLevel = LODLevels[0];
-
-	// Spawn rate
-	LODLevel->SpawnModule->LODValidity = 1;
-	UDistributionFloatConstant* SpawnRateDist = Cast<UDistributionFloatConstant>(LODLevel->SpawnModule->Rate.Distribution);
-	if (SpawnRateDist)
-	{
-		SpawnRateDist->Constant = 20.f;
-	}
-
-	// Create basic set of modules
-
-	// Lifetime module
-	UParticleModuleLifetime* LifetimeModule = NewObject<UParticleModuleLifetime>(GetOuter());
-	UDistributionFloatUniform* LifetimeDist = Cast<UDistributionFloatUniform>(LifetimeModule->Lifetime.Distribution);
-	if (LifetimeDist)
-	{
-		LifetimeDist->Min = 1.0f;
-		LifetimeDist->Max = 1.0f;
-		LifetimeDist->bIsDirty = true;
-	}
-	LifetimeModule->LODValidity = 1;
-	LODLevel->Modules.Add(LifetimeModule);
-
-	// Size module
-	UParticleModuleSize* SizeModule = NewObject<UParticleModuleSize>(GetOuter());
-	UDistributionVectorUniform* SizeDist = Cast<UDistributionVectorUniform>(SizeModule->StartSize.Distribution);
-	if (SizeDist)
-	{
-		SizeDist->Min = FVector(25.f, 25.f, 25.f);
-		SizeDist->Max = FVector(25.f, 25.f, 25.f);
-		SizeDist->bIsDirty = true;
-	}
-	SizeModule->LODValidity = 1;
-	LODLevel->Modules.Add(SizeModule);
-
-	// Initial velocity module
-	UParticleModuleVelocity* VelModule = NewObject<UParticleModuleVelocity>(GetOuter());
-	UDistributionVectorUniform* VelDist = Cast<UDistributionVectorUniform>(VelModule->StartVelocity.Distribution);
-	if (VelDist)
-	{
-		VelDist->Min = FVector(-10.f, -10.f, 50.f);
-		VelDist->Max = FVector(10.f, 10.f, 100.f);
-		VelDist->bIsDirty = true;
-	}
-	VelModule->LODValidity = 1;
-	LODLevel->Modules.Add(VelModule);
-
-	// Color over life module
-	UParticleModuleColorOverLife* ColorModule = NewObject<UParticleModuleColorOverLife>(GetOuter());
-	UDistributionVectorConstantCurve* ColorCurveDist = Cast<UDistributionVectorConstantCurve>(ColorModule->ColorOverLife.Distribution);
-	if (ColorCurveDist)
-	{
-		// Add two points, one at time 0.0f and one at 1.0f
-		for (int32 Key = 0; Key < 2; Key++)
-		{
-			int32	KeyIndex = ColorCurveDist->CreateNewKey(Key * 1.0f);
-			for (int32 SubIndex = 0; SubIndex < 3; SubIndex++)
-			{
-				ColorCurveDist->SetKeyOut(SubIndex, KeyIndex, 1.0f);
-			}
-		}
-		ColorCurveDist->bIsDirty = true;
-	}
-	ColorModule->AlphaOverLife.Distribution = NewObject<UDistributionFloatConstantCurve>(ColorModule);
-	UDistributionFloatConstantCurve* AlphaCurveDist = Cast<UDistributionFloatConstantCurve>(ColorModule->AlphaOverLife.Distribution);
-	if (AlphaCurveDist)
-	{
-		// Add two points, one at time 0.0f and one at 1.0f
-		for (int32 Key = 0; Key < 2; Key++)
-		{
-			int32	KeyIndex = AlphaCurveDist->CreateNewKey(Key * 1.0f);
-			if (Key == 0)
-			{
-				AlphaCurveDist->SetKeyOut(0, KeyIndex, 1.0f);
-			}
-			else
-			{
-				AlphaCurveDist->SetKeyOut(0, KeyIndex, 0.0f);
-			}
-		}
-		AlphaCurveDist->bIsDirty = true;
-	}
-	ColorModule->LODValidity = 1;
-	LODLevel->Modules.Add(ColorModule);
-
-#if WITH_EDITOR
-	PostEditChange();
-#endif // WITH_EDITOR
-}
-
-/*-----------------------------------------------------------------------------
-	UFXSystemAsset implementation.
------------------------------------------------------------------------------*/
-
-UFXSystemAsset::UFXSystemAsset(const FObjectInitializer& ObjectInitializer)
-	: Super(ObjectInitializer)
-{
-}
-
-/*-----------------------------------------------------------------------------
-	UParticleSystem implementation.
------------------------------------------------------------------------------*/
-
-UParticleSystem::UParticleSystem(const FObjectInitializer& ObjectInitializer)
-	: Super(ObjectInitializer)
-	, OcclusionBoundsMethod(EPSOBM_ParticleBounds)
-	, bAnyEmitterLoopsForever(false)
-	, HighestSignificance(EParticleSignificanceLevel::Critical)
-	, LowestSignificance(EParticleSignificanceLevel::Low)
-	, bShouldManageSignificance(false)
-	, bIsImmortal(false)
-	, bWillBecomeZombie(false)
-{
-#if WITH_EDITORONLY_DATA
-	ThumbnailDistance = 200.0;
-	ThumbnailWarmup = 1.0;
-#endif // WITH_EDITORONLY_DATA
-	UpdateTime_FPS = 60.0f;
-	UpdateTime_Delta = 1.0f/60.0f;
-	WarmupTime = 0.0f;
-	WarmupTickRate = 0.0f;
-#if WITH_EDITORONLY_DATA
-	EditorLODSetting = 0;
-#endif // WITH_EDITORONLY_DATA
-	FixedRelativeBoundingBox.Min = FVector(-1.0f, -1.0f, -1.0f);
-	FixedRelativeBoundingBox.Max = FVector(1.0f, 1.0f, 1.0f);
-	FixedRelativeBoundingBox.IsValid = true;
-
-	LODMethod = PARTICLESYSTEMLODMETHOD_Automatic;
-	LODDistanceCheckTime = 0.25f;
-	bRegenerateLODDuplicate = false;
-	ThumbnailImageOutOfDate = true;
-#if WITH_EDITORONLY_DATA
-	FloorMesh = TEXT("/Engine/EditorMeshes/AnimTreeEd_PreviewFloor.AnimTreeEd_PreviewFloor");
-	FloorPosition = FVector(0.0f, 0.0f, 0.0f);
-	FloorRotation = FRotator(0.0f, 0.0f, 0.0f);
-	FloorScale = 1.0f;
-	FloorScale3D = FVector(1.0f, 1.0f, 1.0f);
-#endif // WITH_EDITORONLY_DATA
-
-	MacroUVPosition = FVector(0.0f, 0.0f, 0.0f);
-
-	MacroUVRadius = 200.0f;
-	bAutoDeactivate = true;
-	MinTimeBetweenTicks = 0;
-	InsignificantReaction = EParticleSystemInsignificanceReaction::Auto;
-	InsignificanceDelay = 0.0f;
-	MaxSignificanceLevel = EParticleSignificanceLevel::Critical;
-	MaxPoolSize = 32;
-
-
-	bAllowManagedTicking = true;
-}
-
-
-
-
-ParticleSystemLODMethod UParticleSystem::GetCurrentLODMethod()
-{
-	return ParticleSystemLODMethod(LODMethod);
-}
-
-
-int32 UParticleSystem::GetLODLevelCount()
-{
-	return LODDistances.Num();
-}
-
-
-float UParticleSystem::GetLODDistance(int32 LODLevelIndex)
-{
-	if (LODLevelIndex >= LODDistances.Num())
-	{
-		return -1.0f;
-	}
-
-	return LODDistances[LODLevelIndex];
-}
-
-void UParticleSystem::SetCurrentLODMethod(ParticleSystemLODMethod InMethod)
-{
-	LODMethod = InMethod;
-}
-
-
-bool UParticleSystem::SetLODDistance(int32 LODLevelIndex, float InDistance)
-{
-	if (LODLevelIndex >= LODDistances.Num())
-	{
-		return false;
-	}
-
-	LODDistances[LODLevelIndex] = InDistance;
-
-	return true;
-}
-
-bool UParticleSystem::DoesAnyEmitterHaveMotionBlur(int32 LODLevelIndex) const
-{
-	for (auto& EmitterIter : Emitters)
-	{
-		if (EmitterIter)
-		{
-			auto* EmitterLOD = EmitterIter->GetLODLevel(LODLevelIndex);
-			if (!EmitterLOD)
-			{
-				continue;
-			}
-
-			if (EmitterLOD->TypeDataModule && EmitterLOD->TypeDataModule->IsMotionBlurEnabled())
-			{
-				return true;
-			}
-
-			if (EmitterLOD->RequiredModule && EmitterLOD->RequiredModule->ShouldUseVelocityForMotionBlur())
-			{
-				return true;
-			}
-		}
-	}
-
-	return false;
-}
-
-#if WITH_EDITOR
-void UParticleSystem::PostEditChangeProperty(FPropertyChangedEvent& PropertyChangedEvent)
-{
-	UpdateTime_Delta = 1.0f / UpdateTime_FPS;
-
-	bIsElligibleForAsyncTickComputed = false;
-
-	//If the property is NULL then we don't really know what's happened. 
-	//Could well be a module change, requiring all instances to be destroyed and recreated.
-	bool bEmptyInstances = PropertyChangedEvent.Property == NULL;
-	for (TObjectIterator<UParticleSystemComponent> It;It;++It)
-	{
-		if (It->Template == this)
-		{
-			It->UpdateInstances(bEmptyInstances);
-		}
-	}
-
-	// Ensure the bounds have a positive size
-	if (FixedRelativeBoundingBox.IsValid)
-	{
-		if (FixedRelativeBoundingBox.Min.X > FixedRelativeBoundingBox.Max.X)
-		{
-			Swap(FixedRelativeBoundingBox.Min.X, FixedRelativeBoundingBox.Max.X);
-		}
-		if (FixedRelativeBoundingBox.Min.Y > FixedRelativeBoundingBox.Max.Y)
-		{
-			Swap(FixedRelativeBoundingBox.Min.Y, FixedRelativeBoundingBox.Max.Y);
-		}
-		if (FixedRelativeBoundingBox.Min.Z > FixedRelativeBoundingBox.Max.Z)
-		{
-			Swap(FixedRelativeBoundingBox.Min.Z, FixedRelativeBoundingBox.Max.Z);
-		}
-	}
-
-	// Recompute the looping flag
-	bAnyEmitterLoopsForever = false;
-	bIsImmortal = false;
-	bWillBecomeZombie = false;
-	HighestSignificance = EParticleSignificanceLevel::Low;
-	LowestSignificance = EParticleSignificanceLevel::Critical;
-	for (UParticleEmitter* Emitter : Emitters)
-	{
-		if (Emitter != nullptr)
-		{
-			for (const UParticleLODLevel* LODLevel : Emitter->LODLevels)
-			{
-				if (LODLevel != nullptr)
-				{
-					if (LODLevel->bEnabled)
-					{
-						const UParticleModuleRequired* RequiredModule = LODLevel->RequiredModule;
-						if ((RequiredModule != nullptr) && (RequiredModule->EmitterLoops == 0))
-						{
-							bAnyEmitterLoopsForever = true;
-
-							UParticleModuleSpawn* SpawnModule = LODLevel->SpawnModule;
-							check(SpawnModule);
-
-							// check if any emitter will cause the system to never be deleted
-							// terms here are zombie (burst-only, so will stop spawning but emitter instances and psys component will continue existing)
-							// and immortal (any emitter will loop indefinitely and does not have finite duration)
-							if (RequiredModule->EmitterDuration == 0.0f)
-							{
-								bIsImmortal = true;
-								if (SpawnModule->GetMaximumSpawnRate() == 0.0f && !bAutoDeactivate)
-								{
-									bWillBecomeZombie = true;
-								}
-							}
-						}
-					}
-				}
-			}
-
-			EParticleSignificanceLevel EmitterSignificance = FMath::Min(MaxSignificanceLevel, Emitter->SignificanceLevel);
-			if (EmitterSignificance > HighestSignificance)
-			{
-				HighestSignificance = EmitterSignificance;
-			}
-			if (EmitterSignificance < LowestSignificance)
-			{
-				LowestSignificance = EmitterSignificance;
-			}
-		}
-	}
-
-	bShouldManageSignificance = GetLowestSignificance() != EParticleSignificanceLevel::Critical/* && !ContainsEmitterType(UParticleModuleTypeDataBeam2::StaticClass())*/;
-
-	//cap the WarmupTickRate to realistic values
-	if (WarmupTickRate <= 0)
-	{
-		WarmupTickRate = 0;
-	}
-	else if (WarmupTickRate > WarmupTime)
-	{
-		WarmupTickRate = WarmupTime;
-	}
-
-	ThumbnailImageOutOfDate = true;
-
-	Super::PostEditChangeProperty(PropertyChangedEvent);
-}
-#endif // WITH_EDITOR
-
-void UParticleSystem::PreSave(const class ITargetPlatform* TargetPlatform)
-{
-	PRAGMA_DISABLE_DEPRECATION_WARNINGS;
-	Super::PreSave(TargetPlatform);
-	PRAGMA_ENABLE_DEPRECATION_WARNINGS;
-}
-
-void UParticleSystem::PreSave(FObjectPreSaveContext ObjectSaveContext)
-{
-	Super::PreSave(ObjectSaveContext);
-#if WITH_EDITORONLY_DATA
-	// Ensure that soloing is undone...
-	int32 NumEmitters = FMath::Min(Emitters.Num(),SoloTracking.Num());
-	for (int32 EmitterIdx = 0; EmitterIdx < NumEmitters; EmitterIdx++)
-	{
-		UParticleEmitter* Emitter = Emitters[EmitterIdx];
-		if (Emitter != nullptr)
-		{
-			Emitter->bIsSoloing = false;
-			FLODSoloTrack& SoloTrack = SoloTracking[EmitterIdx];
-			int32 NumLODs = FMath::Min(Emitter->LODLevels.Num(), SoloTrack.SoloEnableSetting.Num());
-			for (int32 LODIdx = 0; LODIdx < NumLODs; LODIdx++)
-			{
-				UParticleLODLevel* LODLevel = Emitter->LODLevels[LODIdx];
-				{
-					// Restore the enabled settings - ie turn off soloing...
-					LODLevel->bEnabled = SoloTrack.SoloEnableSetting[LODIdx];
-				}
-			}
-		}
-	}
-#endif // WITH_EDITORONLY_DATA
-}
-
-bool UParticleSystem::IsPostLoadThreadSafe() const
-{
-	return false;
-}
-
-void UParticleSystem::PostLoad()
-{
-	Super::PostLoad();
-
-	// Run thru all of the emitters, load them up and compute some flags based on them
-	bHasPhysics = false;
-	bAnyEmitterLoopsForever = false;
-	bIsImmortal = false;
-	bWillBecomeZombie = false;
-	HighestSignificance = EParticleSignificanceLevel::Low;
-	LowestSignificance = EParticleSignificanceLevel::Critical;
-	for (int32 i = Emitters.Num() - 1; i >= 0; i--)
-	{
-		// Remove any old emitters
-		UParticleEmitter* Emitter = Emitters[i];
-		if (Emitter == NULL)
-		{
-			// Empty emitter slots are ok with cooked content.
-			if( !FPlatformProperties::RequiresCookedData() && !GIsServer)
-			{
-				UE_LOG(LogParticles, Warning, TEXT("ParticleSystem contains empty emitter slots - %s"), *GetFullName());
-			}
-			continue;
-		}
-
-		Emitter->ConditionalPostLoad();
-
-		bool bCookedOut = false;
-		if (UParticleSpriteEmitter* SpriteEmitter = Cast<UParticleSpriteEmitter>(Emitter))
-		{
-			bCookedOut = SpriteEmitter->bCookedOut;
-		}
-
-		if (!bCookedOut)
-		{
-			if (Emitter->LODLevels.Num() == 0)
-			{
-				UE_LOG(LogParticles, Warning, TEXT("ParticleSystem contains emitter with no lod levels - %s - %s"), *GetFullName(), *Emitter->GetFullName());
-				continue;
-			}
-
-			UParticleLODLevel* LODLevel = Emitter->LODLevels[0];
-			check(LODLevel);
-
-			LODLevel->ConditionalPostLoad();
-				
-			//@todo. Move these flag calculations into the editor and serialize?
-			//Should mirror implementation in PostEditChangeProperty.
-			for (UParticleLODLevel* ParticleLODLevel : Emitter->LODLevels)
-			{
-				if (ParticleLODLevel)
-				{
-					//@todo. This is a temporary fix for emitters that apply physics.
-					// Check for collision modules with bApplyPhysics set to true
-					for (int32 ModuleIndex = 0; ModuleIndex < LODLevel->Modules.Num(); ModuleIndex++)
-					{
-						UParticleModuleCollision* CollisionModule = Cast<UParticleModuleCollision>(ParticleLODLevel->Modules[ModuleIndex]);
-						if (CollisionModule)
-						{
-							if (CollisionModule->bApplyPhysics == true)
-							{
-								bHasPhysics = true;
-								break;
-							}
-						}
-					}
-
-					if (LODLevel->bEnabled)
-					{
-						const UParticleModuleRequired* RequiredModule = LODLevel->RequiredModule;
-						if ((RequiredModule != nullptr) && (RequiredModule->EmitterLoops == 0))
-						{
-							bAnyEmitterLoopsForever = true;
-
-							UParticleModuleSpawn* SpawnModule = LODLevel->SpawnModule;
-							check(SpawnModule);
-
-							if (RequiredModule->EmitterDuration == 0.0f)
-							{
-								bIsImmortal = true;
-								if (SpawnModule->GetMaximumSpawnRate() == 0.0f && !bAutoDeactivate)
-								{
-									bWillBecomeZombie = true;
-								}
-							}
-						}
-					}
-				}
-			}
-
-			EParticleSignificanceLevel EmitterSignificance = FMath::Min(MaxSignificanceLevel, Emitter->SignificanceLevel);
-			if (EmitterSignificance > HighestSignificance)
-			{
-				HighestSignificance = EmitterSignificance;
-			}
-			if (EmitterSignificance < LowestSignificance)
-			{
-				LowestSignificance = EmitterSignificance;
-			}
-		}
-	}
-
-	bShouldManageSignificance = GetLowestSignificance() != EParticleSignificanceLevel::Critical /* && !ContainsEmitterType(UParticleModuleTypeDataBeam2::StaticClass())*/;
-
-	if (LODSettings.Num() == 0)
-	{
-		if (Emitters.Num() > 0)
-		{
-			UParticleEmitter* Emitter = Emitters[0];
-			if (Emitter)
-			{
-				LODSettings.AddUninitialized(Emitter->LODLevels.Num());
-				for (int32 LODIndex = 0; LODIndex < LODSettings.Num(); LODIndex++)
-				{
-					LODSettings[LODIndex] = FParticleSystemLOD::CreateParticleSystemLOD();
-				}
-			}
-		}
-		else
-		{
-			LODSettings.AddUninitialized();
-			LODSettings[0] = FParticleSystemLOD::CreateParticleSystemLOD();
-		}
-	}
-
-	// Add default LOD Distances
-	if( LODDistances.Num() == 0 && Emitters.Num() > 0 )
-	{
-		UParticleEmitter* Emitter = Emitters[0];
-		if (Emitter)
-		{
-			LODDistances.AddUninitialized(Emitter->LODLevels.Num());
-			for (int32 LODIndex = 0; LODIndex < LODDistances.Num(); LODIndex++)
-			{
-				LODDistances[LODIndex] = LODIndex * 2500.0f;
-			}
-		}
-	}
-
-	if (GCascadePSOPrecachingTime == 1)
-	{
-		PrecachePSOs();
-	}
-
-#if WITH_EDITOR
-	// Due to there still being some ways that LODLevel counts get mismatched,
-	// when loading in the editor LOD levels will always be checked and fixed
-	// up... This can be removed once all the edge cases that lead to the
-	// problem are found and fixed.
-	if (GIsEditor)
-	{
-		// Fix the LOD distance array and mismatched lod levels
-		int32 LODCount_0 = -1;
-		for (int32 EmitterIndex = 0; EmitterIndex < Emitters.Num(); EmitterIndex++)
-		{
-			UParticleEmitter* Emitter  = Emitters[EmitterIndex];
-			if (Emitter)
-			{
-				if (LODCount_0 == -1)
-				{
-					LODCount_0 = Emitter->LODLevels.Num();
-				}
-				else
-				{
-					int32 EmitterLODCount = Emitter->LODLevels.Num();
-					if (EmitterLODCount != LODCount_0)
-					{
-						UE_LOG(LogParticles, Warning, TEXT("Emitter %d has mismatched LOD level count - expected %d, found %d. PS = %s"),
-							EmitterIndex, LODCount_0, EmitterLODCount, *GetPathName());
-						UE_LOG(LogParticles, Warning, TEXT("Fixing up now... Package = %s"), *(GetOutermost()->GetPathName()));
-
-						if (EmitterLODCount > LODCount_0)
-						{
-							Emitter->LODLevels.RemoveAt(LODCount_0, EmitterLODCount - LODCount_0);
-						}
-						else
-						{
-							for (int32 NewLODIndex = EmitterLODCount; NewLODIndex < LODCount_0; NewLODIndex++)
-							{
-								if (Emitter->CreateLODLevel(NewLODIndex) != NewLODIndex)
-								{
-									UE_LOG(LogParticles, Warning, TEXT("Failed to add LOD level %d"), NewLODIndex);
-								}
-							}
-						}
-					}
-				}
-			}
-		}
-
-		if (LODCount_0 > 0)
-		{
-			if (LODDistances.Num() < LODCount_0)
-			{
-				for (int32 DistIndex = LODDistances.Num(); DistIndex < LODCount_0; DistIndex++)
-				{
-					float Distance = DistIndex * 2500.0f;
-					LODDistances.Add(Distance);
-				}
-			}
-			else
-			if (LODDistances.Num() > LODCount_0)
-			{
-				LODDistances.RemoveAt(LODCount_0, LODDistances.Num() - LODCount_0);
-			}
-		}
-		else
-		{
-			LODDistances.Empty();
-		}
-
-		if (LODCount_0 > 0)
-		{
-			if (LODSettings.Num() < LODCount_0)
-			{
-				for (int32 DistIndex = LODSettings.Num(); DistIndex < LODCount_0; DistIndex++)
-				{
-					LODSettings.Add(FParticleSystemLOD::CreateParticleSystemLOD());
-				}
-			}
-			else
-				if (LODSettings.Num() > LODCount_0)
-				{
-					LODSettings.RemoveAt(LODCount_0, LODSettings.Num() - LODCount_0);
-				}
-		}
-		else
-		{
-			LODSettings.Empty();
-		}
-	}
-#endif
-
-#if WITH_EDITORONLY_DATA
-	// Reset cascade's UI LOD setting to 0.
-	EditorLODSetting = 0;
-#endif // WITH_EDITORONLY_DATA
-
-	FixedRelativeBoundingBox.IsValid = true;
-
-	// Set up the SoloTracking...
-	SetupSoloing();
-}
-
-void UParticleSystem::PrecachePSOs()
-{
-	if (HasLaunchedPSOPrecaching() || (!IsComponentPSOPrecachingEnabled() && !IsResourcePSOPrecachingEnabled()))
-	{
-		return;
-	}
-
-	FMaterialInterfacePSOPrecacheParamsList PSOPrecacheParamsList;
-
-	FMaterialInterfacePSOPrecacheParams NewEntry;
-	NewEntry.PSOPrecacheParams.SetMobility(EComponentMobility::Movable);
-
-	// No per component emitter materials known at this point in time
-	TArray<UMaterialInterface*> EmptyEmitterMaterials;
-	// Cached array to collect all materials used for LOD level
-	TArray<UMaterialInterface*> Materials;
-
-	for (int32 EmitterIdx = 0; EmitterIdx < Emitters.Num(); ++EmitterIdx)
-	{
-		const UParticleEmitter* Emitter = Emitters[EmitterIdx];
-		if (!Emitter)
-		{
-			continue;
-		}
-
-		for (int32 LodIndex = 0; LodIndex < Emitter->LODLevels.Num(); ++LodIndex)
-		{
-			const UParticleLODLevel* LOD = Emitter->LODLevels[LodIndex];
-			if (LOD && LOD->bEnabled)
-			{
-				UParticleModuleTypeDataBase::FPSOPrecacheParams PrecacheParams;
-				if (LOD->TypeDataModule)
-				{
-					LOD->TypeDataModule->CollectPSOPrecacheData(Emitter, PrecacheParams);
-				}
-				else
-				{
-					bool bUsesDynamicParameter = (Emitter->DynamicParameterDataOffset > 0);
-					FPSOPrecacheVertexFactoryData VFData;
-					VFData.VertexFactoryType = &FParticleSpriteVertexFactory::StaticType;
-					VFData.CustomDefaultVertexDeclaration = FParticleSpriteVertexFactory::GetPSOPrecacheVertexDeclaration(bUsesDynamicParameter);
-					PrecacheParams.VertexFactoryDataList.Add(VFData);
-					PrecacheParams.PrimitiveType = PT_TriangleList;
-				}
-
-				Materials.Empty();
-				LOD->GetUsedMaterials(Materials, NamedMaterialSlots, EmptyEmitterMaterials);
-
-				for (UMaterialInterface* MaterialInterface : Materials)
-				{
-					NewEntry.MaterialInterface = MaterialInterface;
-					NewEntry.VertexFactoryDataList = PrecacheParams.VertexFactoryDataList;
-					NewEntry.PSOPrecacheParams.PrimitiveType = PrecacheParams.PrimitiveType;
-
-					AddMaterialInterfacePSOPrecacheParamsToList(NewEntry, PSOPrecacheParamsList);
-				}
-			}
-		}
-	}
-
-	LaunchPSOPrecaching(PSOPrecacheParamsList);
-}
-
-void UParticleSystem::Serialize(FArchive& Ar)
-{
-	Ar.UsingCustomVersion(FParticleSystemCustomVersion::GUID);
-
-	int32 CookTargetPlatformDetailModeMask = 0xFFFFFFFF;
-#if WITH_EDITOR
-	if (Ar.IsCooking())
-	{
-		// If we're cooking, check the device profile for whether we want to eliminate all emitters that don't match the detail mode. 
-		// This will only work if scalability settings affecting detail mode can not be changed at runtime (depends on platform)!
-		//
-		if (UDeviceProfile* DeviceProfile = UDeviceProfileManager::Get().FindProfile(Ar.CookingTarget()->IniPlatformName()))
-		{
-			// if we don't prune, we assume all detail modes
-			int32 CVarDoPrune = 0;
-			if (DeviceProfile->GetConsolidatedCVarValue(TEXT("fx.PruneEmittersOnCookByDetailMode"), CVarDoPrune) && CVarDoPrune == 1)
-			{
-				// get the detail mode from the device platform ini; if it's not there, we assume all detail modes
-				int32 CVarDetailMode = PDM_DefaultValue;
-				if (DeviceProfile->GetConsolidatedCVarValue(TEXT("r.DetailMode"), CVarDetailMode))
-				{
-					CookTargetPlatformDetailModeMask = (1 << CVarDetailMode);
-				}
-			}
-		}
-
-		// if we're cooking, save only emitters with matching detail modes
-		for (int32 EmitterIdx = 0; EmitterIdx<Emitters.Num(); EmitterIdx++)
-		{
-			// null out if detail mode doesn't match
-			if (Emitters[EmitterIdx] && !(Emitters[EmitterIdx]->DetailModeBitmask & CookTargetPlatformDetailModeMask))
-			{
-				UE_LOG(LogParticles, Display, TEXT("Pruning emitter, detail mode mismatch (PDM %i) (only works if platform can't change detail mode at runtime!): %s - set fx.PruneEmittersOnCookByDetailMode to 0 in DeviceProfile.ini for the target profile to avoid"), Emitters[EmitterIdx]->DetailModeBitmask, *Emitters[EmitterIdx]->EmitterName.ToString());
-				Emitters[EmitterIdx] = nullptr;
-			}
-		}
-	}
-#endif
-
-	Super::Serialize(Ar);
-
-}
-
-
-void UParticleSystem::UpdateColorModuleClampAlpha(UParticleModuleColorBase* ColorModule)
-{
-	if (ColorModule)
-	{
-		TArray<const FCurveEdEntry*> CurveEntries;
-		ColorModule->RemoveModuleCurvesFromEditor(CurveEdSetup);
-		ColorModule->AddModuleCurvesToEditor(CurveEdSetup, CurveEntries);
-	}
-}
-
-void UParticleSystem::GetAssetRegistryTags(TArray<FAssetRegistryTag>& OutTags) const
-{
-	PRAGMA_DISABLE_DEPRECATION_WARNINGS;
-	Super::GetAssetRegistryTags(OutTags);
-	PRAGMA_ENABLE_DEPRECATION_WARNINGS;
-}
-
-void UParticleSystem::GetAssetRegistryTags(FAssetRegistryTagsContext Context) const
-{
-	Context.AddTag( FAssetRegistryTag("HasGPUEmitter", HasGPUEmitter() ? TEXT("True") : TEXT("False"), FAssetRegistryTag::TT_Alphabetical) );
-
-	const float BoundsSize = FixedRelativeBoundingBox.GetSize().GetMax();
-	Context.AddTag(FAssetRegistryTag("FixedBoundsSize", bUseFixedRelativeBoundingBox ? FString::Printf(TEXT("%.2f"), BoundsSize) : FString(TEXT("None")), FAssetRegistryTag::TT_Numerical));
-
-	Context.AddTag(FAssetRegistryTag("NumEmitters", LexToString(Emitters.Num()), FAssetRegistryTag::TT_Numerical));
-
-	Context.AddTag(FAssetRegistryTag("NumLODs", LexToString(LODDistances.Num()), FAssetRegistryTag::TT_Numerical));
-
-	Context.AddTag(FAssetRegistryTag("WarmupTime", LexToString(WarmupTime), FAssetRegistryTag::TT_Numerical));
-
-	// Done here instead of as an AssetRegistrySearchable string to avoid the long prefix on the enum value string
-	FString LODMethodString = TEXT("Unknown");
-	switch (LODMethod)
-	{
-	case PARTICLESYSTEMLODMETHOD_Automatic:
-		LODMethodString = TEXT("Automatic");
-		break;
-	case PARTICLESYSTEMLODMETHOD_DirectSet:
-		LODMethodString = TEXT("DirectSet");
-		break;
-	case PARTICLESYSTEMLODMETHOD_ActivateAutomatic:
-		LODMethodString = TEXT("Activate Automatic");
-		break;
-	default:
-		check(false); // Missing enum entry
-		break;
-	}
-	Context.AddTag(FAssetRegistryTag("LODMethod", LODMethodString, FAssetRegistryTag::TT_Alphabetical));
-
-	Context.AddTag(FAssetRegistryTag("CPUCollision", UsesCPUCollision() ? TEXT("True") : TEXT("False"), FAssetRegistryTag::TT_Alphabetical));
-	Context.AddTag(FAssetRegistryTag("Looping", bAnyEmitterLoopsForever ? TEXT("True") : TEXT("False"), FAssetRegistryTag::TT_Alphabetical));
-	Context.AddTag(FAssetRegistryTag("Immortal", IsImmortal() ? TEXT("True") : TEXT("False"), FAssetRegistryTag::TT_Alphabetical));
-	Context.AddTag(FAssetRegistryTag("Becomes Zombie", WillBecomeZombie() ? TEXT("True") : TEXT("False"), FAssetRegistryTag::TT_Alphabetical));
-	Context.AddTag(FAssetRegistryTag("CanBeOccluded", OcclusionBoundsMethod == EParticleSystemOcclusionBoundsMethod::EPSOBM_None ? TEXT("False") : TEXT("True"), FAssetRegistryTag::TT_Alphabetical));
-
-	uint32 NumEmittersAtEachSig[(int32)EParticleSignificanceLevel::Num] = { 0, 0, 0, 0 };
-	for (UParticleEmitter* Emitter : Emitters)
-	{
-		if (Emitter)
-		{
-			++NumEmittersAtEachSig[(int32)Emitter->SignificanceLevel];			
-		}
-	}
-	Context.AddTag(FAssetRegistryTag("Critical Emitters", LexToString(NumEmittersAtEachSig[(int32)EParticleSignificanceLevel::Critical]), FAssetRegistryTag::TT_Numerical));
-	Context.AddTag(FAssetRegistryTag("High Emitters", LexToString(NumEmittersAtEachSig[(int32)EParticleSignificanceLevel::High]), FAssetRegistryTag::TT_Numerical));
-	Context.AddTag(FAssetRegistryTag("Medium Emitters", LexToString(NumEmittersAtEachSig[(int32)EParticleSignificanceLevel::Medium]), FAssetRegistryTag::TT_Numerical));
-	Context.AddTag(FAssetRegistryTag("Low Emitters", LexToString(NumEmittersAtEachSig[(int32)EParticleSignificanceLevel::Low]), FAssetRegistryTag::TT_Numerical));
-
-	Super::GetAssetRegistryTags(Context);
-}
-
-bool UParticleSystem::UsesCPUCollision() const
-{
-	for (const UParticleEmitter* Emitter : Emitters)
-	{
-		if (Emitter != nullptr)
-		{
-			// If we have not yet found a CPU collision module, and we have some enabled LODs to look in..
-			if (Emitter->HasAnyEnabledLODs() && Emitter->LODLevels.Num() > 0)
-			{
-				if (const UParticleLODLevel* HighLODLevel = Emitter->LODLevels[0])
-				{
-					// Iterate over modules of highest LOD (will have all the modules)
-					for (const UParticleModule* Module : HighLODLevel->Modules)
-					{
-						// If an enabled CPU collision module 
-						if (Module->bEnabled && Module->IsA<UParticleModuleCollision>())
-						{
-							return true;
-						}
-					}
-				}
-			}
-		}
-	}
-
-	return false;
-}
-
-bool UParticleSystem::CanBeClusterRoot() const
-{
-	return true;
-}
-
-bool UParticleSystem::CanBePooled()const
-{
-	if (MaxPoolSize == 0)
-	{
-		return false;
-	}
-
-	return true;
-}
-
-bool UParticleSystem::CalculateMaxActiveParticleCounts()
-{
-	bool bSuccess = true;
-
-	for (int32 EmitterIndex = 0; EmitterIndex < Emitters.Num(); EmitterIndex++)
-	{
-		UParticleEmitter* Emitter = Emitters[EmitterIndex];
-		if (Emitter)
-		{
-			if (Emitter->CalculateMaxActiveParticleCount() == false)
-			{
-				bSuccess = false;
-			}
-		}
-	}
-
-	return bSuccess;
-}
-
-
-void UParticleSystem::GetParametersUtilized(TArray<TArray<FString> >& ParticleSysParamList,
-											TArray<TArray<FString> >& ParticleParameterList)
-{
-	ParticleSysParamList.Empty();
-	ParticleParameterList.Empty();
-
-	for (int32 EmitterIndex = 0; EmitterIndex < Emitters.Num(); EmitterIndex++)
-	{
-		int32 CheckIndex;
-		CheckIndex = ParticleSysParamList.AddZeroed();
-		check(CheckIndex == EmitterIndex);
-		CheckIndex = ParticleParameterList.AddZeroed();
-		check(CheckIndex == EmitterIndex);
-
-		UParticleEmitter* Emitter = Emitters[EmitterIndex];
-		if (Emitter)
-		{
-			Emitter->GetParametersUtilized(
-				ParticleSysParamList[EmitterIndex],
-				ParticleParameterList[EmitterIndex]);
-		}
-	}
-}
-
-
-void UParticleSystem::SetupSoloing()
-{
-#if WITH_EDITOR
-	if (GIsEditor)
-	{
-		if (Emitters.Num())
-		{
-			// Store the settings of bEnabled for each LODLevel in each emitter
-			SoloTracking.Empty(Emitters.Num());
-			SoloTracking.AddZeroed(Emitters.Num());
-			for (int32 EmitterIdx = 0; EmitterIdx < Emitters.Num(); EmitterIdx++)
-			{
-				UParticleEmitter* Emitter = Emitters[EmitterIdx];
-				if (Emitter != nullptr)
-				{
-					FLODSoloTrack& SoloTrack = SoloTracking[EmitterIdx];
-					SoloTrack.SoloEnableSetting.Empty(Emitter->LODLevels.Num());
-					SoloTrack.SoloEnableSetting.AddZeroed(Emitter->LODLevels.Num());
-				}
-			}
-
-			for (int32 EmitterIdx = 0; EmitterIdx < Emitters.Num(); EmitterIdx++)
-			{
-				UParticleEmitter* Emitter = Emitters[EmitterIdx];
-				if (Emitter != NULL)
-				{
-					FLODSoloTrack& SoloTrack = SoloTracking[EmitterIdx];
-					int32 MaxLOD = FMath::Min(SoloTrack.SoloEnableSetting.Num(), Emitter->LODLevels.Num());
-					for (int32 LODIdx = 0; LODIdx < MaxLOD; LODIdx++)
-					{
-						UParticleLODLevel* LODLevel = Emitter->LODLevels[LODIdx];
-						check(LODLevel);
-						SoloTrack.SoloEnableSetting[LODIdx] = LODLevel->bEnabled;
-					}
-				}
-			}
-		}
-	}
-#endif
-}
-
-
-bool UParticleSystem::ToggleSoloing(class UParticleEmitter* InEmitter)
-{
-	bool bSoloingReturn = false;
-	if (InEmitter != NULL)
-	{
-		bool bOtherEmitterIsSoloing = false;
-		// Set the given one
-		int32 SelectedIndex = -1;
-		for (int32 EmitterIdx = 0; EmitterIdx < Emitters.Num(); EmitterIdx++)
-		{
-			UParticleEmitter* Emitter = Emitters[EmitterIdx];
-			check(Emitter != NULL);
-			if (Emitter == InEmitter)
-			{
-				SelectedIndex = EmitterIdx;
-			}
-			else
-			{
-				if (Emitter->bIsSoloing == true)
-				{
-					bOtherEmitterIsSoloing = true;
-					bSoloingReturn = true;
-				}
-			}
-		}
-
-		if (SelectedIndex != -1)
-		{
-			InEmitter->bIsSoloing = !InEmitter->bIsSoloing;
-			for (int32 EmitterIdx = 0; EmitterIdx < Emitters.Num(); EmitterIdx++)
-			{
-				UParticleEmitter* Emitter = Emitters[EmitterIdx];
-				FLODSoloTrack& SoloTrack = SoloTracking[EmitterIdx];
-				if (EmitterIdx == SelectedIndex)
-				{
-					for (int32 LODIdx = 0; LODIdx < InEmitter->LODLevels.Num(); LODIdx++)
-					{
-						UParticleLODLevel* LODLevel = InEmitter->LODLevels[LODIdx];
-						if (InEmitter->bIsSoloing == false)
-						{
-							if (bOtherEmitterIsSoloing == false)
-							{
-								// Restore the enabled settings - ie turn off soloing...
-								LODLevel->bEnabled = SoloTrack.SoloEnableSetting[LODIdx];
-							}
-							else
-							{
-								// Disable the emitter
-								LODLevel->bEnabled = false;
-							}
-						}
-						else 
-						if (bOtherEmitterIsSoloing == true)
-						{
-							// Need to restore old settings of this emitter as it is now soloing
-							LODLevel->bEnabled = SoloTrack.SoloEnableSetting[LODIdx];
-						}
-					}
-				}
-				else
-				{
-					// Restore all other emitters if this disables soloing...
-					if ((InEmitter->bIsSoloing == false) && (bOtherEmitterIsSoloing == false))
-					{
-						for (int32 LODIdx = 0; LODIdx < Emitter->LODLevels.Num(); LODIdx++)
-						{
-							UParticleLODLevel* LODLevel = Emitter->LODLevels[LODIdx];
-							// Restore the enabled settings - ie turn off soloing...
-							LODLevel->bEnabled = SoloTrack.SoloEnableSetting[LODIdx];
-						}
-					}
-					else
-					{
-						if (Emitter->bIsSoloing == false)
-						{
-							for (int32 LODIdx = 0; LODIdx < Emitter->LODLevels.Num(); LODIdx++)
-							{
-								UParticleLODLevel* LODLevel = Emitter->LODLevels[LODIdx];
-								// Disable the emitter
-								LODLevel->bEnabled = false;
-							}
-						}
-					}
-				}
-			}
-		}
-
-		// We checked the other emitters above...
-		// Make sure we catch the case of the first one toggled to true!
-		if (InEmitter->bIsSoloing == true)
-		{
-			bSoloingReturn = true;
-		}
-	}
-
-	return bSoloingReturn;
-}
-
-
-bool UParticleSystem::TurnOffSoloing()
-{
-	for (int32 EmitterIdx = 0; EmitterIdx < Emitters.Num(); EmitterIdx++)
-	{
-		UParticleEmitter* Emitter = Emitters[EmitterIdx];
-		if (Emitter != NULL)
-		{
-			FLODSoloTrack& SoloTrack = SoloTracking[EmitterIdx];
-			for (int32 LODIdx = 0; LODIdx < Emitter->LODLevels.Num(); LODIdx++)
-			{
-				UParticleLODLevel* LODLevel = Emitter->LODLevels[LODIdx];
-				if (LODLevel != NULL)
-				{
-					// Restore the enabled settings - ie turn off soloing...
-					LODLevel->bEnabled = SoloTrack.SoloEnableSetting[LODIdx];
-				}
-			}
-			Emitter->bIsSoloing = false;
-		}
-	}
-
-	return true;
-}
-
-
-void UParticleSystem::SetupLODValidity()
-{
-	for (int32 EmitterIdx = 0; EmitterIdx < Emitters.Num(); EmitterIdx++)
-	{
-		UParticleEmitter* Emitter = Emitters[EmitterIdx];
-		if (Emitter != NULL)
-		{
-			for (int32 Pass = 0; Pass < 2; Pass++)
-			{
-				for (int32 LODIdx = 0; LODIdx < Emitter->LODLevels.Num(); LODIdx++)
-				{
-					UParticleLODLevel* LODLevel = Emitter->LODLevels[LODIdx];
-					if (LODLevel != NULL)
-					{
-						for (int32 ModuleIdx = -3; ModuleIdx < LODLevel->Modules.Num(); ModuleIdx++)
-						{
-							int32 ModuleFetchIdx;
-							switch (ModuleIdx)
-							{
-							case -3:	ModuleFetchIdx = INDEX_REQUIREDMODULE;	break;
-							case -2:	ModuleFetchIdx = INDEX_SPAWNMODULE;		break;
-							case -1:	ModuleFetchIdx = INDEX_TYPEDATAMODULE;	break;
-							default:	ModuleFetchIdx = ModuleIdx;				break;
-							}
-
-							UParticleModule* Module = LODLevel->GetModuleAtIndex(ModuleFetchIdx);
-							if (Module != NULL)
-							{
-								// On pass 1, clear the LODValidity flags
-								// On pass 2, set it
-								if (Pass == 0)
-								{
-									Module->LODValidity = 0;
-								}
-								else
-								{
-									Module->LODValidity |= (1 << LODIdx);
-								}
-							}
-						}
-					}
-				}
-			}
-		}
-	}
-}
-
-void UParticleSystem::SetDelay(float InDelay)
-{
-	Delay = InDelay;
-}
-
-#if WITH_EDITOR
-
-bool UParticleSystem::RemoveAllDuplicateModules(bool bInMarkForCooker, TMap<UObject*,bool>* OutRemovedModules)
-{
-	// Generate a map of module classes used to instances of those modules 
-	TMap<UClass*,TMap<UParticleModule*,int32> > ClassToModulesMap;
-	for (int32 EmitterIdx = 0; EmitterIdx < Emitters.Num(); EmitterIdx++)
-	{
-		UParticleEmitter* Emitter = Emitters[EmitterIdx];
-		if (Emitter != NULL)
-		{
-			if (Emitter->bCookedOut == false)
-			{
-				for (int32 LODIdx = 0; LODIdx < Emitter->LODLevels.Num(); LODIdx++)
-				{
-					UParticleLODLevel* LODLevel = Emitter->LODLevels[LODIdx];
-					if (LODLevel != NULL)
-					{
-						for (int32 ModuleIdx = -1; ModuleIdx < LODLevel->Modules.Num(); ModuleIdx++)
-						{
-							UParticleModule* Module = NULL;
-							if (ModuleIdx == -1)
-							{
-								Module = LODLevel->SpawnModule;
-							}
-							else
-							{
-								Module = LODLevel->Modules[ModuleIdx];
-							}
-							if (Module != NULL)
-							{
-								TMap<UParticleModule*,int32>* ModuleList = ClassToModulesMap.Find(Module->GetClass());
-								if (ModuleList == NULL)
-								{
-									TMap<UParticleModule*,int32> TempModuleList;
-									ClassToModulesMap.Add(Module->GetClass(), TempModuleList);
-									ModuleList = ClassToModulesMap.Find(Module->GetClass());
-								}
-								check(ModuleList);
-								int32* ModuleCount = ModuleList->Find(Module);
-								if (ModuleCount == NULL)
-								{
-									int32 TempModuleCount = 0;
-									ModuleList->Add(Module, TempModuleCount);
-									ModuleCount = ModuleList->Find(Module);
-								}
-								check(ModuleCount);
-								(*ModuleCount)++;
-							}
-						}
-					}
-				}
-			}
-		}
-	}
-
-	// Now we have a list of module classes and the modules they contain...
-	// Find modules of the same class that have the exact same settings.
-	TMap<UParticleModule*, TArray<UParticleModule*> > DuplicateModules;
-	TMap<UParticleModule*,bool> FoundAsADupeModules;
-	TMap<UParticleModule*, UParticleModule*> ReplaceModuleMap;
-	bool bRemoveDuplicates = true;
-	for (TMap<UClass*,TMap<UParticleModule*,int32> >::TIterator ModClassIt(ClassToModulesMap); ModClassIt; ++ModClassIt)
-	{
-		UClass* ModuleClass = ModClassIt.Key();
-		TMap<UParticleModule*,int32>& ModuleMap = ModClassIt.Value();
-		if (ModuleMap.Num() > 1)
-		{
-			// There is more than one of this module, so see if there are dupes...
-			TArray<UParticleModule*> ModuleArray;
-			for (TMap<UParticleModule*,int32>::TIterator ModuleIt(ModuleMap); ModuleIt; ++ModuleIt)
-			{
-				ModuleArray.Add(ModuleIt.Key());
-			}
-
-			// For each module, see if it it a duplicate of another
-			for (int32 ModuleIdx = 0; ModuleIdx < ModuleArray.Num(); ModuleIdx++)
-			{
-				UParticleModule* SourceModule = ModuleArray[ModuleIdx];
-				if (FoundAsADupeModules.Find(SourceModule) == NULL)
-				{
-					for (int32 InnerModuleIdx = ModuleIdx + 1; InnerModuleIdx < ModuleArray.Num(); InnerModuleIdx++)
-					{
-						UParticleModule* CheckModule = ModuleArray[InnerModuleIdx];
-						if (FoundAsADupeModules.Find(CheckModule) == NULL)
-						{
-							bool bIsDifferent = false;
-							static const FName CascadeCategory(TEXT("Cascade"));
-							// Copy non component properties from the old actor to the new actor
-							for (FProperty* Property = ModuleClass->PropertyLink; Property != NULL; Property = Property->PropertyLinkNext)
-							{
-								bool bIsTransient = (Property->PropertyFlags & CPF_Transient) != 0;
-								bool bIsEditorOnly = (Property->PropertyFlags & CPF_EditorOnly) != 0;
-								bool bIsCascade = (FObjectEditorUtils::GetCategoryFName(Property) == CascadeCategory);
-								// Ignore 'Cascade' category, transient, native and EditorOnly properties...
-								if (!bIsTransient && !bIsEditorOnly && !bIsCascade)
-								{
-									for( int32 iProp=0; iProp<Property->ArrayDim; iProp++ )
-									{
-										bool bIsIdentical = Property->Identical_InContainer(SourceModule, CheckModule, iProp, PPF_DeepComparison);
-										if (bIsIdentical == false)
-										{
-											bIsDifferent = true;
-											break;
-										}
-									}
-								}
-							}
-
-							if (bIsDifferent == false)
-							{
-								TArray<UParticleModule*>* DupedModules = DuplicateModules.Find(SourceModule);
-								if (DupedModules == NULL)
-								{
-									TArray<UParticleModule*> TempDupedModules;
-									DuplicateModules.Add(SourceModule, TempDupedModules);
-									DupedModules = DuplicateModules.Find(SourceModule);
-								}
-								check(DupedModules);
-								if (ReplaceModuleMap.Find(CheckModule) == NULL)
-								{
-									ReplaceModuleMap.Add(CheckModule, SourceModule);
-								}
-								else
-								{
-									UE_LOG(LogParticles, Error, TEXT("Module already in replacement map - ABORTING CONVERSION!!!!"));
-									bRemoveDuplicates = false;
-								}
-								DupedModules->AddUnique(CheckModule);
-								FoundAsADupeModules.Add(CheckModule, true);
-							}
-						}
-					}
-				}
-			}
-		}
-	}
+{
+	bool			bUseTemplateDenyList = false;
+	TSet<FName>		TemplateDenyList;
+	FString			TemplateDenyListString;
 
-	// If not errors were found, and there are duplicates, remove them...
-	if (bRemoveDuplicates && (ReplaceModuleMap.Num() > 0))
+	static void UpdateTemplateDenyList(IConsoleVariable*)
 	{
-		TArray<UParticleModule*> RemovedModules;
-		for (int32 EmitterIdx = 0; EmitterIdx < Emitters.Num(); EmitterIdx++)
-		{
-			UParticleEmitter* Emitter = Emitters[EmitterIdx];
-			if (Emitter != NULL)
-			{
-				if (Emitter->bCookedOut == false)
-				{
-					for (int32 LODIdx = 0; LODIdx < Emitter->LODLevels.Num(); LODIdx++)
-					{
-						UParticleLODLevel* LODLevel = Emitter->LODLevels[LODIdx];
-						if (LODLevel != NULL)
-						{
-							for (int32 ModuleIdx = -1; ModuleIdx < LODLevel->Modules.Num(); ModuleIdx++)
-							{
-								UParticleModule* Module = NULL;
-								if (ModuleIdx == -1)
-								{
-									Module = LODLevel->SpawnModule;
-								}
-								else
-								{
-									Module = LODLevel->Modules[ModuleIdx];
-								}
-								if (Module != NULL)
-								{
-									UParticleModule** ReplacementModule = ReplaceModuleMap.Find(Module);
-									if (ReplacementModule != NULL)
-									{
-										UParticleModule* ReplaceMod = *ReplacementModule;
-										if (ModuleIdx == -1)
-										{
-											LODLevel->SpawnModule = CastChecked<UParticleModuleSpawn>(ReplaceMod);
-										}
-										else
-										{
-											LODLevel->Modules[ModuleIdx] = ReplaceMod;
-										}
-
-										if (bInMarkForCooker == true)
-										{
-											RemovedModules.AddUnique(Module);
-										}
-									}
-								}
-							}
-						}
-					}
-				}
-			}
-		}
+		TArray<FString> Names;
+		TemplateDenyListString.ParseIntoArray(Names, TEXT(","));
 
-		if (bInMarkForCooker == true)
+		TemplateDenyList.Empty();
+		for (const FString& Name : Names)
 		{
-			for (int32 MarkIdx = 0; MarkIdx < RemovedModules.Num(); MarkIdx++)
-			{
-				UParticleModule* RemovedModule = RemovedModules[MarkIdx];
-				RemovedModule->SetFlags(RF_Transient);
-				if (OutRemovedModules != NULL)
-				{
-					OutRemovedModules->Add(RemovedModule, true);
-				}
-			}
+			TemplateDenyList.Emplace(Name);
 		}
 
-		// Update the list of modules in each emitter
-		UpdateAllModuleLists();
+		bUseTemplateDenyList = TemplateDenyList.Num() > 0;
+	}
+
+	bool AllowTemplate(UParticleSystem* Template)
+	{
+		return !bUseTemplateDenyList || (Template && !TemplateDenyList.Contains(Template->GetFName()));
 	}
 
-	return true;
+	static FAutoConsoleVariableRef CVarCascadeSetTemplateDenyList(
+		TEXT("fx.Cascade.SetTemplateDenyList"),
+		TemplateDenyListString,
+		TEXT("Set the template deny List to use. (i.e. P_SystemA,P_SystemB)"),
+		FConsoleVariableDelegate::CreateStatic(UpdateTemplateDenyList),
+		ECVF_Scalability | ECVF_Default
+	);
 }
 
+// Comment this in to debug empty emitter instance templates...
+//#define _PSYSCOMP_DEBUG_INVALID_EMITTER_INSTANCE_TEMPLATES_
+
+/*-----------------------------------------------------------------------------
+	Particle scene view
+-----------------------------------------------------------------------------*/
+FSceneView*			GParticleView = NULL;
 
-void UParticleSystem::UpdateAllModuleLists()
+/*-----------------------------------------------------------------------------
+	Conversion functions
+-----------------------------------------------------------------------------*/
+void Particle_ModifyFloatDistribution(UDistributionFloat* pkDistribution, float fScale)
 {
-	for (int32 EmitterIdx = 0; EmitterIdx < Emitters.Num(); EmitterIdx++)
+	if (pkDistribution->IsA(UDistributionFloatConstant::StaticClass()))
+	{
+		UDistributionFloatConstant* pkDistConstant = Cast<UDistributionFloatConstant>(pkDistribution);
+		pkDistConstant->Constant *= fScale;
+	}
+	else
+	if (pkDistribution->IsA(UDistributionFloatUniform::StaticClass()))
+	{
+		UDistributionFloatUniform* pkDistUniform = Cast<UDistributionFloatUniform>(pkDistribution);
+		pkDistUniform->Min *= fScale;
+		pkDistUniform->Max *= fScale;
+	}
+	else
+	if (pkDistribution->IsA(UDistributionFloatConstantCurve::StaticClass()))
 	{
-		UParticleEmitter* Emitter = Emitters[EmitterIdx];
-		if (Emitter != NULL)
+		UDistributionFloatConstantCurve* pkDistCurve = Cast<UDistributionFloatConstantCurve>(pkDistribution);
+
+		int32 iKeys = pkDistCurve->GetNumKeys();
+		int32 iCurves = pkDistCurve->GetNumSubCurves();
+
+		for (int32 KeyIndex = 0; KeyIndex < iKeys; KeyIndex++)
 		{
-			for (int32 LODIdx = 0; LODIdx < Emitter->LODLevels.Num(); LODIdx++)
+			float fKeyIn = pkDistCurve->GetKeyIn(KeyIndex);
+			for (int32 SubIndex = 0; SubIndex < iCurves; SubIndex++)
 			{
-				UParticleLODLevel* LODLevel = Emitter->LODLevels[LODIdx];
-				if (LODLevel != NULL)
-				{
-					LODLevel->UpdateModuleLists();
-				}
-			}
+				float fKeyOut = pkDistCurve->GetKeyOut(SubIndex, KeyIndex);
+				float ArriveTangent;
+				float LeaveTangent;
+				pkDistCurve->GetTangents(SubIndex, KeyIndex, ArriveTangent, LeaveTangent);
 
-			// Allow type data module to cache any module info
-			if(Emitter->LODLevels.Num() > 0)
-			{
-				UParticleLODLevel* HighLODLevel = Emitter->LODLevels[0];
-				if (HighLODLevel != nullptr && HighLODLevel->TypeDataModule != nullptr)
-				{
-					// Allow TypeData module to cache pointers to modules
-					HighLODLevel->TypeDataModule->CacheModuleInfo(Emitter);
-				}
+				pkDistCurve->SetKeyOut(SubIndex, KeyIndex, fKeyOut * fScale);
+				pkDistCurve->SetTangents(SubIndex, KeyIndex, ArriveTangent * fScale, LeaveTangent * fScale);
 			}
-
-			// Update any cached info from modules on the emitter
-			Emitter->CacheEmitterModuleInfo();
 		}
 	}
 }
-#endif
 
-
-void UParticleSystem::BuildEmitters()
+void Particle_ModifyVectorDistribution(UDistributionVector* pkDistribution, FVector& vScale)
 {
-	const int32 EmitterCount = Emitters.Num();
-	for ( int32 EmitterIndex = 0; EmitterIndex < EmitterCount; ++EmitterIndex )
+	if (pkDistribution->IsA(UDistributionVectorConstant::StaticClass()))
 	{
-		if (UParticleEmitter* Emitter = Emitters[EmitterIndex])
-		{
-			Emitter->Build();
-		}
+		UDistributionVectorConstant* pkDistConstant = Cast<UDistributionVectorConstant>(pkDistribution);
+		pkDistConstant->Constant *= vScale;
 	}
-}
-
-static bool LogReasoningForAnyThreadTicking()
-{
-	static bool bLogThreadedParticleTicking = FParse::Param(FCommandLine::Get(), TEXT("LogThreadedParticleTicking"));
-	return bLogThreadedParticleTicking;
-}
+	else
+	if (pkDistribution->IsA(UDistributionVectorUniform::StaticClass()))
+	{
+		UDistributionVectorUniform* pkDistUniform = Cast<UDistributionVectorUniform>(pkDistribution);
+		pkDistUniform->Min *= vScale;
+		pkDistUniform->Max *= vScale;
+	}
+	else
+	if (pkDistribution->IsA(UDistributionVectorConstantCurve::StaticClass()))
+	{
+		UDistributionVectorConstantCurve* pkDistCurve = Cast<UDistributionVectorConstantCurve>(pkDistribution);
 
-void UParticleSystem::ComputeCanTickInAnyThread()
-{
-	check(!bIsElligibleForAsyncTickComputed);
-	bIsElligibleForAsyncTickComputed = true;
+		int32 iKeys = pkDistCurve->GetNumKeys();
+		int32 iCurves = pkDistCurve->GetNumSubCurves();
 
-	bIsElligibleForAsyncTick = true; // assume everything is async
-	int32 EmitterIndex;
-	for (EmitterIndex = 0; EmitterIndex < Emitters.Num(); EmitterIndex++)
-	{
-		UParticleEmitter* Emitter = Emitters[EmitterIndex];
-		if (Emitter)
+		for (int32 KeyIndex = 0; KeyIndex < iKeys; KeyIndex++)
 		{
-			for (int32 LevelIndex = 0; LevelIndex < Emitter->LODLevels.Num(); LevelIndex++)
+			float fKeyIn = pkDistCurve->GetKeyIn(KeyIndex);
+			for (int32 SubIndex = 0; SubIndex < iCurves; SubIndex++)
 			{
-				UParticleLODLevel* LODLevel	= Emitter->LODLevels[LevelIndex];
-				if (LODLevel)
+				float fKeyOut = pkDistCurve->GetKeyOut(SubIndex, KeyIndex);
+				float ArriveTangent;
+				float LeaveTangent;
+				pkDistCurve->GetTangents(SubIndex, KeyIndex, ArriveTangent, LeaveTangent);
+
+				switch (SubIndex)
 				{
-					for (int32 ModuleIndex = 0; ModuleIndex < LODLevel->Modules.Num(); ModuleIndex++)
-					{
-						UParticleModule* Module	= LODLevel->Modules[ModuleIndex];
-						if (Module && !Module->CanTickInAnyThread())
-						{
-							bIsElligibleForAsyncTick = false;
-							if (LogReasoningForAnyThreadTicking())
-							{
-								UE_LOG(LogParticles, Display, TEXT("Cannot tick %s in parallel because module %s in Emitter %s cannot tick in in parallel."), *GetFullName(), *Module->GetFullName(), *Emitter->GetFullName());
-							}
-							else
-							{
-								return;
-							}
-						}
-					}
+				case 1:
+					pkDistCurve->SetKeyOut(SubIndex, KeyIndex, fKeyOut * vScale.Y);
+					pkDistCurve->SetTangents(SubIndex, KeyIndex, ArriveTangent * vScale.Y, LeaveTangent * vScale.Y);
+					break;
+				case 2:
+					pkDistCurve->SetKeyOut(SubIndex, KeyIndex, fKeyOut * vScale.Z);
+					pkDistCurve->SetTangents(SubIndex, KeyIndex, ArriveTangent * vScale.Z, LeaveTangent * vScale.Z);
+					break;
+				case 0:
+				default:
+					pkDistCurve->SetKeyOut(SubIndex, KeyIndex, fKeyOut * vScale.X);
+					pkDistCurve->SetTangents(SubIndex, KeyIndex, ArriveTangent * vScale.X, LeaveTangent * vScale.X);
+					break;
 				}
 			}
-
 		}
 	}
-	if (LogReasoningForAnyThreadTicking() && bIsElligibleForAsyncTick)
-	{
-		UE_LOG(LogParticles, Display, TEXT("Can tick %s in parallel."), *GetFullName());
-	}
 }
 
-bool UParticleSystem::ContainsEmitterType(UClass* TypeData)
+/** Console command to reset all particle components. */
+static void ResetAllParticleComponents()
 {
-	for ( int32 EmitterIndex = 0; EmitterIndex < Emitters.Num(); ++EmitterIndex)
+	for (TObjectIterator<UParticleSystemComponent> It; It; ++It)
 	{
-		UParticleEmitter* Emitter = Emitters[EmitterIndex];
-		if (Emitter)
-		{
-			UParticleLODLevel* LODLevel = Emitter->LODLevels[0];
-			if (LODLevel && LODLevel->TypeDataModule && LODLevel->TypeDataModule->IsA(TypeData))
-			{
-				return true;
-			}
-		}
+		UParticleSystemComponent* ParticleSystemComponent = *It;
+		ParticleSystemComponent->ResetParticles();
+		ParticleSystemComponent->ActivateSystem(true);
+		ParticleSystemComponent->bIsViewRelevanceDirty = true;
+		ParticleSystemComponent->CachedViewRelevanceFlags.Empty();
+		ParticleSystemComponent->ConditionalCacheViewRelevanceFlags();
+		ParticleSystemComponent->ReregisterComponent();
 	}
-	
-	return false;
 }
+static FAutoConsoleCommand GResetAllParticleComponentsCmd(
+	TEXT("FX.RestartAll"),
+	TEXT("Restarts all particle system components"),
+	FConsoleCommandDelegate::CreateStatic(ResetAllParticleComponents)
+	);
 
-bool UParticleSystem::HasGPUEmitter() const
-{
-	for (int32 EmitterIndex = 0; EmitterIndex < Emitters.Num(); ++EmitterIndex)
-	{
-		if (Emitters[EmitterIndex] == nullptr)
-		{
-			continue;
-		}
-		// We can just check for the GPU type data at the highest LOD.
-		UParticleLODLevel* LODLevel = Emitters[EmitterIndex]->LODLevels[0];
-		if( LODLevel )
-		{
-			UParticleModule* TypeDataModule = LODLevel->TypeDataModule;
-			if( TypeDataModule && TypeDataModule->IsA(UParticleModuleTypeDataGpu::StaticClass()) )
-			{
-				return true;
-			}
-		}
-	}
-	return false;
-}
 
 UFXSystemComponent::UFXSystemComponent(const FObjectInitializer& ObjectInitializer)
 	: Super(ObjectInitializer)
@@ -3657,7 +406,7 @@ public:
 	virtual bool GetJustRegistered() const override { return Component->bJustRegistered; }
 	virtual float GetWarmupTime() const override { return Component->WarmupTime; }
 	virtual float GetEmitterDelay() const override { return Component->EmitterDelay; }
-	virtual FRandomStream& GetRandomStream() const override { return Component->RandomStream; }
+	virtual FRandomStream& GetRandomStream() override { return Component->RandomStream; }
 
 	virtual void SetComponentToWorld(const FTransform& NewComponentToWorld) override { Component->SetComponentToWorld(NewComponentToWorld); }
 	virtual void DeactivateNextTick() override { Component->DeactivaateNextTick(); }
@@ -8008,139 +4757,6 @@ void UParticleSystemComponent::AutoPopulateInstanceProperties()
 	}
 }
 
-void UParticleLODLevel::GetUsedMaterials(TArray<UMaterialInterface*>& OutMaterials, const TArray<FNamedEmitterMaterial>& Slots, const TArray<class UMaterialInterface*>& EmitterMaterials) const
-{
-	// Only process enabled emitters
-	if (bEnabled)
-	{
-		const UParticleModuleTypeDataMesh* MeshTypeData = Cast<UParticleModuleTypeDataMesh>(TypeDataModule);
-
-		if (MeshTypeData && MeshTypeData->Mesh && MeshTypeData->Mesh->GetRenderData())
-		{
-			const FStaticMeshLODResources& LODModel = MeshTypeData->Mesh->GetRenderData()->LODResources[0];
-
-			// Gather the materials applied to the LOD.
-			for (int32 SectionIndex = 0; SectionIndex < LODModel.Sections.Num(); SectionIndex++)
-			{
-				UMaterialInterface* Material = NULL;
-							
-				TArray<FName>& NamedOverrides = RequiredModule->NamedMaterialOverrides;
-
-				if (NamedOverrides.IsValidIndex(SectionIndex))
-				{	
-					//If we have named material overrides then get it's index into the emitter materials array.	
-					for (int32 CheckIdx = 0; CheckIdx < Slots.Num(); ++CheckIdx)
-					{
-						if (NamedOverrides[SectionIndex] == Slots[CheckIdx].Name)
-						{
-							//Default to the default material for that slot.
-							Material = Slots[CheckIdx].Material;
-							if (EmitterMaterials.IsValidIndex(CheckIdx) && nullptr != EmitterMaterials[CheckIdx] )
-							{
-								//This material has been overridden externally, e.g. from a BP so use that one.
-								Material = EmitterMaterials[CheckIdx];
-							}
-
-							break;
-						}
-					}
-				}
-
-				// See if there is a mesh material module.
-				if (Material == NULL)
-				{
-					// Walk in reverse order as in the case of multiple modules, only the final result will be applied
-					for (int32 ModuleIndex = Modules.Num()-1; ModuleIndex >= 0; --ModuleIndex)
-					{
-						UParticleModuleMeshMaterial* MeshMatModule = Cast<UParticleModuleMeshMaterial>(Modules[ModuleIndex]);
-						if (MeshMatModule && MeshMatModule->bEnabled)
-						{
-							if (SectionIndex < MeshMatModule->MeshMaterials.Num())
-							{
-								Material = MeshMatModule->MeshMaterials[SectionIndex];
-								break;
-							}
-						}
-					}
-				}
-
-				// Overriding the material?
-				if (Material == NULL && MeshTypeData->bOverrideMaterial == true)
-				{
-					Material = RequiredModule->Material;
-				}
-
-				// Use the material set on the mesh.
-				if (Material == NULL)
-				{
-					Material = MeshTypeData->Mesh->GetMaterial(LODModel.Sections[SectionIndex].MaterialIndex);
-				}
-
-				if (Material)
-				{
-					OutMaterials.Add(Material);
-				}
-			}
-		}
-		else
-		{
-			UMaterialInterface* Material = NULL;
-							
-			TArray<FName>& NamedOverrides = RequiredModule->NamedMaterialOverrides;
-
-			if (NamedOverrides.Num() > 0)
-			{
-				for (int32 CheckIdx = 0; CheckIdx < Slots.Num(); ++CheckIdx)
-				{
-					if (NamedOverrides[0] == Slots[CheckIdx].Name)
-					{
-						//Default to the default material for that slot.
-						Material = Slots[CheckIdx].Material;
-						if (EmitterMaterials.IsValidIndex(CheckIdx) && nullptr != EmitterMaterials[CheckIdx])
-						{
-							//This material has been overridden externally, e.g. from a BP so use that one.
-							Material = EmitterMaterials[CheckIdx];
-						}
-        
-						break;
-					}
-				}
-			}
-
-			if (!Material)
-			{
-				Material = RequiredModule->Material;
-			}
-
-			OutMaterials.Add(Material);
-		}
-	}
-}
-
-void UParticleLODLevel::GetStreamingMeshInfo(const FBoxSphereBounds& Bounds, TArray<FStreamingRenderAssetPrimitiveInfo>& OutStreamingRenderAssets) const
-{
-	if (bEnabled)
-	{
-		if (const UParticleModuleTypeDataMesh* MeshTypeData = Cast<UParticleModuleTypeDataMesh>(TypeDataModule))
-		{
-			if (UStaticMesh* Mesh = MeshTypeData->Mesh)
-			{
-				if (Mesh->RenderResourceSupportsStreaming() && Mesh->GetRenderAssetType() == EStreamableRenderAssetType::StaticMesh)
-				{
-					const FBoxSphereBounds MeshBounds = Mesh->GetBounds();
-					const FBoxSphereBounds StreamingBounds = FBoxSphereBounds(
-						Bounds.Origin + MeshBounds.Origin,
-						MeshBounds.BoxExtent * MeshTypeData->LODSizeScale,
-						MeshBounds.SphereRadius * MeshTypeData->LODSizeScale);
-					const float MeshTexelFactor = MeshBounds.SphereRadius * 2.0f;
-
-					new (OutStreamingRenderAssets) FStreamingRenderAssetPrimitiveInfo(Mesh, StreamingBounds, MeshTexelFactor);
-				}
-			}
-		}
-	}
-}
-
 void UParticleSystemComponent::GetUsedMaterials( TArray<UMaterialInterface*>& OutMaterials, bool bGetDebugMaterials ) const
 {
 	if (Template)
diff --git a/Engine/Source/Runtime/Engine/Private/Particles/ParticleEmitter.cpp b/Engine/Source/Runtime/Engine/Private/Particles/ParticleEmitter.cpp
new file mode 100644
--- /dev/null
+++ b/Engine/Source/Runtime/Engine/Private/Particles/ParticleEmitter.cpp
@@ -0,0 +1,1338 @@
+// Copyright Epic Games, Inc. All Rights Reserved.
+
+#include "Particles/ParticleEmitter.h"
+#include "Distributions/DistributionFloatConstant.h"
+#include "Distributions/DistributionFloatConstantCurve.h"
+#include "Distributions/DistributionFloatUniform.h"
+#include "Distributions/DistributionVectorConstantCurve.h"
+#include "Distributions/DistributionVectorUniform.h"
+#include "Engine/Engine.h"
+#include "Engine/InterpCurveEdSetup.h"
+#include "ParticleEmitterInstanceOwner.h"
+#include "Particles/Camera/ParticleModuleCameraOffset.h"
+#include "Particles/Color/ParticleModuleColorOverLife.h"
+#include "Particles/Lifetime/ParticleModuleLifetime.h"
+#include "Particles/Light/ParticleModuleLight.h"
+#include "Particles/Location/ParticleModuleLocationBoneSocket.h"
+#include "Particles/Material/ParticleModuleMeshMaterial.h"
+#include "Particles/Modules/Location/ParticleModulePivotOffset.h"
+#include "Particles/Parameter/ParticleModuleParameterDynamic.h"
+#include "Particles/ParticleLODLevel.h"
+#include "Particles/ParticleModuleRequired.h"
+#include "Particles/ParticleSystem.h"
+#include "Particles/ParticleSystemComponent.h"
+#include "Particles/Size/ParticleModuleSize.h"
+#include "Particles/Spawn/ParticleModuleSpawn.h"
+#include "Particles/SubUV/ParticleModuleSubUV.h"
+#include "Particles/TypeData/ParticleModuleTypeDataBase.h"
+#include "Particles/TypeData/ParticleModuleTypeDataBeam2.h"
+#include "Particles/Velocity/ParticleModuleVelocity.h"
+#include "Scalability.h"
+#include "UObject/UObjectIterator.h"
+
+#include UE_INLINE_GENERATED_CPP_BY_NAME(ParticleEmitter)
+
+////////////////////////////////////////////////////////////////////////////////////////////////////
+
+static TAutoConsoleVariable<float> CVarQLSpawnRateReferenceLevel(
+	TEXT("fx.QualityLevelSpawnRateScaleReferenceLevel"),
+	2,
+	TEXT("Controls the reference level for quality level based spawn rate scaling. This is the FX quality level\n")
+	TEXT("at which spawn rate is not scaled down; Spawn rate scaling will happen by each emitter's\n")
+	TEXT("QualityLevelSpawnRateScale value for each reduction in level below the reference level.\n")
+	TEXT("\n")
+	TEXT("Default = 2. Value should range from 0 to the maximum FX quality level."),
+	ECVF_Scalability);
+
+////////////////////////////////////////////////////////////////////////////////////////////////////
+
+UParticleEmitter::UParticleEmitter(const FObjectInitializer& ObjectInitializer)
+	: Super(ObjectInitializer)
+	, SignificanceLevel(EParticleSignificanceLevel::Critical)
+	, bUseLegacySpawningBehavior(false)
+	, bDisabledLODsKeepEmitterAlive(false)
+	, bDisableWhenInsignficant(0)
+	, QualityLevelSpawnRateScale(1.0f)
+	, DetailModeBitmask(PDM_DefaultValue)
+{
+	// Structure to hold one-time initialization
+	struct FConstructorStatics
+	{
+		FName NAME_Particle_Emitter;
+		FConstructorStatics()
+			: NAME_Particle_Emitter(TEXT("Particle Emitter"))
+		{
+		}
+	};
+	static FConstructorStatics ConstructorStatics;
+
+	EmitterName = ConstructorStatics.NAME_Particle_Emitter;
+	ConvertedModules = true;
+	PeakActiveParticles = 0;
+#if WITH_EDITORONLY_DATA
+	EmitterEditorColor = FColor(0, 150, 150, 255);
+#endif // WITH_EDITORONLY_DATA
+
+}
+
+FParticleEmitterInstance* UParticleEmitter::CreateInstance(IParticleEmitterInstanceOwner& InComponent)
+{
+	UE_LOG(LogParticles, Fatal,TEXT("UParticleEmitter::CreateInstance is pure virtual")); 
+	return NULL; 
+}
+
+void UParticleEmitter::UpdateModuleLists()
+{
+	for (int32 LODIndex = 0; LODIndex < LODLevels.Num(); LODIndex++)
+	{
+		UParticleLODLevel* LODLevel = LODLevels[LODIndex];
+		if (LODLevel)
+		{
+			LODLevel->UpdateModuleLists();
+		}
+	}
+	Build();
+}
+
+/**
+ *	Helper function for fixing up LODValidity issues on particle modules...
+ *
+ *	@param	LODIndex		The index of the LODLevel the module is from.
+ *	@param	ModuleIndex		The index of the module being checked.
+ *	@param	Emitter			The emitter owner.
+ *	@param	CurrModule		The module being checked.
+ *
+ *	@return	 0		If there was no problem.
+ *			 1		If there was a problem and it was fixed.
+ *			-1		If there was a problem that couldn't be fixed.
+ */
+int32 ParticleEmitterHelper_FixupModuleLODErrors( int32 LODIndex, int32 ModuleIndex, 
+	const UParticleEmitter* Emitter, UParticleModule* CurrModule )
+{
+	int32 Result = 1;
+	bool bIsDirty = false;
+
+	UObject* ModuleOuter = CurrModule->GetOuter();
+	UObject* EmitterOuter = Emitter->GetOuter();
+	if (ModuleOuter != EmitterOuter)
+	{
+		// Module has an incorrect outer
+		CurrModule->Rename(NULL, EmitterOuter, REN_DoNotDirty);
+		bIsDirty = true;
+	}
+
+	if (CurrModule->LODValidity == 0)
+	{
+		// Immediately tag it for this lod level...
+		CurrModule->LODValidity = (1 << LODIndex);
+		bIsDirty = true;
+	}
+	else
+	if (CurrModule->IsUsedInLODLevel(LODIndex) == false)
+	{
+		// Why was this even called here?? 
+		// The assumption is that it should be called for the module in the given lod level...
+		// So, we will tag it with this index.
+		CurrModule->LODValidity |= (1 << LODIndex);
+		bIsDirty = true;
+	}
+
+	if (LODIndex > 0)
+	{
+		int32 CheckIndex = LODIndex - 1;
+		while (CheckIndex >= 0)
+		{
+			if (CurrModule->IsUsedInLODLevel(CheckIndex))
+			{
+				// Ensure that it is the same as the one it THINKS it is shared with...
+				UParticleLODLevel* CheckLODLevel = Emitter->LODLevels[CheckIndex];
+
+				if (CurrModule->IsA(UParticleModuleSpawn::StaticClass()))
+				{
+					if (CheckLODLevel->SpawnModule != CurrModule)
+					{
+						// Fix it up... Turn off the higher LOD flag
+						CurrModule->LODValidity &= ~(1 << CheckIndex);
+						bIsDirty = true;
+					}
+				}
+				else
+				if (CurrModule->IsA(UParticleModuleRequired::StaticClass()))
+				{
+					if (CheckLODLevel->RequiredModule != CurrModule)
+					{
+						// Fix it up... Turn off the higher LOD flag
+						CurrModule->LODValidity &= ~(1 << CheckIndex);
+						bIsDirty = true;
+					}
+				}
+				else
+				if (CurrModule->IsA(UParticleModuleTypeDataBase::StaticClass()))
+				{
+					if (CheckLODLevel->TypeDataModule != CurrModule)
+					{
+						// Fix it up... Turn off the higher LOD flag
+						CurrModule->LODValidity &= ~(1 << CheckIndex);
+						bIsDirty = true;
+					}
+				}
+				else
+				{
+					if (ModuleIndex >= CheckLODLevel->Modules.Num())
+					{
+						UE_LOG(LogParticles, Warning, TEXT("\t\tMismatched module count at %2d in %s"), LODIndex, *(Emitter->GetPathName()));
+						Result = -1;
+					}
+					else
+					{
+						UParticleModule* CheckModule = CheckLODLevel->Modules[ModuleIndex];
+						if (CheckModule != CurrModule)
+						{
+							// Fix it up... Turn off the higher LOD flag
+							CurrModule->LODValidity &= ~(1 << CheckIndex);
+							bIsDirty = true;
+						}
+					}
+				}
+			}
+
+			CheckIndex--;
+		}
+	}
+
+	if ((bIsDirty == true) && IsRunningCommandlet())
+	{
+		CurrModule->MarkPackageDirty();
+		Emitter->MarkPackageDirty();
+	}
+
+	return Result;
+}
+
+bool UParticleEmitter::IsPostLoadThreadSafe() const
+{
+	return false;
+}
+
+void UParticleEmitter::Serialize(FArchive& Ar)
+{
+	Super::Serialize(Ar);
+
+	Ar.UsingCustomVersion(FParticleSystemCustomVersion::GUID);
+}
+
+void UParticleEmitter::PostLoad()
+{
+	Super::PostLoad();
+
+	const int32 PSysVer = GetLinkerCustomVersion(FParticleSystemCustomVersion::GUID);
+	if (PSysVer < FParticleSystemCustomVersion::FixLegacySpawningBugs)
+	{
+		bUseLegacySpawningBehavior = true;
+	}
+
+	if (PSysVer < FParticleSystemCustomVersion::AddEpicDetailMode)
+	{
+		// Init epic detail mode to enabled if high is set
+		if (DetailModeBitmask & (1 << EParticleDetailMode::PDM_High))
+		{
+			DetailModeBitmask |= (1 << EParticleDetailMode::PDM_Epic);
+		}
+	}
+
+	for (int32 LODIndex = 0; LODIndex < LODLevels.Num(); LODIndex++)
+	{
+		UParticleLODLevel* LODLevel = LODLevels[LODIndex];
+		if (LODLevel)
+		{
+			LODLevel->ConditionalPostLoad();
+
+			FLinkerLoad* LODLevelLinker = LODLevel->GetLinker();
+			if (LODLevel->SpawnModule == NULL)
+			{
+				// Force the conversion to SpawnModule
+				UParticleSystem* PSys = Cast<UParticleSystem>(GetOuter());
+				if (PSys)
+				{
+					UE_LOG(LogParticles, Warning, TEXT("LODLevel %d was not converted to spawn module - forcing: %s"), 
+						LODLevel->Level, *(PSys->GetPathName()));
+				}
+				LODLevel->ConvertToSpawnModule();
+			}
+
+			check(LODLevel->SpawnModule);
+
+		}
+	}
+
+#if WITH_EDITORONLY_DATA
+	UpdateDetailModeDisplayString();
+#endif
+
+#if WITH_EDITOR
+	if ((GIsEditor == true) && 1)//(IsRunningCommandlet() == false))
+	{
+		ConvertedModules = false;
+		PeakActiveParticles = 0;
+
+		// Check for improper outers...
+		UObject* EmitterOuter = GetOuter();
+		bool bWarned = false;
+		for (int32 LODIndex = 0; (LODIndex < LODLevels.Num()) && !bWarned; LODIndex++)
+		{
+			UParticleLODLevel* LODLevel = LODLevels[LODIndex];
+			if (LODLevel)
+			{
+				LODLevel->ConditionalPostLoad();
+
+				UParticleModule* Module = LODLevel->TypeDataModule;
+				if (Module)
+				{
+					Module->ConditionalPostLoad();
+
+					UObject* OuterObj = Module->GetOuter();
+					check(OuterObj);
+					if (OuterObj != EmitterOuter)
+					{
+						UE_LOG(LogParticles, Warning, TEXT("UParticleModule %s has an incorrect outer on %s... run FixupEmitters on package %s (%s)"),
+							*(Module->GetPathName()), 
+							*(EmitterOuter->GetPathName()),
+							*(OuterObj->GetOutermost()->GetPathName()),
+							*(GetOutermost()->GetPathName()));
+						UE_LOG(LogParticles, Warning, TEXT("\tModule Outer..............%s"), *(OuterObj->GetPathName()));
+						UE_LOG(LogParticles, Warning, TEXT("\tModule Outermost..........%s"), *(Module->GetOutermost()->GetPathName()));
+						UE_LOG(LogParticles, Warning, TEXT("\tEmitter Outer.............%s"), *(EmitterOuter->GetPathName()));
+						UE_LOG(LogParticles, Warning, TEXT("\tEmitter Outermost.........%s"), *(GetOutermost()->GetPathName()));
+						bWarned = true;
+					}
+				}
+
+				if (!bWarned)
+				{
+					for (int32 ModuleIndex = 0; (ModuleIndex < LODLevel->Modules.Num()) && !bWarned; ModuleIndex++)
+					{
+						Module = LODLevel->Modules[ModuleIndex];
+						if (Module)
+						{
+							Module->ConditionalPostLoad();
+
+							UObject* OuterObj = Module->GetOuter();
+							check(OuterObj);
+							if (OuterObj != EmitterOuter)
+							{
+								UE_LOG(LogParticles, Warning, TEXT("UParticleModule %s has an incorrect outer on %s... run FixupEmitters on package %s (%s)"),
+									*(Module->GetPathName()), 
+									*(EmitterOuter->GetPathName()),
+									*(OuterObj->GetOutermost()->GetPathName()),
+									*(GetOutermost()->GetPathName()));
+								UE_LOG(LogParticles, Warning, TEXT("\tModule Outer..............%s"), *(OuterObj->GetPathName()));
+								UE_LOG(LogParticles, Warning, TEXT("\tModule Outermost..........%s"), *(Module->GetOutermost()->GetPathName()));
+								UE_LOG(LogParticles, Warning, TEXT("\tEmitter Outer.............%s"), *(EmitterOuter->GetPathName()));
+								UE_LOG(LogParticles, Warning, TEXT("\tEmitter Outermost.........%s"), *(GetOutermost()->GetPathName()));
+								bWarned = true;
+							}
+						}
+					}
+				}
+			}
+		}
+	}
+	else
+#endif
+	{
+		for (int32 LODIndex = 0; LODIndex < LODLevels.Num(); LODIndex++)
+		{
+			UParticleLODLevel* LODLevel = LODLevels[LODIndex];
+			if (LODLevel)
+			{
+				LODLevel->ConditionalPostLoad();
+			}
+		}
+	}
+   
+
+	ConvertedModules = true;
+
+	// this will look at all of the emitters and then remove ones that some how have become NULL (e.g. from a removal of an Emitter where content
+	// is still referencing it)
+	for (int32 LODIndex = 0; LODIndex < LODLevels.Num(); LODIndex++)
+	{
+		UParticleLODLevel* LODLevel = LODLevels[LODIndex];
+		if (LODLevel)
+		{
+			for (int32 ModuleIndex = LODLevel->Modules.Num()-1; ModuleIndex >= 0; ModuleIndex--)
+			{
+				UParticleModule* ParticleModule = LODLevel->Modules[ModuleIndex];
+				if( ParticleModule == NULL )
+				{
+					LODLevel->Modules.RemoveAt(ModuleIndex);
+					MarkPackageDirty();
+				}
+			}
+		}
+	}
+
+
+	UObject* MyOuter = GetOuter();
+	UParticleSystem* PSysOuter = Cast<UParticleSystem>(MyOuter);
+	bool bRegenDup = false;
+	if (PSysOuter)
+	{
+		bRegenDup = PSysOuter->bRegenerateLODDuplicate;
+	}
+
+	// Clamp the detail spawn rate scale...
+	QualityLevelSpawnRateScale = FMath::Clamp<float>(QualityLevelSpawnRateScale, 0.0f, 1.0f);
+
+	UpdateModuleLists();
+}
+
+#if WITH_EDITOR
+void UParticleEmitter::PostEditChangeProperty(FPropertyChangedEvent& PropertyChangedEvent)
+{
+	check(GIsEditor);
+
+	// Reset the peak active particle counts.
+	// This could check for changes to SpawnRate and Burst and only reset then,
+	// but since we reset the particle system after any edited property, it
+	// may as well just autoreset the peak counts.
+	for (int32 LODIndex = 0; LODIndex < LODLevels.Num(); LODIndex++)
+	{
+		UParticleLODLevel* LODLevel = LODLevels[LODIndex];
+		if (LODLevel)
+		{
+			LODLevel->PeakActiveParticles	= 1;
+		}
+	}
+
+	UpdateModuleLists();
+
+	for (TObjectIterator<UParticleSystemComponent> It;It;++It)
+	{
+		if (It->Template)
+		{
+			int32 i;
+
+			for (i=0; i<It->Template->Emitters.Num(); i++)
+			{
+				if (It->Template->Emitters[i] == this)
+				{
+					It->UpdateInstances();
+				}
+			}
+		}
+	}
+
+	Super::PostEditChangeProperty(PropertyChangedEvent);
+
+	if (CalculateMaxActiveParticleCount() == false)
+	{
+		//
+	}
+
+	// Clamp the detail spawn rate scale...
+	QualityLevelSpawnRateScale = FMath::Clamp<float>(QualityLevelSpawnRateScale, 0.0f, 1.0f);
+
+#if	WITH_EDITORONLY_DATA
+	UpdateDetailModeDisplayString();
+#endif
+}
+#endif // WITH_EDITOR
+
+void UParticleEmitter::SetEmitterName(FName Name)
+{
+	EmitterName = Name;
+}
+
+FName& UParticleEmitter::GetEmitterName()
+{
+	return EmitterName;
+}
+
+void UParticleEmitter::SetLODCount(int32 LODCount)
+{
+	// 
+}
+
+void UParticleEmitter::AddEmitterCurvesToEditor(UInterpCurveEdSetup* EdSetup)
+{
+	UE_LOG(LogParticles, Log, TEXT("UParticleEmitter::AddEmitterCurvesToEditor> Should no longer be called..."));
+	return;
+}
+
+void UParticleEmitter::RemoveEmitterCurvesFromEditor(UInterpCurveEdSetup* EdSetup)
+{
+	for (int32 LODIndex = 0; LODIndex < LODLevels.Num(); LODIndex++)
+	{
+		UParticleLODLevel* LODLevel = LODLevels[LODIndex];
+		// Remove the typedata curves...
+		if (LODLevel->TypeDataModule && LODLevel->TypeDataModule->IsDisplayedInCurveEd(EdSetup))
+		{
+			LODLevel->TypeDataModule->RemoveModuleCurvesFromEditor(EdSetup);
+		}
+
+		// Remove the spawn module curves...
+		if (LODLevel->SpawnModule && LODLevel->SpawnModule->IsDisplayedInCurveEd(EdSetup))
+		{
+			LODLevel->SpawnModule->RemoveModuleCurvesFromEditor(EdSetup);
+		}
+
+		// Remove each modules curves as well.
+		for (int32 ii = 0; ii < LODLevel->Modules.Num(); ii++)
+		{
+			if (LODLevel->Modules[ii]->IsDisplayedInCurveEd(EdSetup))
+			{
+				// Remove it from the curve editor!
+				LODLevel->Modules[ii]->RemoveModuleCurvesFromEditor(EdSetup);
+			}
+		}
+	}
+}
+
+void UParticleEmitter::ChangeEditorColor(FColor& Color, UInterpCurveEdSetup* EdSetup)
+{
+#if WITH_EDITORONLY_DATA
+	UParticleLODLevel* LODLevel = LODLevels[0];
+	EmitterEditorColor = Color;
+	for (int32 TabIndex = 0; TabIndex < EdSetup->Tabs.Num(); TabIndex++)
+	{
+		FCurveEdTab*	Tab = &(EdSetup->Tabs[TabIndex]);
+		for (int32 CurveIndex = 0; CurveIndex < Tab->Curves.Num(); CurveIndex++)
+		{
+			FCurveEdEntry* Entry	= &(Tab->Curves[CurveIndex]);
+			if (LODLevel->SpawnModule->Rate.Distribution == Entry->CurveObject)
+			{
+				Entry->CurveColor	= Color;
+			}
+		}
+	}
+#endif // WITH_EDITORONLY_DATA
+}
+
+void UParticleEmitter::AutoPopulateInstanceProperties(UParticleSystemComponent* PSysComp)
+{
+	for (int32 LODIndex = 0; LODIndex < LODLevels.Num(); LODIndex++)
+	{
+		UParticleLODLevel* LODLevel	= LODLevels[LODIndex];
+		for (int32 ModuleIndex = 0; ModuleIndex < LODLevel->Modules.Num(); ModuleIndex++)
+		{
+			UParticleModule* Module = LODLevel->Modules[ModuleIndex];	
+			LODLevel->SpawnModule->AutoPopulateInstanceProperties(PSysComp);
+			LODLevel->RequiredModule->AutoPopulateInstanceProperties(PSysComp);
+			if (LODLevel->TypeDataModule)
+			{
+				LODLevel->TypeDataModule->AutoPopulateInstanceProperties(PSysComp);
+			}
+
+			Module->AutoPopulateInstanceProperties(PSysComp);
+		}
+	}
+}
+
+
+int32 UParticleEmitter::CreateLODLevel(int32 LODLevel, bool bGenerateModuleData)
+{
+	int32					LevelIndex		= -1;
+	UParticleLODLevel*	CreatedLODLevel	= NULL;
+
+	if (LODLevels.Num() == 0)
+	{
+		LODLevel = 0;
+	}
+
+	// Is the requested index outside a viable range?
+	if ((LODLevel < 0) || (LODLevel > LODLevels.Num()))
+	{
+		return -1;
+	}
+
+	// NextHighestLODLevel is the one that will be 'copied'
+	UParticleLODLevel*	NextHighestLODLevel	= NULL;
+	int32 NextHighIndex = -1;
+	// NextLowestLODLevel is the one (and all ones lower than it) that will have their LOD indices updated
+	UParticleLODLevel*	NextLowestLODLevel	= NULL;
+	int32 NextLowIndex = -1;
+
+	// Grab the two surrounding LOD levels...
+	if (LODLevel == 0)
+	{
+		// It is being added at the front of the list... (highest)
+		if (LODLevels.Num() > 0)
+		{
+			NextHighestLODLevel = LODLevels[0];
+			NextHighIndex = 0;
+			NextLowestLODLevel = NextHighestLODLevel;
+			NextLowIndex = 0;
+		}
+	}
+	else
+	if (LODLevel > 0)
+	{
+		NextHighestLODLevel = LODLevels[LODLevel - 1];
+		NextHighIndex = LODLevel - 1;
+		if (LODLevel < LODLevels.Num())
+		{
+			NextLowestLODLevel = LODLevels[LODLevel];
+			NextLowIndex = LODLevel;
+		}
+	}
+	
+	// Update the LODLevel index for the lower levels and
+	// offset the LOD validity flags for the modules...
+	if (NextLowestLODLevel)
+	{
+		NextLowestLODLevel->ConditionalPostLoad();
+		for (int32 LowIndex = LODLevels.Num() - 1; LowIndex >= NextLowIndex; LowIndex--)
+		{
+			UParticleLODLevel* LowRemapLevel = LODLevels[LowIndex];
+			if (LowRemapLevel)
+			{
+				LowRemapLevel->SetLevelIndex(LowIndex + 1);
+			}
+		}
+	}
+
+	// Create a ParticleLODLevel
+	CreatedLODLevel = NewObject<UParticleLODLevel>(this);
+	check(CreatedLODLevel);
+
+	CreatedLODLevel->Level = LODLevel;
+	CreatedLODLevel->bEnabled = true;
+	CreatedLODLevel->ConvertedModules = true;
+	CreatedLODLevel->PeakActiveParticles = 0;
+
+	// Determine where to place it...
+	if (LODLevels.Num() == 0)
+	{
+		LODLevels.InsertZeroed(0, 1);
+		LODLevels[0] = CreatedLODLevel;
+		CreatedLODLevel->Level	= 0;
+	}
+	else
+	{
+		LODLevels.InsertZeroed(LODLevel, 1);
+		LODLevels[LODLevel] = CreatedLODLevel;
+		CreatedLODLevel->Level = LODLevel;
+	}
+
+	if (NextHighestLODLevel)
+	{
+		NextHighestLODLevel->ConditionalPostLoad();
+
+		// Generate from the higher LOD level
+		if (CreatedLODLevel->GenerateFromLODLevel(NextHighestLODLevel, 100.0, bGenerateModuleData) == false)
+		{
+			UE_LOG(LogParticles, Warning, TEXT("Failed to generate LOD level %d from level %d"), LODLevel, NextHighestLODLevel->Level);
+		}
+	}
+	else
+	{
+		// Create the RequiredModule
+		UParticleModuleRequired* RequiredModule = NewObject<UParticleModuleRequired>(GetOuter());
+		check(RequiredModule);
+		RequiredModule->SetToSensibleDefaults(this);
+		CreatedLODLevel->RequiredModule	= RequiredModule;
+
+		// The SpawnRate for the required module
+		RequiredModule->bUseLocalSpace			= false;
+		RequiredModule->bKillOnDeactivate		= false;
+		RequiredModule->bKillOnCompleted		= false;
+		RequiredModule->EmitterDuration			= 1.0f;
+		RequiredModule->EmitterLoops			= 0;
+		RequiredModule->ParticleBurstMethod		= EPBM_Instant;
+#if WITH_EDITORONLY_DATA
+		RequiredModule->ModuleEditorColor		= FColor::MakeRandomColor();
+#endif // WITH_EDITORONLY_DATA
+		RequiredModule->InterpolationMethod		= PSUVIM_None;
+		RequiredModule->SubImages_Horizontal	= 1;
+		RequiredModule->SubImages_Vertical		= 1;
+		RequiredModule->bScaleUV				= false;
+		RequiredModule->RandomImageTime			= 0.0f;
+		RequiredModule->RandomImageChanges		= 0;
+		RequiredModule->bEnabled				= true;
+
+		RequiredModule->LODValidity = (1 << LODLevel);
+
+		// There must be a spawn module as well...
+		UParticleModuleSpawn* SpawnModule = NewObject<UParticleModuleSpawn>(GetOuter());
+		check(SpawnModule);
+		CreatedLODLevel->SpawnModule = SpawnModule;
+		SpawnModule->LODValidity = (1 << LODLevel);
+		UDistributionFloatConstant* ConstantSpawn	= Cast<UDistributionFloatConstant>(SpawnModule->Rate.Distribution);
+		ConstantSpawn->Constant					= 10;
+		ConstantSpawn->bIsDirty					= true;
+		SpawnModule->BurstList.Empty();
+
+		// Copy the TypeData module
+		CreatedLODLevel->TypeDataModule			= NULL;
+	}
+
+	LevelIndex	= CreatedLODLevel->Level;
+
+	MarkPackageDirty();
+
+	return LevelIndex;
+}
+
+
+bool UParticleEmitter::IsLODLevelValid(int32 LODLevel)
+{
+	for (int32 LODIndex = 0; LODIndex < LODLevels.Num(); LODIndex++)
+	{
+		UParticleLODLevel* CheckLODLevel	= LODLevels[LODIndex];
+		if (CheckLODLevel->Level == LODLevel)
+		{
+			return true;
+		}
+	}
+
+	return false;
+}
+
+UParticleLODLevel* UParticleEmitter::GetCurrentLODLevel(FParticleEmitterInstance* Instance)
+{
+	if (!FPlatformProperties::HasEditorOnlyData())
+	{
+		return Instance->CurrentLODLevel;
+	}
+	else
+	{
+		// for the game (where we care about perf) we don't branch
+		if (Instance->Component.IsGameWorld() )
+		{
+			return Instance->CurrentLODLevel;
+		}
+		else
+		{
+			EditorUpdateCurrentLOD( Instance );
+			return Instance->CurrentLODLevel;
+		}
+	}
+}
+
+void UParticleEmitter::EditorUpdateCurrentLOD(FParticleEmitterInstance* Instance)
+{
+#if WITH_EDITORONLY_DATA
+	UParticleLODLevel*	CurrentLODLevel	= NULL;
+	UParticleLODLevel*	Higher			= NULL;
+
+	int32 SetLODLevel = -1;
+	if (const UParticleSystem* Template = Instance->Component.GetTemplate())
+	{
+		int32 DesiredLODLevel = Template->EditorLODSetting;
+		if (GIsEditor && GEngine->bEnableEditorPSysRealtimeLOD)
+		{
+			DesiredLODLevel = Instance->Component.GetCurrentLODIndex();
+		}
+
+		for (int32 LevelIndex = 0; LevelIndex < LODLevels.Num(); LevelIndex++)
+		{
+			Higher	= LODLevels[LevelIndex];
+			if (Higher && (Higher->Level == DesiredLODLevel))
+			{
+				SetLODLevel = LevelIndex;
+				break;
+			}
+		}
+	}
+
+	if (SetLODLevel == -1)
+	{
+		SetLODLevel = 0;
+	}
+	Instance->SetCurrentLODIndex(SetLODLevel, false);
+#endif // WITH_EDITORONLY_DATA
+}
+
+
+
+UParticleLODLevel* UParticleEmitter::GetLODLevel(int32 LODLevel)
+{
+	if (LODLevel >= LODLevels.Num())
+	{
+		return NULL;
+	}
+
+	return LODLevels[LODLevel];
+}
+
+
+bool UParticleEmitter::AutogenerateLowestLODLevel(bool bDuplicateHighest)
+{
+	// Didn't find it?
+	if (LODLevels.Num() == 1)
+	{
+		// We need to generate it...
+		LODLevels.InsertZeroed(1, 1);
+		UParticleLODLevel* LODLevel = NewObject<UParticleLODLevel>(this);
+		check(LODLevel);
+		LODLevels[1]					= LODLevel;
+		LODLevel->Level					= 1;
+		LODLevel->ConvertedModules		= true;
+		LODLevel->PeakActiveParticles	= 0;
+
+		// Grab LODLevel 0 for creation
+		UParticleLODLevel* SourceLODLevel	= LODLevels[0];
+
+		LODLevel->bEnabled				= SourceLODLevel->bEnabled;
+
+		float Percentage	= 10.0f;
+		if (SourceLODLevel->TypeDataModule)
+		{
+			UParticleModuleTypeDataBeam2*	Beam2TD		= Cast<UParticleModuleTypeDataBeam2>(SourceLODLevel->TypeDataModule);
+
+			if (Beam2TD)
+			{
+				// For now, don't support LOD on beams and trails
+				Percentage	= 100.0f;
+			}
+		}
+
+		if (bDuplicateHighest == true)
+		{
+			Percentage = 100.0f;
+		}
+
+		if (LODLevel->GenerateFromLODLevel(SourceLODLevel, Percentage) == false)
+		{
+			UE_LOG(LogParticles, Warning, TEXT("Failed to generate LOD level %d from LOD level 0"), 1);
+			return false;
+		}
+
+		MarkPackageDirty();
+		return true;
+	}
+
+	return true;
+}
+
+
+bool UParticleEmitter::CalculateMaxActiveParticleCount()
+{
+	int32	CurrMaxAPC = 0;
+
+	int32 MaxCount = 0;
+	
+	for (int32 LODIndex = 0; LODIndex < LODLevels.Num(); LODIndex++)
+	{
+		UParticleLODLevel* LODLevel = LODLevels[LODIndex];
+		if (LODLevel && LODLevel->bEnabled)
+		{
+			bool bForceMaxCount = false;
+			// Check for beams or trails
+			if ((LODLevel->Level == 0) && (LODLevel->TypeDataModule != NULL))
+			{
+				UParticleModuleTypeDataBeam2* BeamTD = Cast<UParticleModuleTypeDataBeam2>(LODLevel->TypeDataModule);
+				if (BeamTD)
+				{
+					bForceMaxCount = true;
+					MaxCount = BeamTD->MaxBeamCount + 2;
+				}
+			}
+
+			int32 LODMaxAPC = LODLevel->CalculateMaxActiveParticleCount();
+			if (bForceMaxCount == true)
+			{
+				LODLevel->PeakActiveParticles = MaxCount;
+				LODMaxAPC = MaxCount;
+			}
+
+			if (LODMaxAPC > CurrMaxAPC)
+			{
+				if (LODIndex > 0)
+				{
+					// Check for a ridiculous difference in counts...
+					if ((CurrMaxAPC > 0) && (LODMaxAPC / CurrMaxAPC) > 2)
+					{
+						//UE_LOG(LogParticles, Log, TEXT("MaxActiveParticleCount Discrepancy?\n\tLOD %2d, Emitter %16s"), LODIndex, *GetName());
+					}
+				}
+				CurrMaxAPC = LODMaxAPC;
+			}
+		}
+	}
+
+#if WITH_EDITOR
+	if ((GIsEditor == true) && (CurrMaxAPC > 500))
+	{
+		//@todo. Added an option to the emitter to disable this warning - for 
+		// the RARE cases where it is really required to render that many.
+		UE_LOG(LogParticles, Warning, TEXT("MaxCount = %4d for Emitter %s (%s)"),
+			CurrMaxAPC, *(GetName()), GetOuter() ? *(GetOuter()->GetPathName()) : TEXT("????"));
+	}
+#endif
+	return true;
+}
+
+
+void UParticleEmitter::GetParametersUtilized(TArray<FString>& ParticleSysParamList,
+											 TArray<FString>& ParticleParameterList)
+{
+	// Clear the lists
+	ParticleSysParamList.Empty();
+	ParticleParameterList.Empty();
+
+	TArray<UParticleModule*> ProcessedModules;
+	ProcessedModules.Empty();
+
+	for (int32 LODIndex = 0; LODIndex < LODLevels.Num(); LODIndex++)
+	{
+		UParticleLODLevel* LODLevel = LODLevels[LODIndex];
+		if (LODLevel)
+		{
+			int32 FindIndex;
+			// Grab that parameters from each module...
+			check(LODLevel->RequiredModule);
+			if (ProcessedModules.Find(LODLevel->RequiredModule, FindIndex) == false)
+			{
+				LODLevel->RequiredModule->GetParticleSysParamsUtilized(ParticleSysParamList);
+				LODLevel->RequiredModule->GetParticleParametersUtilized(ParticleParameterList);
+				ProcessedModules.AddUnique(LODLevel->RequiredModule);
+			}
+
+			check(LODLevel->SpawnModule);
+			if (ProcessedModules.Find(LODLevel->SpawnModule, FindIndex) == false)
+			{
+				LODLevel->SpawnModule->GetParticleSysParamsUtilized(ParticleSysParamList);
+				LODLevel->SpawnModule->GetParticleParametersUtilized(ParticleParameterList);
+				ProcessedModules.AddUnique(LODLevel->SpawnModule);
+			}
+
+			if (LODLevel->TypeDataModule)
+			{
+				if (ProcessedModules.Find(LODLevel->TypeDataModule, FindIndex) == false)
+				{
+					LODLevel->TypeDataModule->GetParticleSysParamsUtilized(ParticleSysParamList);
+					LODLevel->TypeDataModule->GetParticleParametersUtilized(ParticleParameterList);
+					ProcessedModules.AddUnique(LODLevel->TypeDataModule);
+				}
+			}
+			
+			for (int32 ModuleIndex = 0; ModuleIndex < LODLevel->Modules.Num(); ModuleIndex++)
+			{
+				UParticleModule* Module = LODLevel->Modules[ModuleIndex];
+				if (Module)
+				{
+					if (ProcessedModules.Find(Module, FindIndex) == false)
+					{
+						Module->GetParticleSysParamsUtilized(ParticleSysParamList);
+						Module->GetParticleParametersUtilized(ParticleParameterList);
+						ProcessedModules.AddUnique(Module);
+					}
+				}
+			}
+		}
+	}
+}
+
+
+void UParticleEmitter::Build()
+{
+	const int32 LODCount = LODLevels.Num();
+	if ( LODCount > 0 )
+	{
+		UParticleLODLevel* HighLODLevel = LODLevels[0];
+		check(HighLODLevel);
+		if (HighLODLevel->TypeDataModule != nullptr)
+		{
+			if(HighLODLevel->TypeDataModule->RequiresBuild())
+			{
+				FParticleEmitterBuildInfo EmitterBuildInfo;
+#if WITH_EDITOR
+				if(!GetOutermost()->bIsCookedForEditor)
+				{
+					HighLODLevel->CompileModules( EmitterBuildInfo );
+				}
+#endif
+				HighLODLevel->TypeDataModule->Build( EmitterBuildInfo );
+			}
+
+			// Allow TypeData module to cache pointers to modules
+			HighLODLevel->TypeDataModule->CacheModuleInfo(this);
+		}
+
+		// Cache particle size/offset data for all LOD Levels
+		CacheEmitterModuleInfo();
+	}
+}
+
+void UParticleEmitter::CacheEmitterModuleInfo()
+{
+	// This assert makes sure that packing is as expected.
+	// Added FBaseColor...
+	// Linear color change
+	// Added Flags field
+
+	bRequiresLoopNotification = false;
+	bAxisLockEnabled = false;
+	bMeshRotationActive = false;
+	LockAxisFlags = EPAL_NONE;
+	ModuleOffsetMap.Empty();
+	ModuleInstanceOffsetMap.Empty();
+	ModuleRandomSeedInstanceOffsetMap.Empty();
+	ModulesNeedingInstanceData.Empty();
+	ModulesNeedingRandomSeedInstanceData.Empty();
+	MeshMaterials.Empty();
+	DynamicParameterDataOffset = 0;
+	LightDataOffset = 0;
+	LightVolumetricScatteringIntensity = 0;
+	CameraPayloadOffset = 0;
+	ParticleSize = sizeof(FBaseParticle);
+	ReqInstanceBytes = 0;
+	PivotOffset = FVector2D(-0.5f, -0.5f);
+	TypeDataOffset = 0;
+	TypeDataInstanceOffset = -1;
+	SubUVAnimation = nullptr;
+
+	UParticleLODLevel* HighLODLevel = GetLODLevel(0);
+	check(HighLODLevel);
+
+	UParticleModuleTypeDataBase* HighTypeData = HighLODLevel->TypeDataModule;
+	if (HighTypeData)
+	{
+		int32 ReqBytes = HighTypeData->RequiredBytes(static_cast<UParticleModuleTypeDataBase*>(nullptr));
+		if (ReqBytes)
+		{
+			TypeDataOffset = ParticleSize;
+			ParticleSize += ReqBytes;
+		}
+
+		int32 TempInstanceBytes = HighTypeData->RequiredBytesPerInstance();
+		if (TempInstanceBytes)
+		{
+			TypeDataInstanceOffset = ReqInstanceBytes;
+			ReqInstanceBytes += TempInstanceBytes;
+		}
+	}
+
+	// Grab required module
+	UParticleModuleRequired* RequiredModule = HighLODLevel->RequiredModule;
+	check(RequiredModule);
+	// mesh rotation active if alignment is set
+	bMeshRotationActive = (RequiredModule->ScreenAlignment == PSA_Velocity || RequiredModule->ScreenAlignment == PSA_AwayFromCenter);
+
+	// NOTE: This code assumes that the same module order occurs in all LOD levels
+
+	for (int32 ModuleIdx = 0; ModuleIdx < HighLODLevel->Modules.Num(); ModuleIdx++)
+	{
+		UParticleModule* ParticleModule = HighLODLevel->Modules[ModuleIdx];
+		check(ParticleModule);
+
+		// Loop notification?
+		bRequiresLoopNotification |= (ParticleModule->bEnabled && ParticleModule->RequiresLoopingNotification());
+
+		if (ParticleModule->IsA(UParticleModuleTypeDataBase::StaticClass()) == false)
+		{
+			int32 ReqBytes = ParticleModule->RequiredBytes(HighTypeData);
+			if (ReqBytes)
+			{
+				ModuleOffsetMap.Add(ParticleModule, ParticleSize);
+				if (ParticleModule->IsA(UParticleModuleParameterDynamic::StaticClass()) && (DynamicParameterDataOffset == 0))
+				{
+					DynamicParameterDataOffset = ParticleSize;
+				}
+				if (ParticleModule->IsA(UParticleModuleLight::StaticClass()) && (LightDataOffset == 0))
+				{
+					UParticleModuleLight* ParticleModuleLight = Cast<UParticleModuleLight>(ParticleModule);
+					LightVolumetricScatteringIntensity = ParticleModuleLight->VolumetricScatteringIntensity;
+					LightDataOffset = ParticleSize;
+				}
+				if (ParticleModule->IsA(UParticleModuleCameraOffset::StaticClass()) && (CameraPayloadOffset == 0))
+				{
+					CameraPayloadOffset = ParticleSize;
+				}
+				ParticleSize += ReqBytes;
+			}
+
+			int32 TempInstanceBytes = ParticleModule->RequiredBytesPerInstance();
+			if (TempInstanceBytes > 0)
+			{
+				// Add the high-lodlevel offset to the lookup map
+				ModuleInstanceOffsetMap.Add(ParticleModule, ReqInstanceBytes);
+				// Remember that this module has emitter-instance data
+				ModulesNeedingInstanceData.Add(ParticleModule);
+
+				// Add all the other LODLevel modules, using the same offset.
+				// This removes the need to always also grab the HighestLODLevel pointer.
+				for (int32 LODIdx = 1; LODIdx < LODLevels.Num(); LODIdx++)
+				{
+					UParticleLODLevel* CurLODLevel = LODLevels[LODIdx];
+					ModuleInstanceOffsetMap.Add(CurLODLevel->Modules[ModuleIdx], ReqInstanceBytes);
+				}
+				ReqInstanceBytes += TempInstanceBytes;
+			}
+
+			// Add space for per instance random seed value if required
+			if (FApp::bUseFixedSeed || ParticleModule->bSupportsRandomSeed)
+			{
+				// Add the high-lodlevel offset to the lookup map
+				ModuleRandomSeedInstanceOffsetMap.Add(ParticleModule, ReqInstanceBytes);
+				// Remember that this module has emitter-instance data
+				ModulesNeedingRandomSeedInstanceData.Add(ParticleModule);
+
+				// Add all the other LODLevel modules, using the same offset.
+				// This removes the need to always also grab the HighestLODLevel pointer.
+				for (int32 LODIdx = 1; LODIdx < LODLevels.Num(); LODIdx++)
+				{
+					UParticleLODLevel* CurLODLevel = LODLevels[LODIdx];
+					ModuleRandomSeedInstanceOffsetMap.Add(CurLODLevel->Modules[ModuleIdx], ReqInstanceBytes);
+				}
+
+				ReqInstanceBytes += sizeof(FParticleRandomSeedInstancePayload);
+			}
+		}
+
+		if (ParticleModule->IsA(UParticleModuleOrientationAxisLock::StaticClass()))
+		{
+			UParticleModuleOrientationAxisLock* Module_AxisLock = CastChecked<UParticleModuleOrientationAxisLock>(ParticleModule);
+			bAxisLockEnabled = Module_AxisLock->bEnabled;
+			LockAxisFlags = Module_AxisLock->LockAxisFlags;
+		}
+		else if (ParticleModule->IsA(UParticleModulePivotOffset::StaticClass()))
+		{
+			PivotOffset += Cast<UParticleModulePivotOffset>(ParticleModule)->PivotOffset;
+		}
+		else if (ParticleModule->IsA(UParticleModuleMeshMaterial::StaticClass()))
+		{
+			UParticleModuleMeshMaterial* MeshMaterialModule = CastChecked<UParticleModuleMeshMaterial>(ParticleModule);
+			if (MeshMaterialModule->bEnabled)
+			{
+				MeshMaterials = MeshMaterialModule->MeshMaterials;
+			}
+		}
+		else if (ParticleModule->IsA(UParticleModuleSubUV::StaticClass()))
+		{
+			USubUVAnimation* ModuleSubUVAnimation = Cast<UParticleModuleSubUV>(ParticleModule)->Animation;
+			SubUVAnimation = ModuleSubUVAnimation && ModuleSubUVAnimation->SubUVTexture && ModuleSubUVAnimation->IsBoundingGeometryValid()
+				? ModuleSubUVAnimation
+				: NULL;
+		}
+		// Perform validation / fixup on some modules that can cause crashes if LODs / Modules are out of sync
+		// This should only be applied on uncooked builds to avoid wasting cycles
+		else if ( !FPlatformProperties::RequiresCookedData() )
+		{
+			if (ParticleModule->IsA(UParticleModuleLocationBoneSocket::StaticClass()))
+			{
+				UParticleModuleLocationBoneSocket::ValidateLODLevels(this, ModuleIdx);
+			}
+		}
+
+		// Set bMeshRotationActive if module says so
+		if(!bMeshRotationActive && ParticleModule->TouchesMeshRotation())
+		{
+			bMeshRotationActive = true;
+		}
+	}
+}
+
+float UParticleEmitter::GetQualityLevelSpawnRateMult()
+{
+	int32 EffectsQuality = Scalability::GetEffectsQualityDirect(IsInGameThread() || IsInParallelGameThread());
+	int32 ReferenceLevel = CVarQLSpawnRateReferenceLevel.GetValueOnAnyThread(true);
+	float Level = (ReferenceLevel - EffectsQuality);
+	float Q = FMath::Pow(QualityLevelSpawnRateScale, Level);
+	return FMath::Min(1.0f, Q);
+}
+
+bool UParticleEmitter::HasAnyEnabledLODs()const
+{
+	for (UParticleLODLevel* LodLevel : LODLevels)
+	{
+		if (LodLevel && LodLevel->bEnabled)
+		{
+			return true;
+		}
+	}
+	
+	return false;
+}
+
+#if STATS
+
+DEFINE_STAT(STAT_EmittersStatGroupTester);
+DEFINE_STAT(STAT_EmittersRTStatGroupTester);
+
+void UParticleEmitter::CreateStatID() const
+{
+	QUICK_SCOPE_CYCLE_COUNTER(STAT_UParticleEmitterCreateStatID);
+
+	UObject* Outer = GetOuter();
+	FName OuterName = Outer ? Outer->GetFName() : NAME_None;
+	FString LongName;
+	LongName = FString(TEXT("Emitter")) / OuterName.ToString() / EmitterName.ToString();
+	StatID = FDynamicStats::CreateStatId<FStatGroup_STATGROUP_Emitters>(LongName);
+	StatIDRT = FDynamicStats::CreateStatId<FStatGroup_STATGROUP_EmittersRT>(LongName / TEXT("RT"));
+}
+#endif
+
+bool UParticleEmitter::IsSignificant(EParticleSignificanceLevel RequiredSignificance)
+{
+	UParticleSystem* PSysOuter = CastChecked<UParticleSystem>(GetOuter());
+	EParticleSignificanceLevel Significance = FMath::Min(PSysOuter->MaxSignificanceLevel, SignificanceLevel);
+	return Significance >= RequiredSignificance;
+}
+
+/*-----------------------------------------------------------------------------
+	UParticleSpriteEmitter implementation.
+-----------------------------------------------------------------------------*/
+UParticleSpriteEmitter::UParticleSpriteEmitter(const FObjectInitializer& ObjectInitializer)
+	: Super(ObjectInitializer)
+{
+}
+
+
+void UParticleSpriteEmitter::PostLoad()
+{
+	Super::PostLoad();
+
+	// Postload the materials
+	for (int32 LODIndex = 0; LODIndex < LODLevels.Num(); LODIndex++)
+	{
+		UParticleLODLevel* LODLevel = LODLevels[LODIndex];
+		if (LODLevel)
+		{
+			UParticleModuleRequired* RequiredModule = LODLevel->RequiredModule;
+			if (RequiredModule)
+			{
+				if (RequiredModule->Material)
+				{
+					RequiredModule->Material->ConditionalPostLoad();
+				}
+			}
+		}
+	}
+}
+
+FParticleEmitterInstance* UParticleSpriteEmitter::CreateInstance(IParticleEmitterInstanceOwner& InComponent)
+{
+	// If this emitter was cooked out or has no valid LOD levels don't create an instance for it.
+	if ((bCookedOut == true) || (LODLevels.Num() == 0))
+	{
+		return NULL;
+	}
+
+	FParticleEmitterInstance* Instance = 0;
+
+	UParticleLODLevel* LODLevel	= GetLODLevel(0);
+	check(LODLevel);
+
+	if (LODLevel->TypeDataModule)
+	{
+		//@todo. This will NOT work for trails/beams!
+		Instance = LODLevel->TypeDataModule->CreateInstance(this, InComponent);
+	}
+	else
+	{
+		Instance = new FParticleSpriteEmitterInstance(InComponent);
+		check(Instance);
+		Instance->InitParameters(this);
+	}
+
+	if (Instance)
+	{
+		Instance->CurrentLODLevelIndex	= 0;
+		Instance->CurrentLODLevel		= LODLevels[Instance->CurrentLODLevelIndex];
+		Instance->Init();
+	}
+
+	return Instance;
+}
+
+void UParticleSpriteEmitter::SetToSensibleDefaults()
+{
+#if WITH_EDITOR
+	PreEditChange(NULL);
+#endif // WITH_EDITOR
+
+	UParticleLODLevel* LODLevel = LODLevels[0];
+
+	// Spawn rate
+	LODLevel->SpawnModule->LODValidity = 1;
+	UDistributionFloatConstant* SpawnRateDist = Cast<UDistributionFloatConstant>(LODLevel->SpawnModule->Rate.Distribution);
+	if (SpawnRateDist)
+	{
+		SpawnRateDist->Constant = 20.f;
+	}
+
+	// Create basic set of modules
+
+	// Lifetime module
+	UParticleModuleLifetime* LifetimeModule = NewObject<UParticleModuleLifetime>(GetOuter());
+	UDistributionFloatUniform* LifetimeDist = Cast<UDistributionFloatUniform>(LifetimeModule->Lifetime.Distribution);
+	if (LifetimeDist)
+	{
+		LifetimeDist->Min = 1.0f;
+		LifetimeDist->Max = 1.0f;
+		LifetimeDist->bIsDirty = true;
+	}
+	LifetimeModule->LODValidity = 1;
+	LODLevel->Modules.Add(LifetimeModule);
+
+	// Size module
+	UParticleModuleSize* SizeModule = NewObject<UParticleModuleSize>(GetOuter());
+	UDistributionVectorUniform* SizeDist = Cast<UDistributionVectorUniform>(SizeModule->StartSize.Distribution);
+	if (SizeDist)
+	{
+		SizeDist->Min = FVector(25.f, 25.f, 25.f);
+		SizeDist->Max = FVector(25.f, 25.f, 25.f);
+		SizeDist->bIsDirty = true;
+	}
+	SizeModule->LODValidity = 1;
+	LODLevel->Modules.Add(SizeModule);
+
+	// Initial velocity module
+	UParticleModuleVelocity* VelModule = NewObject<UParticleModuleVelocity>(GetOuter());
+	UDistributionVectorUniform* VelDist = Cast<UDistributionVectorUniform>(VelModule->StartVelocity.Distribution);
+	if (VelDist)
+	{
+		VelDist->Min = FVector(-10.f, -10.f, 50.f);
+		VelDist->Max = FVector(10.f, 10.f, 100.f);
+		VelDist->bIsDirty = true;
+	}
+	VelModule->LODValidity = 1;
+	LODLevel->Modules.Add(VelModule);
+
+	// Color over life module
+	UParticleModuleColorOverLife* ColorModule = NewObject<UParticleModuleColorOverLife>(GetOuter());
+	UDistributionVectorConstantCurve* ColorCurveDist = Cast<UDistributionVectorConstantCurve>(ColorModule->ColorOverLife.Distribution);
+	if (ColorCurveDist)
+	{
+		// Add two points, one at time 0.0f and one at 1.0f
+		for (int32 Key = 0; Key < 2; Key++)
+		{
+			int32	KeyIndex = ColorCurveDist->CreateNewKey(Key * 1.0f);
+			for (int32 SubIndex = 0; SubIndex < 3; SubIndex++)
+			{
+				ColorCurveDist->SetKeyOut(SubIndex, KeyIndex, 1.0f);
+			}
+		}
+		ColorCurveDist->bIsDirty = true;
+	}
+	ColorModule->AlphaOverLife.Distribution = NewObject<UDistributionFloatConstantCurve>(ColorModule);
+	UDistributionFloatConstantCurve* AlphaCurveDist = Cast<UDistributionFloatConstantCurve>(ColorModule->AlphaOverLife.Distribution);
+	if (AlphaCurveDist)
+	{
+		// Add two points, one at time 0.0f and one at 1.0f
+		for (int32 Key = 0; Key < 2; Key++)
+		{
+			int32	KeyIndex = AlphaCurveDist->CreateNewKey(Key * 1.0f);
+			if (Key == 0)
+			{
+				AlphaCurveDist->SetKeyOut(0, KeyIndex, 1.0f);
+			}
+			else
+			{
+				AlphaCurveDist->SetKeyOut(0, KeyIndex, 0.0f);
+			}
+		}
+		AlphaCurveDist->bIsDirty = true;
+	}
+	ColorModule->LODValidity = 1;
+	LODLevel->Modules.Add(ColorModule);
+
+#if WITH_EDITOR
+	PostEditChange();
+#endif // WITH_EDITOR
+}
+
+////////////////////////////////////////////////////////////////////////////////////////////////////
diff --git a/Engine/Source/Runtime/Engine/Private/Particles/ParticleLODLevel.cpp b/Engine/Source/Runtime/Engine/Private/Particles/ParticleLODLevel.cpp
new file mode 100644
--- /dev/null
+++ b/Engine/Source/Runtime/Engine/Private/Particles/ParticleLODLevel.cpp
@@ -0,0 +1,652 @@
+// Copyright Epic Games, Inc. All Rights Reserved.
+
+#include "Particles/ParticleLODLevel.h"
+#include "Engine/StaticMesh.h"
+#include "ParticleEmitterInstances.h"
+#include "Particles/Event/ParticleModuleEventGenerator.h"
+#include "Particles/Event/ParticleModuleEventReceiverBase.h"
+#include "Particles/Lifetime/ParticleModuleLifetimeBase.h"
+#include "Particles/Material/ParticleModuleMeshMaterial.h"
+#include "Particles/Orbit/ParticleModuleOrbit.h"
+#include "Particles/ParticleModuleRequired.h"
+#include "Particles/Spawn/ParticleModuleSpawn.h"
+#include "Particles/TypeData/ParticleModuleTypeDataBase.h"
+#include "Particles/TypeData/ParticleModuleTypeDataMesh.h"
+#include "StaticMeshResources.h"
+
+#include UE_INLINE_GENERATED_CPP_BY_NAME(ParticleLODLevel)
+
+////////////////////////////////////////////////////////////////////////////////////////////////////
+
+UParticleLODLevel::UParticleLODLevel(const FObjectInitializer& ObjectInitializer)
+	: Super(ObjectInitializer)
+{
+	bEnabled = true;
+	ConvertedModules = true;
+	PeakActiveParticles = 0;
+}
+
+
+void UParticleLODLevel::CompileModules( FParticleEmitterBuildInfo& EmitterBuildInfo )
+{
+	check( RequiredModule );
+	check( SpawnModule );
+
+	// Store a few special modules.
+	EmitterBuildInfo.RequiredModule = RequiredModule;
+	EmitterBuildInfo.SpawnModule = SpawnModule;
+
+	// Compile those special modules.
+	RequiredModule->CompileModule( EmitterBuildInfo );
+	if ( SpawnModule->bEnabled )
+	{
+		SpawnModule->CompileModule( EmitterBuildInfo );
+	}
+
+	// Compile all remaining modules.
+	const int32 ModuleCount = Modules.Num();
+	for ( int32 ModuleIndex = 0; ModuleIndex < ModuleCount; ++ModuleIndex )
+	{
+		UParticleModule* Module = Modules[ModuleIndex];
+		if ( Module && Module->bEnabled )
+		{
+			Module->CompileModule( EmitterBuildInfo );
+		}
+	}
+
+	// Estimate the maximum number of active particles.
+	EmitterBuildInfo.EstimatedMaxActiveParticleCount = CalculateMaxActiveParticleCount();
+}
+
+bool UParticleLODLevel::IsPostLoadThreadSafe() const
+{
+	return false;
+}
+
+void UParticleLODLevel::PostLoad()
+{
+	Super::PostLoad();
+
+#if WITH_EDITORONLY_DATA
+	checkf(SpawnModule, TEXT("Missing spawn module on %s (%s)"), *(GetPathName()), 
+		(GetOuter() ? (GetOuter()->GetOuter() ? *(GetOuter()->GetOuter()->GetPathName()) : *(GetOuter()->GetPathName())) : TEXT("???")));
+#endif // WITH_EDITORONLY_DATA
+
+	{
+		RequiredModule->ConditionalPostLoad();
+	}
+	if ( SpawnModule )
+	{
+		SpawnModule->ConditionalPostLoad();
+	}
+
+	for (UParticleModule* ParticleModule : Modules)
+	{
+		ParticleModule->ConditionalPostLoad();
+	}
+
+	// shouldn't ever set another UObjects serialized variable in post load
+	// this causes determinisitc cooking issues due to load order being different
+	/*if (RequiredModule)
+	{
+		RequiredModule->ConditionalPostLoad();
+		if (RequiredModule->bEnabled != bEnabled)
+		{
+			RequiredModule->bEnabled = bEnabled;
+		}
+	}*/
+}
+
+void UParticleLODLevel::UpdateModuleLists()
+{
+	LLM_SCOPE(ELLMTag::Particles);
+
+	SpawningModules.Empty();
+	SpawnModules.Empty();
+	UpdateModules.Empty();
+	OrbitModules.Empty();
+	EventReceiverModules.Empty();
+	EventGenerator = NULL;
+
+	UParticleModule* Module;
+	int32 TypeDataModuleIndex = -1;
+
+	for (int32 i = 0; i < Modules.Num(); i++)
+	{
+		Module = Modules[i];
+		if (!Module)
+		{
+			continue;
+		}
+
+		if (Module->bSpawnModule)
+		{
+			SpawnModules.Add(Module);
+		}
+		if (Module->bUpdateModule || Module->bFinalUpdateModule)
+		{
+			UpdateModules.Add(Module);
+		}
+
+		if (Module->IsA(UParticleModuleTypeDataBase::StaticClass()))
+		{
+			TypeDataModule = CastChecked<UParticleModuleTypeDataBase>(Module);
+			if (!Module->bSpawnModule && !Module->bUpdateModule)
+			{
+				// For now, remove it from the list and set it as the TypeDataModule
+				TypeDataModuleIndex = i;
+			}
+		}
+		else
+		if (Module->IsA(UParticleModuleSpawnBase::StaticClass()))
+		{
+			UParticleModuleSpawnBase* SpawnBase = CastChecked<UParticleModuleSpawnBase>(Module);
+			SpawningModules.Add(SpawnBase);
+		}
+		else
+		if (Module->IsA(UParticleModuleOrbit::StaticClass()))
+		{
+			UParticleModuleOrbit* Orbit = CastChecked<UParticleModuleOrbit>(Module);
+			OrbitModules.Add(Orbit);
+		}
+		else
+		if (Module->IsA(UParticleModuleEventGenerator::StaticClass()))
+		{
+			EventGenerator = CastChecked<UParticleModuleEventGenerator>(Module);
+		}
+		else
+		if (Module->IsA(UParticleModuleEventReceiverBase::StaticClass()))
+		{
+			UParticleModuleEventReceiverBase* Event = CastChecked<UParticleModuleEventReceiverBase>(Module);
+			EventReceiverModules.Add(Event);
+		}
+	}
+
+	if (EventGenerator)
+	{
+		// Force the event generator module to the top of the module stack...
+		Modules.RemoveSingle(EventGenerator);
+		Modules.Insert(EventGenerator, 0);
+	}
+
+	if (TypeDataModuleIndex != -1)
+	{
+		Modules.RemoveAt(TypeDataModuleIndex);
+	}
+
+	if (TypeDataModule /**&& (Level == 0)**/)
+	{
+		UParticleModuleTypeDataMesh* MeshTD = Cast<UParticleModuleTypeDataMesh>(TypeDataModule);
+		if (MeshTD
+			&& MeshTD->Mesh
+			&& MeshTD->Mesh->HasValidRenderData(false))
+		{
+			UParticleSpriteEmitter* SpriteEmitter = Cast<UParticleSpriteEmitter>(GetOuter());
+			if (SpriteEmitter && (MeshTD->bOverrideMaterial == false))
+			{
+				FStaticMeshSection& Section = MeshTD->Mesh->GetRenderData()->LODResources[0].Sections[0];
+				UMaterialInterface* Material = MeshTD->Mesh->GetMaterial(Section.MaterialIndex);
+				if (Material)
+				{
+					RequiredModule->Material = Material;
+				}
+			}
+		}
+	}
+}
+
+
+bool UParticleLODLevel::GenerateFromLODLevel(UParticleLODLevel* SourceLODLevel, float Percentage, bool bGenerateModuleData)
+{
+	// See if there are already modules in place
+	if (Modules.Num() > 0)
+	{
+		UE_LOG(LogParticles, Log, TEXT("ERROR? - GenerateFromLODLevel - modules already present!"));
+		return false;
+	}
+
+	bool	bResult	= true;
+
+	// Allocate slots in the array...
+	Modules.InsertZeroed(0, SourceLODLevel->Modules.Num());
+
+	// Set the enabled flag
+	bEnabled = SourceLODLevel->bEnabled;
+
+	// Set up for undo/redo!
+	SetFlags(RF_Transactional);
+
+	// Required module...
+	RequiredModule = CastChecked<UParticleModuleRequired>(
+		SourceLODLevel->RequiredModule->GenerateLODModule(SourceLODLevel, this, Percentage, bGenerateModuleData));
+
+	// Spawn module...
+	SpawnModule = CastChecked<UParticleModuleSpawn>(
+		SourceLODLevel->SpawnModule->GenerateLODModule(SourceLODLevel, this, Percentage, bGenerateModuleData));
+
+	// TypeData module, if present...
+	if (SourceLODLevel->TypeDataModule)
+	{
+		TypeDataModule = 
+			CastChecked<UParticleModuleTypeDataBase>(
+			SourceLODLevel->TypeDataModule->GenerateLODModule(SourceLODLevel, this, Percentage, bGenerateModuleData));
+		check(TypeDataModule == SourceLODLevel->TypeDataModule); // Code expects typedata to be the same across LODs
+	}
+
+	// The remaining modules...
+	for (int32 ModuleIndex = 0; ModuleIndex < SourceLODLevel->Modules.Num(); ModuleIndex++)
+	{
+		if (SourceLODLevel->Modules[ModuleIndex])
+		{
+			Modules[ModuleIndex] = SourceLODLevel->Modules[ModuleIndex]->GenerateLODModule(SourceLODLevel, this, Percentage, bGenerateModuleData);
+		}
+		else
+		{
+			Modules[ModuleIndex] = NULL;
+		}
+	}
+
+	return bResult;
+}
+
+
+int32	UParticleLODLevel::CalculateMaxActiveParticleCount()
+{
+	check(RequiredModule != NULL);
+
+	// Determine the lifetime for particles coming from the emitter
+	float ParticleLifetime = 0.0f;
+	float MaxSpawnRate = SpawnModule->GetEstimatedSpawnRate();
+	int32 MaxBurstCount = SpawnModule->GetMaximumBurstCount();
+	for (int32 ModuleIndex = 0; ModuleIndex < Modules.Num(); ModuleIndex++)
+	{
+		UParticleModuleLifetimeBase* LifetimeMod = Cast<UParticleModuleLifetimeBase>(Modules[ModuleIndex]);
+		if (LifetimeMod != NULL)
+		{
+			ParticleLifetime += LifetimeMod->GetMaxLifetime();
+		}
+
+		UParticleModuleSpawnBase* SpawnMod = Cast<UParticleModuleSpawnBase>(Modules[ModuleIndex]);
+		if (SpawnMod != NULL)
+		{
+			MaxSpawnRate += SpawnMod->GetEstimatedSpawnRate();
+			MaxBurstCount += SpawnMod->GetMaximumBurstCount();
+		}
+	}
+
+	// Determine the maximum duration for this particle system
+	float MaxDuration = 0.0f;
+	float TotalDuration = 0.0f;
+	int32 TotalLoops = 0;
+	if (RequiredModule != NULL)
+	{
+		// We don't care about delay wrt spawning...
+		MaxDuration = FMath::Max<float>(RequiredModule->EmitterDuration, RequiredModule->EmitterDurationLow);
+		TotalLoops = RequiredModule->EmitterLoops;
+		TotalDuration = MaxDuration * TotalLoops;
+	}
+
+	// Determine the max
+	int32 MaxAPC = 0;
+
+	if (TotalDuration != 0.0f)
+	{
+		if (TotalLoops == 1)
+		{
+			// Special case for one loop... 
+			if (ParticleLifetime < MaxDuration)
+			{
+				MaxAPC += FMath::CeilToInt(ParticleLifetime * MaxSpawnRate);
+			}
+			else
+			{
+				MaxAPC += FMath::CeilToInt(MaxDuration * MaxSpawnRate);
+			}
+			// Safety zone...
+			MaxAPC += 1;
+			// Add in the bursts...
+			MaxAPC += MaxBurstCount;
+		}
+		else
+		{
+			if (ParticleLifetime < MaxDuration)
+			{
+				MaxAPC += FMath::CeilToInt(ParticleLifetime * MaxSpawnRate);
+			}
+			else
+			{
+				MaxAPC += (FMath::CeilToInt(FMath::CeilToInt(MaxDuration * MaxSpawnRate) * ParticleLifetime));
+			}
+			// Safety zone...
+			MaxAPC += 1;
+			// Add in the bursts...
+			MaxAPC += MaxBurstCount;
+			if (ParticleLifetime > MaxDuration)
+			{
+				MaxAPC += MaxBurstCount * FMath::CeilToInt(ParticleLifetime - MaxDuration);
+			}
+		}
+	}
+	else
+	{
+		// We are infinite looping... 
+		// Single loop case is all we will worry about. Safer base estimate - but not ideal.
+		if (ParticleLifetime < MaxDuration)
+		{
+			MaxAPC += FMath::CeilToInt(ParticleLifetime * FMath::CeilToInt(MaxSpawnRate));
+		}
+		else
+		{
+			if (ParticleLifetime != 0.0f)
+			{
+				if (ParticleLifetime <= MaxDuration)
+				{
+					MaxAPC += FMath::CeilToInt(MaxDuration * MaxSpawnRate);
+				}
+				else //if (ParticleLifetime > MaxDuration)
+				{
+					MaxAPC += FMath::CeilToInt(MaxDuration * MaxSpawnRate) * ParticleLifetime;
+				}
+			}
+			else
+			{
+				// No lifetime, no duration...
+				MaxAPC += FMath::CeilToInt(MaxSpawnRate);
+			}
+		}
+		// Safety zone...
+		MaxAPC += FMath::Max<int32>(FMath::CeilToInt(MaxSpawnRate * 0.032f), 2);
+		// Burst
+		MaxAPC += MaxBurstCount;
+	}
+
+	PeakActiveParticles = MaxAPC;
+
+	return MaxAPC;
+}
+
+
+void UParticleLODLevel::ConvertToSpawnModule()
+{
+#if WITH_EDITOR
+	// Move the required module SpawnRate and Burst information to a new SpawnModule.
+	if (SpawnModule)
+	{
+//		UE_LOG(LogParticles, Warning, TEXT("LOD Level already has a spawn module!"));
+		return;
+	}
+
+	UParticleEmitter* EmitterOuter = CastChecked<UParticleEmitter>(GetOuter());
+	SpawnModule = NewObject<UParticleModuleSpawn>(EmitterOuter->GetOuter());
+	check(SpawnModule);
+
+	UDistributionFloat* SourceDist = RequiredModule->SpawnRate.Distribution;
+	if (SourceDist)
+	{
+		SpawnModule->Rate.Distribution = Cast<UDistributionFloat>(StaticDuplicateObject(SourceDist, SpawnModule));
+		SpawnModule->Rate.Distribution->bIsDirty = true;
+		SpawnModule->Rate.Initialize();
+	}
+
+	// Now the burst list.
+	int32 BurstCount = RequiredModule->BurstList.Num();
+	if (BurstCount > 0)
+	{
+		SpawnModule->BurstList.AddZeroed(BurstCount);
+		for (int32 BurstIndex = 0; BurstIndex < BurstCount; BurstIndex++)
+		{
+			SpawnModule->BurstList[BurstIndex].Count = RequiredModule->BurstList[BurstIndex].Count;
+			SpawnModule->BurstList[BurstIndex].CountLow = RequiredModule->BurstList[BurstIndex].CountLow;
+			SpawnModule->BurstList[BurstIndex].Time = RequiredModule->BurstList[BurstIndex].Time;
+		}
+	}
+
+	MarkPackageDirty();
+#endif	//#if WITH_EDITOR
+}
+
+
+int32 UParticleLODLevel::GetModuleIndex(UParticleModule* InModule)
+{
+	if (InModule)
+	{
+		if (InModule == RequiredModule)
+		{
+			return INDEX_REQUIREDMODULE;
+		}
+		else if (InModule == SpawnModule)
+		{
+			return INDEX_SPAWNMODULE;
+		}
+		else if (InModule == TypeDataModule)
+		{
+			return INDEX_TYPEDATAMODULE;
+		}
+		else
+		{
+			for (int32 ModuleIndex = 0; ModuleIndex < Modules.Num(); ModuleIndex++)
+			{
+				if (InModule == Modules[ModuleIndex])
+				{
+					return ModuleIndex;
+				}
+			}
+		}
+	}
+
+	return INDEX_NONE;
+}
+
+
+UParticleModule* UParticleLODLevel::GetModuleAtIndex(int32 InIndex)
+{
+	// 'Normal' modules
+	if (InIndex > INDEX_NONE)
+	{
+		if (InIndex < Modules.Num())
+		{
+			return Modules[InIndex];
+		}
+
+		return NULL;
+	}
+
+	switch (InIndex)
+	{
+	case INDEX_REQUIREDMODULE:		return RequiredModule;
+	case INDEX_SPAWNMODULE:			return SpawnModule;
+	case INDEX_TYPEDATAMODULE:		return TypeDataModule;
+	}
+
+	return NULL;
+}
+
+
+void UParticleLODLevel::SetLevelIndex(int32 InLevelIndex)
+{
+	// Remove the 'current' index from the validity flags and set the new one.
+	RequiredModule->LODValidity &= ~(1 << Level);
+	RequiredModule->LODValidity |= (1 << InLevelIndex);
+	SpawnModule->LODValidity &= ~(1 << Level);
+	SpawnModule->LODValidity |= (1 << InLevelIndex);
+	if (TypeDataModule)
+	{
+		TypeDataModule->LODValidity &= ~(1 << Level);
+		TypeDataModule->LODValidity |= (1 << InLevelIndex);
+	}
+	for (int32 ModuleIndex = 0; ModuleIndex < Modules.Num(); ModuleIndex++)
+	{
+		UParticleModule* CheckModule = Modules[ModuleIndex];
+		if (CheckModule)
+		{
+			CheckModule->LODValidity &= ~(1 << Level);
+			CheckModule->LODValidity |= (1 << InLevelIndex);
+		}
+	}
+
+	Level = InLevelIndex;
+}
+
+
+bool UParticleLODLevel::IsModuleEditable(UParticleModule* InModule)
+{
+	// If the module validity flag is not set for this level, it is not editable.
+	if ((InModule->LODValidity & (1 << Level)) == 0)
+	{
+		return false;
+	}
+
+	// If the module is shared w/ higher LOD levels, then it is not editable...
+	int32 Validity = 0;
+	if (Level > 0)
+	{
+		int32 Check = Level - 1;
+		while (Check >= 0)
+		{
+			Validity |= (1 << Check);
+			Check--;
+		}
+
+		if ((Validity & InModule->LODValidity) != 0)
+		{
+			return false;
+		}
+	}
+
+	return true;
+}
+
+void UParticleLODLevel::GetUsedMaterials(TArray<UMaterialInterface*>& OutMaterials, const TArray<FNamedEmitterMaterial>& Slots, const TArray<class UMaterialInterface*>& EmitterMaterials) const
+{
+	// Only process enabled emitters
+	if (bEnabled)
+	{
+		const UParticleModuleTypeDataMesh* MeshTypeData = Cast<UParticleModuleTypeDataMesh>(TypeDataModule);
+
+		if (MeshTypeData && MeshTypeData->Mesh && MeshTypeData->Mesh->GetRenderData())
+		{
+			const FStaticMeshLODResources& LODModel = MeshTypeData->Mesh->GetRenderData()->LODResources[0];
+
+			// Gather the materials applied to the LOD.
+			for (int32 SectionIndex = 0; SectionIndex < LODModel.Sections.Num(); SectionIndex++)
+			{
+				UMaterialInterface* Material = NULL;
+							
+				TArray<FName>& NamedOverrides = RequiredModule->NamedMaterialOverrides;
+
+				if (NamedOverrides.IsValidIndex(SectionIndex))
+				{	
+					//If we have named material overrides then get it's index into the emitter materials array.	
+					for (int32 CheckIdx = 0; CheckIdx < Slots.Num(); ++CheckIdx)
+					{
+						if (NamedOverrides[SectionIndex] == Slots[CheckIdx].Name)
+						{
+							//Default to the default material for that slot.
+							Material = Slots[CheckIdx].Material;
+							if (EmitterMaterials.IsValidIndex(CheckIdx) && nullptr != EmitterMaterials[CheckIdx] )
+							{
+								//This material has been overridden externally, e.g. from a BP so use that one.
+								Material = EmitterMaterials[CheckIdx];
+							}
+
+							break;
+						}
+					}
+				}
+
+				// See if there is a mesh material module.
+				if (Material == NULL)
+				{
+					// Walk in reverse order as in the case of multiple modules, only the final result will be applied
+					for (int32 ModuleIndex = Modules.Num()-1; ModuleIndex >= 0; --ModuleIndex)
+					{
+						UParticleModuleMeshMaterial* MeshMatModule = Cast<UParticleModuleMeshMaterial>(Modules[ModuleIndex]);
+						if (MeshMatModule && MeshMatModule->bEnabled)
+						{
+							if (SectionIndex < MeshMatModule->MeshMaterials.Num())
+							{
+								Material = MeshMatModule->MeshMaterials[SectionIndex];
+								break;
+							}
+						}
+					}
+				}
+
+				// Overriding the material?
+				if (Material == NULL && MeshTypeData->bOverrideMaterial == true)
+				{
+					Material = RequiredModule->Material;
+				}
+
+				// Use the material set on the mesh.
+				if (Material == NULL)
+				{
+					Material = MeshTypeData->Mesh->GetMaterial(LODModel.Sections[SectionIndex].MaterialIndex);
+				}
+
+				if (Material)
+				{
+					OutMaterials.Add(Material);
+				}
+			}
+		}
+		else
+		{
+			UMaterialInterface* Material = NULL;
+							
+			TArray<FName>& NamedOverrides = RequiredModule->NamedMaterialOverrides;
+
+			if (NamedOverrides.Num() > 0)
+			{
+				for (int32 CheckIdx = 0; CheckIdx < Slots.Num(); ++CheckIdx)
+				{
+					if (NamedOverrides[0] == Slots[CheckIdx].Name)
+					{
+						//Default to the default material for that slot.
+						Material = Slots[CheckIdx].Material;
+						if (EmitterMaterials.IsValidIndex(CheckIdx) && nullptr != EmitterMaterials[CheckIdx])
+						{
+							//This material has been overridden externally, e.g. from a BP so use that one.
+							Material = EmitterMaterials[CheckIdx];
+						}
+        
+						break;
+					}
+				}
+			}
+
+			if (!Material)
+			{
+				Material = RequiredModule->Material;
+			}
+
+			OutMaterials.Add(Material);
+		}
+	}
+}
+
+void UParticleLODLevel::GetStreamingMeshInfo(const FBoxSphereBounds& Bounds, TArray<FStreamingRenderAssetPrimitiveInfo>& OutStreamingRenderAssets) const
+{
+	if (bEnabled)
+	{
+		if (const UParticleModuleTypeDataMesh* MeshTypeData = Cast<UParticleModuleTypeDataMesh>(TypeDataModule))
+		{
+			if (UStaticMesh* Mesh = MeshTypeData->Mesh)
+			{
+				if (Mesh->RenderResourceSupportsStreaming() && Mesh->GetRenderAssetType() == EStreamableRenderAssetType::StaticMesh)
+				{
+					const FBoxSphereBounds MeshBounds = Mesh->GetBounds();
+					const FBoxSphereBounds StreamingBounds = FBoxSphereBounds(
+						Bounds.Origin + MeshBounds.Origin,
+						MeshBounds.BoxExtent * MeshTypeData->LODSizeScale,
+						MeshBounds.SphereRadius * MeshTypeData->LODSizeScale);
+					const float MeshTexelFactor = MeshBounds.SphereRadius * 2.0f;
+
+					new (OutStreamingRenderAssets) FStreamingRenderAssetPrimitiveInfo(Mesh, StreamingBounds, MeshTexelFactor);
+				}
+			}
+		}
+	}
+}
+
+////////////////////////////////////////////////////////////////////////////////////////////////////
diff --git a/Engine/Source/Runtime/Engine/Private/Particles/ParticleSystem.cpp b/Engine/Source/Runtime/Engine/Private/Particles/ParticleSystem.cpp
new file mode 100644
--- /dev/null
+++ b/Engine/Source/Runtime/Engine/Private/Particles/ParticleSystem.cpp
@@ -0,0 +1,1420 @@
+// Copyright Epic Games, Inc. All Rights Reserved.
+
+#include "Particles/ParticleSystem.h"
+#include "DeviceProfiles/DeviceProfile.h"
+#include "DeviceProfiles/DeviceProfileManager.h"
+#include "Interfaces/ITargetPlatform.h"
+#include "ObjectEditorUtils.h"
+#include "ParticleHelper.h"
+#include "Particles/Collision/ParticleModuleCollision.h"
+#include "Particles/Color/ParticleModuleColorBase.h"
+#include "Particles/ParticleEmitter.h"
+#include "Particles/ParticleLODLevel.h"
+#include "Particles/ParticleModuleRequired.h"
+#include "Particles/ParticleSystemComponent.h"
+#include "Particles/Spawn/ParticleModuleSpawn.h"
+#include "Particles/TypeData/ParticleModuleTypeDataBase.h"
+#include "Particles/TypeData/ParticleModuleTypeDataGpu.h"
+#include "PSOPrecacheMaterial.h"
+#include "UObject/AssetRegistryTagsContext.h"
+#include "UObject/ObjectSaveContext.h"
+#include "UObject/UObjectIterator.h"
+
+#include UE_INLINE_GENERATED_CPP_BY_NAME(ParticleSystem)
+
+////////////////////////////////////////////////////////////////////////////////////////////////////
+
+/** When to precache Cascade systems' PSOs */
+extern int32 GCascadePSOPrecachingTime;
+
+////////////////////////////////////////////////////////////////////////////////////////////////////
+
+UFXSystemAsset::UFXSystemAsset(const FObjectInitializer& ObjectInitializer)
+	: Super(ObjectInitializer)
+{
+}
+
+void UFXSystemAsset::LaunchPSOPrecaching(const FMaterialInterfacePSOPrecacheParamsList& PSOPrecacheParamsList)
+{
+	FGraphEventArray PrecachePSOsEvents;
+	if (IsComponentPSOPrecachingEnabled())
+	{
+		PrecacheMaterialPSOs(PSOPrecacheParamsList, MaterialPSOPrecacheRequestIDs, PrecachePSOsEvents);
+	} 
+
+	// Create task to signal that the PSO precache events are done by adding them as prerequisite to the task.
+	if (PrecachePSOsEvents.Num() > 0)
+	{
+		struct FReleasePrecachePSOsEventTask
+		{
+			explicit FReleasePrecachePSOsEventTask(UFXSystemAsset* OwnerAsset)
+				: WeakOwnerAsset(OwnerAsset)
+			{
+			}
+
+			static TStatId GetStatId() { return TStatId(); }
+			static ENamedThreads::Type GetDesiredThread() { return ENamedThreads::GameThread; }
+			static ESubsequentsMode::Type GetSubsequentsMode() { return ESubsequentsMode::TrackSubsequents; }
+
+			void DoTask(ENamedThreads::Type CurrentThread, const FGraphEventRef& MyCompletionGraphEvent)
+			{
+				if (UFXSystemAsset* Asset = WeakOwnerAsset.Get())
+				{
+					Asset->PrecachePSOsEvent = nullptr;
+				}
+			}
+
+			TWeakObjectPtr<UFXSystemAsset> WeakOwnerAsset;
+		};
+
+		// need to set `PrecachePSOsEvent` before the task is launched to not race with its execution
+		TGraphTask<FReleasePrecachePSOsEventTask>* ReleasePrecachePSOsEventTask = TGraphTask<FReleasePrecachePSOsEventTask>::CreateTask(&PrecachePSOsEvents).ConstructAndHold(this);
+		PrecachePSOsEvent = ReleasePrecachePSOsEventTask->GetCompletionEvent();
+		ReleasePrecachePSOsEventTask->Unlock();
+	}
+
+	PSOPrecachingLaunched = true;
+}
+
+////////////////////////////////////////////////////////////////////////////////////////////////////
+
+UParticleSystem::UParticleSystem(const FObjectInitializer& ObjectInitializer)
+	: Super(ObjectInitializer)
+	, OcclusionBoundsMethod(EPSOBM_ParticleBounds)
+	, bAnyEmitterLoopsForever(false)
+	, HighestSignificance(EParticleSignificanceLevel::Critical)
+	, LowestSignificance(EParticleSignificanceLevel::Low)
+	, bShouldManageSignificance(false)
+	, bIsImmortal(false)
+	, bWillBecomeZombie(false)
+{
+#if WITH_EDITORONLY_DATA
+	ThumbnailDistance = 200.0;
+	ThumbnailWarmup = 1.0;
+#endif // WITH_EDITORONLY_DATA
+	UpdateTime_FPS = 60.0f;
+	UpdateTime_Delta = 1.0f/60.0f;
+	WarmupTime = 0.0f;
+	WarmupTickRate = 0.0f;
+#if WITH_EDITORONLY_DATA
+	EditorLODSetting = 0;
+#endif // WITH_EDITORONLY_DATA
+	FixedRelativeBoundingBox.Min = FVector(-1.0f, -1.0f, -1.0f);
+	FixedRelativeBoundingBox.Max = FVector(1.0f, 1.0f, 1.0f);
+	FixedRelativeBoundingBox.IsValid = true;
+
+	LODMethod = PARTICLESYSTEMLODMETHOD_Automatic;
+	LODDistanceCheckTime = 0.25f;
+	bRegenerateLODDuplicate = false;
+	ThumbnailImageOutOfDate = true;
+#if WITH_EDITORONLY_DATA
+	FloorMesh = TEXT("/Engine/EditorMeshes/AnimTreeEd_PreviewFloor.AnimTreeEd_PreviewFloor");
+	FloorPosition = FVector(0.0f, 0.0f, 0.0f);
+	FloorRotation = FRotator(0.0f, 0.0f, 0.0f);
+	FloorScale = 1.0f;
+	FloorScale3D = FVector(1.0f, 1.0f, 1.0f);
+#endif // WITH_EDITORONLY_DATA
+
+	MacroUVPosition = FVector(0.0f, 0.0f, 0.0f);
+
+	MacroUVRadius = 200.0f;
+	bAutoDeactivate = true;
+	MinTimeBetweenTicks = 0;
+	InsignificantReaction = EParticleSystemInsignificanceReaction::Auto;
+	InsignificanceDelay = 0.0f;
+	MaxSignificanceLevel = EParticleSignificanceLevel::Critical;
+	MaxPoolSize = 32;
+
+
+	bAllowManagedTicking = true;
+}
+
+ParticleSystemLODMethod UParticleSystem::GetCurrentLODMethod()
+{
+	return ParticleSystemLODMethod(LODMethod);
+}
+
+
+int32 UParticleSystem::GetLODLevelCount()
+{
+	return LODDistances.Num();
+}
+
+
+float UParticleSystem::GetLODDistance(int32 LODLevelIndex)
+{
+	if (LODLevelIndex >= LODDistances.Num())
+	{
+		return -1.0f;
+	}
+
+	return LODDistances[LODLevelIndex];
+}
+
+void UParticleSystem::SetCurrentLODMethod(ParticleSystemLODMethod InMethod)
+{
+	LODMethod = InMethod;
+}
+
+
+bool UParticleSystem::SetLODDistance(int32 LODLevelIndex, float InDistance)
+{
+	if (LODLevelIndex >= LODDistances.Num())
+	{
+		return false;
+	}
+
+	LODDistances[LODLevelIndex] = InDistance;
+
+	return true;
+}
+
+bool UParticleSystem::DoesAnyEmitterHaveMotionBlur(int32 LODLevelIndex) const
+{
+	for (auto& EmitterIter : Emitters)
+	{
+		if (EmitterIter)
+		{
+			auto* EmitterLOD = EmitterIter->GetLODLevel(LODLevelIndex);
+			if (!EmitterLOD)
+			{
+				continue;
+			}
+
+			if (EmitterLOD->TypeDataModule && EmitterLOD->TypeDataModule->IsMotionBlurEnabled())
+			{
+				return true;
+			}
+
+			if (EmitterLOD->RequiredModule && EmitterLOD->RequiredModule->ShouldUseVelocityForMotionBlur())
+			{
+				return true;
+			}
+		}
+	}
+
+	return false;
+}
+
+#if WITH_EDITOR
+void UParticleSystem::PostEditChangeProperty(FPropertyChangedEvent& PropertyChangedEvent)
+{
+	UpdateTime_Delta = 1.0f / UpdateTime_FPS;
+
+	bIsElligibleForAsyncTickComputed = false;
+
+	//If the property is NULL then we don't really know what's happened. 
+	//Could well be a module change, requiring all instances to be destroyed and recreated.
+	bool bEmptyInstances = PropertyChangedEvent.Property == NULL;
+	for (TObjectIterator<UParticleSystemComponent> It;It;++It)
+	{
+		if (It->Template == this)
+		{
+			It->UpdateInstances(bEmptyInstances);
+		}
+	}
+
+	// Ensure the bounds have a positive size
+	if (FixedRelativeBoundingBox.IsValid)
+	{
+		if (FixedRelativeBoundingBox.Min.X > FixedRelativeBoundingBox.Max.X)
+		{
+			Swap(FixedRelativeBoundingBox.Min.X, FixedRelativeBoundingBox.Max.X);
+		}
+		if (FixedRelativeBoundingBox.Min.Y > FixedRelativeBoundingBox.Max.Y)
+		{
+			Swap(FixedRelativeBoundingBox.Min.Y, FixedRelativeBoundingBox.Max.Y);
+		}
+		if (FixedRelativeBoundingBox.Min.Z > FixedRelativeBoundingBox.Max.Z)
+		{
+			Swap(FixedRelativeBoundingBox.Min.Z, FixedRelativeBoundingBox.Max.Z);
+		}
+	}
+
+	// Recompute the looping flag
+	bAnyEmitterLoopsForever = false;
+	bIsImmortal = false;
+	bWillBecomeZombie = false;
+	HighestSignificance = EParticleSignificanceLevel::Low;
+	LowestSignificance = EParticleSignificanceLevel::Critical;
+	for (UParticleEmitter* Emitter : Emitters)
+	{
+		if (Emitter != nullptr)
+		{
+			for (const UParticleLODLevel* LODLevel : Emitter->LODLevels)
+			{
+				if (LODLevel != nullptr)
+				{
+					if (LODLevel->bEnabled)
+					{
+						const UParticleModuleRequired* RequiredModule = LODLevel->RequiredModule;
+						if ((RequiredModule != nullptr) && (RequiredModule->EmitterLoops == 0))
+						{
+							bAnyEmitterLoopsForever = true;
+
+							UParticleModuleSpawn* SpawnModule = LODLevel->SpawnModule;
+							check(SpawnModule);
+
+							// check if any emitter will cause the system to never be deleted
+							// terms here are zombie (burst-only, so will stop spawning but emitter instances and psys component will continue existing)
+							// and immortal (any emitter will loop indefinitely and does not have finite duration)
+							if (RequiredModule->EmitterDuration == 0.0f)
+							{
+								bIsImmortal = true;
+								if (SpawnModule->GetMaximumSpawnRate() == 0.0f && !bAutoDeactivate)
+								{
+									bWillBecomeZombie = true;
+								}
+							}
+						}
+					}
+				}
+			}
+
+			EParticleSignificanceLevel EmitterSignificance = FMath::Min(MaxSignificanceLevel, Emitter->SignificanceLevel);
+			if (EmitterSignificance > HighestSignificance)
+			{
+				HighestSignificance = EmitterSignificance;
+			}
+			if (EmitterSignificance < LowestSignificance)
+			{
+				LowestSignificance = EmitterSignificance;
+			}
+		}
+	}
+
+	bShouldManageSignificance = GetLowestSignificance() != EParticleSignificanceLevel::Critical/* && !ContainsEmitterType(UParticleModuleTypeDataBeam2::StaticClass())*/;
+
+	//cap the WarmupTickRate to realistic values
+	if (WarmupTickRate <= 0)
+	{
+		WarmupTickRate = 0;
+	}
+	else if (WarmupTickRate > WarmupTime)
+	{
+		WarmupTickRate = WarmupTime;
+	}
+
+	ThumbnailImageOutOfDate = true;
+
+	Super::PostEditChangeProperty(PropertyChangedEvent);
+}
+#endif // WITH_EDITOR
+
+void UParticleSystem::PreSave(const class ITargetPlatform* TargetPlatform)
+{
+	PRAGMA_DISABLE_DEPRECATION_WARNINGS;
+	Super::PreSave(TargetPlatform);
+	PRAGMA_ENABLE_DEPRECATION_WARNINGS;
+}
+
+void UParticleSystem::PreSave(FObjectPreSaveContext ObjectSaveContext)
+{
+	Super::PreSave(ObjectSaveContext);
+#if WITH_EDITORONLY_DATA
+	// Ensure that soloing is undone...
+	int32 NumEmitters = FMath::Min(Emitters.Num(),SoloTracking.Num());
+	for (int32 EmitterIdx = 0; EmitterIdx < NumEmitters; EmitterIdx++)
+	{
+		UParticleEmitter* Emitter = Emitters[EmitterIdx];
+		if (Emitter != nullptr)
+		{
+			Emitter->bIsSoloing = false;
+			FLODSoloTrack& SoloTrack = SoloTracking[EmitterIdx];
+			int32 NumLODs = FMath::Min(Emitter->LODLevels.Num(), SoloTrack.SoloEnableSetting.Num());
+			for (int32 LODIdx = 0; LODIdx < NumLODs; LODIdx++)
+			{
+				UParticleLODLevel* LODLevel = Emitter->LODLevels[LODIdx];
+				{
+					// Restore the enabled settings - ie turn off soloing...
+					LODLevel->bEnabled = SoloTrack.SoloEnableSetting[LODIdx];
+				}
+			}
+		}
+	}
+#endif // WITH_EDITORONLY_DATA
+}
+
+bool UParticleSystem::IsPostLoadThreadSafe() const
+{
+	return false;
+}
+
+void UParticleSystem::PostLoad()
+{
+	Super::PostLoad();
+
+	// Run thru all of the emitters, load them up and compute some flags based on them
+	bHasPhysics = false;
+	bAnyEmitterLoopsForever = false;
+	bIsImmortal = false;
+	bWillBecomeZombie = false;
+	HighestSignificance = EParticleSignificanceLevel::Low;
+	LowestSignificance = EParticleSignificanceLevel::Critical;
+	for (int32 i = Emitters.Num() - 1; i >= 0; i--)
+	{
+		// Remove any old emitters
+		UParticleEmitter* Emitter = Emitters[i];
+		if (Emitter == NULL)
+		{
+			// Empty emitter slots are ok with cooked content.
+			if( !FPlatformProperties::RequiresCookedData() && !GIsServer)
+			{
+				UE_LOG(LogParticles, Warning, TEXT("ParticleSystem contains empty emitter slots - %s"), *GetFullName());
+			}
+			continue;
+		}
+
+		Emitter->ConditionalPostLoad();
+
+		bool bCookedOut = false;
+		if (UParticleSpriteEmitter* SpriteEmitter = Cast<UParticleSpriteEmitter>(Emitter))
+		{
+			bCookedOut = SpriteEmitter->bCookedOut;
+		}
+
+		if (!bCookedOut)
+		{
+			if (Emitter->LODLevels.Num() == 0)
+			{
+				UE_LOG(LogParticles, Warning, TEXT("ParticleSystem contains emitter with no lod levels - %s - %s"), *GetFullName(), *Emitter->GetFullName());
+				continue;
+			}
+
+			UParticleLODLevel* LODLevel = Emitter->LODLevels[0];
+			check(LODLevel);
+
+			LODLevel->ConditionalPostLoad();
+				
+			//@todo. Move these flag calculations into the editor and serialize?
+			//Should mirror implementation in PostEditChangeProperty.
+			for (UParticleLODLevel* ParticleLODLevel : Emitter->LODLevels)
+			{
+				if (ParticleLODLevel)
+				{
+					//@todo. This is a temporary fix for emitters that apply physics.
+					// Check for collision modules with bApplyPhysics set to true
+					for (int32 ModuleIndex = 0; ModuleIndex < LODLevel->Modules.Num(); ModuleIndex++)
+					{
+						UParticleModuleCollision* CollisionModule = Cast<UParticleModuleCollision>(ParticleLODLevel->Modules[ModuleIndex]);
+						if (CollisionModule)
+						{
+							if (CollisionModule->bApplyPhysics == true)
+							{
+								bHasPhysics = true;
+								break;
+							}
+						}
+					}
+
+					if (LODLevel->bEnabled)
+					{
+						const UParticleModuleRequired* RequiredModule = LODLevel->RequiredModule;
+						if ((RequiredModule != nullptr) && (RequiredModule->EmitterLoops == 0))
+						{
+							bAnyEmitterLoopsForever = true;
+
+							UParticleModuleSpawn* SpawnModule = LODLevel->SpawnModule;
+							check(SpawnModule);
+
+							if (RequiredModule->EmitterDuration == 0.0f)
+							{
+								bIsImmortal = true;
+								if (SpawnModule->GetMaximumSpawnRate() == 0.0f && !bAutoDeactivate)
+								{
+									bWillBecomeZombie = true;
+								}
+							}
+						}
+					}
+				}
+			}
+
+			EParticleSignificanceLevel EmitterSignificance = FMath::Min(MaxSignificanceLevel, Emitter->SignificanceLevel);
+			if (EmitterSignificance > HighestSignificance)
+			{
+				HighestSignificance = EmitterSignificance;
+			}
+			if (EmitterSignificance < LowestSignificance)
+			{
+				LowestSignificance = EmitterSignificance;
+			}
+		}
+	}
+
+	bShouldManageSignificance = GetLowestSignificance() != EParticleSignificanceLevel::Critical /* && !ContainsEmitterType(UParticleModuleTypeDataBeam2::StaticClass())*/;
+
+	if (LODSettings.Num() == 0)
+	{
+		if (Emitters.Num() > 0)
+		{
+			UParticleEmitter* Emitter = Emitters[0];
+			if (Emitter)
+			{
+				LODSettings.AddUninitialized(Emitter->LODLevels.Num());
+				for (int32 LODIndex = 0; LODIndex < LODSettings.Num(); LODIndex++)
+				{
+					LODSettings[LODIndex] = FParticleSystemLOD::CreateParticleSystemLOD();
+				}
+			}
+		}
+		else
+		{
+			LODSettings.AddUninitialized();
+			LODSettings[0] = FParticleSystemLOD::CreateParticleSystemLOD();
+		}
+	}
+
+	// Add default LOD Distances
+	if( LODDistances.Num() == 0 && Emitters.Num() > 0 )
+	{
+		UParticleEmitter* Emitter = Emitters[0];
+		if (Emitter)
+		{
+			LODDistances.AddUninitialized(Emitter->LODLevels.Num());
+			for (int32 LODIndex = 0; LODIndex < LODDistances.Num(); LODIndex++)
+			{
+				LODDistances[LODIndex] = LODIndex * 2500.0f;
+			}
+		}
+	}
+
+	if (GCascadePSOPrecachingTime == 1)
+	{
+		PrecachePSOs();
+	}
+
+#if WITH_EDITOR
+	// Due to there still being some ways that LODLevel counts get mismatched,
+	// when loading in the editor LOD levels will always be checked and fixed
+	// up... This can be removed once all the edge cases that lead to the
+	// problem are found and fixed.
+	if (GIsEditor)
+	{
+		// Fix the LOD distance array and mismatched lod levels
+		int32 LODCount_0 = -1;
+		for (int32 EmitterIndex = 0; EmitterIndex < Emitters.Num(); EmitterIndex++)
+		{
+			UParticleEmitter* Emitter  = Emitters[EmitterIndex];
+			if (Emitter)
+			{
+				if (LODCount_0 == -1)
+				{
+					LODCount_0 = Emitter->LODLevels.Num();
+				}
+				else
+				{
+					int32 EmitterLODCount = Emitter->LODLevels.Num();
+					if (EmitterLODCount != LODCount_0)
+					{
+						UE_LOG(LogParticles, Warning, TEXT("Emitter %d has mismatched LOD level count - expected %d, found %d. PS = %s"),
+							EmitterIndex, LODCount_0, EmitterLODCount, *GetPathName());
+						UE_LOG(LogParticles, Warning, TEXT("Fixing up now... Package = %s"), *(GetOutermost()->GetPathName()));
+
+						if (EmitterLODCount > LODCount_0)
+						{
+							Emitter->LODLevels.RemoveAt(LODCount_0, EmitterLODCount - LODCount_0);
+						}
+						else
+						{
+							for (int32 NewLODIndex = EmitterLODCount; NewLODIndex < LODCount_0; NewLODIndex++)
+							{
+								if (Emitter->CreateLODLevel(NewLODIndex) != NewLODIndex)
+								{
+									UE_LOG(LogParticles, Warning, TEXT("Failed to add LOD level %d"), NewLODIndex);
+								}
+							}
+						}
+					}
+				}
+			}
+		}
+
+		if (LODCount_0 > 0)
+		{
+			if (LODDistances.Num() < LODCount_0)
+			{
+				for (int32 DistIndex = LODDistances.Num(); DistIndex < LODCount_0; DistIndex++)
+				{
+					float Distance = DistIndex * 2500.0f;
+					LODDistances.Add(Distance);
+				}
+			}
+			else
+			if (LODDistances.Num() > LODCount_0)
+			{
+				LODDistances.RemoveAt(LODCount_0, LODDistances.Num() - LODCount_0);
+			}
+		}
+		else
+		{
+			LODDistances.Empty();
+		}
+
+		if (LODCount_0 > 0)
+		{
+			if (LODSettings.Num() < LODCount_0)
+			{
+				for (int32 DistIndex = LODSettings.Num(); DistIndex < LODCount_0; DistIndex++)
+				{
+					LODSettings.Add(FParticleSystemLOD::CreateParticleSystemLOD());
+				}
+			}
+			else
+				if (LODSettings.Num() > LODCount_0)
+				{
+					LODSettings.RemoveAt(LODCount_0, LODSettings.Num() - LODCount_0);
+				}
+		}
+		else
+		{
+			LODSettings.Empty();
+		}
+	}
+#endif
+
+#if WITH_EDITORONLY_DATA
+	// Reset cascade's UI LOD setting to 0.
+	EditorLODSetting = 0;
+#endif // WITH_EDITORONLY_DATA
+
+	FixedRelativeBoundingBox.IsValid = true;
+
+	// Set up the SoloTracking...
+	SetupSoloing();
+}
+
+void UParticleSystem::PrecachePSOs()
+{
+	if (HasLaunchedPSOPrecaching() || (!IsComponentPSOPrecachingEnabled() && !IsResourcePSOPrecachingEnabled()))
+	{
+		return;
+	}
+
+	FMaterialInterfacePSOPrecacheParamsList PSOPrecacheParamsList;
+
+	FMaterialInterfacePSOPrecacheParams NewEntry;
+	NewEntry.PSOPrecacheParams.SetMobility(EComponentMobility::Movable);
+
+	// No per component emitter materials known at this point in time
+	TArray<UMaterialInterface*> EmptyEmitterMaterials;
+	// Cached array to collect all materials used for LOD level
+	TArray<UMaterialInterface*> Materials;
+
+	for (int32 EmitterIdx = 0; EmitterIdx < Emitters.Num(); ++EmitterIdx)
+	{
+		const UParticleEmitter* Emitter = Emitters[EmitterIdx];
+		if (!Emitter)
+		{
+			continue;
+		}
+
+		for (int32 LodIndex = 0; LodIndex < Emitter->LODLevels.Num(); ++LodIndex)
+		{
+			const UParticleLODLevel* LOD = Emitter->LODLevels[LodIndex];
+			if (LOD && LOD->bEnabled)
+			{
+				UParticleModuleTypeDataBase::FPSOPrecacheParams PrecacheParams;
+				if (LOD->TypeDataModule)
+				{
+					LOD->TypeDataModule->CollectPSOPrecacheData(Emitter, PrecacheParams);
+				}
+				else
+				{
+					bool bUsesDynamicParameter = (Emitter->DynamicParameterDataOffset > 0);
+					FPSOPrecacheVertexFactoryData VFData;
+					VFData.VertexFactoryType = &FParticleSpriteVertexFactory::StaticType;
+					VFData.CustomDefaultVertexDeclaration = FParticleSpriteVertexFactory::GetPSOPrecacheVertexDeclaration(bUsesDynamicParameter);
+					PrecacheParams.VertexFactoryDataList.Add(VFData);
+					PrecacheParams.PrimitiveType = PT_TriangleList;
+				}
+
+				Materials.Empty();
+				LOD->GetUsedMaterials(Materials, NamedMaterialSlots, EmptyEmitterMaterials);
+
+				for (UMaterialInterface* MaterialInterface : Materials)
+				{
+					NewEntry.MaterialInterface = MaterialInterface;
+					NewEntry.VertexFactoryDataList = PrecacheParams.VertexFactoryDataList;
+					NewEntry.PSOPrecacheParams.PrimitiveType = PrecacheParams.PrimitiveType;
+
+					AddMaterialInterfacePSOPrecacheParamsToList(NewEntry, PSOPrecacheParamsList);
+				}
+			}
+		}
+	}
+
+	LaunchPSOPrecaching(PSOPrecacheParamsList);
+}
+
+void UParticleSystem::Serialize(FArchive& Ar)
+{
+	Ar.UsingCustomVersion(FParticleSystemCustomVersion::GUID);
+
+	int32 CookTargetPlatformDetailModeMask = 0xFFFFFFFF;
+#if WITH_EDITOR
+	if (Ar.IsCooking())
+	{
+		// If we're cooking, check the device profile for whether we want to eliminate all emitters that don't match the detail mode. 
+		// This will only work if scalability settings affecting detail mode can not be changed at runtime (depends on platform)!
+		//
+		if (UDeviceProfile* DeviceProfile = UDeviceProfileManager::Get().FindProfile(Ar.CookingTarget()->IniPlatformName()))
+		{
+			// if we don't prune, we assume all detail modes
+			int32 CVarDoPrune = 0;
+			if (DeviceProfile->GetConsolidatedCVarValue(TEXT("fx.PruneEmittersOnCookByDetailMode"), CVarDoPrune) && CVarDoPrune == 1)
+			{
+				// get the detail mode from the device platform ini; if it's not there, we assume all detail modes
+				int32 CVarDetailMode = PDM_DefaultValue;
+				if (DeviceProfile->GetConsolidatedCVarValue(TEXT("r.DetailMode"), CVarDetailMode))
+				{
+					CookTargetPlatformDetailModeMask = (1 << CVarDetailMode);
+				}
+			}
+		}
+
+		// if we're cooking, save only emitters with matching detail modes
+		for (int32 EmitterIdx = 0; EmitterIdx<Emitters.Num(); EmitterIdx++)
+		{
+			// null out if detail mode doesn't match
+			if (Emitters[EmitterIdx] && !(Emitters[EmitterIdx]->DetailModeBitmask & CookTargetPlatformDetailModeMask))
+			{
+				UE_LOG(LogParticles, Display, TEXT("Pruning emitter, detail mode mismatch (PDM %i) (only works if platform can't change detail mode at runtime!): %s - set fx.PruneEmittersOnCookByDetailMode to 0 in DeviceProfile.ini for the target profile to avoid"), Emitters[EmitterIdx]->DetailModeBitmask, *Emitters[EmitterIdx]->EmitterName.ToString());
+				Emitters[EmitterIdx] = nullptr;
+			}
+		}
+	}
+#endif
+
+	Super::Serialize(Ar);
+
+}
+
+
+void UParticleSystem::UpdateColorModuleClampAlpha(UParticleModuleColorBase* ColorModule)
+{
+	if (ColorModule)
+	{
+		TArray<const FCurveEdEntry*> CurveEntries;
+		ColorModule->RemoveModuleCurvesFromEditor(CurveEdSetup);
+		ColorModule->AddModuleCurvesToEditor(CurveEdSetup, CurveEntries);
+	}
+}
+
+void UParticleSystem::GetAssetRegistryTags(TArray<FAssetRegistryTag>& OutTags) const
+{
+	PRAGMA_DISABLE_DEPRECATION_WARNINGS;
+	Super::GetAssetRegistryTags(OutTags);
+	PRAGMA_ENABLE_DEPRECATION_WARNINGS;
+}
+
+void UParticleSystem::GetAssetRegistryTags(FAssetRegistryTagsContext Context) const
+{
+	Context.AddTag( FAssetRegistryTag("HasGPUEmitter", HasGPUEmitter() ? TEXT("True") : TEXT("False"), FAssetRegistryTag::TT_Alphabetical) );
+
+	const float BoundsSize = FixedRelativeBoundingBox.GetSize().GetMax();
+	Context.AddTag(FAssetRegistryTag("FixedBoundsSize", bUseFixedRelativeBoundingBox ? FString::Printf(TEXT("%.2f"), BoundsSize) : FString(TEXT("None")), FAssetRegistryTag::TT_Numerical));
+
+	Context.AddTag(FAssetRegistryTag("NumEmitters", LexToString(Emitters.Num()), FAssetRegistryTag::TT_Numerical));
+
+	Context.AddTag(FAssetRegistryTag("NumLODs", LexToString(LODDistances.Num()), FAssetRegistryTag::TT_Numerical));
+
+	Context.AddTag(FAssetRegistryTag("WarmupTime", LexToString(WarmupTime), FAssetRegistryTag::TT_Numerical));
+
+	// Done here instead of as an AssetRegistrySearchable string to avoid the long prefix on the enum value string
+	FString LODMethodString = TEXT("Unknown");
+	switch (LODMethod)
+	{
+	case PARTICLESYSTEMLODMETHOD_Automatic:
+		LODMethodString = TEXT("Automatic");
+		break;
+	case PARTICLESYSTEMLODMETHOD_DirectSet:
+		LODMethodString = TEXT("DirectSet");
+		break;
+	case PARTICLESYSTEMLODMETHOD_ActivateAutomatic:
+		LODMethodString = TEXT("Activate Automatic");
+		break;
+	default:
+		check(false); // Missing enum entry
+		break;
+	}
+	Context.AddTag(FAssetRegistryTag("LODMethod", LODMethodString, FAssetRegistryTag::TT_Alphabetical));
+
+	Context.AddTag(FAssetRegistryTag("CPUCollision", UsesCPUCollision() ? TEXT("True") : TEXT("False"), FAssetRegistryTag::TT_Alphabetical));
+	Context.AddTag(FAssetRegistryTag("Looping", bAnyEmitterLoopsForever ? TEXT("True") : TEXT("False"), FAssetRegistryTag::TT_Alphabetical));
+	Context.AddTag(FAssetRegistryTag("Immortal", IsImmortal() ? TEXT("True") : TEXT("False"), FAssetRegistryTag::TT_Alphabetical));
+	Context.AddTag(FAssetRegistryTag("Becomes Zombie", WillBecomeZombie() ? TEXT("True") : TEXT("False"), FAssetRegistryTag::TT_Alphabetical));
+	Context.AddTag(FAssetRegistryTag("CanBeOccluded", OcclusionBoundsMethod == EParticleSystemOcclusionBoundsMethod::EPSOBM_None ? TEXT("False") : TEXT("True"), FAssetRegistryTag::TT_Alphabetical));
+
+	uint32 NumEmittersAtEachSig[(int32)EParticleSignificanceLevel::Num] = { 0, 0, 0, 0 };
+	for (UParticleEmitter* Emitter : Emitters)
+	{
+		if (Emitter)
+		{
+			++NumEmittersAtEachSig[(int32)Emitter->SignificanceLevel];			
+		}
+	}
+	Context.AddTag(FAssetRegistryTag("Critical Emitters", LexToString(NumEmittersAtEachSig[(int32)EParticleSignificanceLevel::Critical]), FAssetRegistryTag::TT_Numerical));
+	Context.AddTag(FAssetRegistryTag("High Emitters", LexToString(NumEmittersAtEachSig[(int32)EParticleSignificanceLevel::High]), FAssetRegistryTag::TT_Numerical));
+	Context.AddTag(FAssetRegistryTag("Medium Emitters", LexToString(NumEmittersAtEachSig[(int32)EParticleSignificanceLevel::Medium]), FAssetRegistryTag::TT_Numerical));
+	Context.AddTag(FAssetRegistryTag("Low Emitters", LexToString(NumEmittersAtEachSig[(int32)EParticleSignificanceLevel::Low]), FAssetRegistryTag::TT_Numerical));
+
+	Super::GetAssetRegistryTags(Context);
+}
+
+bool UParticleSystem::UsesCPUCollision() const
+{
+	for (const UParticleEmitter* Emitter : Emitters)
+	{
+		if (Emitter != nullptr)
+		{
+			// If we have not yet found a CPU collision module, and we have some enabled LODs to look in..
+			if (Emitter->HasAnyEnabledLODs() && Emitter->LODLevels.Num() > 0)
+			{
+				if (const UParticleLODLevel* HighLODLevel = Emitter->LODLevels[0])
+				{
+					// Iterate over modules of highest LOD (will have all the modules)
+					for (const UParticleModule* Module : HighLODLevel->Modules)
+					{
+						// If an enabled CPU collision module 
+						if (Module->bEnabled && Module->IsA<UParticleModuleCollision>())
+						{
+							return true;
+						}
+					}
+				}
+			}
+		}
+	}
+
+	return false;
+}
+
+bool UParticleSystem::CanBeClusterRoot() const
+{
+	return true;
+}
+
+bool UParticleSystem::CanBePooled()const
+{
+	if (MaxPoolSize == 0)
+	{
+		return false;
+	}
+
+	return true;
+}
+
+bool UParticleSystem::CalculateMaxActiveParticleCounts()
+{
+	bool bSuccess = true;
+
+	for (int32 EmitterIndex = 0; EmitterIndex < Emitters.Num(); EmitterIndex++)
+	{
+		UParticleEmitter* Emitter = Emitters[EmitterIndex];
+		if (Emitter)
+		{
+			if (Emitter->CalculateMaxActiveParticleCount() == false)
+			{
+				bSuccess = false;
+			}
+		}
+	}
+
+	return bSuccess;
+}
+
+
+void UParticleSystem::GetParametersUtilized(TArray<TArray<FString> >& ParticleSysParamList,
+											TArray<TArray<FString> >& ParticleParameterList)
+{
+	ParticleSysParamList.Empty();
+	ParticleParameterList.Empty();
+
+	for (int32 EmitterIndex = 0; EmitterIndex < Emitters.Num(); EmitterIndex++)
+	{
+		int32 CheckIndex;
+		CheckIndex = ParticleSysParamList.AddZeroed();
+		check(CheckIndex == EmitterIndex);
+		CheckIndex = ParticleParameterList.AddZeroed();
+		check(CheckIndex == EmitterIndex);
+
+		UParticleEmitter* Emitter = Emitters[EmitterIndex];
+		if (Emitter)
+		{
+			Emitter->GetParametersUtilized(
+				ParticleSysParamList[EmitterIndex],
+				ParticleParameterList[EmitterIndex]);
+		}
+	}
+}
+
+
+void UParticleSystem::SetupSoloing()
+{
+#if WITH_EDITOR
+	if (GIsEditor)
+	{
+		if (Emitters.Num())
+		{
+			// Store the settings of bEnabled for each LODLevel in each emitter
+			SoloTracking.Empty(Emitters.Num());
+			SoloTracking.AddZeroed(Emitters.Num());
+			for (int32 EmitterIdx = 0; EmitterIdx < Emitters.Num(); EmitterIdx++)
+			{
+				UParticleEmitter* Emitter = Emitters[EmitterIdx];
+				if (Emitter != nullptr)
+				{
+					FLODSoloTrack& SoloTrack = SoloTracking[EmitterIdx];
+					SoloTrack.SoloEnableSetting.Empty(Emitter->LODLevels.Num());
+					SoloTrack.SoloEnableSetting.AddZeroed(Emitter->LODLevels.Num());
+				}
+			}
+
+			for (int32 EmitterIdx = 0; EmitterIdx < Emitters.Num(); EmitterIdx++)
+			{
+				UParticleEmitter* Emitter = Emitters[EmitterIdx];
+				if (Emitter != NULL)
+				{
+					FLODSoloTrack& SoloTrack = SoloTracking[EmitterIdx];
+					int32 MaxLOD = FMath::Min(SoloTrack.SoloEnableSetting.Num(), Emitter->LODLevels.Num());
+					for (int32 LODIdx = 0; LODIdx < MaxLOD; LODIdx++)
+					{
+						UParticleLODLevel* LODLevel = Emitter->LODLevels[LODIdx];
+						check(LODLevel);
+						SoloTrack.SoloEnableSetting[LODIdx] = LODLevel->bEnabled;
+					}
+				}
+			}
+		}
+	}
+#endif
+}
+
+
+bool UParticleSystem::ToggleSoloing(class UParticleEmitter* InEmitter)
+{
+	bool bSoloingReturn = false;
+	if (InEmitter != NULL)
+	{
+		bool bOtherEmitterIsSoloing = false;
+		// Set the given one
+		int32 SelectedIndex = -1;
+		for (int32 EmitterIdx = 0; EmitterIdx < Emitters.Num(); EmitterIdx++)
+		{
+			UParticleEmitter* Emitter = Emitters[EmitterIdx];
+			check(Emitter != NULL);
+			if (Emitter == InEmitter)
+			{
+				SelectedIndex = EmitterIdx;
+			}
+			else
+			{
+				if (Emitter->bIsSoloing == true)
+				{
+					bOtherEmitterIsSoloing = true;
+					bSoloingReturn = true;
+				}
+			}
+		}
+
+		if (SelectedIndex != -1)
+		{
+			InEmitter->bIsSoloing = !InEmitter->bIsSoloing;
+			for (int32 EmitterIdx = 0; EmitterIdx < Emitters.Num(); EmitterIdx++)
+			{
+				UParticleEmitter* Emitter = Emitters[EmitterIdx];
+				FLODSoloTrack& SoloTrack = SoloTracking[EmitterIdx];
+				if (EmitterIdx == SelectedIndex)
+				{
+					for (int32 LODIdx = 0; LODIdx < InEmitter->LODLevels.Num(); LODIdx++)
+					{
+						UParticleLODLevel* LODLevel = InEmitter->LODLevels[LODIdx];
+						if (InEmitter->bIsSoloing == false)
+						{
+							if (bOtherEmitterIsSoloing == false)
+							{
+								// Restore the enabled settings - ie turn off soloing...
+								LODLevel->bEnabled = SoloTrack.SoloEnableSetting[LODIdx];
+							}
+							else
+							{
+								// Disable the emitter
+								LODLevel->bEnabled = false;
+							}
+						}
+						else 
+						if (bOtherEmitterIsSoloing == true)
+						{
+							// Need to restore old settings of this emitter as it is now soloing
+							LODLevel->bEnabled = SoloTrack.SoloEnableSetting[LODIdx];
+						}
+					}
+				}
+				else
+				{
+					// Restore all other emitters if this disables soloing...
+					if ((InEmitter->bIsSoloing == false) && (bOtherEmitterIsSoloing == false))
+					{
+						for (int32 LODIdx = 0; LODIdx < Emitter->LODLevels.Num(); LODIdx++)
+						{
+							UParticleLODLevel* LODLevel = Emitter->LODLevels[LODIdx];
+							// Restore the enabled settings - ie turn off soloing...
+							LODLevel->bEnabled = SoloTrack.SoloEnableSetting[LODIdx];
+						}
+					}
+					else
+					{
+						if (Emitter->bIsSoloing == false)
+						{
+							for (int32 LODIdx = 0; LODIdx < Emitter->LODLevels.Num(); LODIdx++)
+							{
+								UParticleLODLevel* LODLevel = Emitter->LODLevels[LODIdx];
+								// Disable the emitter
+								LODLevel->bEnabled = false;
+							}
+						}
+					}
+				}
+			}
+		}
+
+		// We checked the other emitters above...
+		// Make sure we catch the case of the first one toggled to true!
+		if (InEmitter->bIsSoloing == true)
+		{
+			bSoloingReturn = true;
+		}
+	}
+
+	return bSoloingReturn;
+}
+
+
+bool UParticleSystem::TurnOffSoloing()
+{
+	for (int32 EmitterIdx = 0; EmitterIdx < Emitters.Num(); EmitterIdx++)
+	{
+		UParticleEmitter* Emitter = Emitters[EmitterIdx];
+		if (Emitter != NULL)
+		{
+			FLODSoloTrack& SoloTrack = SoloTracking[EmitterIdx];
+			for (int32 LODIdx = 0; LODIdx < Emitter->LODLevels.Num(); LODIdx++)
+			{
+				UParticleLODLevel* LODLevel = Emitter->LODLevels[LODIdx];
+				if (LODLevel != NULL)
+				{
+					// Restore the enabled settings - ie turn off soloing...
+					LODLevel->bEnabled = SoloTrack.SoloEnableSetting[LODIdx];
+				}
+			}
+			Emitter->bIsSoloing = false;
+		}
+	}
+
+	return true;
+}
+
+
+void UParticleSystem::SetupLODValidity()
+{
+	for (int32 EmitterIdx = 0; EmitterIdx < Emitters.Num(); EmitterIdx++)
+	{
+		UParticleEmitter* Emitter = Emitters[EmitterIdx];
+		if (Emitter != NULL)
+		{
+			for (int32 Pass = 0; Pass < 2; Pass++)
+			{
+				for (int32 LODIdx = 0; LODIdx < Emitter->LODLevels.Num(); LODIdx++)
+				{
+					UParticleLODLevel* LODLevel = Emitter->LODLevels[LODIdx];
+					if (LODLevel != NULL)
+					{
+						for (int32 ModuleIdx = -3; ModuleIdx < LODLevel->Modules.Num(); ModuleIdx++)
+						{
+							int32 ModuleFetchIdx;
+							switch (ModuleIdx)
+							{
+							case -3:	ModuleFetchIdx = INDEX_REQUIREDMODULE;	break;
+							case -2:	ModuleFetchIdx = INDEX_SPAWNMODULE;		break;
+							case -1:	ModuleFetchIdx = INDEX_TYPEDATAMODULE;	break;
+							default:	ModuleFetchIdx = ModuleIdx;				break;
+							}
+
+							UParticleModule* Module = LODLevel->GetModuleAtIndex(ModuleFetchIdx);
+							if (Module != NULL)
+							{
+								// On pass 1, clear the LODValidity flags
+								// On pass 2, set it
+								if (Pass == 0)
+								{
+									Module->LODValidity = 0;
+								}
+								else
+								{
+									Module->LODValidity |= (1 << LODIdx);
+								}
+							}
+						}
+					}
+				}
+			}
+		}
+	}
+}
+
+void UParticleSystem::SetDelay(float InDelay)
+{
+	Delay = InDelay;
+}
+
+#if WITH_EDITOR
+
+bool UParticleSystem::RemoveAllDuplicateModules(bool bInMarkForCooker, TMap<UObject*,bool>* OutRemovedModules)
+{
+	// Generate a map of module classes used to instances of those modules 
+	TMap<UClass*,TMap<UParticleModule*,int32> > ClassToModulesMap;
+	for (int32 EmitterIdx = 0; EmitterIdx < Emitters.Num(); EmitterIdx++)
+	{
+		UParticleEmitter* Emitter = Emitters[EmitterIdx];
+		if (Emitter != NULL)
+		{
+			if (Emitter->bCookedOut == false)
+			{
+				for (int32 LODIdx = 0; LODIdx < Emitter->LODLevels.Num(); LODIdx++)
+				{
+					UParticleLODLevel* LODLevel = Emitter->LODLevels[LODIdx];
+					if (LODLevel != NULL)
+					{
+						for (int32 ModuleIdx = -1; ModuleIdx < LODLevel->Modules.Num(); ModuleIdx++)
+						{
+							UParticleModule* Module = NULL;
+							if (ModuleIdx == -1)
+							{
+								Module = LODLevel->SpawnModule;
+							}
+							else
+							{
+								Module = LODLevel->Modules[ModuleIdx];
+							}
+							if (Module != NULL)
+							{
+								TMap<UParticleModule*,int32>* ModuleList = ClassToModulesMap.Find(Module->GetClass());
+								if (ModuleList == NULL)
+								{
+									TMap<UParticleModule*,int32> TempModuleList;
+									ClassToModulesMap.Add(Module->GetClass(), TempModuleList);
+									ModuleList = ClassToModulesMap.Find(Module->GetClass());
+								}
+								check(ModuleList);
+								int32* ModuleCount = ModuleList->Find(Module);
+								if (ModuleCount == NULL)
+								{
+									int32 TempModuleCount = 0;
+									ModuleList->Add(Module, TempModuleCount);
+									ModuleCount = ModuleList->Find(Module);
+								}
+								check(ModuleCount);
+								(*ModuleCount)++;
+							}
+						}
+					}
+				}
+			}
+		}
+	}
+
+	// Now we have a list of module classes and the modules they contain...
+	// Find modules of the same class that have the exact same settings.
+	TMap<UParticleModule*, TArray<UParticleModule*> > DuplicateModules;
+	TMap<UParticleModule*,bool> FoundAsADupeModules;
+	TMap<UParticleModule*, UParticleModule*> ReplaceModuleMap;
+	bool bRemoveDuplicates = true;
+	for (TMap<UClass*,TMap<UParticleModule*,int32> >::TIterator ModClassIt(ClassToModulesMap); ModClassIt; ++ModClassIt)
+	{
+		UClass* ModuleClass = ModClassIt.Key();
+		TMap<UParticleModule*,int32>& ModuleMap = ModClassIt.Value();
+		if (ModuleMap.Num() > 1)
+		{
+			// There is more than one of this module, so see if there are dupes...
+			TArray<UParticleModule*> ModuleArray;
+			for (TMap<UParticleModule*,int32>::TIterator ModuleIt(ModuleMap); ModuleIt; ++ModuleIt)
+			{
+				ModuleArray.Add(ModuleIt.Key());
+			}
+
+			// For each module, see if it it a duplicate of another
+			for (int32 ModuleIdx = 0; ModuleIdx < ModuleArray.Num(); ModuleIdx++)
+			{
+				UParticleModule* SourceModule = ModuleArray[ModuleIdx];
+				if (FoundAsADupeModules.Find(SourceModule) == NULL)
+				{
+					for (int32 InnerModuleIdx = ModuleIdx + 1; InnerModuleIdx < ModuleArray.Num(); InnerModuleIdx++)
+					{
+						UParticleModule* CheckModule = ModuleArray[InnerModuleIdx];
+						if (FoundAsADupeModules.Find(CheckModule) == NULL)
+						{
+							bool bIsDifferent = false;
+							static const FName CascadeCategory(TEXT("Cascade"));
+							// Copy non component properties from the old actor to the new actor
+							for (FProperty* Property = ModuleClass->PropertyLink; Property != NULL; Property = Property->PropertyLinkNext)
+							{
+								bool bIsTransient = (Property->PropertyFlags & CPF_Transient) != 0;
+								bool bIsEditorOnly = (Property->PropertyFlags & CPF_EditorOnly) != 0;
+								bool bIsCascade = (FObjectEditorUtils::GetCategoryFName(Property) == CascadeCategory);
+								// Ignore 'Cascade' category, transient, native and EditorOnly properties...
+								if (!bIsTransient && !bIsEditorOnly && !bIsCascade)
+								{
+									for( int32 iProp=0; iProp<Property->ArrayDim; iProp++ )
+									{
+										bool bIsIdentical = Property->Identical_InContainer(SourceModule, CheckModule, iProp, PPF_DeepComparison);
+										if (bIsIdentical == false)
+										{
+											bIsDifferent = true;
+											break;
+										}
+									}
+								}
+							}
+
+							if (bIsDifferent == false)
+							{
+								TArray<UParticleModule*>* DupedModules = DuplicateModules.Find(SourceModule);
+								if (DupedModules == NULL)
+								{
+									TArray<UParticleModule*> TempDupedModules;
+									DuplicateModules.Add(SourceModule, TempDupedModules);
+									DupedModules = DuplicateModules.Find(SourceModule);
+								}
+								check(DupedModules);
+								if (ReplaceModuleMap.Find(CheckModule) == NULL)
+								{
+									ReplaceModuleMap.Add(CheckModule, SourceModule);
+								}
+								else
+								{
+									UE_LOG(LogParticles, Error, TEXT("Module already in replacement map - ABORTING CONVERSION!!!!"));
+									bRemoveDuplicates = false;
+								}
+								DupedModules->AddUnique(CheckModule);
+								FoundAsADupeModules.Add(CheckModule, true);
+							}
+						}
+					}
+				}
+			}
+		}
+	}
+
+	// If not errors were found, and there are duplicates, remove them...
+	if (bRemoveDuplicates && (ReplaceModuleMap.Num() > 0))
+	{
+		TArray<UParticleModule*> RemovedModules;
+		for (int32 EmitterIdx = 0; EmitterIdx < Emitters.Num(); EmitterIdx++)
+		{
+			UParticleEmitter* Emitter = Emitters[EmitterIdx];
+			if (Emitter != NULL)
+			{
+				if (Emitter->bCookedOut == false)
+				{
+					for (int32 LODIdx = 0; LODIdx < Emitter->LODLevels.Num(); LODIdx++)
+					{
+						UParticleLODLevel* LODLevel = Emitter->LODLevels[LODIdx];
+						if (LODLevel != NULL)
+						{
+							for (int32 ModuleIdx = -1; ModuleIdx < LODLevel->Modules.Num(); ModuleIdx++)
+							{
+								UParticleModule* Module = NULL;
+								if (ModuleIdx == -1)
+								{
+									Module = LODLevel->SpawnModule;
+								}
+								else
+								{
+									Module = LODLevel->Modules[ModuleIdx];
+								}
+								if (Module != NULL)
+								{
+									UParticleModule** ReplacementModule = ReplaceModuleMap.Find(Module);
+									if (ReplacementModule != NULL)
+									{
+										UParticleModule* ReplaceMod = *ReplacementModule;
+										if (ModuleIdx == -1)
+										{
+											LODLevel->SpawnModule = CastChecked<UParticleModuleSpawn>(ReplaceMod);
+										}
+										else
+										{
+											LODLevel->Modules[ModuleIdx] = ReplaceMod;
+										}
+
+										if (bInMarkForCooker == true)
+										{
+											RemovedModules.AddUnique(Module);
+										}
+									}
+								}
+							}
+						}
+					}
+				}
+			}
+		}
+
+		if (bInMarkForCooker == true)
+		{
+			for (int32 MarkIdx = 0; MarkIdx < RemovedModules.Num(); MarkIdx++)
+			{
+				UParticleModule* RemovedModule = RemovedModules[MarkIdx];
+				RemovedModule->SetFlags(RF_Transient);
+				if (OutRemovedModules != NULL)
+				{
+					OutRemovedModules->Add(RemovedModule, true);
+				}
+			}
+		}
+
+		// Update the list of modules in each emitter
+		UpdateAllModuleLists();
+	}
+
+	return true;
+}
+
+
+void UParticleSystem::UpdateAllModuleLists()
+{
+	for (int32 EmitterIdx = 0; EmitterIdx < Emitters.Num(); EmitterIdx++)
+	{
+		UParticleEmitter* Emitter = Emitters[EmitterIdx];
+		if (Emitter != NULL)
+		{
+			for (int32 LODIdx = 0; LODIdx < Emitter->LODLevels.Num(); LODIdx++)
+			{
+				UParticleLODLevel* LODLevel = Emitter->LODLevels[LODIdx];
+				if (LODLevel != NULL)
+				{
+					LODLevel->UpdateModuleLists();
+				}
+			}
+
+			// Allow type data module to cache any module info
+			if(Emitter->LODLevels.Num() > 0)
+			{
+				UParticleLODLevel* HighLODLevel = Emitter->LODLevels[0];
+				if (HighLODLevel != nullptr && HighLODLevel->TypeDataModule != nullptr)
+				{
+					// Allow TypeData module to cache pointers to modules
+					HighLODLevel->TypeDataModule->CacheModuleInfo(Emitter);
+				}
+			}
+
+			// Update any cached info from modules on the emitter
+			Emitter->CacheEmitterModuleInfo();
+		}
+	}
+}
+#endif
+
+
+void UParticleSystem::BuildEmitters()
+{
+	const int32 EmitterCount = Emitters.Num();
+	for ( int32 EmitterIndex = 0; EmitterIndex < EmitterCount; ++EmitterIndex )
+	{
+		if (UParticleEmitter* Emitter = Emitters[EmitterIndex])
+		{
+			Emitter->Build();
+		}
+	}
+}
+
+static bool LogReasoningForAnyThreadTicking()
+{
+	static bool bLogThreadedParticleTicking = FParse::Param(FCommandLine::Get(), TEXT("LogThreadedParticleTicking"));
+	return bLogThreadedParticleTicking;
+}
+
+void UParticleSystem::ComputeCanTickInAnyThread()
+{
+	check(!bIsElligibleForAsyncTickComputed);
+	bIsElligibleForAsyncTickComputed = true;
+
+	bIsElligibleForAsyncTick = true; // assume everything is async
+	int32 EmitterIndex;
+	for (EmitterIndex = 0; EmitterIndex < Emitters.Num(); EmitterIndex++)
+	{
+		UParticleEmitter* Emitter = Emitters[EmitterIndex];
+		if (Emitter)
+		{
+			for (int32 LevelIndex = 0; LevelIndex < Emitter->LODLevels.Num(); LevelIndex++)
+			{
+				UParticleLODLevel* LODLevel	= Emitter->LODLevels[LevelIndex];
+				if (LODLevel)
+				{
+					for (int32 ModuleIndex = 0; ModuleIndex < LODLevel->Modules.Num(); ModuleIndex++)
+					{
+						UParticleModule* Module	= LODLevel->Modules[ModuleIndex];
+						if (Module && !Module->CanTickInAnyThread())
+						{
+							bIsElligibleForAsyncTick = false;
+							if (LogReasoningForAnyThreadTicking())
+							{
+								UE_LOG(LogParticles, Display, TEXT("Cannot tick %s in parallel because module %s in Emitter %s cannot tick in in parallel."), *GetFullName(), *Module->GetFullName(), *Emitter->GetFullName());
+							}
+							else
+							{
+								return;
+							}
+						}
+					}
+				}
+			}
+
+		}
+	}
+	if (LogReasoningForAnyThreadTicking() && bIsElligibleForAsyncTick)
+	{
+		UE_LOG(LogParticles, Display, TEXT("Can tick %s in parallel."), *GetFullName());
+	}
+}
+
+bool UParticleSystem::ContainsEmitterType(UClass* TypeData)
+{
+	for ( int32 EmitterIndex = 0; EmitterIndex < Emitters.Num(); ++EmitterIndex)
+	{
+		UParticleEmitter* Emitter = Emitters[EmitterIndex];
+		if (Emitter)
+		{
+			UParticleLODLevel* LODLevel = Emitter->LODLevels[0];
+			if (LODLevel && LODLevel->TypeDataModule && LODLevel->TypeDataModule->IsA(TypeData))
+			{
+				return true;
+			}
+		}
+	}
+	
+	return false;
+}
+
+bool UParticleSystem::HasGPUEmitter() const
+{
+	for (int32 EmitterIndex = 0; EmitterIndex < Emitters.Num(); ++EmitterIndex)
+	{
+		if (Emitters[EmitterIndex] == nullptr)
+		{
+			continue;
+		}
+		// We can just check for the GPU type data at the highest LOD.
+		UParticleLODLevel* LODLevel = Emitters[EmitterIndex]->LODLevels[0];
+		if( LODLevel )
+		{
+			UParticleModule* TypeDataModule = LODLevel->TypeDataModule;
+			if( TypeDataModule && TypeDataModule->IsA(UParticleModuleTypeDataGpu::StaticClass()) )
+			{
+				return true;
+			}
+		}
+	}
+	return false;
+}
diff --git a/Engine/Source/Runtime/Engine/Public/ParticleEmitterInstanceOwner.h b/Engine/Source/Runtime/Engine/Public/ParticleEmitterInstanceOwner.h
--- a/Engine/Source/Runtime/Engine/Public/ParticleEmitterInstanceOwner.h
+++ b/Engine/Source/Runtime/Engine/Public/ParticleEmitterInstanceOwner.h
@@ -65,7 +65,7 @@ public:
 	virtual bool GetJustRegistered() const = 0;
 	virtual float GetWarmupTime() const = 0;
 	virtual float GetEmitterDelay() const = 0;
-	virtual FRandomStream& GetRandomStream() const = 0;
+	virtual FRandomStream& GetRandomStream() = 0;
 
 	virtual void SetComponentToWorld(const FTransform& NewComponentToWorld) = 0;
 	virtual void DeactivateNextTick() = 0;
