Fix seamless travel crashes by explicitly clearing the network object list during the transition to make sure nothing from the old world is still there

diff --git a/Engine/Source/Runtime/Engine/Private/NetDriver.cpp b/Engine/Source/Runtime/Engine/Private/NetDriver.cpp
--- a/Engine/Source/Runtime/Engine/Private/NetDriver.cpp
+++ b/Engine/Source/Runtime/Engine/Private/NetDriver.cpp
@@ -6941,6 +6941,9 @@ void UNetDriver::CleanupWorldForSeamlessTravel()
 void UNetDriver::PreSeamlessTravelGarbageCollect()
 {
 	ResetGameWorldState();
+
+	// Clear out the network objects list, the traveling actors are in SeamlessTravelingObjects
+	NetworkObjects->Reset();
 }
 
 void UNetDriver::PostSeamlessTravelGarbageCollect()
