UOLs: Use the FTemporaryPlayInEditorIDOverride in the ActorLocatorFragment

We never want to resolve to objects with a different PIE instance ID, even if the current callstack is being executed inside a different GPlayInEditorID scope. Since ResolveObject will always call FixupForPIE in editor based on GPlayInEditorID, we always override the current GPlayInEditorID to be the current PIE instance of the provided context.

Tested MRQ, Auto Bind to PIE/SIE toggles, switching between bound contexts, two standalone PIE instances all work as expected

diff --git a/Engine/Source/Runtime/Engine/Private/UniversalObjectLocators/ActorLocatorFragment.cpp b/Engine/Source/Runtime/Engine/Private/UniversalObjectLocators/ActorLocatorFragment.cpp
--- a/Engine/Source/Runtime/Engine/Private/UniversalObjectLocators/ActorLocatorFragment.cpp
+++ b/Engine/Source/Runtime/Engine/Private/UniversalObjectLocators/ActorLocatorFragment.cpp
@@ -12,6 +12,7 @@
 #include "LevelUtils.h"
 #include "Engine/Level.h"
 #include "Engine/LevelStreaming.h"
+#include "UnrealEngine.h"
 #include "Engine/World.h"
 #include "WorldPartition/WorldPartitionLevelHelper.h"
 #include "Misc/EditorPathHelper.h"
@@ -119,6 +120,12 @@ UE::UniversalObjectLocator::FResolveResult FActorLocatorFragment::Resolve(const
 	const UPackage* ContextPackage = Params.Context ? Params.Context->GetOutermost() : nullptr;
 	const int32     PIEInstanceID  = ContextPackage ? ContextPackage->GetPIEInstanceID() : INDEX_NONE;
 
+	// The Actor Fragment is explicit about providing a resolution context for its bindings. We never want to resolve to objects
+	// with a different PIE instance ID, even if the current callstack is being executed inside a different GPlayInEditorID
+	// scope. Since ResolveObject will always call FixupForPIE in editor based on GPlayInEditorID, we always override the current
+	// GPlayInEditorID to be the current PIE instance of the provided context.
+	FTemporaryPlayInEditorIDOverride PIEGuard(PIEInstanceID);
+
 	// Finally fallback to just trying to resolve the path directly
 	auto ResolvePathWithPIEHandling = [PIEInstanceID](const FSoftObjectPath& PathPtr)
 	{
