Fix for uninitialized variable in FActiveSound

#rnx

diff --git a/Engine/Source/Runtime/Engine/Private/ActiveSound.cpp b/Engine/Source/Runtime/Engine/Private/ActiveSound.cpp
--- a/Engine/Source/Runtime/Engine/Private/ActiveSound.cpp
+++ b/Engine/Source/Runtime/Engine/Private/ActiveSound.cpp
@@ -822,6 +822,17 @@ void FActiveSound::UpdateInterfaceParameters(const TArray<FListener>& InListener
 	InstanceTransmitter->SetParameters(MoveTemp(ParamsToUpdate));
 }
 
+int32 FActiveSound::GetClosestListenerIndex() const
+{
+	// If we haven't cached the cloest listener index, just directly query it from the audio device.
+	// This is cached in UpdateWaveInstances.
+	if (ClosestListenerIndex == INDEX_NONE && AudioDevice)
+	{
+		return AudioDevice->FindClosestListenerIndex(Transform);
+	}
+	return ClosestListenerIndex;
+}
+
 void FActiveSound::UpdateWaveInstances(TArray<FWaveInstance*> &InWaveInstances, const float DeltaTime)
 {
 	// Reset whether or not the active sound is playing audio.
diff --git a/Engine/Source/Runtime/Engine/Public/ActiveSound.h b/Engine/Source/Runtime/Engine/Public/ActiveSound.h
--- a/Engine/Source/Runtime/Engine/Public/ActiveSound.h
+++ b/Engine/Source/Runtime/Engine/Public/ActiveSound.h
@@ -366,7 +366,7 @@ public:
 		bShouldSourceBufferListenerZeroBuffer = bShouldZeroBuffer;
 	}
 
-	int32 GetClosestListenerIndex() const { return ClosestListenerIndex; }
+	ENGINE_API int32 GetClosestListenerIndex() const;
 
 	/** Returns whether or not the active sound can be deleted. */
 	bool CanDelete() const { return !bAsyncOcclusionPending; }
@@ -803,7 +803,7 @@ private:
 	static ENGINE_API FTraceDelegate ActiveSoundTraceDelegate;
 
 	/** Cached index to the closest listener. So we don't have to do the work to find it twice. */
-	int32 ClosestListenerIndex;
+	int32 ClosestListenerIndex = INDEX_NONE;
 
 	/** This is a friend so the audio device can call Stop() on the active sound. */
 	friend class FAudioDevice;
