New material translator push:
 - After transforming BinaryOperator to simply Operator, this change also adds support for unary operators like Negate and Not.
- Implemented OneMinus expression.

#rnx

diff --git a/Engine/Source/Runtime/Engine/Private/Materials/MaterialExpressionsIR.cpp b/Engine/Source/Runtime/Engine/Private/Materials/MaterialExpressionsIR.cpp
--- a/Engine/Source/Runtime/Engine/Private/Materials/MaterialExpressionsIR.cpp
+++ b/Engine/Source/Runtime/Engine/Private/Materials/MaterialExpressionsIR.cpp
@@ -456,7 +456,29 @@ void UMaterialExpressionAppendVector::Build(MIR::FEmitter& Emitter)
 	Emitter.Output(0, Output);
 }
 
-/* Mathematical Operations */
+/* Unary Operators */
+
+static void BuildUnaryArithmeticOperator(
+	MIR::FEmitter& Emitter,
+	MIR::EOperator Op,
+	FExpressionInput* Input,
+	float Const)
+{
+	// Default input to its relative constants if disconnected, then get it after checking it has primitive type.
+	MIR::FValue* Value = Emitter.DefaultInputToFloatScalar(Input, Const).Input(Input);
+
+	UE_MIR_CHECKPOINT();
+
+	// Finally emit the unary operator.
+	Value = Emitter.EmitOperator(Op, Value);
+
+	// And flow it out of the expression's only output.
+	Emitter.Output(0, Value);
+}
+
+// TODO: Add the missing expressions 
+
+/* Binary Operators */
 
 static void BuildBinaryArithmeticOperator(
 	MIR::FEmitter& Emitter,
@@ -510,6 +532,25 @@ void UMaterialExpressionDivide::Build(MIR::FEmitter& Emitter)
 	BuildBinaryArithmeticOperator(Emitter, MIR::BO_Divide, &A, ConstA, &B, ConstB);
 }
 
+void UMaterialExpressionOneMinus::Build(MIR::FEmitter& Emitter)
+{
+	// Default input to zero if not connected, then get it as a primitive.
+	MIR::FValue* Value = Emitter.DefaultInputToFloatScalar(&Input, 0.0f).InputPrimitive(&Input);
+	
+	UE_MIR_CHECKPOINT();
+
+	// Make a "One" value of the same type and dimension as input's.
+	MIR::FValue* OneValue = Emitter.EmitConstruct(Value->Type, Emitter.EmitConstantFloatScalar(1.0f));
+
+	UE_MIR_CHECKPOINT();
+
+	// Finally emit the binary operator.
+	MIR::FValue* Output = Emitter.EmitOperator(MIR::BO_Subtract, OneValue, Value);
+
+	// And flow it out of the expression's only output.
+	Emitter.Output(0, Output);
+}
+
 void UMaterialExpressionIfThenElse::Build(MIR::FEmitter& Emitter)
 {
 	// If Condition isn't connected, default its value to false
diff --git a/Engine/Source/Runtime/Engine/Private/Materials/MaterialIR.cpp b/Engine/Source/Runtime/Engine/Private/Materials/MaterialIR.cpp
--- a/Engine/Source/Runtime/Engine/Private/Materials/MaterialIR.cpp
+++ b/Engine/Source/Runtime/Engine/Private/Materials/MaterialIR.cpp
@@ -374,24 +374,29 @@ FBlock* FInstruction::GetDesiredBlockForUse(int32 UseIndex)
 	return Block;
 }
 
-bool OperatorIsComparison(EOperator Op)
+bool IsComparisonOperator(EOperator Op)
 {
 	return Op >= BO_GreaterThan && Op <= BO_NotEquals;
 }
 
-bool OperatorIsArithmetic(EOperator Op)
+bool IsBitOperator(EOperator Op)
 {
-	return Op >= BO_Add;
+	return Op == UO_Not || (Op >= BO_Or && Op <= BO_Xor);
 }
 
-bool OperatorIsTernary(EOperator Op)
+bool IsBinaryOperator(EOperator Op)
+{
+	return Op >= BO_GreaterThan && Op <= BO_Max;
+}
+
+bool IsTernaryOperator(EOperator Op)
 {
 	return Op >= TO_Lerp;
 }
 
-bool OperatorIsBinaryAlgebraic(EOperator Op)
+bool IsBinaryInfixOperator(EOperator Op)
 {
-	return Op >= BO_Add && Op <= BO_Divide;
+	return Op >= BO_Add && Op <= BO_Xor;
 }
 
 const TCHAR* OperatorToString(EOperator Op)
diff --git a/Engine/Source/Runtime/Engine/Private/Materials/MaterialIREmitter.cpp b/Engine/Source/Runtime/Engine/Private/Materials/MaterialIREmitter.cpp
--- a/Engine/Source/Runtime/Engine/Private/Materials/MaterialIREmitter.cpp
+++ b/Engine/Source/Runtime/Engine/Private/Materials/MaterialIREmitter.cpp
@@ -656,6 +656,7 @@ static bool FoldComparisonOperatorScalar(EOperator Operator, T A, T B)
 {
 	switch (Operator)
 	{
+		case UO_Not: return !A;
 		case BO_GreaterThan: return A > B;
 		case BO_GreaterThanOrEquals: return A >= B;
 		case BO_LowerThan: return A < B;
@@ -671,6 +672,7 @@ static T FoldScalarArithmeticOperator(EOperator Operator, T A, T B, T C)
 {
 	switch (Operator)
 	{
+		case UO_Negate: return -A;
 		case BO_Add: return A + B;
 		case BO_Subtract: return A - B;
 		case BO_Multiply: return A * B;
@@ -683,6 +685,18 @@ static T FoldScalarArithmeticOperator(EOperator Operator, T A, T B, T C)
 	}
 }
 
+static TInteger FoldScalarBitOperator(EOperator Operator, TInteger A, TInteger B)
+{
+	switch (Operator)
+	{
+		case UO_Not: return !A;
+		case BO_And: return A & B;
+		case BO_Or: return A | B;
+		case BO_Xor: return A ^ B;
+		default: UE_MIR_UNREACHABLE();
+	}
+}
+
 static FValue* TryFoldOperatorScalar(FEmitter* Emitter, EOperator Op, FValue* A, FValue* B, FValue* C)
 {
 	FPrimitiveTypePtr PrimitiveType = A->Type->AsPrimitive();
@@ -691,7 +705,7 @@ static FValue* TryFoldOperatorScalar(FEmitter* Emitter, EOperator Op, FValue* A,
 	// These simplifications don't requre that both lhs and rhs are constant.
 	switch (Op)
 	{
-		/* Comparisons */
+		/* Binary Comparisons */
 		case BO_GreaterThan:
 		case BO_LowerThan:
 		case BO_NotEquals:
@@ -755,10 +769,15 @@ static FValue* TryFoldOperatorScalar(FEmitter* Emitter, EOperator Op, FValue* A,
 		case TO_Lerp:
 			break; // todo
 
+		case TO_Smoothstep:
+			break; // todo
+
+		/* No simplifications for these operations */
+		case UO_Negate:
+		case UO_Not:
 		case BO_Min:
 		case BO_Max:
-		case TO_Smoothstep:
-			break; // no simplifications for these operations
+			break;
 
 		default:
 			UE_MIR_UNREACHABLE();
@@ -768,12 +787,13 @@ static FValue* TryFoldOperatorScalar(FEmitter* Emitter, EOperator Op, FValue* A,
 	const FConstant* AConstant = A->As<FConstant>();
 	const FConstant* BConstant = B->As<FConstant>();
 	const FConstant* CConstant = C->As<FConstant>();
-	if (!AConstant || !BConstant || (OperatorIsTernary(Op) && !CConstant))
+	if (!AConstant ||  (IsBinaryOperator(Op) && !BConstant) || (IsTernaryOperator(Op) && !CConstant))
 	{
 		return nullptr;
 	}
 
-	if (OperatorIsComparison(Op))
+	// Call the appropriate helper function depending on what type of operator this is
+	if (IsComparisonOperator(Op))
 	{
 		bool Result;
 		switch (PrimitiveType->ScalarKind)
@@ -791,20 +811,39 @@ static FValue* TryFoldOperatorScalar(FEmitter* Emitter, EOperator Op, FValue* A,
 		}
 		return Emitter->EmitConstantBoolScalar(Result);
 	}
+	else if (IsBitOperator(Op))
+	{
+		switch (PrimitiveType->ScalarKind)
+		{
+			case SK_Bool:
+			{
+				bool Result = (bool)FoldScalarBitOperator(Op, AConstant->Boolean, BConstant ? BConstant->Boolean : 0);
+				return Emitter->EmitConstantBoolScalar(Result);
+			}
+
+			case SK_Int:
+			{
+				TInteger Result = FoldScalarBitOperator(Op, AConstant->Integer, BConstant ? BConstant->Integer : 0);
+				return Emitter->EmitConstantIntScalar(Result);
+			}
+			
+			default:
+				UE_MIR_UNREACHABLE();
+		}
+	}
 	else
 	{
-		check(OperatorIsArithmetic(Op));
 		switch (PrimitiveType->ScalarKind)
 		{
 			case SK_Int:
 			{
-				int Result = FoldScalarArithmeticOperator<TInteger>(Op, AConstant->Integer, BConstant->Integer, CConstant ? CConstant->Integer : 0);
+				int Result = FoldScalarArithmeticOperator<TInteger>(Op, AConstant->Integer, BConstant ? BConstant->Integer : 0, CConstant ? CConstant->Integer : 0);
 				return Emitter->EmitConstantIntScalar(Result);
 			}
 
 			case SK_Float:
 			{
-				TFloat Result = FoldScalarArithmeticOperator<TFloat>(Op, AConstant->Float, BConstant->Float, CConstant ? CConstant->Float : 0);
+				TFloat Result = FoldScalarArithmeticOperator<TFloat>(Op, AConstant->Float, BConstant ? BConstant->Float : 0, CConstant ? CConstant->Float : 0);
 				return Emitter->EmitConstantFloatScalar(Result);
 			}
 
@@ -816,11 +855,18 @@ static FValue* TryFoldOperatorScalar(FEmitter* Emitter, EOperator Op, FValue* A,
 
 FValue* FEmitter::EmitOperator(EOperator Op, FValue* A, FValue* B, FValue* C)
 {
+	// Assert that if C is specified, B must too
+	check(!C || B);
+
+	// Verify that B argument has been provided if operator is binary.
+	check(!MIR::IsBinaryOperator(Op) || B);
+
 	// Verify that C argument has been provided if operator is ternary.
-	check(!MIR::OperatorIsTernary(Op) || C);
+	check(!MIR::IsTernaryOperator(Op) || C);
 
 	// Argument value types must always match.
-	if (A->Type != B->Type || (C && (A->Type != C->Type || B->Type != C->Type)))
+	if ((B && A->Type != B->Type) || 
+		(C && (A->Type != C->Type || B->Type != C->Type)))
 	{
 		Errorf(TEXT("Operands of operator `%s` do not have the same type."), OperatorToString(Op));
 		return MIR::FPoison::Get();
@@ -830,15 +876,23 @@ FValue* FEmitter::EmitOperator(EOperator Op, FValue* A, FValue* B, FValue* C)
 	FPrimitiveTypePtr PrimitiveType = A->Type->AsPrimitive();
 	if (!PrimitiveType)
 	{
-		Errorf(TEXT("Invalid operator on non primitive types."));
+		Errorf(TEXT("Operator on operands of non-primitive type."));
+		return MIR::FPoison::Get();
+	}
+	
+	// If operator is logical, arguments must be booleans.
+	if (IsBitOperator(Op) && !A->Type->IsBoolean() && !A->Type->IsInteger())
+	{
+		Errorf(TEXT("Logical operator on non-boolean operands."));
 		return MIR::FPoison::Get();
 	}
 
 	// Determine the result type. If the operator is arithmetic, the result type will be the same
-	// as the arguments type. Otherwise it will have the same number of components but bool.
-	FPrimitiveTypePtr ResultType = OperatorIsArithmetic(Op)
-		? PrimitiveType
-		: FPrimitiveType::Get(SK_Bool, PrimitiveType->NumRows, PrimitiveType->NumColumns);
+	// as the arguments type. Otherwise it will have the same number of components but boolean.
+	FPrimitiveTypePtr ResultType = IsComparisonOperator(Op) || IsBitOperator(Op)
+		? FPrimitiveType::Get(SK_Bool, PrimitiveType->NumRows, PrimitiveType->NumColumns)
+		: PrimitiveType;
+
 
 	// Check that at least one component of the resulting dimensional value would folded.
 	// If all components of resulting dimensional value are not folded, then instead of emitting
@@ -857,7 +911,7 @@ FValue* FEmitter::EmitOperator(EOperator Op, FValue* A, FValue* B, FValue* C)
 	{
 		// Extract the arguments individual components
 		FValue* AComponent = EmitSubscript(A, i);
-		FValue* BComponent = EmitSubscript(B, i);
+		FValue* BComponent = B ? EmitSubscript(B, i) : nullptr;
 		FValue* CComponent = C ? EmitSubscript(C, i) : nullptr;
 
 		// Try folding the operation, it may return null
@@ -866,7 +920,7 @@ FValue* FEmitter::EmitOperator(EOperator Op, FValue* A, FValue* B, FValue* C)
 		// Update the flags
 		bSomeResultComponentWasFolded |= (bool)ResultComponent;
 		bResultIsIdenticalToA &= ResultComponent && ResultComponent->Equals(AComponent);
-		bResultIsIdenticalToB &= ResultComponent && ResultComponent->Equals(BComponent);
+		bResultIsIdenticalToB &= BComponent && ResultComponent && ResultComponent->Equals(BComponent);
 		bResultIsIdenticalToC &= CComponent && ResultComponent && ResultComponent->Equals(CComponent);
 
 		// Cache the results
@@ -923,7 +977,7 @@ FValue* FEmitter::EmitOperator(EOperator Op, FValue* A, FValue* B, FValue* C)
 				FOperator Proto = MakePrototype<FOperator>(ComponentType);
 				Proto.Op = Op;
 				Proto.AArg = EmitSubscript(A, i);
-				Proto.BArg = EmitSubscript(B, i);
+				Proto.BArg = B ? EmitSubscript(B, i) : nullptr;
 				Proto.CArg = C ? EmitSubscript(C, i) : nullptr;
 				ResultComponents[i] = EmitPrototype(this, Proto);
 			}
@@ -938,6 +992,7 @@ FValue* FEmitter::EmitOperator(EOperator Op, FValue* A, FValue* B, FValue* C)
 	Proto.AArg = A;
 	Proto.BArg = B;
 	Proto.CArg = C;
+
 	return EmitPrototype(this, Proto);
 }
 
diff --git a/Engine/Source/Runtime/Engine/Private/Materials/MaterialIRToHLSLTranslator.cpp b/Engine/Source/Runtime/Engine/Private/Materials/MaterialIRToHLSLTranslator.cpp
--- a/Engine/Source/Runtime/Engine/Private/Materials/MaterialIRToHLSLTranslator.cpp
+++ b/Engine/Source/Runtime/Engine/Private/Materials/MaterialIRToHLSLTranslator.cpp
@@ -388,7 +388,7 @@ struct FTranslator : FMaterialIRToHLSLTranslation
 
 	void LowerOperator(const MIR::FOperator* Operator)
 	{
-		if (MIR::OperatorIsComparison(Operator->Op) || MIR::OperatorIsBinaryAlgebraic(Operator->Op))
+		if (MIR::IsComparisonOperator(Operator->Op) || MIR::IsBinaryInfixOperator(Operator->Op))
 		{
 			LowerValue(Operator->AArg);
 
diff --git a/Engine/Source/Runtime/Engine/Private/Materials/MaterialIRTypes.cpp b/Engine/Source/Runtime/Engine/Private/Materials/MaterialIRTypes.cpp
--- a/Engine/Source/Runtime/Engine/Private/Materials/MaterialIRTypes.cpp
+++ b/Engine/Source/Runtime/Engine/Private/Materials/MaterialIRTypes.cpp
@@ -168,6 +168,18 @@ bool FType::IsPoison() const
 	return Kind == TK_Poison;
 }
 
+bool FType::IsBoolean() const
+{
+	FPrimitiveTypePtr Ptr = AsPrimitive();
+	return Ptr && Ptr->ScalarKind == SK_Bool;
+}
+
+bool FType::IsInteger() const
+{
+	FPrimitiveTypePtr Ptr = AsPrimitive();
+	return Ptr && Ptr->ScalarKind == SK_Int;
+}
+
 bool FType::IsBoolScalar() const
 {
 	return this == FPrimitiveType::GetBool1();
diff --git a/Engine/Source/Runtime/Engine/Public/Materials/MaterialExpressionOneMinus.h b/Engine/Source/Runtime/Engine/Public/Materials/MaterialExpressionOneMinus.h
--- a/Engine/Source/Runtime/Engine/Public/Materials/MaterialExpressionOneMinus.h
+++ b/Engine/Source/Runtime/Engine/Public/Materials/MaterialExpressionOneMinus.h
@@ -20,6 +20,7 @@ class UMaterialExpressionOneMinus : public UMaterialExpression
 
 	//~ Begin UMaterialExpression Interface
 #if WITH_EDITOR
+	virtual void Build(UE::MIR::FEmitter& Emitter) override;
 	virtual bool GenerateHLSLExpression(FMaterialHLSLGenerator& Generator, UE::HLSLTree::FScope& Scope, int32 OutputIndex, UE::HLSLTree::FExpression const*& OutExpression) const override;
 	virtual int32 Compile(class FMaterialCompiler* Compiler, int32 OutputIndex) override;
 	virtual void GetCaption(TArray<FString>& OutCaptions) const override;
diff --git a/Engine/Source/Runtime/Engine/Public/Materials/MaterialIR.h b/Engine/Source/Runtime/Engine/Public/Materials/MaterialIR.h
--- a/Engine/Source/Runtime/Engine/Public/Materials/MaterialIR.h
+++ b/Engine/Source/Runtime/Engine/Public/Materials/MaterialIR.h
@@ -285,10 +285,15 @@ struct FSetMaterialOutput : TInstruction<VK_SetMaterialOutput>
 	FValue* Arg;
 };
 
-//
+// Note: if you modify this enum, update the implementations of the
+// helper functions below.
 enum EOperator
 {
 	O_Invalid,
+
+	/* Unary */
+	UO_Negate, 	// Arithmetic negation: negate(5) -> -5
+	UO_Not,		// Logical negation: not(true) -> false
 	
 	/* Binary Comparison Operators */
 	BO_GreaterThan,
@@ -303,6 +308,9 @@ enum EOperator
 	BO_Subtract,
 	BO_Multiply,
 	BO_Divide,
+	BO_Or,
+	BO_And,
+	BO_Xor,
 	BO_Min,
 	BO_Max,
 
@@ -312,16 +320,19 @@ enum EOperator
 };
 
 //
-bool OperatorIsComparison(EOperator Op);
+bool IsComparisonOperator(EOperator Op);
+
+//
+bool IsBitOperator(EOperator Op);
 
 //
-bool OperatorIsArithmetic(EOperator Op);
+bool IsBinaryOperator(EOperator Op);
 
 //
-bool OperatorIsTernary(EOperator Op);
+bool IsTernaryOperator(EOperator Op);
 
 //
-bool OperatorIsBinaryAlgebraic(EOperator Op);
+bool IsBinaryInfixOperator(EOperator Op);
 
 //
 const TCHAR* OperatorToString(EOperator Op);
diff --git a/Engine/Source/Runtime/Engine/Public/Materials/MaterialIREmitter.h b/Engine/Source/Runtime/Engine/Public/Materials/MaterialIREmitter.h
--- a/Engine/Source/Runtime/Engine/Public/Materials/MaterialIREmitter.h
+++ b/Engine/Source/Runtime/Engine/Public/Materials/MaterialIREmitter.h
@@ -153,7 +153,7 @@ public:
 	/* Instructions */
 
 	FSetMaterialOutput* EmitSetMaterialOutput(EMaterialProperty InProperty, FValue* InArgValue);
-	FValue* EmitOperator(EOperator Operator, FValue* A, FValue* B, FValue* C = nullptr);
+	FValue* EmitOperator(EOperator Operator, FValue* A, FValue* B = nullptr, FValue* C = nullptr);
 	FValue* EmitBranch(FValue* Condition, FValue* True, FValue* False);
 	FValue* EmitConstruct(FTypePtr Type, FValue* Initializer);
 	FValue* EmitTextureSample(UTexture* Texture, FValue* TexCoord, ESamplerSourceMode SamplerSourceMode, ETextureMipValueMode MipValueMode, EMaterialSamplerType SamplerType);
diff --git a/Engine/Source/Runtime/Engine/Public/Materials/MaterialIRTypes.h b/Engine/Source/Runtime/Engine/Public/Materials/MaterialIRTypes.h
--- a/Engine/Source/Runtime/Engine/Public/Materials/MaterialIRTypes.h
+++ b/Engine/Source/Runtime/Engine/Public/Materials/MaterialIRTypes.h
@@ -44,6 +44,12 @@ struct FType
 	// Returns whether this is the poison type.
 	bool IsPoison() const;
 
+	// Returns whether this is a boolean primitive type (of any dimension).
+	bool IsBoolean() const;
+	
+	// Returns whether this is an integer primitive type (of any dimension).
+	bool IsInteger() const;
+
 	// Returns whether this type is a `bool` scalar.
 	bool IsBoolScalar() const;
 
