Add Nanite Software VRS support for per-material "Allow Variable Rate Shading" checkbox (default on). Automatically disable VRS for materials that use Custom Primitive Data ,since this causes edge artifacts when used with Nanite Software VRS.

diff --git a/Engine/Source/Runtime/Engine/Private/Materials/Material.cpp b/Engine/Source/Runtime/Engine/Private/Materials/Material.cpp
--- a/Engine/Source/Runtime/Engine/Private/Materials/Material.cpp
+++ b/Engine/Source/Runtime/Engine/Private/Materials/Material.cpp
@@ -7174,6 +7174,26 @@ bool UMaterial::HasCustomizedUVs() const
 	return GetCachedExpressionData().bHasCustomizedUVs;
 }
 
+bool UMaterial::HasCustomPrimitiveData() const
+{
+	const FMaterialCachedExpressionData& CachedData = GetCachedExpressionData();
+	for (int32 Index : CachedData.ScalarPrimitiveDataIndexValues)
+	{
+		if (Index != INDEX_NONE)
+		{
+			return true;
+		}
+	}
+	for (int32 Index : CachedData.VectorPrimitiveDataIndexValues)
+	{
+		if (Index != INDEX_NONE)
+		{
+			return true;
+		}
+	}
+	return false;
+}
+
 USubsurfaceProfile* UMaterial::GetSubsurfaceProfile_Internal() const
 {
 	checkSlow(IsInGameThread());
diff --git a/Engine/Source/Runtime/Engine/Private/Materials/MaterialInstance.cpp b/Engine/Source/Runtime/Engine/Private/Materials/MaterialInstance.cpp
--- a/Engine/Source/Runtime/Engine/Private/Materials/MaterialInstance.cpp
+++ b/Engine/Source/Runtime/Engine/Private/Materials/MaterialInstance.cpp
@@ -2026,6 +2026,11 @@ bool UMaterialInstance::HasMeshPaintTexture() const
 	return Parent ? Parent->HasMeshPaintTexture() : false;
 }
 
+bool UMaterialInstance::HasCustomPrimitiveData() const
+{
+	return Parent ? Parent->HasCustomPrimitiveData() : false;
+}
+
 const FMaterialResource* UMaterialInstance::GetMaterialResource(ERHIFeatureLevel::Type InFeatureLevel, EMaterialQualityLevel::Type QualityLevel) const
 {
 	return const_cast<UMaterialInstance*>(this)->GetMaterialResource(InFeatureLevel, QualityLevel);
diff --git a/Engine/Source/Runtime/Engine/Private/Materials/MaterialInterface.cpp b/Engine/Source/Runtime/Engine/Private/Materials/MaterialInterface.cpp
--- a/Engine/Source/Runtime/Engine/Private/Materials/MaterialInterface.cpp
+++ b/Engine/Source/Runtime/Engine/Private/Materials/MaterialInterface.cpp
@@ -1469,6 +1469,11 @@ bool UMaterialInterface::HasMeshPaintTexture() const
 	return false;
 }
 
+bool UMaterialInterface::HasCustomPrimitiveData() const
+{
+	return false;
+}
+
 FMaterialShadingModelField UMaterialInterface::GetShadingModels() const
 {
 	return MSM_DefaultLit;
diff --git a/Engine/Source/Runtime/Engine/Private/Materials/MaterialShared.cpp b/Engine/Source/Runtime/Engine/Private/Materials/MaterialShared.cpp
--- a/Engine/Source/Runtime/Engine/Private/Materials/MaterialShared.cpp
+++ b/Engine/Source/Runtime/Engine/Private/Materials/MaterialShared.cpp
@@ -1762,12 +1762,10 @@ bool FMaterialResource::HasMaterialAttributesConnected() const { return (Materia
 #endif
 EMaterialShadingRate FMaterialResource::GetShadingRate() const { return Material->ShadingRate; }
 bool FMaterialResource::IsVariableRateShadingAllowed() const {
-	// * Automatically disable VRS on anything using pixel discard as coarse shading causes the whole block to get 
-	//   discarded resulting in noticeable artifacts
-	// * For now disable VRS when Substrate is enabled due to incorrect data replication between neighbor pixels
-	return Material->bAllowVariableRateShading 
-		&& !IsMasked() 
-		&& !Substrate::IsSubstrateEnabled();
+	return Material->bAllowVariableRateShading
+		&& !IsMasked()								// When using pixel discard, coarse shading causes the whole block to get discarded resulting in noticeable artifacts
+		&& !Material->HasCustomPrimitiveData()		// When using custom primitive data, Nanite can't determine which primitive ID to use for shading clusters
+		&& !Substrate::IsSubstrateEnabled();		// For now disable VRS when Substrate is enabled due to incorrect data replication between neighbor pixels
 }
 FString FMaterialResource::GetBaseMaterialPathName() const { return Material->GetPathName(); }
 FString FMaterialResource::GetDebugName() const
diff --git a/Engine/Source/Runtime/Engine/Public/Materials/Material.h b/Engine/Source/Runtime/Engine/Public/Materials/Material.h
--- a/Engine/Source/Runtime/Engine/Public/Materials/Material.h
+++ b/Engine/Source/Runtime/Engine/Public/Materials/Material.h
@@ -1230,6 +1230,7 @@ public:
 	ENGINE_API virtual bool HasCustomizedUVs() const override;
 	ENGINE_API virtual bool WritesToRuntimeVirtualTexture() const override;
 	ENGINE_API virtual bool HasMeshPaintTexture() const override;
+	ENGINE_API virtual bool HasCustomPrimitiveData() const override;
 
 	ENGINE_API virtual FGraphEventArray PrecachePSOs(const FPSOPrecacheVertexFactoryDataList& VertexFactoryDataList, const FPSOPrecacheParams& PreCacheParams, EPSOPrecachePriority Priority, TArray<FMaterialPSOPrecacheRequestID>& OutMaterialPSORequestIDs) override;
 
diff --git a/Engine/Source/Runtime/Engine/Public/Materials/MaterialInstance.h b/Engine/Source/Runtime/Engine/Public/Materials/MaterialInstance.h
--- a/Engine/Source/Runtime/Engine/Public/Materials/MaterialInstance.h
+++ b/Engine/Source/Runtime/Engine/Public/Materials/MaterialInstance.h
@@ -893,6 +893,7 @@ public:
 	ENGINE_API virtual bool IsMasked() const override;
 	ENGINE_API virtual bool WritesToRuntimeVirtualTexture() const override;
 	ENGINE_API virtual bool HasMeshPaintTexture() const override;
+	ENGINE_API virtual bool HasCustomPrimitiveData() const override;
 	ENGINE_API virtual FDisplacementScaling GetDisplacementScaling() const override;
 	ENGINE_API virtual bool IsDisplacementFadeEnabled() const override;
 	ENGINE_API virtual FDisplacementFadeRange GetDisplacementFadeRange() const override;
diff --git a/Engine/Source/Runtime/Engine/Public/Materials/MaterialInterface.h b/Engine/Source/Runtime/Engine/Public/Materials/MaterialInterface.h
--- a/Engine/Source/Runtime/Engine/Public/Materials/MaterialInterface.h
+++ b/Engine/Source/Runtime/Engine/Public/Materials/MaterialInterface.h
@@ -948,6 +948,7 @@ public:
 	ENGINE_API virtual bool IsPostProcessMaterial() const;
 	ENGINE_API virtual bool WritesToRuntimeVirtualTexture() const;
 	ENGINE_API virtual bool HasMeshPaintTexture() const;
+	ENGINE_API virtual bool HasCustomPrimitiveData() const;
 	ENGINE_API virtual FDisplacementScaling GetDisplacementScaling() const;
 	ENGINE_API virtual bool IsDisplacementFadeEnabled() const;
 	ENGINE_API virtual FDisplacementFadeRange GetDisplacementFadeRange() const;
