[SceneComponent] Add a new CVar to control whether USceneComponent::CheckStaticMobilityAndWarn should output warnings to the message/output log. In debug/dev builds, these warnings can cripple performance.

The CVar is `s.WarnSceneComponentStaticMobilityMove` and defaults to true. It is compiled out of test/shipping builds to match its usage.

diff --git a/Engine/Source/Runtime/Engine/Private/Components/SceneComponent.cpp b/Engine/Source/Runtime/Engine/Private/Components/SceneComponent.cpp
--- a/Engine/Source/Runtime/Engine/Private/Components/SceneComponent.cpp
+++ b/Engine/Source/Runtime/Engine/Private/Components/SceneComponent.cpp
@@ -57,6 +57,16 @@ namespace SceneComponentCVars
 		TEXT("Scene Component OnRep_AttachedChildren:\n")
 		TEXT("false: fix up any children that are missing the parent, true: fixes up all children except non-replicated root components (default)"),
 		ECVF_Default);
+
+#if !(UE_BUILD_SHIPPING || UE_BUILD_TEST)
+	bool bWarnSceneComponentStaticMobilityMove = true;
+	static FAutoConsoleVariableRef CVarWarnSceneComponentStaticMobilityMove(
+		TEXT("s.WarnSceneComponentStaticMobilityMove"),
+		bWarnSceneComponentStaticMobilityMove,
+		TEXT("If enabled, output a warning to the message and output log about static mobility moves for scene components."),
+		ECVF_Default
+	);
+#endif // !(UE_BUILD_SHIPPING || UE_BUILD_TEST)
 }
 
 DEFINE_LOG_CATEGORY_STATIC(LogSceneComponent, Log, All);
@@ -3131,8 +3141,11 @@ bool USceneComponent::CheckStaticMobilityAndWarn(const FText& ActionText) const
 				if (MyOwner && MyOwner->IsActorInitialized())
 				{
 #if !(UE_BUILD_SHIPPING || UE_BUILD_TEST)
-					FMessageLog("PIE").Warning(FText::Format(LOCTEXT("InvalidMustBeMovable", "Mobility of {0} : {1} has to be 'Movable' if you'd like to {2}. "),
-						FText::FromString(GetPathNameSafe(GetOwner())), FText::FromString(GetName()), ActionText));
+					if (SceneComponentCVars::bWarnSceneComponentStaticMobilityMove)
+					{
+						FMessageLog("PIE").Warning(FText::Format(LOCTEXT("InvalidMustBeMovable", "Mobility of {0} : {1} has to be 'Movable' if you'd like to {2}. "),
+							FText::FromString(GetPathNameSafe(GetOwner())), FText::FromString(GetName()), ActionText));
+					}
 #endif
 					return true;
 				}
