- Added incremental pre-unregistration of actor/components in level streaming (optional).
- When enabled, allows asynchronous physics state destruction when making a streaming level invisible (depends on p.Chaos.EnableAsyncInitBody being set to true).
- Can be enabled with s.AllowIncrementalPreUnregisterComponents=true.

diff --git a/Engine/Source/Runtime/Engine/Classes/Components/ActorComponent.h b/Engine/Source/Runtime/Engine/Classes/Components/ActorComponent.h
--- a/Engine/Source/Runtime/Engine/Classes/Components/ActorComponent.h
+++ b/Engine/Source/Runtime/Engine/Classes/Components/ActorComponent.h
@@ -204,7 +204,11 @@ protected:
 		/** Component is pre-registering (not registered) */
 		PreRegistering,
 		/** Component is pre-registered (not registered) */
-		PreRegistered
+		PreRegistered,
+		/** Component is pre-unregistering (registered) */
+		PreUnregistering,
+		/** Component is pre-unregistered (registered) */
+		PreUnregistered
 	};
 
 	/** Indicates if this ActorComponent is currently pre-registering or pre-registered with a scene. */
@@ -375,8 +379,18 @@ private:
 	/** Whether to use use the async physics tick with this component. */
 	uint8 bAsyncPhysicsTickEnabled : 1;
 
-	/** If non-null, physics state creation is done asynchronously by FPhysScene_Chaos::ProcessAsyncPhysicsStateJobs. */
-	FPhysScene* AsyncCreatePhysicsStateScene;
+	enum struct EPhysicsStateAsyncCreationState : uint8
+	{
+		None,
+		Creating,
+		Destroying
+	};
+
+	/** Whether the component is creating/destroying its physics state asynchronously (or not) */
+	EPhysicsStateAsyncCreationState PhysicsStateAsyncCreationState : 2;
+
+	/** If non-null, physics state creation/destruction is done asynchronously by FPhysScene_Chaos::ProcessAsyncPhysicsStateJobs. */
+	FPhysScene* PhysicsStateAsyncCreationScene;
 
 public:
 
@@ -405,7 +419,16 @@ private:
 public:
 
 	/** Whether asynchronous physics state creation is running or not. */
-	ENGINE_API bool IsAsyncCreatePhysicsStateRunning() const { return AsyncCreatePhysicsStateScene != nullptr; }
+	ENGINE_API bool IsAsyncCreatePhysicsStateRunning() const
+	{
+		return PhysicsStateAsyncCreationScene && (PhysicsStateAsyncCreationState == EPhysicsStateAsyncCreationState::Creating);
+	}
+
+	/** Whether asynchronous physics state destruction is running or not. */
+	ENGINE_API bool IsAsyncDestroyPhysicsStateRunning() const
+	{
+		return PhysicsStateAsyncCreationScene && (PhysicsStateAsyncCreationState == EPhysicsStateAsyncCreationState::Destroying);
+	}
 
 	/** Tracks whether the component has been added to one of the world's end of frame update lists */
 	uint32 GetMarkedForEndOfFrameUpdateState() const { return MarkedForEndOfFrameUpdateState; }
@@ -731,12 +754,21 @@ protected:
 	/** Returns whether this component should create its physics state asynchronously. */
 	ENGINE_API virtual bool ShouldAsyncCreatePhysicsState(UWorld* WorldContext) const;
 
+	/** Returns whether this component should destroy its physics state asynchronously. */
+	ENGINE_API virtual bool ShouldAsyncDestroyPhysicsState() const;
+
 	/** Called by PreRegisterComponentWithWorld when the component is being pre-registered. */
 	ENGINE_API virtual void OnPreRegister();
 
 	/** Called once the component is pre-registered. */
 	ENGINE_API virtual void OnPreRegistered();
 
+	/** Called by PreUnregisterComponentFromWorld when the component is being pre-unregistered. */
+	ENGINE_API virtual void OnPreUnregister();
+
+	/** Called once the component is pre-unregistered. */
+	ENGINE_API virtual void OnPreUnregistered();
+
 	/**
 	 * Called when a component is registered, after Scene is set, but before CreateRenderState_Concurrent or OnCreatePhysicsState are called.
 	 */
@@ -804,14 +836,37 @@ public:
 	/** Returns whether PreRegisterComponentWithWorld should be called for this component when its Level is being added to the world. */
 	ENGINE_API virtual bool ShouldIncrementalPreRegister(UWorld* WorldContext) const;
 
-	/** Returns whether this component allows having its physics state being created and destroyed asynchronously (outside of the GameThread). */
-	ENGINE_API virtual bool AllowsAsyncPhysicsStateCreationAndDestruction() const { return false; }
+	/** Returns whether PreUnregisterComponentFromWorld should be called for this component when its Level is being removed from the world. */
+	ENGINE_API virtual bool ShouldIncrementalPreUnregister() const;
+
+	/** Returns whether this component allows having its physics state being created asynchronously (outside of the GameThread). */
+	ENGINE_API virtual bool AllowsAsyncPhysicsStateCreation() const
+	{
+		return false;
+	}
+
+	/** Returns whether this component allows having its physics state being destroyed asynchronously (outside of the GameThread). */
+	ENGINE_API virtual bool AllowsAsyncPhysicsStateDestruction() const
+	{
+		return false;
+	}
 
 	/** Used to create any physics engine information for this component outside of the GameThread. */
 	ENGINE_API virtual void OnAsyncCreatePhysicsState();
 
 	/** Called on the GameThread once the component's physic engine information is created. */
-	ENGINE_API virtual void OnAsyncCreatePhysicsStateDone_GameThread();
+	ENGINE_API virtual void OnAsyncCreatePhysicsStateEnd_GameThread();
+
+	/** Used to destroy any physics engine information for this component outside of the GameThread. */
+	ENGINE_API virtual void OnAsyncDestroyPhysicsState()
+	{
+	}
+
+	/** Called on the GameThread before the component's physic engine information is destroyed. */
+	ENGINE_API virtual void OnAsyncDestroyPhysicsStateBegin_GameThread();
+
+	/** Called on the GameThread once the component's physic engine information is destroyed. */
+	ENGINE_API virtual void OnAsyncDestroyPhysicsStateEnd_GameThread();
 
 	/**
 	 * Initializes the component.  Occurs at level startup or actor spawn. This is before BeginPlay (Actor or Component).  
@@ -940,6 +995,11 @@ public:
 	 */
 	ENGINE_API void PreRegisterComponentWithWorld(UWorld* InWorld);
 
+	/**
+	 * Pre-unregisters a component from a specific world. Only called on components that return ShouldIncrementalPreUnregister().
+	 */
+	ENGINE_API void PreUnregisterComponentFromWorld();
+
 	/**
 	 * Registers a component with a specific world, which creates any visual/physical state
 	 * @param InWorld - The world to register the component with.
@@ -1170,10 +1230,28 @@ public:
 	ENGINE_API bool IsOwnerRunningUserConstructionScript() const;
 
 	/** See if this component is currently pre-registering */
-	FORCEINLINE bool IsPreRegistering() const { return RegistrationState == EComponentRegistrationState::PreRegistering; }
+	FORCEINLINE bool IsPreRegistering() const
+	{
+		return RegistrationState == EComponentRegistrationState::PreRegistering;
+	}
 
 	/** See if this component is currently pre-registered */
-	FORCEINLINE bool IsPreRegistered() const { return RegistrationState == EComponentRegistrationState::PreRegistered; }
+	FORCEINLINE bool IsPreRegistered() const
+	{
+		return RegistrationState == EComponentRegistrationState::PreRegistered;
+	}
+
+	/** See if this component is currently pre-unregistering */
+	FORCEINLINE bool IsPreUnregistering() const
+	{
+		return RegistrationState == EComponentRegistrationState::PreUnregistering;
+	}
+
+	/** See if this component is currently pre-unregistered */
+	FORCEINLINE bool IsPreUnregistered() const
+	{
+		return RegistrationState == EComponentRegistrationState::PreUnregistered;
+	}
 
 	/** See if this component is currently registered */
 	FORCEINLINE bool IsRegistered() const { return bRegistered; }
diff --git a/Engine/Source/Runtime/Engine/Classes/Components/InstancedStaticMeshComponent.h b/Engine/Source/Runtime/Engine/Classes/Components/InstancedStaticMeshComponent.h
--- a/Engine/Source/Runtime/Engine/Classes/Components/InstancedStaticMeshComponent.h
+++ b/Engine/Source/Runtime/Engine/Classes/Components/InstancedStaticMeshComponent.h
@@ -433,6 +433,9 @@ public:
 private:
 	/** Current cached conservativ bounds. */
 	FBox CachedConservativeInstanceBounds;
+
+	/** Payload used by asynchronous destruction of physics state (see OnAsyncDestroyPhysicsState). */
+	TArray<FBodyInstance*> AsyncDestroyPhysicsStatePayload;
 public:
 
 #if WITH_EDITOR
@@ -472,6 +475,9 @@ public:
 protected:
 	ENGINE_API virtual void OnCreatePhysicsState() override;
 	ENGINE_API virtual void OnDestroyPhysicsState() override;
+	ENGINE_API virtual void OnAsyncDestroyPhysicsState() override;
+	ENGINE_API virtual void OnAsyncDestroyPhysicsStateBegin_GameThread() override;
+
 public:
 	ENGINE_API virtual bool CanEditSimulatePhysics() override;
 
diff --git a/Engine/Source/Runtime/Engine/Classes/Components/PrimitiveComponent.h b/Engine/Source/Runtime/Engine/Classes/Components/PrimitiveComponent.h
--- a/Engine/Source/Runtime/Engine/Classes/Components/PrimitiveComponent.h
+++ b/Engine/Source/Runtime/Engine/Classes/Components/PrimitiveComponent.h
@@ -821,6 +821,9 @@ private:
 	/** Custom data that can be read by a material through a material parameter expression. Set data using SetCustomPrimitiveData* functions */
 	UPROPERTY(Transient)
 	FCustomPrimitiveData CustomPrimitiveDataInternal;
+
+	/** Payload used to release BodyInstance resources in asynchronous mode (see OnAsyncDestroyPhysicsState). */
+	TOptional<FBodyInstance::FAsyncTermBodyPayload> AsyncTermBodyPayload;
 	
 public:
 
@@ -2446,7 +2449,10 @@ public:
 	ENGINE_API virtual void OnUnregister()  override;
 	ENGINE_API virtual void DestroyRenderState_Concurrent() override;
 	ENGINE_API virtual void OnCreatePhysicsState() override;
-	ENGINE_API virtual void OnAsyncCreatePhysicsStateDone_GameThread() override;
+	ENGINE_API virtual void OnAsyncCreatePhysicsStateEnd_GameThread() override;
+	ENGINE_API virtual void OnAsyncDestroyPhysicsState() override;
+	ENGINE_API virtual void OnAsyncDestroyPhysicsStateBegin_GameThread() override;
+	ENGINE_API virtual void OnAsyncDestroyPhysicsStateEnd_GameThread() override;
 	ENGINE_API virtual void OnDestroyPhysicsState() override;
 	ENGINE_API virtual void OnActorEnableCollisionChanged() override;
 	ENGINE_API virtual void InvalidateLightingCacheDetailed(bool bInvalidateBuildEnqueuedLighting, bool bTranslationOnly) override;
diff --git a/Engine/Source/Runtime/Engine/Classes/Components/StaticMeshComponent.h b/Engine/Source/Runtime/Engine/Classes/Components/StaticMeshComponent.h
--- a/Engine/Source/Runtime/Engine/Classes/Components/StaticMeshComponent.h
+++ b/Engine/Source/Runtime/Engine/Classes/Components/StaticMeshComponent.h
@@ -556,9 +556,11 @@ protected:
 	ENGINE_API virtual void CreateRenderState_Concurrent(FRegisterComponentContext* Context) override;
 	ENGINE_API virtual void OnCreatePhysicsState() override;
 	ENGINE_API virtual void OnDestroyPhysicsState() override;
+	ENGINE_API virtual void OnAsyncDestroyPhysicsStateEnd_GameThread() override;
 	ENGINE_API virtual bool ShouldCreatePhysicsState() const override;
 	ENGINE_API virtual bool ShouldCreateRenderState() const override;
-	ENGINE_API virtual bool AllowsAsyncPhysicsStateCreationAndDestruction() const override;
+	ENGINE_API virtual bool AllowsAsyncPhysicsStateCreation() const override;
+	ENGINE_API virtual bool AllowsAsyncPhysicsStateDestruction() const override;
 
 public:
 	ENGINE_API virtual void PrecachePSOs() override;
diff --git a/Engine/Source/Runtime/Engine/Classes/Components/UnregisterComponentContext.h b/Engine/Source/Runtime/Engine/Classes/Components/UnregisterComponentContext.h
new file mode 100644
--- /dev/null
+++ b/Engine/Source/Runtime/Engine/Classes/Components/UnregisterComponentContext.h
@@ -0,0 +1,29 @@
+// Copyright Epic Games, Inc. All Rights Reserved.
+
+#pragma once
+
+class UWorld;
+
+class FUnregisterComponentContext
+{
+public:
+	FUnregisterComponentContext(UWorld* InWorld)
+		: World(InWorld)
+		, IncrementalRemoveTimeLimitMs(0)
+	{}
+
+	FUnregisterComponentContext& SetIncrementalRemoveTimeLimitMs(double Value)
+	{
+		IncrementalRemoveTimeLimitMs = Value;
+		return *this;
+	}
+	
+	double GetIncrementalRemoveTimeLimitMs() const
+	{
+		return IncrementalRemoveTimeLimitMs;
+	}
+
+private:
+	UWorld* World;
+	double	IncrementalRemoveTimeLimitMs;
+};
diff --git a/Engine/Source/Runtime/Engine/Classes/Engine/Level.h b/Engine/Source/Runtime/Engine/Classes/Engine/Level.h
--- a/Engine/Source/Runtime/Engine/Classes/Engine/Level.h
+++ b/Engine/Source/Runtime/Engine/Classes/Engine/Level.h
@@ -37,6 +37,7 @@ class UTexture2D;
 struct FLevelCollection;
 class ULevelActorContainer;
 class FRegisterComponentContext;
+class FUnregisterComponentContext;
 class SNotificationItem;
 class UActorFolder;
 class IWorldPartitionCell;
@@ -763,6 +764,16 @@ public:
 #endif
 
 private:
+	enum class EIncrementalUnregisterComponentState : uint8
+	{
+		Init,
+		PreUnregisterComponents,
+		UnregisterComponents,
+		Finalize
+	};
+
+	EIncrementalUnregisterComponentState IncrementalUnregisterComponentState;
+
 	FLevelBoundsActorUpdatedEvent LevelBoundsActorUpdatedEvent; 
 
 	UPROPERTY()
@@ -777,10 +788,51 @@ private:
 	/** Cached level collection that this level is contained in, for faster access than looping through the collections in the world. */
 	FLevelCollection* CachedLevelCollection;
 
-	ENGINE_API void OnPreRegisterComponent(UActorComponent* Component) { check(IsValid(Component)); PreRegisteringActorComponents.Add(Component); }
-	ENGINE_API void OnRemovedPreRegisteringComponent(UActorComponent* Component) { check(IsValid(Component)); verify(PreRegisteringActorComponents.Remove(Component)); }
-	ENGINE_API void OnPreRegisteredComponent(UActorComponent* Component) { OnRemovedPreRegisteringComponent(Component); }
-	ENGINE_API bool HasPreRegisteringComponents() const { return PreRegisteringActorComponents.Num() > 0; }
+	ENGINE_API void OnPreRegisterComponent(UActorComponent* Component)
+	{
+		check(IsValid(Component));
+		PreRegisteringActorComponents.Add(Component);
+	}
+
+	ENGINE_API void OnRemovedPreRegisteringComponent(UActorComponent* Component)
+	{
+		check(IsValid(Component));
+		verify(PreRegisteringActorComponents.Remove(Component));
+	}
+
+	ENGINE_API void OnPreRegisteredComponent(UActorComponent* Component)
+	{
+		OnRemovedPreRegisteringComponent(Component);
+	}
+
+	ENGINE_API bool HasPreRegisteringComponents() const
+	{
+		return PreRegisteringActorComponents.Num() > 0;
+	}
+
+	ENGINE_API void OnPreUnregisterComponent(UActorComponent* Component) 
+	{
+		check(IsValid(Component));
+		PreUnregisteringActorComponents.Add(Component);
+	}
+
+	ENGINE_API void OnRemovedPreUnregisteringComponent(UActorComponent* Component) 
+	{
+		check(IsValid(Component));
+		verify(PreUnregisteringActorComponents.Remove(Component));
+	}
+
+	ENGINE_API void OnPreUnregisteredComponent(UActorComponent* Component) 
+	{
+		OnRemovedPreUnregisteringComponent(Component);
+	}
+
+	ENGINE_API bool HasPreUnregisteringComponents() const
+	{
+		return PreUnregisteringActorComponents.Num() > 0;
+	}
+
+	TSet<UActorComponent*> PreUnregisteringActorComponents;
 	TSet<UActorComponent*> PreRegisteringActorComponents;
 	friend struct FLevelRegistrationAccessor;
 
@@ -936,7 +988,7 @@ public:
     *
 	* @param NumComponentsToUnregister		Minimum number of components to unregister in this run, 0 for all
 	*/
-	bool IncrementalUnregisterComponents(int32 NumComponentsToUnregister);
+	bool IncrementalUnregisterComponents(int32 NumComponentsToUnregister, FUnregisterComponentContext* Context = nullptr);
 
 
 	/**
@@ -1408,6 +1460,7 @@ public:
 private:
 	bool IncrementalPreRegisterComponents(FRegisterComponentContext& Context);
 	bool IncrementalRegisterComponents(FRegisterComponentContext& Context);
+	bool IncrementalPreUnregisterComponents(FUnregisterComponentContext& Context);
 #if WITH_EDITOR
 	bool IncrementalRunConstructionScripts(bool bProcessAllActors);
 	TOptional<bool> bCachedHasStaticMeshCompilationPending;
@@ -1468,14 +1521,50 @@ private:
 #endif // WITH_EDITOR
 };
 
-/** Accessor used by World/ActorComponent for component registration. */
+/** Accessor used by World/ActorComponent for component registration/unregistration. */
 struct FLevelRegistrationAccessor
 {
 private:
-	ENGINE_API static void OnPreRegisterComponent(ULevel* InLevel, UActorComponent* InComponent) { InLevel->OnPreRegisterComponent(InComponent); }
-	ENGINE_API static void OnPreRegisteredComponent(ULevel* InLevel, UActorComponent* InComponent) { InLevel->OnPreRegisteredComponent(InComponent); }
-	ENGINE_API static void OnRemovedPreRegisteringComponent(ULevel* InLevel, UActorComponent* InComponent) { InLevel->OnRemovedPreRegisteringComponent(InComponent); }
-	ENGINE_API static bool HasPreRegisteringComponents(const ULevel* InLevel) { return InLevel->HasPreRegisteringComponents(); }
+	ENGINE_API static void OnPreRegisterComponent(ULevel* InLevel, UActorComponent* InComponent)
+	{
+		InLevel->OnPreRegisterComponent(InComponent);
+	}
+
+	ENGINE_API static void OnPreRegisteredComponent(ULevel* InLevel, UActorComponent* InComponent)
+	{
+		InLevel->OnPreRegisteredComponent(InComponent);
+	}
+
+	ENGINE_API static void OnRemovedPreRegisteringComponent(ULevel* InLevel, UActorComponent* InComponent)
+	{
+		InLevel->OnRemovedPreRegisteringComponent(InComponent);
+	}
+
+	ENGINE_API static bool HasPreRegisteringComponents(const ULevel* InLevel)
+	{
+		return InLevel->HasPreRegisteringComponents();
+	}
+
+	ENGINE_API static void OnPreUnregisterComponent(ULevel* InLevel, UActorComponent* InComponent)
+	{
+		InLevel->OnPreUnregisterComponent(InComponent);
+	}
+
+	ENGINE_API static void OnPreUnregisteredComponent(ULevel* InLevel, UActorComponent* InComponent)
+	{
+		InLevel->OnPreUnregisteredComponent(InComponent);
+	}
+
+	ENGINE_API static void OnRemovedPreUnregisteringComponent(ULevel* InLevel, UActorComponent* InComponent)
+	{
+		InLevel->OnRemovedPreUnregisteringComponent(InComponent);
+	}
+
+	ENGINE_API static bool HasPreUnregisteringComponents(const ULevel* InLevel)
+	{
+		return InLevel->HasPreUnregisteringComponents();
+	}
+
 	friend class UActorComponent;
 	friend class UWorld;
 };
diff --git a/Engine/Source/Runtime/Engine/Classes/GameFramework/Actor.h b/Engine/Source/Runtime/Engine/Classes/GameFramework/Actor.h
--- a/Engine/Source/Runtime/Engine/Classes/GameFramework/Actor.h
+++ b/Engine/Source/Runtime/Engine/Classes/GameFramework/Actor.h
@@ -1479,12 +1479,15 @@ private:
 	ENGINE_API TArray<const UDataLayerInstance*> GetDataLayerInstancesInternal(bool bUseLevelContext, bool bIncludeParentDataLayers = true) const;
 	ENGINE_API const IWorldPartitionCell* GetWorldPartitionRuntimeCell() const;
 
-	/** Incrementally pre-registers components associated with this actor, used during level streaming */
+	/** Incrementally pre-registers components associated with this actor (used during level streaming). */
 	ENGINE_API void IncrementalPreRegisterComponents();
 
 	/** Returns bHasPreRegisteredAllComponents which indicates whether this actor has pre-registered all its components. */
 	ENGINE_API bool HasPreRegisteredAllComponents() const { return bHasPreRegisteredAllComponents; }
 
+	/** Incrementally pre-unregisters components associated with this actor (used during level streaming). */
+	ENGINE_API void IncrementalPreUnregisterComponents();
+
 	//~=============================================================================
 	// General functions.
 public:
@@ -4591,8 +4594,20 @@ inline FActorLastRenderTime* FActorLastRenderTime::GetPtr(AActor* InActor)
 struct FActorPreRegistrationAccessor
 {
 private:
-	ENGINE_API static void IncrementalPreRegisterComponents(AActor* InActor) { InActor->IncrementalPreRegisterComponents(); }
-	ENGINE_API static bool HasPreRegisteredAllComponents(const AActor* InActor) { return InActor->HasPreRegisteredAllComponents(); }
+	ENGINE_API static void IncrementalPreRegisterComponents(AActor* InActor)
+	{
+		InActor->IncrementalPreRegisterComponents();
+	}
+
+	ENGINE_API static bool HasPreRegisteredAllComponents(const AActor* InActor)
+	{
+		return InActor->HasPreRegisteredAllComponents();
+	}
+
+	ENGINE_API static void IncrementalPreUnregisterComponents(AActor* InActor)
+	{
+		InActor->IncrementalPreUnregisterComponents();
+	}
 	friend class ULevel;
 };
 
diff --git a/Engine/Source/Runtime/Engine/Classes/PhysicsEngine/BodyInstance.h b/Engine/Source/Runtime/Engine/Classes/PhysicsEngine/BodyInstance.h
--- a/Engine/Source/Runtime/Engine/Classes/PhysicsEngine/BodyInstance.h
+++ b/Engine/Source/Runtime/Engine/Classes/PhysicsEngine/BodyInstance.h
@@ -777,6 +777,30 @@ public:
 	 */
 	ENGINE_API void TermBody(bool bNeverDeferRelease = false);
 
+	// Contains the payload necessary for FBodyInstance::AsyncTermBody to release the body's resources asynchronously (outside of game thread).
+	struct ENGINE_API FAsyncTermBodyPayload
+	{
+		FAsyncTermBodyPayload();
+
+	private:
+		FPhysicsActorHandle ActorHandle;
+		TWeakObjectPtr<class UPrimitiveComponent> OwnerComponent;
+		bool bNotifyRigidBodyCollision;
+
+		friend struct FBodyInstance;
+	};
+
+	/**
+	 * Prepares the body to release its resources outside of the game thread. 
+	 * @return the payload used by AsyncTermBody to release the resources.
+	 */
+	ENGINE_API FAsyncTermBodyPayload StartAsyncTermBody_GameThread();
+
+	/** Releases body resources provided by Payload which was returned by StartAsyncTermBody_GameThread (can be called outside of game thread).
+	 * @param Payload body resources to release.
+	 */
+	static ENGINE_API void AsyncTermBody(FAsyncTermBodyPayload& Payload);
+
 	/** 
 	 * Takes two body instances and welds them together to create a single simulated rigid body. Returns true if success.
 	 */
diff --git a/Engine/Source/Runtime/Engine/Private/Actor.cpp b/Engine/Source/Runtime/Engine/Private/Actor.cpp
--- a/Engine/Source/Runtime/Engine/Private/Actor.cpp
+++ b/Engine/Source/Runtime/Engine/Private/Actor.cpp
@@ -5587,6 +5587,33 @@ static USceneComponent* GetNonPreRegisteredParent(UActorComponent* Component)
 } // namespace Private
 } // namespace UE
 
+void AActor::IncrementalPreUnregisterComponents()
+{
+	TRACE_CPUPROFILER_EVENT_SCOPE(AActor::IncrementalPreUnregisterComponents);
+
+	TInlineComponentArray<UActorComponent*> Components;
+	for (UActorComponent* Component : GetComponents())
+	{
+		if (Component->ShouldIncrementalPreUnregister())
+		{
+			Components.Add(Component);
+		}
+	}
+
+	if (!Components.Num())
+	{
+		return;
+	}
+
+	// Unregister all tick functions before pre-unregistering components as this will start asynchronous tasks (like destruction of the physics state)
+	RegisterAllActorTickFunctions(false, true); 
+
+	for (UActorComponent* Component : Components)
+	{
+		Component->PreUnregisterComponentFromWorld();
+	}
+}
+
 void AActor::IncrementalPreRegisterComponents()
 {
 	TRACE_CPUPROFILER_EVENT_SCOPE(AActor::IncrementalPreRegisterComponents);
diff --git a/Engine/Source/Runtime/Engine/Private/Components/ActorComponent.cpp b/Engine/Source/Runtime/Engine/Private/Components/ActorComponent.cpp
--- a/Engine/Source/Runtime/Engine/Private/Components/ActorComponent.cpp
+++ b/Engine/Source/Runtime/Engine/Private/Components/ActorComponent.cpp
@@ -442,7 +442,8 @@ FTransactionallySafeRWLock UActorComponent::AllUCSModifiedPropertiesLock;
 UActorComponent::UActorComponent(const FObjectInitializer& ObjectInitializer /*= FObjectInitializer::Get()*/)
 	: Super(ObjectInitializer)
 	, RegistrationState(EComponentRegistrationState::None)
-	, AsyncCreatePhysicsStateScene(nullptr)
+	, PhysicsStateAsyncCreationState(EPhysicsStateAsyncCreationState::None)
+	, PhysicsStateAsyncCreationScene(nullptr)
 {
 	OwnerPrivate = GetTypedOuter<AActor>();
 
@@ -1130,18 +1131,29 @@ bool UActorComponent::ShouldAsyncCreatePhysicsState(UWorld* WorldContext) const
 	check(WorldContext);
 	check(FPhysScene::SupportsAsyncPhysicsStateCreation());
 	TGuardValue<EComponentRegistrationState> Guard(const_cast<UActorComponent*>(this)->RegistrationState, EComponentRegistrationState::PreRegistering);
-	return (!bPhysicsStateCreated && WorldContext->GetPhysicsScene() && AllowsAsyncPhysicsStateCreationAndDestruction() && ShouldCreatePhysicsState());
+	return (!bPhysicsStateCreated && WorldContext->GetPhysicsScene() && AllowsAsyncPhysicsStateCreation() && ShouldCreatePhysicsState());
+}
+
+bool UActorComponent::ShouldAsyncDestroyPhysicsState() const
+{
+	return bPhysicsStateCreated && WorldPrivate && WorldPrivate->GetPhysicsScene() && AllowsAsyncPhysicsStateDestruction();
 }
 
 bool UActorComponent::ShouldIncrementalPreRegister(UWorld* WorldContext) const
 {
-	// For now the only criteria is if the component should create physics state asynchronously
+	// For now the only criteria is if the component should create its physics state asynchronously
 	return ShouldAsyncCreatePhysicsState(WorldContext);
 }
 
+bool UActorComponent::ShouldIncrementalPreUnregister() const
+{
+	// For now the only criteria is if the component should destroy its physics state asynchronously
+	return ShouldAsyncDestroyPhysicsState();
+}
+
 void UActorComponent::OnPreRegister()
 {
-	TRACE_CPUPROFILER_EVENT_SCOPE(UActorComponent::OnPreRegister)
+	TRACE_CPUPROFILER_EVENT_SCOPE(UActorComponent::OnPreRegister);
 	checkf(WorldPrivate, TEXT("OnPreRegister: %s to %s"), *GetDetailedInfo(), *GetFullNameSafe(GetOwner()));
 	checkf(RegistrationState == EComponentRegistrationState::None, TEXT("OnPreRegister: %s to %s"), *GetDetailedInfo(), *GetFullNameSafe(GetOwner()));
 	checkf(!bRegistered, TEXT("OnPreRegister: %s to %s"), *GetDetailedInfo(), *GetFullNameSafe(GetOwner()));
@@ -1161,8 +1173,10 @@ void UActorComponent::OnPreRegister()
 		FPhysScene_Chaos* PhysScene = WorldPrivate->GetPhysicsScene();
 		check(PhysScene);
 		check(!IsAsyncCreatePhysicsStateRunning());
+		check(!IsAsyncDestroyPhysicsStateRunning());
 		verify(PhysScene->PushAsyncCreatePhysicsStateJob(this));
-		AsyncCreatePhysicsStateScene = PhysScene;
+		PhysicsStateAsyncCreationState = EPhysicsStateAsyncCreationState::Creating;
+		PhysicsStateAsyncCreationScene = PhysScene;
 		check(IsAsyncCreatePhysicsStateRunning());
 	}
 	else
@@ -1173,7 +1187,7 @@ void UActorComponent::OnPreRegister()
 
 void UActorComponent::OnPreRegistered()
 {
-	check(RegistrationState == EComponentRegistrationState::PreRegistering);
+	check(IsPreRegistering());
 	RegistrationState = EComponentRegistrationState::PreRegistered;
 
 	AActor* MyOwner = GetOwner();
@@ -1185,7 +1199,7 @@ void UActorComponent::OnPreRegistered()
 
 void UActorComponent::OnRegister()
 {
-	TRACE_CPUPROFILER_EVENT_SCOPE(UActorComponent::OnRegister)
+	TRACE_CPUPROFILER_EVENT_SCOPE(UActorComponent::OnRegister);
 #if !UE_BUILD_SHIPPING
 	// These are removed in shipping because they are still likely to fail in Test and Development builds, and checks in shipping makes this rather expensive.
 	checkf(!IsUnreachable(), TEXT("%s"), *GetDetailedInfo());
@@ -1224,6 +1238,7 @@ void UActorComponent::OnUnregister()
 	bRegistered = false;
 
 	check(!IsPreRegistering());
+	check(!IsPreUnregistering());
 	RegistrationState = EComponentRegistrationState::None;
 
 	ClearNeedEndOfFrameUpdate();
@@ -1667,6 +1682,38 @@ void UActorComponent::RegisterComponent()
 	}
 }
 
+void UActorComponent::PreUnregisterComponentFromWorld()
+{
+	check(RegistrationState == EComponentRegistrationState::None);
+	OnPreUnregister();
+	check(IsPreUnregistering());
+}
+
+void UActorComponent::OnPreUnregister()
+{
+	checkf(WorldPrivate, TEXT("OnPreUnregister: %s to %s"), *GetDetailedInfo(), *GetFullNameSafe(GetOwner()));
+	checkf(RegistrationState == EComponentRegistrationState::None, TEXT("OnPreUnregister: %s to %s"), *GetDetailedInfo(), *GetFullNameSafe(GetOwner()));
+	checkf(bRegistered, TEXT("OnPreUnregister: %s to %s"), *GetDetailedInfo(), *GetFullNameSafe(GetOwner()));
+
+	RegistrationState = EComponentRegistrationState::PreUnregistering;
+
+	AActor* MyOwner = GetOwner();
+	if (ULevel* Level = MyOwner ? MyOwner->GetLevel() : nullptr)
+	{
+		FLevelRegistrationAccessor::OnPreUnregisterComponent(Level, this);
+	}
+
+	check(ShouldAsyncDestroyPhysicsState());
+	FPhysScene_Chaos* PhysScene = WorldPrivate->GetPhysicsScene();
+	check(PhysScene);
+	check(!IsAsyncCreatePhysicsStateRunning());
+	check(!IsAsyncDestroyPhysicsStateRunning());
+	verify(PhysScene->PushAsyncDestroyPhysicsStateJob(this));
+	PhysicsStateAsyncCreationState = EPhysicsStateAsyncCreationState::Destroying;
+	PhysicsStateAsyncCreationScene = PhysScene;
+	check(IsAsyncDestroyPhysicsStateRunning());
+}
+
 void UActorComponent::UnregisterComponent()
 {
 	SCOPE_CYCLE_COUNTER(STAT_UnregisterComponent);
@@ -1843,7 +1890,7 @@ void UActorComponent::DestroyRenderState_Concurrent()
 
 void UActorComponent::OnCreatePhysicsState()
 {
-	TRACE_CPUPROFILER_EVENT_SCOPE(UActorComponent::OnCreatePhysicsState)
+	TRACE_CPUPROFILER_EVENT_SCOPE(UActorComponent::OnCreatePhysicsState);
 	check(IsRegistered() || IsPreRegistering());
 	check(ShouldCreatePhysicsState());
 	check(WorldPrivate->GetPhysicsScene());
@@ -1853,12 +1900,12 @@ void UActorComponent::OnCreatePhysicsState()
 
 void UActorComponent::OnAsyncCreatePhysicsState()
 {
-	TRACE_CPUPROFILER_EVENT_SCOPE(UActorComponent::OnAsyncCreatePhysicsState)
-	check(AllowsAsyncPhysicsStateCreationAndDestruction());
+	TRACE_CPUPROFILER_EVENT_SCOPE(UActorComponent::OnAsyncCreatePhysicsState);
+	check(AllowsAsyncPhysicsStateCreation());
 	OnCreatePhysicsState();
 }
 
-void UActorComponent::OnAsyncCreatePhysicsStateDone_GameThread()
+void UActorComponent::OnAsyncCreatePhysicsStateEnd_GameThread()
 {
 	check(IsPhysicsStateCreated());
 	check(IsPreRegistering())
@@ -1866,20 +1913,53 @@ void UActorComponent::OnAsyncCreatePhysicsStateDone_GameThread()
 	check(!IsA<UPrimitiveComponent>() || !Cast<UPrimitiveComponent>(this)->DeferredCreatePhysicsStateScene);
 
 	// Reset members used by asynchronous physics state creation
-	AsyncCreatePhysicsStateScene = nullptr;
+	PhysicsStateAsyncCreationState = EPhysicsStateAsyncCreationState::None;
+	PhysicsStateAsyncCreationScene = nullptr;
 	bCallAsyncPhysicsStateCreatedDelegate = true;
 
 	OnPreRegistered();
 	check(IsPreRegistered());
 }
 
+void UActorComponent::OnAsyncDestroyPhysicsStateBegin_GameThread()
+{
+	TRACE_CPUPROFILER_EVENT_SCOPE(UActorComponent::OnAsyncDestroyPhysicsStateBegin_GameThread);
+	GlobalDestroyPhysicsDelegate.Broadcast(this);
+	ensureMsgf(bRegistered, TEXT("Component has physics state when not registered (%s)"), *GetFullName()); // should not have physics state unless we are registered
+}
+
+void UActorComponent::OnAsyncDestroyPhysicsStateEnd_GameThread()
+{
+	check(IsAsyncDestroyPhysicsStateRunning());
+	PhysicsStateAsyncCreationState = EPhysicsStateAsyncCreationState::None;
+	PhysicsStateAsyncCreationScene = nullptr;
+
+	checkf(!HasValidPhysicsState(), TEXT("Failed to destroy physics state (%s)"), *GetFullName());
+	check(bPhysicsStateCreated);
+	bPhysicsStateCreated = false;
+
+	OnPreUnregistered();
+	check(IsPreUnregistered());
+}
+
+void UActorComponent::OnPreUnregistered()
+{
+	check(IsPreUnregistering());
+	RegistrationState = EComponentRegistrationState::PreUnregistered;
+
+	AActor* MyOwner = GetOwner();
+	if (ULevel* Level = MyOwner ? MyOwner->GetLevel() : nullptr)
+	{
+		FLevelRegistrationAccessor::OnPreUnregisteredComponent(Level, this);
+	}
+}
+
 void UActorComponent::OnDestroyPhysicsState()
 {
 	ensure(bPhysicsStateCreated);
 	bPhysicsStateCreated = false;
 }
 
-
 void UActorComponent::CreatePhysicsState(bool bAllowDeferral)
 {
 	LLM_SCOPE(ELLMTag::Chaos);
@@ -1887,7 +1967,7 @@ void UActorComponent::CreatePhysicsState(bool bAllowDeferral)
 	UE_TRACE_METADATA_SCOPE_ASSET_FNAME(NAME_None, NAME_None, GetPackage()->GetFName());
 
 	SCOPE_CYCLE_COUNTER(STAT_ComponentCreatePhysicsState);
-	TRACE_CPUPROFILER_EVENT_SCOPE(UActorComponent::CreatePhysicsState)
+	TRACE_CPUPROFILER_EVENT_SCOPE(UActorComponent::CreatePhysicsState);
 
 	if (!bPhysicsStateCreated && WorldPrivate->GetPhysicsScene() && ShouldCreatePhysicsState())
 	{
@@ -1941,17 +2021,19 @@ void UActorComponent::CreatePhysicsState(bool bAllowDeferral)
 
 void UActorComponent::DestroyPhysicsState()
 {
-	TRACE_CPUPROFILER_EVENT_SCOPE(UActorComponent::DestroyPhysicsState)
+	TRACE_CPUPROFILER_EVENT_SCOPE(UActorComponent::DestroyPhysicsState);
 	SCOPE_CYCLE_COUNTER(STAT_ComponentDestroyPhysicsState);
 
 	if (IsAsyncCreatePhysicsStateRunning())
 	{
+		check(!IsAsyncDestroyPhysicsStateRunning());
 		check(!IsA<UPrimitiveComponent>() || !Cast<UPrimitiveComponent>(this)->DeferredCreatePhysicsStateScene);
-		AsyncCreatePhysicsStateScene->RemoveAsyncPhysicsStateJobs(this);
+		PhysicsStateAsyncCreationScene->RemoveAsyncPhysicsStateJobs(this);
 		
 		// Reset members used by asynchronous physics state creation
-		AsyncCreatePhysicsStateScene = nullptr;
 		bCallAsyncPhysicsStateCreatedDelegate = false;
+		PhysicsStateAsyncCreationState = EPhysicsStateAsyncCreationState::None;
+		PhysicsStateAsyncCreationScene = nullptr;
 		
 		AActor* MyOwner = GetOwner();
 		if (ULevel* Level = MyOwner ? MyOwner->GetLevel() : nullptr)
@@ -1962,6 +2044,7 @@ void UActorComponent::DestroyPhysicsState()
 
 	if (bPhysicsStateCreated)
 	{
+		check(!IsAsyncDestroyPhysicsStateRunning());
 		// Broadcast delegate
 		GlobalDestroyPhysicsDelegate.Broadcast(this);
 
diff --git a/Engine/Source/Runtime/Engine/Private/Components/PrimitiveComponent.cpp b/Engine/Source/Runtime/Engine/Private/Components/PrimitiveComponent.cpp
--- a/Engine/Source/Runtime/Engine/Private/Components/PrimitiveComponent.cpp
+++ b/Engine/Source/Runtime/Engine/Private/Components/PrimitiveComponent.cpp
@@ -924,13 +924,13 @@ void UPrimitiveComponent::OnCreatePhysicsState()
 	}
 }
 
-void UPrimitiveComponent::OnAsyncCreatePhysicsStateDone_GameThread()
+void UPrimitiveComponent::OnAsyncCreatePhysicsStateEnd_GameThread()
 {
 	check(IsPhysicsStateCreated());
 
 	OnComponentPhysicsStateChanged.Broadcast(this, EComponentPhysicsStateChange::Created);
 
-	Super::OnAsyncCreatePhysicsStateDone_GameThread();
+	Super::OnAsyncCreatePhysicsStateEnd_GameThread();
 }
 
 void UPrimitiveComponent::EnsurePhysicsStateCreated()
@@ -1016,6 +1016,54 @@ void UPrimitiveComponent::SendPhysicsTransform(ETeleportType Teleport)
 	BodyInstance.UpdateBodyScale(GetComponentTransform().GetScale3D());
 }
 
+void UPrimitiveComponent::OnAsyncDestroyPhysicsStateBegin_GameThread()
+{
+	TRACE_CPUPROFILER_EVENT_SCOPE(UPrimitiveComponent::OnAsyncDestroyPhysicsStateBegin_GameThread);
+	check(AllowsAsyncPhysicsStateDestruction());
+
+	// we remove welding related to this component
+	UnWeldFromParent();
+	UnWeldChildren();
+
+	// Remove all user defined entities here
+	TArray<Chaos::FPhysicsObject*> PhysicsObjects = GetAllPhysicsObjects();
+	FPhysicsObjectExternalInterface::LockWrite(PhysicsObjects)->SetUserDefinedEntity(PhysicsObjects, nullptr);
+
+	check(!AsyncTermBodyPayload.IsSet());
+	AsyncTermBodyPayload = BodyInstance.StartAsyncTermBody_GameThread();
+	check(!BodyInstance.IsValidBodyInstance());
+	
+	Super::OnAsyncDestroyPhysicsStateBegin_GameThread();
+}
+
+void UPrimitiveComponent::OnAsyncDestroyPhysicsState()
+{
+	TRACE_CPUPROFILER_EVENT_SCOPE(UPrimitiveComponent::OnAsyncDestroyPhysicsState);
+	
+	// We tell the BodyInstance to shut down the physics-engine data.
+	if (ensure(AsyncTermBodyPayload.IsSet()))
+	{
+		FBodyInstance::AsyncTermBody(AsyncTermBodyPayload.GetValue());
+		AsyncTermBodyPayload.Reset();
+	}
+
+	Super::OnAsyncDestroyPhysicsState();
+}
+
+void UPrimitiveComponent::OnAsyncDestroyPhysicsStateEnd_GameThread()
+{
+	Super::OnAsyncDestroyPhysicsStateEnd_GameThread();
+
+	check(!BodyInstance.IsValidBodyInstance());
+	check(!BodyInstance.OwnerComponent.IsValid());
+
+#if UE_ENABLE_DEBUG_DRAWING
+	SendRenderDebugPhysics();
+#endif
+
+	OnComponentPhysicsStateChanged.Broadcast(this, EComponentPhysicsStateChange::Destroyed);
+}
+
 void UPrimitiveComponent::OnDestroyPhysicsState()
 {
 	// we remove welding related to this component
@@ -1728,7 +1776,7 @@ bool UPrimitiveComponent::ShouldCreatePhysicsState() const
 		return false;
 	}
 
-	bool bShouldCreatePhysicsState = (IsRegistered() || (IsPreRegistering() && AllowsAsyncPhysicsStateCreationAndDestruction())) && (bAlwaysCreatePhysicsState || BodyInstance.GetCollisionEnabled() != ECollisionEnabled::NoCollision);
+	bool bShouldCreatePhysicsState = (IsRegistered() || (IsPreRegistering() && AllowsAsyncPhysicsStateCreation())) && (bAlwaysCreatePhysicsState || BodyInstance.GetCollisionEnabled() != ECollisionEnabled::NoCollision);
 
 #if WITH_EDITOR
 	if (BodyInstance.bSimulatePhysics)
diff --git a/Engine/Source/Runtime/Engine/Private/Components/StaticMeshComponent.cpp b/Engine/Source/Runtime/Engine/Private/Components/StaticMeshComponent.cpp
--- a/Engine/Source/Runtime/Engine/Private/Components/StaticMeshComponent.cpp
+++ b/Engine/Source/Runtime/Engine/Private/Components/StaticMeshComponent.cpp
@@ -956,6 +956,15 @@ void UStaticMeshComponent::OnDestroyPhysicsState()
 	FNavigationSystem::UpdateComponentData(*this);
 }
 
+void UStaticMeshComponent::OnAsyncDestroyPhysicsStateEnd_GameThread()
+{
+	TRACE_CPUPROFILER_EVENT_SCOPE(UStaticMeshComponent::OnAsyncDestroyPhysicsStateEnd_GameThread);
+	Super::OnAsyncDestroyPhysicsStateEnd_GameThread();
+
+	bNavigationRelevant = IsNavigationRelevant();
+	FNavigationSystem::UpdateComponentData(*this);
+}
+
 #if WITH_EDITORONLY_DATA
 
 /**
@@ -2308,12 +2317,17 @@ bool UStaticMeshComponent::ShouldCreatePhysicsState() const
 	return Super::ShouldCreatePhysicsState() && GetStaticMesh() && !GetStaticMesh()->IsCompiling();
 }
 
-bool UStaticMeshComponent::AllowsAsyncPhysicsStateCreationAndDestruction() const
+bool UStaticMeshComponent::AllowsAsyncPhysicsStateCreation() const
 {
 	// Async physics state creation is supported for static collision or primitive component not using welding
 	return !BodyInstance.bAutoWeld || !GetAttachParent() || IsStaticPhysics();
 }
 
+bool UStaticMeshComponent::AllowsAsyncPhysicsStateDestruction() const
+{
+	return AllowsAsyncPhysicsStateCreation();
+}
+
 void UStaticMeshComponent::SetStaticMeshInternal(UStaticMesh* NewMesh)
 {
 	StaticMesh = NewMesh;
diff --git a/Engine/Source/Runtime/Engine/Private/InstancedStaticMesh.cpp b/Engine/Source/Runtime/Engine/Private/InstancedStaticMesh.cpp
--- a/Engine/Source/Runtime/Engine/Private/InstancedStaticMesh.cpp
+++ b/Engine/Source/Runtime/Engine/Private/InstancedStaticMesh.cpp
@@ -2247,7 +2247,7 @@ void UInstancedStaticMeshComponent::SendRenderInstanceData_Concurrent()
 
 FBodyInstance* UInstancedStaticMeshComponent::GetBodyInstance(FName BoneName, bool bGetWelded, int32 Index) const
 {
-	if (Index != INDEX_NONE && IsValidInstance(Index))
+	if (Index != INDEX_NONE && InstanceBodies.IsValidIndex(Index))
 	{
 		return const_cast<FBodyInstance*>(InstanceBodies[Index]);
 	}
@@ -2626,6 +2626,7 @@ void UInstancedStaticMeshComponent::CreateAllInstanceBodies()
 
 void UInstancedStaticMeshComponent::ClearAllInstanceBodies()
 {
+	TRACE_CPUPROFILER_EVENT_SCOPE(UInstancedStaticMeshComponent::ClearAllInstanceBodies);
 	QUICK_SCOPE_CYCLE_COUNTER(STAT_UInstancedStaticMeshComponent_ClearAllInstanceBodies);
 	STAT(FScopeCycleCounter Context(StatId);)
 
@@ -2666,6 +2667,7 @@ void UInstancedStaticMeshComponent::OnCreatePhysicsState()
 
 void UInstancedStaticMeshComponent::OnDestroyPhysicsState()
 {
+	TRACE_CPUPROFILER_EVENT_SCOPE(UInstancedStaticMeshComponent::OnDestroyPhysicsState);
 	USceneComponent::OnDestroyPhysicsState();
 
 	// Release all physics representations
@@ -2676,6 +2678,35 @@ void UInstancedStaticMeshComponent::OnDestroyPhysicsState()
 	bNavigationRelevant = IsNavigationRelevant();
 }
 
+void UInstancedStaticMeshComponent::OnAsyncDestroyPhysicsStateBegin_GameThread()
+{
+	TRACE_CPUPROFILER_EVENT_SCOPE(UInstancedStaticMeshComponent::OnAsyncDestroyPhysicsStateBegin_GameThread);
+	USceneComponent::OnAsyncDestroyPhysicsStateBegin_GameThread();
+
+	// Move InstanceBodies in AsyncDestroyPhysicsStatePayload
+	check(AsyncDestroyPhysicsStatePayload.IsEmpty());
+	AsyncDestroyPhysicsStatePayload = MoveTemp(InstanceBodies);
+
+	bNavigationRelevant = IsNavigationRelevant();
+}
+
+void UInstancedStaticMeshComponent::OnAsyncDestroyPhysicsState()
+{
+	TRACE_CPUPROFILER_EVENT_SCOPE(UInstancedStaticMeshComponent::OnAsyncDestroyPhysicsState);
+	USceneComponent::OnAsyncDestroyPhysicsState();
+
+	check(InstanceBodies.IsEmpty());
+	for (FBodyInstance*& Instance : AsyncDestroyPhysicsStatePayload)
+	{
+		if (Instance)
+		{
+			Instance->TermBody();
+			delete Instance;
+		}
+	}
+	AsyncDestroyPhysicsStatePayload.Empty();
+}
+
 Chaos::FPhysicsObject* UInstancedStaticMeshComponent::GetPhysicsObjectById(Chaos::FPhysicsObjectId Id) const
 {
 	if (!InstanceBodies.IsValidIndex(Id) || !InstanceBodies[Id] || !InstanceBodies[Id]->ActorHandle)
@@ -4574,7 +4605,7 @@ TArray<int32> UInstancedStaticMeshComponent::GetInstancesOverlappingBox(const FB
 
 bool UInstancedStaticMeshComponent::ShouldCreatePhysicsState() const
 {
-	return !bDisableCollision && (IsRegistered() || (IsPreRegistering() && AllowsAsyncPhysicsStateCreationAndDestruction())) && !IsBeingDestroyed() && GetStaticMesh() && !GetStaticMesh()->IsCompiling() && (bAlwaysCreatePhysicsState || IsCollisionEnabled());
+	return !bDisableCollision && (IsRegistered() || (IsPreRegistering() && AllowsAsyncPhysicsStateCreation())) && !IsBeingDestroyed() && GetStaticMesh() && !GetStaticMesh()->IsCompiling() && (bAlwaysCreatePhysicsState || IsCollisionEnabled());
 }
 
 float UInstancedStaticMeshComponent::GetTextureStreamingTransformScale() const
@@ -4748,7 +4779,7 @@ void UInstancedStaticMeshComponent::SetupNewInstanceData(FInstancedStaticMeshIns
 	InOutNewInstanceData.Transform = InInstanceTransform.ToMatrixWithScale();
 	PrimitiveInstanceDataManager.TransformChanged(InInstanceIndex);
 
-	if (bPhysicsStateCreated)
+	if (bPhysicsStateCreated && !IsAsyncDestroyPhysicsStateRunning())
 	{
 		if (InInstanceTransform.GetScale3D().IsNearlyZero())
 		{
diff --git a/Engine/Source/Runtime/Engine/Private/Level.cpp b/Engine/Source/Runtime/Engine/Private/Level.cpp
--- a/Engine/Source/Runtime/Engine/Private/Level.cpp
+++ b/Engine/Source/Runtime/Engine/Private/Level.cpp
@@ -82,6 +82,7 @@ Level.cpp: Level-related functions
 #include "UObject/PropertyBagRepository.h"
 #include "WorldPartition/WorldPartitionRuntimeCell.h"
 #include "LevelInstance/LevelInstanceSubsystem.h"
+#include "Components/UnregisterComponentContext.h"
 
 #include UE_INLINE_GENERATED_CPP_BY_NAME(Level)
 
@@ -126,11 +127,17 @@ static FAutoConsoleVariableRef CVarRouteActorInitializationWorkUnitWeighting(
 );
 
 bool GAllowIncrementalPreRegisterComponents = false;
-static FAutoConsoleVariableRef CVarAllowIncrementalPreregisterComponents(
+static FAutoConsoleVariableRef CVarAllowIncrementalPreRegisterComponents(
 	TEXT("s.AllowIncrementalPreRegisterComponents"),
 	GAllowIncrementalPreRegisterComponents,
 	TEXT("Whether to allow level streaming incremental pre-registration of actor components."));
 
+bool GAllowIncrementalPreUnregisterComponents = false;
+static FAutoConsoleVariableRef CVarAllowIncrementalPreUnregisterComponents(
+	TEXT("s.AllowIncrementalPreUnregisterComponents"),
+	GAllowIncrementalPreUnregisterComponents,
+	TEXT("Whether to allow level streaming incremental pre-unregistration of actor components."));
+
 #if WITH_EDITOR
 
 TArray<ULevel::FLevelExternalActorsPathsProviderDelegate> ULevel::LevelExternalActorsPathsProviders;
@@ -438,6 +445,9 @@ ULevel::ULevel( const FObjectInitializer& ObjectInitializer )
 	bIsPartitioned = false;
 	bStaticComponentsRegisteredInStreamingManager = false;
 	IncrementalComponentState = EIncrementalComponentState::Init;
+	IncrementalUnregisterComponentState = EIncrementalUnregisterComponentState::Init;
+	CurrentActorIndexForIncrementalUpdate = 0;
+	CurrentActorIndexForUnregisterComponents = 0;
 }
 
 void ULevel::Initialize(const FURL& InURL)
@@ -2015,41 +2025,162 @@ bool ULevel::IncrementalRunConstructionScripts(bool bProcessAllActors)
 }
 #endif
 
-bool ULevel::IncrementalUnregisterComponents(int32 NumComponentsToUnregister)
+bool ULevel::IncrementalPreUnregisterComponents(FUnregisterComponentContext& Context)
 {
-	// A value of 0 means that we want to unregister all components.
-	if (NumComponentsToUnregister != 0)
+	auto ProcessPreUnregisteringComponents = [this]()
 	{
-		// Only the game can use incremental update functionality.
-		checkf(OwningWorld->IsGameWorld(), TEXT("Cannot call IncrementalUnregisterComponents with non 0 argument in the Editor/ commandlets."));
-	}
+		if (HasPreUnregisteringComponents())
+		{
+			if (FPhysScene* PhysScene = OwningWorld->GetPhysicsScene())
+			{
+				PhysScene->ProcessAsyncPhysicsStateJobs();
+			}
+		}
+	};
+
+	const double StartTime = FPlatformTime::Seconds();
+	const double TimeLimitMs = Context.GetIncrementalRemoveTimeLimitMs();
+	const bool bConsiderTimeLimit = TimeLimitMs >= 0;
+
+	ON_SCOPE_EXIT
+	{
+		if (bConsiderTimeLimit)
+		{
+			const double ElapsedTimeMs = (FPlatformTime::Seconds() - StartTime) * 1000;
+			const double IncrementalUpdateTimeLimitMs = bConsiderTimeLimit ? FMath::Max<double>(0.0, TimeLimitMs - ElapsedTimeMs) : -1;
+			Context.SetIncrementalRemoveTimeLimitMs(IncrementalUpdateTimeLimitMs);
+		}
+	};
+
+	check(OwningWorld);
 
-	// Find next valid actor to process components unregistration
-	int32 NumComponentsUnregistered = 0;
 	while (CurrentActorIndexForUnregisterComponents < Actors.Num())
 	{
 		AActor* Actor = Actors[CurrentActorIndexForUnregisterComponents];
-		if (Actor)
+		CurrentActorIndexForUnregisterComponents++;
+		if (IsValid(Actor))
 		{
-			int32 NumComponents = Actor->GetComponents().Num();
-			NumComponentsUnregistered += NumComponents;
-			Actor->UnregisterAllComponents();
+			FActorPreRegistrationAccessor::IncrementalPreUnregisterComponents(Actor);
 		}
-		CurrentActorIndexForUnregisterComponents++;
-		if (NumComponentsUnregistered > NumComponentsToUnregister )
+
+		if (bConsiderTimeLimit)
 		{
-			break;
+			if (const double ElapseTimeMs = (FPlatformTime::Seconds() - StartTime) * 1000; ElapseTimeMs > TimeLimitMs)
+			{
+				break;
+			}
 		}
 	}
 
-	if (CurrentActorIndexForUnregisterComponents >= Actors.Num())
+	// Process pending preregistering components
+	ProcessPreUnregisteringComponents();
+
+	if (!HasPreUnregisteringComponents() && (CurrentActorIndexForUnregisterComponents >= Actors.Num()))
 	{
+		// Reset CurrentActorIndexForUnregisterComponents for IncrementalUnregisterComponents
 		CurrentActorIndexForUnregisterComponents = 0;
 		return true;
 	}
 	return false;
 }
 
+bool ULevel::IncrementalUnregisterComponents(int32 NumComponentsToUnregister, FUnregisterComponentContext* InContext)
+{
+	// A value of 0 means that we want to unregister all components.
+	const bool bIncrementalUnregisterComponents = NumComponentsToUnregister != 0;
+	if (bIncrementalUnregisterComponents)
+	{
+		// Only the game can use incremental update functionality.
+		checkf(OwningWorld->IsGameWorld(), TEXT("Cannot call IncrementalUnregisterComponents with non 0 argument in the Editor/ commandlets."));
+	}
+
+	auto IsServer = [this]()
+	{
+		const ENetMode NetMode = OwningWorld->GetNetMode();
+		return NetMode == NM_DedicatedServer || NetMode == NM_ListenServer;
+	};
+
+	auto IncrementalUnregisterComponentsInternal = [this, NumComponentsToUnregister]()
+	{
+		// Find next valid actor to process components unregistration
+		int32 NumComponentsUnregistered = 0;
+		while (CurrentActorIndexForUnregisterComponents < Actors.Num())
+		{
+			AActor* Actor = Actors[CurrentActorIndexForUnregisterComponents];
+			if (Actor)
+			{
+				int32 NumComponents = Actor->GetComponents().Num();
+				NumComponentsUnregistered += NumComponents;
+				const bool bForReregister = false;
+				Actor->UnregisterAllComponents(bForReregister);
+			}
+			CurrentActorIndexForUnregisterComponents++;
+			if (NumComponentsUnregistered > NumComponentsToUnregister)
+			{
+				break;
+			}
+		}
+
+		if (CurrentActorIndexForUnregisterComponents >= Actors.Num())
+		{
+			CurrentActorIndexForUnregisterComponents = 0;
+			return true;
+		}
+		return false;
+	};
+
+	FUnregisterComponentContext LocalContext(OwningWorld);
+	FUnregisterComponentContext& Context = InContext ? *InContext : LocalContext;
+	bool bAllowLoop;
+	do
+	{
+		bAllowLoop = false;
+		switch (IncrementalUnregisterComponentState)
+		{
+		case EIncrementalUnregisterComponentState::Init:
+			{
+				// Skip PreUnregisterInitialComponents step if any of these conditions are true
+				// - running on server 
+				// - not running in incremental unregister mode
+				// - asynchronous physics state destruction is not supported
+				// - world is inside a blocking level streaming
+				const bool bIsPreUnregisterInitialComponentsAllowed = bIncrementalUnregisterComponents && 
+					(NumComponentsToUnregister != MAX_int32) &&
+					GAllowIncrementalPreUnregisterComponents &&
+					!IsServer() &&
+					FPhysScene::SupportsAsyncPhysicsStateDestruction() &&
+					!OwningWorld->bRequestedBlockOnAsyncLoading &&
+					!OwningWorld->GetIsInBlockTillLevelStreamingCompleted();
+				// Because pre-unregister is only done during incremental unregister, it is implicitly only used in a game world
+				check(!bIsPreUnregisterInitialComponentsAllowed || OwningWorld->IsGameWorld());
+				IncrementalUnregisterComponentState = bIsPreUnregisterInitialComponentsAllowed ? EIncrementalUnregisterComponentState::PreUnregisterComponents : EIncrementalUnregisterComponentState::UnregisterComponents;
+				bAllowLoop = true;
+			}
+			break;
+		case EIncrementalUnregisterComponentState::PreUnregisterComponents:
+			if (IncrementalPreUnregisterComponents(Context))
+			{
+				IncrementalUnregisterComponentState = EIncrementalUnregisterComponentState::UnregisterComponents;
+				bAllowLoop = true;
+			}
+			break;
+		case EIncrementalUnregisterComponentState::UnregisterComponents:
+			if (IncrementalUnregisterComponentsInternal())
+			{
+				IncrementalUnregisterComponentState = EIncrementalUnregisterComponentState::Finalize;
+				bAllowLoop = true;
+			}
+			break;
+		case EIncrementalUnregisterComponentState::Finalize:
+			IncrementalUnregisterComponentState = EIncrementalUnregisterComponentState::Init;
+			return true;
+			break;
+		}
+	} while (bAllowLoop);
+
+	return false;
+}
+
 void ULevel::MarkLevelComponentsRenderStateDirty()
 {
 	for (UModelComponent* ModelComponent : ModelComponents)
diff --git a/Engine/Source/Runtime/Engine/Private/PhysicsEngine/BodyInstance.cpp b/Engine/Source/Runtime/Engine/Private/PhysicsEngine/BodyInstance.cpp
--- a/Engine/Source/Runtime/Engine/Private/PhysicsEngine/BodyInstance.cpp
+++ b/Engine/Source/Runtime/Engine/Private/PhysicsEngine/BodyInstance.cpp
@@ -324,6 +324,12 @@ bool FCollisionResponse::operator==(const FCollisionResponse& Other) const
 
 	return bCollisionResponseEqual;
 }
+
+FBodyInstance::FAsyncTermBodyPayload::FAsyncTermBodyPayload()
+: ActorHandle(DefaultPhysicsActorHandle())
+, bNotifyRigidBodyCollision(false)
+{
+}
 ////////////////////////////////////////////////////////////////////////////
 
 
@@ -1742,6 +1748,60 @@ const FTransform& FBodyInstance::GetRelativeBodyTransform(const FPhysicsShapeHan
 	return Result ? Result->RelativeTM : FTransform::Identity;
 }
 
+FBodyInstance::FAsyncTermBodyPayload FBodyInstance::StartAsyncTermBody_GameThread()
+{
+	TRACE_CPUPROFILER_EVENT_SCOPE(FBodyInstance::StartAsyncTermBody_GameThread);
+	SCOPE_CYCLE_COUNTER(STAT_TermBody);
+
+	// Prepare payload for async task before resetting the ActorHandle/OwnerComponent
+	FAsyncTermBodyPayload Payload;
+	Payload.ActorHandle = ActorHandle;
+	Payload.OwnerComponent = OwnerComponent;
+	Payload.bNotifyRigidBodyCollision = bNotifyRigidBodyCollision;
+	ActorHandle = nullptr;
+	OwnerComponent = nullptr;
+
+	CurrentSceneState = BodyInstanceSceneState::NotAdded;
+	BodySetup = NULL;
+	ExternalCollisionProfileBodySetup = nullptr;
+
+	if (DOFConstraint)
+	{
+		DOFConstraint->TermConstraint();
+		FConstraintInstance::Free(DOFConstraint);
+		DOFConstraint = NULL;
+	}
+	return Payload;
+}
+
+void FBodyInstance::AsyncTermBody(FAsyncTermBodyPayload& Payload)
+{
+	SCOPE_CYCLE_COUNTER(STAT_TermBody);
+
+	if (UPrimitiveComponent* PrimComp = Payload.OwnerComponent.Get())
+	{
+		if (FPhysScene_Chaos* PhysScene = PrimComp->GetWorld()->GetPhysicsScene())
+		{
+			if (FPhysicsInterface::IsValid(Payload.ActorHandle))
+			{
+				PhysScene->RemoveFromComponentMaps(Payload.ActorHandle);
+			}
+			if (Payload.bNotifyRigidBodyCollision)
+			{
+				PhysScene->UnRegisterForCollisionEvents(PrimComp);
+			}
+		}
+	}
+
+	if (FPhysicsInterface::IsValid(Payload.ActorHandle))
+	{
+		FPhysScene* PhysScene = static_cast<FPhysScene*>(FPhysicsInterface::GetCurrentScene(Payload.ActorHandle));
+		FPhysicsInterface::ReleaseActor(Payload.ActorHandle, PhysScene);
+	}
+
+	check(!FPhysicsInterface::IsValid(Payload.ActorHandle));
+}
+
 /**
  *	Clean up the physics engine info for this instance.
  */
@@ -1769,7 +1829,6 @@ void FBodyInstance::TermBody(bool bNeverDeferRelease)
 		FPhysicsInterface::ReleaseActor(ActorHandle, GetPhysicsScene(), bNeverDeferRelease);
 	}
 
-
 	// @TODO: Release spring body here
 
 	CurrentSceneState = BodyInstanceSceneState::NotAdded;
@@ -1781,9 +1840,8 @@ void FBodyInstance::TermBody(bool bNeverDeferRelease)
 	{
 		DOFConstraint->TermConstraint();
 		FConstraintInstance::Free(DOFConstraint);
-			DOFConstraint = NULL;
+		DOFConstraint = NULL;
 	}
-	
 }
 
 bool FBodyInstance::Weld(FBodyInstance* TheirBody, const FTransform& TheirTM)
diff --git a/Engine/Source/Runtime/Engine/Private/PhysicsEngine/Experimental/PhysScene_AsyncPhysicsStateJobQueue.cpp b/Engine/Source/Runtime/Engine/Private/PhysicsEngine/Experimental/PhysScene_AsyncPhysicsStateJobQueue.cpp
--- a/Engine/Source/Runtime/Engine/Private/PhysicsEngine/Experimental/PhysScene_AsyncPhysicsStateJobQueue.cpp
+++ b/Engine/Source/Runtime/Engine/Private/PhysicsEngine/Experimental/PhysScene_AsyncPhysicsStateJobQueue.cpp
@@ -7,6 +7,14 @@
 #include "PhysicsEngine/BodySetup.h"
 #include "Engine/World.h"
 
+FPhysScene_AsyncPhysicsStateJobQueue::~FPhysScene_AsyncPhysicsStateJobQueue()
+{
+	Launch();
+	WaitForAsyncTasksExecution();
+	OnAsyncTasksExecuted();
+	check(IsCompleted());
+}
+
 void FPhysScene_AsyncPhysicsStateJobQueue::AddJob(const FJob& Job)
 {
 	TRACE_CPUPROFILER_EVENT_SCOPE(FPhysScene_AsyncPhysicsStateJobQueue::AddJob);
@@ -30,17 +38,44 @@ void FPhysScene_AsyncPhysicsStateJobQueue::RemoveJob(const FJob& Job)
 		return;
 	}
 
-	WaitForCompletion();
+	WaitForAsyncTasksExecution();
 	
 	for (FJobs& Jobs : PipedJobs)
 	{
 		if (Jobs->Remove(Job))
 		{
+			Job.OnPostExecute_GameThread();
 			break;
 		}
 	}
 }
 
+void FPhysScene_AsyncPhysicsStateJobQueue::OnAsyncTasksExecuted()
+{
+	check(AreAsyncTasksExecuted());
+	for (const FJobs& Jobs : PipedJobs)
+	{
+		Jobs->OnPostExecute_GameThread();
+	}
+	PipedJobs.Reset();
+	PipeTasks.Reset();
+}
+
+bool FPhysScene_AsyncPhysicsStateJobQueue::AreAsyncTasksExecuted() const
+{
+	return !PendingJobs.IsValid() && !Pipe.HasWork();
+}
+
+bool FPhysScene_AsyncPhysicsStateJobQueue::IsCompleted() const
+{ 
+	if (!PendingJobs.IsValid() && PipedJobs.IsEmpty() && PipeTasks.IsEmpty())
+	{
+		check(AreAsyncTasksExecuted());
+		return true;
+	}
+	return false;
+}
+
 void FPhysScene_AsyncPhysicsStateJobQueue::Launch()
 {
 	TRACE_CPUPROFILER_EVENT_SCOPE(FPhysScene_AsyncPhysicsStateJobQueue::Launch);
@@ -60,9 +95,9 @@ void FPhysScene_AsyncPhysicsStateJobQueue::Launch()
 	PipeTasks.Emplace(Pipe.Launch(UE_SOURCE_LOCATION, [this, NewJobs = NewJobs.Get()]() { NewJobs->Execute(); }, UE::Tasks::ETaskPriority::BackgroundHigh));
 }
 
-void FPhysScene_AsyncPhysicsStateJobQueue::WaitForCompletion()
+void FPhysScene_AsyncPhysicsStateJobQueue::WaitForAsyncTasksExecution()
 {
-	TRACE_CPUPROFILER_EVENT_SCOPE(FPhysScene_AsyncTaskQueue::WaitForCompletion);
+	TRACE_CPUPROFILER_EVENT_SCOPE(FPhysScene_AsyncTaskQueue::WaitForAsyncTasksExecution);
 	check(IsInGameThread());
 
 	UE::Tasks::Wait(PipeTasks);
@@ -79,20 +114,12 @@ void FPhysScene_AsyncPhysicsStateJobQueue::Tick()
 	// Wait for tasks if the world is inside a blocking load
 	if (PhysScene->GetOwningWorld()->GetIsInBlockTillLevelStreamingCompleted())
 	{
-		WaitForCompletion();
+		WaitForAsyncTasksExecution();
 	}
 
-	if (IsCompleted())
+	if (AreAsyncTasksExecuted())
 	{
-		TRACE_CPUPROFILER_EVENT_SCOPE(FPhysScene_AsyncPhysicsStateJobQueue::OnCompleted);
-		for (const FJobs& Jobs : PipedJobs)
-		{
-			Jobs->OnPostExecute_GameThread();
-			
-		}
-		PipedJobs.Reset();
-		PipeTasks.Reset();
-		check(!PendingJobs.IsValid());
+		OnAsyncTasksExecuted();
 	}
 }
 
@@ -107,30 +134,50 @@ void FPhysScene_AsyncPhysicsStateJobQueue::FJobSet::Execute() const
 
 void FPhysScene_AsyncPhysicsStateJobQueue::FJobSet::OnPreExecute_GameThread() const
 {
-	TSet<UBodySetup*> UniqueBodySetups;
-
-	for (const FJob& Job : Jobs)
+	// Pre-pass to make sure physics meshes are created
 	{
-		if (Job.Type == EJobType::CreatePhysicsState)
+		TSet<UBodySetup*> UniqueBodySetups;
+		for (const FJob& Job : Jobs)
 		{
-			UPrimitiveComponent* PrimitiveComponent = Cast<UPrimitiveComponent>(Job.Component);
-			UBodySetup* BodySetup = PrimitiveComponent ? PrimitiveComponent->GetBodySetup() : nullptr;
-			if (BodySetup && !BodySetup->bCreatedPhysicsMeshes)
+			if (Job.Type == EJobType::CreatePhysicsState)
 			{
-				UniqueBodySetups.Add(BodySetup);
+				UPrimitiveComponent* PrimitiveComponent = Cast<UPrimitiveComponent>(Job.Component);
+				UBodySetup* BodySetup = PrimitiveComponent ? PrimitiveComponent->GetBodySetup() : nullptr;
+				if (BodySetup && !BodySetup->bCreatedPhysicsMeshes)
+				{
+					UniqueBodySetups.Add(BodySetup);
+				}
 			}
 		}
+		if (UniqueBodySetups.Num())
+		{
+			TArray<UBodySetup*> BodySetups = UniqueBodySetups.Array();
+			ParallelFor(TEXT("CreatePhysicsMeshes.PF"), BodySetups.Num(), 1, [this, &BodySetups](int32 Index) { BodySetups[Index]->CreatePhysicsMeshes(); });
+		}
 	}
 
-	if (UniqueBodySetups.Num())
+	for (const FJob& Job : Jobs)
+	{
+		Job.OnPreExecute_GameThread();
+	}
+}
+
+void FPhysScene_AsyncPhysicsStateJobQueue::FJob::OnPreExecute_GameThread() const
+{
+	TRACE_CPUPROFILER_EVENT_SCOPE(FPhysScene_AsyncPhysicsStateJobQueue::OnPreExecute_GameThread);
+	switch (Type)
 	{
-		TArray<UBodySetup*> BodySetups = UniqueBodySetups.Array();
-		ParallelFor(TEXT("CreatePhysicsMeshes.PF"), BodySetups.Num(), 1, [this, &BodySetups](int32 Index) { BodySetups[Index]->CreatePhysicsMeshes(); });
+		case EJobType::DestroyPhysicsState:
+		{
+			Component->OnAsyncDestroyPhysicsStateBegin_GameThread();
+		}
+		break;
 	}
 }
 
 void FPhysScene_AsyncPhysicsStateJobQueue::FJobSet::OnPostExecute_GameThread() const
 {
+	TRACE_CPUPROFILER_EVENT_SCOPE(FPhysScene_AsyncPhysicsStateJobQueue::OnPostExecute_GameThread);
 	for (const FJob& Job : Jobs)
 	{
 		Job.OnPostExecute_GameThread();
@@ -149,6 +196,12 @@ void FPhysScene_AsyncPhysicsStateJobQueue::FJob::Execute() const
 				Component->OnAsyncCreatePhysicsState();
 			}
 			break;
+			case EJobType::DestroyPhysicsState:
+			{
+				TRACE_CPUPROFILER_EVENT_SCOPE(FPhysScene_AsyncPhysicsStateJobQueue::DestroyPhysicsState);
+				Component->OnAsyncDestroyPhysicsState();
+			}
+			break;
 		}
 	}
 }
@@ -161,7 +214,12 @@ void FPhysScene_AsyncPhysicsStateJobQueue::FJob::OnPostExecute_GameThread() cons
 		{
 			case EJobType::CreatePhysicsState:
 			{
-				Component->OnAsyncCreatePhysicsStateDone_GameThread();
+				Component->OnAsyncCreatePhysicsStateEnd_GameThread();
+			}
+			break;
+			case EJobType::DestroyPhysicsState:
+			{
+				Component->OnAsyncDestroyPhysicsStateEnd_GameThread();
 			}
 			break;
 		}
diff --git a/Engine/Source/Runtime/Engine/Private/PhysicsEngine/Experimental/PhysScene_Chaos.cpp b/Engine/Source/Runtime/Engine/Private/PhysicsEngine/Experimental/PhysScene_Chaos.cpp
--- a/Engine/Source/Runtime/Engine/Private/PhysicsEngine/Experimental/PhysScene_Chaos.cpp
+++ b/Engine/Source/Runtime/Engine/Private/PhysicsEngine/Experimental/PhysScene_Chaos.cpp
@@ -19,6 +19,7 @@
 #include "PhysicsEngine/SkeletalBodySetup.h"
 #include "Chaos/PhysicsObjectInternalInterface.h"
 #include "Physics/Experimental/ChaosEventRelay.h"
+#include "Physics/Experimental/PhysScene_AsyncPhysicsStateJobQueue.h"
 #include "EngineUtils.h"
 
 #include "ChaosSolversModule.h"
@@ -1756,13 +1757,18 @@ bool FPhysScene_Chaos::SupportsAsyncPhysicsStateCreation()
 	return Chaos::CVars::bEnableAsyncInitBody;
 }
 
+bool FPhysScene_Chaos::SupportsAsyncPhysicsStateDestruction()
+{
+	return SupportsAsyncPhysicsStateCreation();
+}
+
 bool FPhysScene_Chaos::PushAsyncCreatePhysicsStateJob(UActorComponent* Component)
 {
 	if (SupportsAsyncPhysicsStateCreation())
 	{
 		check(IsInGameThread());
 		check(!Component->IsPhysicsStateCreated());
-		check(Component->AllowsAsyncPhysicsStateCreationAndDestruction());
+		check(Component->AllowsAsyncPhysicsStateCreation());
 
 		if (!AsyncPhysicsStateJobQueue.IsValid())
 		{
@@ -1775,30 +1781,43 @@ bool FPhysScene_Chaos::PushAsyncCreatePhysicsStateJob(UActorComponent* Component
 	return false;
 }
 
-void FPhysScene_Chaos::RemoveAsyncPhysicsStateJobs(UActorComponent* Component)
+bool FPhysScene_Chaos::PushAsyncDestroyPhysicsStateJob(UActorComponent* Component)
 {
-	if (SupportsAsyncPhysicsStateCreation())
+	if (SupportsAsyncPhysicsStateDestruction())
 	{
 		check(IsInGameThread());
-		if (AsyncPhysicsStateJobQueue.IsValid())
+		check(Component->IsPhysicsStateCreated());
+		check(Component->AllowsAsyncPhysicsStateDestruction());
+
+		if (!AsyncPhysicsStateJobQueue.IsValid())
 		{
-			AsyncPhysicsStateJobQueue->RemoveJob({ Component, FPhysScene_AsyncPhysicsStateJobQueue::EJobType::CreatePhysicsState });
+			AsyncPhysicsStateJobQueue = MakeUnique<FPhysScene_AsyncPhysicsStateJobQueue>(this);
 		}
+
+		AsyncPhysicsStateJobQueue->AddJob({ Component, FPhysScene_AsyncPhysicsStateJobQueue::EJobType::DestroyPhysicsState });
+	}
+	return true;
+}
+
+void FPhysScene_Chaos::RemoveAsyncPhysicsStateJobs(UActorComponent* Component)
+{
+	check(IsInGameThread());
+	if (AsyncPhysicsStateJobQueue.IsValid())
+	{
+		AsyncPhysicsStateJobQueue->RemoveJob({ Component, FPhysScene_AsyncPhysicsStateJobQueue::EJobType::CreatePhysicsState });
+		AsyncPhysicsStateJobQueue->RemoveJob({ Component, FPhysScene_AsyncPhysicsStateJobQueue::EJobType::DestroyPhysicsState });
 	}
 }
 
 void FPhysScene_Chaos::ProcessAsyncPhysicsStateJobs()
 {
-	if (SupportsAsyncPhysicsStateCreation())
+	check(IsInGameThread());
+	if (AsyncPhysicsStateJobQueue.IsValid())
 	{
-		check(IsInGameThread());
-		if (AsyncPhysicsStateJobQueue.IsValid())
+		AsyncPhysicsStateJobQueue->Tick();
+		if (AsyncPhysicsStateJobQueue->IsCompleted())
 		{
-			AsyncPhysicsStateJobQueue->Tick();
-			if (AsyncPhysicsStateJobQueue->IsCompleted())
-			{
-				AsyncPhysicsStateJobQueue.Reset();
-			}
+			AsyncPhysicsStateJobQueue.Reset();
 		}
 	}
 }
diff --git a/Engine/Source/Runtime/Engine/Private/World.cpp b/Engine/Source/Runtime/Engine/Private/World.cpp
--- a/Engine/Source/Runtime/Engine/Private/World.cpp
+++ b/Engine/Source/Runtime/Engine/Private/World.cpp
@@ -82,6 +82,7 @@
 #include "Streaming/ServerStreamingLevelsVisibility.h"
 #include "Streaming/LevelStreamingDelegates.h"
 #include "Streaming/StreamingWorldSubsystemInterface.h"
+#include "Components/UnregisterComponentContext.h"
 
 #if WITH_EDITOR
 	#include "DerivedDataCacheInterface.h"
@@ -3575,7 +3576,7 @@ void UWorld::RemoveFromWorld( ULevel* Level, bool bAllowIncrementalRemoval, FNet
 	SCOPE_CYCLE_COUNTER(STAT_RemoveFromWorldTime);
 	CSV_SCOPED_TIMING_STAT_EXCLUSIVE(RemoveFromWorld);
 
-	FScopeCycleCounterUObject Context(Level);
+	FScopeCycleCounterUObject ScopeContext(Level);
 	check(IsValid(Level));
 	check(!Level->IsUnreachable());
 
@@ -3650,6 +3651,9 @@ void UWorld::RemoveFromWorld( ULevel* Level, bool bAllowIncrementalRemoval, FNet
 				double TimeLimit = FMath::Max<double>(0.0, GLevelStreamingUnregisterComponentsTimeLimit - GRemoveFromWorldUnregisterComponentTimeCumul);
 				if (TimeLimit > 0.0)
 				{
+					FUnregisterComponentContext Context = FUnregisterComponentContext(this)
+						.SetIncrementalRemoveTimeLimitMs(TimeLimit);
+
 					double StartTime = FPlatformTime::Seconds();
 					// Incrementally unregister actor components. 
 					// This avoids spikes on the renderthread and gamethread when we subsequently call ClearLevelComponents() further down
@@ -3657,13 +3661,17 @@ void UWorld::RemoveFromWorld( ULevel* Level, bool bAllowIncrementalRemoval, FNet
 					int32 NumComponentsToUnregister = GLevelStreamingComponentsUnregistrationGranularity;
 					do
 					{
-						if (Level->IncrementalUnregisterComponents(NumComponentsToUnregister))
+						if (Level->IncrementalUnregisterComponents(NumComponentsToUnregister, &Context))
 						{
 							// We're done, so the level can be removed
 							CurrentLevelPendingInvisibility = nullptr;
 							bFinishRemovingLevel = true;
 							break;
 						}
+						else if (FLevelRegistrationAccessor::HasPreUnregisteringComponents(Level))
+						{
+							break;
+						}
 					} while (!IsTimeLimitExceeded(TEXT("unregistering components"), StartTime, Level, TimeLimit));
 
 					double DeltaTime = (FPlatformTime::Seconds() - StartTime) * 1000;
diff --git a/Engine/Source/Runtime/Engine/Public/Physics/Experimental/PhysScene_AsyncPhysicsStateJobQueue.h b/Engine/Source/Runtime/Engine/Public/Physics/Experimental/PhysScene_AsyncPhysicsStateJobQueue.h
--- a/Engine/Source/Runtime/Engine/Public/Physics/Experimental/PhysScene_AsyncPhysicsStateJobQueue.h
+++ b/Engine/Source/Runtime/Engine/Public/Physics/Experimental/PhysScene_AsyncPhysicsStateJobQueue.h
@@ -16,6 +16,8 @@ struct FPhysScene_AsyncPhysicsStateJobQueue
 		: PhysScene(InPhysicScene)
 	{}
 
+	~FPhysScene_AsyncPhysicsStateJobQueue();
+
 	enum class EJobType
 	{
 		CreatePhysicsState,
@@ -27,11 +29,11 @@ struct FPhysScene_AsyncPhysicsStateJobQueue
 		FJob(UActorComponent* InComponent, EJobType InType)
 			: Component(InComponent)
 			, Type(InType)
-			, bIsCancelled(false)
 		{}
 
 		bool IsValid() const { return ::IsValid(Component); }
 		void Execute() const;
+		void OnPreExecute_GameThread() const;
 		void OnPostExecute_GameThread() const;
 
 		inline bool operator == (const FJob& Other) const
@@ -46,14 +48,24 @@ struct FPhysScene_AsyncPhysicsStateJobQueue
 
 		UActorComponent* Component;
 		EJobType Type;
-		bool bIsCancelled;
 	};
 
 	struct FJobSet
 	{
-		void Add(const FJob& Job) { Jobs.Add(Job); }
-		int32 Remove(const FJob& Job) { return Jobs.Remove(Job); }
-		bool IsEmpty() const { return Jobs.IsEmpty(); }
+		void Add(const FJob& Job)
+		{
+			Jobs.Add(Job);
+		}
+
+		int32 Remove(const FJob& Job)
+		{
+			return Jobs.Remove(Job);
+		}
+
+		bool IsEmpty() const
+		{
+			return Jobs.IsEmpty();
+		}
 
 		void OnPreExecute_GameThread() const;
 		void Execute() const;
@@ -66,12 +78,14 @@ struct FPhysScene_AsyncPhysicsStateJobQueue
 	void Tick();
 	void AddJob(const FJob& Job);
 	void RemoveJob(const FJob& Job);
-	bool IsCompleted() const { return !Pipe.HasWork(); }
+	bool IsCompleted() const;
 
 private:
 
 	void Launch();
-	void WaitForCompletion();
+	void WaitForAsyncTasksExecution();
+	bool AreAsyncTasksExecuted() const;
+	void OnAsyncTasksExecuted();
 
 	FPhysScene* PhysScene;
 	using FJobs = TUniquePtr<FJobSet>;
diff --git a/Engine/Source/Runtime/Engine/Public/Physics/Experimental/PhysScene_Chaos.h b/Engine/Source/Runtime/Engine/Public/Physics/Experimental/PhysScene_Chaos.h
--- a/Engine/Source/Runtime/Engine/Public/Physics/Experimental/PhysScene_Chaos.h
+++ b/Engine/Source/Runtime/Engine/Public/Physics/Experimental/PhysScene_Chaos.h
@@ -7,7 +7,6 @@
 #include "EventsData.h"
 #include "Physics/PhysScene.h"
 #include "Physics/Experimental/ChaosEventType.h"
-#include "Physics/Experimental/PhysScene_AsyncPhysicsStateJobQueue.h"
 #include "GameFramework/Actor.h"
 #include "PhysicsPublic.h"
 #include "PhysInterface_Chaos.h"
@@ -47,6 +46,7 @@ class AWorldSettings;
 class FPhysicsReplicationFactory;
 class FContactModifyCallbackFactory;
 struct FConstraintInstanceBase;
+struct FPhysScene_AsyncPhysicsStateJobQueue;
 
 namespace Chaos
 {
@@ -263,7 +263,9 @@ public:
 	ENGINE_API void ProcessDeferredCreatePhysicsState();
 
 	ENGINE_API static bool SupportsAsyncPhysicsStateCreation();
+	ENGINE_API static bool SupportsAsyncPhysicsStateDestruction();
 	ENGINE_API bool PushAsyncCreatePhysicsStateJob(UActorComponent* Component);
+	ENGINE_API bool PushAsyncDestroyPhysicsStateJob(UActorComponent* Component);
 	ENGINE_API void RemoveAsyncPhysicsStateJobs(UActorComponent* Component);
 	ENGINE_API void ProcessAsyncPhysicsStateJobs();
 
