MGPU:  Fix for editor startup crash with more than two GPUs.  Particle related allocator for cross GPU transfers was using a TFixedAllocator hard coded to only support one extra GPU -- switched to TInlineAllocator.

#rnx

diff --git a/Engine/Source/Runtime/Engine/Private/Particles/ParticleGpuSimulation.cpp b/Engine/Source/Runtime/Engine/Private/Particles/ParticleGpuSimulation.cpp
--- a/Engine/Source/Runtime/Engine/Private/Particles/ParticleGpuSimulation.cpp
+++ b/Engine/Source/Runtime/Engine/Private/Particles/ParticleGpuSimulation.cpp
@@ -4771,7 +4771,7 @@ void FFXSystem::SimulateGPUParticles(
 
 
 #if WITH_MGPU
-	TArray<FTransferResourceParams, TFixedAllocator<4>> CrossGPUTransferResources;
+	TArray<FTransferResourceParams, TInlineAllocator<4>> CrossGPUTransferResources;
 	const bool bCrossTransferEnabled = GNumExplicitGPUsForRendering > 1;
 	auto AddCrossGPUTransferResource =
 		[&](FRHITexture* TextureToTransfer)
