Add versioning to Streaming Audio Seektables

#rnx

diff --git a/Engine/Source/Runtime/Engine/Classes/Sound/StreamedAudioChunkSeekTable.h b/Engine/Source/Runtime/Engine/Classes/Sound/StreamedAudioChunkSeekTable.h
--- a/Engine/Source/Runtime/Engine/Classes/Sound/StreamedAudioChunkSeekTable.h
+++ b/Engine/Source/Runtime/Engine/Classes/Sound/StreamedAudioChunkSeekTable.h
@@ -23,8 +23,13 @@ enum class EChunkSeekTableMode : uint8
  */
 class FStreamedAudioChunkSeekTable
 {
-public:
-	
+public:	
+	/**
+	 * Gets the current version of Seek table. This can be used in the DDC key.
+	 * @return int16 Seek-table version
+	 */
+	static ENGINE_API int16 GetVersion();
+
 	/**
 	 * Parse a seek-table from memory.
 	 *
diff --git a/Engine/Source/Runtime/Engine/Private/AudioDerivedData.cpp b/Engine/Source/Runtime/Engine/Private/AudioDerivedData.cpp
--- a/Engine/Source/Runtime/Engine/Private/AudioDerivedData.cpp
+++ b/Engine/Source/Runtime/Engine/Private/AudioDerivedData.cpp
@@ -242,6 +242,7 @@ static void GetStreamedAudioDerivedDataKeySuffix(
 	)
 {
 	uint16 Version = 0;
+	bool bFormatUsingStreamingSeekTables = false;
 
 	// get the version for this soundwave's platform format
 	ITargetPlatformManagerModule* TPM = GetTargetPlatformManager();
@@ -251,6 +252,7 @@ static void GetStreamedAudioDerivedDataKeySuffix(
 		if (AudioFormat)
 		{
 			Version = AudioFormat->GetVersion(AudioFormatName);
+			bFormatUsingStreamingSeekTables = AudioFormat->RequiresStreamingSeekTable();
 		}
 	}
 
@@ -262,6 +264,13 @@ static void GetStreamedAudioDerivedDataKeySuffix(
 		FPlatformAudioCookOverrides::GetHashSuffix(CompressionOverrides, AudioFormatNameString);
 	}
 
+	// If format uses streaming seek-tables append the version into the hash.
+	// Initial version '0' won't change the key.
+	if (const int16 SeekTableVersion = FStreamedAudioChunkSeekTable::GetVersion(); SeekTableVersion != 0 && bFormatUsingStreamingSeekTables)
+	{
+		AudioFormatNameString += FString::Printf(TEXT("_STVER=%d"), SeekTableVersion);
+	}
+
 	const FString SoundWaveHash = GetSoundWaveHash(SoundWave, InTargetPlatform);
 		
 	// build the key
diff --git a/Engine/Source/Runtime/Engine/Private/StreamedAudioChunkSeekTable.cpp b/Engine/Source/Runtime/Engine/Private/StreamedAudioChunkSeekTable.cpp
--- a/Engine/Source/Runtime/Engine/Private/StreamedAudioChunkSeekTable.cpp
+++ b/Engine/Source/Runtime/Engine/Private/StreamedAudioChunkSeekTable.cpp
@@ -31,6 +31,12 @@ namespace StreamedAudioChunkSeekTable_Private
 	}
 }
 
+int16 FStreamedAudioChunkSeekTable::GetVersion()
+{
+	static constexpr int16 Version = 0;		
+	return Version;
+}
+
 void FStreamedAudioChunkSeekTable::SetMode(EChunkSeekTableMode InMode)
 {
 	Mode = InMode;
