[ShaderCompiler]
 - Fix stats atomicity, since they are modified by multiple threads

#rnx

diff --git a/Engine/Source/Runtime/Engine/Private/ShaderCompiler/ShaderCompilerJobCache.cpp b/Engine/Source/Runtime/Engine/Private/ShaderCompiler/ShaderCompilerJobCache.cpp
--- a/Engine/Source/Runtime/Engine/Private/ShaderCompiler/ShaderCompilerJobCache.cpp
+++ b/Engine/Source/Runtime/Engine/Private/ShaderCompiler/ShaderCompilerJobCache.cpp
@@ -127,13 +127,13 @@ static FAutoConsoleVariableRef CVarShaderCompilerDebugDiscardCacheOutputs(
 	ECVF_Default);
 
 
-TRACE_DECLARE_INT_COUNTER(Shaders_JobCacheSearchAttempts, TEXT("Shaders/JobCache/SearchAttempts"));
-TRACE_DECLARE_INT_COUNTER(Shaders_JobCacheHits, TEXT("Shaders/JobCache/Hits"));
+TRACE_DECLARE_ATOMIC_INT_COUNTER(Shaders_JobCacheSearchAttempts, TEXT("Shaders/JobCache/SearchAttempts"));
+TRACE_DECLARE_ATOMIC_INT_COUNTER(Shaders_JobCacheHits, TEXT("Shaders/JobCache/Hits"));
 
-TRACE_DECLARE_INT_COUNTER(Shaders_JobCacheDDCRequests, TEXT("Shaders/JobCache/DDCRequests"));
-TRACE_DECLARE_INT_COUNTER(Shaders_JobCacheDDCHits, TEXT("Shaders/JobCache/DDCHits"));
-TRACE_DECLARE_MEMORY_COUNTER(Shaders_JobCacheDDCBytesReceived, TEXT("Shaders/JobCache/DDCBytesRecieved"));
-TRACE_DECLARE_MEMORY_COUNTER(Shaders_JobCacheDDCBytesSent, TEXT("Shaders/JobCache/DDCBytesSent"));
+TRACE_DECLARE_ATOMIC_INT_COUNTER(Shaders_JobCacheDDCRequests, TEXT("Shaders/JobCache/DDCRequests"));
+TRACE_DECLARE_ATOMIC_INT_COUNTER(Shaders_JobCacheDDCHits, TEXT("Shaders/JobCache/DDCHits"));
+TRACE_DECLARE_ATOMIC_MEMORY_COUNTER(Shaders_JobCacheDDCBytesReceived, TEXT("Shaders/JobCache/DDCBytesRecieved"));
+TRACE_DECLARE_ATOMIC_MEMORY_COUNTER(Shaders_JobCacheDDCBytesSent, TEXT("Shaders/JobCache/DDCBytesSent"));
 
 #if WITH_EDITOR
 namespace
