[HWRT] Fix tracking of TotalResidentSize when changes to FRayTracingGeometry initializer cause size to change.

Example:
FRayTracingDynamicGeometryCollection modifies Initializer and recreates FRayTracingGeometry when parameters such as number of triangles change.

diff --git a/Engine/Source/Runtime/Engine/Private/Rendering/RayTracingGeometryManager.cpp b/Engine/Source/Runtime/Engine/Private/Rendering/RayTracingGeometryManager.cpp
--- a/Engine/Source/Runtime/Engine/Private/Rendering/RayTracingGeometryManager.cpp
+++ b/Engine/Source/Runtime/Engine/Private/Rendering/RayTracingGeometryManager.cpp
@@ -476,7 +476,7 @@ void FRayTracingGeometryManager::RefreshRegisteredGeometry(RayTracingGeometryHan
 
 			if (NumRemoved > 0)
 			{
-				TotalResidentSize -= RegisteredGeometry.Size;
+				TotalResidentSize -= OldSize;
 			}
 
 			EvictableGeometries.Remove(Handle);
