Seektable version stuff (more of original CL)

#rnx

diff --git a/Engine/Source/Runtime/Engine/Classes/Sound/StreamedAudioChunkSeekTable.h b/Engine/Source/Runtime/Engine/Classes/Sound/StreamedAudioChunkSeekTable.h
--- a/Engine/Source/Runtime/Engine/Classes/Sound/StreamedAudioChunkSeekTable.h
+++ b/Engine/Source/Runtime/Engine/Classes/Sound/StreamedAudioChunkSeekTable.h
@@ -19,7 +19,7 @@ enum class EChunkSeekTableMode : uint8
 };
 
 /**
- * Class representing a audio seek-table. Typically embedded in a bit-stream.
+ * Class representing an audio seek-table. Typically embedded in a bit-stream.
  */
 class FStreamedAudioChunkSeekTable
 {
@@ -39,7 +39,7 @@ public:
 	 * @param OutTable Table that's loaded
 	 * @return true on Success, false on Failure
 	 */
-	static bool Parse(const uint8* InMemory, uint32 InSize, uint32& InOutOffset, FStreamedAudioChunkSeekTable& OutTable);
+	static ENGINE_API bool Parse(const uint8* InMemory, uint32 InSize, uint32& InOutOffset, FStreamedAudioChunkSeekTable& OutTable);
 
 	/**
 	 * Calculate the size needed for the table. 
@@ -49,8 +49,15 @@ public:
 	 * @param InMode The Mode the table is running in (i.e. constant/variable)
 	 * @return The Size of the Table in bytes including any header.
 	 */
-	static int32 CalcSize(int32 NumEntries, EChunkSeekTableMode InMode = EChunkSeekTableMode::ConstantSamplesPerEntry);
+	static ENGINE_API int32 CalcSize(int32 NumEntries, EChunkSeekTableMode InMode = EChunkSeekTableMode::ConstantSamplesPerEntry);
 
+	/**
+	 * Member version. Calculate the size needed for the table.	 	 
+	 *	 
+	 * @return The Size of the Table in bytes including any header.
+	 */
+	ENGINE_API int32 CalcSize() const;
+	
 	struct ISeekTableImpl
 	{
 		virtual ~ISeekTableImpl() = default;
@@ -59,6 +66,7 @@ public:
 		virtual uint32 FindTime(uint32 InOffset) const = 0;
 		virtual void Add(uint32 InTimeInAudioFrames, uint32 InOffset) = 0;
 		virtual bool Serialize(FArchive& Ar) = 0;
+		virtual bool GetAt(const uint32 InIndex, uint32& OutOffset, uint32& OutTime) const = 0;
 	};
 
 	/**
@@ -66,7 +74,7 @@ public:
 	 *
 	 * @param InMode Mode the table will operate in.
 	 */
-	FStreamedAudioChunkSeekTable(EChunkSeekTableMode InMode = EChunkSeekTableMode::ConstantSamplesPerEntry);	
+	ENGINE_API FStreamedAudioChunkSeekTable(EChunkSeekTableMode InMode = EChunkSeekTableMode::ConstantSamplesPerEntry);	
 
 	/**
 	 * Add an item to the seek table.
@@ -74,7 +82,7 @@ public:
 	 * @param InTimeInAudioFrames Time of this entry in audio frames
 	 * @param InOffset Offset in the file of where this time happens
 	 */
-	void Add(uint32 InTimeInAudioFrames, uint32 InOffset)
+	ENGINE_API void Add(uint32 InTimeInAudioFrames, uint32 InOffset)
 	{
 		Impl->Add(InTimeInAudioFrames,InOffset);
 	}
@@ -116,14 +124,26 @@ public:
 	 * @param Ar Archive to load/save
 	 * @return True on success, false on failure.
 	 */
-	bool Serialize(FArchive& Ar);
+	ENGINE_API bool Serialize(FArchive& Ar);
 
 	/**
-	 * Empties the table.
+	 * Empties the table
 	 *
 	 */
 	void Reset();
 
+	/**
+	 * Retrieves the Indexed item from the table
+	 * @param InIndex Index of item to retrieve (if out of bounds will fail)
+	 * @param OutOffset The found offset (if successful)
+	 * @param OutTime The found time (if successful)
+	 * @return True on Success, False otherwise.
+	 */
+	ENGINE_API bool GetAt(const uint32 InIndex, uint32& OutOffset, uint32& OutTime) const
+	{
+		return Impl->GetAt(InIndex, OutOffset, OutTime);
+	}
+
 private:
 
 	static uint32 GetMagic();
diff --git a/Engine/Source/Runtime/Engine/Private/StreamedAudioChunkSeekTable.cpp b/Engine/Source/Runtime/Engine/Private/StreamedAudioChunkSeekTable.cpp
--- a/Engine/Source/Runtime/Engine/Private/StreamedAudioChunkSeekTable.cpp
+++ b/Engine/Source/Runtime/Engine/Private/StreamedAudioChunkSeekTable.cpp
@@ -45,7 +45,7 @@ void FStreamedAudioChunkSeekTable::SetMode(EChunkSeekTableMode InMode)
 
 void FStreamedAudioChunkSeekTable::Reset()
 {
-	// Will delete recreate the instance.
+	// Will delete and recreate the instance.
 	SetMode(Mode);
 }
 
@@ -203,6 +203,17 @@ struct FConstantRateSeekTable : public FStreamedAudioChunkSeekTable::ISeekTableI
 		TArray<uint8> Bytes;
 		return true;
 	}
+
+	virtual bool GetAt(const uint32 InIndex, uint32& OutOffset, uint32& OutTime) const override
+	{
+		if (Table.IsValidIndex(InIndex))
+		{
+			OutTime = AudioFramesPerEntry * InIndex;
+			OutOffset = Table[InIndex];
+			return true;
+		}
+		return false;
+	}
 };
 
 struct FVariableRateSeekTable : public FStreamedAudioChunkSeekTable::ISeekTableImpl
@@ -302,6 +313,17 @@ struct FVariableRateSeekTable : public FStreamedAudioChunkSeekTable::ISeekTableI
 		StreamedAudioChunkSeekTable_Private::DeltaDecode(DeltaTimes, Times, FirstTimeItem);
 		return true;
 	}
+
+	virtual bool GetAt(const uint32 InIndex, uint32& OutOffset, uint32& OutTime) const override
+	{
+		if (Offsets.IsValidIndex(InIndex) && Times.IsValidIndex(InIndex))
+		{
+			OutTime = Times[InIndex];
+			OutOffset = Offsets[InIndex];
+			return true;
+		}
+		return false;
+	}
 };
 
 TUniquePtr<FStreamedAudioChunkSeekTable::ISeekTableImpl> FStreamedAudioChunkSeekTable::CreateImpl(EChunkSeekTableMode InMode)
@@ -344,3 +366,8 @@ int32 FStreamedAudioChunkSeekTable::CalcSize(int32 InNumEntries, EChunkSeekTable
 	checkNoEntry();
 	return 0;
 }
+
+int32 FStreamedAudioChunkSeekTable::CalcSize() const
+{
+	return CalcSize(Num(), Mode); 
+}
