Fix for ortho depth buffer scaling calculations from sampling the bottom of the view rect incorrectly.
Primarily due to the calculation of the unit to object size ratio when auto calculating near and far planes being based on the ortho width rather than the ortho height. This meant on certain hardware with specific resolutions, the values in the scaled depth were slightly misaligned across multiple frames when sampled.

diff --git a/Engine/Source/Runtime/Engine/Private/Camera/CameraStackTypes.cpp b/Engine/Source/Runtime/Engine/Private/Camera/CameraStackTypes.cpp
--- a/Engine/Source/Runtime/Engine/Private/Camera/CameraStackTypes.cpp
+++ b/Engine/Source/Runtime/Engine/Private/Camera/CameraStackTypes.cpp
@@ -418,11 +418,11 @@ bool FMinimalViewInfo::AutoCalculateOrthoPlanes(FSceneViewProjectionData& InOutP
 		 */		
 		float FarPlane;
 		float UnitPerPixelRatio = 1.0f;
-		float ViewportSizeX = InOutProjectionData.GetViewRect().Width();
-		if (bUse16bitDepth && OrthoWidth > 0  && ViewportSizeX > 0)
+		const FIntRect& ViewportSize = InOutProjectionData.GetViewRect();
+		if (bUse16bitDepth && OrthoHeight > 0  && ViewportSize.Area() > 0)
 		{
 			//The CmPerPixelRatio determines the far plane depth scale required for the scene
-			UnitPerPixelRatio = OrthoWidth / (float)ViewportSizeX;
+			UnitPerPixelRatio = FMath::FloorToFloat(OrthoHeight / (float)ViewportSize.Height());
 			if (bScaleIncrementingUnits)
 			{	
 				//This scales the min/max depending on the dynamic scale of the unit to pixel as the ortho width increases at the sacrifice of the max FarPlane, allowing scaling to LWC
