UPhysicsConstraintComponent::BreakConstraint breaks a constraint by destroying the constraint with TermConstraint. UPhysicsConstraint::IsBroken checks whether a constraint is broken by first seeing if the constraint reference is valid, then checking whether the constraint is broken. As a result, if a user calls BreakConstraint and then IsBroken, IsBroken will return false because the reference is invalid.

This change adds a local to keep track of whether a user has broken the constraint manually with BreakConstraint to make IsBroken return the correct result.

diff --git a/Engine/Source/Runtime/Engine/Classes/PhysicsEngine/PhysicsConstraintComponent.h b/Engine/Source/Runtime/Engine/Classes/PhysicsEngine/PhysicsConstraintComponent.h
--- a/Engine/Source/Runtime/Engine/Classes/PhysicsEngine/PhysicsConstraintComponent.h
+++ b/Engine/Source/Runtime/Engine/Classes/PhysicsEngine/PhysicsConstraintComponent.h
@@ -415,6 +415,9 @@ protected:
 	ENGINE_API float GetConstraintScale() const;
 
 private:
+
+	bool bIsBrokenLocal;
+
 	/** Wrapper that calls our constraint broken delegate */
 	ENGINE_API void OnConstraintBrokenWrapper(int32 ConstraintIndex);
 
diff --git a/Engine/Source/Runtime/Engine/Private/PhysicsEngine/PhysicsConstraintComponent.cpp b/Engine/Source/Runtime/Engine/Private/PhysicsEngine/PhysicsConstraintComponent.cpp
--- a/Engine/Source/Runtime/Engine/Private/PhysicsEngine/PhysicsConstraintComponent.cpp
+++ b/Engine/Source/Runtime/Engine/Private/PhysicsEngine/PhysicsConstraintComponent.cpp
@@ -26,6 +26,7 @@ UPhysicsConstraintComponent::UPhysicsConstraintComponent(const FObjectInitialize
 #if WITH_EDITORONLY_DATA
 	bVisualizeComponent = true;
 #endif
+	bIsBrokenLocal = false;
 }
 
 
@@ -300,6 +301,8 @@ void UPhysicsConstraintComponent::InitComponentConstraint()
 	{
 		ConstraintInstance.InitConstraint(Object1, Object2, GetConstraintScale(), this, FOnConstraintBroken::CreateUObject(this, &UPhysicsConstraintComponent::OnConstraintBrokenWrapper));
 	}
+
+	bIsBrokenLocal = false;
 }
 
 void UPhysicsConstraintComponent::TermComponentConstraint()
@@ -354,6 +357,7 @@ void UPhysicsConstraintComponent::SetConstrainedComponents(UPrimitiveComponent*
 void UPhysicsConstraintComponent::BreakConstraint()
 {
 	ConstraintInstance.TermConstraint();
+	bIsBrokenLocal = true;
 }
 
 
@@ -589,7 +593,7 @@ void UPhysicsConstraintComponent::GetConstraintForce(FVector& OutLinearForce, FV
 
 bool UPhysicsConstraintComponent::IsBroken()
 {
-	return ConstraintInstance.IsBroken();
+	return bIsBrokenLocal || ConstraintInstance.IsBroken();
 }
 
 void UPhysicsConstraintComponent::SetDisableCollision(bool bDisableCollision)
