Level Instance: force level instance actors to be cooked when used in a non-partitioned world.

diff --git a/Engine/Source/Runtime/Engine/Private/LevelInstance/LevelInstanceActor.cpp b/Engine/Source/Runtime/Engine/Private/LevelInstance/LevelInstanceActor.cpp
--- a/Engine/Source/Runtime/Engine/Private/LevelInstance/LevelInstanceActor.cpp
+++ b/Engine/Source/Runtime/Engine/Private/LevelInstance/LevelInstanceActor.cpp
@@ -3,6 +3,7 @@
 #include "LevelInstance/LevelInstanceActor.h"
 #include "LevelInstance/LevelInstanceSubsystem.h"
 #include "LevelInstance/LevelInstanceComponent.h"
+#include "Engine/Level.h"
 #include "Engine/World.h"
 #include "LevelInstancePrivate.h"
 #include "Net/UnrealNetwork.h"
@@ -379,6 +380,11 @@ bool ALevelInstance::EditorCanAttachFrom(const AActor* InChild, FText& OutReason
 
 bool ALevelInstance::IsEditorOnly() const
 {
+	if (IsRunningCookCommandlet() && GetLevel() && !GetLevel()->bIsPartitioned)
+	{
+		return false;
+	}
+
 	if (DesiredRuntimeBehavior == ELevelInstanceRuntimeBehavior::Partitioned)
 	{
 		return Super::IsEditorOnly();
