Substrate - Clarification of Select node parameter blending.

#tests PC substrate dx12
#rnx

diff --git a/Engine/Source/Runtime/Engine/Private/Materials/MaterialExpressionSubstrate.cpp b/Engine/Source/Runtime/Engine/Private/Materials/MaterialExpressionSubstrate.cpp
--- a/Engine/Source/Runtime/Engine/Private/Materials/MaterialExpressionSubstrate.cpp
+++ b/Engine/Source/Runtime/Engine/Private/Materials/MaterialExpressionSubstrate.cpp
@@ -3116,7 +3116,7 @@ int32 UMaterialExpressionSubstrateSelect::Compile(class FMaterialCompiler* Compi
 	}
 	else
 	{
-		Compiler->Errorf(TEXT("The Select node must be compiled with parameter blending to only select between one of two BSDF"));
+		Compiler->Errorf(TEXT("The Select node can only use parameter blending to only select between one of two BSDF."));
 	}
 
 	// Propagate the parameter blended normal
@@ -3127,7 +3127,14 @@ int32 UMaterialExpressionSubstrateSelect::Compile(class FMaterialCompiler* Compi
 
 void UMaterialExpressionSubstrateSelect::GetCaption(TArray<FString>& OutCaptions) const
 {
-	OutCaptions.Add(TEXT("Substrate Select"));
+	if (bUseParameterBlending)
+	{
+		OutCaptions.Add(TEXT("Substrate Select (Parameter Blend)"));
+	}
+	else
+	{
+		OutCaptions.Add(TEXT("Substrate Select"));
+	}
 }
 
 uint32 UMaterialExpressionSubstrateSelect::GetOutputType(int32 OutputIndex)
@@ -3173,7 +3180,6 @@ void UMaterialExpressionSubstrateSelect::GatherSubstrateMaterialInfo(FSubstrateM
 
 FSubstrateOperator* UMaterialExpressionSubstrateSelect::SubstrateGenerateMaterialTopologyTree(class FMaterialCompiler* Compiler, class UMaterialExpression* Parent, int32 OutputIndex)
 {
-	const bool bUseParameterBlending = true; // The select node only works with a single BSDF on each side.
 	FSubstrateOperator& SubstrateOperator = Compiler->SubstrateCompilationRegisterOperator(SUBSTRATE_OPERATOR_SELECT, Compiler->SubstrateTreeStackGetPathUniqueId(), this, Parent, Compiler->SubstrateTreeStackGetParentPathUniqueId(), bUseParameterBlending);
 	if (Compiler->GetSubstrateTreeOutOfStackDepthOccurred())
 	{
diff --git a/Engine/Source/Runtime/Engine/Public/Materials/MaterialExpressionSubstrate.h b/Engine/Source/Runtime/Engine/Public/Materials/MaterialExpressionSubstrate.h
--- a/Engine/Source/Runtime/Engine/Public/Materials/MaterialExpressionSubstrate.h
+++ b/Engine/Source/Runtime/Engine/Public/Materials/MaterialExpressionSubstrate.h
@@ -1092,6 +1092,12 @@ class UMaterialExpressionSubstrateSelect : public UMaterialExpressionSubstrateBS
 	UPROPERTY(EditAnywhere, Category = Mode, meta = (UIMin = 0.0, UIMax = 1.0, ClampMin = 0.0, ClampMax = 1.0))
 	float Threshold;
 
+	/**
+	 * Select A and B using parameter blending for material evaluation. This makes lighting evaluation cheaper (Default: On)
+	 * As of today, parameter blending must be used since a single Substrate Tree cannot be used with hair or slabs for instance.
+	 */
+	const bool bUseParameterBlending = true;
+
 	//~ Begin UMaterialExpression Interface
 #if WITH_EDITOR
 	virtual int32 Compile(class FMaterialCompiler* Compiler, int32 OutputIndex) override;
