Fixed invalid checks in UWorldPartitionLevelStreamingDynamic::Unload/Deactivate not properly handling level load failures

diff --git a/Engine/Source/Runtime/Engine/Private/WorldPartition/WorldPartitionLevelStreamingDynamic.cpp b/Engine/Source/Runtime/Engine/Private/WorldPartition/WorldPartitionLevelStreamingDynamic.cpp
--- a/Engine/Source/Runtime/Engine/Private/WorldPartition/WorldPartitionLevelStreamingDynamic.cpp
+++ b/Engine/Source/Runtime/Engine/Private/WorldPartition/WorldPartitionLevelStreamingDynamic.cpp
@@ -627,7 +627,7 @@ void UWorldPartitionLevelStreamingDynamic::Unload()
 {
 	UE_LOG(LogLevelStreaming, Verbose, TEXT("UWorldPartitionLevelStreamingDynamic::Unloading %s"), *GetWorldAssetPackageName());
 
-	check(ShouldBeLoaded());
+	check(ShouldBeLoaded() || (GetLevelStreamingState() == ELevelStreamingState::FailedToLoad));
 
 	SetShouldBeLoaded(false);
 	SetShouldBeVisible(false);
@@ -661,8 +661,8 @@ void UWorldPartitionLevelStreamingDynamic::Deactivate()
 {
 	UE_LOG(LogLevelStreaming, Verbose, TEXT("UWorldPartitionLevelStreamingDynamic::Deactivating %s"), *GetWorldAssetPackageName());
 
-	check(ShouldBeLoaded());
-	check(ShouldBeVisible());
+	check(ShouldBeLoaded() || (GetLevelStreamingState() == ELevelStreamingState::FailedToLoad));
+	check(ShouldBeVisible() || (GetLevelStreamingState() == ELevelStreamingState::FailedToLoad));
 
 	SetShouldBeVisible(false);
 }
