Substrate - clarified when it is recommended to specify the MFP directly, or to use the TransmittanceToMFP node.

#rnx
#tests PC Editor

diff --git a/Engine/Source/Runtime/Engine/Private/Materials/MaterialExpressionSubstrate.cpp b/Engine/Source/Runtime/Engine/Private/Materials/MaterialExpressionSubstrate.cpp
--- a/Engine/Source/Runtime/Engine/Private/Materials/MaterialExpressionSubstrate.cpp
+++ b/Engine/Source/Runtime/Engine/Private/Materials/MaterialExpressionSubstrate.cpp
@@ -3326,8 +3326,9 @@ void UMaterialExpressionSubstrateTransmittanceToMFP::GetConnectorToolTip(int32 I
 
 void UMaterialExpressionSubstrateTransmittanceToMFP::GetExpressionToolTip(TArray<FString>& OutToolTip)
 {
-	ConvertToMultilineToolTip(TEXT("Convert a transmittance color corresponding to a slab of participating media viewed perpendicularly to its surface.\nThis node directly maps to the Slab BSDF input."), 80, OutToolTip);
-
+	OutToolTip.Add(TEXT("Convert a transmittance color corresponding to a slab of participating media viewed perpendicularly to its surface."));
+	OutToolTip.Add(TEXT("This node directly maps to the Slab MFP input. It is recommended to use it when specifying the colored transmittance of a top layer slab."));
+	OutToolTip.Add(TEXT("For Subsurface scattering, you might prefer to specify the MFP(light mean free path) as world space centimeter directly."));
 }
 #endif // WITH_EDITOR
 
diff --git a/Engine/Source/Runtime/Engine/Public/Materials/MaterialExpressionSubstrate.h b/Engine/Source/Runtime/Engine/Public/Materials/MaterialExpressionSubstrate.h
--- a/Engine/Source/Runtime/Engine/Public/Materials/MaterialExpressionSubstrate.h
+++ b/Engine/Source/Runtime/Engine/Public/Materials/MaterialExpressionSubstrate.h
@@ -265,7 +265,9 @@ class UMaterialExpressionSubstrateSlabBSDF : public UMaterialExpressionSubstrate
 	FExpressionInput Tangent;
 
 	/**
-	 * Chromatic mean free path . Only used when there is not any sub-surface profile provided. (type = float3, unit = centimeters, default = 0)
+	 * Chromatic mean free path. Only used when there is not any sub-surface profile provided. (type = float3, unit = centimeters, default = 0).
+	 * For subsurface scattering, it is recommended to specify the MFP as world space centimeter directly. As it is an intuitive way to specify how far the light will scatter/bleed.
+	 * For translucent coat layer, it is recommended to use the TransmittanceToMFP node as it might be easier to specify the transmittance color.
 	 */
 	UPROPERTY()
 	FExpressionInput SSSMFP;
