[HWRT] Register ray tracing geometry group using number of LODs in RayTracingProxy instead of LODResources array.

diff --git a/Engine/Source/Runtime/Engine/Private/StaticMesh.cpp b/Engine/Source/Runtime/Engine/Private/StaticMesh.cpp
--- a/Engine/Source/Runtime/Engine/Private/StaticMesh.cpp
+++ b/Engine/Source/Runtime/Engine/Private/StaticMesh.cpp
@@ -2341,12 +2341,12 @@ void FStaticMeshRenderData::InitResources(ERHIFeatureLevel::Type InFeatureLevel,
 		ENQUEUE_RENDER_COMMAND(InitStaticMeshRayTracingGeometry)(
 			[this, Owner](FRHICommandListImmediate& RHICmdList)
 			{
-				RayTracingGeometryGroupHandle = GRayTracingGeometryManager->RegisterRayTracingGeometryGroup(LODResources.Num(), RayTracingProxy->bUsingRenderingLODs ? CurrentFirstLODIdx : 0);
-				
 				FStaticMeshRayTracingProxyLODArray& RayTracingLODs = RayTracingProxy->LODs;
 
 				check(!RayTracingProxy->bUsingRenderingLODs || RayTracingLODs.Num() == LODResources.Num());
 
+				RayTracingGeometryGroupHandle = GRayTracingGeometryManager->RegisterRayTracingGeometryGroup(RayTracingLODs.Num(), RayTracingProxy->bUsingRenderingLODs ? CurrentFirstLODIdx : 0);
+
 				for (int32 LODIndex = 0; LODIndex < RayTracingLODs.Num(); ++LODIndex)
 				{
 					FStaticMeshRayTracingProxyLOD& RayTracingLOD = RayTracingLODs[LODIndex];
@@ -2358,8 +2358,11 @@ void FStaticMeshRenderData::InitResources(ERHIFeatureLevel::Type InFeatureLevel,
 
 						const bool bHasStreamableData = RayTracingLOD.StreamableData.GetBulkDataSize() > 0;
 
-						checkf(RayTracingLOD.bOwnsRayTracingGeometry || !bHasStreamableData,
-							TEXT("Unexpected StreamableData on '%s' (StreamableData is only expected when ray tracing proxy owns the geometry)."), *GetPathNameSafe(Owner));
+						if (bHasStreamableData)
+						{
+							checkf(RayTracingLOD.bOwnsRayTracingGeometry, TEXT("Unexpected StreamableData on '%s' (StreamableData is only expected when ray tracing proxy owns the geometry)."), *GetPathNameSafe(Owner));
+							checkf(!RayTracingLOD.bBuffersInlined, TEXT("RayTracing LOD with streamable data shouldn't have buffers inlined"));
+						}
 
 						if (bHasStreamableData || LODIndex < CurrentFirstLODIdx)
 						{
