Add lock on skel. mesh DDC key building to avoid ensure the function more thread-safe when different groom binding asset are build in parallel.

diff --git a/Engine/Source/Runtime/Engine/Classes/Engine/SkeletalMesh.h b/Engine/Source/Runtime/Engine/Classes/Engine/SkeletalMesh.h
--- a/Engine/Source/Runtime/Engine/Classes/Engine/SkeletalMesh.h
+++ b/Engine/Source/Runtime/Engine/Classes/Engine/SkeletalMesh.h
@@ -2458,6 +2458,8 @@ public:
 
 private:	
 	bool bTransacting = false;
+	FCriticalSection DDCKeyBuildingCS;
+
 public:
 	bool IsTransacting() const;
 	
diff --git a/Engine/Source/Runtime/Engine/Private/SkeletalMesh.cpp b/Engine/Source/Runtime/Engine/Private/SkeletalMesh.cpp
--- a/Engine/Source/Runtime/Engine/Private/SkeletalMesh.cpp
+++ b/Engine/Source/Runtime/Engine/Private/SkeletalMesh.cpp
@@ -5259,6 +5259,8 @@ extern int32 GSkeletalMeshKeepMobileMinLODSettingOnDesktop;
 
 FString USkeletalMesh::BuildDerivedDataKey(const ITargetPlatform* TargetPlatform)
 {
+	FScopeLock ScopeLock(&DDCKeyBuildingCS);
+
 	TArray<uint8> TempBytes;
 	TempBytes.Reserve(64);
 
