- Fixed IsStreamingCompleted using streaming sources to consider streaming as completed if the content is activated and the target state is loaded
- Fixed IsStreamingCompleted using query to respect exact state for cells that are not part of Data Layers (query=loaded, cell=activated was considered as completed)

diff --git a/Engine/Source/Runtime/Engine/Private/WorldPartition/WorldPartitionStreamingPolicy.cpp b/Engine/Source/Runtime/Engine/Private/WorldPartition/WorldPartitionStreamingPolicy.cpp
--- a/Engine/Source/Runtime/Engine/Private/WorldPartition/WorldPartitionStreamingPolicy.cpp
+++ b/Engine/Source/Runtime/Engine/Private/WorldPartition/WorldPartitionStreamingPolicy.cpp
@@ -1028,6 +1028,10 @@ bool UWorldPartitionStreamingPolicy::IsStreamingCompleted(const TArray<FWorldPar
 	check(World->IsGameWorld());
 	const UDataLayerManager* DataLayerManager = WorldPartition->GetDataLayerManager();
 	const bool bTestProvidedStreamingSource = !!InStreamingSources;
+	
+	// IsStreamingCompleted using streaming sources will be considered as completed
+	// if the content is activated and the target state is loaded.
+	const bool bExactState = false;
 
 	// Always test non-spatial cells
 	{
@@ -1037,7 +1041,7 @@ bool UWorldPartitionStreamingPolicy::IsStreamingCompleted(const TArray<FWorldPar
 		QuerySource.bSpatialQuery = false;
 		QuerySource.bDataLayersOnly = false;
 		QuerySource.DataLayers = DataLayerManager->GetEffectiveActiveDataLayerNames().Array();
-		if (!IsStreamingCompleted(EWorldPartitionRuntimeCellState::Activated, QuerySources, true))
+		if (!IsStreamingCompleted(EWorldPartitionRuntimeCellState::Activated, QuerySources, bExactState))
 		{
 			return false;
 		}
@@ -1047,7 +1051,7 @@ bool UWorldPartitionStreamingPolicy::IsStreamingCompleted(const TArray<FWorldPar
 		{
 			QuerySource.bDataLayersOnly = true;
 			QuerySource.DataLayers = DataLayerManager->GetEffectiveLoadedDataLayerNames().Array();
-			if (!IsStreamingCompleted(EWorldPartitionRuntimeCellState::Loaded, QuerySources, true))
+			if (!IsStreamingCompleted(EWorldPartitionRuntimeCellState::Loaded, QuerySources, bExactState))
 			{
 				return false;
 			}
@@ -1074,7 +1078,7 @@ bool UWorldPartitionStreamingPolicy::IsStreamingCompleted(const TArray<FWorldPar
 
 		// Execute query
 		const EWorldPartitionRuntimeCellState QueryState = (StreamingSource.TargetState == EStreamingSourceTargetState::Loaded) ? EWorldPartitionRuntimeCellState::Loaded : EWorldPartitionRuntimeCellState::Activated;
-		if (!IsStreamingCompleted(QueryState, QuerySources, true))
+		if (!IsStreamingCompleted(QueryState, QuerySources, bExactState))
 		{
 			return false;
 		}
@@ -1106,7 +1110,10 @@ bool UWorldPartitionStreamingPolicy::IsStreamingCompleted(EWorldPartitionRuntime
 					bSkipCell = Cell->GetIsHLOD();
 				}
 
-				if (!bSkipCell)
+				// Test if cell is already in the effective wanted state.
+				// Note that GetCellEffectiveWantedState always return Activated when the cell has no data layers.
+				// In this case, continue testing the QueryState with the CellState to respect bExactState.
+				if (!bSkipCell && Cell->HasDataLayers())
 				{
 					const EDataLayerRuntimeState CellWantedState = Cell->GetCellEffectiveWantedState(StreamingContext);
 					bSkipCell = (CellState == EWorldPartitionRuntimeCellState::Unloaded && CellWantedState == EDataLayerRuntimeState::Unloaded) || 
