Fixed case where replacing an actor from and external data layer with a BP actor from another external data layer under the same plugin would not keep the original external data layer.

diff --git a/Engine/Source/Runtime/Engine/Private/WorldPartition/DataLayer/ExternalDataLayerManager.cpp b/Engine/Source/Runtime/Engine/Private/WorldPartition/DataLayer/ExternalDataLayerManager.cpp
--- a/Engine/Source/Runtime/Engine/Private/WorldPartition/DataLayer/ExternalDataLayerManager.cpp
+++ b/Engine/Source/Runtime/Engine/Private/WorldPartition/DataLayer/ExternalDataLayerManager.cpp
@@ -535,10 +535,20 @@ const UExternalDataLayerAsset* UExternalDataLayerManager::GetMatchingExternalDat
 {
 	const FName MountPoint = FPackageName::GetPackageMountPoint(InObjectPath.ToString());
 
+	FExternalDataLayerUID ExternalDataLayerUIDFromPackage;
+	FExternalDataLayerHelper::IsExternalDataLayerPath(InObjectPath.ToString(), &ExternalDataLayerUIDFromPackage);
+
 	// Find injected EDL assets with the same mount point than the provided object path
 	TArray<const UExternalDataLayerAsset*> Candidates;
 	for (const UExternalDataLayerAsset* ExternalDataLayerAsset : InjectedExternalDataLayerAssets)
 	{
+		// If the provided object path contains a valid EDL UID and we find a matching external data layer asset, use this asset.
+		if (ExternalDataLayerUIDFromPackage.IsValid() && (ExternalDataLayerUIDFromPackage == ExternalDataLayerAsset->GetUID()))
+		{
+			check(FPackageName::GetPackageMountPoint(ExternalDataLayerAsset->GetPackage()->GetName()) == MountPoint);
+			return ExternalDataLayerAsset;
+		}
+
 		if (FPackageName::GetPackageMountPoint(ExternalDataLayerAsset->GetPackage()->GetName()) == MountPoint)
 		{
 			Candidates.Add(ExternalDataLayerAsset);
