Abort component registration when attempted on a cleaned-up/cleaning-up World and log an errror
Incorrectly unregistered component detection added to UWorld::CleanupWorldInternal, logs a fatal error in test builds since this will most probablu crash later on and the it's very hard to recover the root cause from the crash. Devs + Debug build relly on an ensure to attract devs attention to this issue.

diff --git a/Engine/Source/Editor/UnrealEd/Private/EditorEngine.cpp b/Engine/Source/Editor/UnrealEd/Private/EditorEngine.cpp
--- a/Engine/Source/Editor/UnrealEd/Private/EditorEngine.cpp
+++ b/Engine/Source/Editor/UnrealEd/Private/EditorEngine.cpp
@@ -4217,6 +4217,9 @@ bool UEditorEngine::InitializePhysicsSceneForSaveIfNecessary(UWorld* World, bool
 		}
 		else
 		{
+			// We allow multiple reinit for the world in this case, so we need to reset the cleanup flag to avoid tripping the checks in IsCleanedUp
+			World->ResetCleanedUpState();
+
 			// If we aren't already initialized, initialize now and create a physics scene. Don't create an FX system because it uses too much video memory for bulk operations
 			World->InitWorld(
 				GetEditorWorldInitializationValues()
diff --git a/Engine/Source/Runtime/Engine/Classes/Engine/World.h b/Engine/Source/Runtime/Engine/Classes/Engine/World.h
--- a/Engine/Source/Runtime/Engine/Classes/Engine/World.h
+++ b/Engine/Source/Runtime/Engine/Classes/Engine/World.h
@@ -3030,6 +3030,15 @@ public:
 	/** Returns true if the world is in the process of being cleaned up. */
 	bool IsBeingCleanedUp() const { return bIsBeingCleanedUp; }
 
+	/** Returns true if the world has or is in the process of being cleaned up. */
+	bool IsCleanedUp() const { return CleanupWorldTag != 0; }
+
+#if WITH_EDITOR
+	/** Force reset world to a state where it's not considered cleaned up. Use with caution this could disable important validations regarding component registration. */
+	void ResetCleanedUpState() { CleanupWorldTag = 0; }
+#endif
+
+
 	/** 
 	 * Marks this world and all objects within as pending kill
 	 */
diff --git a/Engine/Source/Runtime/Engine/Private/Components/ActorComponent.cpp b/Engine/Source/Runtime/Engine/Private/Components/ActorComponent.cpp
--- a/Engine/Source/Runtime/Engine/Private/Components/ActorComponent.cpp
+++ b/Engine/Source/Runtime/Engine/Private/Components/ActorComponent.cpp
@@ -1591,6 +1591,12 @@ void UActorComponent::RegisterComponentWithWorld(UWorld* InWorld, FRegisterCompo
 		return;
 	}
 
+	if (!IsValid(InWorld) || InWorld->IsCleanedUp())
+	{
+		UE_LOG(LogActorComponent, Error, TEXT("RegisterComponentWithWorld: (%s) Trying to register component in invalid or cleaned up world. Aborting."), *GetPathName());
+		return;
+	}
+
 	// If not registered/pre-registered, should not have a world
 	checkf(WorldPrivate == nullptr || (IsPreRegistered()), TEXT("%s"), *GetFullName());
 
diff --git a/Engine/Source/Runtime/Engine/Private/World.cpp b/Engine/Source/Runtime/Engine/Private/World.cpp
--- a/Engine/Source/Runtime/Engine/Private/World.cpp
+++ b/Engine/Source/Runtime/Engine/Private/World.cpp
@@ -164,6 +164,22 @@ FAutoConsoleVariableRef CVarCreateStaticLevelCollection(TEXT("s.World.CreateStat
 	TEXT("When set, create a separate level collection for static streaming levels that will not be duplicated by DuplicateRequestedLevels.\n")
 	TEXT("If this is 0, static streaming levels will be part of the main DynamicSourceLevels collection."));
 
+// This cvar activates detection of incorrectly unregistered components after a world is cleaned-up.
+// This is important since components left as registered/renderstatecreated after their world is cleaned up will 
+// crash in FGlobalComponentRecreateRenderStateContext unless a GC happens before the recreate.
+//
+// Current intended config across builds....
+// Debug + Dev = LogError + an ensure() to prompt devs to fix
+// Test = LogFatalError so that hopefully it doesn't make it out of testing broken
+// Ship = LogError, and hope for a GC pass before usage of a FGlobalComponentRecreateRenderStateContext, or LogFatalError if turned on in the config
+#if UE_BUILD_TEST
+static bool bIncorrectComponentUnregistrationIsFatal = true;
+#else
+static bool bIncorrectComponentUnregistrationIsFatal = false;
+#endif
+
+FAutoConsoleVariableRef CVarIncorrectComponentUnregistrationIsFatal(TEXT("s.World.IncorrectComponentUnregistrationIsFatal"), bIncorrectComponentUnregistrationIsFatal, TEXT("When set, incorrect component unregistration will log a fatal error."));
+
 // Now that it's possible for subclasses of ULevelStreaming to indicate which async loads are necessary for loading,
 // it's possible existing subclasses haven't added their required loads to their StreamingLevel->GetAsyncRequestIDs() array. As a fallback,
 // allow users to force flushing of all async loads during level streaming, as was done in UE 5.4 and lower.
@@ -5549,6 +5565,39 @@ void UWorld::CleanupWorld(bool bSessionEnded, bool bCleanupResources, UWorld* Ne
 	bIsWorldInitialized = false;
 }
 
+void ValidateComponentsRegistration(UWorld* World)
+{
+	TStringBuilder<2048>	Errors;
+
+	for (UActorComponent* Component : TObjectRange<UActorComponent>())
+	{
+		UWorld* ComponentWorld = Component->GetWorld();
+		if (ComponentWorld == World)
+		{
+			bool bComponentUnregisteredIncorrectly = Component->IsRegistered() && Component->IsRenderStateCreated();
+			
+			if (bComponentUnregisteredIncorrectly)
+			{
+				Errors.Appendf(TEXT("Component %s (Outer: %s) still registered and renderstate created after world cleanup. It was either re-registered during clean up or registered in the world without being owned by an actor and never unregistered.\n"), * Component->GetFullName(), * Component->GetOuter()->GetFullName());
+			}
+		}
+	}
+
+	if (Errors.Len())
+	{
+		if (bIncorrectComponentUnregistrationIsFatal)
+		{
+			UE_LOG(LogWorld, Fatal, TEXT("Found components to be incorrectly unregistered: \n%s"), Errors.GetData());
+		}
+		else
+		{
+			UE_LOG(LogWorld, Error, TEXT("Found components to be incorrectly unregistered: \n%s"), Errors.GetData());
+		}
+	}
+
+	ensureMsgf(Errors.Len() == 0, TEXT("Found components to be incorrectly unregistered: \n%s"), Errors.GetData());
+}
+
 void UWorld::CleanupWorldInternal(bool bSessionEnded, bool bCleanupResources, bool bWorldChanged)
 {
 	TGuardValue<bool> IsBeingCleanedUp(bIsBeingCleanedUp, true);
@@ -5745,6 +5794,8 @@ void UWorld::CleanupWorldInternal(bool bSessionEnded, bool bCleanupResources, bo
 		Scene->SetFXSystem(NULL);
 		FXSystem = NULL;
 	}
+
+	ValidateComponentsRegistration(this);
 }
 
 UGameViewportClient* UWorld::GetGameViewport() const
