Do not mark Blueprint assets dirty when reconstructing graph nodes outside of a transaction scope.

diff --git a/Engine/Source/Editor/BlueprintGraph/Private/K2Node.cpp b/Engine/Source/Editor/BlueprintGraph/Private/K2Node.cpp
--- a/Engine/Source/Editor/BlueprintGraph/Private/K2Node.cpp
+++ b/Engine/Source/Editor/BlueprintGraph/Private/K2Node.cpp
@@ -703,7 +703,8 @@ void UK2Node::PostReconstructNode()
 
 void UK2Node::ReconstructNode()
 {
-	Modify();
+	// Only mark as dirty if inside a transactional scope; otherwise, there's no need to re-save the package
+	Modify(/*bAlwaysMarkDirty =*/ false);
 
 	// Clear previously set messages
 	ErrorMsg.Reset();
@@ -1495,8 +1496,8 @@ void UK2Node::DestroyPinList(TArray<UEdGraphPin*>& InPins)
 	// Throw away the original pins
 	for (UEdGraphPin* Pin : InPins)
 	{
-		Pin->Modify();
-		Pin->BreakAllPinLinks(!Blueprint->bIsRegeneratingOnLoad);
+		Pin->Modify(/*bAlwaysMarkDirty =*/ false);
+		Pin->BreakAllPinLinks(!Blueprint->bIsRegeneratingOnLoad, /*bAlwaysMarkDirty =*/ false);
 
 		UEdGraphNode::DestroyPin(Pin);
 	}
diff --git a/Engine/Source/Runtime/Engine/Private/EdGraph/EdGraphPin.cpp b/Engine/Source/Runtime/Engine/Private/EdGraph/EdGraphPin.cpp
--- a/Engine/Source/Runtime/Engine/Private/EdGraph/EdGraphPin.cpp
+++ b/Engine/Source/Runtime/Engine/Private/EdGraph/EdGraphPin.cpp
@@ -597,7 +597,7 @@ enum class ETransferPersistentDataMode : uint8
 template<class T>
 void TransferPersistentDataFromOldPin(UEdGraphPin& DestPin, T& SourcePin, const ETransferPersistentDataMode TransferMode)
 {
-	DestPin.Modify();
+	DestPin.Modify(/*bAlwaysMarkDirty =*/ false);
 
 	// The name matches already, doesn't get copied here
 	// The PinType, Direction, and bNotConnectable are properties generated from the schema
@@ -624,7 +624,7 @@ void TransferPersistentDataFromOldPin(UEdGraphPin& DestPin, T& SourcePin, const
 	{
 		check(OtherPin);
 
-		OtherPin->Modify();
+		OtherPin->Modify(/*bAlwaysMarkDirty =*/ false);
 
 		DestPin.LinkedTo.Add(OtherPin);
 
@@ -678,7 +678,7 @@ void TransferPersistentDataFromOldPin(UEdGraphPin& DestPin, T& SourcePin, const
 
 void UEdGraphPin::MovePersistentDataFromOldPin(UEdGraphPin& SourcePin)
 {
-	SourcePin.Modify();
+	SourcePin.Modify(/*bAlwaysMarkDirty =*/ false);
 
 	TransferPersistentDataFromOldPin(*this, SourcePin, ETransferPersistentDataMode::Move);
 
