Fixing SweepComponent not working properly with zero length queries.
These will now properly do a sweep instead of an overlap and return all relevant info, including handling the correct first impact result.

#jia UE-227438

diff --git a/Engine/Source/Runtime/Engine/Private/PhysicsEngine/Experimental/PhysInterface_Chaos.cpp b/Engine/Source/Runtime/Engine/Private/PhysicsEngine/Experimental/PhysInterface_Chaos.cpp
--- a/Engine/Source/Runtime/Engine/Private/PhysicsEngine/Experimental/PhysInterface_Chaos.cpp
+++ b/Engine/Source/Runtime/Engine/Private/PhysicsEngine/Experimental/PhysInterface_Chaos.cpp
@@ -950,71 +950,83 @@ bool FPhysInterface_Chaos::Sweep_Geom(FHitResult& OutHit, const FBodyInstance* I
 
 				const FVector Delta = InEnd - InStart;
 				const float DeltaMag = Delta.Size();
-				if (DeltaMag > UE_KINDA_SMALL_NUMBER)
-				{
-					const FTransform ActorTM(Actor->GetGameThreadAPI().R(), Actor->GetGameThreadAPI().X());
+				
+				const FTransform ActorTM(Actor->GetGameThreadAPI().R(), Actor->GetGameThreadAPI().X());
 
-					UPrimitiveComponent* OwnerComponentInst = InInstance->OwnerComponent.Get();
-					FTransform StartTM(ShapeAdapter.GetGeomOrientation(), InStart);
-					FTransform CompTM(OwnerComponentInst->GetComponentTransform());
+				UPrimitiveComponent* OwnerComponentInst = InInstance->OwnerComponent.Get();
+				FTransform StartTM(ShapeAdapter.GetGeomOrientation(), InStart);
+				FTransform CompTM(OwnerComponentInst->GetComponentTransform());
 
-					Chaos::FVec3 Dir = Delta / DeltaMag;
+				const Chaos::FVec3 Dir = DeltaMag > 0 ? (Delta / DeltaMag) : Chaos::FVec3(1, 0, 0);
 
-					FSweepHit Hit;
+				FSweepHit Hit;
+				Hit.Distance = TNumericLimits<float>::Max();
 
-					// Get all the shapes from the actor
-					PhysicsInterfaceTypes::FInlineShapeArray Shapes;
-					// #PHYS2 - SHAPES - Resolve this function to not use px stuff
-					const int32 NumShapes = FillInlineShapeArray_AssumesLocked(Shapes, Actor); // #PHYS2 - Need a lock/execute here?
+				// Get all the shapes from the actor
+				PhysicsInterfaceTypes::FInlineShapeArray Shapes;
+				// #PHYS2 - SHAPES - Resolve this function to not use px stuff
+				const int32 NumShapes = FillInlineShapeArray_AssumesLocked(Shapes, Actor); // #PHYS2 - Need a lock/execute here?
 
-					// Iterate over each shape
-					for (int32 ShapeIdx = 0; ShapeIdx < NumShapes; ShapeIdx++)
-					{
-						FPhysicsShapeReference_Chaos& ShapeRef = Shapes[ShapeIdx];
-						Chaos::FPerShapeData* Shape = ShapeRef.Shape;
-						check(Shape);
+				// Iterate over each shape
+				for (int32 ShapeIdx = 0; ShapeIdx < NumShapes; ShapeIdx++)
+				{
+					FPhysicsShapeReference_Chaos& ShapeRef = Shapes[ShapeIdx];
+					Chaos::FPerShapeData* Shape = ShapeRef.Shape;
+					check(Shape);
 
-						// Skip shapes not bound to this instance
-						if (!TargetInstance->IsShapeBoundToBody(ShapeRef))
-						{
-							continue;
-						}
+					// Skip shapes not bound to this instance
+					if (!TargetInstance->IsShapeBoundToBody(ShapeRef))
+					{
+						continue;
+					}
 
-						// Filter so we trace against the right kind of collision
-						FCollisionFilterData ShapeFilter = Shape->GetQueryData();
-						const bool bShapeIsComplex = (ShapeFilter.Word3 & EPDF_ComplexCollision) != 0;
-						const bool bShapeIsSimple = (ShapeFilter.Word3 & EPDF_SimpleCollision) != 0;
-						if ((bSweepComplex && bShapeIsComplex) || (!bSweepComplex && bShapeIsSimple))
-						{
-							//question: this is returning first result, is that valid? Keeping it the same as physx for now
-							Chaos::FVec3 WorldPosition;
-							Chaos::FVec3 WorldNormal;
-							Chaos::FReal Distance;
-							int32 FaceIdx;
-							Chaos::FVec3 FaceNormal;
-							if (Chaos::Utilities::CastHelper(ShapeAdapter.GetGeometry(), ActorTM, [&](const auto& Downcast, const auto& FullActorTM) { return Chaos::SweepQuery(*Shape->GetGeometry(), FullActorTM, Downcast, StartTM, Dir, DeltaMag, Distance, WorldPosition, WorldNormal, FaceIdx, FaceNormal, 0.f, false); }))
+					// Filter so we trace against the right kind of collision
+					FCollisionFilterData ShapeFilter = Shape->GetQueryData();
+					const bool bShapeIsComplex = (ShapeFilter.Word3 & EPDF_ComplexCollision) != 0;
+					const bool bShapeIsSimple = (ShapeFilter.Word3 & EPDF_SimpleCollision) != 0;
+					if ((bSweepComplex && bShapeIsComplex) || (!bSweepComplex && bShapeIsSimple))
+					{
+						//question: this is returning first result, is that valid? Keeping it the same as physx for now
+						Chaos::FVec3 WorldPosition;
+						Chaos::FVec3 WorldNormal;
+						Chaos::FReal Distance;
+						int32 FaceIdx;
+						Chaos::FVec3 FaceNormal;
+
+						// This is what it needs to be, update later...
+						const FPhysicsGeometry& QueryGeometry = ShapeAdapter.GetGeometry();
+						const bool bShapeHit = Chaos::Utilities::CastHelper(QueryGeometry, ActorTM,
+							[&](const auto& Downcast, const auto& FullActorTM)
 							{
-								// we just like to make sure if the hit is made
-								FCollisionFilterData QueryFilter;
-								QueryFilter.Word2 = 0xFFFFF;
-
-								// we don't get Shape information when we access via PShape, so I filled it up
-								Hit.Shape = Shape;
-								Hit.Actor = ShapeRef.ActorRef ? ShapeRef.ActorRef->GetParticle_LowLevel() : nullptr;
-								Hit.WorldPosition = WorldPosition;
-								Hit.WorldNormal = WorldNormal;
-								Hit.Distance = (float)Distance; // we should eventually have the Hit structure to use a Chaos::FReal equivalent instead
-								Hit.FaceIndex = FaceIdx;
-								if (!HadInitialOverlap(Hit))
-								{
-									Hit.FaceIndex = FindFaceIndex(Hit, Dir);
-								}
-								SetFlags(Hit, EHitFlags::Distance | EHitFlags::Normal | EHitFlags::Position | EHitFlags::FaceIndex);
-
-								FTransform StartTransform(InStart);
-								ConvertQueryImpactHit(OwnerComponentInst->GetWorld(), Hit, OutHit, DeltaMag, QueryFilter, InStart, InEnd, nullptr, StartTransform, false, false);
-								bSweepHit = true;
+								return Chaos::SweepQuery(*Shape->GetGeometry(), FullActorTM, Downcast, StartTM, Dir, DeltaMag, Distance, WorldPosition, WorldNormal, FaceIdx, FaceNormal, 0.f, true);
 							}
+						);
+
+						// we should eventually have the Hit structure to use a Chaos::FReal equivalent instead;
+						const float FloatDistance = static_cast<float>(Distance);
+						if (bShapeHit && FloatDistance < Hit.Distance)
+						{
+							// we just like to make sure if the hit is made
+							FCollisionFilterData QueryFilter;
+							QueryFilter.Word2 = 0xFFFFF;
+
+							// we don't get Shape information when we access via PShape, so I filled it up
+							Hit.Shape = Shape;
+							Hit.Actor = ShapeRef.ActorRef ? ShapeRef.ActorRef->GetParticle_LowLevel() : nullptr;
+							Hit.WorldPosition = WorldPosition;
+							Hit.WorldNormal = WorldNormal;
+							Hit.Distance = FloatDistance;
+							Hit.FaceIndex = FaceIdx;
+							Hit.FaceNormal = FaceNormal;
+
+							// Always compute the best face index. This used to check for non-initial overlap hits, however there's no reason to do that as we should always have valid data.
+							Hit.FaceIndex = FindFaceIndex(Hit, Dir);
+
+							SetFlags(Hit, EHitFlags::Distance | EHitFlags::Normal | EHitFlags::Position | EHitFlags::FaceIndex);
+
+							FTransform StartTransform(InStart);
+							ConvertQueryImpactHit(OwnerComponentInst->GetWorld(), Hit, OutHit, DeltaMag, QueryFilter, InStart, InEnd, &QueryGeometry, StartTransform, false, false);
+							bSweepHit = true;
 						}
 					}
 				}
diff --git a/Engine/Source/Runtime/Experimental/Chaos/Public/Chaos/GeometryQueries.h b/Engine/Source/Runtime/Experimental/Chaos/Public/Chaos/GeometryQueries.h
--- a/Engine/Source/Runtime/Experimental/Chaos/Public/Chaos/GeometryQueries.h
+++ b/Engine/Source/Runtime/Experimental/Chaos/Public/Chaos/GeometryQueries.h
@@ -312,7 +312,7 @@ namespace Chaos
 		}
 
 		//put back into world space
-		if (bResult && (OutTime > 0 || bComputeMTD))
+		if (bResult && (OutTime >= 0 || bComputeMTD))
 		{
 			OutNormal = ATM.TransformVectorNoScale(LocalNormal).GetSafeNormal(UE_KINDA_SMALL_NUMBER, FVec3::AxisVector(0));
 			OutPosition = ATM.TransformPositionNoScale(LocalPosition);
diff --git a/Engine/Source/Runtime/Experimental/Chaos/Public/Chaos/PhysicsObjectCollisionInterface.h b/Engine/Source/Runtime/Experimental/Chaos/Public/Chaos/PhysicsObjectCollisionInterface.h
--- a/Engine/Source/Runtime/Experimental/Chaos/Public/Chaos/PhysicsObjectCollisionInterface.h
+++ b/Engine/Source/Runtime/Experimental/Chaos/Public/Chaos/PhysicsObjectCollisionInterface.h
@@ -191,24 +191,8 @@ namespace Chaos
 			const FVector Delta = EndPos - StartPos;
 			const FReal DeltaMag = Delta.Size();
 
-			// This maintains compatability with SQ querying code (see TSQVisitor::Visit where it checks Sweeps with 0 length).
-			if (DeltaMag < UE_SMALL_NUMBER)
-			{
-				TArray<ChaosInterface::TThreadOverlapHit<Id>> OverlapHits;
-				FMTDInfo MTDInfo;
-				const bool bOverlap = ShapeOverlapWithMTD(InObjects, InGeom, StartTM, OverlapHits, &MTDInfo);
-				if (!OverlapHits.IsEmpty())
-				{
-					const ChaosInterface::TThreadQueryHit<Id>& OverlapQueryHit = OverlapHits[0];
-					ChaosInterface::TThreadQueryHit<Id>& OutQueryHit = OutBestHit;
-					OutQueryHit = OverlapQueryHit;
-					OutBestHit.WorldNormal = MTDInfo.Normal * MTDInfo.Penetration;
-					OutBestHit.Distance = -MTDInfo.Penetration;
-				}
-				return bOverlap;
-			}
-
-			const FVector Dir = Delta / DeltaMag;
+			const bool bSweepHasLength = DeltaMag > 0.f;
+			const FVector Dir = bSweepHasLength ? (Delta / DeltaMag) : FVector(1, 0, 0);
 
 			OutBestHit.Distance = TNumericLimits<float>::Max();
 			for (const FConstPhysicsObjectHandle Object : InObjects)
@@ -252,12 +236,13 @@ namespace Chaos
 								OutBestHit.WorldNormal = WorldNormal;
 								OutBestHit.Distance = FloatDistance;
 								OutBestHit.FaceIndex = FaceIdx;
-								if (OutBestHit.Distance > 0.f)
-								{
-									const FVector LocalPosition = WorldTM.InverseTransformPositionNoScale(OutBestHit.WorldPosition);
-									const FVector LocalUnitDir = WorldTM.InverseTransformVectorNoScale(Dir);
-									OutBestHit.FaceIndex = Shape->GetGeometry()->FindMostOpposingFace(LocalPosition, LocalUnitDir, OutBestHit.FaceIndex, 1);
-								}
+								OutBestHit.FaceNormal = FaceNormal;
+
+								// Always compute the best face index. This used to check for non-initial overlap hits, however there's no reason to do that as we should always have valid data.
+								const FVector LocalPosition = WorldTM.InverseTransformPositionNoScale(OutBestHit.WorldPosition);
+								const FVector LocalUnitDir = WorldTM.InverseTransformVectorNoScale(Dir);
+								OutBestHit.FaceIndex = Shape->GetGeometry()->FindMostOpposingFace(LocalPosition, LocalUnitDir, OutBestHit.FaceIndex, 1);
+
 								OutBestHit.Actor = Interface.GetParticle(IterObject);
 							}
 						}
@@ -265,6 +250,12 @@ namespace Chaos
 					}
 				);
 
+				// Previously, when a zero length sweep was performed, normal was scaled by the penetration. 
+				// This preserves that old behavior until we determine if this is needed.
+				if (bHit && !bSweepHasLength)
+				{
+					OutBestHit.WorldNormal = OutBestHit.WorldNormal * -OutBestHit.Distance;
+				}
 			}
 			return bHit;
 		}
