Heterogeneous Volume Component: Prevent material access outside of game thread

diff --git a/Engine/Source/Runtime/Engine/Private/Components/HeterogeneousVolumeComponent.cpp b/Engine/Source/Runtime/Engine/Private/Components/HeterogeneousVolumeComponent.cpp
--- a/Engine/Source/Runtime/Engine/Private/Components/HeterogeneousVolumeComponent.cpp
+++ b/Engine/Source/Runtime/Engine/Private/Components/HeterogeneousVolumeComponent.cpp
@@ -392,10 +392,13 @@ void UHeterogeneousVolumeComponent::PostInitProperties()
 {
 	Super::PostInitProperties();
 
-	MaterialInstanceDynamic = nullptr;
-	if (UMaterialInterface* MaterialInterface = GetHeterogeneousVolumeMaterial())
+	if (IsInGameThread())
 	{
-		MaterialInstanceDynamic = CreateOrCastToMID(MaterialInterface);
+		MaterialInstanceDynamic = nullptr;
+		if (UMaterialInterface* MaterialInterface = GetHeterogeneousVolumeMaterial())
+		{
+			MaterialInstanceDynamic = CreateOrCastToMID(MaterialInterface);
+		}
 	}
 }
 
@@ -403,10 +406,13 @@ void UHeterogeneousVolumeComponent::PostReinitProperties()
 {
 	Super::PostReinitProperties();
 
-	MaterialInstanceDynamic = nullptr;
-	if (UMaterialInterface* MaterialInterface = GetHeterogeneousVolumeMaterial())
+	if (IsInGameThread())
 	{
-		MaterialInstanceDynamic = CreateOrCastToMID(MaterialInterface);
+		MaterialInstanceDynamic = nullptr;
+		if (UMaterialInterface* MaterialInterface = GetHeterogeneousVolumeMaterial())
+		{
+			MaterialInstanceDynamic = CreateOrCastToMID(MaterialInterface);
+		}
 	}
 }
 
