Morph Targets: mark morph targets as garbage and change their outer when deleting them to avoid serializing them

#rnx

diff --git a/Engine/Source/Runtime/Engine/Private/SkeletalMesh.cpp b/Engine/Source/Runtime/Engine/Private/SkeletalMesh.cpp
--- a/Engine/Source/Runtime/Engine/Private/SkeletalMesh.cpp
+++ b/Engine/Source/Runtime/Engine/Private/SkeletalMesh.cpp
@@ -1423,6 +1423,16 @@ void USkeletalMesh::PostEditUndo()
 		}
 	}
 
+	// ensure that morph targets belong to this skeletal mesh
+	// note: removing a morph target can re-outer it to the transient package and mark it as garbage so it has to be reverted on post-undo if needed  
+	for (UMorphTarget* MorphTarget: GetMorphTargets())
+	{
+		if (MorphTarget)
+		{
+			MorphTarget->ClearGarbage();
+		}
+	}
+	
 	if(GetMorphTargets().Num() > GetMorphTargetIndexMap().Num())
 	{
 		// A morph target remove has been undone, reinitialise
@@ -4442,6 +4452,9 @@ bool USkeletalMesh::RemoveMorphTargets(TConstArrayView<FName> InMorphTargetNames
 
 			UnregisterMorphTarget(MorphTarget);
 
+			MorphTarget->Rename(nullptr, GetTransientPackage(), REN_DoNotDirty | REN_DontCreateRedirectors);
+			MorphTarget->MarkAsGarbage();
+			
 			bRemoved = true;
 		}
 
