FNetGUIDCache::GetObjectFromNetGUID() will no longer attempt to resolve remote objects

diff --git a/Engine/Source/Runtime/Engine/Private/PackageMapClient.cpp b/Engine/Source/Runtime/Engine/Private/PackageMapClient.cpp
--- a/Engine/Source/Runtime/Engine/Private/PackageMapClient.cpp
+++ b/Engine/Source/Runtime/Engine/Private/PackageMapClient.cpp
@@ -3684,10 +3684,16 @@ UObject* FNetGUIDCache::GetObjectFromNetGUID( const FNetworkGUID& NetGUID, const
 		return NULL;
 	}
 
-	UObject* Object = CacheObjectPtr->Object.Get();
+	UObject* Object = nullptr;
 
-	if ( Object != NULL )
+	if (CacheObjectPtr->Object.IsValid()
+#if UE_WITH_REMOTE_OBJECT_HANDLE
+		&& !UE::RemoteObject::Handle::IsRemote(FWeakObjectPtr(CacheObjectPtr->Object).GetRemoteId())
+#endif
+		)
 	{
+		Object = CacheObjectPtr->Object.Get();
+
 		// Either the name should match, or this is dynamic, or we're on the server
 		checkf( Object->GetFName() == CacheObjectPtr->PathName || NetGUID.IsDynamic() || IsNetGUIDAuthority(),
 		        TEXT("ObjectName: '%s', CacheObjectPathName: '%s'"), *Object->GetName(), *CacheObjectPtr->PathName.ToString());
