Fix FExternalPackageHelper::DuplicateExternalPackages duplicating objects not from external actors or external objects, (wrongly) converting them to external objects.
Saving as a level with external entities is still broken, but with this change entities aren't converted to external objects anymore.

diff --git a/Engine/Source/Runtime/Engine/Private/ExternalPackageHelper.cpp b/Engine/Source/Runtime/Engine/Private/ExternalPackageHelper.cpp
--- a/Engine/Source/Runtime/Engine/Private/ExternalPackageHelper.cpp
+++ b/Engine/Source/Runtime/Engine/Private/ExternalPackageHelper.cpp
@@ -63,6 +63,9 @@ void FExternalPackageHelper::DuplicateExternalPackages(const UObject* InObject,
 	if (InDuplicationParameters.DuplicateMode != EDuplicateMode::PIE && InDuplicationParameters.bAssignExternalPackages)
 	{
 		const UPackage* SourcePackage = InObject->GetPackage();
+		const FString SourcePackageName = SourcePackage->GetName();
+		const FString SourcePackageExternalActorsPath = ULevel::GetExternalActorsPath(SourcePackageName);
+		const FString SourcePackageExternalObjectsPath = GetExternalObjectsPath(SourcePackageName);
 		UPackage* DestinationPackage = InDuplicationParameters.DestOuter->GetPackage();
 		
 		FString ReplaceFrom = FPaths::GetBaseFilename(*SourcePackage->GetName());
@@ -75,16 +78,27 @@ void FExternalPackageHelper::DuplicateExternalPackages(const UObject* InObject,
 		{
 			if (UPackage* Package = Object ? Object->GetExternalPackage() : nullptr)
 			{
-				FString Path = Object->GetPathName();
-				if (DestinationPackage != SourcePackage)
+				const FString PackageName = Package->GetName();
+
+				FString SplitPackageRoot;
+				FString SplitPackagePath;
+				FString SplitPackageName;
+				if (FPackageName::SplitLongPackageName(PackageName, SplitPackageRoot, SplitPackagePath, SplitPackageName))
 				{
-					Path = Path.Replace(*ReplaceFrom, *ReplaceTo);
-				}
-				UPackage* DupPackage = Object->IsA<AActor>() ? ULevel::CreateActorPackage(DestinationPackage, ActorPackagingScheme, Path, Object) : FExternalPackageHelper::CreateExternalPackage(DestinationPackage, Path);
-				DupPackage->MarkAsFullyLoaded();
-				DupPackage->MarkPackageDirty();
+					if (SplitPackagePath.StartsWith(FPackagePath::GetExternalActorsFolderName()) || SplitPackagePath.StartsWith(FPackagePath::GetExternalObjectsFolderName()))
+					{
+						FString Path = Object->GetPathName();
+						if (DestinationPackage != SourcePackage)
+						{
+							Path = Path.Replace(*ReplaceFrom, *ReplaceTo);
+						}
+						UPackage* DupPackage = Object->IsA<AActor>() ? ULevel::CreateActorPackage(DestinationPackage, ActorPackagingScheme, Path, Object) : FExternalPackageHelper::CreateExternalPackage(DestinationPackage, Path);
+						DupPackage->MarkAsFullyLoaded();
+						DupPackage->MarkPackageDirty();
 				
-				InDuplicationParameters.DuplicationSeed.Add(Package, DupPackage);
+						InDuplicationParameters.DuplicationSeed.Add(Package, DupPackage);
+					}
+				}
 			}
 		}, /*bIncludeNestedObjects*/ true);
 	}
