Iris
* Don't unregister NetEngineFactories if they were never registered in the first place
* Prevents an ensure on engine shutdown if the engine init process failed before the factory init was called.

diff --git a/Engine/Source/Runtime/Engine/Private/Net/Iris/ReplicationSystem/NetEngineFactories.cpp b/Engine/Source/Runtime/Engine/Private/Net/Iris/ReplicationSystem/NetEngineFactories.cpp
--- a/Engine/Source/Runtime/Engine/Private/Net/Iris/ReplicationSystem/NetEngineFactories.cpp
+++ b/Engine/Source/Runtime/Engine/Private/Net/Iris/ReplicationSystem/NetEngineFactories.cpp
@@ -14,16 +14,33 @@
 namespace UE::Net
 {
 
+namespace Private
+{
+	static bool bAreFactoriesRegistered = false;
+}
+
 void InitEngineNetObjectFactories()
 {
-	FNetObjectFactoryRegistry::RegisterFactory(UNetActorFactory::StaticClass(), UNetActorFactory::GetFactoryName());
-	FNetObjectFactoryRegistry::RegisterFactory(UNetSubObjectFactory::StaticClass(), UNetSubObjectFactory::GetFactoryName());
+	using namespace UE::Net::Private;
+
+	if (ensure(bAreFactoriesRegistered == false))
+	{
+		FNetObjectFactoryRegistry::RegisterFactory(UNetActorFactory::StaticClass(), UNetActorFactory::GetFactoryName());
+		FNetObjectFactoryRegistry::RegisterFactory(UNetSubObjectFactory::StaticClass(), UNetSubObjectFactory::GetFactoryName());
+		bAreFactoriesRegistered = true;
+	}
 }
 
 void ShutdownEngineNetObjectFactories()
 {
-	FNetObjectFactoryRegistry::UnregisterFactory(UNetActorFactory::GetFactoryName());
-	FNetObjectFactoryRegistry::UnregisterFactory(UNetSubObjectFactory::GetFactoryName());
+	using namespace UE::Net::Private;
+
+	if (bAreFactoriesRegistered)
+	{
+		FNetObjectFactoryRegistry::UnregisterFactory(UNetActorFactory::GetFactoryName());
+		FNetObjectFactoryRegistry::UnregisterFactory(UNetSubObjectFactory::GetFactoryName());
+		bAreFactoriesRegistered = false;
+	}
 }
 
 } // end namespace UE::Net
