ULevelStreaming::RequestLevel no longer tests DoesPackageExist before calling LoadPackageAsync (for performance reasons).
Instead, it lets ULevelStreaming::AsyncLevelLoadComplete handle and report the error.

diff --git a/Engine/Source/Runtime/Engine/Classes/Engine/LevelStreaming.h b/Engine/Source/Runtime/Engine/Classes/Engine/LevelStreaming.h
--- a/Engine/Source/Runtime/Engine/Classes/Engine/LevelStreaming.h
+++ b/Engine/Source/Runtime/Engine/Classes/Engine/LevelStreaming.h
@@ -267,6 +267,9 @@ protected:
 	/** Whether the streaming level can safely skip making visible transaction request from the client to the server */
 	uint8 bSkipClientUseMakingVisibleTransactionRequest:1;
 
+	/** Last async load request package path */
+	FPackagePath LastAsyncLoadPackagePath;
+
 	/** IDs used to track the underlying async load requests for the level */
 	TArray<int32> AsyncRequestIDs;
 
diff --git a/Engine/Source/Runtime/Engine/Private/LevelStreaming.cpp b/Engine/Source/Runtime/Engine/Private/LevelStreaming.cpp
--- a/Engine/Source/Runtime/Engine/Private/LevelStreaming.cpp
+++ b/Engine/Source/Runtime/Engine/Private/LevelStreaming.cpp
@@ -1697,72 +1697,64 @@ bool ULevelStreaming::RequestLevel(UWorld* PersistentWorld, bool bAllowLevelLoad
 		NormalizedPackageName = FPackageName::ObjectPathToPackageName(NormalizedPackageName);
 		FPackagePath PackagePath = FPackagePath::FromPackageNameChecked(NormalizedPackageName);
 
-		if (FPackageName::DoesPackageExist(PackagePath, &PackagePath))
-		{
-			SetCurrentState(ELevelStreamingState::Loading);
-			OnLoadingStarted();
+		SetCurrentState(ELevelStreamingState::Loading);
+		OnLoadingStarted();
 			
-			ULevel::StreamedLevelsOwningWorld.Add(DesiredPackageName, PersistentWorld);
-			UWorld::WorldTypePreLoadMap.FindOrAdd(DesiredPackageName) = PersistentWorld->WorldType;
+		ULevel::StreamedLevelsOwningWorld.Add(DesiredPackageName, PersistentWorld);
+		UWorld::WorldTypePreLoadMap.FindOrAdd(DesiredPackageName) = PersistentWorld->WorldType;
 
-			// Kick off async load request.
-			STAT_ADD_CUSTOMMESSAGE_NAME( STAT_NamedMarker, *(FString( TEXT( "RequestLevel - " ) + DesiredPackageName.ToString() )) );
-			TRACE_BOOKMARK(TEXT("RequestLevel - %s"), *DesiredPackageName.ToString());
+		// Kick off async load request.
+		STAT_ADD_CUSTOMMESSAGE_NAME( STAT_NamedMarker, *(FString( TEXT( "RequestLevel - " ) + DesiredPackageName.ToString() )) );
+		TRACE_BOOKMARK(TEXT("RequestLevel - %s"), *DesiredPackageName.ToString());
 			
-			FLinkerInstancingContext* InstancingContextPtr = nullptr;
+		FLinkerInstancingContext* InstancingContextPtr = nullptr;
 #if WITH_EDITOR
-			FLinkerInstancingContext InstancingContext;
-			if (DesiredPackageName != NAME_None && PackagePath.GetPackageFName() != DesiredPackageName)
+		FLinkerInstancingContext InstancingContext;
+		if (DesiredPackageName != NAME_None && PackagePath.GetPackageFName() != DesiredPackageName)
+		{
+			// When loading an instanced package we want to avoid it being processed as an asset so we make sure to set it RF_Transient.
+			// If the package is not created here, it will get created by the LoadPackageAsync call. 
+			// Gameworld packages (PIE) are already ignored so we can let LoadPackageAsync do its job.
+			if (!bIsGameWorld)
 			{
-				// When loading an instanced package we want to avoid it being processed as an asset so we make sure to set it RF_Transient.
-				// If the package is not created here, it will get created by the LoadPackageAsync call. 
-				// Gameworld packages (PIE) are already ignored so we can let LoadPackageAsync do its job.
-				if (!bIsGameWorld)
-				{
-					UPackage* NewPackage = CreatePackage(*DesiredPackageName.ToString());
-					NewPackage->SetFlags(RF_Transient);
-				}
-
-				// When loading an instanced package we need to invoke an instancing context function in case non external actors part of the level are 
-				// pulling on external actors.
-				const FString ExternalActorsPathStr = ULevel::GetExternalActorsPath(PackagePath.GetPackageName());
-				const FString DesiredPackageNameStr = DesiredPackageName.ToString();
-
-				InstancingContext.AddPackageMappingFunc([ExternalActorsPathStr, DesiredPackageNameStr](FName Original)
-				{
-					const FString OriginalStr = Original.ToString();
-					if (OriginalStr.StartsWith(ExternalActorsPathStr))
-					{
-						return FName(*ULevel::GetExternalActorPackageInstanceName(DesiredPackageNameStr, OriginalStr));
-					}
-					return Original;
-				});
-
-				InstancingContextPtr = &InstancingContext;
+				UPackage* NewPackage = CreatePackage(*DesiredPackageName.ToString());
+				NewPackage->SetFlags(RF_Transient);
 			}
-#endif
-			AsyncRequestIDs.Add(LoadPackageAsync(PackagePath, DesiredPackageName, FLoadPackageAsyncDelegate::CreateUObject(this, &ULevelStreaming::AsyncLevelLoadComplete), PackageFlags, PIEInstanceID, GetPriority(), InstancingContextPtr));
 
-			// streamingServer: server loads everything?
-			// Editor immediately blocks on load and we also block if background level streaming is disabled.
-			if (BlockPolicy == AlwaysBlock || (ShouldBeAlwaysLoaded() && BlockPolicy != NeverBlock))
+			// When loading an instanced package we need to invoke an instancing context function in case non external actors part of the level are 
+			// pulling on external actors.
+			const FString ExternalActorsPathStr = ULevel::GetExternalActorsPath(PackagePath.GetPackageName());
+			const FString DesiredPackageNameStr = DesiredPackageName.ToString();
+
+			InstancingContext.AddPackageMappingFunc([ExternalActorsPathStr, DesiredPackageNameStr](FName Original)
 			{
-				if (IsAsyncLoading())
+				const FString OriginalStr = Original.ToString();
+				if (OriginalStr.StartsWith(ExternalActorsPathStr))
 				{
-					UE_LOG(LogStreaming, Display, TEXT("ULevelStreaming::RequestLevel(%s) is flushing async loading"), *DesiredPackageName.ToString());
+					return FName(*ULevel::GetExternalActorPackageInstanceName(DesiredPackageNameStr, OriginalStr));
 				}
+				return Original;
+			});
 
-				// Finish all async loading. Since we will clear our requests upon completion of all loads, 
-				// we take a copy so FlushAsyncLoading won't touch an invalidated array view
-				TArray<int32> LocalRequestIDs = AsyncRequestIDs;
-				FlushAsyncLoading(LocalRequestIDs);
-			}
+			InstancingContextPtr = &InstancingContext;
 		}
-		else
+#endif
+		LastAsyncLoadPackagePath = PackagePath;
+		AsyncRequestIDs.Add(LoadPackageAsync(PackagePath, DesiredPackageName, FLoadPackageAsyncDelegate::CreateUObject(this, &ULevelStreaming::AsyncLevelLoadComplete), PackageFlags, PIEInstanceID, GetPriority(), InstancingContextPtr));
+
+		// streamingServer: server loads everything?
+		// Editor immediately blocks on load and we also block if background level streaming is disabled.
+		if (BlockPolicy == AlwaysBlock || (ShouldBeAlwaysLoaded() && BlockPolicy != NeverBlock))
 		{
-			UE_LOG(LogStreaming, Error,TEXT("Couldn't find file for package %s."), *PackagePath.GetDebugName());
-			SetCurrentState(ELevelStreamingState::FailedToLoad);
-			return false;
+			if (IsAsyncLoading())
+			{
+				UE_LOG(LogStreaming, Display, TEXT("ULevelStreaming::RequestLevel(%s) is flushing async loading"), *DesiredPackageName.ToString());
+			}
+
+			// Finish all async loading. Since we will clear our requests upon completion of all loads, 
+			// we take a copy so FlushAsyncLoading won't touch an invalidated array view
+			TArray<int32> LocalRequestIDs = AsyncRequestIDs;
+			FlushAsyncLoading(LocalRequestIDs);
 		}
 	}
 
@@ -1774,6 +1766,7 @@ void ULevelStreaming::AsyncLevelLoadComplete(const FName& InPackageName, UPackag
 	// TODO: Should not set state here so that observers will have access to LoadedLevel
 	SetCurrentState(ELevelStreamingState::LoadedNotVisible);
 	OnLoadingFinished();
+	FPackagePath PackagePath = MoveTemp(LastAsyncLoadPackagePath);
 
 	if (InLoadedPackage)
 	{
@@ -1882,8 +1875,11 @@ void ULevelStreaming::AsyncLevelLoadComplete(const FName& InPackageName, UPackag
 	}
 	else
 	{
-		UE_LOG(LogLevelStreaming, Warning, TEXT("Failed to load package '%s'"), *InPackageName.ToString() );
-		
+		UE_LOG(LogLevelStreaming, Warning, TEXT("Failed to load package '%s'"), *InPackageName.ToString());
+		if (!FPackageName::DoesPackageExist(PackagePath, &PackagePath))
+		{
+			UE_LOG(LogLevelStreaming, Error, TEXT("Couldn't find file for package %s."), *PackagePath.GetDebugName());
+		}
 		SetCurrentState(ELevelStreamingState::FailedToLoad);
  		SetShouldBeLoaded(false);
 	}
