Fixed UInstancedStaticMeshComponent invalid navigation bounds when UpdateBounds is called while static mesh was compiling.

diff --git a/Engine/Source/Runtime/Engine/Classes/Components/InstancedStaticMeshComponent.h b/Engine/Source/Runtime/Engine/Classes/Components/InstancedStaticMeshComponent.h
--- a/Engine/Source/Runtime/Engine/Classes/Components/InstancedStaticMeshComponent.h
+++ b/Engine/Source/Runtime/Engine/Classes/Components/InstancedStaticMeshComponent.h
@@ -525,6 +525,7 @@ public:
 	ENGINE_API virtual void GetResourceSizeEx(FResourceSizeEx& CumulativeResourceSize) override;
 	ENGINE_API void BeginDestroy() override;
 #if WITH_EDITOR
+	ENGINE_API virtual void PostStaticMeshCompilation() override;
 	ENGINE_API virtual void PostEditChangeChainProperty(FPropertyChangedChainEvent& PropertyChangedEvent) override;
 	ENGINE_API virtual void PostEditUndo() override;
 	/**
diff --git a/Engine/Source/Runtime/Engine/Private/InstancedStaticMesh.cpp b/Engine/Source/Runtime/Engine/Private/InstancedStaticMesh.cpp
--- a/Engine/Source/Runtime/Engine/Private/InstancedStaticMesh.cpp
+++ b/Engine/Source/Runtime/Engine/Private/InstancedStaticMesh.cpp
@@ -5621,6 +5621,21 @@ void UInstancedStaticMeshComponent::SetLODDistanceScale(float InLODDistanceScale
 }
 
 #if WITH_EDITOR
+void UInstancedStaticMeshComponent::PostStaticMeshCompilation()
+{
+	// When UpdateBounds calls CalcAndCacheNavigationBounds, it will call GetInstanceNavigationBounds which calls UStaticMesh::GetNavCollision.
+	// If the mesh was compiling, it will wait for compilation to finish. 
+	// We override PostStaticMeshCompilation to intercept the call and make sure that UpdateBounds completes before 
+	// UStaticMeshComponent::PostStaticMeshCompilation calls FNavigationSubsystem::UpdateComponentData, which will end up calling 
+	// FNavigationElement::InitializeFromInterface (which calls UInstancedStaticMeshComponent::GetNavigationBounds).
+	if (bIsUpdatingBounds)
+	{
+		UpdateBounds();
+	}
+
+	Super::PostStaticMeshCompilation();
+}
+
 void UInstancedStaticMeshComponent::PostEditChangeChainProperty(FPropertyChangedChainEvent& PropertyChangedEvent)
 {
 	if(PropertyChangedEvent.Property != NULL)
