Substrate - remove some loops from he HSL translator. (preparing for select different BSDFs than only slab)

#tests pc dx12 and dx11.
#rnx

diff --git a/Engine/Source/Runtime/Core/Private/UObject/FortniteMainBranchObjectVersion.cpp b/Engine/Source/Runtime/Core/Private/UObject/FortniteMainBranchObjectVersion.cpp
--- a/Engine/Source/Runtime/Core/Private/UObject/FortniteMainBranchObjectVersion.cpp
+++ b/Engine/Source/Runtime/Core/Private/UObject/FortniteMainBranchObjectVersion.cpp
@@ -14,6 +14,6 @@ TMap<FGuid, FGuid> FFortniteMainBranchObjectVersion::GetSystemGuids()
 	SystemGuids.Add(DevGuids.Niagara_LatestScriptCompileVersion, FGuid("393081597FA7234C92F00869DE262D73"));
 	SystemGuids.Add(DevGuids.SkeletalMeshDerivedDataVersion, FGuid("B9AEAA3EE7AC4FCFB3860E8F3E06DD75"));
 	SystemGuids.Add(DevGuids.STATICMESH_DERIVEDDATA_VER, FGuid("CDFDB82FE0A83DB602AE2EFFD8BB7F42"));
-	SystemGuids.Add(DevGuids.MaterialTranslationDDCVersion, FGuid("E779012A63684F46B85596E4E89D489A"));
+	SystemGuids.Add(DevGuids.MaterialTranslationDDCVersion, FGuid("77BE52870E484BE78419760A6CD0C677"));
 	return SystemGuids;
 }
diff --git a/Engine/Source/Runtime/Engine/Private/Materials/HLSLMaterialTranslator.cpp b/Engine/Source/Runtime/Engine/Private/Materials/HLSLMaterialTranslator.cpp
--- a/Engine/Source/Runtime/Engine/Private/Materials/HLSLMaterialTranslator.cpp
+++ b/Engine/Source/Runtime/Engine/Private/Materials/HLSLMaterialTranslator.cpp
@@ -12714,7 +12714,8 @@ bool FHLSLMaterialTranslator::FSubstrateCompilationContext::SubstrateGenerateDer
 	}
 
 	//
-	// Evaluate the one and only root node
+	// Evaluate the one and only root node.
+	// And make sure each and every path of the Substrate tree have valid children and path.
 	//
 	int32 RootIndex = INDEX_NONE;
 	for (auto& It : SubstrateMaterialExpressionRegisteredOperators)
@@ -12731,19 +12732,7 @@ bool FHLSLMaterialTranslator::FSubstrateCompilationContext::SubstrateGenerateDer
 			bSubstrateWritesEmissive |= It.bBSDFWritesEmissive > 0;
 			bSubstrateWritesAmbientOcclusion |= It.bBSDFWritesAmbientOcclusion > 0;
 		}
-	}
-	SubstrateMaterialRootOperator = &SubstrateMaterialExpressionRegisteredOperators[RootIndex];
-	if (!SubstrateMaterialRootOperator)
-	{
-		Compiler->Errorf(TEXT("Cannot find the root of the Substrate Tree for Material %s (asset: %s).\r\n"), *CompilerMaterial->GetDebugName(), *CompilerMaterial->GetAssetPath().ToString());
-		return false;
-	}
 
-	//
-	// Make sure each and every path of the Substrate tree have valid children and path.
-	//
-	for (auto& It : SubstrateMaterialExpressionRegisteredOperators)
-	{
 		if (It.IsDiscarded())
 		{
 			continue; // ignore discarded operations in sub tree using parameter blending
@@ -12789,17 +12778,23 @@ bool FHLSLMaterialTranslator::FSubstrateCompilationContext::SubstrateGenerateDer
 			return false;
 		}
 	}
+	SubstrateMaterialRootOperator = &SubstrateMaterialExpressionRegisteredOperators[RootIndex];
+	if (!SubstrateMaterialRootOperator)
+	{
+		Compiler->Errorf(TEXT("Cannot find the root of the Substrate Tree for Material %s (asset: %s).\r\n"), *CompilerMaterial->GetDebugName(), *CompilerMaterial->GetAssetPath().ToString());
+		return false;
+	}
 
+	// Evaluate if simplification is needed.
 	SubstrateSimplificationStatus.bRunFullSimplification = Compiler->SubstrateCompilationConfig.bFullSimplify || CompilationContextIndex == ESubstrateCompilationContext::SCC_FullySimplified;
-
 	if (!SubstrateSimplificationStatus.bRunFullSimplification)
 	{
 		//
-		// Generate LayerDepth value for all operators/bsdfs for progressive material simplification
+		// Generate LayerDepth value for all operators/bsdfs for progressive material simplification and sort them.
 		//
 		int VOpTopBranchCountTaken = 0;
 		int VOpBottomBranchCountTaken = 0;
-		std::function<void(FSubstrateOperator&)> WalkOperatorsForDepth = [&](FSubstrateOperator& CurrentOperator) -> void
+		std::function<void(FSubstrateOperator&)> WalkOperatorsForSimplification = [&](FSubstrateOperator& CurrentOperator) -> void
 		{
 			CurrentOperator.LayerDepth = VOpBottomBranchCountTaken;
 			switch (CurrentOperator.OperatorType)
@@ -12807,10 +12802,10 @@ bool FHLSLMaterialTranslator::FSubstrateCompilationContext::SubstrateGenerateDer
 			case SUBSTRATE_OPERATOR_VERTICAL:
 			{
 				VOpTopBranchCountTaken++;
-				WalkOperatorsForDepth(SubstrateMaterialExpressionRegisteredOperators[CurrentOperator.LeftIndex]);
+				WalkOperatorsForSimplification(SubstrateMaterialExpressionRegisteredOperators[CurrentOperator.LeftIndex]);
 				VOpTopBranchCountTaken--;
 				VOpBottomBranchCountTaken++;
-				WalkOperatorsForDepth(SubstrateMaterialExpressionRegisteredOperators[CurrentOperator.RightIndex]);
+				WalkOperatorsForSimplification(SubstrateMaterialExpressionRegisteredOperators[CurrentOperator.RightIndex]);
 				VOpBottomBranchCountTaken--;
 				break;
 			}
@@ -12818,13 +12813,13 @@ bool FHLSLMaterialTranslator::FSubstrateCompilationContext::SubstrateGenerateDer
 			case SUBSTRATE_OPERATOR_ADD:
 			case SUBSTRATE_OPERATOR_SELECT:
 			{
-				WalkOperatorsForDepth(SubstrateMaterialExpressionRegisteredOperators[CurrentOperator.LeftIndex]);
-				WalkOperatorsForDepth(SubstrateMaterialExpressionRegisteredOperators[CurrentOperator.RightIndex]);
+				WalkOperatorsForSimplification(SubstrateMaterialExpressionRegisteredOperators[CurrentOperator.LeftIndex]);
+				WalkOperatorsForSimplification(SubstrateMaterialExpressionRegisteredOperators[CurrentOperator.RightIndex]);
 				break;
 			}
 			case SUBSTRATE_OPERATOR_WEIGHT:
 			{
-				WalkOperatorsForDepth(SubstrateMaterialExpressionRegisteredOperators[CurrentOperator.LeftIndex]);
+				WalkOperatorsForSimplification(SubstrateMaterialExpressionRegisteredOperators[CurrentOperator.LeftIndex]);
 				break;
 			}
 			case SUBSTRATE_OPERATOR_BSDF:
@@ -12832,42 +12827,25 @@ bool FHLSLMaterialTranslator::FSubstrateCompilationContext::SubstrateGenerateDer
 				break;
 			}
 			}
-		};
-		WalkOperatorsForDepth(*SubstrateMaterialRootOperator);
 
-		//
-		// Make sure each and every path of the Substrate tree have valid children and path.
-		//
-		for (auto& It : SubstrateMaterialExpressionRegisteredOperators)
-		{
-			switch (It.OperatorType)
+			// Add operators to simplify with a priority.
+			switch (CurrentOperator.OperatorType)
 			{
-				// Operators without any child
-			case SUBSTRATE_OPERATOR_BSDF:
-				break;
-
-				// Operators with a single child
-			case SUBSTRATE_OPERATOR_WEIGHT:
-				break;
-
-				// Operators always has parameter blending enabled
-			case SUBSTRATE_OPERATOR_SELECT:
-				break;
-
-			// Operators with two children
-			case SUBSTRATE_OPERATOR_HORIZONTAL:
 			case SUBSTRATE_OPERATOR_VERTICAL:
+			case SUBSTRATE_OPERATOR_HORIZONTAL:
 			case SUBSTRATE_OPERATOR_ADD:
 			{
 				FSubstrateSimplificationStatus::FOperatorToSimplify OperatorToSimplify;
-				OperatorToSimplify.Data.Index = It.Index;
-				OperatorToSimplify.Data.Depth = It.LayerDepth;
+				OperatorToSimplify.Data.Index = CurrentOperator.Index;
+				OperatorToSimplify.Data.Depth = CurrentOperator.LayerDepth;
 				SubstrateSimplificationStatus.OperatorSimplificationOrder.Push(OperatorToSimplify);
+				break;
 			}
-			break;
 			}
-			SubstrateSimplificationStatus.OperatorSimplificationOrder.Sort();	// sort according to depth
-		}
+		};
+		WalkOperatorsForSimplification(*SubstrateMaterialRootOperator);
+
+		SubstrateSimplificationStatus.OperatorSimplificationOrder.Sort();	// sort according to depth
 	}
 
 	EShaderPlatform ShaderPlatform = Compiler->GetShaderPlatform();
@@ -13055,7 +13033,7 @@ bool FHLSLMaterialTranslator::FSubstrateCompilationContext::SubstrateGenerateDer
 		}
 
 		//
-		// Make sure all the types have valid children operator indices
+		// Make sure all the types have valid children operator indices according to the parameter blending setup for this simplification status.
 		//
 		for (const auto& It : SubstrateMaterialExpressionRegisteredOperators)
 		{
@@ -13143,7 +13121,7 @@ bool FHLSLMaterialTranslator::FSubstrateCompilationContext::SubstrateGenerateDer
 
 				if (It.OperatorType == SUBSTRATE_OPERATOR_BSDF)
 				{
-					// Recursively parse all node from BSDF to the root node and update the necessary properties.
+					// Recursively parse all nodes from BSDF to the root node and update the necessary properties.
 					WalkOperatorsToRoot(It);
 				}
 			}
@@ -13151,7 +13129,7 @@ bool FHLSLMaterialTranslator::FSubstrateCompilationContext::SubstrateGenerateDer
 
 		//
 		// Compute IsTop or IsBottom layer using a depth first tree visit while counting vertical right and left branches taken.
-		// When a BSDF is encoutered, and it is the root of a parameter blend subtree, we continue to parse, setting the same information as the root.
+		// When a BSDF is encountered, and it is the root of a parameter blend subtree, we continue to parse, setting the same information as the root.
 		//
 		{
 			int VOpTopBranchCountTaken = 0;
@@ -13239,7 +13217,7 @@ bool FHLSLMaterialTranslator::FSubstrateCompilationContext::SubstrateGenerateDer
 		}
 
 		//
-		// Compute the size of the material
+		// Compute the per pixel byte count required by the materials.
 		//
 		{
 			// Compute the shared local basis count only
