[HWRT] Only update spline mesh BLAS if spline parameters changed, it's using WPO or geometry is not valid.
* Add r.RayTracing.Geometry.SplineMeshes.ForceUpdate to force build spline BLAS every frame.

#rnx

diff --git a/Engine/Source/Runtime/Engine/Private/SplineMeshSceneProxy.cpp b/Engine/Source/Runtime/Engine/Private/SplineMeshSceneProxy.cpp
--- a/Engine/Source/Runtime/Engine/Private/SplineMeshSceneProxy.cpp
+++ b/Engine/Source/Runtime/Engine/Private/SplineMeshSceneProxy.cpp
@@ -19,6 +19,10 @@ static TAutoConsoleVariable<int32> CVarRayTracingSplineMeshes(
 	1,
 	TEXT("Include splines meshes in ray tracing effects (default = 1 (spline meshes enabled in ray tracing))"));
 
+static TAutoConsoleVariable<int32> CVarRayTracingSplineMeshesForceUpdate(
+	TEXT("r.RayTracing.Geometry.SplineMeshes.ForceUpdate"),
+	0,
+	TEXT("Forces ray tracing geometry for spline meshes to be updated every frame"));
 
 bool FSplineMeshVertexFactory::ShouldCompilePermutation(const FVertexFactoryShaderPermutationParameters& Parameters)
 {
@@ -223,7 +227,9 @@ void FSplineMeshSceneProxy::GetDynamicRayTracingInstances(FRayTracingInstanceCol
 	RayTracingInstance.InstanceTransformsView = MakeArrayView(&ThisLocalToWorld, 1);
 	RayTracingInstance.MaterialsView = MakeArrayView(CachedRayTracingMaterials);
 
-	if (RenderData->LODVertexFactories[LODIndex].VertexFactory.GetType()->SupportsRayTracingDynamicGeometry())
+	const bool bNeedsUpdate = bUpdateRayTracingGeometry || bAnyMaterialHasWorldPositionOffset || !Geometry.IsValid() || (CVarRayTracingSplineMeshesForceUpdate.GetValueOnRenderThread() != 0);
+
+	if (RenderData->LODVertexFactories[LODIndex].VertexFactory.GetType()->SupportsRayTracingDynamicGeometry() && bNeedsUpdate)
 	{
 		Collector.AddRayTracingGeometryUpdate(
 			FRayTracingDynamicGeometryUpdateParams
@@ -234,10 +240,12 @@ void FSplineMeshSceneProxy::GetDynamicRayTracingInstances(FRayTracingInstanceCol
 				uint32((SIZE_T)LODModel.GetNumVertices() * sizeof(FVector3f)),
 				Geometry.Initializer.TotalPrimitiveCount,
 				&Geometry,
-				nullptr /* VertexBuffer */,
+				&RayTracingDynamicVertexBuffer,
 				true
 			}
 		);
+
+		bUpdateRayTracingGeometry = false;
 	}
 
 	check(CachedRayTracingMaterials.Num() == RayTracingInstance.GetMaterials().Num());
diff --git a/Engine/Source/Runtime/Engine/Public/SplineMeshSceneProxy.h b/Engine/Source/Runtime/Engine/Public/SplineMeshSceneProxy.h
--- a/Engine/Source/Runtime/Engine/Public/SplineMeshSceneProxy.h
+++ b/Engine/Source/Runtime/Engine/Public/SplineMeshSceneProxy.h
@@ -138,6 +138,8 @@ protected:
 				);
 			}
 		}
+
+		bUpdateRayTracingGeometry = true;
 	}
 
 protected:
@@ -145,6 +147,9 @@ protected:
 	FSplineMeshShaderParams SplineParams;
 
 	FSplineMeshSceneInstanceDataBuffers SplineMeshInstanceData;
+
+	bool bUpdateRayTracingGeometry = true;
+	FRWBuffer RayTracingDynamicVertexBuffer;
 private:
 	/** implemented by derived to provide access by parent */
 	virtual TDerived& Downcast() = 0;
