[Engine]
* Changed so StaticMeshSceneProxy is getting feature level from scene instead of world (same as skeletalmeshsceneproxy)

diff --git a/Engine/Source/Runtime/Engine/Private/StaticMeshSceneProxy.cpp b/Engine/Source/Runtime/Engine/Private/StaticMeshSceneProxy.cpp
--- a/Engine/Source/Runtime/Engine/Private/StaticMeshSceneProxy.cpp
+++ b/Engine/Source/Runtime/Engine/Private/StaticMeshSceneProxy.cpp
@@ -2143,7 +2143,7 @@ FStaticMeshSceneProxy::FLODInfo::FLODInfo(const FStaticMeshSceneProxyDesc& InPro
 	, OverrideColorVertexBuffer(nullptr)
 	, bUsesMeshModifyingMaterials(false)
 {
-	const auto FeatureLevel =  InProxyDesc.GetWorld()->GetFeatureLevel();
+	const auto FeatureLevel =  InProxyDesc.Scene->GetFeatureLevel();
 
 	FStaticMeshRenderData* MeshRenderData = InProxyDesc.GetStaticMesh()->GetRenderData();
 	FStaticMeshLODResources& LODModel = MeshRenderData->LODResources[LODIndex];
@@ -2593,7 +2593,7 @@ void FStaticMeshSceneProxyDesc::InitializeFromStaticMeshComponent(const UStaticM
 	WireframeColor = InComponent->GetWireframeColor();
 	LODParentPrimitive = InComponent->GetLODParentPrimitive();	
 
-	SetMaterialRelevance(InComponent->GetMaterialRelevance(World->GetFeatureLevel()));
+	SetMaterialRelevance(InComponent->GetMaterialRelevance(GetScene()->GetFeatureLevel()));
 	SetCollisionResponseToChannels(InComponent->GetCollisionResponseToChannels());
 
 	MeshPaintTexture = InComponent->MeshPaintTextureOverride ? InComponent->MeshPaintTextureOverride.Get() : InComponent->GetMeshPaintTexture();
