Fixed infinite retries of ULevelStreaming::RequestLevel when provided package doesn't contain a world.
Instead set level streaming current state to FailedToLoad and ShouldBeLoaded to false.

diff --git a/Engine/Source/Runtime/Engine/Private/LevelStreaming.cpp b/Engine/Source/Runtime/Engine/Private/LevelStreaming.cpp
--- a/Engine/Source/Runtime/Engine/Private/LevelStreaming.cpp
+++ b/Engine/Source/Runtime/Engine/Private/LevelStreaming.cpp
@@ -1865,6 +1865,12 @@ void ULevelStreaming::AsyncLevelLoadComplete(const FName& InPackageName, UPackag
 					SetWorldAsset(DestinationWorld);
 				}
 			}
+			else
+			{
+				UE_LOG(LogLevelStreaming, Warning, TEXT("Failed to find a world in package '%s'"), *InPackageName.ToString());
+				SetCurrentState(ELevelStreamingState::FailedToLoad);
+				SetShouldBeLoaded(false);
+			}
 		}
 	}
 	else if (Result == EAsyncLoadingResult::Canceled)
