Rename confusing IsNonPIEEditorOnly actor function to IsEditorOnlyLoadedInPIE.

diff --git a/Engine/Plugins/MovieScene/MovieRenderPipeline/Source/MovieRenderPipelineEditor/Private/MoviePipelineFunctionalTestBase.h b/Engine/Plugins/MovieScene/MovieRenderPipeline/Source/MovieRenderPipelineEditor/Private/MoviePipelineFunctionalTestBase.h
--- a/Engine/Plugins/MovieScene/MovieRenderPipeline/Source/MovieRenderPipelineEditor/Private/MoviePipelineFunctionalTestBase.h
+++ b/Engine/Plugins/MovieScene/MovieRenderPipeline/Source/MovieRenderPipelineEditor/Private/MoviePipelineFunctionalTestBase.h
@@ -31,7 +31,7 @@ protected:
 	virtual bool IsReady_Implementation() override;
 	virtual void StartTest() override;
 	// ~AFunctionalTest
-	virtual bool IsNonPIEEditorOnly() const override
+	virtual bool IsEditorOnlyLoadedInPIE() const override
 	{
 		return true;
 	}
diff --git a/Engine/Plugins/Tests/FunctionalTestingEditor/Source/Public/EditorFunctionalTest.h b/Engine/Plugins/Tests/FunctionalTestingEditor/Source/Public/EditorFunctionalTest.h
--- a/Engine/Plugins/Tests/FunctionalTestingEditor/Source/Public/EditorFunctionalTest.h
+++ b/Engine/Plugins/Tests/FunctionalTestingEditor/Source/Public/EditorFunctionalTest.h
@@ -16,7 +16,7 @@ class AEditorFunctionalTest : public AFunctionalTest
 		return true;
 	}
 
-	virtual bool IsNonPIEEditorOnly() const override
+	virtual bool IsEditorOnlyLoadedInPIE() const override
 	{
 		return true;
 	}
@@ -33,7 +33,7 @@ class AEditorScreenshotFunctionalTest : public AScreenshotFunctionalTest
 		return true;
 	}
 
-	virtual bool IsNonPIEEditorOnly() const override
+	virtual bool IsEditorOnlyLoadedInPIE() const override
 	{
 		return true;
 	}
diff --git a/Engine/Source/Developer/FunctionalTesting/Private/FunctionalTest.cpp b/Engine/Source/Developer/FunctionalTesting/Private/FunctionalTest.cpp
--- a/Engine/Source/Developer/FunctionalTesting/Private/FunctionalTest.cpp
+++ b/Engine/Source/Developer/FunctionalTesting/Private/FunctionalTest.cpp
@@ -253,15 +253,15 @@ bool AFunctionalTest::RunTest(const TArray<FString>& Params)
 	UWorld* World = GetWorld();
 	#if WITH_EDITOR
 	const bool bIsEditorOnlyObject = IsEditorOnlyObject(this);
-	const bool bIsNonPIEEditorOnly = IsNonPIEEditorOnly();
+	const bool bIsEditorOnlyLoadedInPIE = IsEditorOnlyLoadedInPIE();
 
-	ensure(World->HasBegunPlay() || (World->WorldType == EWorldType::Editor && bIsEditorOnlyObject && !bIsNonPIEEditorOnly));
+	ensure(World->HasBegunPlay() || (World->WorldType == EWorldType::Editor && bIsEditorOnlyObject && !bIsEditorOnlyLoadedInPIE));
 
-	if (bIsEditorOnlyObject && !bIsNonPIEEditorOnly)
+	if (bIsEditorOnlyObject && !bIsEditorOnlyLoadedInPIE)
 	{
 		if (GetLinkerCustomVersion(FFortniteMainBranchObjectVersion::GUID) < FFortniteMainBranchObjectVersion::FunctionalTestCanRunInEditorWorld)
 		{
-			AddWarning(FString::Printf(TEXT("%s changed behavior, return true in IsNonPIEEditorOnly to run in PIE world, otherwise resave the asset to acknowledge behavior change"), *GetName()));
+			AddWarning(FString::Printf(TEXT("%s changed behavior, return true in IsEditorOnlyLoadedInPIE to run in PIE world, otherwise resave the asset to acknowledge behavior change"), *GetName()));
 		}
 	}
 	#else
diff --git a/Engine/Source/Developer/FunctionalTesting/Private/Tests/ClientFuncTestPerforming.cpp b/Engine/Source/Developer/FunctionalTesting/Private/Tests/ClientFuncTestPerforming.cpp
--- a/Engine/Source/Developer/FunctionalTesting/Private/Tests/ClientFuncTestPerforming.cpp
+++ b/Engine/Source/Developer/FunctionalTesting/Private/Tests/ClientFuncTestPerforming.cpp
@@ -203,7 +203,7 @@ bool FClientFunctionalTestingMapsEditor::RunTest(const FString& Parameters)
 	{
 		// we have a world to use - but does it have the test we now need to run?
 		AFunctionalTest* Test = FindObjectFast<AFunctionalTest>(WorldToUse->PersistentLevel, FName(MapTestName));
-		if(Test && (!IsEditorOnlyObject(Test) || Test->IsNonPIEEditorOnly()))
+		if(Test && (!IsEditorOnlyObject(Test) || Test->IsEditorOnlyLoadedInPIE()))
 		{
 			// it does, just use base test running logic, as it will be faster than
 			// reopening the map:
@@ -233,7 +233,7 @@ bool FClientFunctionalTestingMapsEditor::RunTest(const FString& Parameters)
 		AddCommand(new FFunctionLatentCommand([MapTestName, MapPackageName]() {
 			UWorld* World = GEditor->GetEditorWorldContext().World();
 			AFunctionalTest* Test = FindObjectFast<AFunctionalTest>(World->PersistentLevel, FName(MapTestName));
-			if(ensure(Test) && Test->IsNonPIEEditorOnly())
+			if(ensure(Test) && Test->IsEditorOnlyLoadedInPIE())
 			{
 				// request PIE, as this test actor claims it supports PIE:
 				const bool bCanProceed = AutomationOpenMap(MapPackageName);
diff --git a/Engine/Source/Editor/Blutility/Private/BlutilityModule.cpp b/Engine/Source/Editor/Blutility/Private/BlutilityModule.cpp
--- a/Engine/Source/Editor/Blutility/Private/BlutilityModule.cpp
+++ b/Engine/Source/Editor/Blutility/Private/BlutilityModule.cpp
@@ -240,7 +240,7 @@ public:
 				{
 					if(const AActor* NativeCDO = Cast<AActor>(NativeParent->ClassDefaultObject))
 					{
-						return !NativeCDO->IsNonPIEEditorOnly();
+						return !NativeCDO->IsEditorOnlyLoadedInPIE();
 					}
 				}
 				return true;
diff --git a/Engine/Source/Runtime/Engine/Classes/GameFramework/Actor.h b/Engine/Source/Runtime/Engine/Classes/GameFramework/Actor.h
--- a/Engine/Source/Runtime/Engine/Classes/GameFramework/Actor.h
+++ b/Engine/Source/Runtime/Engine/Classes/GameFramework/Actor.h
@@ -2303,7 +2303,7 @@ public:
 	ENGINE_API virtual bool CanBeInCluster() const override;
 	static ENGINE_API void AddReferencedObjects(UObject* InThis, FReferenceCollector& Collector);
 	ENGINE_API virtual bool IsEditorOnly() const override;
-	ENGINE_API virtual bool IsNonPIEEditorOnly() const { return false; }
+	ENGINE_API virtual bool IsEditorOnlyLoadedInPIE() const { return false; }
 	ENGINE_API virtual bool IsRuntimeOnly() const { return false; }
 	ENGINE_API virtual bool IsAsset() const override;
 
diff --git a/Engine/Source/Runtime/Engine/Private/WorldPartition/WorldPartition.cpp b/Engine/Source/Runtime/Engine/Private/WorldPartition/WorldPartition.cpp
--- a/Engine/Source/Runtime/Engine/Private/WorldPartition/WorldPartition.cpp
+++ b/Engine/Source/Runtime/Engine/Private/WorldPartition/WorldPartition.cpp
@@ -582,7 +582,7 @@ void UWorldPartition::PostDuplicate(EDuplicateMode::Type DuplicateMode)
 				{
 					if (AActor* Actor = Cast<AActor>(InObj))
 					{
-						if (Actor->IsEditorOnly() && !Actor->IsNonPIEEditorOnly())
+						if (Actor->IsEditorOnly() && !Actor->IsEditorOnlyLoadedInPIE())
 						{
 							InObj = nullptr;
 							UE_LOG(LogWorldPartition, Log, TEXT("Cleared actor editor-only reference '%s' -> '%s'."), *Root->GetActorNameOrLabel(), *Actor->GetActorNameOrLabel());							
diff --git a/Engine/Source/Runtime/Engine/Private/WorldPartition/WorldPartitionActorDesc.cpp b/Engine/Source/Runtime/Engine/Private/WorldPartition/WorldPartitionActorDesc.cpp
--- a/Engine/Source/Runtime/Engine/Private/WorldPartition/WorldPartitionActorDesc.cpp
+++ b/Engine/Source/Runtime/Engine/Private/WorldPartition/WorldPartitionActorDesc.cpp
@@ -432,9 +432,9 @@ void FWorldPartitionActorDesc::SerializeTo(TArray<uint8>& OutData, FWorldPartiti
 	OutData.Append(PayloadData);
 }
 
-bool FWorldPartitionActorDesc::GetActorIsNonPIEEditorOnly() const
+bool FWorldPartitionActorDesc::GetActorIsEditorOnlyLoadedInPIE() const
 {
-	return ActorNativeClass->GetDefaultObject<AActor>()->IsNonPIEEditorOnly();
+	return ActorNativeClass->GetDefaultObject<AActor>()->IsEditorOnlyLoadedInPIE();
 }
 
 TArray<FName> FWorldPartitionActorDesc::GetDataLayers(bool bIncludeExternalDataLayer) const
diff --git a/Engine/Source/Runtime/Engine/Private/WorldPartition/WorldPartitionRuntimeLevelStreamingCell.cpp b/Engine/Source/Runtime/Engine/Private/WorldPartition/WorldPartitionRuntimeLevelStreamingCell.cpp
--- a/Engine/Source/Runtime/Engine/Private/WorldPartition/WorldPartitionRuntimeLevelStreamingCell.cpp
+++ b/Engine/Source/Runtime/Engine/Private/WorldPartition/WorldPartitionRuntimeLevelStreamingCell.cpp
@@ -291,7 +291,7 @@ void UWorldPartitionRuntimeLevelStreamingCell::AddActorToCell(const FStreamingGe
 {
 	// Non-spatially loaded actors coming from level instances are not moved into the persistent level in PIE but rather placed into an always loaded cell,
 	// because that would imply loading them through an instancing context.
-	check(!ActorDescView.GetActorIsEditorOnly() || ActorDescView.GetActorIsNonPIEEditorOnly());
+	check(!ActorDescView.GetActorIsEditorOnly() || ActorDescView.GetActorIsEditorOnlyLoadedInPIE());
 
 	const UActorDescContainerInstance* ContainerInstance = ActorDescView.GetContainerInstance();
 	check(ContainerInstance);
diff --git a/Engine/Source/Runtime/Engine/Private/WorldPartition/WorldPartitionStreamingGeneration.cpp b/Engine/Source/Runtime/Engine/Private/WorldPartition/WorldPartitionStreamingGeneration.cpp
--- a/Engine/Source/Runtime/Engine/Private/WorldPartition/WorldPartitionStreamingGeneration.cpp
+++ b/Engine/Source/Runtime/Engine/Private/WorldPartition/WorldPartitionStreamingGeneration.cpp
@@ -799,7 +799,7 @@ class FWorldPartitionStreamingGenerator
 			}
 
 			const bool bActorIsEditorOnly = InActor ? InActor->IsEditorOnly() : InActorDescView.GetActorIsEditorOnly();
-			const bool bActorIsNonPIEEditorOnly = InActorDescView.GetActorIsNonPIEEditorOnly();
+			const bool bActorIsNonPIEEditorOnly = InActorDescView.GetActorIsEditorOnlyLoadedInPIE();
 
 			if (!bActorIsEditorOnly || (bIsPIE && bActorIsNonPIEEditorOnly))
 			{
@@ -1322,11 +1322,11 @@ class FWorldPartitionStreamingGenerator
 				}
 
 				// Non-PIE editor-only actors can only be non-spatially loaded and can't have data layers
-				if (ActorDescView.GetActorIsNonPIEEditorOnly() && (ActorDescView.GetIsSpatiallyLoaded() || ActorDescView.GetRuntimeDataLayerInstanceNames().GetNonExternalDataLayers().Num()))
+				if (ActorDescView.GetActorIsEditorOnlyLoadedInPIE() && (ActorDescView.GetIsSpatiallyLoaded() || ActorDescView.GetRuntimeDataLayerInstanceNames().GetNonExternalDataLayers().Num()))
 				{
 					if (ensure(PassType == EPassType::ErrorReporting))
 					{
-						// No error reporting here because this is not a data issue rather than a code issue. Actor classes returning true to IsNonPIEEditorOnly should be non-spatially loaded
+						// No error reporting here because this is not a data issue rather than a code issue. Actor classes returning true to IsEditorOnlyLoadedInPIE should be non-spatially loaded
 						// and can't be associated to data layers, so make sure they also have CanChangeIsSpatiallyLoadedFlag and ActorTypeSupportsDataLayer return false.
 					}
 					else
diff --git a/Engine/Source/Runtime/Engine/Public/WorldPartition/WorldPartitionActorDesc.h b/Engine/Source/Runtime/Engine/Public/WorldPartition/WorldPartitionActorDesc.h
--- a/Engine/Source/Runtime/Engine/Public/WorldPartition/WorldPartitionActorDesc.h
+++ b/Engine/Source/Runtime/Engine/Public/WorldPartition/WorldPartitionActorDesc.h
@@ -166,7 +166,7 @@ public:
 	inline bool GetIsSpatiallyLoaded() const { return RuntimeBounds.IsValid ? bIsSpatiallyLoaded : false; }
 	inline bool GetIsSpatiallyLoadedRaw() const { return bIsSpatiallyLoaded; }
 	inline bool GetActorIsEditorOnly() const { return bActorIsEditorOnly; }
-	ENGINE_API bool GetActorIsNonPIEEditorOnly() const;
+	ENGINE_API bool GetActorIsEditorOnlyLoadedInPIE() const;
 	inline bool GetActorIsRuntimeOnly() const { return bActorIsRuntimeOnly; }
 
 	inline bool GetActorIsHLODRelevant() const { return bActorIsHLODRelevant; }
diff --git a/Engine/Source/Runtime/Engine/Public/WorldPartition/WorldPartitionActorDescInstance.h b/Engine/Source/Runtime/Engine/Public/WorldPartition/WorldPartitionActorDescInstance.h
--- a/Engine/Source/Runtime/Engine/Public/WorldPartition/WorldPartitionActorDescInstance.h
+++ b/Engine/Source/Runtime/Engine/Public/WorldPartition/WorldPartitionActorDescInstance.h
@@ -45,7 +45,7 @@ public:
 	virtual FName GetRuntimeGrid() const { return ActorDesc->GetRuntimeGrid(); }
 	virtual bool GetIsSpatiallyLoaded() const { return !GetForceNonSpatiallyLoaded() && ActorDesc->GetIsSpatiallyLoaded(); }
 	virtual bool GetActorIsEditorOnly() const { return ActorDesc->GetActorIsEditorOnly(); }
-	virtual bool GetActorIsNonPIEEditorOnly() const { return ActorDesc->GetActorIsNonPIEEditorOnly(); }
+	virtual bool GetActorIsEditorOnlyLoadedInPIE() const { return ActorDesc->GetActorIsEditorOnlyLoadedInPIE(); }
 	virtual bool GetActorIsRuntimeOnly() const { return ActorDesc->GetActorIsRuntimeOnly(); }
 	ENGINE_API virtual bool IsRuntimeRelevant() const;
 	ENGINE_API virtual bool IsEditorRelevant() const;
diff --git a/Engine/Source/Runtime/Engine/Public/WorldPartition/WorldPartitionActorDescInstanceView.h b/Engine/Source/Runtime/Engine/Public/WorldPartition/WorldPartitionActorDescInstanceView.h
--- a/Engine/Source/Runtime/Engine/Public/WorldPartition/WorldPartitionActorDescInstanceView.h
+++ b/Engine/Source/Runtime/Engine/Public/WorldPartition/WorldPartitionActorDescInstanceView.h
@@ -26,7 +26,7 @@ public:
 	virtual FName GetRuntimeGrid() const override { return ActorDescInstance->GetRuntimeGrid(); }
 	virtual bool GetIsSpatiallyLoaded() const override { return ActorDescInstance->GetIsSpatiallyLoaded(); }
 	virtual bool GetActorIsEditorOnly() const override { return ActorDescInstance->GetActorIsEditorOnly(); }
-	virtual bool GetActorIsNonPIEEditorOnly() const override { return ActorDescInstance->GetActorIsNonPIEEditorOnly(); }
+	virtual bool GetActorIsEditorOnlyLoadedInPIE() const override { return ActorDescInstance->GetActorIsEditorOnlyLoadedInPIE(); }
 	virtual bool GetActorIsRuntimeOnly() const override { return ActorDescInstance->GetActorIsRuntimeOnly(); }
 	virtual bool IsRuntimeRelevant() const override { return ActorDescInstance->IsRuntimeRelevant(); }
 	virtual bool IsEditorRelevant() const override { return ActorDescInstance->IsEditorRelevant(); }
diff --git a/Engine/Source/Runtime/Engine/Public/WorldPartition/WorldPartitionActorDescInstanceViewInterface.h b/Engine/Source/Runtime/Engine/Public/WorldPartition/WorldPartitionActorDescInstanceViewInterface.h
--- a/Engine/Source/Runtime/Engine/Public/WorldPartition/WorldPartitionActorDescInstanceViewInterface.h
+++ b/Engine/Source/Runtime/Engine/Public/WorldPartition/WorldPartitionActorDescInstanceViewInterface.h
@@ -29,7 +29,7 @@ public:
 	virtual FName GetRuntimeGrid() const = 0;
 	virtual bool GetIsSpatiallyLoaded() const = 0;
 	virtual bool GetActorIsEditorOnly() const = 0;
-	virtual bool GetActorIsNonPIEEditorOnly() const = 0;
+	virtual bool GetActorIsEditorOnlyLoadedInPIE() const = 0;
 	virtual bool GetActorIsRuntimeOnly() const = 0;
 	virtual bool IsRuntimeRelevant() const = 0;
 	virtual bool IsEditorRelevant() const = 0;
