Fix for AActor's ForEachComponent crashing when using default constructed TSubclassOf.

diff --git a/Engine/Source/Runtime/Engine/Classes/GameFramework/Actor.h b/Engine/Source/Runtime/Engine/Classes/GameFramework/Actor.h
--- a/Engine/Source/Runtime/Engine/Classes/GameFramework/Actor.h
+++ b/Engine/Source/Runtime/Engine/Classes/GameFramework/Actor.h
@@ -3755,47 +3755,50 @@ private:
 	 * Use template parameter bIncludeFromChildActors to recurse in to ChildActor components and find components of the appropriate type in those actors as well
 	 */
 	template<class ComponentType, bool bClassIsActorComponent, bool bIncludeFromChildActors, typename Func>
-	void ForEachComponent_Internal(TSubclassOf<UActorComponent> ComponentClass, Func InFunc) const
+	void ForEachComponent_Internal(TSubclassOf<UActorComponent> InComponentClass, Func InFunc) const
 	{
-		check(bClassIsActorComponent == false || ComponentClass == UActorComponent::StaticClass());
-		check(ComponentClass->IsChildOf(ComponentType::StaticClass()));
-
-		// static check, so that the most common case (bIncludeFromChildActors) doesn't need to allocate an additional array : 
-		if (bIncludeFromChildActors)
+		if (const UClass* ComponentClass = InComponentClass)
 		{
-			TArray<AActor*, TInlineAllocator<NumInlinedActorComponents>> ChildActors;
-			for (UActorComponent* OwnedComponent : OwnedComponents)
+			check(bClassIsActorComponent == false || ComponentClass == UActorComponent::StaticClass());
+			check(ComponentClass->IsChildOf(ComponentType::StaticClass()));
+
+			// static check, so that the most common case (bIncludeFromChildActors) doesn't need to allocate an additional array : 
+			if (bIncludeFromChildActors)
 			{
-				if (OwnedComponent)
+				TArray<AActor*, TInlineAllocator<NumInlinedActorComponents>> ChildActors;
+				for (UActorComponent* OwnedComponent : OwnedComponents)
 				{
-					if (bClassIsActorComponent || OwnedComponent->IsA(ComponentClass))
-					{
-						InFunc(static_cast<ComponentType*>(OwnedComponent));
-					}
-					if (UChildActorComponent* ChildActorComponent = Cast<UChildActorComponent>(OwnedComponent))
+					if (OwnedComponent)
 					{
-						if (AActor* ChildActor = ChildActorComponent->GetChildActor())
+						if (bClassIsActorComponent || OwnedComponent->IsA(ComponentClass))
 						{
-							ChildActors.Add(ChildActor);
+							InFunc(static_cast<ComponentType *>(OwnedComponent));
+						}
+						if (UChildActorComponent* ChildActorComponent = Cast<UChildActorComponent>(OwnedComponent))
+						{
+							if (AActor* ChildActor = ChildActorComponent->GetChildActor())
+							{
+								ChildActors.Add(ChildActor);
+							}
 						}
 					}
 				}
-			}
 
-			for (AActor* ChildActor : ChildActors)
-			{
-				ChildActor->ForEachComponent_Internal<ComponentType, bClassIsActorComponent, bIncludeFromChildActors>(ComponentClass, InFunc);
+				for (AActor* ChildActor : ChildActors)
+				{
+					ChildActor->ForEachComponent_Internal<ComponentType, bClassIsActorComponent, bIncludeFromChildActors>(InComponentClass, InFunc);
+				}
 			}
-		}
-		else
-		{
-			for (UActorComponent* OwnedComponent : OwnedComponents)
+			else
 			{
-				if (OwnedComponent)
+				for (UActorComponent* OwnedComponent : OwnedComponents)
 				{
-					if (bClassIsActorComponent || OwnedComponent->IsA(ComponentClass))
+					if (OwnedComponent)
 					{
-						InFunc(static_cast<ComponentType*>(OwnedComponent));
+						if (bClassIsActorComponent || OwnedComponent->IsA(ComponentClass))
+						{
+							InFunc(static_cast<ComponentType *>(OwnedComponent));
+						}
 					}
 				}
 			}
