Chaos: add "SkinnedTriangleMesh" implicit type and add the ability to create it from the preview mesh in the PhysicsAssetEditor.

The new type has limited functionality used by Cloth. For example, it does not support PhiWithNormal, Overlap, etc. This could be added in the future if desired (likely wanting to share functions with the static TriangleMeshImplicitObject).

Additionally some housekeeping while I was in the area:
* Cleaned up several things in the AggregateGeom related to getting elements, transforms, etc. as there were some things that wouldn't work as intended with invalid indices.
* Updated the UI for creating bodies to hide data it doesn't actually use.
* Cleaned up MLLevelSet a bit.

diff --git a/Engine/Plugins/ChaosCloth/Source/ChaosCloth/Private/ChaosCloth/ChaosClothingSimulationCollider.cpp b/Engine/Plugins/ChaosCloth/Source/ChaosCloth/Private/ChaosCloth/ChaosClothingSimulationCollider.cpp
--- a/Engine/Plugins/ChaosCloth/Source/ChaosCloth/Private/ChaosCloth/ChaosClothingSimulationCollider.cpp
+++ b/Engine/Plugins/ChaosCloth/Source/ChaosCloth/Private/ChaosCloth/ChaosClothingSimulationCollider.cpp
@@ -759,7 +759,6 @@ void FClothingSimulationCollider::ExtractPhysicsAssetCollision(
 			{				
 				if (MLLevelSetElem.GetMLLevelSet().IsValid())
 				{
-					FName ParentBoneName = MLLevelSetElem.GetMLLevelSet()->GetParentBoneName();
 					TArray<int32> MappedActiveBones;
 					const TArray<FName>& ActiveBoneNames = MLLevelSetElem.GetMLLevelSet()->GetActiveBoneNames();
 					MappedActiveBones.Reserve(ActiveBoneNames.Num());
diff --git a/Engine/Source/Developer/PhysicsUtilities/Private/PhysicsAssetUtils.cpp b/Engine/Source/Developer/PhysicsUtilities/Private/PhysicsAssetUtils.cpp
--- a/Engine/Source/Developer/PhysicsUtilities/Private/PhysicsAssetUtils.cpp
+++ b/Engine/Source/Developer/PhysicsUtilities/Private/PhysicsAssetUtils.cpp
@@ -24,9 +24,11 @@
 #include "SkinnedLevelSetBuilder.h"
 #include "LevelSetHelpers.h"
 #include "Chaos/Levelset.h"
+#include "Chaos/SkinnedTriangleMesh.h"
 #include "AssetToolsModule.h"
 #include "IAssetTools.h"
 #include "Misc/NamePermissionList.h"
+#include "BoneWeights.h"
 
 namespace FPhysicsAssetUtils
 {
@@ -1014,6 +1016,190 @@ bool CreateCollisionsFromBones(UPhysicsAsset* PhysicsAsset, USkeletalMesh* SkelM
 	return bAllSuccessful;
 }
 
+int32 GenerateSkinnedTriangleMesh(UPhysicsAsset* PhysicsAsset, USkeletalMesh* SkeletalMesh, int32 LODIndex, const FPhysAssetCreateParams& Params, FText& OutErrorMessage, bool bShowProgress)
+{
+	if (!ensure(PhysicsAsset))
+	{
+		OutErrorMessage = FText(NSLOCTEXT("GenerateSkinnedTriangleMesh", "MissingPAFailure", "Failed to generate the skinned triangle mesh: missing input PhysicsAsset"));
+		return INDEX_NONE;
+	}
+	if (!ensure(SkeletalMesh))
+	{
+		OutErrorMessage = FText(NSLOCTEXT("GenerateSkinnedTriangleMesh", "MissingSKMFailure", "Failed to generate the skinned triangle mesh: missing input Skeletal Mesh"));
+		return INDEX_NONE;
+	}
+
+	// Gather data from the SkelMesh.
+	const FSkeletalMeshRenderData* const RenderData = SkeletalMesh->GetResourceForRendering();
+	if (!ensure(RenderData))
+	{
+		OutErrorMessage = FText::Format(NSLOCTEXT("GenerateSkinnedTriangleMesh", "MissingRenderDataFailure", "Failed to generate the skinned triangle mesh: missing Render Data on Skeletal Mesh '{0}'"), FText::FromString(SkeletalMesh->GetName()));
+		return INDEX_NONE;
+	}
+	if (!ensure(RenderData->LODRenderData.IsValidIndex(LODIndex)))
+	{
+		OutErrorMessage = FText::Format(NSLOCTEXT("GenerateSkinnedTriangleMesh", "MissingLODRenderDataFailure", "Failed to generate the skinned triangle mesh: missing LOD Render Data on Skeletal Mesh '{0}' and LOD {1}"), FText::FromString(SkeletalMesh->GetName()), LODIndex);
+		return INDEX_NONE;
+	}
+
+	const FSkeletalMeshLODRenderData& LodRenderData = RenderData->LODRenderData[LODIndex];
+	const int32 NumOrigVertices = LodRenderData.StaticVertexBuffers.PositionVertexBuffer.GetNumVertices();
+	if (NumOrigVertices == 0)
+	{
+		// This implies the actual mesh has no positions, not that they're just missing on the CPU.
+		// Rather than generate an empty object, generate an error
+		OutErrorMessage = FText::Format(NSLOCTEXT("GenerateSkinnedTriangleMesh", "EmptyLODRenderDataFailure", "Failed to generate the skinned triangle mesh: empty LOD Render Data on Skeletal Mesh '{0}' at LOD {1}"), FText::FromString(SkeletalMesh->GetName()), LODIndex);
+		return INDEX_NONE;
+	}
+	if (!ensure(NumOrigVertices == LodRenderData.SkinWeightVertexBuffer.GetNumVertices()))
+	{
+		OutErrorMessage = FText::Format(NSLOCTEXT("GenerateSkinnedTriangleMesh", "InvalidLODRenderData", "Failed to generate the skinned triangle mesh: invalid LOD Render Data on Skeletal Mesh '{0}' at LOD {1}"), FText::FromString(SkeletalMesh->GetName()), LODIndex);
+		return INDEX_NONE;
+	}
+	if (!ensure(LodRenderData.StaticVertexBuffers.PositionVertexBuffer.GetVertexData() && LodRenderData.SkinWeightVertexBuffer.GetDataVertexBuffer() && LodRenderData.SkinWeightVertexBuffer.GetDataVertexBuffer()->GetWeightData()))
+	{
+		OutErrorMessage = FText::Format(NSLOCTEXT("GenerateSkinnedTriangleMesh", "MissingCPURenderDataFailure", "Failed to generate the skinned triangle mesh: missing CPU LOD Render Data on Skeletal Mesh '{0}' and LOD {1}"), FText::FromString(SkeletalMesh->GetName()), LODIndex);
+		return INDEX_NONE;
+	}	
+
+	// Weld mesh verts to get rid of duplicates within a section. BoneData is taken from lower indexed duplicate. This matches existing cloth behavior.
+	TArray<FSkinWeightInfo> SkinWeightInfos;
+	LodRenderData.SkinWeightVertexBuffer.GetSkinWeights(SkinWeightInfos);
+	check(SkinWeightInfos.Num() == NumOrigVertices);
+
+	TArray<int32> OrigToWeldedVertices;
+	TArray<Chaos::FVec3f> WeldedPositions;
+	TArray<Chaos::FWeightedInfluenceData> BoneData;
+	OrigToWeldedVertices.Reserve(NumOrigVertices);
+	WeldedPositions.Reserve(NumOrigVertices);
+	BoneData.Reserve(NumOrigVertices);
+	for (const FSkelMeshRenderSection& Section : LodRenderData.RenderSections)
+	{
+		constexpr Chaos::FRealSingle ThreshSq = UE_THRESH_POINTS_ARE_SAME * UE_THRESH_POINTS_ARE_SAME;
+		const int32 WeldedSectionStartIndex = WeldedPositions.Num();
+
+		for (uint32 OrigVertexIndex = Section.BaseVertexIndex; OrigVertexIndex < Section.BaseVertexIndex + Section.NumVertices; ++OrigVertexIndex)
+		{
+			const Chaos::FVec3f Position(LodRenderData.StaticVertexBuffers.PositionVertexBuffer.VertexPosition(OrigVertexIndex));
+			bool bIsUnique = true;
+			int32 WeldedIndex = INDEX_NONE;
+			for (int32 WeldedVertexIndex = WeldedSectionStartIndex; WeldedVertexIndex < WeldedPositions.Num(); ++WeldedVertexIndex)
+			{
+				if (Chaos::FVec3f::DistSquared(Position, WeldedPositions[WeldedVertexIndex]) <= ThreshSq)
+				{
+					bIsUnique = false;
+					WeldedIndex = WeldedVertexIndex;
+					break;
+				}
+			}
+			if (bIsUnique)
+			{
+				WeldedIndex = WeldedPositions.Emplace(Position);
+
+				// Copy bone data
+				const FSkinWeightInfo& SkinWeightInfo = SkinWeightInfos[OrigVertexIndex];
+				static_assert(MAX_TOTAL_INFLUENCES <= Chaos::FWeightedInfluenceData::MaxTotalInfluences);
+				Chaos::FWeightedInfluenceData& NewBoneData = BoneData.AddDefaulted_GetRef();
+				check(NewBoneData.NumInfluences == 0);
+				for (int32 InfluenceIdx = 0; InfluenceIdx < MAX_TOTAL_INFLUENCES; ++InfluenceIdx)
+				{
+					if (SkinWeightInfo.InfluenceWeights[InfluenceIdx] > 0)
+					{
+						NewBoneData.BoneIndices[NewBoneData.NumInfluences] = Section.BoneMap[SkinWeightInfo.InfluenceBones[InfluenceIdx]];
+						NewBoneData.BoneWeights[NewBoneData.NumInfluences] = (float)SkinWeightInfo.InfluenceWeights[InfluenceIdx] * UE::AnimationCore::InvMaxRawBoneWeightFloat;
+						++NewBoneData.NumInfluences;
+					}
+				}
+			}
+			OrigToWeldedVertices.Emplace(WeldedIndex);
+		}
+	}
+	check(OrigToWeldedVertices.Num() == NumOrigVertices);
+
+	// Get UsedBones names and convert BoneIndices from RefSkeleton global indexing to local indexing.
+	const FReferenceSkeleton& RefSkeleton = SkeletalMesh->GetRefSkeleton();
+	
+	// This is copied from WeightedLevelSetImplicitObjectBuilder... Todo refactor all this to be shared.
+	TSet<int32> UsedBoneSet;
+	for (int32 Idx = 0; Idx < BoneData.Num(); ++Idx)
+	{
+		const Chaos::FWeightedInfluenceData& Data = BoneData[Idx];
+		for (uint8 InfIdx = 0; InfIdx < Data.NumInfluences; ++InfIdx)
+		{
+			UsedBoneSet.Add(Data.BoneIndices[InfIdx]);
+		}
+	}
+	const TArray<int32> UsedBoneIndices = UsedBoneSet.Array();
+	// Calculate the root bone index
+	const int32 RootBoneIndex = CalculateCommonRootBoneIndex(RefSkeleton, UsedBoneIndices);
+	const FTransform RootRefBaseInv(FMatrix(SkeletalMesh->GetRefBasesInvMatrix()[RootBoneIndex]));
+	FTransform ReferenceRootTransform = RootRefBaseInv.Inverse();
+	TArray<FName> UsedBones;
+	TArray<FTransform> ReferenceRelativeTransforms;
+	UsedBones.SetNumUninitialized(UsedBoneIndices.Num());
+	ReferenceRelativeTransforms.SetNumUninitialized(UsedBoneIndices.Num());
+
+	TMap<int32, int32> UsedBoneReverseLookup;
+	for (int32 ArrayIdx = 0; ArrayIdx < UsedBoneIndices.Num(); ++ArrayIdx)
+	{
+		UsedBones[ArrayIdx] = RefSkeleton.GetBoneName(UsedBoneIndices[ArrayIdx]);
+		ReferenceRelativeTransforms[ArrayIdx] = ReferenceRootTransform * FTransform(FMatrix(SkeletalMesh->GetRefBasesInvMatrix()[UsedBoneIndices[ArrayIdx]]));
+		UsedBoneReverseLookup.Add(UsedBoneIndices[ArrayIdx], ArrayIdx);
+	}
+
+	for (int32 Idx = 0; Idx < BoneData.Num(); ++Idx)
+	{
+		Chaos::FWeightedInfluenceData& Data = BoneData[Idx];
+		for (uint8 InfIdx = 0; InfIdx < Data.NumInfluences; ++InfIdx)
+		{
+			Data.BoneIndices[InfIdx] = (uint16)UsedBoneReverseLookup[Data.BoneIndices[InfIdx]];
+		}
+
+		WeldedPositions[Idx] = Chaos::FVec3f(RootRefBaseInv.TransformPosition((FVector)WeldedPositions[Idx]));
+	}
+
+	// Copy over triangles with updated indices from welding.
+	TArray<uint32> OrigIndices;
+	LodRenderData.MultiSizeIndexContainer.GetIndexBuffer(OrigIndices);
+	check(OrigIndices.Num() % 3 == 0);
+	const int32 NumTriangles = OrigIndices.Num() / 3;
+	TArray<Chaos::TVec3<int32>> Elements;
+	Elements.Reserve(NumTriangles);
+	for (int32 TriIndex = 0; TriIndex < NumTriangles; ++TriIndex)
+	{
+		Elements.Emplace(
+			OrigToWeldedVertices[OrigIndices[TriIndex * 3]],
+			OrigToWeldedVertices[OrigIndices[TriIndex * 3 + 1]],
+			OrigToWeldedVertices[OrigIndices[TriIndex * 3 + 2]]);
+	}
+
+	// Build TriangleMesh.
+	Chaos::FTriangleMesh TriangleMesh(MoveTemp(Elements), 0, WeldedPositions.Num());
+
+	// Make the SkinnedTriangleMesh
+	TRefCountPtr<Chaos::FSkinnedTriangleMesh> SkinnedTriangleMesh(new Chaos::FSkinnedTriangleMesh(
+		MoveTemp(TriangleMesh),
+		MoveTemp(WeldedPositions),
+		MoveTemp(BoneData),
+		MoveTemp(UsedBones),
+		MoveTemp(ReferenceRootTransform),
+		MoveTemp(ReferenceRelativeTransforms)));
+
+	// Make the SkinnedTriangleMeshElem
+	FKSkinnedTriangleMeshElem SkinnedTriangleMeshElem;
+	SkinnedTriangleMeshElem.SetSkinnedTriangleMesh(MoveTemp(SkinnedTriangleMesh));
+
+
+	// Get or create a new body for RootBoneIndex.
+	const FName RootBoneName = RefSkeleton.GetBoneName(RootBoneIndex);
+	const int32 NewBodyIndex = CreateNewBody(PhysicsAsset, RootBoneName, Params);
+	UBodySetup* NewBodySetup = PhysicsAsset->SkeletalBodySetups[NewBodyIndex];
+	check(NewBodySetup->BoneName == RootBoneName);
+	NewBodySetup->AggGeom.SkinnedTriangleMeshElems.Emplace(MoveTemp(SkinnedTriangleMeshElem));
+
+	return NewBodyIndex;
+}
+
 void WeldBodies(UPhysicsAsset* PhysAsset, int32 BaseBodyIndex, int32 AddBodyIndex, USkeletalMeshComponent* SkelComp)
 {
 	if(BaseBodyIndex == INDEX_NONE || AddBodyIndex == INDEX_NONE)
diff --git a/Engine/Source/Developer/PhysicsUtilities/Public/PhysicsAssetUtils.h b/Engine/Source/Developer/PhysicsUtilities/Public/PhysicsAssetUtils.h
--- a/Engine/Source/Developer/PhysicsUtilities/Public/PhysicsAssetUtils.h
+++ b/Engine/Source/Developer/PhysicsUtilities/Public/PhysicsAssetUtils.h
@@ -22,9 +22,10 @@ enum EPhysAssetFitGeomType : int
 	EFG_TaperedCapsule		UMETA(DisplayName="Tapered Capsule (Cloth Only)"),
 	EFG_SingleConvexHull	UMETA(DisplayName="Single Convex Hull"),
 	EFG_MultiConvexHull		UMETA(DisplayName="Multi Convex Hull"),
-	EFG_LevelSet			UMETA(DisplayName="Level Set"),
-	EFG_SkinnedLevelSet		UMETA(DisplayName="Skinned Level Set"),
-	EFG_MLLevelSet			UMETA(DisplayName="ML Level Set"),
+	EFG_LevelSet			UMETA(DisplayName="Level Set (Experimental, Cloth Only)"),
+	EFG_SkinnedLevelSet		UMETA(DisplayName="Skinned Level Set (Experimental, Cloth Only)"),
+	EFG_MLLevelSet			UMETA(Hidden),
+	EFG_SkinnedTriangleMesh UMETA(Hidden),
 };
 
 UENUM()
@@ -99,21 +100,21 @@ struct FPhysAssetCreateParams
 	TEnumAsByte<EAngularConstraintMotion> AngularConstraintMode;
 
 	/** When creating multiple convex hulls, the maximum number that will be created. */
-	UPROPERTY(EditAnywhere, AdvancedDisplay, Category = "Body Creation")
+	UPROPERTY(EditAnywhere, AdvancedDisplay, Category = "Body Creation", meta = (EditCondition = "GeomType == EPhysAssetFitGeomType::EFG_MultiConvexHull", EditConditionHides))
 	int32								HullCount;
 
 	/** When creating convex hulls, the maximum verts that should be created */
-	UPROPERTY(EditAnywhere, AdvancedDisplay, Category = "Body Creation")
+	UPROPERTY(EditAnywhere, AdvancedDisplay, Category = "Body Creation", meta = (EditCondition = "GeomType == EPhysAssetFitGeomType::EFG_MultiConvexHull || GeomType == EPhysAssetFitGeomType::EFG_SingleConvexHull", EditConditionHides))
 	int32								MaxHullVerts;
 
 	/** When creating level sets, the grid resolution to use */
 	UPROPERTY(EditAnywhere, AdvancedDisplay, Category = "Body Creation", 
-		meta = (ClampMin = 1, UIMin = 10, UIMax = 100, ClampMax = 500, EditCondition = "GeomType == EPhysAssetFitGeomType::EFG_LevelSet || GeomType == EPhysAssetFitGeomType::EFG_SkinnedLevelSet"))
+		meta = (ClampMin = 1, UIMin = 10, UIMax = 100, ClampMax = 500, EditCondition = "GeomType == EPhysAssetFitGeomType::EFG_LevelSet || GeomType == EPhysAssetFitGeomType::EFG_SkinnedLevelSet", EditConditionHides))
 	int32								LevelSetResolution;
 
 	/** When creating skinned level sets, the embedding grid resolution to use*/
 	UPROPERTY(EditAnywhere, AdvancedDisplay, Category = "Body Creation",
-		meta = (ClampMin = 1, UIMin = 10, UIMax = 100, ClampMax = 500, EditCondition = "GeomType == EPhysAssetFitGeomType::EFG_SkinnedLevelSet"))
+		meta = (ClampMin = 1, UIMin = 10, UIMax = 100, ClampMax = 500, EditCondition = "GeomType == EPhysAssetFitGeomType::EFG_SkinnedLevelSet", EditConditionHides))
 	int32								LatticeResolution;
 };
 
@@ -174,6 +175,17 @@ namespace FPhysicsAssetUtils
 	 */
 	PHYSICSUTILITIES_API bool CreateCollisionsFromBones(UPhysicsAsset* PhysicsAsset, USkeletalMesh* SkelMesh, const TArray<int32>& BodyIndices, const FPhysAssetCreateParams& Params, const TArray<FBoneVertInfo>& Infos, TArray<int32>& OutSuccessfulBodyIndices);
 
+	/**
+	 * Given a USkeletalMesh, generate a Skinned Triangle Mesh primitive and add it to the root body (creating a new one if needed).
+	 *
+	 * @param	PhysicsAsset		The PhysicsAsset instance to update
+	 * @param	SkelMesh			The Skeletal Mesh to create the physics asset from
+	 * @param	LODIndex			SkelMesh LOD index to use.
+	 * @param	Params				Additional creation parameters
+	 * @param	OutErrorMessage		Additional error information
+	 * @return  Index of body under which the new Skinned Triangle Mesh primitive was created (or INDEX_NONE if failure)
+	 */
+	PHYSICSUTILITIES_API int32 GenerateSkinnedTriangleMesh(UPhysicsAsset* PhysicsAsset, USkeletalMesh* SkelMesh, int32 LODIndex, const FPhysAssetCreateParams& Params, FText& OutErrorMessage, bool bShowProgress = true);
 	/**
 	 * Does a few things:
 	 * - add any collision primitives from body2 into body1 (adjusting the tm of each).
diff --git a/Engine/Source/Editor/Persona/Private/PhysicsAssetRenderUtils.cpp b/Engine/Source/Editor/Persona/Private/PhysicsAssetRenderUtils.cpp
--- a/Engine/Source/Editor/Persona/Private/PhysicsAssetRenderUtils.cpp
+++ b/Engine/Source/Editor/Persona/Private/PhysicsAssetRenderUtils.cpp
@@ -13,6 +13,7 @@
 #include "PhysicsEngine/PhysicsConstraintTemplate.h"
 #include "PhysicsEngine/SkeletalBodySetup.h"
 #include "Chaos/Core.h"
+#include "Chaos/SkinnedTriangleMesh.h"
 #include "SkeletalMeshTypes.h"
 #include "Styling/AppStyle.h"
 #include "Animation/DebugSkelMeshComponent.h"
@@ -426,6 +427,7 @@ namespace PhysicsAssetRender
 				CollectDrawElements(DrawElements, EAggCollisionShape::LevelSet, AggGeom->LevelSetElems, BoneTM, Scale, i, RenderSettings, PhysicsAsset, PDI, GetPrimitiveColor, GetPrimitiveMaterial, GetPrimitiveTransform, CreateHitProxy);
 				CollectDrawElements(DrawElements, EAggCollisionShape::SkinnedLevelSet, AggGeom->SkinnedLevelSetElems, BoneTM, Scale, i, RenderSettings, PhysicsAsset, PDI, GetPrimitiveColor, GetPrimitiveMaterial, GetPrimitiveTransform, CreateHitProxy);
 				CollectDrawElements(DrawElements, EAggCollisionShape::MLLevelSet, AggGeom->MLLevelSetElems, BoneTM, Scale, i, RenderSettings, PhysicsAsset, PDI, GetPrimitiveColor, GetPrimitiveMaterial, GetPrimitiveTransform, CreateHitProxy);
+				CollectDrawElements(DrawElements, EAggCollisionShape::SkinnedTriangleMesh, AggGeom->SkinnedTriangleMeshElems, BoneTM, Scale, i, RenderSettings, PhysicsAsset, PDI, GetPrimitiveColor, GetPrimitiveMaterial, GetPrimitiveTransform, CreateHitProxy);
 
 				// Sort elements
 				DrawElements.Sort([](const DrawElement& DE1, const DrawElement& DE2) {
diff --git a/Engine/Source/Editor/PhysicsAssetEditor/Private/PhysicsAssetEditor.cpp b/Engine/Source/Editor/PhysicsAssetEditor/Private/PhysicsAssetEditor.cpp
--- a/Engine/Source/Editor/PhysicsAssetEditor/Private/PhysicsAssetEditor.cpp
+++ b/Engine/Source/Editor/PhysicsAssetEditor/Private/PhysicsAssetEditor.cpp
@@ -858,6 +858,7 @@ void FPhysicsAssetEditor::ExtendMenu()
 			Section.AddMenuEntry(Commands.SelectAllConstraints);
 			Section.AddMenuEntry(Commands.ToggleSelectionType);
 			Section.AddMenuEntry(Commands.ToggleSelectionTypeWithUserConstraints);
+			Section.AddMenuEntry(Commands.GenerateSkinnedTriangleMesh);
 		}
 		{
 			FToolMenuSection& Section = EditMenu->AddSection("Shapes", LOCTEXT("PhatEditSelectionShapes", "Shapes"));
@@ -1403,6 +1404,11 @@ void FPhysicsAssetEditor::BindCommands()
 		FExecuteAction::CreateSP(this, &FPhysicsAssetEditor::ImportMLLevelSetFromDataTable),
 		FCanExecuteAction::CreateSP(this, &FPhysicsAssetEditor::IsNotSimulation));
 
+	ToolkitCommands->MapAction(
+		Commands.GenerateSkinnedTriangleMesh,
+		FExecuteAction::CreateSP(this, &FPhysicsAssetEditor::GenerateSkinnedTriangleMesh),
+		FCanExecuteAction::CreateSP(this, &FPhysicsAssetEditor::IsNotSimulation));
+
 	ToolkitCommands->MapAction(
 		Commands.SelectAllConstraints,
 		FExecuteAction::CreateSP(this, &FPhysicsAssetEditor::OnSelectAllConstraints),
@@ -2045,7 +2051,12 @@ void FPhysicsAssetEditor::BuildMenuWidgetSelection(FMenuBuilder& InMenuBuilder)
 		}
 		InMenuBuilder.AddMenuEntry( Commands.SelectShapesQueryAndProbe);
 		InMenuBuilder.AddMenuEntry( Commands.SelectShapesProbeOnly);
+		InMenuBuilder.EndSection();
+
+
+		InMenuBuilder.BeginSection("CreateAdvancedPrimitives", LOCTEXT("CreateAdvancedPrimitives", "Create Advanced Primitives"));
 		InMenuBuilder.AddMenuEntry( Commands.ImportMLLevelSet);
+		InMenuBuilder.AddMenuEntry( Commands.GenerateSkinnedTriangleMesh );
 		InMenuBuilder.EndSection();
 	}
 	InMenuBuilder.PopCommandList();
@@ -3911,7 +3922,7 @@ bool FPhysicsAssetEditor::LoadMLLevelSetFromDataTable(UDataTable* DataTableMLLev
 										const FString MLModelIncorrectZoneWeightsString = RowModelInferenceForIncorrectZoneInfo->MLModelWeights;
 										NNEModelDataArr.Add({ RowModelInferenceForIncorrectZoneInfo->ModelArchitectureActivationNodeSizes, MLModelIncorrectZoneWeightsString, NNIModelIncorrectZoneDir, nullptr });
 
-										Chaos::FMLLevelSetImportData MLLevelSetImportData = { ParentBoneName,ActiveBoneNames, MoveTemp(NNEModelDataArr),
+										Chaos::FMLLevelSetImportData MLLevelSetImportData = { ActiveBoneNames, MoveTemp(NNEModelDataArr),
 											MoveTemp(ReferenceBoneRotations), MoveTemp(ReferenceBoneTranslations),
 											SignedDistanceScaling, MoveTemp(ActiveBonesJointRotationComponents),TrainingGridMin, TrainingGridAxes, DebugGridResolution };
 										
@@ -3923,7 +3934,7 @@ bool FPhysicsAssetEditor::LoadMLLevelSetFromDataTable(UDataTable* DataTableMLLev
 										TArray<Chaos::FMLLevelSetNNEModelData> NNEModelDataArr;
 										NNEModelDataArr.Add({ RowModelInferenceInfo->ModelArchitectureActivationNodeSizes, MLModelWeightsString, NNIModelDir, nullptr });
 										
-										Chaos::FMLLevelSetImportData MLLevelSetImportData = { ParentBoneName,ActiveBoneNames, MoveTemp(NNEModelDataArr),
+										Chaos::FMLLevelSetImportData MLLevelSetImportData = { ActiveBoneNames, MoveTemp(NNEModelDataArr),
 											MoveTemp(ReferenceBoneRotations), MoveTemp(ReferenceBoneTranslations),
 											SignedDistanceScaling, MoveTemp(ActiveBonesJointRotationComponents),TrainingGridMin, TrainingGridAxes, DebugGridResolution };
 
@@ -3969,6 +3980,50 @@ bool FPhysicsAssetEditor::LoadMLLevelSetFromDataTable(UDataTable* DataTableMLLev
 	return bCanLoadAll;
 }
 
+void FPhysicsAssetEditor::GenerateSkinnedTriangleMesh()
+{
+	USkeletalMesh* const EditorSkelMesh = SharedData->PhysicsAsset->GetPreviewMesh();
+	if (EditorSkelMesh == nullptr)
+	{
+		return;
+	}
+
+	int32 LODIndex = 0;
+	if (SharedData->EditorSkelComp)
+	{
+		LODIndex = SharedData->EditorSkelComp->GetPredictedLODLevel();
+	}
+	if (!ensure(EditorSkelMesh->IsValidLODIndex(LODIndex)))
+	{
+		return;
+	}
+
+	// Make sure rendering is done - so we are not changing data being used by collision drawing.
+	FlushRenderingCommands();
+
+	const FPhysAssetCreateParams& NewBodyData = GetDefault<UPhysicsAssetGenerationSettings>()->CreateParams;
+	{
+		const FScopedTransaction Transaction(LOCTEXT("GenerateSkinnedTriangleMesh", "Generate Skinned Triangle Mesh"));
+		SharedData->PhysicsAsset->Modify();
+
+		FText ErrorMessage;
+		const int32 BodyIndex = FPhysicsAssetUtils::GenerateSkinnedTriangleMesh(SharedData->PhysicsAsset, EditorSkelMesh, LODIndex, NewBodyData, ErrorMessage);
+		if (BodyIndex == INDEX_NONE)
+		{
+			FMessageDialog::Open(EAppMsgType::Ok, ErrorMessage);
+		}
+		else
+		{
+			SharedData->AutoNameAllPrimitives(BodyIndex, EFG_SkinnedTriangleMesh);
+			SharedData->SetSelectedBodyAnyPrimitive(BodyIndex, true);
+		}
+	}
+	RecreatePhysicsState();
+	SharedData->RefreshPhysicsAssetChange(SharedData->PhysicsAsset);
+	RefreshPreviewViewport();
+	RefreshHierachyTree();
+}
+
 void FPhysicsAssetEditor::ImportMLLevelSetFromDataTable()
 {
 	FContentBrowserModule& ContentBrowserModule = FModuleManager::LoadModuleChecked<FContentBrowserModule>(TEXT("ContentBrowser"));
diff --git a/Engine/Source/Editor/PhysicsAssetEditor/Private/PhysicsAssetEditor.h b/Engine/Source/Editor/PhysicsAssetEditor/Private/PhysicsAssetEditor.h
--- a/Engine/Source/Editor/PhysicsAssetEditor/Private/PhysicsAssetEditor.h
+++ b/Engine/Source/Editor/PhysicsAssetEditor/Private/PhysicsAssetEditor.h
@@ -335,8 +335,11 @@ private:
 	void OnDeselectAll();
 	bool LoadMLLevelSetFromDataTable(UDataTable* DataTableMLLevelSetModelBoneBinningInfo, TArray<FString>& ErrorMessages);
 	void ImportMLLevelSetFromDataTable();
+
 	bool ShouldCreateConstraintsWhenCreatingBodies() const;
 	
+	void GenerateSkinnedTriangleMesh();
+
 	FText GetRepeatLastSimulationToolTip() const;
 	FSlateIcon GetRepeatLastSimulationIcon() const;
 
diff --git a/Engine/Source/Editor/PhysicsAssetEditor/Private/PhysicsAssetEditorActions.cpp b/Engine/Source/Editor/PhysicsAssetEditor/Private/PhysicsAssetEditorActions.cpp
--- a/Engine/Source/Editor/PhysicsAssetEditor/Private/PhysicsAssetEditorActions.cpp
+++ b/Engine/Source/Editor/PhysicsAssetEditor/Private/PhysicsAssetEditorActions.cpp
@@ -116,6 +116,7 @@ void FPhysicsAssetEditorCommands::RegisterCommands()
 	UI_COMMAND(SelectShapesQueryAndProbe, "Select QueryAndProbe Shapes", "Select Kinematic Bodies", EUserInterfaceActionType::Button, FInputChord());
 	UI_COMMAND(SelectShapesProbeOnly, "Select ProbeOnly Shapes", "Select Kinematic Bodies", EUserInterfaceActionType::Button, FInputChord());
 	UI_COMMAND(ImportMLLevelSet, "Import MLLevelset Data", "Import MLLevelset Physical Asset", EUserInterfaceActionType::Button, FInputChord());
+	UI_COMMAND(GenerateSkinnedTriangleMesh, "Generate Skinned Triangle Mesh", "Generate a Skinned Triangle Mesh primitive from the current preview skeletal mesh.", EUserInterfaceActionType::Button, FInputChord());
 	UI_COMMAND(SelectAllConstraints, "Select All Constraints", "Select All Constraints", EUserInterfaceActionType::Button, FInputChord(EModifierKey::Control | EModifierKey::Shift, EKeys::A));
 	UI_COMMAND(ToggleSelectionType, "Toggle Selection Type", "Select bodies from constraints, or constraints from bodies, depending on selection.", EUserInterfaceActionType::Button, FInputChord(EModifierKey::Control, EKeys::T));
 	UI_COMMAND(ToggleSelectionTypeWithUserConstraints, "Toggle Selection Type (With User Constraints)", "Select bodies from constraints, or constraints from bodies, depending on selection (account for user constraints).", EUserInterfaceActionType::Button, FInputChord(EModifierKey::Control | EModifierKey::Shift, EKeys::T));
diff --git a/Engine/Source/Editor/PhysicsAssetEditor/Private/PhysicsAssetEditorActions.h b/Engine/Source/Editor/PhysicsAssetEditor/Private/PhysicsAssetEditorActions.h
--- a/Engine/Source/Editor/PhysicsAssetEditor/Private/PhysicsAssetEditorActions.h
+++ b/Engine/Source/Editor/PhysicsAssetEditor/Private/PhysicsAssetEditorActions.h
@@ -131,6 +131,7 @@ public:
 	TSharedPtr<FUICommandInfo> SelectShapesQueryAndProbe;
 	TSharedPtr<FUICommandInfo> SelectShapesProbeOnly;
 	TSharedPtr<FUICommandInfo> ImportMLLevelSet; 
+	TSharedPtr<FUICommandInfo> GenerateSkinnedTriangleMesh;
 	TSharedPtr<FUICommandInfo> SelectAllConstraints;
 	TSharedPtr<FUICommandInfo> ToggleSelectionType;
 	TSharedPtr<FUICommandInfo> ToggleSelectionTypeWithUserConstraints;
diff --git a/Engine/Source/Editor/PhysicsAssetEditor/Private/PhysicsAssetEditorSkeletalMeshComponent.cpp b/Engine/Source/Editor/PhysicsAssetEditor/Private/PhysicsAssetEditorSkeletalMeshComponent.cpp
--- a/Engine/Source/Editor/PhysicsAssetEditor/Private/PhysicsAssetEditorSkeletalMeshComponent.cpp
+++ b/Engine/Source/Editor/PhysicsAssetEditor/Private/PhysicsAssetEditorSkeletalMeshComponent.cpp
@@ -206,52 +206,9 @@ FTransform UPhysicsAssetEditorSkeletalMeshComponent::GetPrimitiveTransform(const
 		}
 	}
 
-
-	if (PrimType == EAggCollisionShape::Sphere)
-	{
-		FTransform PrimTM = ManTM * SharedBodySetup->AggGeom.SphereElems[PrimIndex].GetTransform();
-		PrimTM.ScaleTranslation(Scale3D);
-		return PrimTM * BoneTM;
-	}
-	else if (PrimType == EAggCollisionShape::Box)
-	{
-		FTransform PrimTM = ManTM * SharedBodySetup->AggGeom.BoxElems[PrimIndex].GetTransform();
-		PrimTM.ScaleTranslation(Scale3D);
-		return PrimTM * BoneTM;
-	}
-	else if (PrimType == EAggCollisionShape::Sphyl)
-	{
-		FTransform PrimTM = ManTM * SharedBodySetup->AggGeom.SphylElems[PrimIndex].GetTransform();
-		PrimTM.ScaleTranslation(Scale3D);
-		return PrimTM * BoneTM;
-	}
-	else if (PrimType == EAggCollisionShape::Convex)
-	{
-		FTransform PrimTM = ManTM * SharedBodySetup->AggGeom.ConvexElems[PrimIndex].GetTransform();
-		PrimTM.ScaleTranslation(Scale3D);
-		return PrimTM * BoneTM;
-	}
-	else if (PrimType == EAggCollisionShape::TaperedCapsule)
-	{
-		FTransform PrimTM = ManTM * SharedBodySetup->AggGeom.TaperedCapsuleElems[PrimIndex].GetTransform();
-		PrimTM.ScaleTranslation(Scale3D);
-		return PrimTM * BoneTM;
-	}
-	else if (PrimType == EAggCollisionShape::LevelSet)
-	{
-		FTransform PrimTM = ManTM * SharedBodySetup->AggGeom.LevelSetElems[PrimIndex].GetTransform();
-		PrimTM.ScaleTranslation(Scale3D);
-		return PrimTM * BoneTM;
-	}
-	else if (PrimType == EAggCollisionShape::SkinnedLevelSet)
-	{
-		FTransform PrimTM = ManTM * SharedBodySetup->AggGeom.SkinnedLevelSetElems[PrimIndex].GetTransform();
-		PrimTM.ScaleTranslation(Scale3D);
-		return PrimTM * BoneTM;
-	}
-	else if (PrimType == EAggCollisionShape::MLLevelSet)
+	if (const FKShapeElem* const Prim = SharedBodySetup->AggGeom.GetElement(PrimType, PrimIndex))
 	{
-		FTransform PrimTM = ManTM * SharedBodySetup->AggGeom.MLLevelSetElems[PrimIndex].GetTransform();
+		FTransform PrimTM = ManTM * Prim->GetTransform();
 		PrimTM.ScaleTranslation(Scale3D);
 		return PrimTM * BoneTM;
 	}
@@ -380,6 +337,7 @@ void UPhysicsAssetEditorSkeletalMeshComponent::RefreshBoneTransforms(FActorCompo
 	}
 	UpdateSkinnedLevelSets();
 	UpdateMLLevelSets();
+	UpdateSkinnedTriangleMeshes();
 }
 
 void UPhysicsAssetEditorSkeletalMeshComponent::AddImpulseAtLocation(FVector Impulse, FVector Location, FName BoneName)
@@ -533,4 +491,50 @@ void UPhysicsAssetEditorSkeletalMeshComponent::UpdateMLLevelSets()
 			}
 		}
 	}
-}
\ No newline at end of file
+}
+
+void UPhysicsAssetEditorSkeletalMeshComponent::UpdateSkinnedTriangleMeshes()
+{
+	UPhysicsAsset* const PhysicsAsset = GetPhysicsAsset();
+	if (!PhysicsAsset)
+	{
+		return;
+	}
+	for (int32 i = 0; i < PhysicsAsset->SkeletalBodySetups.Num(); ++i)
+	{
+		const int32 BoneIndex = GetBoneIndex(PhysicsAsset->SkeletalBodySetups[i]->BoneName);
+		if (BoneIndex != INDEX_NONE)
+		{
+			FKAggregateGeom* const AggGeom = &PhysicsAsset->SkeletalBodySetups[i]->AggGeom;
+			if (AggGeom)
+			{
+				for (FKSkinnedTriangleMeshElem& SkinnedTriangleMesh : AggGeom->SkinnedTriangleMeshElems)
+				{
+					if (SkinnedTriangleMesh.GetSkinnedTriangleMesh().IsValid())
+					{
+						const TArray<FName>& UsedBoneNames = SkinnedTriangleMesh.GetSkinnedTriangleMesh()->GetUsedBones();
+
+						const FTransform RootTransformInv = GetBoneTransform(BoneIndex, FTransform::Identity).Inverse();
+						TArray<FTransform> Transforms;
+						Transforms.SetNum(UsedBoneNames.Num());
+
+						for (int32 LocalIdx = 0; LocalIdx < UsedBoneNames.Num(); ++LocalIdx)
+						{
+							const int32 LocalBoneIndex = GetBoneIndex(UsedBoneNames[LocalIdx]);
+							if (LocalBoneIndex != INDEX_NONE)
+							{
+								const FTransform BoneTransformTimesRootTransformInv = GetBoneTransform(LocalBoneIndex, RootTransformInv);
+								Transforms[LocalIdx] = BoneTransformTimesRootTransformInv;
+							}
+							else
+							{
+								Transforms[LocalIdx] = RootTransformInv;
+							}
+						}
+						SkinnedTriangleMesh.GetSkinnedTriangleMesh()->SkinPositions(Transforms);
+					}
+				}
+			}
+		}
+	}
+}
diff --git a/Engine/Source/Editor/PhysicsAssetEditor/Private/PhysicsAssetEditorSkeletalMeshComponent.h b/Engine/Source/Editor/PhysicsAssetEditor/Private/PhysicsAssetEditorSkeletalMeshComponent.h
--- a/Engine/Source/Editor/PhysicsAssetEditor/Private/PhysicsAssetEditorSkeletalMeshComponent.h
+++ b/Engine/Source/Editor/PhysicsAssetEditor/Private/PhysicsAssetEditorSkeletalMeshComponent.h
@@ -82,4 +82,5 @@ public:
 private:
 	void UpdateSkinnedLevelSets();
 	void UpdateMLLevelSets();
+	void UpdateSkinnedTriangleMeshes();
 };
diff --git a/Engine/Source/Editor/PhysicsAssetEditor/Private/PhysicsAssetEditorToolsSummoner.cpp b/Engine/Source/Editor/PhysicsAssetEditor/Private/PhysicsAssetEditorToolsSummoner.cpp
--- a/Engine/Source/Editor/PhysicsAssetEditor/Private/PhysicsAssetEditorToolsSummoner.cpp
+++ b/Engine/Source/Editor/PhysicsAssetEditor/Private/PhysicsAssetEditorToolsSummoner.cpp
@@ -42,7 +42,8 @@ TSharedRef<SWidget> FPhysicsAssetEditorToolsSummoner::CreateTabBody(const FWorkf
 		TAttribute<bool>::Create(TAttribute<bool>::FGetter::CreateLambda([this]()
 		{
 			TSharedRef<FPhysicsAssetEditor> PhysicsAssetEditor = StaticCastSharedRef<FPhysicsAssetEditor>(HostingApp.Pin().ToSharedRef());
-			return !PhysicsAssetEditor->GetSharedData()->bRunningSimulation;
+			
+			return PhysicsAssetEditor->IsNotSimulation();
 		})),
 		TAttribute<FText>::Create(TAttribute<FText>::FGetter::CreateLambda([this]()
 		{
diff --git a/Engine/Source/Runtime/Engine/Classes/PhysicsEngine/AggregateGeom.h b/Engine/Source/Runtime/Engine/Classes/PhysicsEngine/AggregateGeom.h
--- a/Engine/Source/Runtime/Engine/Classes/PhysicsEngine/AggregateGeom.h
+++ b/Engine/Source/Runtime/Engine/Classes/PhysicsEngine/AggregateGeom.h
@@ -9,6 +9,7 @@
 #include "PhysicsEngine/MLLevelSetElem.h"
 #include "PhysicsEngine/BoxElem.h"
 #include "PhysicsEngine/SkinnedLevelSetElem.h"
+#include "PhysicsEngine/SkinnedTriangleMeshElem.h"
 #include "PhysicsEngine/SphereElem.h"
 #include "PhysicsEngine/SphylElem.h"
 #include "PhysicsEngine/TaperedCapsuleElem.h"
@@ -44,9 +45,12 @@ struct FKAggregateGeom
 	UPROPERTY(EditAnywhere, editfixedsize, Category = "Aggregate Geometry", meta = (DisplayName = "(Experimental) Skinned Level Sets", TitleProperty = "Name"), Experimental)
 	TArray<FKSkinnedLevelSetElem> SkinnedLevelSetElems;
 
-	UPROPERTY(EditAnywhere, editfixedsize, Category = "Aggregate Geometry", meta = (DisplayName = "(Experimental) ML Level Sets"))
+	UPROPERTY(EditAnywhere, editfixedsize, Category = "Aggregate Geometry", meta = (DisplayName = "(Experimental) ML Level Sets", TitleProperty = "Name"), Experimental)
 	TArray<FKMLLevelSetElem> MLLevelSetElems;
 
+	UPROPERTY(EditAnywhere, editfixedsize, Category = "Aggregate Geometry", meta = (DisplayName = "(Experimental) Skinned Triangle Meshes", TitleProperty = "Name"), Experimental)
+	TArray<FKSkinnedTriangleMeshElem> SkinnedTriangleMeshElems;
+
 	FKAggregateGeom()
 		: RenderInfoPtr(nullptr)
 	{
@@ -67,55 +71,188 @@ struct FKAggregateGeom
 
 	int32 GetElementCount() const
 	{
-		return SphereElems.Num() + SphylElems.Num() + BoxElems.Num() + ConvexElems.Num() + TaperedCapsuleElems.Num() + LevelSetElems.Num() + SkinnedLevelSetElems.Num() + MLLevelSetElems.Num();
+		return SphereElems.Num() + SphylElems.Num() + BoxElems.Num() + ConvexElems.Num() + TaperedCapsuleElems.Num() + LevelSetElems.Num() + SkinnedLevelSetElems.Num() + MLLevelSetElems.Num() + SkinnedTriangleMeshElems.Num();
 	}
 
 	ENGINE_API int32 GetElementCount(EAggCollisionShape::Type Type) const;
 
-	SIZE_T GetAllocatedSize() const { return SphereElems.GetAllocatedSize() + SphylElems.GetAllocatedSize() + BoxElems.GetAllocatedSize() + ConvexElems.GetAllocatedSize() + TaperedCapsuleElems.GetAllocatedSize() + LevelSetElems.GetAllocatedSize() + SkinnedLevelSetElems.GetAllocatedSize() + MLLevelSetElems.GetAllocatedSize(); }
+	SIZE_T GetAllocatedSize() const { return SphereElems.GetAllocatedSize() + SphylElems.GetAllocatedSize() + BoxElems.GetAllocatedSize() + ConvexElems.GetAllocatedSize() + TaperedCapsuleElems.GetAllocatedSize() + LevelSetElems.GetAllocatedSize() + SkinnedLevelSetElems.GetAllocatedSize() + MLLevelSetElems.GetAllocatedSize() + SkinnedTriangleMeshElems.GetAllocatedSize(); }
 
 	FKShapeElem* GetElement(const EAggCollisionShape::Type Type, const int32 Index)
 	{
 		switch (Type)
 		{
 		case EAggCollisionShape::Sphere:
-			if (ensure(SphereElems.IsValidIndex(Index))) { return &SphereElems[Index]; }
+			if (ensure(SphereElems.IsValidIndex(Index)))
+			{
+				return &SphereElems[Index];
+			}
+			break;
+		case EAggCollisionShape::Box:
+			if (ensure(BoxElems.IsValidIndex(Index)))
+			{
+				return &BoxElems[Index];
+			}
+			break;
+		case EAggCollisionShape::Sphyl:
+			if (ensure(SphylElems.IsValidIndex(Index)))
+			{
+				return &SphylElems[Index];
+			}
+			break;
+		case EAggCollisionShape::Convex:
+			if (ensure(ConvexElems.IsValidIndex(Index)))
+			{
+				return &ConvexElems[Index];
+			}
+			break;
+		case EAggCollisionShape::TaperedCapsule:
+			if (ensure(TaperedCapsuleElems.IsValidIndex(Index)))
+			{
+				return &TaperedCapsuleElems[Index];
+			}
+			break;
+		case EAggCollisionShape::LevelSet:
+			if (ensure(LevelSetElems.IsValidIndex(Index)))
+			{
+				return &LevelSetElems[Index];
+			}
+			break;
+		case EAggCollisionShape::SkinnedLevelSet:
+			if (ensure(SkinnedLevelSetElems.IsValidIndex(Index)))
+			{
+				return &SkinnedLevelSetElems[Index];
+			}
+			break;
+		case EAggCollisionShape::MLLevelSet:
+			if (ensure(MLLevelSetElems.IsValidIndex(Index)))
+			{
+				return &MLLevelSetElems[Index];
+			}
+			break;
+		case EAggCollisionShape::SkinnedTriangleMesh:
+			if (ensure(SkinnedTriangleMeshElems.IsValidIndex(Index)))
+			{
+				return &SkinnedTriangleMeshElems[Index];
+			}
+			break;
+		default:
+			ensure(false);
+		}
+		return nullptr;
+	}
+
+	const FKShapeElem* GetElement(const EAggCollisionShape::Type Type, const int32 Index) const
+	{
+		switch (Type)
+		{
+		case EAggCollisionShape::Sphere:
+			if (ensure(SphereElems.IsValidIndex(Index)))
+			{
+				return &SphereElems[Index];
+			}
+			break;
 		case EAggCollisionShape::Box:
-			if (ensure(BoxElems.IsValidIndex(Index))) { return &BoxElems[Index]; }
+			if (ensure(BoxElems.IsValidIndex(Index)))
+			{
+				return &BoxElems[Index];
+			}
+			break;
 		case EAggCollisionShape::Sphyl:
-			if (ensure(SphylElems.IsValidIndex(Index))) { return &SphylElems[Index]; }
+			if (ensure(SphylElems.IsValidIndex(Index)))
+			{
+				return &SphylElems[Index];
+			}
+			break;
 		case EAggCollisionShape::Convex:
-			if (ensure(ConvexElems.IsValidIndex(Index))) { return &ConvexElems[Index]; }
+			if (ensure(ConvexElems.IsValidIndex(Index)))
+			{
+				return &ConvexElems[Index];
+			}
+			break;
 		case EAggCollisionShape::TaperedCapsule:
-			if (ensure(TaperedCapsuleElems.IsValidIndex(Index))) { return &TaperedCapsuleElems[Index]; }
+			if (ensure(TaperedCapsuleElems.IsValidIndex(Index)))
+			{
+				return &TaperedCapsuleElems[Index];
+			}
+			break;
 		case EAggCollisionShape::LevelSet:
-			if (ensure(LevelSetElems.IsValidIndex(Index))) { return &LevelSetElems[Index]; }
+			if (ensure(LevelSetElems.IsValidIndex(Index)))
+			{
+				return &LevelSetElems[Index];
+			}
+			break;
 		case EAggCollisionShape::SkinnedLevelSet:
-			if (ensure(SkinnedLevelSetElems.IsValidIndex(Index))) { return &SkinnedLevelSetElems[Index]; }
+			if (ensure(SkinnedLevelSetElems.IsValidIndex(Index)))
+			{
+				return &SkinnedLevelSetElems[Index];
+			}
+			break;
 		case EAggCollisionShape::MLLevelSet:
-			if (ensure(MLLevelSetElems.IsValidIndex(Index))) { return &MLLevelSetElems[Index]; }
+			if (ensure(MLLevelSetElems.IsValidIndex(Index)))
+			{
+				return &MLLevelSetElems[Index];
+			}
+			break;
+		case EAggCollisionShape::SkinnedTriangleMesh:
+			if (ensure(SkinnedTriangleMeshElems.IsValidIndex(Index)))
+			{
+				return &SkinnedTriangleMeshElems[Index];
+			}
+			break;
 		default:
 			ensure(false);
-			return nullptr;
 		}
+		return nullptr;
 	}
 
 	FKShapeElem* GetElement(const int32 InIndex)
 	{
 		int Index = InIndex;
-		if (Index < SphereElems.Num()) { return &SphereElems[Index]; }
+		if (SphereElems.IsValidIndex(Index))
+		{
+			return &SphereElems[Index];
+		}
 		Index -= SphereElems.Num();
-		if (Index < BoxElems.Num()) { return &BoxElems[Index]; }
+		if (BoxElems.IsValidIndex(Index))
+		{
+			return &BoxElems[Index];
+		}
 		Index -= BoxElems.Num();
-		if (Index < SphylElems.Num()) { return &SphylElems[Index]; }
+		if (SphylElems.IsValidIndex(Index))
+		{
+			return &SphylElems[Index];
+		}
 		Index -= SphylElems.Num();
-		if (Index < ConvexElems.Num()) { return &ConvexElems[Index]; }
+		if (ConvexElems.IsValidIndex(Index))
+		{
+			return &ConvexElems[Index];
+		}
 		Index -= ConvexElems.Num();
-		if (Index < TaperedCapsuleElems.Num()) { return &TaperedCapsuleElems[Index]; }
+		if (TaperedCapsuleElems.IsValidIndex(Index))
+		{
+			return &TaperedCapsuleElems[Index];
+		}
 		Index -= TaperedCapsuleElems.Num();
-		if (Index < LevelSetElems.Num()) { return &LevelSetElems[Index]; }
+		if (LevelSetElems.IsValidIndex(Index))
+		{
+			return &LevelSetElems[Index];
+		}
 		Index -= LevelSetElems.Num();
-		if (Index < SkinnedLevelSetElems.Num()) { return &SkinnedLevelSetElems[Index]; }
+		if (SkinnedLevelSetElems.IsValidIndex(Index))
+		{
+			return &SkinnedLevelSetElems[Index];
+		}
+		Index -= SkinnedLevelSetElems.Num();
+		if (MLLevelSetElems.IsValidIndex(Index))
+		{
+			return &MLLevelSetElems[Index];
+		}
+		Index -= MLLevelSetElems.Num();
+		if (SkinnedTriangleMeshElems.IsValidIndex(Index))
+		{
+			return &SkinnedTriangleMeshElems[Index];
+		}
 		ensure(false);
 		return nullptr;
 	}
@@ -123,21 +260,50 @@ struct FKAggregateGeom
 	const FKShapeElem* GetElement(const int32 InIndex) const
 	{
 		int Index = InIndex;
-		if (Index < SphereElems.Num()) { return &SphereElems[Index]; }
+		if (SphereElems.IsValidIndex(Index))
+		{
+			return &SphereElems[Index];
+		}
 		Index -= SphereElems.Num();
-		if (Index < BoxElems.Num()) { return &BoxElems[Index]; }
+		if (BoxElems.IsValidIndex(Index))
+		{
+			return &BoxElems[Index];
+		}
 		Index -= BoxElems.Num();
-		if (Index < SphylElems.Num()) { return &SphylElems[Index]; }
+		if (SphylElems.IsValidIndex(Index))
+		{
+			return &SphylElems[Index];
+		}
 		Index -= SphylElems.Num();
-		if (Index < ConvexElems.Num()) { return &ConvexElems[Index]; }
+		if (ConvexElems.IsValidIndex(Index))
+		{
+			return &ConvexElems[Index];
+		}
 		Index -= ConvexElems.Num();
-		if (Index < TaperedCapsuleElems.Num()) { return &TaperedCapsuleElems[Index]; }
+		if (TaperedCapsuleElems.IsValidIndex(Index))
+		{
+			return &TaperedCapsuleElems[Index];
+		}
 		Index -= TaperedCapsuleElems.Num();
-		if (Index < LevelSetElems.Num()) { return &LevelSetElems[Index]; }
+		if (LevelSetElems.IsValidIndex(Index))
+		{
+			return &LevelSetElems[Index];
+		}
 		Index -= LevelSetElems.Num();
-		if (Index < SkinnedLevelSetElems.Num()) { return &SkinnedLevelSetElems[Index]; }
+		if (SkinnedLevelSetElems.IsValidIndex(Index))
+		{
+			return &SkinnedLevelSetElems[Index];
+		}
 		Index -= SkinnedLevelSetElems.Num();
-		if (Index < MLLevelSetElems.Num()) { return &MLLevelSetElems[Index]; }
+		if (MLLevelSetElems.IsValidIndex(Index))
+		{
+			return &MLLevelSetElems[Index];
+		}
+		Index -= MLLevelSetElems.Num();
+		if (SkinnedTriangleMeshElems.IsValidIndex(Index))
+		{
+			return &SkinnedTriangleMeshElems[Index];
+		}
 		ensure(false);
 		return nullptr;
 	}
@@ -176,6 +342,10 @@ struct FKAggregateGeom
 		{
 			return FoundMLLevelSetElem;
 		}
+		else if (const FKShapeElem* FoundSkinnedTriangleMeshElem = GetElementByName<FKSkinnedTriangleMeshElem>(MakeArrayView(SkinnedTriangleMeshElems), InName))
+		{
+			return FoundSkinnedLevelSetElem;
+		}
 
 		return nullptr;
 	}
@@ -234,7 +404,13 @@ struct FKAggregateGeom
 		FoundIndex = GetElementIndexByName<FKMLLevelSetElem>(MakeArrayView(MLLevelSetElems), InName);
 		if (FoundIndex != INDEX_NONE)
 		{
-			return FoundIndex;
+			return FoundIndex + StartIndex;
+		}
+
+		FoundIndex = GetElementIndexByName<FKSkinnedTriangleMeshElem>(MakeArrayView(SkinnedTriangleMeshElems), InName);
+		if (FoundIndex != INDEX_NONE)
+		{
+			return FoundIndex + StartIndex;
 		}
 
 		return INDEX_NONE;
@@ -257,6 +433,8 @@ struct FKAggregateGeom
 		CleanUp(TaperedCapsuleElems);
 		CleanUp(LevelSetElems);
 		CleanUp(SkinnedLevelSetElems);
+		CleanUp(MLLevelSetElems);
+		CleanUp(SkinnedTriangleMeshElems);
 
 		FreeRenderInfo();
 	}
@@ -272,6 +450,7 @@ struct FKAggregateGeom
 		LevelSetElems.Empty();
 		SkinnedLevelSetElems.Empty();
 		MLLevelSetElems.Empty();
+		SkinnedTriangleMeshElems.Empty();
 		FreeRenderInfo();
 	}
 
@@ -317,6 +496,7 @@ private:
 		LevelSetElems = Other.LevelSetElems;
 		SkinnedLevelSetElems = Other.SkinnedLevelSetElems;
 		MLLevelSetElems = Other.MLLevelSetElems;
+		SkinnedTriangleMeshElems = Other.SkinnedTriangleMeshElems;
 	}
 
 	template <class T>
diff --git a/Engine/Source/Runtime/Engine/Classes/PhysicsEngine/BoxElem.h b/Engine/Source/Runtime/Engine/Classes/PhysicsEngine/BoxElem.h
--- a/Engine/Source/Runtime/Engine/Classes/PhysicsEngine/BoxElem.h
+++ b/Engine/Source/Runtime/Engine/Classes/PhysicsEngine/BoxElem.h
@@ -99,7 +99,7 @@ struct FKBoxElem : public FKShapeElem
 	};
 
 	// Utility function that builds an FTransform from the current data
-	FTransform GetTransform() const
+	virtual FTransform GetTransform() const override final
 	{
 		return FTransform( Rotation, Center );
 	};
diff --git a/Engine/Source/Runtime/Engine/Classes/PhysicsEngine/ConvexElem.h b/Engine/Source/Runtime/Engine/Classes/PhysicsEngine/ConvexElem.h
--- a/Engine/Source/Runtime/Engine/Classes/PhysicsEngine/ConvexElem.h
+++ b/Engine/Source/Runtime/Engine/Classes/PhysicsEngine/ConvexElem.h
@@ -124,7 +124,7 @@ public:
 	ENGINE_API TArray<int32> GetChaosConvexIndices() const;
 
 	/** Get current transform applied to convex mesh vertices */
-	FTransform GetTransform() const
+	virtual FTransform GetTransform() const override final
 	{
 		return Transform;
 	};
diff --git a/Engine/Source/Runtime/Engine/Classes/PhysicsEngine/LevelSetElem.h b/Engine/Source/Runtime/Engine/Classes/PhysicsEngine/LevelSetElem.h
--- a/Engine/Source/Runtime/Engine/Classes/PhysicsEngine/LevelSetElem.h
+++ b/Engine/Source/Runtime/Engine/Classes/PhysicsEngine/LevelSetElem.h
@@ -36,7 +36,7 @@ struct FKLevelSetElem : public FKShapeElem
 
 	ENGINE_API void GetLevelSetData(FTransform& OutGridTransform, TArray<double>& OutGridValues, FIntVector& OutGridDims, float& OutGridCellSize) const;
 
-	FTransform GetTransform() const
+	virtual FTransform GetTransform() const override final
 	{
 		return Transform;
 	};
diff --git a/Engine/Source/Runtime/Engine/Classes/PhysicsEngine/MLLevelSetElem.h b/Engine/Source/Runtime/Engine/Classes/PhysicsEngine/MLLevelSetElem.h
--- a/Engine/Source/Runtime/Engine/Classes/PhysicsEngine/MLLevelSetElem.h
+++ b/Engine/Source/Runtime/Engine/Classes/PhysicsEngine/MLLevelSetElem.h
@@ -34,7 +34,10 @@ struct FKMLLevelSetElem : public FKShapeElem
 
 	ENGINE_API void BuildMLLevelSet(Chaos::FMLLevelSetImportData&& MLLevelSetImportData);
 
-	FTransform GetTransform() const { return FTransform(); }
+	virtual FTransform GetTransform() const override final
+	{
+		return FTransform();
+	}
 
 	// Draw helpers
 	/** Get geometry of all cells where the level set function is less than or equal to InteriorThreshold */
diff --git a/Engine/Source/Runtime/Engine/Classes/PhysicsEngine/ShapeElem.h b/Engine/Source/Runtime/Engine/Classes/PhysicsEngine/ShapeElem.h
--- a/Engine/Source/Runtime/Engine/Classes/PhysicsEngine/ShapeElem.h
+++ b/Engine/Source/Runtime/Engine/Classes/PhysicsEngine/ShapeElem.h
@@ -22,6 +22,7 @@ namespace EAggCollisionShape
 		LevelSet,
 		SkinnedLevelSet,
 		MLLevelSet,
+		SkinnedTriangleMesh,
 
 		Unknown
 	};
@@ -108,6 +109,11 @@ struct FKShapeElem
 	/** Get whether this shape should be considered for query or sim collision */
 	ECollisionEnabled::Type GetCollisionEnabled() const { return CollisionEnabled; }
 
+	virtual FTransform GetTransform() const 
+	{
+		return FTransform();
+	}
+
 	virtual void DrawElemWire(class FPrimitiveDrawInterface* PDI, const FTransform& ElemTM, const float Scale, const FColor Color) const {}
 	virtual void DrawElemSolid(class FPrimitiveDrawInterface* PDI, const FTransform& ElemTM, const float Scale, const class FMaterialRenderProxy* MaterialRenderProxy) const {}
 
diff --git a/Engine/Source/Runtime/Engine/Classes/PhysicsEngine/SkinnedLevelSetElem.h b/Engine/Source/Runtime/Engine/Classes/PhysicsEngine/SkinnedLevelSetElem.h
--- a/Engine/Source/Runtime/Engine/Classes/PhysicsEngine/SkinnedLevelSetElem.h
+++ b/Engine/Source/Runtime/Engine/Classes/PhysicsEngine/SkinnedLevelSetElem.h
@@ -33,7 +33,10 @@ struct FKSkinnedLevelSetElem : public FKShapeElem
 	UE_DEPRECATED(5.4, "Please use SetWeightedLevelSet with TRefCountPtr instead")
 	ENGINE_API void SetWeightedLevelSet(TUniquePtr< Chaos::TWeightedLatticeImplicitObject<Chaos::FLevelSet>>&& InWeightedLevelSet) {check(false);}
 
-	ENGINE_API FTransform GetTransform() const;
+	virtual FTransform GetTransform() const override final
+	{
+		return FTransform();
+	}
 
 	// Draw functions
 	ENGINE_API virtual void DrawElemWire(class FPrimitiveDrawInterface* PDI, const FTransform& ElemTM, const float Scale, const FColor Color) const override;
diff --git a/Engine/Source/Runtime/Engine/Classes/PhysicsEngine/SkinnedTriangleMeshElem.h b/Engine/Source/Runtime/Engine/Classes/PhysicsEngine/SkinnedTriangleMeshElem.h
new file mode 100644
--- /dev/null
+++ b/Engine/Source/Runtime/Engine/Classes/PhysicsEngine/SkinnedTriangleMeshElem.h
@@ -0,0 +1,63 @@
+// Copyright Epic Games, Inc. All Rights Reserved.
+
+#pragma once
+
+#include "UObject/ObjectMacros.h"
+#include "PhysicsEngine/ShapeElem.h"
+#include "Chaos/SkinnedTriangleMesh.h"
+#include "SkinnedTriangleMeshElem.generated.h"
+
+USTRUCT()
+struct FKSkinnedTriangleMeshElem : public FKShapeElem
+{
+	GENERATED_BODY()
+
+	FKSkinnedTriangleMeshElem() :
+		FKShapeElem(EAggCollisionShape::SkinnedTriangleMesh)
+	{}
+
+	FKSkinnedTriangleMeshElem(const FKSkinnedTriangleMeshElem& Other)
+	{
+		CloneElem(Other);
+	}
+
+	const FKSkinnedTriangleMeshElem& operator=(const FKSkinnedTriangleMeshElem& Other)
+	{
+		CloneElem(Other);
+		return *this;
+	}
+
+	ENGINE_API void SetSkinnedTriangleMesh(TRefCountPtr<Chaos::FSkinnedTriangleMesh>&& InSkinnedTriangleMesh);
+	ENGINE_API const TRefCountPtr<Chaos::FSkinnedTriangleMesh>& GetSkinnedTriangleMesh() const;
+
+	virtual FTransform GetTransform() const override final
+	{
+		return FTransform();
+	}
+
+	// Draw functions
+	ENGINE_API virtual void DrawElemWire(class FPrimitiveDrawInterface* PDI, const FTransform& ElemTM, const float Scale, const FColor Color) const override;
+	ENGINE_API virtual void DrawElemSolid(class FPrimitiveDrawInterface* PDI, const FTransform& ElemTM, const float Scale, const FMaterialRenderProxy* MaterialRenderProxy) const override;
+	ENGINE_API void GetElemSolid(const FTransform& ElemTM, const FVector& Scale3D, const FMaterialRenderProxy* MaterialRenderProxy, int32 ViewIndex, class FMeshElementCollector& Collector) const;
+
+	ENGINE_API FBox CalcAABB(const FTransform& BoneTM, const FVector& Scale3D) const;
+
+	ENGINE_API bool Serialize(FArchive& Ar);
+
+private:
+
+	TRefCountPtr<Chaos::FSkinnedTriangleMesh> SkinnedTriangleMesh;
+
+	/** Helper function to safely copy instances of this shape*/
+	 ENGINE_API void CloneElem(const FKSkinnedTriangleMeshElem& Other);
+};
+
+/* Enable our own serialization function to handle FKSkinnedTriangleMeshElem */
+template<>
+struct TStructOpsTypeTraits<FKSkinnedTriangleMeshElem> : public TStructOpsTypeTraitsBase2<FKSkinnedTriangleMeshElem>
+{
+	enum
+	{
+		WithSerializer = true
+	};
+};
diff --git a/Engine/Source/Runtime/Engine/Classes/PhysicsEngine/SphereElem.h b/Engine/Source/Runtime/Engine/Classes/PhysicsEngine/SphereElem.h
--- a/Engine/Source/Runtime/Engine/Classes/PhysicsEngine/SphereElem.h
+++ b/Engine/Source/Runtime/Engine/Classes/PhysicsEngine/SphereElem.h
@@ -63,7 +63,7 @@ struct FKSphereElem : public FKShapeElem
 	}
 
 	// Utility function that builds an FTransform from the current data
-	FTransform GetTransform() const
+	virtual FTransform GetTransform() const override final
 	{
 		return FTransform( Center );
 	};
diff --git a/Engine/Source/Runtime/Engine/Classes/PhysicsEngine/SphylElem.h b/Engine/Source/Runtime/Engine/Classes/PhysicsEngine/SphylElem.h
--- a/Engine/Source/Runtime/Engine/Classes/PhysicsEngine/SphylElem.h
+++ b/Engine/Source/Runtime/Engine/Classes/PhysicsEngine/SphylElem.h
@@ -80,7 +80,7 @@ struct FKSphylElem : public FKShapeElem
 	};
 
 	// Utility function that builds an FTransform from the current data
-	FTransform GetTransform() const
+	virtual FTransform GetTransform() const override final
 	{
 		return FTransform(Rotation, Center );
 	};
diff --git a/Engine/Source/Runtime/Engine/Classes/PhysicsEngine/TaperedCapsuleElem.h b/Engine/Source/Runtime/Engine/Classes/PhysicsEngine/TaperedCapsuleElem.h
--- a/Engine/Source/Runtime/Engine/Classes/PhysicsEngine/TaperedCapsuleElem.h
+++ b/Engine/Source/Runtime/Engine/Classes/PhysicsEngine/TaperedCapsuleElem.h
@@ -51,7 +51,7 @@ struct FKTaperedCapsuleElem : public FKShapeElem
 	};
 
 	// Utility function that builds an FTransform from the current data
-	FTransform GetTransform() const
+	virtual FTransform GetTransform() const override final
 	{
 		return FTransform(Rotation, Center );
 	};
diff --git a/Engine/Source/Runtime/Engine/Private/Collision/KAggregateGeom.cpp b/Engine/Source/Runtime/Engine/Private/Collision/KAggregateGeom.cpp
--- a/Engine/Source/Runtime/Engine/Private/Collision/KAggregateGeom.cpp
+++ b/Engine/Source/Runtime/Engine/Private/Collision/KAggregateGeom.cpp
@@ -14,13 +14,15 @@
 #include "PhysicsEngine/SphereElem.h"
 #include "PhysicsEngine/SphylElem.h"
 #include "PhysicsEngine/AggregateGeom.h"
+#include "PhysicsEngine/SkinnedTriangleMeshElem.h"
+#include "PhysicsEngine/TaperedCapsuleElem.h"
 #include "Engine/Polys.h"
 #include "Chaos/Convex.h"
 #include "Chaos/Levelset.h"
-#include "PhysicsEngine/TaperedCapsuleElem.h"
+#include "Chaos/MLLevelset.h"
+#include "Chaos/SkinnedTriangleMesh.h"
 #include "Chaos/WeightedLatticeImplicitObject.h"
 #include "AutoRTFM/AutoRTFM.h"
-#include "Chaos/MLLevelset.h"
 #include "UObject/Package.h"
 
 #if INTEL_ISPC
@@ -130,6 +132,11 @@ FBox FKAggregateGeom::CalcAABB(const FTransform& Transform) const
 		Box += MLLevelSetElems[i].CalcAABB(BoneTM, (FVector)Scale3D);
 	}
 
+	for (int32 i = 0; i < SkinnedTriangleMeshElems.Num(); i++)
+	{
+		Box += SkinnedTriangleMeshElems[i].CalcAABB(BoneTM, (FVector)Scale3D);
+	}
+
 	return Box;
 
 }
@@ -689,12 +696,7 @@ FIntVector3 FKSkinnedLevelSetElem::LatticeGridResolution() const
 
 void FKSkinnedLevelSetElem::SetWeightedLevelSet(TRefCountPtr< Chaos::TWeightedLatticeImplicitObject<Chaos::FLevelSet>>&& InWeightedLevelSet)
 {
-	WeightedLatticeLevelSet = InWeightedLevelSet;
-}
-
-FTransform FKSkinnedLevelSetElem::GetTransform() const
-{
-	return FTransform();
+	WeightedLatticeLevelSet = MoveTemp(InWeightedLevelSet);
 }
 
 ///////////////////////////////////////
@@ -785,6 +787,33 @@ void FKMLLevelSetElem::GetZeroIsosurfaceGridCellFaces(TArray<FVector3f>& Vertice
 	}
 }
 
+
+///////////////////////////////////////
+///// FKSkinnedTriangleMeshElem ///////
+///////////////////////////////////////
+FBox FKSkinnedTriangleMeshElem::CalcAABB(const FTransform& BoneTM, const FVector& Scale3D) const
+{
+	FBox Box(ForceInit);
+	if (SkinnedTriangleMesh.IsValid())
+	{
+		Box = FBox(SkinnedTriangleMesh->BoundingBox().Min(), SkinnedTriangleMesh->BoundingBox().Max());
+		const FTransform LocalToWorld = FTransform(FQuat::Identity, FVector::ZeroVector, Scale3D) * BoneTM;
+		return Box.TransformBy(LocalToWorld);
+	}
+
+	return Box;
+}
+
+void FKSkinnedTriangleMeshElem::SetSkinnedTriangleMesh(TRefCountPtr<Chaos::FSkinnedTriangleMesh>&& InSkinnedTriangleMesh)
+{
+	SkinnedTriangleMesh = MoveTemp(InSkinnedTriangleMesh);
+}
+
+const TRefCountPtr<Chaos::FSkinnedTriangleMesh>& FKSkinnedTriangleMeshElem::GetSkinnedTriangleMesh() const
+{
+	return SkinnedTriangleMesh;
+}
+
 static const float DIST_COMPARE_THRESH = 0.1f;
 static const float DIR_COMPARE_THRESH = 0.0003f; // about 1 degree
 
@@ -1145,3 +1174,18 @@ bool FKMLLevelSetElem::Serialize(FArchive& Ar)
 	return true;
 }
 
+bool FKSkinnedTriangleMeshElem::Serialize(FArchive& Ar)
+{
+	if (Ar.IsLoading())
+	{
+		SkinnedTriangleMesh = TRefCountPtr<Chaos::FSkinnedTriangleMesh>(new Chaos::FSkinnedTriangleMesh());
+	}
+
+	if (SkinnedTriangleMesh.IsValid())
+	{
+		Chaos::FChaosArchive ChaosAr(Ar);
+		SkinnedTriangleMesh->Serialize(ChaosAr);
+	}
+
+	return true;
+}
diff --git a/Engine/Source/Runtime/Engine/Private/PhysicsEngine/BodySetup.cpp b/Engine/Source/Runtime/Engine/Private/PhysicsEngine/BodySetup.cpp
--- a/Engine/Source/Runtime/Engine/Private/PhysicsEngine/BodySetup.cpp
+++ b/Engine/Source/Runtime/Engine/Private/PhysicsEngine/BodySetup.cpp
@@ -40,6 +40,7 @@
 #include "PhysicsEngine/SphylElem.h"
 #include "PhysicsEngine/TaperedCapsuleElem.h"
 #include "PhysicsEngine/MLLevelSetElem.h"
+#include "PhysicsEngine/SkinnedTriangleMeshElem.h"
 
 #if WITH_EDITOR
 #include "Misc/DataValidation.h"
@@ -2352,6 +2353,12 @@ void FKMLLevelSetElem::CloneElem(const FKMLLevelSetElem& Other)
 	NNEIncorrectZoneModelData = Other.NNEIncorrectZoneModelData;
 }
 
+void FKSkinnedTriangleMeshElem::CloneElem(const FKSkinnedTriangleMeshElem& Other)
+{
+	Super::CloneElem(Other);
+	SkinnedTriangleMesh = Other.SkinnedTriangleMesh;
+}
+
 #if WITH_EDITOR
 
 bool FBodySetupObjectTextFactory::CanCreateClass(UClass* InObjectClass, bool& bOmitSubObjs) const
diff --git a/Engine/Source/Runtime/Engine/Private/PhysicsEngine/PhysDrawing.cpp b/Engine/Source/Runtime/Engine/Private/PhysicsEngine/PhysDrawing.cpp
--- a/Engine/Source/Runtime/Engine/Private/PhysicsEngine/PhysDrawing.cpp
+++ b/Engine/Source/Runtime/Engine/Private/PhysicsEngine/PhysDrawing.cpp
@@ -18,10 +18,12 @@
 #include "PhysicsEngine/PhysicsAsset.h"
 #include "PhysicsEngine/SkeletalBodySetup.h"
 #include "PhysicsEngine/SkinnedLevelSetElem.h"
+#include "PhysicsEngine/SkinnedTriangleMeshElem.h"
 #include "PhysicsEngine/TaperedCapsuleElem.h"
 #include "ReferenceSkeleton.h"
 #include "StaticMeshResources.h"
 #include "Chaos/Levelset.h"
+#include "Chaos/SkinnedTriangleMesh.h"
 #include "Chaos/UniformGrid.h"
 #include "Chaos/WeightedLatticeImplicitObject.h"
 #include "Chaos/PBDJointConstraintUtilities.h"
@@ -1100,6 +1102,72 @@ void FKMLLevelSetElem::GetElemSolid(const FTransform& ElemTM, const FVector& Sca
 	MeshBuilder.GetMesh(ElemTM.ToMatrixWithScale(), MaterialRenderProxy, SDPG_World, false, false, ViewIndex, Collector);
 }
 
+/////////////////////////////////////////////////////////////////////////////////////
+// FKSkinnedTriangleMesh
+/////////////////////////////////////////////////////////////////////////////////////
+void FKSkinnedTriangleMeshElem::DrawElemWire(class FPrimitiveDrawInterface* PDI, const FTransform& ElemTM, float Scale, const FColor Color) const
+{
+	if (SkinnedTriangleMesh.IsValid())
+	{
+		const TArray<Chaos::TVec3<int32>>& Elements = SkinnedTriangleMesh->GetTriangleMesh().GetElements();
+		const TConstArrayView<Chaos::Softs::FSolverVec3>& Positions = SkinnedTriangleMesh->GetLocalPositions();
+		PDI->AddReserveLines(SDPG_World, Elements.Num() * 3);
+		for (const Chaos::TVec3<int32>& Element : Elements)
+		{
+			const FVector P0 = ElemTM.TransformPosition(FVector(Positions[Element[0]]));
+			const FVector P1 = ElemTM.TransformPosition(FVector(Positions[Element[1]]));
+			const FVector P2 = ElemTM.TransformPosition(FVector(Positions[Element[2]]));
+			PDI->DrawLine(P0, P1, Color, SDPG_World);
+			PDI->DrawLine(P1, P2, Color, SDPG_World);
+			PDI->DrawLine(P2, P0, Color, SDPG_World);
+		}
+	}
+}
+
+void FKSkinnedTriangleMeshElem::DrawElemSolid(class FPrimitiveDrawInterface* PDI, const FTransform& ElemTM, const float Scale, const FMaterialRenderProxy* MaterialRenderProxy) const
+{
+	if (SkinnedTriangleMesh.IsValid())
+	{
+		const TArray<Chaos::TVec3<int32>>& Elements = SkinnedTriangleMesh->GetTriangleMesh().GetElements();
+		const TConstArrayView<Chaos::Softs::FSolverVec3>& Positions = SkinnedTriangleMesh->GetLocalPositions();
+
+		FDynamicMeshBuilder MeshBuilder(PDI->View->GetFeatureLevel());
+		for (const Chaos::Softs::FSolverVec3& Position : Positions)
+		{
+			MeshBuilder.AddVertex(FDynamicMeshVertex(FVector3f(Position)));
+		}
+		for (const Chaos::TVec3<int32>& Element : Elements)
+		{
+			MeshBuilder.AddTriangle(Element[0], Element[1], Element[2]);
+		}
+
+		MeshBuilder.Draw(PDI, ElemTM.ToMatrixWithScale(), MaterialRenderProxy, SDPG_World, 0.f);
+	}
+}
+
+
+
+void FKSkinnedTriangleMeshElem::GetElemSolid(const FTransform& ElemTM, const FVector& Scale3D, const FMaterialRenderProxy* MaterialRenderProxy, int32 ViewIndex, class FMeshElementCollector& Collector) const
+{
+	if (SkinnedTriangleMesh.IsValid())
+	{
+		const TArray<Chaos::TVec3<int32>>& Elements = SkinnedTriangleMesh->GetTriangleMesh().GetElements();
+		const TConstArrayView<Chaos::Softs::FSolverVec3>& Positions = SkinnedTriangleMesh->GetLocalPositions();
+
+		FDynamicMeshBuilder MeshBuilder(Collector.GetFeatureLevel());
+		for (const Chaos::Softs::FSolverVec3& Position : Positions)
+		{
+			MeshBuilder.AddVertex(FDynamicMeshVertex(FVector3f(Position)));
+		}
+		for (const Chaos::TVec3<int32>& Element : Elements)
+		{
+			MeshBuilder.AddTriangle(Element[0], Element[1], Element[2]);
+		}
+
+		MeshBuilder.GetMesh(ElemTM.ToMatrixWithScale(), MaterialRenderProxy, SDPG_World, false, false, ViewIndex, Collector);
+	}
+}
+
 /////////////////////////////////////////////////////////////////////////////////////
 // FKConvexGeomRenderInfo
 /////////////////////////////////////////////////////////////////////////////////////
@@ -1297,6 +1365,30 @@ void FKAggregateGeom::GetAggGeom(const FTransform& Transform, const FColor Color
 			SkinnedLevelSetElems[i].DrawElemWire(Collector.GetPDI(ViewIndex), Transform, 1.f, Color);
 		}
 	}
+
+	for (int32 i = 0; i < MLLevelSetElems.Num(); i++)
+	{
+		if (bDrawSolid)
+		{
+			MLLevelSetElems[i].GetElemSolid(Transform, Scale3D, MatInst, ViewIndex, Collector);
+		}
+		else
+		{
+			MLLevelSetElems[i].DrawElemWire(Collector.GetPDI(ViewIndex), Transform, 1.f, Color);
+		}
+	}
+
+	for (int32 i = 0; i < SkinnedTriangleMeshElems.Num(); i++)
+	{
+		if (bDrawSolid)
+		{
+			SkinnedTriangleMeshElems[i].GetElemSolid(Transform, Scale3D, MatInst, ViewIndex, Collector);
+		}
+		else
+		{
+			SkinnedTriangleMeshElems[i].DrawElemWire(Collector.GetPDI(ViewIndex), Transform, 1.f, Color);
+		}
+	}
 }
 
 /** Release the RenderInfo (if its there) and safely clean up any resources. Can be called from any thread. */
diff --git a/Engine/Source/Runtime/Experimental/Chaos/Private/Chaos/MLLevelset.cpp b/Engine/Source/Runtime/Experimental/Chaos/Private/Chaos/MLLevelset.cpp
--- a/Engine/Source/Runtime/Experimental/Chaos/Private/Chaos/MLLevelset.cpp
+++ b/Engine/Source/Runtime/Experimental/Chaos/Private/Chaos/MLLevelset.cpp
@@ -15,7 +15,6 @@ static FAutoConsoleVariableRef CVarChaosMLLevelSetUpdatePhiFlag(TEXT("p.MLLevelS
 
 FMLLevelSet::FMLLevelSet(FMLLevelSetImportData&& MLLevelSetImportData)
 	: FImplicitObject(EImplicitObject::HasBoundingBox, ImplicitObjectType::MLLevelSet)
-	, ParentBoneName(MLLevelSetImportData.ParentBoneName)
 	, ActiveBoneNames(MoveTemp(MLLevelSetImportData.ActiveBoneNames))
 	, SignedDistanceScaling(MLLevelSetImportData.SignedDistanceScaling)
 	, TrainingGridMin(MLLevelSetImportData.TrainingGridMin)
@@ -64,7 +63,6 @@ CHAOS_API FMLLevelSet::FMLLevelSet(TObjectPtr<UNNEModelData> InNNESignedDistance
 
 FMLLevelSet::FMLLevelSet(FMLLevelSet&& Other)
 	: FImplicitObject(EImplicitObject::HasBoundingBox, ImplicitObjectType::MLLevelSet)
-	, ParentBoneName(Other.ParentBoneName)
 	, ActiveBoneNames(MoveTemp(Other.ActiveBoneNames))
 	, SolverActiveBoneIndices(MoveTemp(Other.SolverActiveBoneIndices))
 	, SolverParentBoneIndex(Other.SolverParentBoneIndex)
@@ -103,7 +101,6 @@ Chaos::FImplicitObjectPtr FMLLevelSet::CopyGeometry() const
 {
 	FMLLevelSet* Copy = new FMLLevelSet();
 
-	Copy->ParentBoneName = ParentBoneName;
 	Copy->ActiveBoneNames = ActiveBoneNames;
 	Copy->SolverActiveBoneIndices = SolverActiveBoneIndices;
 	Copy->ActiveBonesRelativeTransforms = ActiveBonesRelativeTransforms;
@@ -140,7 +137,6 @@ Chaos::FImplicitObjectPtr FMLLevelSet::CopyGeometryWithScale(const FVec3& Scale)
 {
 	FMLLevelSet* Copy = new FMLLevelSet();
 
-	Copy->ParentBoneName = ParentBoneName;
 	Copy->ActiveBoneNames = ActiveBoneNames;
 	Copy->SolverActiveBoneIndices = SolverActiveBoneIndices;
 	Copy->ActiveBonesRelativeTransforms = ActiveBonesRelativeTransforms;
diff --git a/Engine/Source/Runtime/Experimental/Chaos/Private/Chaos/PerParticlePBDCollisionConstraint.ispc b/Engine/Source/Runtime/Experimental/Chaos/Private/Chaos/PerParticlePBDCollisionConstraint.ispc
--- a/Engine/Source/Runtime/Experimental/Chaos/Private/Chaos/PerParticlePBDCollisionConstraint.ispc
+++ b/Engine/Source/Runtime/Experimental/Chaos/Private/Chaos/PerParticlePBDCollisionConstraint.ispc
@@ -26,8 +26,9 @@ enum ImplicitObjectType
 	Triangle,
 	UnionClustered,
 	TaperedCapsule,
-	WeightedLatticeBone,
 	MLLevelSet,
+	SkinnedTriangleMesh,
+	WeightedLatticeBone,
 	MLLevelSetBone,
 	WeightedLatticeLevelSetType = 38, //== LevelSet | IsWeightedLattice (cannot use expressions like this in ISPC)
 
diff --git a/Engine/Source/Runtime/Experimental/Chaos/Private/Chaos/SkinnedTriangleMesh.cpp b/Engine/Source/Runtime/Experimental/Chaos/Private/Chaos/SkinnedTriangleMesh.cpp
new file mode 100644
--- /dev/null
+++ b/Engine/Source/Runtime/Experimental/Chaos/Private/Chaos/SkinnedTriangleMesh.cpp
@@ -0,0 +1,156 @@
+// Copyright Epic Games, Inc. All Rights Reserved.
+#include "Chaos/SkinnedTriangleMesh.h"
+#include "Chaos/ImplicitObjectScaled.h"
+
+
+namespace Chaos
+{
+	FSkinnedTriangleMesh::FSkinnedTriangleMesh(FTriangleMesh&& InTriangleMesh, TArray<FVec3f>&& InReferencePositions, TArray<FWeightedInfluenceData>&& InBoneData, 
+		TArray<FName>&& InUsedBones, FTransform&& InReferenceRootTransform, TArray<FTransform>&& InReferenceRelativeTransforms)
+		: FImplicitObject(EImplicitObject::HasBoundingBox, ImplicitObjectType::SkinnedTriangleMesh)
+		, TriangleMesh(MoveTemp(InTriangleMesh))
+		, BoneData(MoveTemp(InBoneData))
+		, UsedBones(MoveTemp(InUsedBones))
+		, ReferenceRootTransform(MoveTemp(InReferenceRootTransform))
+		, ReferenceRelativeTransforms(MoveTemp(InReferenceRelativeTransforms))
+		, ReferencePositions(MoveTemp(InReferencePositions))
+	{
+		ReferenceBoundingBox = CalculateBoundingBox(TConstArrayView<FVec3f>(ReferencePositions));
+		FinalizeConstruction();
+	}
+
+	FSkinnedTriangleMesh::FSkinnedTriangleMesh(FSkinnedTriangleMesh&& Other)
+		: FImplicitObject(EImplicitObject::HasBoundingBox, ImplicitObjectType::SkinnedTriangleMesh)
+		, TriangleMesh(MoveTemp(Other.TriangleMesh))
+		, BoneData(MoveTemp(Other.BoneData))
+		, UsedBones(MoveTemp(Other.UsedBones))
+		, ReferenceRootTransform(MoveTemp(Other.ReferenceRootTransform))
+		, ReferenceRelativeTransforms(MoveTemp(Other.ReferenceRelativeTransforms))
+		, ReferencePositions(MoveTemp(Other.ReferencePositions))
+		, ReferenceBoundingBox(MoveTemp(Other.ReferenceBoundingBox))
+		, LocalBoundingBox(MoveTemp(Other.LocalBoundingBox))
+		, LocalPositions(MoveTemp(Other.LocalPositions))
+	{
+	}
+
+	TAABB<FRealSingle, 3> FSkinnedTriangleMesh::CalculateBoundingBox(const TConstArrayView<FVec3f>& Positions) const
+	{
+		// Should we calculate based on only positions used by Elements, or should we require that Positions contains no extra points?
+		TAABB<FRealSingle, 3> AABB;
+
+		for (const TVec3<int32>& Element : TriangleMesh.GetElements())
+		{
+			AABB.GrowToInclude(Positions[Element[0]]);
+			AABB.GrowToInclude(Positions[Element[1]]);
+			AABB.GrowToInclude(Positions[Element[2]]);
+		}
+		
+		return AABB;
+	}
+
+	void FSkinnedTriangleMesh::FinalizeConstruction()
+	{
+		LocalPositions = ReferencePositions;
+		LocalBoundingBox = ReferenceBoundingBox;
+	}
+
+	FImplicitObjectPtr FSkinnedTriangleMesh::CopyGeometry() const
+	{
+		FSkinnedTriangleMesh* const Copy = new FSkinnedTriangleMesh();
+		constexpr bool bCullDegenerateElementsFalse = false;
+		Copy->TriangleMesh.Init(TriangleMesh.GetElements(), TriangleMesh.GetStartIndex(), TriangleMesh.GetNumIndices(), bCullDegenerateElementsFalse);
+		Copy->BoneData = BoneData;
+		Copy->UsedBones = UsedBones;
+		Copy->ReferenceRootTransform = ReferenceRootTransform;
+		Copy->ReferenceRelativeTransforms = ReferenceRelativeTransforms;
+		Copy->ReferencePositions = ReferencePositions;
+		Copy->ReferenceBoundingBox = ReferenceBoundingBox;
+		Copy->LocalPositions = LocalPositions;
+		Copy->LocalBoundingBox = LocalBoundingBox;
+		return FImplicitObjectPtr(Copy);
+	}
+
+	FImplicitObjectPtr FSkinnedTriangleMesh::CopyGeometryWithScale(const FVec3& Scale) const
+	{
+		FSkinnedTriangleMesh* const Copy = new FSkinnedTriangleMesh();
+		constexpr bool bCullDegenerateElementsFalse = false;
+		Copy->TriangleMesh.Init(TriangleMesh.GetElements(), TriangleMesh.GetStartIndex(), TriangleMesh.GetNumIndices(), bCullDegenerateElementsFalse);
+		Copy->BoneData = BoneData;
+		Copy->UsedBones = UsedBones;
+		Copy->ReferenceRootTransform = ReferenceRootTransform;
+		Copy->ReferenceRelativeTransforms = ReferenceRelativeTransforms;
+		Copy->ReferencePositions = ReferencePositions;
+		Copy->ReferenceBoundingBox = ReferenceBoundingBox;
+		Copy->LocalPositions = LocalPositions;
+		Copy->LocalBoundingBox = LocalBoundingBox;
+		return MakeImplicitObjectPtr<TImplicitObjectScaled<FSkinnedTriangleMesh>>(Copy, Scale);
+	}
+
+	uint32 FSkinnedTriangleMesh::GetTypeHash() const
+	{
+		uint32 Result = GetArrayHash(TriangleMesh.GetElements().GetData(), TriangleMesh.GetElements().Num());
+		for (const FWeightedInfluenceData& Data : BoneData)
+		{
+			Result = HashCombine(Result, Data.GetTypeHash());
+		}
+		Result = GetArrayHash(ReferencePositions.GetData(), ReferencePositions.Num(), Result);
+		return Result;
+	}
+
+	void FSkinnedTriangleMesh::SkinPositions(const TArray<FTransform>& RelativeTransforms, const TArrayView<FVec3f>& Positions) const
+	{
+		check(Positions.Num() == BoneData.Num());
+		check(RelativeTransforms.Num() == ReferenceRelativeTransforms.Num());
+		TArray<FTransform3f> BoneTransforms;
+		BoneTransforms.SetNum(RelativeTransforms.Num());
+		for (int32 Index = 0; Index < RelativeTransforms.Num(); ++Index)
+		{
+			BoneTransforms[Index] = FTransform3f(ReferenceRelativeTransforms[Index] * RelativeTransforms[Index]);
+		}
+
+		// TODO: ISPC version
+
+		auto AddInfluence = [&BoneTransforms](FVec3f& OutPosition, const FVec3f& RefPosition, const uint16 BoneIndex, const float BoneWeight)
+			{
+				OutPosition += BoneTransforms[BoneIndex].TransformPosition(RefPosition) * BoneWeight;
+			};
+
+		for (int32 VertIndex = 0; VertIndex < BoneData.Num(); ++VertIndex)
+		{
+			// This is taken from ChaosClothingSimulationMesh although it uses FMatrix44f matrices instead of FTransform3f.
+			const uint16* const RESTRICT BoneIndices = BoneData[VertIndex].BoneIndices;
+			const float* const RESTRICT BoneWeights = BoneData[VertIndex].BoneWeights;
+			// WARNING - HORRIBLE UNROLLED LOOP + JUMP TABLE BELOW
+			// done this way because this is a pretty tight and performance critical loop. essentially
+			// rather than checking each influence we can just jump into this switch and fall through
+			// everything to compose the final skinned data
+
+			const FVec3f& RefPosition = ReferencePositions[VertIndex];
+			FVec3f Position((FRealSingle)0.f);
+
+			switch (BoneData[VertIndex].NumInfluences)
+			{
+			default: // Intentional fallthrough
+			case 12: AddInfluence(Position, RefPosition, BoneIndices[11], BoneWeights[11]); // Intentional fallthrough
+			case 11: AddInfluence(Position, RefPosition, BoneIndices[10], BoneWeights[10]); // Intentional fallthrough
+			case 10: AddInfluence(Position, RefPosition, BoneIndices[ 9], BoneWeights[ 9]); // Intentional fallthrough
+			case  9: AddInfluence(Position, RefPosition, BoneIndices[ 8], BoneWeights[ 8]); // Intentional fallthrough
+			case  8: AddInfluence(Position, RefPosition, BoneIndices[ 7], BoneWeights[ 7]); // Intentional fallthrough
+			case  7: AddInfluence(Position, RefPosition, BoneIndices[ 6], BoneWeights[ 6]); // Intentional fallthrough
+			case  6: AddInfluence(Position, RefPosition, BoneIndices[ 5], BoneWeights[ 5]); // Intentional fallthrough
+			case  5: AddInfluence(Position, RefPosition, BoneIndices[ 4], BoneWeights[ 4]); // Intentional fallthrough
+			case  4: AddInfluence(Position, RefPosition, BoneIndices[ 3], BoneWeights[ 3]); // Intentional fallthrough
+			case  3: AddInfluence(Position, RefPosition, BoneIndices[ 2], BoneWeights[ 2]); // Intentional fallthrough
+			case  2: AddInfluence(Position, RefPosition, BoneIndices[ 1], BoneWeights[ 1]); // Intentional fallthrough
+			case  1: AddInfluence(Position, RefPosition, BoneIndices[ 0], BoneWeights[ 0]); // Intentional fallthrough
+			case  0: break;
+			}
+			Positions[VertIndex] = Position;
+		}
+	}
+
+	void FSkinnedTriangleMesh::UpdateSpatialHierarchy(const FRealSingle MinLodSize)
+	{		
+		TriangleMesh.BuildSpatialHash(TConstArrayView<FVec3f>(LocalPositions), SpatialHash, MinLodSize);
+	}
+}
diff --git a/Engine/Source/Runtime/Experimental/Chaos/Private/Chaos/TriangleMesh.cpp b/Engine/Source/Runtime/Experimental/Chaos/Private/Chaos/TriangleMesh.cpp
--- a/Engine/Source/Runtime/Experimental/Chaos/Private/Chaos/TriangleMesh.cpp
+++ b/Engine/Source/Runtime/Experimental/Chaos/Private/Chaos/TriangleMesh.cpp
@@ -229,6 +229,17 @@ void FTriangleMesh::ResetAuxiliaryStructures()
 	MEdgeToFaces.Reset();
 }
 
+void FTriangleMesh::Serialize(FArchive& Ar)
+{
+	Ar << MElements;
+	Ar << MStartIdx;
+	Ar << MNumIndices;
+	if (Ar.IsLoading())
+	{
+		ResetAuxiliaryStructures();
+	}
+}
+
 TVec2<int32> FTriangleMesh::GetVertexRange() const
 {
 	return TVec2<int32>(MStartIdx, MStartIdx + MNumIndices - 1);
diff --git a/Engine/Source/Runtime/Experimental/Chaos/Public/Chaos/ImplicitFwd.h b/Engine/Source/Runtime/Experimental/Chaos/Public/Chaos/ImplicitFwd.h
--- a/Engine/Source/Runtime/Experimental/Chaos/Public/Chaos/ImplicitFwd.h
+++ b/Engine/Source/Runtime/Experimental/Chaos/Public/Chaos/ImplicitFwd.h
@@ -18,6 +18,7 @@ namespace Chaos
 	class FTriangleMeshImplicitObject;
 	class FImplicitObjectUnionClustered;
 	class FImplicitObjectUnion;
+	class FSkinnedTriangleMesh;
 
 	using FImplicitBox3 = TBox<FReal, 3>;
 	using FImplicitCapsule3 = FCapsule;
@@ -42,6 +43,7 @@ namespace Chaos
 	using FSpherePtr = TRefCountPtr<FImplicitSphere3>;
 	using FImplicitObjectUnionClusteredPtr = TRefCountPtr<FImplicitObjectUnionClustered>;
 	using FImplicitObjectUnionPtr = TRefCountPtr<FImplicitObjectUnion>;
+	using FSkinnedTriangleMeshPtr = TRefCountPtr<FSkinnedTriangleMesh>;
 	
 	using FImplicitObjectRef = FImplicitObject*;
 	using FConstImplicitObjectRef = const FImplicitObject*;
@@ -56,6 +58,7 @@ namespace Chaos
 	using FSphereRef = FImplicitSphere3*;
 	using FImplicitObjectUnionClusteredRef = FImplicitObjectUnionClustered*;
 	using FImplicitObjectUnionRef = FImplicitObjectUnion*;
+	using FSkinnedTriangleMeshRef = FSkinnedTriangleMesh*;
 	
 	using FImplicitObjectsArray = TArray<FImplicitObjectPtr>;
 }
\ No newline at end of file
diff --git a/Engine/Source/Runtime/Experimental/Chaos/Public/Chaos/ImplicitObjectType.h b/Engine/Source/Runtime/Experimental/Chaos/Public/Chaos/ImplicitObjectType.h
--- a/Engine/Source/Runtime/Experimental/Chaos/Public/Chaos/ImplicitObjectType.h
+++ b/Engine/Source/Runtime/Experimental/Chaos/Public/Chaos/ImplicitObjectType.h
@@ -27,8 +27,9 @@ namespace Chaos
 			Triangle,
 			UnionClustered,
 			TaperedCapsule,
-			WeightedLatticeBone,
 			MLLevelSet,
+			SkinnedTriangleMesh,
+			WeightedLatticeBone,
 			MLLevelSetBone,
 			//Add entries above this line for serialization
 			ConcreteObjectCount, // Used to ensure bitflags do not overlap concrete type
@@ -95,8 +96,9 @@ namespace Chaos
 				TEXT("EImplicitObjectType:Triangle"),
 				TEXT("EImplicitObjectType:UnionClustered"),
 				TEXT("EImplicitObjectType:TaperedCapsule"),
-				TEXT("EImplicitObjectType:WeightedLatticeBone"),
 				TEXT("EImplicitObjectType:MLLevelSet"),
+				TEXT("EImplicitObjectType:SkinnedTriangleMesh"),
+				TEXT("EImplicitObjectType:WeightedLatticeBone"),
 				TEXT("EImplicitObjectType:MLLevelSetBone"),
 			};
 
diff --git a/Engine/Source/Runtime/Experimental/Chaos/Public/Chaos/MLLevelset.h b/Engine/Source/Runtime/Experimental/Chaos/Public/Chaos/MLLevelset.h
--- a/Engine/Source/Runtime/Experimental/Chaos/Public/Chaos/MLLevelset.h
+++ b/Engine/Source/Runtime/Experimental/Chaos/Public/Chaos/MLLevelset.h
@@ -47,7 +47,6 @@ struct FMLLevelSetNNEModelData
 
 struct FMLLevelSetImportData
 {
-	const FName ParentBoneName;
 	TArray<FName> ActiveBoneNames;
 	TArray<FMLLevelSetNNEModelData> NNEModelDataArr; //Either {SignedDistanceModel} or {SignedDistanceModel, IncorrectZoneModel} 
 	TArray<FVector3f> ActiveBonesReferenceRotations;
@@ -89,7 +88,6 @@ class FMLLevelSet final : public FImplicitObject
 	const int32 GetNumberOfActiveBones() const { return ActiveBoneNames.Num(); }
 	void SetSolverParentBoneIndex(int32 InSolverParentBoneIndex) { SolverParentBoneIndex = InSolverParentBoneIndex; }
 
-	const FName GetParentBoneName() const { return ParentBoneName; }
 	const TArray<FName>& GetActiveBoneNames() const { return ActiveBoneNames; }
 	virtual const FAABB3 BoundingBox() const override { return LocalBoundingBox; } 
 	const FVector3f GetTrainingGridMin() const { return TrainingGridMin; }
@@ -110,7 +108,6 @@ class FMLLevelSet final : public FImplicitObject
 	{
 		FImplicitObject::SerializeImp(Ar);
 		TBox<FReal, 3>::SerializeAsAABB(Ar, LocalBoundingBox);
-		Ar << ParentBoneName;
 		Ar << ActiveBoneNames;
 		Ar << ActiveBonesRelativeTransforms;
 		Ar << SignedDistanceModelWeights;
@@ -197,7 +194,6 @@ class FMLLevelSet final : public FImplicitObject
 
   private:
 	//Relative Bone Transform Related Data
-	FName ParentBoneName;
 	TArray<FName> ActiveBoneNames; 
 	TArray<int32> SolverActiveBoneIndices; //indices set by the solver for bone transform query
 	int32 SolverParentBoneIndex;
diff --git a/Engine/Source/Runtime/Experimental/Chaos/Public/Chaos/SkinnedTriangleMesh.h b/Engine/Source/Runtime/Experimental/Chaos/Public/Chaos/SkinnedTriangleMesh.h
new file mode 100644
--- /dev/null
+++ b/Engine/Source/Runtime/Experimental/Chaos/Public/Chaos/SkinnedTriangleMesh.h
@@ -0,0 +1,175 @@
+// Copyright Epic Games, Inc. All Rights Reserved.
+#pragma once
+
+#include "Chaos/ImplicitObject.h"
+#include "Chaos/AABB.h"
+#include "Chaos/HierarchicalSpatialHash.h"
+#include "Chaos/TriangleMesh.h"
+
+struct FKSkinnedTriangleMeshElem;
+
+namespace Chaos
+{
+	struct FWeightedInfluenceData
+	{
+		FWeightedInfluenceData()
+			: NumInfluences(0)
+		{
+			FMemory::Memset(BoneIndices, (uint8)INDEX_NONE, sizeof(BoneIndices));
+			FMemory::Memset(BoneWeights, 0, sizeof(BoneWeights));
+		}
+
+		uint32 GetTypeHash() const
+		{
+			uint32 Result = ::GetTypeHash(NumInfluences);
+			for (uint8 Index = 0; Index < NumInfluences; ++Index)
+			{
+				Result = HashCombine(Result, ::GetTypeHash(BoneIndices[Index]));
+				Result = HashCombine(Result, ::GetTypeHash(BoneWeights[Index]));
+			}
+			return Result;
+		}
+
+		static const uint8 MaxTotalInfluences = 12;
+		uint8 NumInfluences;
+		uint16 BoneIndices[MaxTotalInfluences]; // This is index into FSkinnedTriangleMesh::UsedBones (which then maps into skeletalmesh bones)
+		float BoneWeights[MaxTotalInfluences];
+	};
+
+	inline FArchive& operator<<(FArchive& Ar, FWeightedInfluenceData& Value)
+	{
+		Ar << Value.NumInfluences;
+		for (int32 Index = 0; Index < Value.NumInfluences; ++Index)
+		{
+			Ar << Value.BoneIndices[Index];
+			Ar << Value.BoneWeights[Index];
+		}
+		return Ar;
+	}
+
+	/**
+	 * Skinned Triangle Mesh
+	 */
+	class FSkinnedTriangleMesh : public FImplicitObject
+	{
+	public:
+		using ObjectType = FImplicitObjectPtr;
+		using FImplicitObject::GetTypeName;
+
+		CHAOS_API FSkinnedTriangleMesh(FTriangleMesh&& TriangleMesh, TArray<FVec3f>&& InReferencePositions, TArray<FWeightedInfluenceData>&& InBoneData, 
+			TArray<FName>&& InUsedBones, FTransform&& InReferenceRootTransform, TArray<FTransform>&& InReferenceRelativeTransforms);
+
+		CHAOS_API FSkinnedTriangleMesh(FSkinnedTriangleMesh&& Other);
+
+		virtual ~FSkinnedTriangleMesh() override = default;
+
+		static constexpr EImplicitObjectType StaticType()
+		{
+			return ImplicitObjectType::SkinnedTriangleMesh;
+		}
+
+		virtual const FAABB3 BoundingBox() const override 
+		{
+			return LocalBoundingBox.IsEmpty() ? FAABB3() : FAABB3(LocalBoundingBox); 
+		}
+		const TArray<FWeightedInfluenceData>& GetBoneData() const 
+		{
+			return BoneData; 
+		}
+		const TArray<FName>& GetUsedBones() const 
+		{
+			return UsedBones;
+		}
+		const FTriangleMesh& GetTriangleMesh() const 
+		{
+			return TriangleMesh;
+		}
+		// Original reference positions
+		const TArray<FVec3f>& GetReferencePositions() const 
+		{
+			return ReferencePositions;
+		}
+		// Current skinned positions
+		const TArray<FVec3f>& GetLocalPositions() const 
+		{
+			return LocalPositions;
+		}
+		// Current Skinned positions for writing
+		TArrayView<FVec3f> GetLocalPositions() 
+		{
+			return TArrayView<FVec3f>(LocalPositions);
+		}
+		const FTriangleMesh::TSpatialHashType<FRealSingle>& GetSpatialHierarchy() const 
+		{
+			return SpatialHash;
+		}
+
+		// Skin positions
+		CHAOS_API void SkinPositions(const TArray<FTransform>& RelativeTransforms, const TArrayView<FVec3f>& Positions) const;
+		void SkinPositions(const TArray<FTransform>& RelativeTransforms)
+		{
+			SkinPositions(RelativeTransforms, GetLocalPositions());
+			UpdateLocalBoundingBox();
+		}
+		void UpdateLocalBoundingBox()
+		{
+			LocalBoundingBox = CalculateBoundingBox(LocalPositions);
+		}
+		CHAOS_API void UpdateSpatialHierarchy(const FRealSingle MinLodSize = 0.f);
+
+		virtual void Serialize(FChaosArchive& Ar) override
+		{
+			SerializeImp(Ar);
+		}
+
+		void SerializeImp(FArchive& Ar)
+		{
+			FImplicitObject::SerializeImp(Ar);
+			TriangleMesh.Serialize(Ar);
+			Ar << BoneData;
+			Ar << UsedBones;
+			Ar << ReferenceRootTransform;
+			Ar << ReferenceRelativeTransforms;
+			Ar << ReferencePositions;
+			TBox<FRealSingle, 3>::SerializeAsAABB(Ar, ReferenceBoundingBox);
+			if (Ar.IsLoading())
+			{
+				FinalizeConstruction();
+			}
+		}
+
+		CHAOS_API virtual FImplicitObjectPtr CopyGeometry() const override;
+		CHAOS_API virtual FImplicitObjectPtr CopyGeometryWithScale(const FVec3& Scale) const override;
+		CHAOS_API virtual uint32 GetTypeHash() const override;
+		virtual FReal PhiWithNormal(const FVec3& x, FVec3& Normal) const override
+		{
+			// Currently no users of PhiWithNormal. Only used for proximity queries.
+			ensure(false);
+			return TNumericLimits<FReal>::Max();
+		}
+
+	private:
+		friend FKSkinnedTriangleMeshElem;
+
+		FSkinnedTriangleMesh()
+			:FImplicitObject(EImplicitObject::HasBoundingBox, ImplicitObjectType::SkinnedTriangleMesh)
+		{}
+
+		CHAOS_API TAABB<FRealSingle,3> CalculateBoundingBox(const TConstArrayView<FVec3f>& Positions) const;
+		CHAOS_API void FinalizeConstruction();
+
+		// Serialized data. Only non-const because of serialization
+		FTriangleMesh TriangleMesh;
+		TArray<FWeightedInfluenceData> BoneData;
+		TArray<FName> UsedBones;
+		FTransform ReferenceRootTransform;
+		TArray<FTransform> ReferenceRelativeTransforms; // ReferenceRootTransform * RefBaseMatrixInv(UsedBoneIdx)
+		TArray<FVec3f> ReferencePositions;
+		TAABB<FRealSingle, 3> ReferenceBoundingBox;
+
+		// Calculated data
+		TAABB<FRealSingle, 3> LocalBoundingBox;
+		TArray<FVec3f> LocalPositions;
+		FTriangleMesh::TSpatialHashType<FRealSingle> SpatialHash;
+	};
+}
\ No newline at end of file
diff --git a/Engine/Source/Runtime/Experimental/Chaos/Public/Chaos/TriangleMesh.h b/Engine/Source/Runtime/Experimental/Chaos/Public/Chaos/TriangleMesh.h
--- a/Engine/Source/Runtime/Experimental/Chaos/Public/Chaos/TriangleMesh.h
+++ b/Engine/Source/Runtime/Experimental/Chaos/Public/Chaos/TriangleMesh.h
@@ -40,6 +40,8 @@ namespace Chaos
 
 		CHAOS_API void ResetAuxiliaryStructures();
 
+		CHAOS_API void Serialize(FArchive& Ar);
+
 		/**
 		 * Returns the closed interval of the smallest vertex index used by 
 		 * this class, to the largest.
@@ -47,6 +49,14 @@ namespace Chaos
 		 * If this mesh is empty, the second index of the range will be negative.
 		 */
 		CHAOS_API TVec2<int32> GetVertexRange() const;
+		int32 GetStartIndex() const 
+		{
+			return MStartIdx;
+		}
+		int32 GetNumIndices() const
+		{
+			return MNumIndices;
+		}
 
 		/** Returns the set of vertices used by triangles. */
 		CHAOS_API TSet<int32> GetVertices() const;
