Defend against odd sitations where assets can make it to the cooker with no source.

#rnx

diff --git a/Engine/Source/Runtime/Engine/Private/TextureDerivedData.cpp b/Engine/Source/Runtime/Engine/Private/TextureDerivedData.cpp
--- a/Engine/Source/Runtime/Engine/Private/TextureDerivedData.cpp
+++ b/Engine/Source/Runtime/Engine/Private/TextureDerivedData.cpp
@@ -4166,7 +4166,7 @@ void UTexture::SerializeCookedPlatformData(FArchive& Ar, const bool bSerializeMi
 				}
 				PlatformDataToSerialize.Add(RunningPlatformData);
 			}
-			else
+			else if (CanBuildPlatformData())
 			{
 				TMap<FString, FTexturePlatformData*> *CookedPlatformDataPtr = GetCookedPlatformData();
 				if (CookedPlatformDataPtr == nullptr)
