IncrementalCook: UStreamableRenderAsset and FPerQualityLevelProperty: workaround the indeterminism in first and second in-memory saves of packages with FPerQualityLevelProperty, for the specific FPerQualityLevelProperty UStreamableRenderAsset.NoRefStreamingLODBias which is breaking in our IterativeValidate test cases.
#rnx

diff --git a/Engine/Source/Runtime/Engine/Classes/Engine/StreamableRenderAsset.h b/Engine/Source/Runtime/Engine/Classes/Engine/StreamableRenderAsset.h
--- a/Engine/Source/Runtime/Engine/Classes/Engine/StreamableRenderAsset.h
+++ b/Engine/Source/Runtime/Engine/Classes/Engine/StreamableRenderAsset.h
@@ -42,6 +42,10 @@ public:
 	/** Destructor */
 	ENGINE_API virtual ~UStreamableRenderAsset();
 
+#if WITH_EDITOR
+	ENGINE_API virtual void PreSave(FObjectPreSaveContext SaveContext);
+#endif
+
 	/** Get an integer representation of the LOD group */
 	virtual int32 GetLODGroupForStreaming() const
 	{
diff --git a/Engine/Source/Runtime/Engine/Private/StreamableRenderAsset.cpp b/Engine/Source/Runtime/Engine/Private/StreamableRenderAsset.cpp
--- a/Engine/Source/Runtime/Engine/Private/StreamableRenderAsset.cpp
+++ b/Engine/Source/Runtime/Engine/Private/StreamableRenderAsset.cpp
@@ -2,9 +2,11 @@
 
 #include "Engine/StreamableRenderAsset.h"
 #include "Containers/Ticker.h"
+#include "Interfaces/ITargetPlatform.h"
 #include "RenderAssetUpdate.h"
 #include "Rendering/NaniteCoarseMeshStreamingManager.h"
 #include "RenderingThread.h"
+#include "UObject/ObjectSaveContext.h"
 #include "UObject/Package.h"
 #include "UObject/UnrealType.h"
 
@@ -40,6 +42,33 @@ UStreamableRenderAsset::UStreamableRenderAsset(const FObjectInitializer& ObjectI
 
 UStreamableRenderAsset::~UStreamableRenderAsset() = default;
 
+#if WITH_EDITOR
+void UStreamableRenderAsset::PreSave(FObjectPreSaveContext SaveContext)
+{
+	Super::PreSave(SaveContext);
+	// Workaround cook indeterminism in FPerQualityLevelProperty. It writes to its data during cooking
+	// save serialization, by calling StripQualityLevelForCooking. But this function can make the struct
+	// equal to its archetype, and delta-serialization after that then skips serializing the property fname.
+	// This makes a second in-memory save of its package different than the first in-memory save. In some of
+	// our processes (-diffonly cooks and -iterativevalidate cooks) we need the first and second in-memory
+	// saves to be identical.
+	// We don't yet have a way to remove the modification of Property in the save path of
+	// operator<<(FArchive& Ar, FPerQualityLevelProperty<...>& Property).
+	// Other objects within the savepackage call can depend on the modified values, which is surprising since
+	// the values modified by StripQualityLevelForCooking do not change the value read for the targetplatform.
+	// We need to find those dependencies and remove them.
+	// In the meantime, we workaround the problem by triggering the call to StripQualityLevelForCooking before
+	// the first operator<< save path, here in PreSave. We can not do this for all occurrences of the
+	// struct unfortunately because Unreal's savepackage system does not support a PreSave call for structs.
+	// So we currently do it only for the one class using it that is currently causing a problem in our engine tests.
+	if (SaveContext.IsCooking())
+	{
+		const FDataDrivenPlatformInfo& PlatformInfo = SaveContext.GetTargetPlatform()->GetPlatformInfo();
+		NoRefStreamingLODBias.StripQualityLevelForCooking(*(PlatformInfo.IniPlatformName.ToString()));
+	}
+}
+#endif
+
 void UStreamableRenderAsset::RegisterMipLevelChangeCallback(UPrimitiveComponent* Component, int32 LODIndex, float TimeoutSecs, bool bOnStreamIn, FLODStreamingCallback&& Callback)
 {
 	check(IsInGameThread());
