[Renderer]
* Added bool to turn on/off statestream path for static meshes.. note this is hidden behind WITH_STATE_STREAM_ACTOR so will not affect anyhone not explictly having that set to 1

diff --git a/Engine/Source/Runtime/Engine/Private/Components/StaticMeshComponent.cpp b/Engine/Source/Runtime/Engine/Private/Components/StaticMeshComponent.cpp
--- a/Engine/Source/Runtime/Engine/Private/Components/StaticMeshComponent.cpp
+++ b/Engine/Source/Runtime/Engine/Private/Components/StaticMeshComponent.cpp
@@ -57,6 +57,10 @@
 
 #define LOCTEXT_NAMESPACE "StaticMeshComponent"
 
+#if WITH_STATE_STREAM_ACTOR
+constexpr bool UseStaticMeshStateStream = true;
+#endif
+
 DECLARE_MEMORY_STAT( TEXT( "StaticMesh VxColor Inst Mem" ), STAT_InstVertexColorMemory, STATGROUP_MemoryStaticMesh );
 DECLARE_MEMORY_STAT( TEXT( "StaticMesh PreCulled Index Memory" ), STAT_StaticMeshPreCulledIndexMemory, STATGROUP_MemoryStaticMesh );
 
@@ -899,6 +903,7 @@ void UStaticMeshComponent::CreateRenderState_Concurrent(FRegisterComponentContex
 	LLM_SCOPE(ELLMTag::StaticMesh);
 
 #if WITH_STATE_STREAM_ACTOR
+	if (UseStaticMeshStateStream)
 	{
 		FStaticMeshStaticState Ss;
 		FStaticMeshDynamicState Ds;
@@ -919,6 +924,7 @@ void UStaticMeshComponent::CreateRenderState_Concurrent(FRegisterComponentContex
 		}
 		MeshHandle = GetWorld()->GetStateStream<IStaticMeshStateStream>().Game_CreateInstance(Ss, Ds);
 		bRenderStateCreated = true;
+		return;
 	}
 #else
 	Super::CreateRenderState_Concurrent(Context);
@@ -2326,9 +2332,12 @@ void UStaticMeshComponent::SetStaticMeshInternal(UStaticMesh* NewMesh)
 	StaticMesh = NewMesh;
 
 #if WITH_STATE_STREAM_ACTOR
-	FStaticMeshDynamicState Ds;
-	Ds.SetMesh(StaticMesh);
-	MeshHandle.Update(Ds);
+	if (UseStaticMeshStateStream)
+	{
+		FStaticMeshDynamicState Ds;
+		Ds.SetMesh(StaticMesh);
+		MeshHandle.Update(Ds);
+	}
 #endif
 
 	NotifyIfStaticMeshChanged();
