Fixed broken rendering when energy conservation is enabled but path tracer is disabled.

#tests PC dx12 editor as mentioned in the JIRA
#rnx

diff --git a/Engine/Source/Runtime/Engine/Private/UnrealEngine.cpp b/Engine/Source/Runtime/Engine/Private/UnrealEngine.cpp
--- a/Engine/Source/Runtime/Engine/Private/UnrealEngine.cpp
+++ b/Engine/Source/Runtime/Engine/Private/UnrealEngine.cpp
@@ -3516,7 +3516,13 @@ void UEngine::LoadEnergyTextures()
 	static const auto CVarSubstrate = IConsoleManager::Get().FindTConsoleVariableDataInt(TEXT("r.Substrate"));
 	const bool bSubstrateEnabled = CVarSubstrate && CVarSubstrate->GetValueOnAnyThread() > 0;
 
-	if ((bPathTracingEnabled || bSubstrateEnabled) && GGXReflectionEnergyTexture == nullptr && GGXReflectionEnergyTextureName.IsValid())
+	static const auto CVarEnergyConservation = IConsoleManager::Get().FindTConsoleVariableDataInt(TEXT("r.Material.EnergyConservation"));
+	const bool bEnergyConservationEnabled = CVarEnergyConservation && CVarEnergyConservation->GetValueOnAnyThread() > 0;
+
+	// NOTE: this logic must match GetSettings from ShadingEnergyConservation.cpp
+	const bool bLoadEnergyConservationTextures = bPathTracingEnabled || bSubstrateEnabled || bEnergyConservationEnabled;
+
+	if (bLoadEnergyConservationTextures && GGXReflectionEnergyTexture == nullptr && GGXReflectionEnergyTextureName.IsValid())
 	{
 		LoadEngineTexture(GGXReflectionEnergyTexture, *GGXReflectionEnergyTextureName.ToString());
 	}
diff --git a/Engine/Source/Runtime/Renderer/Private/ShadingEnergyConservation.cpp b/Engine/Source/Runtime/Renderer/Private/ShadingEnergyConservation.cpp
--- a/Engine/Source/Runtime/Renderer/Private/ShadingEnergyConservation.cpp
+++ b/Engine/Source/Runtime/Renderer/Private/ShadingEnergyConservation.cpp
@@ -1,7 +1,7 @@
 // Copyright Epic Games, Inc. All Rights Reserved.
 
 /*=============================================================================
-	DeferredShadingRenderer.cpp: Top level rendering loop for deferred shading
+	ShadingEnergyConservation.cpp: private energy conservation implementation
 =============================================================================*/
 
 #include "ShadingEnergyConservation.h"
@@ -384,11 +384,6 @@ IMPLEMENT_SHADER_TYPE(, FBuildShadingEnergyConservationTableCS, TEXT("/Engine/Pr
 namespace ShadingEnergyConservation
 {
 
-bool IsEnable()
-{
-	return CVarShadingEnergyConservation.GetValueOnAnyThread() > 0;
-}
-
 void Init(FRDGBuilder& GraphBuilder, const FViewInfo& View)
 {
 	if (GetSettings(View).bNeedData)
diff --git a/Engine/Source/Runtime/Renderer/Private/ShadingEnergyConservation.h b/Engine/Source/Runtime/Renderer/Private/ShadingEnergyConservation.h
--- a/Engine/Source/Runtime/Renderer/Private/ShadingEnergyConservation.h
+++ b/Engine/Source/Runtime/Renderer/Private/ShadingEnergyConservation.h
@@ -1,7 +1,7 @@
 // Copyright Epic Games, Inc. All Rights Reserved.
 
 /*=============================================================================
-	DeferredShadingRenderer.h: Scene rendering definitions.
+	ShadingEnergyConservation.h: private energy conservation related data
 =============================================================================*/
 
 #pragma once
@@ -27,6 +27,4 @@ namespace ShadingEnergyConservation
 	void Init(FRDGBuilder& GraphBuilder, const FViewInfo& View);
 	void Debug(FRDGBuilder& GraphBuilder, const FViewInfo& View, FSceneTextures& SceneTextures);
 	FShadingEnergyConservationData GetData(const FViewInfo& View);
-
-	RENDERER_API bool IsEnable();
 }
