SkyLight - fixed sky light capture to work after toggling global illumination from medium to high when there is a landscape.

#tests PC DX12 repro on JIRA
#rnx

diff --git a/Engine/Source/Runtime/Engine/Private/Components/SkyLightComponent.cpp b/Engine/Source/Runtime/Engine/Private/Components/SkyLightComponent.cpp
--- a/Engine/Source/Runtime/Engine/Private/Components/SkyLightComponent.cpp
+++ b/Engine/Source/Runtime/Engine/Private/Components/SkyLightComponent.cpp
@@ -25,6 +25,7 @@
 #include "UObject/ReleaseObjectVersion.h"
 #include "UObject/UE5MainStreamObjectVersion.h"
 #include "Engine/TextureCube.h"
+#include "ComponentRecreateRenderStateContext.h"
 
 #if RHI_RAYTRACING
 #endif
@@ -107,6 +108,10 @@ void OnChangeSkylightRealTimeReflectionCapture(IConsoleVariable* Var)
 			SkylightComponent->SetCaptureIsDirty();
 		}
 	}
+
+	// We need to recreate the render state of everything after the skylight.
+	// This is required for correct update due to recreation order when setting global illumination back from low/medium to high or more.
+	FGlobalComponentRecreateRenderStateContext Context;
 }
 
 int32 GSkylightRealTimeReflectionCapture = 1;
