Fix for race condition leading to rare crash in cooking of body setup derived data

A call to FFormatContainer::GetFormat will potentially allocate a new bulk data and push it into the format map, which can reallocate the underlying storage for the map.
This isn't a threadsafe operation and cooking physics data is performed on multiple threads. Rarely two threads will get into the cook at the same time. This change adds
 a double checked lock condition to make sure only one thread can build the physics data, any other threads will wait on the build before proceeding

diff --git a/Engine/Source/Runtime/Engine/Classes/PhysicsEngine/BodySetup.h b/Engine/Source/Runtime/Engine/Classes/PhysicsEngine/BodySetup.h
--- a/Engine/Source/Runtime/Engine/Classes/PhysicsEngine/BodySetup.h
+++ b/Engine/Source/Runtime/Engine/Classes/PhysicsEngine/BodySetup.h
@@ -220,6 +220,15 @@ class UBodySetup : public UBodySetupCore
 	/** Cooked physics data for each format */
 	FFormatContainer CookedFormatData;
 
+#if WITH_EDITORONLY_DATA
+	/** 
+	 * The format container may be queried from multiple threads and mutated if a thread fails to find a format.
+	 * This mutex protects the cook path in TryGetCookedData so we don't have multiple threads attempting to
+	 * write to the format data
+	 */
+	FCriticalSection CookedFormatDataMutex;
+#endif
+
 	/** GUID used to uniquely identify this setup so it can be found in the DDC */
 	FGuid BodySetupGuid;
 	
diff --git a/Engine/Source/Runtime/Engine/Private/PhysicsEngine/BodySetup.cpp b/Engine/Source/Runtime/Engine/Private/PhysicsEngine/BodySetup.cpp
--- a/Engine/Source/Runtime/Engine/Private/PhysicsEngine/BodySetup.cpp
+++ b/Engine/Source/Runtime/Engine/Private/PhysicsEngine/BodySetup.cpp
@@ -1439,30 +1439,44 @@ UBodySetup::FBodySetupTryCookResult UBodySetup::TryGetCookedData(FName Format)
 	FFormatContainer* UseCookedData = CookedFormatDataOverride ? CookedFormatDataOverride : &CookedFormatData;
 #endif
 
-	bool bContainedData = UseCookedData->Contains(Format);
 	FByteBulkData* Result = nullptr;
-	bool bIsRuntime = IsRuntime(this);
+	const bool bIsRuntime = IsRuntime(this);
 
 #if WITH_EDITOR
-	if(!bContainedData)
+	if(!UseCookedData->Contains(Format))
 	{
 		SCOPE_CYCLE_COUNTER(STAT_PhysXCooking);
 
-		// Note: Check ContainsPhysicsTriMeshData before looking at the number of convex elems, to ensure the side effects of ContainsPhysicsTriMeshData happen
-		// (specifically, for static mesh this will ensure the mesh render data is already built)
-		if((CDP == nullptr || CDP->ContainsPhysicsTriMeshData(bMeshCollideAll) == false) && AggGeom.ConvexElems.Num() == 0)
+		// Only one thread should build the physics data. Below calling GetFormat and TryGetDDCBuiltData will
+		// both mutate shared state so we protect this operation with a lock.
+		FScopeLock ScopeLock(&CookedFormatDataMutex);
+
+		// Check the lock condition again, we only want the first thread in here to actually build the data,
+		// while it's building other threads should wait and proceed after the build completes.
+		// Double-checked locking avoids always taking the lock - we only do it for unbuilt data
+		if(!UseCookedData->Contains(Format))
 		{
-			return { nullptr, CookDataResult };
-		}
+			// Note: Check ContainsPhysicsTriMeshData before looking at the number of convex elems, to ensure the side effects of ContainsPhysicsTriMeshData happen
+			// (specifically, for static mesh this will ensure the mesh render data is already built)
+			if ((CDP == nullptr || CDP->ContainsPhysicsTriMeshData(bMeshCollideAll) == false) && AggGeom.ConvexElems.Num() == 0)
+			{
+				return { nullptr, CookDataResult };
+			}
 
-		// We do not want a FGCObject to be created to prevent garbage collection of our own UBodySetup*
-		// because the scope and lifetime is well defined and FGCObject can't be created on other threads
-		// during garbage collection, which would prevent this function from being run asynchronously.
-		const bool bUseRefHolder = false;
-		FChaosDerivedDataCooker* PhysicsDerivedCooker = new FChaosDerivedDataCooker(this, Format, bUseRefHolder);
+			// We do not want a FGCObject to be created to prevent garbage collection of our own UBodySetup*
+			// because the scope and lifetime is well defined and FGCObject can't be created on other threads
+			// during garbage collection, which would prevent this function from being run asynchronously.
+			const bool bUseRefHolder = false;
+			FChaosDerivedDataCooker* PhysicsDerivedCooker = new FChaosDerivedDataCooker(this, Format, bUseRefHolder);
 
-		Result = &UseCookedData->GetFormat(Format);
-		CookDataResult = TryGetDDCBuiltData(Result, *PhysicsDerivedCooker, this, bIsRuntime);
+			Result = &UseCookedData->GetFormat(Format);
+			CookDataResult = TryGetDDCBuiltData(Result, *PhysicsDerivedCooker, this, bIsRuntime);
+		}
+		else
+		{
+			// We were a waiting thread while another thread built the data, just take it from the container now
+			Result = &UseCookedData->GetFormat(Format);
+		}
 	}
 	else
 #endif // #if WITH_EDITOR
