Substrate - do not use indices to be able to reorder input if we want later.

#rnx
#tests PC dx12 editor

diff --git a/Engine/Source/Runtime/Engine/Private/Materials/Material.cpp b/Engine/Source/Runtime/Engine/Private/Materials/Material.cpp
--- a/Engine/Source/Runtime/Engine/Private/Materials/Material.cpp
+++ b/Engine/Source/Runtime/Engine/Private/Materials/Material.cpp
@@ -3375,9 +3375,8 @@ bool UMaterial::ConvertMaterialToSubstrateMaterial(bool bAllowEmptyMaterialUpdat
 		CurrentNodePosX = FMath::Max(NodeToReplace->MaterialExpressionEditorX + TranslationOffsetX, CurrentNodePosX);
 	};
 
-
-	// for ExpressionInput
-	auto ConnectionTo = [](auto& OldNodeInput, UMaterialExpression* NewNode, uint32 NewInputIndex, uint32 OperationType = SUBSTRATE_MOVE_CONNECTION)
+	// Recommendation: use the non indexed version whenever possible. However, indexed version are mandatory for material function which do not have explicit input.
+	auto ConnectionToIndexed = [](auto& OldNodeInput, UMaterialExpression* NewNode, uint32 NewInputIndex, uint32 OperationType = SUBSTRATE_MOVE_CONNECTION)
 	{
 		if (OldNodeInput.IsConnected())
 		{
@@ -3388,15 +3387,24 @@ bool UMaterial::ConvertMaterialToSubstrateMaterial(bool bAllowEmptyMaterialUpdat
 			}
 		}
 	};
-
-	// For material input
-	auto ScalarMatInputConnectionTo = [&](FScalarMaterialInput& OldNodeInput, UMaterialExpression* NewNode, uint32 NewInputIndex, EMaterialProperty Property, uint32 OperationType = SUBSTRATE_MOVE_CONNECTION)
+	auto ConnectionTo = [](auto& OldNodeInput, FExpressionInput& NewNodeInput, uint32 OperationType = SUBSTRATE_MOVE_CONNECTION)
+	{
+		if (OldNodeInput.IsConnected())
+		{
+			NewNodeInput.Connect(OldNodeInput.OutputIndex, OldNodeInput.Expression);
+			if (OperationType == SUBSTRATE_MOVE_CONNECTION)
+			{
+				OldNodeInput.Expression = nullptr;
+			}
+		}
+	};
+	auto ScalarMatInputConnectionToIndexed = [&](FScalarMaterialInput& OldNodeInput, UMaterialExpression* NewNode, uint32 NewInputIndex, EMaterialProperty Property, uint32 OperationType = SUBSTRATE_MOVE_CONNECTION)
 	{
 		FVector4f DefaultPinValue = FMaterialAttributeDefinitionMap::GetDefaultValue(Property);
 
 		if (OldNodeInput.IsConnected())
 		{
-			ConnectionTo(OldNodeInput, NewNode, NewInputIndex, OperationType);
+			ConnectionToIndexed(OldNodeInput, NewNode, NewInputIndex, OperationType);
 		}
 		else if (OldNodeInput.UseConstant && DefaultPinValue.X != OldNodeInput.Constant)
 		{
@@ -3409,7 +3417,26 @@ bool UMaterial::ConvertMaterialToSubstrateMaterial(bool bAllowEmptyMaterialUpdat
 			}
 		}
 	};
-	auto ColorMatInputConnectionTo = [&](FColorMaterialInput& OldNodeInput, UMaterialExpression* NewNode, uint32 NewInputIndex, EMaterialProperty Property, uint32 OperationType = SUBSTRATE_MOVE_CONNECTION)
+	auto ScalarMatInputConnectionTo = [&](FScalarMaterialInput& OldNodeInput, FExpressionInput& NewNodeInput, EMaterialProperty Property, uint32 OperationType = SUBSTRATE_MOVE_CONNECTION)
+	{
+		FVector4f DefaultPinValue = FMaterialAttributeDefinitionMap::GetDefaultValue(Property);
+
+		if (OldNodeInput.IsConnected())
+		{
+			ConnectionTo(OldNodeInput, NewNodeInput, OperationType);
+		}
+		else if (OldNodeInput.UseConstant && DefaultPinValue.X != OldNodeInput.Constant)
+		{
+			UMaterialExpressionConstant* ExpressionConstantScalar = NewObject<UMaterialExpressionConstant>(this);
+			ExpressionConstantScalar->R = OldNodeInput.Constant;
+			NewNodeInput.Connect(0, ExpressionConstantScalar);
+			if (OperationType == SUBSTRATE_MOVE_CONNECTION)
+			{
+				OldNodeInput.Expression = nullptr;
+			}
+		}
+	};
+	auto ColorMatInputConnectionToIndexed = [&](FColorMaterialInput& OldNodeInput, UMaterialExpression* NewNode, uint32 NewInputIndex, EMaterialProperty Property, uint32 OperationType = SUBSTRATE_MOVE_CONNECTION)
 	{
 		FVector4f DefaultPinValue = FMaterialAttributeDefinitionMap::GetDefaultValue(Property);
 		FLinearColor DefaultPinLinearColor = FLinearColor(DefaultPinValue.X, DefaultPinValue.Y, DefaultPinValue.Z);
@@ -3417,7 +3444,7 @@ bool UMaterial::ConvertMaterialToSubstrateMaterial(bool bAllowEmptyMaterialUpdat
 
 		if (OldNodeInput.IsConnected())
 		{
-			ConnectionTo(OldNodeInput, NewNode, NewInputIndex, OperationType);
+			ConnectionToIndexed(OldNodeInput, NewNode, NewInputIndex, OperationType);
 		}
 		else if (OldNodeInput.UseConstant && DefaultPinLinearColor != CurrentPinLinearColor)
 		{
@@ -3430,7 +3457,28 @@ bool UMaterial::ConvertMaterialToSubstrateMaterial(bool bAllowEmptyMaterialUpdat
 			}
 		}
 	};
-	auto Vector3MatInputConnectionTo = [&](FVectorMaterialInput& OldNodeInput, UMaterialExpression* NewNode, uint32 NewInputIndex, EMaterialProperty Property, uint32 OperationType = SUBSTRATE_MOVE_CONNECTION)
+	auto ColorMatInputConnectionTo = [&](FColorMaterialInput& OldNodeInput, FExpressionInput& NewNodeInput, EMaterialProperty Property, uint32 OperationType = SUBSTRATE_MOVE_CONNECTION)
+	{
+		FVector4f DefaultPinValue = FMaterialAttributeDefinitionMap::GetDefaultValue(Property);
+		FLinearColor DefaultPinLinearColor = FLinearColor(DefaultPinValue.X, DefaultPinValue.Y, DefaultPinValue.Z);
+		FLinearColor CurrentPinLinearColor = OldNodeInput.Constant;
+
+		if (OldNodeInput.IsConnected())
+		{
+			ConnectionTo(OldNodeInput, NewNodeInput, OperationType);
+		}
+		else if (OldNodeInput.UseConstant && DefaultPinLinearColor != CurrentPinLinearColor)
+		{
+			UMaterialExpressionConstant3Vector* ExpressionConstantScalar = NewObject<UMaterialExpressionConstant3Vector>(this);
+			ExpressionConstantScalar->Constant = CurrentPinLinearColor;
+			NewNodeInput.Connect(0, ExpressionConstantScalar);
+			if (OperationType == SUBSTRATE_MOVE_CONNECTION)
+			{
+				OldNodeInput.Expression = nullptr;
+			}
+		}
+	};
+	auto Vector3MatInputConnectionToIndexed = [&](FVectorMaterialInput& OldNodeInput, UMaterialExpression* NewNode, uint32 NewInputIndex, EMaterialProperty Property, uint32 OperationType = SUBSTRATE_MOVE_CONNECTION)
 	{
 		FVector4f DefaultPinValue = FMaterialAttributeDefinitionMap::GetDefaultValue(Property);
 		FVector3f DefaultPinVec3f = FVector3f(DefaultPinValue.X, DefaultPinValue.Y, DefaultPinValue.Z);
@@ -3438,7 +3486,7 @@ bool UMaterial::ConvertMaterialToSubstrateMaterial(bool bAllowEmptyMaterialUpdat
 
 		if (OldNodeInput.IsConnected())
 		{
-			ConnectionTo(OldNodeInput, NewNode, NewInputIndex, OperationType);
+			ConnectionToIndexed(OldNodeInput, NewNode, NewInputIndex, OperationType);
 		}
 		else if (OldNodeInput.UseConstant && DefaultPinVec3f != CurrentPinVec3f)
 		{
@@ -3451,6 +3499,27 @@ bool UMaterial::ConvertMaterialToSubstrateMaterial(bool bAllowEmptyMaterialUpdat
 			}
 		}
 	};
+	auto Vector3MatInputConnectionTo = [&](FVectorMaterialInput& OldNodeInput, FExpressionInput& NewNodeInput, EMaterialProperty Property, uint32 OperationType = SUBSTRATE_MOVE_CONNECTION)
+	{
+		FVector4f DefaultPinValue = FMaterialAttributeDefinitionMap::GetDefaultValue(Property);
+		FVector3f DefaultPinVec3f = FVector3f(DefaultPinValue.X, DefaultPinValue.Y, DefaultPinValue.Z);
+		FVector3f CurrentPinVec3f = FVector3f(OldNodeInput.Constant.X, OldNodeInput.Constant.Y, OldNodeInput.Constant.Z);
+
+		if (OldNodeInput.IsConnected())
+		{
+			ConnectionTo(OldNodeInput, NewNodeInput, OperationType);
+		}
+		else if (OldNodeInput.UseConstant && DefaultPinVec3f != CurrentPinVec3f)
+		{
+			UMaterialExpressionConstant3Vector* ExpressionConstantScalar = NewObject<UMaterialExpressionConstant3Vector>(this);
+			ExpressionConstantScalar->Constant = CurrentPinVec3f;
+			NewNodeInput.Connect(0, ExpressionConstantScalar);
+			if (OperationType == SUBSTRATE_MOVE_CONNECTION)
+			{
+				OldNodeInput.Expression = nullptr;
+			}
+		}
+	};
 
 	bool bCustomNodesGathered = false;
 	UMaterialExpressionThinTranslucentMaterialOutput* ThinTranslucentOutput = nullptr;
@@ -3582,7 +3651,7 @@ bool UMaterial::ConvertMaterialToSubstrateMaterial(bool bAllowEmptyMaterialUpdat
 			// * Leave the material attribute existing connection plugged to the root node, 
 			//   so that other input (PixelDepthOffset, WorldPositionOffset, ...) get pull 
 			//   from the material attributes node
-			ConnectionTo(EditorOnly->MaterialAttributes, ConvertAttributeNode, 0, SUBSTRATE_COPY_CONNECTION);
+			ConnectionTo(EditorOnly->MaterialAttributes, ConvertAttributeNode->MaterialAttributes, SUBSTRATE_COPY_CONNECTION);
 
 			// Reconnect custom output to material attribute conversion node
 			{
@@ -3591,10 +3660,10 @@ bool UMaterial::ConvertMaterialToSubstrateMaterial(bool bAllowEmptyMaterialUpdat
 
 				if (SingleLayerWaterOutput)
 				{
-					ConnectionTo(*SingleLayerWaterOutput->GetInput(0), ConvertAttributeNode, 1); // WaterScatteringCoefficients
-					ConnectionTo(*SingleLayerWaterOutput->GetInput(1), ConvertAttributeNode, 2); // WaterAbsorptionCoefficients
-					ConnectionTo(*SingleLayerWaterOutput->GetInput(2), ConvertAttributeNode, 3); // WaterPhaseG
-					ConnectionTo(*SingleLayerWaterOutput->GetInput(3), ConvertAttributeNode, 4); // ColorScaleBehindWater
+					ConnectionTo(*SingleLayerWaterOutput->GetInput(0), ConvertAttributeNode->WaterScatteringCoefficients);
+					ConnectionTo(*SingleLayerWaterOutput->GetInput(1), ConvertAttributeNode->WaterAbsorptionCoefficients);
+					ConnectionTo(*SingleLayerWaterOutput->GetInput(2), ConvertAttributeNode->WaterPhaseG);
+					ConnectionTo(*SingleLayerWaterOutput->GetInput(3), ConvertAttributeNode->ColorScaleBehindWater);
 				}
 			}
 
@@ -3678,13 +3747,13 @@ bool UMaterial::ConvertMaterialToSubstrateMaterial(bool bAllowEmptyMaterialUpdat
 					MFCallNode->PostEditChange();
 					MFCallNode->ConditionalPostLoad();
 
-					ColorMatInputConnectionTo(EditorOnly->BaseColor,		MFCallNode, 0, MP_BaseColor);
-					ScalarMatInputConnectionTo(EditorOnly->Metallic,		MFCallNode, 1, MP_Metallic);
-					ScalarMatInputConnectionTo(EditorOnly->Specular,		MFCallNode, 2, MP_Specular);
-					ScalarMatInputConnectionTo(EditorOnly->Roughness,		MFCallNode, 3, MP_Roughness);
-					Vector3MatInputConnectionTo(EditorOnly->Normal,			MFCallNode, 4, MP_Normal);
-					ColorMatInputConnectionTo(EditorOnly->EmissiveColor,	MFCallNode, 5, MP_EmissiveColor);
-					ScalarMatInputConnectionTo(EditorOnly->Opacity,			MFCallNode, 6, MP_Opacity);
+					ColorMatInputConnectionToIndexed(EditorOnly->BaseColor,		MFCallNode, 0, MP_BaseColor);
+					ScalarMatInputConnectionToIndexed(EditorOnly->Metallic,		MFCallNode, 1, MP_Metallic);
+					ScalarMatInputConnectionToIndexed(EditorOnly->Specular,		MFCallNode, 2, MP_Specular);
+					ScalarMatInputConnectionToIndexed(EditorOnly->Roughness,	MFCallNode, 3, MP_Roughness);
+					Vector3MatInputConnectionToIndexed(EditorOnly->Normal,		MFCallNode, 4, MP_Normal);
+					ColorMatInputConnectionToIndexed(EditorOnly->EmissiveColor,	MFCallNode, 5, MP_EmissiveColor);
+					ScalarMatInputConnectionToIndexed(EditorOnly->Opacity,		MFCallNode, 6, MP_Opacity);
 				}
 			}
 			else
@@ -3716,9 +3785,9 @@ bool UMaterial::ConvertMaterialToSubstrateMaterial(bool bAllowEmptyMaterialUpdat
 						UMaterialExpressionSubstrateMetalnessToDiffuseAlbedoF0* SubstrateMetalnessToDiffuseAlbedoF0 = NewObject<UMaterialExpressionSubstrateMetalnessToDiffuseAlbedoF0>(this);
 						EditorOnly->ExpressionCollection.AddExpression(SubstrateMetalnessToDiffuseAlbedoF0);
 						SetPosXAndMoveReferenceToTheRight(SubstrateMetalnessToDiffuseAlbedoF0);
-						ColorMatInputConnectionTo(EditorOnly->BaseColor, SubstrateMetalnessToDiffuseAlbedoF0, 0, MP_BaseColor);
-						ScalarMatInputConnectionTo(EditorOnly->Metallic, SubstrateMetalnessToDiffuseAlbedoF0, 1, MP_Metallic);
-						ScalarMatInputConnectionTo(EditorOnly->Specular, SubstrateMetalnessToDiffuseAlbedoF0, 2, MP_Specular);
+						ColorMatInputConnectionTo(EditorOnly->BaseColor, SubstrateMetalnessToDiffuseAlbedoF0->BaseColor, MP_BaseColor);
+						ScalarMatInputConnectionTo(EditorOnly->Metallic, SubstrateMetalnessToDiffuseAlbedoF0->Metallic, MP_Metallic);
+						ScalarMatInputConnectionTo(EditorOnly->Specular, SubstrateMetalnessToDiffuseAlbedoF0->Specular, MP_Specular);
 					
 						// Top slab BSDF as a simple Disney material
 						UMaterialExpressionSubstrateSlabBSDF* BottomSlabBSDF = NewObject<UMaterialExpressionSubstrateSlabBSDF>(this);
@@ -3728,12 +3797,12 @@ bool UMaterial::ConvertMaterialToSubstrateMaterial(bool bAllowEmptyMaterialUpdat
 						BottomSlabBSDF->GetInput(0)->Connect(0, SubstrateMetalnessToDiffuseAlbedoF0);
 						BottomSlabBSDF->GetInput(1)->Connect(1, SubstrateMetalnessToDiffuseAlbedoF0);
 						BottomSlabBSDF->GetInput(2)->Connect(2, SubstrateMetalnessToDiffuseAlbedoF0);
-						ScalarMatInputConnectionTo(EditorOnly->Roughness, BottomSlabBSDF, 3, MP_Roughness);
-						ScalarMatInputConnectionTo(EditorOnly->Anisotropy, BottomSlabBSDF, 4, MP_Anisotropy, SUBSTRATE_COPY_CONNECTION);
-						Vector3MatInputConnectionTo(EditorOnly->Tangent, BottomSlabBSDF, 6, MP_Tangent);
+						ScalarMatInputConnectionTo(EditorOnly->Roughness, BottomSlabBSDF->Roughness, MP_Roughness);
+						ScalarMatInputConnectionTo(EditorOnly->Anisotropy, BottomSlabBSDF->Anisotropy, MP_Anisotropy, SUBSTRATE_COPY_CONNECTION);
+						Vector3MatInputConnectionTo(EditorOnly->Tangent, BottomSlabBSDF->Tangent, MP_Tangent);
 
 						check(ClearCoatBottomNormalOutput);
-						ConnectionTo(*ClearCoatBottomNormalOutput->GetInput(0), BottomSlabBSDF, 5, SUBSTRATE_COPY_CONNECTION);// ClearColorBottomNormal -> BottomSlabBSDF.Normal
+						ConnectionTo(*ClearCoatBottomNormalOutput->GetInput(0), BottomSlabBSDF->Normal, SUBSTRATE_COPY_CONNECTION);
 
 						// Now weight the top base material by opacity.
 						UMaterialExpressionSubstrateSlabBSDF* TopSlabBSDF = NewObject<UMaterialExpressionSubstrateSlabBSDF>(this);
@@ -3741,9 +3810,9 @@ bool UMaterial::ConvertMaterialToSubstrateMaterial(bool bAllowEmptyMaterialUpdat
 						TopSlabBSDF->Material = this;
 						TopSlabBSDF->MaterialExpressionEditorX = BottomSlabBSDF->MaterialExpressionEditorX;
 						TopSlabBSDF->MaterialExpressionEditorY = BottomSlabBSDF->MaterialExpressionEditorY + 650;
-						ColorMatInputConnectionTo(EditorOnly->EmissiveColor, TopSlabBSDF, 10, MP_EmissiveColor);
-						ScalarMatInputConnectionTo(EditorOnly->ClearCoatRoughness, TopSlabBSDF, 3, MP_CustomData0);	// ClearCoatRoughness => Roughness
-						Vector3MatInputConnectionTo(EditorOnly->Normal, TopSlabBSDF, 5, MP_Normal);
+						ColorMatInputConnectionTo(EditorOnly->EmissiveColor, TopSlabBSDF->EmissiveColor, MP_EmissiveColor);
+						ScalarMatInputConnectionTo(EditorOnly->ClearCoatRoughness, TopSlabBSDF->Roughness, MP_CustomData0);
+						Vector3MatInputConnectionTo(EditorOnly->Normal, TopSlabBSDF->Normal, MP_Normal);
 
 						//  The top layer has a hard coded specular value of 0.5 (F0 = 0.04)
 						UMaterialExpressionConstant* ConstantHalf = NewObject<UMaterialExpressionConstant>(this);
@@ -3766,20 +3835,19 @@ bool UMaterial::ConvertMaterialToSubstrateMaterial(bool bAllowEmptyMaterialUpdat
 						EditorOnly->ExpressionCollection.AddExpression(Constant075);
 						ReplaceNodeAndMoveToTheRight(TopSlabBSDF, Constant075);
 						Constant075->R = 0.75f;
-						UMaterialExpressionSubstrateTransmittanceToMFP* TransToMDFP = NewObject<UMaterialExpressionSubstrateTransmittanceToMFP>(this);
-						EditorOnly->ExpressionCollection.AddExpression(TransToMDFP);
-						ReplaceNodeAndMoveToTheRight(TopSlabBSDF, TransToMDFP);
-						TransToMDFP->GetInput(0)->Connect(0, Constant075);
-						TopSlabBSDF->GetInput(7)->Connect(0, TransToMDFP);							// MFP -> MFP
-						TopSlabBSDF->GetInput(13)->Connect(1, TransToMDFP);							// Thickness -> Thickness
+						UMaterialExpressionSubstrateTransmittanceToMFP* TransToMFP = NewObject<UMaterialExpressionSubstrateTransmittanceToMFP>(this);
+						EditorOnly->ExpressionCollection.AddExpression(TransToMFP);
+						ReplaceNodeAndMoveToTheRight(TopSlabBSDF, TransToMFP);
+						TransToMFP->TransmittanceColor.Connect(0, Constant075);
+						TopSlabBSDF->SSSMFP.Connect(0, TransToMFP); // Connect the TransmittanceToMFP result to the Top slab MFP
 
 						// Now weight the top base material by ClearCoat
 						UMaterialExpressionSubstrateWeight* TopSlabBSDFWithCoverage = NewObject<UMaterialExpressionSubstrateWeight>(this);
 						EditorOnly->ExpressionCollection.AddExpression(SubstrateMetalnessToDiffuseAlbedoF0); // Add the SetAttributesNode to ensure the MaterialAttribute input are correctly tracked & cached
 						SetPosXAndMoveReferenceToTheRight(TopSlabBSDFWithCoverage);
 						TopSlabBSDFWithCoverage->GetInput(0)->Connect(0, TopSlabBSDF);												// TopSlabBSDF -> A
-						ScalarMatInputConnectionTo(EditorOnly->ClearCoat, TopSlabBSDFWithCoverage, 1, MP_CustomData0);				// ClearCoat -> Weight
-						ScalarMatInputConnectionTo(EditorOnly->ClearCoatRoughness, TopSlabBSDFWithCoverage, 1, MP_CustomData1);		// ClearCoat -> Weight
+						ScalarMatInputConnectionTo(EditorOnly->ClearCoat, TopSlabBSDFWithCoverage->Weight, MP_CustomData0);
+						ScalarMatInputConnectionTo(EditorOnly->ClearCoatRoughness, TopSlabBSDF->Roughness, MP_CustomData1);
 
 						UMaterialExpressionSubstrateVerticalLayering* VerticalLayering = NewObject<UMaterialExpressionSubstrateVerticalLayering>(this);
 						EditorOnly->ExpressionCollection.AddExpression(VerticalLayering);
@@ -3800,18 +3868,18 @@ bool UMaterial::ConvertMaterialToSubstrateMaterial(bool bAllowEmptyMaterialUpdat
 					ConvertNode->Material = this;
 					SetPosXAndMoveReferenceToTheRight(ConvertNode);
 					ConvertNode->SubsurfaceProfile = bRequireNoSubsurfaceProfile ? nullptr : SubsurfaceProfile;
-					ColorMatInputConnectionTo(EditorOnly->BaseColor, ConvertNode, 0, MP_BaseColor);
-					ScalarMatInputConnectionTo(EditorOnly->Metallic, ConvertNode, 1, MP_Metallic);
-					ScalarMatInputConnectionTo(EditorOnly->Specular, ConvertNode, 2, MP_Specular);
-					ScalarMatInputConnectionTo(EditorOnly->Roughness, ConvertNode, 3, MP_Roughness);
-					ScalarMatInputConnectionTo(EditorOnly->Anisotropy, ConvertNode, 4, MP_Anisotropy);
-					ColorMatInputConnectionTo(EditorOnly->EmissiveColor, ConvertNode, 5, MP_EmissiveColor);
-					Vector3MatInputConnectionTo(EditorOnly->Normal, ConvertNode, 6, MP_Normal, SUBSTRATE_COPY_CONNECTION);
-					Vector3MatInputConnectionTo(EditorOnly->Tangent, ConvertNode, 7, MP_Tangent);
-					ColorMatInputConnectionTo(EditorOnly->SubsurfaceColor, ConvertNode, 8, MP_SubsurfaceColor);
-					ScalarMatInputConnectionTo(EditorOnly->ClearCoat, ConvertNode, 9, MP_CustomData0);
-					ScalarMatInputConnectionTo(EditorOnly->ClearCoatRoughness, ConvertNode, 10, MP_CustomData1);
-					ScalarMatInputConnectionTo(EditorOnly->Opacity, ConvertNode, 11, MP_Opacity, SUBSTRATE_COPY_CONNECTION);	// We only copy, to keep Opacity on the root node in case BLEND_AlphaComposite is selected.
+					ColorMatInputConnectionTo(EditorOnly->BaseColor, ConvertNode->BaseColor, MP_BaseColor);
+					ScalarMatInputConnectionTo(EditorOnly->Metallic, ConvertNode->Metallic, MP_Metallic);
+					ScalarMatInputConnectionTo(EditorOnly->Specular, ConvertNode->Specular, MP_Specular);
+					ScalarMatInputConnectionTo(EditorOnly->Roughness, ConvertNode->Roughness, MP_Roughness);
+					ScalarMatInputConnectionTo(EditorOnly->Anisotropy, ConvertNode->Anisotropy, MP_Anisotropy);
+					ColorMatInputConnectionTo(EditorOnly->EmissiveColor, ConvertNode->EmissiveColor, MP_EmissiveColor);
+					Vector3MatInputConnectionTo(EditorOnly->Normal, ConvertNode->Normal, MP_Normal, SUBSTRATE_COPY_CONNECTION);
+					Vector3MatInputConnectionTo(EditorOnly->Tangent, ConvertNode->Tangent, MP_Tangent);
+					ColorMatInputConnectionTo(EditorOnly->SubsurfaceColor, ConvertNode->SubSurfaceColor, MP_SubsurfaceColor);
+					ScalarMatInputConnectionTo(EditorOnly->ClearCoat, ConvertNode->ClearCoat, MP_CustomData0);
+					ScalarMatInputConnectionTo(EditorOnly->ClearCoatRoughness, ConvertNode->ClearCoatRoughness, MP_CustomData1);
+					ScalarMatInputConnectionTo(EditorOnly->Opacity, ConvertNode->Opacity, MP_Opacity, SUBSTRATE_COPY_CONNECTION);	// We only copy, to keep Opacity on the root node in case BLEND_AlphaComposite is selected.
 					bRelinkCustomOutputNodes = true;
 			
 					// Shading Model
@@ -3856,10 +3924,10 @@ bool UMaterial::ConvertMaterialToSubstrateMaterial(bool bAllowEmptyMaterialUpdat
 				EditorOnly->ExpressionCollection.AddExpression(VolBSDF);
 				VolBSDF->Material = this;
 				SetPosXAndMoveReferenceToTheRight(VolBSDF);
-				ColorMatInputConnectionTo(EditorOnly->BaseColor, VolBSDF, 0, MP_BaseColor);	
-				ColorMatInputConnectionTo(EditorOnly->SubsurfaceColor, VolBSDF, 1, MP_SubsurfaceColor);
-				ColorMatInputConnectionTo(EditorOnly->EmissiveColor, VolBSDF, 2, MP_EmissiveColor);	
-				ScalarMatInputConnectionTo(EditorOnly->AmbientOcclusion, VolBSDF, 3, MP_AmbientOcclusion);
+				ColorMatInputConnectionTo(EditorOnly->BaseColor, VolBSDF->Albedo, MP_BaseColor);	
+				ColorMatInputConnectionTo(EditorOnly->SubsurfaceColor, VolBSDF->Extinction, MP_SubsurfaceColor);
+				ColorMatInputConnectionTo(EditorOnly->EmissiveColor, VolBSDF->EmissiveColor, MP_EmissiveColor);	
+				ScalarMatInputConnectionTo(EditorOnly->AmbientOcclusion, VolBSDF->AmbientOcclusion, MP_AmbientOcclusion);
 
 				VolBSDF->bEmissiveOnly = ShadingModel == MSM_Unlit;
 
@@ -3879,7 +3947,7 @@ bool UMaterial::ConvertMaterialToSubstrateMaterial(bool bAllowEmptyMaterialUpdat
 				EditorOnly->ExpressionCollection.AddExpression(LightFunctionNode);
 				LightFunctionNode->Material = this;
 				SetPosXAndMoveReferenceToTheRight(LightFunctionNode);
-				ColorMatInputConnectionTo(EditorOnly->EmissiveColor, LightFunctionNode, 0, MP_EmissiveColor);
+				ColorMatInputConnectionTo(EditorOnly->EmissiveColor, LightFunctionNode->Color, MP_EmissiveColor);
 
 				EditorOnly->FrontMaterial.Connect(0, LightFunctionNode);
 				bInvalidateShader = true;
@@ -3901,8 +3969,8 @@ bool UMaterial::ConvertMaterialToSubstrateMaterial(bool bAllowEmptyMaterialUpdat
 				PostProcNode->Material = this;
 				SetPosXAndMoveReferenceToTheRight(PostProcNode);
 
-				ColorMatInputConnectionTo(EditorOnly->EmissiveColor, PostProcNode, 0, MP_EmissiveColor);
-				ScalarMatInputConnectionTo(EditorOnly->Opacity, PostProcNode, 1, MP_Opacity, SUBSTRATE_COPY_CONNECTION);	// We only copy, to keep Opacity on the root node in case BLEND_AlphaComposite is selected.
+				ColorMatInputConnectionTo(EditorOnly->EmissiveColor, PostProcNode->Color, MP_EmissiveColor);
+				ScalarMatInputConnectionTo(EditorOnly->Opacity, PostProcNode->Opacity, MP_Opacity, SUBSTRATE_COPY_CONNECTION);	// We only copy, to keep Opacity on the root node in case BLEND_AlphaComposite is selected.
 
 				EditorOnly->FrontMaterial.Connect(0, PostProcNode);
 				bInvalidateShader = true;
@@ -3918,18 +3986,18 @@ bool UMaterial::ConvertMaterialToSubstrateMaterial(bool bAllowEmptyMaterialUpdat
 				EditorOnly->ExpressionCollection.AddExpression(ConvertNode);
 				ConvertNode->Material = this;
 				SetPosXAndMoveReferenceToTheRight(ConvertNode);
-				ColorMatInputConnectionTo(EditorOnly->BaseColor, ConvertNode, 0, MP_BaseColor);
-				ScalarMatInputConnectionTo(EditorOnly->Metallic, ConvertNode, 1, MP_Metallic);
-				ScalarMatInputConnectionTo(EditorOnly->Specular, ConvertNode, 2, MP_Specular);
-				ScalarMatInputConnectionTo(EditorOnly->Roughness, ConvertNode, 3, MP_Roughness);
-				ScalarMatInputConnectionTo(EditorOnly->Anisotropy, ConvertNode, 4, MP_Anisotropy);
-				ColorMatInputConnectionTo(EditorOnly->EmissiveColor, ConvertNode, 5, MP_EmissiveColor);
-				Vector3MatInputConnectionTo(EditorOnly->Normal, ConvertNode, 6, MP_Normal, SUBSTRATE_COPY_CONNECTION);
-				Vector3MatInputConnectionTo(EditorOnly->Tangent, ConvertNode, 7, MP_Tangent);
-				ColorMatInputConnectionTo(EditorOnly->SubsurfaceColor, ConvertNode, 8, MP_SubsurfaceColor);
-				ScalarMatInputConnectionTo(EditorOnly->ClearCoat, ConvertNode, 9, MP_CustomData0);
-				ScalarMatInputConnectionTo(EditorOnly->ClearCoatRoughness, ConvertNode, 10, MP_CustomData1);
-				ScalarMatInputConnectionTo(EditorOnly->Opacity, ConvertNode, 11, MP_Opacity, SUBSTRATE_COPY_CONNECTION);	// We only copy, to keep Opacity on the root node in case BLEND_AlphaComposite is selected.
+				ColorMatInputConnectionTo(EditorOnly->BaseColor, ConvertNode->BaseColor, MP_BaseColor);
+				ScalarMatInputConnectionTo(EditorOnly->Metallic, ConvertNode->Metallic, MP_Metallic);
+				ScalarMatInputConnectionTo(EditorOnly->Specular, ConvertNode->Specular, MP_Specular);
+				ScalarMatInputConnectionTo(EditorOnly->Roughness, ConvertNode->Roughness, MP_Roughness);
+				ScalarMatInputConnectionTo(EditorOnly->Anisotropy, ConvertNode->Anisotropy, MP_Anisotropy);
+				ColorMatInputConnectionTo(EditorOnly->EmissiveColor, ConvertNode->EmissiveColor, MP_EmissiveColor);
+				Vector3MatInputConnectionTo(EditorOnly->Normal, ConvertNode->Normal, MP_Normal, SUBSTRATE_COPY_CONNECTION);
+				Vector3MatInputConnectionTo(EditorOnly->Tangent, ConvertNode->Tangent, MP_Tangent);
+				ColorMatInputConnectionTo(EditorOnly->SubsurfaceColor, ConvertNode->SubSurfaceColor, MP_SubsurfaceColor);
+				ScalarMatInputConnectionTo(EditorOnly->ClearCoat, ConvertNode->ClearCoat, MP_CustomData0);
+				ScalarMatInputConnectionTo(EditorOnly->ClearCoatRoughness, ConvertNode->ClearCoatRoughness, MP_CustomData1);
+				ScalarMatInputConnectionTo(EditorOnly->Opacity, ConvertNode->Opacity, MP_Opacity, SUBSTRATE_COPY_CONNECTION);	// We only copy, to keep Opacity on the root node in case BLEND_AlphaComposite is selected.
 
 				// Add constant for the Unlit shading model
 				ConvertNode->ShadingModelOverride = ShadingModel;
@@ -3956,8 +4024,8 @@ bool UMaterial::ConvertMaterialToSubstrateMaterial(bool bAllowEmptyMaterialUpdat
 				EditorOnly->ExpressionCollection.AddExpression(UINode);
 				UINode->Material = this;
 				SetPosXAndMoveReferenceToTheRight(UINode);
-				ColorMatInputConnectionTo(EditorOnly->EmissiveColor, UINode, 0, MP_EmissiveColor);
-				ScalarMatInputConnectionTo(EditorOnly->Opacity, UINode, 1, MP_Opacity, SUBSTRATE_COPY_CONNECTION);	// We only copy, to keep Opacity on the root node in case BLEND_AlphaComposite is selected.
+				ColorMatInputConnectionTo(EditorOnly->EmissiveColor, UINode->Color, MP_EmissiveColor);
+				ScalarMatInputConnectionTo(EditorOnly->Opacity, UINode->Opacity, MP_Opacity, SUBSTRATE_COPY_CONNECTION);	// We only copy, to keep Opacity on the root node in case BLEND_AlphaComposite is selected.
 
 				EditorOnly->FrontMaterial.Connect(0, UINode);
 				bInvalidateShader = true;
@@ -3988,18 +4056,18 @@ bool UMaterial::ConvertMaterialToSubstrateMaterial(bool bAllowEmptyMaterialUpdat
 		}
 		if (SingleLayerWaterOutput)
 		{
-			ConnectionTo(*SingleLayerWaterOutput->GetInput(0), ConvertNode, 13); // WaterScatteringCoefficients
-			ConnectionTo(*SingleLayerWaterOutput->GetInput(1), ConvertNode, 14); // WaterAbsorptionCoefficients
-			ConnectionTo(*SingleLayerWaterOutput->GetInput(2), ConvertNode, 15); // WaterPhaseG
-			ConnectionTo(*SingleLayerWaterOutput->GetInput(3), ConvertNode, 16); // ColorScaleBehindWater
+			ConnectionTo(*SingleLayerWaterOutput->GetInput(0), ConvertNode->WaterScatteringCoefficients);
+			ConnectionTo(*SingleLayerWaterOutput->GetInput(1), ConvertNode->WaterAbsorptionCoefficients);
+			ConnectionTo(*SingleLayerWaterOutput->GetInput(2), ConvertNode->WaterPhaseG);
+			ConnectionTo(*SingleLayerWaterOutput->GetInput(3), ConvertNode->ColorScaleBehindWater);
 		}
 		if (ClearCoatBottomNormalOutput)
 		{
-			ConnectionTo(*ClearCoatBottomNormalOutput->GetInput(0), ConvertNode, 17, SUBSTRATE_COPY_CONNECTION); // ClearCoatNormal
+			ConnectionTo(*ClearCoatBottomNormalOutput->GetInput(0), ConvertNode->ClearCoatNormal, SUBSTRATE_COPY_CONNECTION);
 		}
 		if (TangentOutput)
 		{
-			ConnectionTo(*TangentOutput->GetInput(0), ConvertNode, 18, SUBSTRATE_COPY_CONNECTION);	// TangentOutput
+			ConnectionTo(*TangentOutput->GetInput(0), ConvertNode->Tangent, SUBSTRATE_COPY_CONNECTION);
 		}
 	}
 
@@ -4029,14 +4097,11 @@ bool UMaterial::ConvertMaterialToSubstrateMaterial(bool bAllowEmptyMaterialUpdat
 							SetPosXAndMoveReferenceToTheRight(MFCallNode);
 
 							// Connect GlintValue graph to the MF input
-							ConnectionTo(*SlabExpression->GetInput(16)/*GlintValue*/, MFCallNode, 0/*Density*/);
+							ConnectionToIndexed(SlabExpression->GlintValue, MFCallNode, 0/*Density*/);
 
 							// Connect the MF output to the GlintValue input
 							SlabExpression->GlintValue.Connect(0, MFCallNode);
 
-							/** Connects output of InExpression to this input */
-							ENGINE_API void Connect( int32 InOutputIndex, class UMaterialExpression* InExpression );
-
 							MFCallNode->UpdateFromFunctionResource();
 							MFCallNode->PostEditChange();
 							MFCallNode->ConditionalPostLoad();
@@ -4050,6 +4115,8 @@ bool UMaterial::ConvertMaterialToSubstrateMaterial(bool bAllowEmptyMaterialUpdat
 		// Even if no glint conversion was one, we invalidate the shader in order to save the new Substrate's version (from -1 to 0)
 		bInvalidateShader = true;
 	}
+	//else if (EditorOnly->SubstrateVersion < 1)
+	//{ [...] }
 
 	if (bInvalidateShader)
 	{
