Substrate - do not use GetInput to be able to reorder input if we want later.

#tests PC dx12 editor
#rnx

diff --git a/Engine/Source/Runtime/Engine/Private/Materials/Material.cpp b/Engine/Source/Runtime/Engine/Private/Materials/Material.cpp
--- a/Engine/Source/Runtime/Engine/Private/Materials/Material.cpp
+++ b/Engine/Source/Runtime/Engine/Private/Materials/Material.cpp
@@ -3660,10 +3660,10 @@ bool UMaterial::ConvertMaterialToSubstrateMaterial(bool bAllowEmptyMaterialUpdat
 
 				if (SingleLayerWaterOutput)
 				{
-					ConnectionTo(*SingleLayerWaterOutput->GetInput(0), ConvertAttributeNode->WaterScatteringCoefficients);
-					ConnectionTo(*SingleLayerWaterOutput->GetInput(1), ConvertAttributeNode->WaterAbsorptionCoefficients);
-					ConnectionTo(*SingleLayerWaterOutput->GetInput(2), ConvertAttributeNode->WaterPhaseG);
-					ConnectionTo(*SingleLayerWaterOutput->GetInput(3), ConvertAttributeNode->ColorScaleBehindWater);
+					ConnectionTo(SingleLayerWaterOutput->ScatteringCoefficients,	ConvertAttributeNode->WaterScatteringCoefficients);
+					ConnectionTo(SingleLayerWaterOutput->AbsorptionCoefficients,	ConvertAttributeNode->WaterAbsorptionCoefficients);
+					ConnectionTo(SingleLayerWaterOutput->PhaseG,					ConvertAttributeNode->WaterPhaseG);
+					ConnectionTo(SingleLayerWaterOutput->ColorScaleBehindWater,		ConvertAttributeNode->ColorScaleBehindWater);
 				}
 			}
 
@@ -3679,7 +3679,7 @@ bool UMaterial::ConvertMaterialToSubstrateMaterial(bool bAllowEmptyMaterialUpdat
 				SetAttributesNode->Material = this;
 				SetAttributesNode->Inputs[0].Connect(0, EditorOnly->MaterialAttributes.Expression);
 				SetAttributesNode->AttributeSetTypes.Add(CustomEyeTangentGuid);
-				SetAttributesNode->Inputs.Add(*TangentOutput->GetInput(0));
+				SetAttributesNode->Inputs.Add(TangentOutput->Input);
 
 				// Connect to conversion node
 				ConvertAttributeNode->MaterialAttributes.Connect(0, SetAttributesNode);
@@ -3785,8 +3785,8 @@ bool UMaterial::ConvertMaterialToSubstrateMaterial(bool bAllowEmptyMaterialUpdat
 						UMaterialExpressionSubstrateMetalnessToDiffuseAlbedoF0* SubstrateMetalnessToDiffuseAlbedoF0 = NewObject<UMaterialExpressionSubstrateMetalnessToDiffuseAlbedoF0>(this);
 						EditorOnly->ExpressionCollection.AddExpression(SubstrateMetalnessToDiffuseAlbedoF0);
 						SetPosXAndMoveReferenceToTheRight(SubstrateMetalnessToDiffuseAlbedoF0);
-						ColorMatInputConnectionTo(EditorOnly->BaseColor, SubstrateMetalnessToDiffuseAlbedoF0->BaseColor, MP_BaseColor);
-						ScalarMatInputConnectionTo(EditorOnly->Metallic, SubstrateMetalnessToDiffuseAlbedoF0->Metallic, MP_Metallic);
+						ColorMatInputConnectionTo(EditorOnly->BaseColor, SubstrateMetalnessToDiffuseAlbedoF0->BaseColor,MP_BaseColor);
+						ScalarMatInputConnectionTo(EditorOnly->Metallic, SubstrateMetalnessToDiffuseAlbedoF0->Metallic,	MP_Metallic);
 						ScalarMatInputConnectionTo(EditorOnly->Specular, SubstrateMetalnessToDiffuseAlbedoF0->Specular, MP_Specular);
 					
 						// Top slab BSDF as a simple Disney material
@@ -3794,15 +3794,14 @@ bool UMaterial::ConvertMaterialToSubstrateMaterial(bool bAllowEmptyMaterialUpdat
 						EditorOnly->ExpressionCollection.AddExpression(BottomSlabBSDF);
 						BottomSlabBSDF->Material = this;
 						SetPosXAndMoveReferenceToTheRight(BottomSlabBSDF);
-						BottomSlabBSDF->GetInput(0)->Connect(0, SubstrateMetalnessToDiffuseAlbedoF0);
-						BottomSlabBSDF->GetInput(1)->Connect(1, SubstrateMetalnessToDiffuseAlbedoF0);
-						BottomSlabBSDF->GetInput(2)->Connect(2, SubstrateMetalnessToDiffuseAlbedoF0);
-						ScalarMatInputConnectionTo(EditorOnly->Roughness, BottomSlabBSDF->Roughness, MP_Roughness);
-						ScalarMatInputConnectionTo(EditorOnly->Anisotropy, BottomSlabBSDF->Anisotropy, MP_Anisotropy, SUBSTRATE_COPY_CONNECTION);
-						Vector3MatInputConnectionTo(EditorOnly->Tangent, BottomSlabBSDF->Tangent, MP_Tangent);
+						BottomSlabBSDF->DiffuseAlbedo.Connect(0, SubstrateMetalnessToDiffuseAlbedoF0);
+						BottomSlabBSDF->F0.Connect(1, SubstrateMetalnessToDiffuseAlbedoF0);
+						ScalarMatInputConnectionTo(EditorOnly->Roughness,	BottomSlabBSDF->Roughness,	MP_Roughness);
+						ScalarMatInputConnectionTo(EditorOnly->Anisotropy,	BottomSlabBSDF->Anisotropy,	MP_Anisotropy, SUBSTRATE_COPY_CONNECTION);
+						Vector3MatInputConnectionTo(EditorOnly->Tangent,	BottomSlabBSDF->Tangent,	MP_Tangent);
 
 						check(ClearCoatBottomNormalOutput);
-						ConnectionTo(*ClearCoatBottomNormalOutput->GetInput(0), BottomSlabBSDF->Normal, SUBSTRATE_COPY_CONNECTION);
+						ConnectionTo(ClearCoatBottomNormalOutput->Input, BottomSlabBSDF->Normal, SUBSTRATE_COPY_CONNECTION);
 
 						// Now weight the top base material by opacity.
 						UMaterialExpressionSubstrateSlabBSDF* TopSlabBSDF = NewObject<UMaterialExpressionSubstrateSlabBSDF>(this);
@@ -3810,16 +3809,16 @@ bool UMaterial::ConvertMaterialToSubstrateMaterial(bool bAllowEmptyMaterialUpdat
 						TopSlabBSDF->Material = this;
 						TopSlabBSDF->MaterialExpressionEditorX = BottomSlabBSDF->MaterialExpressionEditorX;
 						TopSlabBSDF->MaterialExpressionEditorY = BottomSlabBSDF->MaterialExpressionEditorY + 650;
-						ColorMatInputConnectionTo(EditorOnly->EmissiveColor, TopSlabBSDF->EmissiveColor, MP_EmissiveColor);
-						ScalarMatInputConnectionTo(EditorOnly->ClearCoatRoughness, TopSlabBSDF->Roughness, MP_CustomData0);
-						Vector3MatInputConnectionTo(EditorOnly->Normal, TopSlabBSDF->Normal, MP_Normal);
+						ColorMatInputConnectionTo(EditorOnly->EmissiveColor,		TopSlabBSDF->EmissiveColor,	MP_EmissiveColor);
+						ScalarMatInputConnectionTo(EditorOnly->ClearCoatRoughness,	TopSlabBSDF->Roughness,		MP_CustomData0);
+						Vector3MatInputConnectionTo(EditorOnly->Normal,				TopSlabBSDF->Normal,		MP_Normal);
 
 						//  The top layer has a hard coded specular value of 0.5 (F0 = 0.04)
 						UMaterialExpressionConstant* ConstantHalf = NewObject<UMaterialExpressionConstant>(this);
 						EditorOnly->ExpressionCollection.AddExpression(ConstantHalf);
 						ReplaceNodeAndMoveToTheRight(TopSlabBSDF, ConstantHalf);
 						ConstantHalf->R = 0.5f * 0.08f;
-						TopSlabBSDF->GetInput(1)->Connect(0, ConstantHalf);
+						TopSlabBSDF->F0.Connect(0, ConstantHalf);
 
 						// The original clear coat is a complex assemblage of arbitrary functions that do not always make sense.
 						// To simplify things, we set the top slab BSDF as having a constant Grey scale transmittance.
@@ -3828,7 +3827,7 @@ bool UMaterial::ConvertMaterialToSubstrateMaterial(bool bAllowEmptyMaterialUpdat
 						EditorOnly->ExpressionCollection.AddExpression(ConstantZero);
 						ReplaceNodeAndMoveToTheRight(TopSlabBSDF, ConstantZero);
 						ConstantZero->R = 0.0f;
-						TopSlabBSDF->GetInput(0)->Connect(0, ConstantZero);							// BaseColor = 0 to only feature absorption, no scattering
+						TopSlabBSDF->DiffuseAlbedo.Connect(0, ConstantZero);		// DiffuseAlbedo = 0 to only feature absorption, no scattering
 
 						// Now setup the mean free path with a hard coded transmittance of 0.75 when viewing the surface perpendicularly
 						UMaterialExpressionConstant* Constant075 = NewObject<UMaterialExpressionConstant>(this);
@@ -3845,15 +3844,15 @@ bool UMaterial::ConvertMaterialToSubstrateMaterial(bool bAllowEmptyMaterialUpdat
 						UMaterialExpressionSubstrateWeight* TopSlabBSDFWithCoverage = NewObject<UMaterialExpressionSubstrateWeight>(this);
 						EditorOnly->ExpressionCollection.AddExpression(SubstrateMetalnessToDiffuseAlbedoF0); // Add the SetAttributesNode to ensure the MaterialAttribute input are correctly tracked & cached
 						SetPosXAndMoveReferenceToTheRight(TopSlabBSDFWithCoverage);
-						TopSlabBSDFWithCoverage->GetInput(0)->Connect(0, TopSlabBSDF);												// TopSlabBSDF -> A
+						TopSlabBSDFWithCoverage->A.Connect(0, TopSlabBSDF);
 						ScalarMatInputConnectionTo(EditorOnly->ClearCoat, TopSlabBSDFWithCoverage->Weight, MP_CustomData0);
 						ScalarMatInputConnectionTo(EditorOnly->ClearCoatRoughness, TopSlabBSDF->Roughness, MP_CustomData1);
 
 						UMaterialExpressionSubstrateVerticalLayering* VerticalLayering = NewObject<UMaterialExpressionSubstrateVerticalLayering>(this);
 						EditorOnly->ExpressionCollection.AddExpression(VerticalLayering);
 						SetPosXAndMoveReferenceToTheRight(VerticalLayering);
-						VerticalLayering->GetInput(0)->Connect(0, TopSlabBSDFWithCoverage);			// Top -> Top
-						VerticalLayering->GetInput(1)->Connect(0, BottomSlabBSDF);					// Bottom -> Base
+						VerticalLayering->Top.Connect(0, TopSlabBSDFWithCoverage);
+						VerticalLayering->Base.Connect(0, BottomSlabBSDF);
 
 						EditorOnly->FrontMaterial.Connect(0, VerticalLayering);
 						bClearCoatConversionDone = true;
@@ -3868,18 +3867,18 @@ bool UMaterial::ConvertMaterialToSubstrateMaterial(bool bAllowEmptyMaterialUpdat
 					ConvertNode->Material = this;
 					SetPosXAndMoveReferenceToTheRight(ConvertNode);
 					ConvertNode->SubsurfaceProfile = bRequireNoSubsurfaceProfile ? nullptr : SubsurfaceProfile;
-					ColorMatInputConnectionTo(EditorOnly->BaseColor, ConvertNode->BaseColor, MP_BaseColor);
-					ScalarMatInputConnectionTo(EditorOnly->Metallic, ConvertNode->Metallic, MP_Metallic);
-					ScalarMatInputConnectionTo(EditorOnly->Specular, ConvertNode->Specular, MP_Specular);
-					ScalarMatInputConnectionTo(EditorOnly->Roughness, ConvertNode->Roughness, MP_Roughness);
-					ScalarMatInputConnectionTo(EditorOnly->Anisotropy, ConvertNode->Anisotropy, MP_Anisotropy);
-					ColorMatInputConnectionTo(EditorOnly->EmissiveColor, ConvertNode->EmissiveColor, MP_EmissiveColor);
-					Vector3MatInputConnectionTo(EditorOnly->Normal, ConvertNode->Normal, MP_Normal, SUBSTRATE_COPY_CONNECTION);
-					Vector3MatInputConnectionTo(EditorOnly->Tangent, ConvertNode->Tangent, MP_Tangent);
-					ColorMatInputConnectionTo(EditorOnly->SubsurfaceColor, ConvertNode->SubSurfaceColor, MP_SubsurfaceColor);
-					ScalarMatInputConnectionTo(EditorOnly->ClearCoat, ConvertNode->ClearCoat, MP_CustomData0);
-					ScalarMatInputConnectionTo(EditorOnly->ClearCoatRoughness, ConvertNode->ClearCoatRoughness, MP_CustomData1);
-					ScalarMatInputConnectionTo(EditorOnly->Opacity, ConvertNode->Opacity, MP_Opacity, SUBSTRATE_COPY_CONNECTION);	// We only copy, to keep Opacity on the root node in case BLEND_AlphaComposite is selected.
+					ColorMatInputConnectionTo(EditorOnly->BaseColor,			ConvertNode->BaseColor,			MP_BaseColor);
+					ScalarMatInputConnectionTo(EditorOnly->Metallic,			ConvertNode->Metallic,			MP_Metallic);
+					ScalarMatInputConnectionTo(EditorOnly->Specular,			ConvertNode->Specular,			MP_Specular);
+					ScalarMatInputConnectionTo(EditorOnly->Roughness,			ConvertNode->Roughness,			MP_Roughness);
+					ScalarMatInputConnectionTo(EditorOnly->Anisotropy,			ConvertNode->Anisotropy,		MP_Anisotropy);
+					ColorMatInputConnectionTo(EditorOnly->EmissiveColor,		ConvertNode->EmissiveColor,		MP_EmissiveColor);
+					Vector3MatInputConnectionTo(EditorOnly->Normal,				ConvertNode->Normal,			MP_Normal, SUBSTRATE_COPY_CONNECTION);
+					Vector3MatInputConnectionTo(EditorOnly->Tangent,			ConvertNode->Tangent,			MP_Tangent);
+					ColorMatInputConnectionTo(EditorOnly->SubsurfaceColor,		ConvertNode->SubSurfaceColor,	MP_SubsurfaceColor);
+					ScalarMatInputConnectionTo(EditorOnly->ClearCoat,			ConvertNode->ClearCoat,			MP_CustomData0);
+					ScalarMatInputConnectionTo(EditorOnly->ClearCoatRoughness,	ConvertNode->ClearCoatRoughness,MP_CustomData1);
+					ScalarMatInputConnectionTo(EditorOnly->Opacity,				ConvertNode->Opacity,			MP_Opacity, SUBSTRATE_COPY_CONNECTION);	// We only copy, to keep Opacity on the root node in case BLEND_AlphaComposite is selected.
 					bRelinkCustomOutputNodes = true;
 			
 					// Shading Model
@@ -3924,10 +3923,10 @@ bool UMaterial::ConvertMaterialToSubstrateMaterial(bool bAllowEmptyMaterialUpdat
 				EditorOnly->ExpressionCollection.AddExpression(VolBSDF);
 				VolBSDF->Material = this;
 				SetPosXAndMoveReferenceToTheRight(VolBSDF);
-				ColorMatInputConnectionTo(EditorOnly->BaseColor, VolBSDF->Albedo, MP_BaseColor);	
-				ColorMatInputConnectionTo(EditorOnly->SubsurfaceColor, VolBSDF->Extinction, MP_SubsurfaceColor);
-				ColorMatInputConnectionTo(EditorOnly->EmissiveColor, VolBSDF->EmissiveColor, MP_EmissiveColor);	
-				ScalarMatInputConnectionTo(EditorOnly->AmbientOcclusion, VolBSDF->AmbientOcclusion, MP_AmbientOcclusion);
+				ColorMatInputConnectionTo(EditorOnly->BaseColor,		VolBSDF->Albedo,			MP_BaseColor);	
+				ColorMatInputConnectionTo(EditorOnly->SubsurfaceColor,	VolBSDF->Extinction,		MP_SubsurfaceColor);
+				ColorMatInputConnectionTo(EditorOnly->EmissiveColor,	VolBSDF->EmissiveColor,		MP_EmissiveColor);	
+				ScalarMatInputConnectionTo(EditorOnly->AmbientOcclusion,VolBSDF->AmbientOcclusion,	MP_AmbientOcclusion);
 
 				VolBSDF->bEmissiveOnly = ShadingModel == MSM_Unlit;
 
@@ -3986,18 +3985,18 @@ bool UMaterial::ConvertMaterialToSubstrateMaterial(bool bAllowEmptyMaterialUpdat
 				EditorOnly->ExpressionCollection.AddExpression(ConvertNode);
 				ConvertNode->Material = this;
 				SetPosXAndMoveReferenceToTheRight(ConvertNode);
-				ColorMatInputConnectionTo(EditorOnly->BaseColor, ConvertNode->BaseColor, MP_BaseColor);
-				ScalarMatInputConnectionTo(EditorOnly->Metallic, ConvertNode->Metallic, MP_Metallic);
-				ScalarMatInputConnectionTo(EditorOnly->Specular, ConvertNode->Specular, MP_Specular);
-				ScalarMatInputConnectionTo(EditorOnly->Roughness, ConvertNode->Roughness, MP_Roughness);
-				ScalarMatInputConnectionTo(EditorOnly->Anisotropy, ConvertNode->Anisotropy, MP_Anisotropy);
-				ColorMatInputConnectionTo(EditorOnly->EmissiveColor, ConvertNode->EmissiveColor, MP_EmissiveColor);
-				Vector3MatInputConnectionTo(EditorOnly->Normal, ConvertNode->Normal, MP_Normal, SUBSTRATE_COPY_CONNECTION);
-				Vector3MatInputConnectionTo(EditorOnly->Tangent, ConvertNode->Tangent, MP_Tangent);
-				ColorMatInputConnectionTo(EditorOnly->SubsurfaceColor, ConvertNode->SubSurfaceColor, MP_SubsurfaceColor);
-				ScalarMatInputConnectionTo(EditorOnly->ClearCoat, ConvertNode->ClearCoat, MP_CustomData0);
-				ScalarMatInputConnectionTo(EditorOnly->ClearCoatRoughness, ConvertNode->ClearCoatRoughness, MP_CustomData1);
-				ScalarMatInputConnectionTo(EditorOnly->Opacity, ConvertNode->Opacity, MP_Opacity, SUBSTRATE_COPY_CONNECTION);	// We only copy, to keep Opacity on the root node in case BLEND_AlphaComposite is selected.
+				ColorMatInputConnectionTo(EditorOnly->BaseColor,			ConvertNode->BaseColor,			MP_BaseColor);
+				ScalarMatInputConnectionTo(EditorOnly->Metallic,			ConvertNode->Metallic,			MP_Metallic);
+				ScalarMatInputConnectionTo(EditorOnly->Specular,			ConvertNode->Specular,			MP_Specular);
+				ScalarMatInputConnectionTo(EditorOnly->Roughness,			ConvertNode->Roughness,			MP_Roughness);
+				ScalarMatInputConnectionTo(EditorOnly->Anisotropy,			ConvertNode->Anisotropy,		MP_Anisotropy);
+				ColorMatInputConnectionTo(EditorOnly->EmissiveColor,		ConvertNode->EmissiveColor,		MP_EmissiveColor);
+				Vector3MatInputConnectionTo(EditorOnly->Normal,				ConvertNode->Normal,			MP_Normal, SUBSTRATE_COPY_CONNECTION);
+				Vector3MatInputConnectionTo(EditorOnly->Tangent,			ConvertNode->Tangent,			MP_Tangent);
+				ColorMatInputConnectionTo(EditorOnly->SubsurfaceColor,		ConvertNode->SubSurfaceColor,	MP_SubsurfaceColor);
+				ScalarMatInputConnectionTo(EditorOnly->ClearCoat,			ConvertNode->ClearCoat,			MP_CustomData0);
+				ScalarMatInputConnectionTo(EditorOnly->ClearCoatRoughness,	ConvertNode->ClearCoatRoughness,MP_CustomData1);
+				ScalarMatInputConnectionTo(EditorOnly->Opacity,				ConvertNode->Opacity,			MP_Opacity, SUBSTRATE_COPY_CONNECTION);	// We only copy, to keep Opacity on the root node in case BLEND_AlphaComposite is selected.
 
 				// Add constant for the Unlit shading model
 				ConvertNode->ShadingModelOverride = ShadingModel;
@@ -4024,8 +4023,8 @@ bool UMaterial::ConvertMaterialToSubstrateMaterial(bool bAllowEmptyMaterialUpdat
 				EditorOnly->ExpressionCollection.AddExpression(UINode);
 				UINode->Material = this;
 				SetPosXAndMoveReferenceToTheRight(UINode);
-				ColorMatInputConnectionTo(EditorOnly->EmissiveColor, UINode->Color, MP_EmissiveColor);
-				ScalarMatInputConnectionTo(EditorOnly->Opacity, UINode->Opacity, MP_Opacity, SUBSTRATE_COPY_CONNECTION);	// We only copy, to keep Opacity on the root node in case BLEND_AlphaComposite is selected.
+				ColorMatInputConnectionTo(EditorOnly->EmissiveColor,UINode->Color,	MP_EmissiveColor);
+				ScalarMatInputConnectionTo(EditorOnly->Opacity,		UINode->Opacity,MP_Opacity, SUBSTRATE_COPY_CONNECTION);	// We only copy, to keep Opacity on the root node in case BLEND_AlphaComposite is selected.
 
 				EditorOnly->FrontMaterial.Connect(0, UINode);
 				bInvalidateShader = true;
@@ -4043,8 +4042,8 @@ bool UMaterial::ConvertMaterialToSubstrateMaterial(bool bAllowEmptyMaterialUpdat
 
 		if (ThinTranslucentOutput)
 		{
-			FExpressionInput* TransmittanceColorInput = ThinTranslucentOutput->GetInput(0);
-			FExpressionInput* ThinTranslucentSurfaceCoverageInput = ThinTranslucentOutput->GetInput(1);
+			FExpressionInput* TransmittanceColorInput = &ThinTranslucentOutput->TransmittanceColor;
+			FExpressionInput* ThinTranslucentSurfaceCoverageInput = &ThinTranslucentOutput->SurfaceCoverage;
 			if (TransmittanceColorInput && TransmittanceColorInput->IsConnected())
 			{
 				ConvertNode->TransmittanceColor.Connect(TransmittanceColorInput->OutputIndex, TransmittanceColorInput->Expression);
@@ -4056,18 +4055,18 @@ bool UMaterial::ConvertMaterialToSubstrateMaterial(bool bAllowEmptyMaterialUpdat
 		}
 		if (SingleLayerWaterOutput)
 		{
-			ConnectionTo(*SingleLayerWaterOutput->GetInput(0), ConvertNode->WaterScatteringCoefficients);
-			ConnectionTo(*SingleLayerWaterOutput->GetInput(1), ConvertNode->WaterAbsorptionCoefficients);
-			ConnectionTo(*SingleLayerWaterOutput->GetInput(2), ConvertNode->WaterPhaseG);
-			ConnectionTo(*SingleLayerWaterOutput->GetInput(3), ConvertNode->ColorScaleBehindWater);
+			ConnectionTo(SingleLayerWaterOutput->ScatteringCoefficients,	ConvertNode->WaterScatteringCoefficients);
+			ConnectionTo(SingleLayerWaterOutput->AbsorptionCoefficients,	ConvertNode->WaterAbsorptionCoefficients);
+			ConnectionTo(SingleLayerWaterOutput->PhaseG,					ConvertNode->WaterPhaseG);
+			ConnectionTo(SingleLayerWaterOutput->ColorScaleBehindWater,		ConvertNode->ColorScaleBehindWater);
 		}
 		if (ClearCoatBottomNormalOutput)
 		{
-			ConnectionTo(*ClearCoatBottomNormalOutput->GetInput(0), ConvertNode->ClearCoatNormal, SUBSTRATE_COPY_CONNECTION);
+			ConnectionTo(ClearCoatBottomNormalOutput->Input, ConvertNode->ClearCoatNormal, SUBSTRATE_COPY_CONNECTION);
 		}
 		if (TangentOutput)
 		{
-			ConnectionTo(*TangentOutput->GetInput(0), ConvertNode->Tangent, SUBSTRATE_COPY_CONNECTION);
+			ConnectionTo(TangentOutput->Input, ConvertNode->Tangent, SUBSTRATE_COPY_CONNECTION);
 		}
 	}
 
