Add references to the "next to play" chunk

diff --git a/Engine/Source/Runtime/Engine/Private/AudioDecompress.cpp b/Engine/Source/Runtime/Engine/Private/AudioDecompress.cpp
--- a/Engine/Source/Runtime/Engine/Private/AudioDecompress.cpp
+++ b/Engine/Source/Runtime/Engine/Private/AudioDecompress.cpp
@@ -19,6 +19,13 @@ namespace AudioDecompressPrivate
 	#define WITH_AUDIO_DECODER_DIAGNOSTICS (0)
 #endif //WITH_AUDIO_DECODER_DIAGNOSTICS
 
+	static int32 bEnableRetainingNextPlayingChunkCVar = 0;
+	FAutoConsoleVariableRef CVarForceDecoderNegativeSamplesOnWave(
+		TEXT("au.streamcache.RetainNextChunk"),
+		bEnableRetainingNextPlayingChunkCVar,
+		TEXT("Keep a reference to next to play chunk" ),
+		ECVF_Default);
+	
 #if WITH_AUDIO_DECODER_DIAGNOSTICS
 	
 	static FString ForceDecoderErrorOnWaveCVar;
@@ -671,6 +678,23 @@ uint32 IStreamedCompressedInfo::ZeroBuffer(uint8* Destination, uint32 BufferSize
 }
 
 
+void IStreamedCompressedInfo::RetainNextChunkInSequence(const FSoundWaveProxyPtr& InSoundWave,const uint32 InChunkIndex)
+{
+	if (AudioDecompressPrivate::bEnableRetainingNextPlayingChunkCVar)
+	{
+		// Keep and hold hard ref to next chunk.
+		if (const int32 NextChunk = FCachedAudioStreamingManager::GetNextChunkIndex(InSoundWave, InChunkIndex); NextChunk != INDEX_NONE)
+		{
+			NextPlayingChunkHandle = IStreamingManager::Get().GetAudioStreamingManager().GetLoadedChunk(
+				InSoundWave, NextChunk, /* bBlockForLoad*/ false, /* bForImmediatePlayback */ true);
+
+			return;
+		}
+	}
+	// Otherwise null it.
+	NextPlayingChunkHandle = {};
+}
+
 const uint8* IStreamedCompressedInfo::GetLoadedChunk(const FSoundWaveProxyPtr& InSoundWave, uint32 ChunkIndex, uint32& OutChunkSize)
 {
 	if (!ensure(InSoundWave.IsValid()))
@@ -692,6 +716,7 @@ const uint8* IStreamedCompressedInfo::GetLoadedChunk(const FSoundWaveProxyPtr& I
 	else
 	{
 		CurCompressedChunkHandle = IStreamingManager::Get().GetAudioStreamingManager().GetLoadedChunk(InSoundWave, ChunkIndex, false, true);
+		RetainNextChunkInSequence(InSoundWave,ChunkIndex);
 		OutChunkSize = CurCompressedChunkHandle.Num();
 		return CurCompressedChunkHandle.GetData();
 	}
diff --git a/Engine/Source/Runtime/Engine/Private/AudioStreamingCache.cpp b/Engine/Source/Runtime/Engine/Private/AudioStreamingCache.cpp
--- a/Engine/Source/Runtime/Engine/Private/AudioStreamingCache.cpp
+++ b/Engine/Source/Runtime/Engine/Private/AudioStreamingCache.cpp
@@ -569,7 +569,7 @@ FAudioChunkCache* FCachedAudioStreamingManager::GetCacheForChunkSize(uint32 InCh
 	return nullptr;
 }
 
-int32 FCachedAudioStreamingManager::GetNextChunkIndex(const FSoundWaveProxyPtr&  InSoundWave, uint32 CurrentChunkIndex) const
+int32 FCachedAudioStreamingManager::GetNextChunkIndex(const FSoundWaveProxyPtr& InSoundWave, uint32 InCurrentChunkIndex)
 {
 	LLM_SCOPE(ELLMTag::AudioStreamCache);
 
@@ -588,7 +588,7 @@ int32 FCachedAudioStreamingManager::GetNextChunkIndex(const FSoundWaveProxyPtr&
 		// We don't need to load anything.
 		return INDEX_NONE;
 	}
-	else if (CurrentChunkIndex == (NumChunksTotal - 1))
+	else if (InCurrentChunkIndex == (NumChunksTotal - 1))
 	{
 		// if we're on the last chunk, load the first chunk after the zeroth chunk.
 		return 1;
@@ -596,7 +596,7 @@ int32 FCachedAudioStreamingManager::GetNextChunkIndex(const FSoundWaveProxyPtr&
 	else
 	{
 		// Otherwise, there's another chunk of audio after this one before the end of this file.
-		return CurrentChunkIndex + 1;
+		return InCurrentChunkIndex + 1;
 	}
 }
 
diff --git a/Engine/Source/Runtime/Engine/Public/AudioDecompress.h b/Engine/Source/Runtime/Engine/Public/AudioDecompress.h
--- a/Engine/Source/Runtime/Engine/Public/AudioDecompress.h
+++ b/Engine/Source/Runtime/Engine/Public/AudioDecompress.h
@@ -286,6 +286,8 @@ protected:
 	 */
 	ENGINE_API const uint8* GetLoadedChunk(const FSoundWaveProxyPtr& InSoundWave, uint32 ChunkIndex, uint32& OutChunkSize);
 
+	void RetainNextChunkInSequence(const FSoundWaveProxyPtr& InSoundWave,const uint32 InChunkIndex);
+
 	/**
 	 * Gets the current chunks seektable instance (or creates one)
 	 * @return the currently loaded chunk seektable
@@ -332,7 +334,9 @@ protected:
 	/** Number of bytes of padding used, overridden in some implementations. Defaults to 0. */
 	uint32 SrcBufferPadding;
 	/** Chunk Handle to ensure that this chunk of streamed audio is not deleted while we are using it. */
-	FAudioChunkHandle CurCompressedChunkHandle;		
+	FAudioChunkHandle CurCompressedChunkHandle;
+	/** Handle to the next chunk in sequence. Added once we're playing, to prevent it being dropped before it's played.  */
+	FAudioChunkHandle NextPlayingChunkHandle;
 	/** If there's a chunked seek-table present, this contains the current chunks portion. */
 	TPimplPtr<FStreamedAudioChunkSeekTable> CurrentChunkSeekTable;
 	/** 
diff --git a/Engine/Source/Runtime/Engine/Public/AudioStreamingCache.h b/Engine/Source/Runtime/Engine/Public/AudioStreamingCache.h
--- a/Engine/Source/Runtime/Engine/Public/AudioStreamingCache.h
+++ b/Engine/Source/Runtime/Engine/Public/AudioStreamingCache.h
@@ -544,7 +544,14 @@ public:
 	virtual void SetProfilingMode(bool bEnabled) override;
 	// End IAudioStreamingManager interface
 
-
+	/**
+	 * Returns the next chunk to kick off a load for, or INDEX_NONE if there is only one chunk to cache.
+	 * @param InSoundWave Wave Proxy
+	 * @param InCurrentChunkIndex Current index to move forwards of
+	 * @return Chunk Index or INDEX_NONE if there is only one chunk 
+	 */
+	static int32 GetNextChunkIndex(const FSoundWaveProxyPtr& InSoundWave, const uint32 InCurrentChunkIndex);
+	
 protected:
 
 	// These are used to reference count consumers of audio chunks.
@@ -563,11 +570,6 @@ protected:
 	FAudioChunkCache* GetCacheForWave(const FSoundWaveProxyPtr&  InSoundWave) const;
 	FAudioChunkCache* GetCacheForChunkSize(uint32 InChunkSize) const;
 
-	/**
-	 * Returns the next chunk to kick off a load for, or INDEX_NONE if there is only one chunk to cache.
-	 */
-	int32 GetNextChunkIndex(const FSoundWaveProxyPtr&  InSoundWave, uint32 CurrentChunkIndex) const;
-
 	/** Audio chunk caches. These are set up on initialization. */
 	TArray<FAudioChunkCache> CacheArray;
 };
