Retain first chunk if inlined (no matter load mode)

diff --git a/Engine/Source/Runtime/Engine/Private/SoundWave.cpp b/Engine/Source/Runtime/Engine/Private/SoundWave.cpp
--- a/Engine/Source/Runtime/Engine/Private/SoundWave.cpp
+++ b/Engine/Source/Runtime/Engine/Private/SoundWave.cpp
@@ -61,7 +61,7 @@ FAutoConsoleVariableRef CVarForceNonStreamingInEditor(
 	TEXT("0: Honor the Play When Silent flag, 1: stop all silent non-procedural sources."),
 	ECVF_Default);
 
-static int32 DisableRetainingCVar = 0;
+int32 DisableRetainingCVar = 0;
 FAutoConsoleVariableRef CVarDisableRetaining(
 	TEXT("au.streamcache.DisableRetaining"),
 	DisableRetainingCVar,
@@ -69,6 +69,13 @@ FAutoConsoleVariableRef CVarDisableRetaining(
 	TEXT("0: Don't disable retaining, 1: disable retaining."),
 	ECVF_Default);
 
+static int32 bRetainInlinedFirstChunkCVar = 0;
+FAutoConsoleVariableRef CVarRetainInlinedFirstChunk(
+	TEXT("au.streamcache.RetainInlinedFirstChunk"),
+	bRetainInlinedFirstChunkCVar,
+	TEXT("0 is off, 1 or anything else is on"),
+	ECVF_Default);
+
 static int32 BlockOnChunkLoadCompletionCVar = 0;
 FAutoConsoleVariableRef CVarBlockOnChunkLoadCompletion(
 	TEXT("au.streamcache.BlockOnChunkLoadCompletion"),
@@ -2074,7 +2081,23 @@ void USoundWave::PostLoad()
 
 		if (ShouldUseStreamCaching())
 		{
-			if (!GIsEditor && !DisableRetainingCVar && ActualLoadingBehavior == ESoundWaveLoadingBehavior::RetainOnLoad)
+			const TIndirectArray<FStreamedAudioChunk>& Chunks = SoundWaveDataPtr->RunningPlatformData.Chunks;
+
+			// First chunk can be optionally inlined for the sake of pre-buffering.
+			const bool bIsFirstChunkInlined = Chunks.IsValidIndex(1) ? Chunks[1].BulkData.IsInlined() : false;
+
+			// Allowing retaining if we're NOT in the editor AND we've NOT disabled retaining (via a cvar).
+			const bool bAllowRetaining = !GIsEditor && !DisableRetainingCVar;
+
+			// Retaining the "inlined" first chunk (which can be any load mode) is enabled via a CVAR.
+			const bool bShouldRetainInlinedFirstChunked = bIsFirstChunkInlined && !!bRetainInlinedFirstChunkCVar;
+
+			// We retain the first chunk if we're set to retain
+			// OR we've inlined the first chunk (we do when we're pre-buffering). We don't want these dumped by the stream cache if we can help it.
+			const bool bShouldRetainFirstChunk = bAllowRetaining &&
+				(ActualLoadingBehavior == ESoundWaveLoadingBehavior::RetainOnLoad || bShouldRetainInlinedFirstChunked);
+				
+			if (bShouldRetainFirstChunk)
 			{
 				RetainCompressedAudio(false);
 			}
