Temp revert back to old CPFUO editinline duplication path for a potential regression (investigating).

#rnx

diff --git a/Engine/Source/Runtime/Engine/Private/UnrealEngine.cpp b/Engine/Source/Runtime/Engine/Private/UnrealEngine.cpp
--- a/Engine/Source/Runtime/Engine/Private/UnrealEngine.cpp
+++ b/Engine/Source/Runtime/Engine/Private/UnrealEngine.cpp
@@ -17631,7 +17631,7 @@ public:
 
 namespace UE::Engine::Private
 {
-	static bool bUseOldDuplicateEditInlineSubObjectsPath = false;
+	static bool bUseOldDuplicateEditInlineSubObjectsPath = true;
 	static FAutoConsoleVariableRef CVarUseOldDuplicateEditInlineSubObjectsPath(
 		TEXT("CPFUO.UseOldDuplicateEditInlineSubObjectsPath"),
 		bUseOldDuplicateEditInlineSubObjectsPath,
