Removed unused FRectLightRayTracingData

#rnx

diff --git a/Engine/Source/Runtime/Engine/Classes/Components/RectLightComponent.h b/Engine/Source/Runtime/Engine/Classes/Components/RectLightComponent.h
--- a/Engine/Source/Runtime/Engine/Classes/Components/RectLightComponent.h
+++ b/Engine/Source/Runtime/Engine/Classes/Components/RectLightComponent.h
@@ -103,9 +103,4 @@ public:
 	ENGINE_API virtual void PostEditChangeProperty(FPropertyChangedEvent& PropertyChangedEvent) override;
 #endif // WITH_EDITOR
 	//~ End UObject Interface
-
-private:
-	ENGINE_API void UpdateRayTracingData();
-	friend class FRectLightSceneProxy;
-	struct FRectLightRayTracingData* RayTracingData;
 };
diff --git a/Engine/Source/Runtime/Engine/Private/Components/RectLightComponent.cpp b/Engine/Source/Runtime/Engine/Private/Components/RectLightComponent.cpp
--- a/Engine/Source/Runtime/Engine/Private/Components/RectLightComponent.cpp
+++ b/Engine/Source/Runtime/Engine/Private/Components/RectLightComponent.cpp
@@ -85,7 +85,6 @@ void CalculateRectLightBarnCorners(float SourceWidth, float SourceHeight, float
 
 URectLightComponent::URectLightComponent(const FObjectInitializer& ObjectInitializer)
 	: Super(ObjectInitializer)
-	, RayTracingData(new FRectLightRayTracingData())
 {
 #if WITH_EDITORONLY_DATA
 	if (!IsRunningCommandlet())
@@ -108,7 +107,6 @@ URectLightComponent::URectLightComponent(const FObjectInitializer& ObjectInitial
 	BarnDoorAngle = GetRectLightBarnDoorMaxAngle();
 	BarnDoorLength = 20.0f;
 	LightFunctionConeAngle = 0.0f;
-	// RayTracingData will be initialised on the render thread.
 }
 
 FLightSceneProxy* URectLightComponent::CreateSceneProxy() const
@@ -123,7 +121,7 @@ void URectLightComponent::SetSourceTexture(UTexture* NewValue)
 	{
 		SourceTexture = NewValue;
 
-		// This will trigger a recreation of the LightSceneProxy and update RayTracingData accordingly if the texture has changed.
+		// This will trigger a recreation of the LightSceneProxy
 		MarkRenderStateDirty();
 	}
 }
@@ -259,13 +257,6 @@ float URectLightComponent::GetUniformPenumbraSize() const
 
 void URectLightComponent::BeginDestroy()
 {
-	FRectLightRayTracingData* DeletedRenderData = RayTracingData;
-	RayTracingData = nullptr;
-	ENQUEUE_RENDER_COMMAND(DeleteBuildRectLightMipTree)(
-		[DeletedRenderData](FRHICommandListImmediate& RHICmdList)
-	{
-		delete DeletedRenderData;
-	});
 	Super::BeginDestroy();
 }
 
@@ -290,7 +281,6 @@ FRectLightSceneProxy::FRectLightSceneProxy(const URectLightComponent* Component)
 	, SourceHeight(Component->SourceHeight)
 	, BarnDoorAngle(FMath::Clamp(Component->BarnDoorAngle, 0.f, GetRectLightBarnDoorMaxAngle()))
 	, BarnDoorLength(FMath::Max(0.1f, Component->BarnDoorLength))
-	, RayTracingData(Component->RayTracingData)
 	, SourceTexture(Component->SourceTexture)
 	, LightFunctionConeAngleTangent(Component->LightFunctionConeAngle > 0 ? FMath::Tan(FMath::Clamp(Component->LightFunctionConeAngle, 0.0f, 89.0f) * (float)UE_PI / 180.0f) : 0.0f)
 	, SourceTextureScaleOffset(FVector4f(
diff --git a/Engine/Source/Runtime/Engine/Public/RectLightSceneProxy.h b/Engine/Source/Runtime/Engine/Public/RectLightSceneProxy.h
--- a/Engine/Source/Runtime/Engine/Public/RectLightSceneProxy.h
+++ b/Engine/Source/Runtime/Engine/Public/RectLightSceneProxy.h
@@ -11,15 +11,6 @@
 #include "LocalLightSceneProxy.h"
 #include "SceneManagement.h"
 
-struct FRectLightRayTracingData
-{
-#if RHI_RAYTRACING
-	FRWBuffer   RectLightMipTree;
-	FIntVector  RectLightMipTreeDimensions;
-	FGuid		TextureLightingGuid;
-	bool		bInitialised = false;
-#endif
-};
 
 class FRectLightSceneProxy : public FLocalLightSceneProxy
 {
@@ -28,7 +19,6 @@ public:
 	float		SourceHeight;
 	float		BarnDoorAngle;
 	float		BarnDoorLength;
-	FRectLightRayTracingData* RayTracingData; // Render thread data only
 	UTexture*	SourceTexture;
 	uint32		RectAtlasId;
 	float		LightFunctionConeAngleTangent;	// Use Ortho projection if 0
