[Material]
 - Fix race condition manipulating the bitfield from different threads

diff --git a/Engine/Source/Runtime/Engine/Private/Materials/MaterialShared.cpp b/Engine/Source/Runtime/Engine/Private/Materials/MaterialShared.cpp
--- a/Engine/Source/Runtime/Engine/Private/Materials/MaterialShared.cpp
+++ b/Engine/Source/Runtime/Engine/Private/Materials/MaterialShared.cpp
@@ -1343,7 +1343,7 @@ void FMaterial::SetGameThreadShaderMap(FMaterialShaderMap* InMaterialShaderMap)
 	{
 		GameThreadShaderMap->GetResource()->SetOwnerName(GetOwnerFName());
 	}
-	bGameThreadShaderMapIsComplete = bIsComplete;
+	bGameThreadShaderMapIsComplete.store(bIsComplete, std::memory_order_relaxed);
 
 	TRefCountPtr<FMaterial> Material = this;
 	TRefCountPtr<FMaterialShaderMap> ShaderMap = InMaterialShaderMap;
@@ -1362,11 +1362,11 @@ void FMaterial::UpdateInlineShaderMapIsComplete()
 	const bool bSilent = false;
 	const bool bIsComplete = GameThreadShaderMap->IsComplete(this, bSilent);
 
-	bGameThreadShaderMapIsComplete = bIsComplete;
+	bGameThreadShaderMapIsComplete.store(bIsComplete, std::memory_order_relaxed);
 	TRefCountPtr<FMaterial> Material = this;
 	ENQUEUE_RENDER_COMMAND(UpdateGameThreadShaderMapIsComplete)([Material = MoveTemp(Material), bIsComplete](FRHICommandListImmediate& RHICmdList) mutable
 	{
-		Material->bRenderingThreadShaderMapIsComplete = bIsComplete;
+		Material->bRenderingThreadShaderMapIsComplete.store(bIsComplete, std::memory_order_relaxed);
 	});
 }
 
@@ -1382,7 +1382,7 @@ void FMaterial::SetInlineShaderMap(FMaterialShaderMap* InMaterialShaderMap)
 
 	// SetInlineShaderMap is called during PostLoad(), before given UMaterial(Instance) is fully initialized
 	// Can't check for completeness yet
-	bGameThreadShaderMapIsComplete = false;
+	bGameThreadShaderMapIsComplete.store(false, std::memory_order_relaxed);
 	GameThreadShaderMapSubmittedPriority = EShaderCompileJobPriority::None;
 
 	TRefCountPtr<FMaterial> Material = this;
@@ -3245,7 +3245,7 @@ void FMaterial::CacheGivenTypes(EShaderPlatform Platform, const TArray<const FVe
 		return;
 	}
 
-	if (bGameThreadShaderMapIsComplete)
+	if (bGameThreadShaderMapIsComplete.load(std::memory_order_relaxed))
 	{
 		UE_LOG(LogMaterial, Verbose, TEXT("Cache given types for a material resource %s with a complete ShaderMap"), *GetFriendlyName());
 		return;
@@ -5436,11 +5436,11 @@ void FMaterial::SetShaderMapsOnMaterialResources(const TMap<TRefCountPtr<FMateri
 		if (LIKELY(Material->GameThreadShaderMap))
 		{
 			Material->GameThreadShaderMap->GetResource()->SetOwnerName(Material->GetOwnerFName());
-			Material->bGameThreadShaderMapIsComplete = ShaderMap->IsComplete(Material, true);
+			Material->bGameThreadShaderMapIsComplete.store(ShaderMap->IsComplete(Material, true), std::memory_order_relaxed);
 		}
 		else
 		{
-			Material->bGameThreadShaderMapIsComplete = false;
+			Material->bGameThreadShaderMapIsComplete.store(false, std::memory_order_relaxed);
 		}
 	}
 
diff --git a/Engine/Source/Runtime/Engine/Public/MaterialShared.h b/Engine/Source/Runtime/Engine/Public/MaterialShared.h
--- a/Engine/Source/Runtime/Engine/Public/MaterialShared.h
+++ b/Engine/Source/Runtime/Engine/Public/MaterialShared.h
@@ -1997,8 +1997,6 @@ public:
 		FeatureLevel(ERHIFeatureLevel::Num),
 		bContainsInlineShaders(false),
 		bLoadedCookedShaderMapId(false),
-		bGameThreadShaderMapIsComplete(false),
-		bRenderingThreadShaderMapIsComplete(false),
 		bAllowPixelDepthOffset(true)
 	{
 		// this option affects only deferred renderer
@@ -2475,13 +2473,13 @@ public:
 	inline bool IsGameThreadShaderMapComplete() const
 	{
 		checkSlow(IsInGameThread() || IsInParallelGameThread());
-		return bGameThreadShaderMapIsComplete;
+		return bGameThreadShaderMapIsComplete.load(std::memory_order_relaxed);
 	}
 
 	inline bool IsRenderingThreadShaderMapComplete() const
 	{
 		checkSlow(IsInParallelRenderingThread());
-		return bRenderingThreadShaderMapIsComplete;
+		return bRenderingThreadShaderMapIsComplete.load(std::memory_order_relaxed);
 	}
 
 	ENGINE_API void SetRenderingThreadShaderMap(TRefCountPtr<FMaterialShaderMap>& InMaterialShaderMap);
@@ -2783,15 +2781,16 @@ private:
 	uint32 bContainsInlineShaders : 1;
 	uint32 bLoadedCookedShaderMapId : 1;
 
-	uint32 bGameThreadShaderMapIsComplete : 1;
-	uint32 bRenderingThreadShaderMapIsComplete : 1;
-	
 	uint32 bAllowPixelDepthOffset : 1;
 
 #if UE_CHECK_FMATERIAL_LIFETIME
 	/** Set when the owner of this FMaterial (typically a UMaterial or UMaterialInstance) has had BeginDestroy() called */
 	uint32 bOwnerBeginDestroyed : 1;
 #endif // UE_CHECK_FMATERIAL_LIFETIME
+	
+	// These are touched by different threads and need to be outside the bitfield above.
+	std::atomic<bool> bGameThreadShaderMapIsComplete{ false };
+	std::atomic<bool> bRenderingThreadShaderMapIsComplete{ false };
 
 	/** MaterialPSO request IDs currently active for this material */
 	TArray<FMaterialPSOPrecacheRequestID> PrecachedPSORequestIDs;
