[Backout] - CL36894548 Vertex interpolators inherit from CustomOutput node, which made them not usable in MaterialFunctions. Reverting this for now to resolve that.

Original CL Desc
-----------------------------------------------------------------
Remove CustomOutput nodes from the node palette in Material Functions

#tests Windows Editor

diff --git a/Engine/Source/Runtime/Engine/Private/Materials/MaterialExpressions.cpp b/Engine/Source/Runtime/Engine/Private/Materials/MaterialExpressions.cpp
--- a/Engine/Source/Runtime/Engine/Private/Materials/MaterialExpressions.cpp
+++ b/Engine/Source/Runtime/Engine/Private/Materials/MaterialExpressions.cpp
@@ -356,8 +356,7 @@ bool IsAllowedExpressionType(const UClass* const Class, const bool bMaterialFunc
 
 	if (bMaterialFunction)
 	{
-		return bSharedAllowed 
-			&& !(Class && Class->IsChildOf(UMaterialExpressionCustomOutput::StaticClass()));
+		return bSharedAllowed;
 	}
 	else
 	{
