First Person: Implemented support for custom material outputs in MaterialAttributes. This  extends support for the handful of existing supported outputs to now actually compile them without needing to get the custom output attribute at the end of the graph and then pipe it into the corresponding custom output node. This was achieved by extending the FMaterialAttributeDefinitionMap class to store additional meta data about custom outputs and by introducing the concept of attribute bitmasks to it. The list of supported custom outputs can be extended in the future as most custom outputs don't require complex logic in the node.
In addition, the FirstPerson output node (FirstPersonInterpolationAlpha to be precise) is now also supported by MaterialAttributes, enabling workflows where the FirstPerson lerp factor can be set from material functions.

diff --git a/Engine/Source/Runtime/Core/Private/UObject/FortniteMainBranchObjectVersion.cpp b/Engine/Source/Runtime/Core/Private/UObject/FortniteMainBranchObjectVersion.cpp
--- a/Engine/Source/Runtime/Core/Private/UObject/FortniteMainBranchObjectVersion.cpp
+++ b/Engine/Source/Runtime/Core/Private/UObject/FortniteMainBranchObjectVersion.cpp
@@ -14,6 +14,6 @@ TMap<FGuid, FGuid> FFortniteMainBranchObjectVersion::GetSystemGuids()
 	SystemGuids.Add(DevGuids.Niagara_LatestScriptCompileVersion, FGuid("393081597FA7234C92F00869DE262D73"));
 	SystemGuids.Add(DevGuids.SkeletalMeshDerivedDataVersion, FGuid("B9AEAA3EE7AC4FCFB3860E8F3E06DD75"));
 	SystemGuids.Add(DevGuids.STATICMESH_DERIVEDDATA_VER, FGuid("CDFDB82FE0A83DB602AE2EFFD8BB7F42"));
-	SystemGuids.Add(DevGuids.MaterialTranslationDDCVersion, FGuid("77BE52870E484BE78419760A6CD0C677"));
+	SystemGuids.Add(DevGuids.MaterialTranslationDDCVersion, FGuid("D318C2D188414B7284D232E41F77BB74"));
 	return SystemGuids;
 }
diff --git a/Engine/Source/Runtime/Engine/Private/Materials/HLSLMaterialTranslator.cpp b/Engine/Source/Runtime/Engine/Private/Materials/HLSLMaterialTranslator.cpp
--- a/Engine/Source/Runtime/Engine/Private/Materials/HLSLMaterialTranslator.cpp
+++ b/Engine/Source/Runtime/Engine/Private/Materials/HLSLMaterialTranslator.cpp
@@ -1038,6 +1038,83 @@ void FHLSLMaterialTranslator::CompileCustomOutputs(TArray<UMaterialExpressionCus
 	}
 }
 
+uint64 FHLSLMaterialTranslator::CompileMaterialAttributesCustomOutputs(uint64 CustomOutputNodeBitmask)
+{
+	// Some proxies return null for this. But the main one we are interested in doesn't.
+	UMaterialInterface* MatIf = Material->GetMaterialInterface();
+	if (!IsValid(MatIf))
+	{
+		return 0;
+	}
+
+	UMaterial* BaseMaterial = MatIf->GetMaterial();
+	if (!BaseMaterial->bUseMaterialAttributes)
+	{
+		return 0;
+	}
+
+	// Get all expressions in the input chain of the MaterialAttributes output and create a bitmask of all connected attributes, including custom outputs.
+	TArray<UMaterialExpression*> InputExpressions;
+	BaseMaterial->GetExpressionsInPropertyChain(MP_MaterialAttributes, InputExpressions, &StaticParameters, ERHIFeatureLevel::Num, EMaterialQualityLevel::Num, ERHIShadingPath::Num, true /*bInRecurseIntoMaterialFunctions*/);
+	const uint64 AttributesBitmask = FMaterialAttributeDefinitionMap::GetConnectedMaterialAttributesBitmask(InputExpressions);
+
+	// Get a list of all custom output attribute definitions.
+	static TArray<FMaterialCustomOutputAttributeDefintion> CustomAttributes;
+	if (CustomAttributes.IsEmpty())
+	{
+		FMaterialAttributeDefinitionMap::GetCustomAttributeList(CustomAttributes);
+	}
+
+	uint64 ProcessedCustomOutputsMask = 0;
+
+	// Iterate over all supported custom outputs and compile them if they're set in the bitmask. Based on the logic in CompileCustomOutputs().
+	for (const FMaterialCustomOutputAttributeDefintion& CustomAttribute : CustomAttributes)
+	{
+		const uint64 CustomAttributeBit = FMaterialAttributeDefinitionMap::GetBitmask(CustomAttribute.AttributeID);
+		if ((CustomAttributeBit & AttributesBitmask) == 0)
+		{
+			continue;
+		}
+
+		if ((CustomOutputNodeBitmask & FMaterialAttributeDefinitionMap::GetBitmask(CustomAttribute.AttributeID)) != 0)
+		{
+			UE_LOG(LogMaterial, Warning, TEXT("The material sets custom output '%s' using both MaterialAttributes and a custom output node! The value set via MaterialAttributes will be ignored."), *CustomAttribute.AttributeName);
+			continue;
+		}
+
+		const int32 NumOutputs = 1; // TODO: Properly support custom output nodes with more than one output. Currently all of the supported custom outputs just require the define to have a value > 0.
+		const FString FunctionNameUpper = CustomAttribute.FunctionName.ToUpper();
+		ResourcesString += FString::Printf(TEXT("#ifndef NUM_MATERIAL_OUTPUTS_%s\n#define NUM_MATERIAL_OUTPUTS_%s %d\n#endif\n"), *FunctionNameUpper, *FunctionNameUpper, NumOutputs);
+
+		const int32 NumEvaluations = CustomAttribute.bNeedsPreviousFrameEvaluation ? 2 : 1;
+		const bool bCachedCompilingPreviousFrame = bCompilingPreviousFrame;
+
+		for (int32 PreviousFrameEvaluation = 0; PreviousFrameEvaluation < NumEvaluations; PreviousFrameEvaluation++)
+		{
+			{
+				ClearFunctionStack(CustomAttribute.ShaderFrequency);
+				FunctionStacks[CustomAttribute.ShaderFrequency].Add(new FMaterialFunctionCompileState(nullptr));
+			}
+			MaterialProperty = MP_MAX; // Indicates we're not compiling any material property.
+			ShaderFrequency = CustomAttribute.ShaderFrequency;
+			bCompilingPreviousFrame = PreviousFrameEvaluation != 0;
+			TArray<FShaderCodeChunk> CustomExpressionChunks;
+			AssignTempScope(CustomExpressionChunks);
+			int32 CodeChunk = Material->CompileCustomAttribute(CustomAttribute.AttributeID, this);
+			CustomOutput(CustomAttribute.AttributeName, CustomAttribute.FunctionName, CustomAttribute.OutputIndex, CodeChunk);
+		}
+
+		bCompilingPreviousFrame = bCachedCompilingPreviousFrame; // Restore the cached value of bCompilingPreviousFrame
+
+		ClearFunctionStack(CustomAttribute.ShaderFrequency);
+		FunctionStacks[CustomAttribute.ShaderFrequency].Add(new FMaterialFunctionCompileState(nullptr));
+
+		ProcessedCustomOutputsMask |= CustomAttributeBit;
+	}
+
+	return ProcessedCustomOutputsMask;
+}
+
 template<typename ExpressionsArrayType>
 EMaterialExpressionVisitResult FHLSLMaterialTranslator::VisitExpressionsRecursive(const ExpressionsArrayType& Expressions, IMaterialExpressionVisitor& InVisitor)
 {
@@ -1694,75 +1771,13 @@ void FHLSLMaterialTranslator::TranslateMaterial()
 	CHECK_DDC_QUERY_FINISHED_ELSE_RETURN();
 	ResourcesString = TEXT("");
 
-#if HANDLE_CUSTOM_OUTPUTS_AS_MATERIAL_ATTRIBUTES
-	// Handle custom outputs when using material attribute output
-	if (Material->HasMaterialAttributesConnected())
-	{
-		TArray<FMaterialCustomOutputAttributeDefintion> CustomAttributeList;
-		FMaterialAttributeDefinitionMap::GetCustomAttributeList(CustomAttributeList);
-		TArray<FShaderCodeChunk> CustomExpressionChunks;
-
-		for (FMaterialCustomOutputAttributeDefintion& Attribute : CustomAttributeList)
-		{
-			// Compile all outputs for attribute
-			bool bValidResultCompiled = false;
-			int32 NumOutputs = 1;//CustomOutput->GetNumOutputs();
-
-			for (int32 OutputIndex = 0; OutputIndex < NumOutputs; ++OutputIndex)
-			{
-				CHECK_DDC_QUERY_FINISHED_ELSE_RETURN();
-				MaterialProperty = Attribute.Property;
-				ShaderFrequency = Attribute.ShaderFrequency;
-				FunctionStacks[ShaderFrequency].Empty();
-				FunctionStacks[ShaderFrequency].Add(FMaterialFunctionCompileState(nullptr));
-
-				CustomExpressionChunks.Empty();
-				AssignTempScope(CustomExpressionChunks);
-				int32 Result = Material->CompileCustomAttribute(Attribute.AttributeID, this);
-
-				// Consider attribute used if varies from default value
-				if (Result != INDEX_NONE)
-				{
-					bool bValueNonDefault = true;
-
-					if (FMaterialUniformExpression* Expression = GetParameterUniformExpression(Result))
-					{
-						FLinearColor Value;
-						FMaterialRenderContext DummyContext(nullptr, *Material, nullptr);
-						Expression->GetNumberValue(DummyContext, Value);
-
-						bool bEqualValue = Value.R == Attribute.DefaultValue.X;
-						bEqualValue &= Value.G == Attribute.DefaultValue.Y || Attribute.ValueType < MCT_Float2;
-						bEqualValue &= Value.B == Attribute.DefaultValue.Z || Attribute.ValueType < MCT_Float3;
-						bEqualValue &= Value.A == Attribute.DefaultValue.W || Attribute.ValueType < MCT_Float4;
-
-						if (Expression->IsConstant() && bEqualValue)
-						{
-							bValueNonDefault = false;
-						}
-					}
-
-					// Valid, non-default value so generate shader code
-					if (bValueNonDefault)
-					{
-						GenerateCustomAttributeCode(OutputIndex, Result, Attribute.ValueType, Attribute.FunctionName);
-						bValidResultCompiled = true;
-					}
-				}
-			}
+	// Compile all custom output nodes that were not compiled in the first batch.
+	CompileCustomOutputs(CustomOutputExpressions, SeenCustomOutputExpressionsClasses, false);
 
-			// If used, add compile data
-			if (bValidResultCompiled)
-			{
-				ResourcesString += FString::Printf(TEXT("#define NUM_MATERIAL_OUTPUTS_%s %d\n"), *Attribute.FunctionName.ToUpper(), NumOutputs);
-			}
-		}
-	}
-	else
-#endif // #if HANDLE_CUSTOM_OUTPUTS_AS_MATERIAL_ATTRIBUTES
-		{
-			CompileCustomOutputs(CustomOutputExpressions, SeenCustomOutputExpressionsClasses, false);
-		}
+	// Compile custom outputs that are set via MaterialAttributes. Nodes take precedence, so any outputs set via both MaterialAttributes and a node will use the value set on the node.
+	const uint64 CustomOutputNodesBitmask = FMaterialAttributeDefinitionMap::GetCustomOutputNodesBitmask(CustomOutputExpressions);
+	const uint64 ProcessedMaterialAttributesCustomOutputs = CompileMaterialAttributesCustomOutputs(CustomOutputNodesBitmask);
+	const uint64 CustomOutputsBitmask = CustomOutputNodesBitmask | ProcessedMaterialAttributesCustomOutputs;
 
 	// No more calls to non-vertex shader CompilePropertyAndSetMaterialProperty beyond this point
 	const uint32 SavedNumUserTexCoords = GetNumUserTexCoords();
@@ -1805,7 +1820,8 @@ void FHLSLMaterialTranslator::TranslateMaterial()
 		Material->IsTessellationEnabled() &&
 		IsMaterialPropertyUsed(MP_Displacement, Chunk[MP_Displacement], FLinearColor(-1, 0, 0, 0), 1);
 
-	const bool bHasFirstPersonOutput = CustomOutputExpressions.FindByPredicate([](UMaterialExpressionCustomOutput* Expression) { return Expression->IsA<UMaterialExpressionFirstPersonOutput>(); }) != nullptr;
+	static const FGuid FirstPersonInterpolationAlphaGuid = FMaterialAttributeDefinitionMap::GetCustomAttributeID(TEXT("FirstPersonInterpolationAlpha"));
+	const bool bHasFirstPersonOutput = (CustomOutputsBitmask & FMaterialAttributeDefinitionMap::GetBitmask(FirstPersonInterpolationAlphaGuid)) != 0;
 	MaterialCompilationOutput.bModifiesMeshPosition = bUsesPixelDepthOffset || bUsesWorldPositionOffset || bUsesDisplacement || bHasFirstPersonOutput;
 	MaterialCompilationOutput.bUsesWorldPositionOffset = bUsesWorldPositionOffset;
 	MaterialCompilationOutput.bUsesPixelDepthOffset = bUsesPixelDepthOffset;
@@ -15441,7 +15457,7 @@ int32 FHLSLMaterialTranslator::CustomExpression( class UMaterialExpressionCustom
 	return Result;
 }
 
-int32 FHLSLMaterialTranslator::CustomOutput(class UMaterialExpressionCustomOutput* Custom, int32 OutputIndex, int32 OutputCode)
+int32 FHLSLMaterialTranslator::CustomOutput(const FString& CustomOutputName, const FString& FunctionName, int32 OutputIndex, int32 OutputCode)
 {
 	if (MaterialProperty != MP_MAX)
 	{
@@ -15456,7 +15472,7 @@ int32 FHLSLMaterialTranslator::CustomOutput(class UMaterialExpressionCustomOutpu
 	EMaterialValueType OutputType = GetParameterType(OutputCode);
 	if (!IsFloatNumericType(OutputType))
 	{
-		return Errorf(TEXT("Bad type %s for %s"), DescribeType(OutputType), *Custom->GetDescription());
+		return Errorf(TEXT("Bad type %s for %s"), DescribeType(OutputType), *CustomOutputName);
 	}
 
 	FString OutputTypeString = HLSLTypeString(OutputType);
@@ -15473,7 +15489,7 @@ int32 FHLSLMaterialTranslator::CustomOutput(class UMaterialExpressionCustomOutpu
 		GetFixedParameterCode(OutputCode, *CurrentScopeChunks, Definitions, Body, CompiledPDV_FiniteDifferences);
 	}
 
-	const FString FunctionNameBase = FString::Printf(TEXT("%s%d%s"), *Custom->GetFunctionName(), OutputIndex, bCompilingPreviousFrame ? TEXT("Prev") : TEXT(""));
+	const FString FunctionNameBase = FString::Printf(TEXT("%s%d%s"), *FunctionName, OutputIndex, bCompilingPreviousFrame ? TEXT("Prev") : TEXT(""));
 
 	// Primary function will have _LWC suffix if it returns an LWC type
 	// We also define a pre-processor symbol to indicate the custom output function is available so that
@@ -15520,6 +15536,11 @@ int32 FHLSLMaterialTranslator::CustomOutput(class UMaterialExpressionCustomOutpu
 	return INDEX_NONE;
 }
 
+int32 FHLSLMaterialTranslator::CustomOutput(class UMaterialExpressionCustomOutput* Custom, int32 OutputIndex, int32 OutputCode)
+{
+	return CustomOutput(Custom->GetDescription(), Custom->GetFunctionName(), OutputIndex, OutputCode);
+}
+
 int32 FHLSLMaterialTranslator::VirtualTextureOutput(uint8 AttributeMask)
 {
 	MaterialCompilationOutput.bHasRuntimeVirtualTextureOutputNode |= AttributeMask != 0;
@@ -15529,50 +15550,6 @@ int32 FHLSLMaterialTranslator::VirtualTextureOutput(uint8 AttributeMask)
 	return INDEX_NONE;
 }
 
-#if HANDLE_CUSTOM_OUTPUTS_AS_MATERIAL_ATTRIBUTES
-/** Used to translate code for custom output attributes such as ClearCoatBottomNormal */
-void FHLSLMaterialTranslator::GenerateCustomAttributeCode(int32 OutputIndex, int32 OutputCode, EMaterialValueType OutputType, FString& DisplayName)
-{
-	check(MaterialProperty == MP_CustomOutput);
-	check(OutputIndex >= 0 && OutputCode != INDEX_NONE);
-
-	FString OutputTypeString;
-	switch (OutputType)
-	{
-		case MCT_Float:
-		case MCT_Float1:
-			OutputTypeString = TEXT("MaterialFloat");
-			break;
-		case MCT_Float2:
-			OutputTypeString += TEXT("MaterialFloat2");
-			break;
-		case MCT_Float3:
-			OutputTypeString += TEXT("MaterialFloat3");
-			break;
-		case MCT_Float4:
-			OutputTypeString += TEXT("MaterialFloat4");
-			break;
-		default:
-			check(0);
-	}
-
-	FString Definitions;
-	FString Body;
-
-	if ((*CurrentScopeChunks)[OutputCode].UniformExpression && !(*CurrentScopeChunks)[OutputCode].UniformExpression->IsConstant())
-	{
-		Body = GetParameterCode(OutputCode);
-	}
-	else
-	{
-		GetFixedParameterCode(OutputCode, *CurrentScopeChunks, Definitions, Body);
-	}
-
-	FString ImplementationCode = FString::Printf(TEXT("%s %s%d(FMaterial%sParameters Parameters)\n{\n%s return %s;\n}\n"), *OutputTypeString, *DisplayName, OutputIndex, ShaderFrequency == SF_Vertex ? TEXT("Vertex") : TEXT("Pixel"), *Definitions, *Body);
-	CustomOutputImplementations.Add(ImplementationCode);
-}
-#endif
-
 /**
 	* Adds code to return a random value shared by all geometry for any given instanced static mesh
 	*
diff --git a/Engine/Source/Runtime/Engine/Private/Materials/HLSLMaterialTranslator.h b/Engine/Source/Runtime/Engine/Private/Materials/HLSLMaterialTranslator.h
--- a/Engine/Source/Runtime/Engine/Private/Materials/HLSLMaterialTranslator.h
+++ b/Engine/Source/Runtime/Engine/Private/Materials/HLSLMaterialTranslator.h
@@ -628,6 +628,7 @@ public:
 	void GatherCustomVertexInterpolators(const ExpressionsArrayType& Expressions);
 
 	void CompileCustomOutputs(TArray<UMaterialExpressionCustomOutput*>& CustomOutputExpressions, TSet<UClass*>& SeenCustomOutputExpressionsClasses, bool bIsBeforeAttributes);
+	uint64 CompileMaterialAttributesCustomOutputs(uint64 CustomOutputNodeBitmask);
 
 	template<typename ExpressionsArrayType>
 	EMaterialExpressionVisitResult VisitExpressionsRecursive(const ExpressionsArrayType& Expressions, IMaterialExpressionVisitor& InVisitor);
@@ -1265,6 +1266,7 @@ protected:
 
 	virtual int32 CustomExpression(class UMaterialExpressionCustom* Custom, int32 OutputIndex, TArray<int32>& CompiledInputs) override;
 	virtual int32 CustomOutput(class UMaterialExpressionCustomOutput* Custom, int32 OutputIndex, int32 OutputCode) override;
+	int32 CustomOutput(const FString& CustomOutputName, const FString& FunctionName, int32 OutputIndex, int32 OutputCode); // Only used internally and intentionally not exposed on FMaterialCompiler.
 
 	virtual int32 VirtualTextureOutput(uint8 MaterialAttributeMask) override;
 
@@ -1364,11 +1366,6 @@ protected:
 	virtual int32 SubstrateVerticalLayeringParameterBlendingBSDFCoverageToNormalMixCodeChunk(int32 TopCodeChunk) override;
 	virtual int32 SubstrateHorizontalMixingParameterBlendingBSDFCoverageToNormalMixCodeChunk(int32 BackgroundCodeChunk, int32 ForegroundCodeChunk, int32 HorizontalMixCodeChunk) override;
 
-#if HANDLE_CUSTOM_OUTPUTS_AS_MATERIAL_ATTRIBUTES
-	/** Used to translate code for custom output attributes such as ClearCoatBottomNormal */
-	void GenerateCustomAttributeCode(int32 OutputIndex, int32 OutputCode, EMaterialValueType OutputType, FString& DisplayName);
-#endif
-
 	/**
 	 * Adds code to return a random value shared by all geometry for any given instanced static mesh
 	 *
diff --git a/Engine/Source/Runtime/Engine/Private/Materials/Material.cpp b/Engine/Source/Runtime/Engine/Private/Materials/Material.cpp
--- a/Engine/Source/Runtime/Engine/Private/Materials/Material.cpp
+++ b/Engine/Source/Runtime/Engine/Private/Materials/Material.cpp
@@ -313,12 +313,15 @@ int32 FMaterialResource::CompilePropertyAndSetMaterialProperty(EMaterialProperty
 #endif // WITH_EDITOR
 }
 
-#if HANDLE_CUSTOM_OUTPUTS_AS_MATERIAL_ATTRIBUTES
 int32 FMaterialResource::CompileCustomAttribute(const FGuid& AttributeID, FMaterialCompiler* Compiler) const
 {
+#if WITH_EDITOR
 	return Material->CompilePropertyEx(Compiler, AttributeID);
+#else // WITH_EDITOR
+	check(0); // This is editor-only function
+	return INDEX_NONE;
+#endif // WITH_EDITOR
 }
-#endif
 
 #if WITH_EDITORONLY_DATA
 void FMaterialResource::GatherCustomOutputExpressions(TArray<UMaterialExpressionCustomOutput*>& OutCustomOutputs) const
diff --git a/Engine/Source/Runtime/Engine/Private/Materials/MaterialAttributeDefinitionMap.cpp b/Engine/Source/Runtime/Engine/Private/Materials/MaterialAttributeDefinitionMap.cpp
--- a/Engine/Source/Runtime/Engine/Private/Materials/MaterialAttributeDefinitionMap.cpp
+++ b/Engine/Source/Runtime/Engine/Private/Materials/MaterialAttributeDefinitionMap.cpp
@@ -7,6 +7,14 @@
 #include "Materials/Material.h"
 #include "Materials/MaterialExpressionShadingModel.h"
 #include "Materials/MaterialExpressionSubstrate.h"
+#include "Materials/MaterialExpressionSetMaterialAttributes.h"
+#include "Materials/MaterialExpressionMakeMaterialAttributes.h"
+#include "Materials/MaterialExpressionMaterialFunctionCall.h"
+#include "Materials/MaterialExpressionBentNormalCustomOutput.h"
+#include "Materials/MaterialExpressionClearCoatNormalCustomOutput.h"
+#include "Materials/MaterialExpressionTangentOutput.h"
+#include "Materials/MaterialExpressionThinTranslucentMaterialOutput.h"
+#include "Materials/MaterialExpressionFirstPersonOutput.h"
 #include "Serialization/ShaderKeyGenerator.h"
 #include "SubstrateDefinitions.h"
 
@@ -101,10 +109,13 @@ int32 FMaterialAttributeDefintion::CompileDefaultValue(FMaterialCompiler* Compil
 //////////////////////////////////////////////////////////////////////////
 
 FMaterialCustomOutputAttributeDefintion::FMaterialCustomOutputAttributeDefintion(
-	const FGuid& InAttributeID, const FString& InAttributeName, const FString& InFunctionName, EMaterialProperty InProperty,
+	const FGuid& InAttributeID, UClass* InExpressionClass, const FString& InAttributeName, const FString& InFunctionName, int32 InOutputIndex, bool bInNeedsPreviousFrameEvaluation, EMaterialProperty InProperty,
 	EMaterialValueType InValueType, const FVector4& InDefaultValue, EShaderFrequency InShaderFrequency, MaterialAttributeBlendFunction InBlendFunction /*= nullptr*/)
 	: FMaterialAttributeDefintion(InAttributeID, InAttributeName, InProperty, InValueType, InDefaultValue, InShaderFrequency, INDEX_NONE, false, InBlendFunction)
+	, CustomOutputExpressionStaticClass(InExpressionClass)
 	, FunctionName(InFunctionName)
+	, OutputIndex(InOutputIndex)
+	, bNeedsPreviousFrameEvaluation(bInNeedsPreviousFrameEvaluation)
 {
 }
 
@@ -214,6 +225,67 @@ EShaderFrequency FMaterialAttributeDefinitionMap::GetShaderFrequency(const FGuid
 	return Attribute->ShaderFrequency;
 }
 
+int32 FMaterialAttributeDefinitionMap::GetIndex(EMaterialProperty Property)
+{
+	if (Property == MP_CustomOutput || Property == MP_MAX)
+	{
+		return INDEX_NONE;
+	}
+	return static_cast<int32>(Property);
+}
+
+int32 FMaterialAttributeDefinitionMap::GetIndex(const FGuid& AttributeID)
+{
+	if (FMaterialAttributeDefintion* Attribute = GMaterialPropertyAttributesMap.Find(AttributeID))
+	{
+		if (Attribute->Property == MP_CustomOutput)
+		{
+			const int32 CustomIndex = GMaterialPropertyAttributesMap.CustomAttributes.IndexOfByPredicate([Attribute](const FMaterialCustomOutputAttributeDefintion& A) { return A.AttributeID == Attribute->AttributeID; });
+			const int32 ResultIndex = MP_MAX + CustomIndex;
+			return ResultIndex;
+		}
+		else if (Attribute->Property != MP_MAX)
+		{
+			return static_cast<int32>(Attribute->Property);
+		}
+	}
+	return INDEX_NONE;
+}
+
+uint64 FMaterialAttributeDefinitionMap::GetBitmask(EMaterialProperty Property)
+{
+	const int32 Index = FMaterialAttributeDefinitionMap::GetIndex(Property);
+	if (Index != INDEX_NONE)
+	{
+		const uint64 Result = 1ull << static_cast<uint32>(Index);
+		return Result;
+	}
+	return 0;
+}
+
+uint64 FMaterialAttributeDefinitionMap::GetBitmask(const FGuid& AttributeID)
+{
+	const int32 Index = FMaterialAttributeDefinitionMap::GetIndex(AttributeID);
+	if (Index != INDEX_NONE)
+	{
+		const uint64 Result = 1ull << static_cast<uint32>(Index);
+		return Result;
+	}
+	return 0;
+}
+
+bool FMaterialAttributeDefinitionMap::IsAttributeInBitmask(uint64 Bitmask, EMaterialProperty Property)
+{
+	const uint64 AttributeMask = GetBitmask(Property);
+	return (AttributeMask & Bitmask) != 0;
+}
+
+bool FMaterialAttributeDefinitionMap::IsAttributeInBitmask(uint64 Bitmask, const FGuid& AttributeID)
+{
+	const uint64 AttributeMask = GetBitmask(AttributeID);
+	return (AttributeMask & Bitmask) != 0;
+}
+
 /** Returns the attribute ID for a matching material property */
 FGuid FMaterialAttributeDefinitionMap::GetID(EMaterialProperty Property)
 {
@@ -304,11 +376,12 @@ void FMaterialAttributeDefinitionMap::InitializeAttributeMap()
 	Add(FGuid(0xA0119D44, 0xC456450D, 0x9C39C933, 0x1F72D8D1), TEXT("TessellationMultiplier"), MP_TessellationMultiplier_DEPRECATED, MCT_Float, FVector4(1, 0, 0, 0), SF_Vertex, INDEX_NONE, bHideAttribute);
 
 	// UMaterialExpression custom outputs
-	AddCustomAttribute(FGuid(0xfbd7b46e, 0xb1234824, 0xbde76b23, 0x609f984c), "BentNormal", "GetBentNormal", MCT_Float3, FVector4(0, 0, 1, 0));
-	AddCustomAttribute(FGuid(0xAA3D5C04, 0x16294716, 0xBBDEC869, 0x6A27DD72), "ClearCoatBottomNormal", "ClearCoatBottomNormal", MCT_Float3, FVector4(0, 0, 1, 0));
-	AddCustomAttribute(FGuid(0x8EAB2CB2, 0x73634A24, 0x8CD14F47, 0x3F9C8E55), "CustomEyeTangent", "GetTangentOutput", MCT_Float3, FVector4(0, 0, 0, 0));
-	AddCustomAttribute(FGuid(0xF2D8C70E, 0x42ECA0D1, 0x4652D0AD, 0xB785A065), "TransmittanceColor", "GetThinTranslucentMaterialOutput0", MCT_Float3, FVector4(0.5, 0.5, 0.5, 0));
-	AddCustomAttribute(FGuid(0xAAA1D0E0, 0xE312458E, 0xF3556BAD, 0xA01277CD), "ThinTranslucentSurfaceCoverage", "GetThinTranslucentMaterialOutput1", MCT_Float, FVector4(1, 0, 0, 0));
+	AddCustomAttribute(FGuid(0xfbd7b46e, 0xb1234824, 0xbde76b23, 0x609f984c), UMaterialExpressionBentNormalCustomOutput::StaticClass(), "BentNormal", "GetBentNormal", 0, MCT_Float3, FVector4(0, 0, 1, 0), SF_Pixel);
+	AddCustomAttribute(FGuid(0xAA3D5C04, 0x16294716, 0xBBDEC869, 0x6A27DD72), UMaterialExpressionClearCoatNormalCustomOutput::StaticClass(), "ClearCoatBottomNormal", "ClearCoatBottomNormal", 0, MCT_Float3, FVector4(0, 0, 1, 0), SF_Pixel);
+	AddCustomAttribute(FGuid(0x8EAB2CB2, 0x73634A24, 0x8CD14F47, 0x3F9C8E55), UMaterialExpressionTangentOutput::StaticClass(), "CustomEyeTangent", "GetTangentOutput", 0, MCT_Float3, FVector4(0, 0, 0, 0), SF_Pixel);
+	AddCustomAttribute(FGuid(0xF2D8C70E, 0x42ECA0D1, 0x4652D0AD, 0xB785A065), UMaterialExpressionThinTranslucentMaterialOutput::StaticClass(), "TransmittanceColor", "GetThinTranslucentMaterialOutput", 0, MCT_Float3, FVector4(0.5, 0.5, 0.5, 0), SF_Pixel);
+	AddCustomAttribute(FGuid(0xAAA1D0E0, 0xE312458E, 0xF3556BAD, 0xA01277CD), UMaterialExpressionThinTranslucentMaterialOutput::StaticClass(), "ThinTranslucentSurfaceCoverage", "GetThinTranslucentMaterialOutput", 1, MCT_Float, FVector4(1, 0, 0, 0), SF_Pixel);
+	AddCustomAttribute(FGuid(0xDB380E51, 0xBDED47B6, 0xB185A8BE, 0x8EACBFA4), UMaterialExpressionFirstPersonOutput::StaticClass(), "FirstPersonInterpolationAlpha", "GetFirstPersonOutput", 0, MCT_Float, FVector4(1, 0, 0, 0), SF_Vertex, true /*bNeedsPreviousFrameEvaluation*/);
 }
 
 void FMaterialAttributeDefinitionMap::Add(const FGuid& AttributeID, const FString& AttributeName, EMaterialProperty Property,
@@ -529,12 +602,22 @@ void FMaterialAttributeDefinitionMap::AppendDDCKey(FShaderKeyGenerator& KeyGen)
 	KeyGen.Append(DDCString);
 }
 
-void FMaterialAttributeDefinitionMap::AddCustomAttribute(const FGuid& AttributeID, const FString& AttributeName, const FString& FunctionName, EMaterialValueType ValueType, const FVector4& DefaultValue, MaterialAttributeBlendFunction BlendFunction /*= nullptr*/)
+void FMaterialAttributeDefinitionMap::AddCustomAttribute(
+	const FGuid& AttributeID,
+	UClass* ExpressionClass,
+	const FString& AttributeName,
+	const FString& FunctionName,
+	int32 OutputIndex,
+	EMaterialValueType ValueType,
+	const FVector4& DefaultValue, 
+	EShaderFrequency ShaderFrequency,
+	bool bNeedsPreviousFrameEvaluation,
+	MaterialAttributeBlendFunction BlendFunction /*= nullptr*/)
 {
 	// Make sure that we init CustomAttributes before DDCString is initialized (before first shader load)
 	check(GMaterialPropertyAttributesMap.AttributeDDCString.Len() == 0);
 
-	FMaterialCustomOutputAttributeDefintion UserAttribute(AttributeID, AttributeName, FunctionName, MP_CustomOutput, ValueType, DefaultValue, SF_Pixel, BlendFunction);
+	FMaterialCustomOutputAttributeDefintion UserAttribute(AttributeID, ExpressionClass, AttributeName, FunctionName, OutputIndex, bNeedsPreviousFrameEvaluation, MP_CustomOutput, ValueType, DefaultValue, ShaderFrequency, BlendFunction);
 #if DO_CHECK
 	for (auto& Attribute : GMaterialPropertyAttributesMap.AttributeMap)
 	{
@@ -598,6 +681,68 @@ void FMaterialAttributeDefinitionMap::GetCustomAttributeList(TArray<FMaterialCus
 	}
 }
 
+#if WITH_EDITOR
+uint64 FMaterialAttributeDefinitionMap::GetConnectedMaterialAttributesBitmask(TConstArrayView<TObjectPtr<UMaterialExpression>> Expressions)
+{
+	return GetConnectedMaterialAttributesBitmask_Internal(Expressions);
+}
+
+uint64 FMaterialAttributeDefinitionMap::GetConnectedMaterialAttributesBitmask(TConstArrayView<UMaterialExpression*> Expressions)
+{
+	return GetConnectedMaterialAttributesBitmask_Internal(Expressions);
+}
+
+template<typename ExpressionsArrayType>
+inline uint64 FMaterialAttributeDefinitionMap::GetConnectedMaterialAttributesBitmask_Internal(const ExpressionsArrayType& Expressions)
+{
+	uint64 Out = 0ull;
+	for (const UMaterialExpression* Expression : Expressions)
+	{
+		if (Expression)
+		{
+			if (Expression->IsA(UMaterialExpressionSetMaterialAttributes::StaticClass()))
+			{
+				const UMaterialExpressionSetMaterialAttributes* Attr = Cast<UMaterialExpressionSetMaterialAttributes>(Expression);
+				Out |= Attr->GetConnectedInputs();
+			}
+			else if (Expression->IsA(UMaterialExpressionMakeMaterialAttributes::StaticClass()))
+			{
+				const UMaterialExpressionMakeMaterialAttributes* Attr = Cast<UMaterialExpressionMakeMaterialAttributes>(Expression);
+				Out |= Attr->GetConnectedInputs();
+			}
+			else if (Expression->IsA(UMaterialExpressionMaterialFunctionCall::StaticClass()))
+			{
+				const UMaterialExpressionMaterialFunctionCall* Attr = Cast<UMaterialExpressionMaterialFunctionCall>(Expression);
+				if (Attr->MaterialFunction)
+				{
+					Out |= GetConnectedMaterialAttributesBitmask(Attr->MaterialFunction->GetExpressions());
+				}
+			}
+		}
+	}
+	return Out;
+}
+
+uint64 FMaterialAttributeDefinitionMap::GetCustomOutputNodesBitmask(TConstArrayView<UMaterialExpressionCustomOutput*> Expressions)
+{
+	uint64 Out = 0ull;
+	for (const UMaterialExpressionCustomOutput* Expression : Expressions)
+	{
+		if (Expression)
+		{
+			for (auto& Attribute : GMaterialPropertyAttributesMap.CustomAttributes)
+			{
+				if (Expression->IsA(Attribute.CustomOutputExpressionStaticClass))
+				{
+					Out |= GetBitmask(Attribute.AttributeID);
+				}
+			}
+		}
+	}
+	return Out;
+}
+#endif // WITH_EDITOR
+
 void FMaterialAttributeDefinitionMap::GetAttributeNameToIDList(TArray<TPair<FString, FGuid>>& NameToIDList)
 {
 	NameToIDList.Empty(GMaterialPropertyAttributesMap.OrderedVisibleAttributeList.Num());
diff --git a/Engine/Source/Runtime/Engine/Private/Materials/MaterialCachedData.cpp b/Engine/Source/Runtime/Engine/Private/Materials/MaterialCachedData.cpp
--- a/Engine/Source/Runtime/Engine/Private/Materials/MaterialCachedData.cpp
+++ b/Engine/Source/Runtime/Engine/Private/Materials/MaterialCachedData.cpp
@@ -331,6 +331,7 @@ void FMaterialCachedExpressionData::UpdateForLayerFunctions(const FMaterialCache
 void FMaterialCachedExpressionData::UpdateForExpressions(const FMaterialCachedExpressionContext& Context, TConstArrayView<TObjectPtr<UMaterialExpression>> Expressions, EMaterialParameterAssociation Association, int32 ParameterIndex)
 {
 	check(EditorOnlyData);
+	static const FGuid FirstPersonInterpolationAlphaGuid = FMaterialAttributeDefinitionMap::GetCustomAttributeID(TEXT("FirstPersonInterpolationAlpha"));
 
 	for (UMaterialExpression* Expression : Expressions)
 	{
@@ -526,6 +527,10 @@ void FMaterialCachedExpressionData::UpdateForExpressions(const FMaterialCachedEx
 				if (AttributeInput.Expression)
 				{
 					SetPropertyConnected(MaterialProperty);
+					if (Guid == FirstPersonInterpolationAlphaGuid)
+					{
+						bHasFirstPersonOutput = true;
+					}
 				}
 			}
 		}
diff --git a/Engine/Source/Runtime/Engine/Private/Materials/MaterialExpressionSubstrate.cpp b/Engine/Source/Runtime/Engine/Private/Materials/MaterialExpressionSubstrate.cpp
--- a/Engine/Source/Runtime/Engine/Private/Materials/MaterialExpressionSubstrate.cpp
+++ b/Engine/Source/Runtime/Engine/Private/Materials/MaterialExpressionSubstrate.cpp
@@ -3672,68 +3672,10 @@ void UMaterialExpressionSubstrateThinFilm::GetExpressionToolTip(TArray<FString>&
 
 // Return a conservative list of connected material attribute inputs
 #if WITH_EDITOR
-static uint64 GetConnectedMaterialAttributesInputs(TConstArrayView<TObjectPtr<UMaterialExpression>> Expressions)
-{
-	uint64 Out = 0ull;
-	for (const UMaterialExpression* Expression : Expressions)
-	{
-		if (Expression)
-		{
-			if (Expression->IsA(UMaterialExpressionSetMaterialAttributes::StaticClass()))
-			{
-				const UMaterialExpressionSetMaterialAttributes* Attr = Cast<UMaterialExpressionSetMaterialAttributes>(Expression);
-				Out |= Attr->GetConnectedInputs();
-			}
-			else if (Expression->IsA(UMaterialExpressionMakeMaterialAttributes::StaticClass()))
-			{
-				const UMaterialExpressionMakeMaterialAttributes* Attr = Cast<UMaterialExpressionMakeMaterialAttributes>(Expression);
-				Out |= Attr->GetConnectedInputs();
-			}
-			else if (Expression->IsA(UMaterialExpressionMaterialFunctionCall::StaticClass()))
-			{
-				const UMaterialExpressionMaterialFunctionCall* Attr = Cast<UMaterialExpressionMaterialFunctionCall>(Expression);
-				if (Attr->MaterialFunction)
-				{
-					Out |= GetConnectedMaterialAttributesInputs(Attr->MaterialFunction->GetExpressions());
-				}
-			}
-		}
-	}
-	return Out;
-}
-
 static uint64 GetConnectedMaterialAttributesInputs(const UMaterial* InMaterial)
 {
 	if (!InMaterial) return 0;
-	return GetConnectedMaterialAttributesInputs(InMaterial->GetExpressions());
-}
-
-static bool IsCustomMaterialAttributeInputConnected(uint64 InCache, FGuid InProperty)
-{
-	static TArray<FMaterialCustomOutputAttributeDefintion> CustomAttributeList;
-	if (CustomAttributeList.Num() == 0)
-	{
-		FMaterialAttributeDefinitionMap::GetCustomAttributeList(CustomAttributeList);
-	}
-
-	const EMaterialProperty Property = FMaterialAttributeDefinitionMap::GetProperty(InProperty);
-	if (Property == MP_CustomOutput)
-	{
-		if (const FMaterialCustomOutputAttributeDefintion* Attribute = FMaterialAttributeDefinitionMap::GetCustomAttribute(InProperty))
-		{
-			const uint32 CustomIndex = CustomAttributeList.IndexOfByPredicate([Attribute](const FMaterialCustomOutputAttributeDefintion& A) { return A.AttributeID == Attribute->AttributeID; });
-			if (CustomIndex != INDEX_NONE)
-			{
-				return !!(InCache & (1ull << (MP_MAX + CustomIndex)));
-			}
-		}
-	}
-	return false;
-}
-
-static bool IsMaterialAttributeInputConnected(uint64 InCache, EMaterialProperty InProperty)
-{
-	return !!(InCache & (1ull << uint64(InProperty)));
+	return FMaterialAttributeDefinitionMap::GetConnectedMaterialAttributesBitmask(InMaterial->GetExpressions());
 }
 #endif
 
@@ -3787,7 +3729,7 @@ int32 UMaterialExpressionSubstrateConvertMaterialAttributes::Compile(class FMate
 	// used in SubstrateGetBSDFSharedBasis cannot be relied on for smooth tangent used for lighting on any mesh.
 
 	const uint64 Cached = GetConnectedMaterialAttributesInputs(Material);
-	const bool bHasAnisotropy = IsMaterialAttributeInputConnected(Cached, MP_Anisotropy);
+	const bool bHasAnisotropy = FMaterialAttributeDefinitionMap::IsAttributeInBitmask(Cached, MP_Anisotropy);
 
 	// Regular normal basis
 	int32 NormalCodeChunk = MaterialAttributes.CompileWithDefault(Compiler, FMaterialAttributeDefinitionMap::GetID(MP_Normal));
@@ -3810,7 +3752,7 @@ int32 UMaterialExpressionSubstrateConvertMaterialAttributes::Compile(class FMate
 	const FSubstrateRegisteredSharedLocalBasis NewRegisteredSharedLocalBasis = SubstrateCompilationInfoCreateSharedLocalBasis(Compiler, NormalCodeChunk, TangentCodeChunk);
 	const FString BasisIndexMacro = Compiler->GetSubstrateSharedLocalBasisIndexMacro(NewRegisteredSharedLocalBasis);
 
-	const bool bHasCoatNormal = IsCustomMaterialAttributeInputConnected(Cached, ClearCoatBottomNormalGuid);
+	const bool bHasCoatNormal = FMaterialAttributeDefinitionMap::IsAttributeInBitmask(Cached, ClearCoatBottomNormalGuid);
 	// Clear coat normal basis
 	int32 ClearCoat_NormalCodeChunk = INDEX_NONE;
 	int32 ClearCoat_TangentCodeChunk = INDEX_NONE;
@@ -3833,7 +3775,7 @@ int32 UMaterialExpressionSubstrateConvertMaterialAttributes::Compile(class FMate
 
 	// Custom tangent. No need to register it as a local basis, as it is only used for eye shading internal conversion
 	int32 CustomTangent_TangentCodeChunk = INDEX_NONE;
-	const bool bHasCustomTangent = IsCustomMaterialAttributeInputConnected(Cached, CustomEyeTangentGuid);
+	const bool bHasCustomTangent = FMaterialAttributeDefinitionMap::IsAttributeInBitmask(Cached, CustomEyeTangentGuid);
 	if (bHasCustomTangent)
 	{
 		// Legacy code doesn't do tangent <-> world basis conversion on tangent output, when provided.
@@ -3880,7 +3822,7 @@ int32 UMaterialExpressionSubstrateConvertMaterialAttributes::Compile(class FMate
 	}
 
 	// Transmittance Color
-	const bool bHasTransmittanceColor = IsCustomMaterialAttributeInputConnected(Cached, TransmittanceColorGuid);
+	const bool bHasTransmittanceColor = FMaterialAttributeDefinitionMap::IsAttributeInBitmask(Cached, TransmittanceColorGuid);
 	int32 TransmittanceColorChunk = INDEX_NONE;
 	if (bHasTransmittanceColor)
 	{
@@ -3891,7 +3833,7 @@ int32 UMaterialExpressionSubstrateConvertMaterialAttributes::Compile(class FMate
 		TransmittanceColorChunk = Compiler->Constant3(0.5f, 0.5f, 0.5f);
 	}
 	// Thin Translucent Surface Coverage
-	const bool bHasThinTranslucentSurfaceCoverage = IsCustomMaterialAttributeInputConnected(Cached, ThinTranslucentSurfaceCoverageGuid);
+	const bool bHasThinTranslucentSurfaceCoverage = FMaterialAttributeDefinitionMap::IsAttributeInBitmask(Cached, ThinTranslucentSurfaceCoverageGuid);
 	int32 ThinTranslucentSurfaceCoverageChunk = INDEX_NONE;
 	if (bHasThinTranslucentSurfaceCoverage)
 	{
@@ -3903,7 +3845,7 @@ int32 UMaterialExpressionSubstrateConvertMaterialAttributes::Compile(class FMate
 	}
 
 	int32 ShadingModelCodeChunk = MaterialAttributes.CompileWithDefault(Compiler, FMaterialAttributeDefinitionMap::GetID(MP_ShadingModel));
-	const bool bHasShadingModelExpression = IsMaterialAttributeInputConnected(Cached, MP_ShadingModel)
+	const bool bHasShadingModelExpression = FMaterialAttributeDefinitionMap::IsAttributeInBitmask(Cached, MP_ShadingModel)
 		|| ShadingModelOverride == MSM_FromMaterialExpression; // In this case, rely on the default compilation to return DefaultLit.
 	if (!bHasShadingModelExpression)
 	{
@@ -4021,20 +3963,20 @@ void UMaterialExpressionSubstrateConvertMaterialAttributes::GatherSubstrateMater
 {
 	const uint64 Cached = GetConnectedMaterialAttributesInputs(Material);
 
-	if (IsMaterialAttributeInputConnected(Cached, MP_BaseColor)) { SubstrateMaterialInfo.AddPropertyConnected(MP_BaseColor); }
-	if (IsMaterialAttributeInputConnected(Cached, MP_Metallic)) { SubstrateMaterialInfo.AddPropertyConnected(MP_Metallic); }
-	if (IsMaterialAttributeInputConnected(Cached, MP_Specular)) { SubstrateMaterialInfo.AddPropertyConnected(MP_Specular); }
-	if (IsMaterialAttributeInputConnected(Cached, MP_Roughness)) { SubstrateMaterialInfo.AddPropertyConnected(MP_Roughness); }
-	if (IsMaterialAttributeInputConnected(Cached, MP_Anisotropy)) { SubstrateMaterialInfo.AddPropertyConnected(MP_Anisotropy); }
-	if (IsMaterialAttributeInputConnected(Cached, MP_EmissiveColor)) { SubstrateMaterialInfo.AddPropertyConnected(MP_EmissiveColor); }
-	if (IsMaterialAttributeInputConnected(Cached, MP_Normal)) { SubstrateMaterialInfo.AddPropertyConnected(MP_Normal); }
-	if (IsMaterialAttributeInputConnected(Cached, MP_Tangent)) { SubstrateMaterialInfo.AddPropertyConnected(MP_Tangent); }
-	if (IsMaterialAttributeInputConnected(Cached, MP_SubsurfaceColor)) { SubstrateMaterialInfo.AddPropertyConnected(MP_SubsurfaceColor); }
-	if (IsMaterialAttributeInputConnected(Cached, MP_CustomData0)) { SubstrateMaterialInfo.AddPropertyConnected(MP_CustomData0); }
-	if (IsMaterialAttributeInputConnected(Cached, MP_CustomData1)) { SubstrateMaterialInfo.AddPropertyConnected(MP_CustomData1); }
-	if (IsMaterialAttributeInputConnected(Cached, MP_Opacity)) { SubstrateMaterialInfo.AddPropertyConnected(MP_Opacity); }
-
-	if (IsMaterialAttributeInputConnected(Cached, MP_ShadingModel) || ShadingModelOverride == MSM_FromMaterialExpression)
+	if (FMaterialAttributeDefinitionMap::IsAttributeInBitmask(Cached, MP_BaseColor)) { SubstrateMaterialInfo.AddPropertyConnected(MP_BaseColor); }
+	if (FMaterialAttributeDefinitionMap::IsAttributeInBitmask(Cached, MP_Metallic)) { SubstrateMaterialInfo.AddPropertyConnected(MP_Metallic); }
+	if (FMaterialAttributeDefinitionMap::IsAttributeInBitmask(Cached, MP_Specular)) { SubstrateMaterialInfo.AddPropertyConnected(MP_Specular); }
+	if (FMaterialAttributeDefinitionMap::IsAttributeInBitmask(Cached, MP_Roughness)) { SubstrateMaterialInfo.AddPropertyConnected(MP_Roughness); }
+	if (FMaterialAttributeDefinitionMap::IsAttributeInBitmask(Cached, MP_Anisotropy)) { SubstrateMaterialInfo.AddPropertyConnected(MP_Anisotropy); }
+	if (FMaterialAttributeDefinitionMap::IsAttributeInBitmask(Cached, MP_EmissiveColor)) { SubstrateMaterialInfo.AddPropertyConnected(MP_EmissiveColor); }
+	if (FMaterialAttributeDefinitionMap::IsAttributeInBitmask(Cached, MP_Normal)) { SubstrateMaterialInfo.AddPropertyConnected(MP_Normal); }
+	if (FMaterialAttributeDefinitionMap::IsAttributeInBitmask(Cached, MP_Tangent)) { SubstrateMaterialInfo.AddPropertyConnected(MP_Tangent); }
+	if (FMaterialAttributeDefinitionMap::IsAttributeInBitmask(Cached, MP_SubsurfaceColor)) { SubstrateMaterialInfo.AddPropertyConnected(MP_SubsurfaceColor); }
+	if (FMaterialAttributeDefinitionMap::IsAttributeInBitmask(Cached, MP_CustomData0)) { SubstrateMaterialInfo.AddPropertyConnected(MP_CustomData0); }
+	if (FMaterialAttributeDefinitionMap::IsAttributeInBitmask(Cached, MP_CustomData1)) { SubstrateMaterialInfo.AddPropertyConnected(MP_CustomData1); }
+	if (FMaterialAttributeDefinitionMap::IsAttributeInBitmask(Cached, MP_Opacity)) { SubstrateMaterialInfo.AddPropertyConnected(MP_Opacity); }
+
+	if (FMaterialAttributeDefinitionMap::IsAttributeInBitmask(Cached, MP_ShadingModel) || ShadingModelOverride == MSM_FromMaterialExpression)
 	{
 		SubstrateMaterialInfo.AddPropertyConnected(MP_ShadingModel);
 
@@ -4072,7 +4014,7 @@ FSubstrateOperator* UMaterialExpressionSubstrateConvertMaterialAttributes::Subst
 	// Note Thickness has no meaning/usage in the context of SubstrateLegacyConversionNode
 	int32 ThicknessIndex = Compiler->SubstrateThicknessStackGetThicknessIndex();
 
-	const bool bHasAnisotropy = IsMaterialAttributeInputConnected(Cached, MP_Anisotropy);
+	const bool bHasAnisotropy = FMaterialAttributeDefinitionMap::IsAttributeInBitmask(Cached, MP_Anisotropy);
 
 	auto AddDefaultWorstCase = [&](uint32 In)
 		{
@@ -4088,7 +4030,7 @@ FSubstrateOperator* UMaterialExpressionSubstrateConvertMaterialAttributes::Subst
 	FMaterialShadingModelField ShadingModels = Compiler->GetMaterialShadingModels();
 
 	// Logic about shading models and complexity should match UMaterialExpressionSubstrateConvertMaterialAttributes::Compile.
-	const bool bHasShadingModelFromExpression = IsMaterialAttributeInputConnected(Cached, MP_ShadingModel) || ShadingModelOverride == MSM_FromMaterialExpression; // We keep HasShadingModelFromExpression in case all shading models cannot be safely recovered from material functions.
+	const bool bHasShadingModelFromExpression = FMaterialAttributeDefinitionMap::IsAttributeInBitmask(Cached, MP_ShadingModel) || ShadingModelOverride == MSM_FromMaterialExpression; // We keep HasShadingModelFromExpression in case all shading models cannot be safely recovered from material functions.
 	if ((ShadingModels.CountShadingModels() > 1) || bHasShadingModelFromExpression)
 	{
 		uint32 Features = ESubstrateBsdfFeature_SSS | ESubstrateBsdfFeature_Fuzz;
diff --git a/Engine/Source/Runtime/Engine/Private/Materials/MaterialExpressions.cpp b/Engine/Source/Runtime/Engine/Private/Materials/MaterialExpressions.cpp
--- a/Engine/Source/Runtime/Engine/Private/Materials/MaterialExpressions.cpp
+++ b/Engine/Source/Runtime/Engine/Private/Materials/MaterialExpressions.cpp
@@ -7715,32 +7715,12 @@ UMaterialExpressionSetMaterialAttributes::UMaterialExpressionSetMaterialAttribut
 #if WITH_EDITOR
 uint64 UMaterialExpressionSetMaterialAttributes::GetConnectedInputs() const
 {
-	static TArray<FMaterialCustomOutputAttributeDefintion> CustomAttributeList;
-	if (CustomAttributeList.Num() == 0)
-	{
-		FMaterialAttributeDefinitionMap::GetCustomAttributeList(CustomAttributeList);
-	}
-
 	uint64 Out = 0ull;
 	const int32 NumInputPins = AttributeSetTypes.Num();
 	for (int32 i = 0; i < NumInputPins; ++i)
 	{
-		const EMaterialProperty Property = FMaterialAttributeDefinitionMap::GetProperty(AttributeSetTypes[i]);
-		if (Property == MP_CustomOutput)
-		{
-			if (const FMaterialCustomOutputAttributeDefintion* Attribute = FMaterialAttributeDefinitionMap::GetCustomAttribute(AttributeSetTypes[i]))
-			{
-				const uint32 CustomIndex = CustomAttributeList.IndexOfByPredicate([Attribute](const FMaterialCustomOutputAttributeDefintion& A) { return A.AttributeID == Attribute->AttributeID; });
-				if (CustomIndex != INDEX_NONE)
-				{
-					Out |= 1ull << (MP_MAX + CustomIndex);
-				}
-			}
-		}
-		else if (Property != MP_MAX)
-		{
-			Out |= 1ull << uint64(Property);
-		}
+		const uint64 Bitmask = FMaterialAttributeDefinitionMap::GetBitmask(AttributeSetTypes[i]);
+		Out |= Bitmask;
 	}
 	return Out;
 }
diff --git a/Engine/Source/Runtime/Engine/Public/MaterialShared.h b/Engine/Source/Runtime/Engine/Public/MaterialShared.h
--- a/Engine/Source/Runtime/Engine/Public/MaterialShared.h
+++ b/Engine/Source/Runtime/Engine/Public/MaterialShared.h
@@ -121,8 +121,6 @@ template <class ElementType> class TLinkedList;
 #define ME_STD_LABEL_PAD		16
 #define ME_STD_TAB_HEIGHT		21
 
-#define HANDLE_CUSTOM_OUTPUTS_AS_MATERIAL_ATTRIBUTES 0
-
 // disallow debug data in shipping or on non-desktop Test
 #define ALLOW_SHADERMAP_DEBUG_DATA (!(UE_BUILD_SHIPPING || (UE_BUILD_TEST && !PLATFORM_DESKTOP)))
 
@@ -2641,10 +2639,8 @@ protected:
 	 */
 	virtual int32 CompilePropertyAndSetMaterialProperty(EMaterialProperty Property, class FMaterialCompiler* Compiler, EShaderFrequency OverrideShaderFrequency = SF_NumFrequencies, bool bUsePreviousFrameTime = false) const = 0;
 	
-#if HANDLE_CUSTOM_OUTPUTS_AS_MATERIAL_ATTRIBUTES
 	/** Used to translate code for custom output attributes such as ClearCoatBottomNormal  */
 	virtual int32 CompileCustomAttribute(const FGuid& AttributeID, class FMaterialCompiler* Compiler) const {return INDEX_NONE;}
-#endif
 
 	/** Useful for debugging. */
 	virtual FString GetBaseMaterialPathName() const { return TEXT(""); }
@@ -3088,10 +3084,8 @@ protected:
 
 	/** Entry point for compiling a specific material property.  This must call SetMaterialProperty. */
 	ENGINE_API virtual int32 CompilePropertyAndSetMaterialProperty(EMaterialProperty Property, class FMaterialCompiler* Compiler, EShaderFrequency OverrideShaderFrequency, bool bUsePreviousFrameTime) const override;
-#if HANDLE_CUSTOM_OUTPUTS_AS_MATERIAL_ATTRIBUTES
 	/** Used to translate code for custom output attributes such as ClearCoatBottomNormal  */
 	ENGINE_API virtual int32 CompileCustomAttribute(const FGuid& AttributeID, FMaterialCompiler* Compiler) const override;
-#endif
 
 #if WITH_EDITORONLY_DATA
 	/* Gives the material a chance to compile any custom output nodes it has added */
diff --git a/Engine/Source/Runtime/Engine/Public/Materials/MaterialAttributeDefinitionMap.h b/Engine/Source/Runtime/Engine/Public/Materials/MaterialAttributeDefinitionMap.h
--- a/Engine/Source/Runtime/Engine/Public/Materials/MaterialAttributeDefinitionMap.h
+++ b/Engine/Source/Runtime/Engine/Public/Materials/MaterialAttributeDefinitionMap.h
@@ -61,7 +61,7 @@ public:
 class FMaterialCustomOutputAttributeDefintion : public FMaterialAttributeDefintion
 {
 public:
-	FMaterialCustomOutputAttributeDefintion(const FGuid& InGUID, const FString& InAttributeName, const FString& InFunctionName, EMaterialProperty InProperty,
+	FMaterialCustomOutputAttributeDefintion(const FGuid& InGUID, UClass* InExpressionClass, const FString& InAttributeName, const FString& InFunctionName, int32 InOutputIndex, bool bInNeedsPreviousFrameEvaluation, EMaterialProperty InProperty,
 		EMaterialValueType InValueType, const FVector4& InDefaultValue, EShaderFrequency InShaderFrequency, MaterialAttributeBlendFunction InBlendFunction = nullptr);
 
 	bool operator==(const FMaterialCustomOutputAttributeDefintion& Other) const
@@ -69,8 +69,10 @@ public:
 		return (AttributeID == Other.AttributeID);
 	}
 
-	// Name of function used to access attribute in shader code
-	FString							FunctionName;
+	UClass*				CustomOutputExpressionStaticClass;
+	FString				FunctionName; // Name of function used to access attribute in shader code
+	int32				OutputIndex;
+	bool				bNeedsPreviousFrameEvaluation;
 };
 
 /**
@@ -117,6 +119,24 @@ public:
 	/** Returns the shader frequency of a material attribute */
 	static ENGINE_API EShaderFrequency GetShaderFrequency(const FGuid& AttributeID);
 
+	/** Returns the bit index of this property in the bitmasks returned by GetBitmask(), GetConnectedMaterialAttributesBitmask() */
+	static ENGINE_API int32 GetIndex(EMaterialProperty Property);
+
+	/** Returns the bit index of this attribute in the bitmasks returned by GetBitmask(), GetConnectedMaterialAttributesBitmask() and GetCustomOutputNodesBitmask() */
+	static ENGINE_API int32 GetIndex(const FGuid& AttributeID);
+
+	/** Returns a bitmask with just the bit corresponding to this property set */
+	static ENGINE_API uint64 GetBitmask(EMaterialProperty Property);
+
+	/** Returns a bitmask with just the bit corresponding to this attribute set */
+	static ENGINE_API uint64 GetBitmask(const FGuid& AttributeID);
+
+	/** Checks if the bit corresponding to the property/attribute is set in the bitmask */
+	static ENGINE_API bool IsAttributeInBitmask(uint64 Bitmask, EMaterialProperty Property);
+
+	/** Checks if the bit corresponding to the property/attribute is set in the bitmask */
+	static ENGINE_API bool IsAttributeInBitmask(uint64 Bitmask, const FGuid& AttributeID);
+
 	/** Returns the attribute ID for a matching material property */
 	static ENGINE_API FGuid GetID(EMaterialProperty Property);
 
@@ -134,7 +154,7 @@ public:
 	static ENGINE_API void AppendDDCKey(FShaderKeyGenerator& KeyGen);
 
 	/** Appends a new attribute definition to the custom output list */
-	static ENGINE_API void AddCustomAttribute(const FGuid& AttributeID, const FString& AttributeName, const FString& FunctionName, EMaterialValueType ValueType, const FVector4& DefaultValue, MaterialAttributeBlendFunction BlendFunction = nullptr);
+	static ENGINE_API void AddCustomAttribute(const FGuid& AttributeID, UClass* ExpressionClass, const FString& AttributeName, const FString& FunctionName, int32 OutputIndex, EMaterialValueType ValueType, const FVector4& DefaultValue, EShaderFrequency ShaderFrequency, bool bNeedsPreviousFrameEvaluation = false, MaterialAttributeBlendFunction BlendFunction = nullptr);
 
 	/** Returns the first custom attribute ID that has the specificed attribute name */
 	static ENGINE_API FGuid GetCustomAttributeID(const FString& AttributeName);
@@ -153,6 +173,15 @@ public:
 		return GMaterialPropertyAttributesMap.OrderedVisibleAttributeList;
 	}
 
+#if WITH_EDITOR
+	/** Returns a bitmask of all attributes/properties that are set via Set/MakeMaterialAttribute expressions, taking material function calls into account */
+	static ENGINE_API uint64 GetConnectedMaterialAttributesBitmask(TConstArrayView<TObjectPtr<class UMaterialExpression>> Expressions);
+	/** Returns a bitmask of all attributes/properties that are set via Set/MakeMaterialAttribute expressions, taking material function calls into account */
+	static ENGINE_API uint64 GetConnectedMaterialAttributesBitmask(TConstArrayView<class UMaterialExpression*> Expressions);
+	/** Returns a bitmask of all supported custom output attributes in the given array of custom output expressions */
+	static ENGINE_API uint64 GetCustomOutputNodesBitmask(TConstArrayView<class UMaterialExpressionCustomOutput*> Expressions);
+#endif
+
 private:
 	// Customization class for displaying data in the material editor
 	friend class FMaterialAttributePropertyDetails;
@@ -174,6 +203,11 @@ private:
 	static ENGINE_API FText GetAttributeOverrideForMaterial(const FGuid& AttributeID, UMaterial* Material);
 	static ENGINE_API FString GetPinNameFromShadingModelField(FMaterialShadingModelField InShadingModels, const TArray<TKeyValuePair<EMaterialShadingModel, FString>>& InCustomShadingModelPinNames, const FString& InDefaultPinName);
 
+#if WITH_EDITOR
+	template<typename ExpressionsArrayType>
+	static ENGINE_API uint64 GetConnectedMaterialAttributesBitmask_Internal(const ExpressionsArrayType& Expressions);
+#endif
+
 	static ENGINE_API FMaterialAttributeDefinitionMap GMaterialPropertyAttributesMap;
 
 	TMap<EMaterialProperty, FMaterialAttributeDefintion>	AttributeMap; // Fixed map of compile-time definitions
