Expose WorldSetting's Asset User Data to the editor

diff --git a/Engine/Source/Runtime/Engine/Classes/GameFramework/WorldSettings.h b/Engine/Source/Runtime/Engine/Classes/GameFramework/WorldSettings.h
--- a/Engine/Source/Runtime/Engine/Classes/GameFramework/WorldSettings.h
+++ b/Engine/Source/Runtime/Engine/Classes/GameFramework/WorldSettings.h
@@ -733,7 +733,7 @@ public:
 protected:
 
 	/** Array of user data stored with the asset */
-	UPROPERTY()
+	UPROPERTY(EditAnywhere, AdvancedDisplay, Instanced, Category = AssetUserData)
 	TArray<TObjectPtr<UAssetUserData>> AssetUserData;
 
 	// If paused, PlayerState of person pausing the game.
