[HWRT] Fix spline proxy ray tracing geometry not being rebuilt after getting evicted.
- also release RayTracingDynamicVertexBuffer when geometry is updated every frame.

diff --git a/Engine/Source/Runtime/Engine/Private/SplineMeshSceneProxy.cpp b/Engine/Source/Runtime/Engine/Private/SplineMeshSceneProxy.cpp
--- a/Engine/Source/Runtime/Engine/Private/SplineMeshSceneProxy.cpp
+++ b/Engine/Source/Runtime/Engine/Private/SplineMeshSceneProxy.cpp
@@ -227,9 +227,24 @@ void FSplineMeshSceneProxy::GetDynamicRayTracingInstances(FRayTracingInstanceCol
 	RayTracingInstance.InstanceTransformsView = MakeArrayView(&ThisLocalToWorld, 1);
 	RayTracingInstance.MaterialsView = MakeArrayView(CachedRayTracingMaterials);
 
-	const bool bNeedsUpdate = bUpdateRayTracingGeometry || bAnyMaterialHasWorldPositionOffset || !Geometry.IsValid() || (CVarRayTracingSplineMeshesForceUpdate.GetValueOnRenderThread() != 0);
+	const bool bAlwaysUpdate = bAnyMaterialHasWorldPositionOffset || (CVarRayTracingSplineMeshesForceUpdate.GetValueOnRenderThread() != 0);
+	const bool bNeedsUpdate = bAlwaysUpdate
+		|| bUpdateRayTracingGeometry
+		|| (!bAlwaysUpdate && Geometry.DynamicGeometrySharedBufferGenerationID != FRayTracingGeometry::NonSharedVertexBuffers) // was using shared VB but won't use it anymore so update once
+		|| !Geometry.IsValid()
+		|| Geometry.IsEvicted()
+		|| Geometry.GetRequiresBuild();
 
-	if (RenderData->LODVertexFactories[LODIndex].VertexFactory.GetType()->SupportsRayTracingDynamicGeometry() && bNeedsUpdate)
+	FRWBuffer* VertexBuffer = &RayTracingDynamicVertexBuffer;
+
+	if (bAlwaysUpdate)
+	{
+		// if updating every frame release memory and use shared VB
+		VertexBuffer->Release();
+		VertexBuffer = nullptr;
+	}
+
+	if (bNeedsUpdate && RenderData->LODVertexFactories[LODIndex].VertexFactory.GetType()->SupportsRayTracingDynamicGeometry())
 	{
 		Collector.AddRayTracingGeometryUpdate(
 			FRayTracingDynamicGeometryUpdateParams
@@ -240,7 +255,7 @@ void FSplineMeshSceneProxy::GetDynamicRayTracingInstances(FRayTracingInstanceCol
 				uint32((SIZE_T)LODModel.GetNumVertices() * sizeof(FVector3f)),
 				Geometry.Initializer.TotalPrimitiveCount,
 				&Geometry,
-				&RayTracingDynamicVertexBuffer,
+				VertexBuffer,
 				true
 			}
 		);
