Fixed crash in UMaterialExpression::GetDescription() when GetCaption() returns an empty array.

#rnx

diff --git a/Engine/Source/Runtime/Engine/Private/Materials/MaterialExpressions.cpp b/Engine/Source/Runtime/Engine/Private/Materials/MaterialExpressions.cpp
--- a/Engine/Source/Runtime/Engine/Private/Materials/MaterialExpressions.cpp
+++ b/Engine/Source/Runtime/Engine/Private/Materials/MaterialExpressions.cpp
@@ -1981,6 +1981,13 @@ FString UMaterialExpression::GetDescription() const
 	TArray<FString> Captions;
 	GetCaption(Captions);
 
+	// Guard against GetCaption() implementations not populating the array
+	if (Captions.IsEmpty())
+	{
+		UMaterialExpression::GetCaption(Captions);
+		check(!Captions.IsEmpty());
+	}
+
 	if (Captions.Num() > 1)
 	{
 		FString Result = Captions[0];
