NetDriver
* Prevent crash when NetDriver destruction is triggered while NetDriver is ticking. Ex: Calling Shutdown from inside an RPC.
*  Now NetDriver will wait until PostTickFlush to shutdown. Previously queued destruction request could also occur earlier in PostTickDispatch

diff --git a/Engine/Plugins/Online/OnlineSubsystemUtils/Source/OnlineSubsystemUtils/Private/OnlineBeacon.cpp b/Engine/Plugins/Online/OnlineSubsystemUtils/Source/OnlineSubsystemUtils/Private/OnlineBeacon.cpp
--- a/Engine/Plugins/Online/OnlineSubsystemUtils/Source/OnlineSubsystemUtils/Private/OnlineBeacon.cpp
+++ b/Engine/Plugins/Online/OnlineSubsystemUtils/Source/OnlineSubsystemUtils/Private/OnlineBeacon.cpp
@@ -57,14 +57,7 @@ void AOnlineBeacon::CleanupNetDriver()
 {
 	if (NetDriver)
 	{
-		if (NetDriver->IsInTick())
-		{
-			NetDriver->SetPendingDestruction(true);
-		}
-		else
-		{
-			GEngine->DestroyNamedNetDriver(GetWorld(), NetDriverName);
-		}
+		NetDriver->RequestNetDriverDestruction();
 		// If the net connection is currently in the middle of processing messages it is
 		// possible for the notifier to be fired again if not cleared.
 		NetDriver->Notify = nullptr;
diff --git a/Engine/Source/Runtime/Engine/Classes/Engine/NetDriver.h b/Engine/Source/Runtime/Engine/Classes/Engine/NetDriver.h
--- a/Engine/Source/Runtime/Engine/Classes/Engine/NetDriver.h
+++ b/Engine/Source/Runtime/Engine/Classes/Engine/NetDriver.h
@@ -1104,10 +1104,17 @@ public:
 	double GetElapsedTime() const { return ElapsedTime; }
 	void ResetElapsedTime() { ElapsedTime = 0.0; }
 
+	/** Trigger the destruction of this NetDriver. If called while the NetDriver is ticking it will queue the destruction request and execute it after all the NetDriver work is finished. */
+	ENGINE_API void RequestNetDriverDestruction();
+
+	/** True only inside one of the NetDriver's tick function. */
 	bool IsInTick() const { return bInTick; }
 
+	/** Is a destruction request queued to be executed later. */
 	bool GetPendingDestruction() const { return bPendingDestruction; }
-	void SetPendingDestruction(bool bDestroy) { bPendingDestruction = bDestroy; }
+
+	UE_DEPRECATED(5.6, "SetPendingDestruction will be made private. Use RequestNetDriverDestruction instead")
+	ENGINE_API void SetPendingDestruction(bool bDestroy);
 
 private:
 	double						ElapsedTime;
diff --git a/Engine/Source/Runtime/Engine/Private/NetDriver.cpp b/Engine/Source/Runtime/Engine/Private/NetDriver.cpp
--- a/Engine/Source/Runtime/Engine/Private/NetDriver.cpp
+++ b/Engine/Source/Runtime/Engine/Private/NetDriver.cpp
@@ -1461,20 +1461,26 @@ void UNetDriver::ProcessLocalClientPackets()
 
 void UNetDriver::PostTickFlush()
 {
-#if UE_WITH_IRIS
-	if (ReplicationSystem && (ServerConnection != nullptr || ClientConnections.Num() > 0))
+	// We are still ticking here
 	{
-		ReplicationSystem->PostSendUpdate();
-	}
+		TGuardValue<bool> GuardInNetTick(bInTick, true);
+#if UE_WITH_IRIS
+		if (ReplicationSystem && (ServerConnection != nullptr || ClientConnections.Num() > 0))
+		{
+			ReplicationSystem->PostSendUpdate();
+		}
 #endif // UE_WITH_IRIS
 
-	if (World && !bSkipClearVoicePackets)
-	{
-		UOnlineEngineInterface::Get()->ClearVoicePackets(World);
+		if (World && !bSkipClearVoicePackets)
+		{
+			UOnlineEngineInterface::Get()->ClearVoicePackets(World);
+		}
 	}
 
+	// All ticks are finished, look if a destruction request was queued
 	if (bPendingDestruction)
 	{
+		check(IsInTick() == false);
 		if (World)
 		{
 			GEngine->DestroyNamedNetDriver(World, NetDriverName);
@@ -2077,8 +2083,8 @@ void UNetDriver::UnregisterTickEvents(class UWorld* InWorld)
 	{
 		InWorld->OnTickDispatch ().Remove(TickDispatchDelegateHandle);
 		InWorld->OnPostTickDispatch().Remove(PostTickDispatchDelegateHandle);
-		InWorld->OnTickFlush    ().Remove(TickFlushDelegateHandle);
-		InWorld->OnPostTickFlush   ().Remove(PostTickFlushDelegateHandle);
+		InWorld->OnTickFlush().Remove(TickFlushDelegateHandle);
+		InWorld->OnPostTickFlush().Remove(PostTickFlushDelegateHandle);
 	}
 }
 
@@ -2483,6 +2489,8 @@ void UNetDriver::TickDispatch( float DeltaTime )
 
 void UNetDriver::PostTickDispatch()
 {
+	TGuardValue<bool> GuardInNetTick(bInTick, true);
+
 	// Flush out of order packet caches for connections that did not receive the missing packets during TickDispatch
 	if (ServerConnection != nullptr)
 	{
@@ -2515,19 +2523,6 @@ void UNetDriver::PostTickDispatch()
 		GRPCCSVTracker.EndTickDispatch();
 		GReceiveRPCTimingEnabled = false;
 	}
-
-	if (bPendingDestruction)
-	{
-		if (World)
-		{ 
-			GEngine->DestroyNamedNetDriver(World, NetDriverName);
-		}
-		else
-		{
-			UE_LOG(LogNet, Error, TEXT("NetDriver %s pending destruction without valid world."), *NetDriverName.ToString());
-		}
-		bPendingDestruction = false;
-	}
 }
 
 void UNetDriver::NotifyRPCProcessed(UFunction* Function, UNetConnection* Connection, double ElapsedTimeSeconds)
@@ -7556,6 +7551,35 @@ bool UNetDriver::ShouldSkipRepNotifies() const
 	return false;
 }
 
+void UNetDriver::RequestNetDriverDestruction()
+{
+	if (GetPendingDestruction())
+	{
+		// destruction already queued
+		return;
+	}
+
+	// If we are ticking, wait before triggering our destruction 
+	if (IsInTick())
+	{
+PRAGMA_DISABLE_DEPRECATION_WARNINGS
+		SetPendingDestruction(true);
+PRAGMA_ENABLE_DEPRECATION_WARNINGS
+	}
+	else
+	{
+		if (ensureMsgf(World, TEXT("RequestNetDriverDestruction on %s failed due to no world assigned"), *GetDescription()))
+		{
+			GEngine->DestroyNamedNetDriver(World, NetDriverName);
+		}		
+	}
+}
+
+void UNetDriver::SetPendingDestruction(bool bDestroy)
+{
+	UE_LOG(LogNet, Log, TEXT("SetPendingDestruction: %d on %s)"), bDestroy, *GetDescription());
+	bPendingDestruction = bDestroy;
+}
 // --------------------------------------------------------------------------------------------------------------------------------------------
 // --------------------------------------------------------------------------------------------------------------------------------------------
 // --------------------------------------------------------------------------------------------------------------------------------------------
diff --git a/Engine/Source/Runtime/Engine/Private/UnrealEngine.cpp b/Engine/Source/Runtime/Engine/Private/UnrealEngine.cpp
--- a/Engine/Source/Runtime/Engine/Private/UnrealEngine.cpp
+++ b/Engine/Source/Runtime/Engine/Private/UnrealEngine.cpp
@@ -14076,7 +14076,7 @@ void UEngine::ShutdownWorldNetDriver( UWorld * World )
 		for (int32 Index = 0; Index < Context.ActiveNetDrivers.Num(); Index++)
 		{
 			NetDriver = Context.ActiveNetDrivers[Index].NetDriver;
-			if (NetDriver && NetDriver->GetWorld() == World)
+			if (NetDriver && NetDriver->GetWorld() == World && !NetDriver->GetPendingDestruction())
 			{
 				UE_LOG(LogNet, Log, TEXT("World NetDriver shutdown %s [%s]"), *NetDriver->GetName(), *NetDriver->NetDriverName.ToString());
 				DestroyNamedNetDriver(World, NetDriver->NetDriverName);
@@ -14468,12 +14468,18 @@ namespace UE::Private
 		{
 			FNamedNetDriver& NamedNetDriver = Context.ActiveNetDrivers[Index];
 			UNetDriver* NetDriver = NamedNetDriver.NetDriver;
+
 			if (NetDriver && NetDriver->NetDriverName == NetDriverName)
 			{
-				UE_LOG(LogNet, Log, TEXT("DestroyNamedNetDriver %s [%s]"), *NetDriver->GetName(), *NetDriverName.ToString());
-				
-				ensureMsgf(!(NetDriver->IsInTick() && NetDriverName != NAME_PendingNetDriver), TEXT("Attempting to destroy NetDriver %s [%s] while it is ticking."), *NetDriver->GetName(), *NetDriverName.ToString());
+				// Can't destroy immediately while inside the NetDriver's tick.
+				if (NetDriver && NetDriver->IsInTick())
+				{
+					NetDriver->RequestNetDriverDestruction();
+					continue;
+				}
 
+				UE_LOG(LogNet, Log, TEXT("DestroyNamedNetDriver: %s "), *NetDriver->GetDescription());
+				
 				Context.ActiveNetDrivers.RemoveAtSwap(Index);
 				Context.GarbageObjectsToVerify.Add(NetDriver);
 				NetDriver->SetWorld(NULL);
