Fixed AActor::SupportsDataLayerType to also test if actor is external.

diff --git a/Engine/Source/Runtime/Engine/Private/ActorEditor.cpp b/Engine/Source/Runtime/Engine/Private/ActorEditor.cpp
--- a/Engine/Source/Runtime/Engine/Private/ActorEditor.cpp
+++ b/Engine/Source/Runtime/Engine/Private/ActorEditor.cpp
@@ -1913,6 +1913,7 @@ bool AActor::SupportsDataLayerType(TSubclassOf<UDataLayerInstance> InDataLayerTy
 	return (!bIsLevelNotPartitioned &&
 		!bHasComponentForceActorNoDataLayers &&
 		bActorTypeSupportsDataLayerType &&
+		IsPackageExternal() &&
 		!FActorEditorUtils::IsABuilderBrush(this) &&
 		!GetClass()->GetDefaultObject<AActor>()->bHiddenEd);
 }
