DrawStatsHUD will now cache `DrawOnscreenDebugMessages` to reduce hitching when drawing a large amount of text to the screen.

#tests Profiled and saw hitch spikes as high as 30 MS go down to < 1 MS. PIE PC
[RN] DrawStatsHUD will now cache `DrawOnscreenDebugMessages` to reduce hitching when drawing a large amount of text to the screen.

diff --git a/Engine/Source/Runtime/Engine/Private/UnrealEngine.cpp b/Engine/Source/Runtime/Engine/Private/UnrealEngine.cpp
--- a/Engine/Source/Runtime/Engine/Private/UnrealEngine.cpp
+++ b/Engine/Source/Runtime/Engine/Private/UnrealEngine.cpp
@@ -61,6 +61,7 @@ UnrealEngine.cpp: Implements the UEngine class and helpers.
 #include "Engine/SkeletalMesh.h"
 #include "Engine/SkinnedAssetCommon.h"
 #include "Components/SkeletalMeshComponent.h"
+#include "EngineFontServices.h"
 #include "Engine/Texture2D.h"
 #include "Engine/Texture2DArray.h"
 #include "Engine/TextureRenderTarget2D.h"
@@ -12821,6 +12822,33 @@ float UEngine::DrawOnscreenDebugMessages(UWorld* World, FViewport* Viewport, FCa
 		// Because some components add their message in concurrent work, we need a CS here.
 		FScopeLock ScopeLock(&GOnScreenMessageCS);
 
+		// Calculate the final amount of vertices for the batch or else we will end up resizing many times as the batching system only reserves the first draw call.
+		TSharedPtr<FSlateFontCache> FontCache = FEngineFontServices::Get().GetFontCache();
+		if (FontCache.IsValid())
+		{
+			int32 TotalCharacterCount = 0;
+			for (const FScreenMessageString& Message : PriorityScreenMessages)
+			{
+				TotalCharacterCount += Message.ScreenMessage.Len();
+			}
+
+			for (const TTuple<int32, FScreenMessageString>& Tuple : ScreenMessages)
+			{
+				TotalCharacterCount += Tuple.Value.ScreenMessage.Len();
+			}
+
+			ISlateFontTexture* SlateFontTexture = FontCache->GetFontTexture(0);
+			const FTexture* FontTexture = SlateFontTexture->GetEngineTexture();
+			FBatchedElements* BatchedElements = Canvas->GetBatchedElements(FCanvas::ET_Triangle, nullptr, FontTexture, SE_BLEND_TranslucentAlphaOnly, FDepthFieldGlowInfo(), false);
+			check(BatchedElements);
+
+			const int32 TextEffects = 1; // Shadow
+			const int32 NumTris = 2 * (TextEffects + 1);
+			const int32 NumVerts = 4 * (TextEffects + 1);
+			BatchedElements->ReserveTriangles(TotalCharacterCount * NumTris, FontTexture, SE_BLEND_TranslucentAlphaOnly);
+			BatchedElements->ReserveVertices(TotalCharacterCount * NumVerts);
+		}
+
 		if (PriorityScreenMessages.Num() > 0)
 		{
 			FCanvasTextItem MessageTextItem(FVector2D(0, 0), FText::GetEmpty(), GetSmallFont(), FLinearColor::White);
