Unshelved from pending changelist '37445705':

Fix to prevent iterating and using montage instances after the animation instance has been uninitialized.

#rnx

diff --git a/Engine/Source/Runtime/Engine/Classes/Animation/AnimInstance.h b/Engine/Source/Runtime/Engine/Classes/Animation/AnimInstance.h
--- a/Engine/Source/Runtime/Engine/Classes/Animation/AnimInstance.h
+++ b/Engine/Source/Runtime/Engine/Classes/Animation/AnimInstance.h
@@ -376,6 +376,9 @@ private:
 	UPROPERTY(Transient)
 	uint8 bQueueMontageEvents : 1;
 
+	/** Flag used to query if the animation instance has been uninitialized via UninitializeAnimation() call. */
+	uint8 bUninitialized : 1;
+	
 #if DO_CHECK
 	/** Guard flag used for checking whether we are in user callbacks for initialization */
 	uint8 bInitializing : 1;
diff --git a/Engine/Source/Runtime/Engine/Private/Animation/AnimInstance.cpp b/Engine/Source/Runtime/Engine/Private/Animation/AnimInstance.cpp
--- a/Engine/Source/Runtime/Engine/Private/Animation/AnimInstance.cpp
+++ b/Engine/Source/Runtime/Engine/Private/Animation/AnimInstance.cpp
@@ -100,6 +100,12 @@ static FAutoConsoleVariableRef CVarRK4SpringInterpolatorMaxIter(TEXT("p.RK4Sprin
 
 namespace MontageCVars
 {
+	static bool bEarlyOutMontageWhenUninitialized = true;
+	static FAutoConsoleVariableRef CVarEarlyOutMontageWhenUninitialized(
+		TEXT("a.Montage.EarlyOutMontageWhenUninitialized"),
+		bEarlyOutMontageWhenUninitialized,
+		TEXT("Exit early when playing or stopping montage(s) if the animation instance was uninitialized while we are performing the action due to triggering montage events. Use this to avoid reading from bad memory."));
+	
 	static bool bFlushCompletedMontagesOnPlay = false;
 	static FAutoConsoleVariableRef CVarFlushCompletedMontagesOnPlay(
 		TEXT("a.Montage.FlushCompletedMontagesOnPlay"),
@@ -130,6 +136,8 @@ UAnimInstance::UAnimInstance(const FObjectInitializer& ObjectInitializer)
 	bPropagateNotifiesToLinkedInstances = false;
 	bUseMainInstanceMontageEvaluationData = false;
 	
+	bUninitialized = true;
+	
 	bUpdateAnimationEnabled = true;
 #if DO_CHECK
 	bInitializing = false;
@@ -259,6 +267,7 @@ void UAnimInstance::InitializeAnimation(bool bInDeferRootNodeInitialization)
 	LLM_SCOPE(ELLMTag::Animation);
 
 	UninitializeAnimation();
+	bUninitialized = false;
 	
 	TRACE_OBJECT_LIFETIME_BEGIN(this);
 
@@ -318,6 +327,8 @@ void UAnimInstance::InitializeAnimation(bool bInDeferRootNodeInitialization)
 
 void UAnimInstance::UninitializeAnimation()
 {
+	bUninitialized = true;
+	
 	NativeUninitializeAnimation();
 
 	GetProxyOnGameThread<FAnimInstanceProxy>().Uninitialize(this);
@@ -2309,6 +2320,12 @@ float UAnimInstance::Montage_PlayInternal(UAnimMontage* MontageToPlay, const FMo
 				// Enforce 'a single montage at once per group' rule
 				StopAllMontagesByGroupName(NewMontageGroupName, BlendInSettings);
 
+				// Early out since we might have called UninitializeAnimation() when trigger montage blending out events above.
+				if (bUninitialized && MontageCVars::bEarlyOutMontageWhenUninitialized)
+				{
+					return 0.0f;
+				}
+
 				if (MontageCVars::bFlushCompletedMontagesOnPlay)
 				{
 					ConditionalFlushCompletedMontages();
@@ -3116,6 +3133,12 @@ void UAnimInstance::StopAllMontagesByGroupName(FName InGroupName, const FMontage
 		if (MontageInstance && MontageInstance->Montage && (MontageInstance->Montage->GetGroupName() == InGroupName))
 		{
 			MontageInstance->Stop(BlendOutSettings, true);
+
+			// UninitializeAnimation() empties the MontagesInstances array thus we need to exit early.
+			if (bUninitialized && MontageCVars::bEarlyOutMontageWhenUninitialized)
+			{
+				return;
+			}
 		}
 	}
 }
