Use transactionally safe locks for the network profiler

diff --git a/Engine/Source/Runtime/Engine/Private/NetworkProfiler.cpp b/Engine/Source/Runtime/Engine/Private/NetworkProfiler.cpp
--- a/Engine/Source/Runtime/Engine/Private/NetworkProfiler.cpp
+++ b/Engine/Source/Runtime/Engine/Private/NetworkProfiler.cpp
@@ -17,7 +17,7 @@
 
 CSV_DEFINE_CATEGORY(NetworkProfiler, true);
 
-#define SCOPE_LOCK_REF(X) FScopeLock ScopeLock(&X);
+#define SCOPE_LOCK_REF(X) UE::TScopeLock ScopeLock(X);
 
 struct FNetworkProfilerCVarHelper
 {
diff --git a/Engine/Source/Runtime/Engine/Public/Net/NetworkProfiler.h b/Engine/Source/Runtime/Engine/Public/Net/NetworkProfiler.h
--- a/Engine/Source/Runtime/Engine/Public/Net/NetworkProfiler.h
+++ b/Engine/Source/Runtime/Engine/Public/Net/NetworkProfiler.h
@@ -13,6 +13,7 @@
 #include "Serialization/MemoryWriter.h"
 #include "ProfilingDebugging/CsvProfiler.h"
 #include "Misc/ScopeLock.h"
+#include "Misc/TransactionallySafeCriticalSection.h"
 
 class AActor;
 class FOutBunch;
@@ -95,7 +96,7 @@ private:
 	FArchive*								FileWriter;
 
 	/** Critical section to sequence tracking.														*/
-	FCriticalSection						CriticalSection;
+	FTransactionallySafeCriticalSection		CriticalSection;
 
 	/** Mapping from name to index in name array.													*/
 	TMap<FString,int32>						NameToNameTableIndexMap;
@@ -378,7 +379,7 @@ public:
 	{
 		if (IsComparisonTrackingEnabled())
 		{
-            FScopeLock ScopeLock(&CriticalSection);
+			UE::TScopeLock ScopeLock(CriticalSection);
 
 			FString ObjectName = GetNameSafe(Object);
 			TArray<uint8> PropertyNameExportData;
