Fix for a regression where CPFUO editinline duplication could bypass inner object properties owned by editable array fields.

Change summary:
- Removed explicit editable and archetype checks; these should be handled or validated instead of bypassed.
- Results are now closer to parity w/ the old path that relied exclusively on the instanced property flags.
- Re-enabled the new duplication path in CPFUO to support types that don't utilize instanced semantics.

#rnx

diff --git a/Engine/Source/Runtime/Engine/Private/UnrealEngine.cpp b/Engine/Source/Runtime/Engine/Private/UnrealEngine.cpp
--- a/Engine/Source/Runtime/Engine/Private/UnrealEngine.cpp
+++ b/Engine/Source/Runtime/Engine/Private/UnrealEngine.cpp
@@ -17665,7 +17665,7 @@ public:
 
 namespace UE::Engine::Private
 {
-	static bool bUseOldDuplicateEditInlineSubObjectsPath = true;
+	static bool bUseOldDuplicateEditInlineSubObjectsPath = false;
 	static FAutoConsoleVariableRef CVarUseOldDuplicateEditInlineSubObjectsPath(
 		TEXT("CPFUO.UseOldDuplicateEditInlineSubObjectsPath"),
 		bUseOldDuplicateEditInlineSubObjectsPath,
@@ -17702,7 +17702,8 @@ PRAGMA_ENABLE_DEPRECATION_WARNINGS
 			return;
 		}
 
-		// Copy old subobjects referenced by editable properties to the new object
+		// Copy old subobjects referenced by non-transient object properties to the new object.
+		// The iterator is inclusive of any container or struct property inner references as well.
 		for (TPropertyValueIterator<FObjectProperty> It(NewObject->GetClass(), NewObject); It; ++It)
 		{
 			const FObjectProperty* ObjectProperty = It.Key();
@@ -17721,22 +17722,6 @@ PRAGMA_ENABLE_DEPRECATION_WARNINGS
 				continue;
 			}
 
-			// We only care about references that can be instanced at edit time by the user. If
-			// this was previously editable, the instance will be cleared during ref replacement.
-			const bool bIsEditable = ObjectProperty->HasAnyPropertyFlags(CPF_Edit);
-			if (!bIsEditable)
-			{
-				continue;
-			}
-
-			// We only care about non-transient references. Transient references still outered to
-			// the old object will either be cleared during ref replacement or instanced elsewhere.
-			const bool bIsTransient = ObjectProperty->HasAnyPropertyFlags(CPF_Transient);
-			if (bIsTransient)
-			{
-				continue;
-			}
-
 			// If the old object doesn't reference it anymore then we don't want to keep it either.
 			// In this case the reference on the new object will be cleared during ref replacement.
 			const bool bKeptByOld = OldObjectReferences.Contains(Obj);
@@ -17745,9 +17730,10 @@ PRAGMA_ENABLE_DEPRECATION_WARNINGS
 				continue;
 			}
 
-			// We only care about references to subobjects not otherwise instanced from a template.
-			const bool bWasInstancedFromArchetype = Obj->GetArchetype()->IsTemplateForSubobjects();
-			if (bWasInstancedFromArchetype)
+			// We only care about non-transient references. Transient references still outered to
+			// the old object will either be cleared during ref replacement or instanced elsewhere.
+			const bool bIsTransient = ObjectProperty->HasAnyPropertyFlags(CPF_Transient);
+			if (bIsTransient)
 			{
 				continue;
 			}
