Small build fix for issue 640664

diff --git a/Engine/Source/Runtime/Engine/Private/Materials/MaterialIREmitter.cpp b/Engine/Source/Runtime/Engine/Private/Materials/MaterialIREmitter.cpp
--- a/Engine/Source/Runtime/Engine/Private/Materials/MaterialIREmitter.cpp
+++ b/Engine/Source/Runtime/Engine/Private/Materials/MaterialIREmitter.cpp
@@ -630,12 +630,14 @@ FValue* FEmitter::EmitTextureObject(UTexture* Texture, EMaterialSamplerType Samp
 
 FValue* FEmitter::EmitParameter(FName Name, const FMaterialParameterMetadata& Metadata)
 {
-	UE_MIR_TODO();
+	// UE_MIR_TODO();
+	return nullptr;
 }
 
 FValue* FEmitter::EmitCustomPrimitiveData(uint32 PrimitiveDataIndex)
 {
-	UE_MIR_TODO();
+	// UE_MIR_TODO();
+	return nullptr;
 }
 
 FSetMaterialOutput* FEmitter::EmitSetMaterialOutput(EMaterialProperty InProperty, FValue* InArgValue)
