Temporarily disable after world clean-up component validation in editor, until we fix all the little issues that cropped up

diff --git a/Engine/Source/Runtime/Engine/Private/World.cpp b/Engine/Source/Runtime/Engine/Private/World.cpp
--- a/Engine/Source/Runtime/Engine/Private/World.cpp
+++ b/Engine/Source/Runtime/Engine/Private/World.cpp
@@ -5553,6 +5553,7 @@ void UWorld::CleanupWorld(bool bSessionEnded, bool bCleanupResources, UWorld* Ne
 
 void ValidateComponentsRegistration(UWorld* World)
 {
+#if !WITH_EDITOR
 	TStringBuilder<2048>	Errors;
 
 	for (UActorComponent* Component : TObjectRange<UActorComponent>())
@@ -5582,6 +5583,7 @@ void ValidateComponentsRegistration(UWorld* World)
 	}
 
 	ensureMsgf(Errors.Len() == 0, TEXT("Found components to be incorrectly unregistered: \n%s"), Errors.GetData());
+#endif
 }
 
 void UWorld::CleanupWorldInternal(bool bSessionEnded, bool bCleanupResources, bool bWorldChanged)
