Fixed case where draw distance on an overlay mesh was affecting base mesh LOD computation

diff --git a/Engine/Source/Runtime/Engine/Private/SceneManagement.cpp b/Engine/Source/Runtime/Engine/Private/SceneManagement.cpp
--- a/Engine/Source/Runtime/Engine/Private/SceneManagement.cpp
+++ b/Engine/Source/Runtime/Engine/Private/SceneManagement.cpp
@@ -1043,7 +1043,7 @@ FLODMask ComputeLODForMeshes(const TArray<class FStaticMeshBatchRelevance>& Stat
 		for (int32 MeshIndex = 0; MeshIndex < StaticMeshRelevances.Num(); ++MeshIndex)
 		{
 			const FStaticMeshBatchRelevance& Mesh = StaticMeshRelevances[MeshIndex];
-			if (Mesh.ScreenSize > 0.0f)
+			if (Mesh.ScreenSize > 0.0f && !Mesh.bOverlayMaterial)
 			{
 				MinLOD = FMath::Min(MinLOD, (int32)Mesh.GetLODIndex());
 				MaxLOD = FMath::Max(MaxLOD, (int32)Mesh.GetLODIndex());
@@ -1065,10 +1065,11 @@ FLODMask ComputeLODForMeshes(const TArray<class FStaticMeshBatchRelevance>& Stat
 				for (int32 MeshIndex = NumMeshes - 1; MeshIndex >= 0; --MeshIndex)
 				{
 					const FStaticMeshBatchRelevance& Mesh = StaticMeshRelevances[MeshIndex];
-					if (Mesh.ScreenSize > 0.0f)
+					// We skip overlay material meshes as they always use base mesh LOD
+					if (Mesh.ScreenSize > 0.0f && !Mesh.bOverlayMaterial)
 					{
 						float MeshScreenSize = Mesh.ScreenSize * ScreenSizeScale;
-
+						
 						if (FMath::Square(MeshScreenSize * 0.5f) >= OutScreenRadiusSquared)
 						{
 							LODToRender.SetLODSample(Mesh.GetLODIndex(), SampleIndex);
@@ -1097,8 +1098,8 @@ FLODMask ComputeLODForMeshes(const TArray<class FStaticMeshBatchRelevance>& Stat
 				const FStaticMeshBatchRelevance& Mesh = StaticMeshRelevances[MeshIndex];
 
 				float MeshScreenSize = Mesh.ScreenSize * ScreenSizeScale;
-
-				if (FMath::Square(MeshScreenSize * 0.5f) >= OutScreenRadiusSquared)
+				// We skip overlay material meshes as they always use base mesh LOD
+				if (FMath::Square(MeshScreenSize * 0.5f) >= OutScreenRadiusSquared && !Mesh.bOverlayMaterial)
 				{
 					LODToRender.SetLOD(Mesh.GetLODIndex());
 					bFoundLOD = true;
