Fix to avoid  check retarget source asset command use large amounts of ram during the scan process
Pakcages loaded during the process will be unloaded
By default, only assets wuth retarget source asset set in the metadata will be loaded in memory (big speed up)

diff --git a/Engine/Source/Runtime/Engine/Private/Animation/AnimRetargetHelpers.cpp b/Engine/Source/Runtime/Engine/Private/Animation/AnimRetargetHelpers.cpp
--- a/Engine/Source/Runtime/Engine/Private/Animation/AnimRetargetHelpers.cpp
+++ b/Engine/Source/Runtime/Engine/Private/Animation/AnimRetargetHelpers.cpp
@@ -12,6 +12,11 @@
 #include "AssetRegistry/AssetData.h"
 #include "AssetRegistry/AssetRegistryModule.h"
 #include "Templates/UnrealTemplate.h"
+#include "PackageTools.h"
+#include "Misc/ScopedSlowTask.h"
+#include "UObject/UObjectGlobals.h"
+#include "CoreGlobals.h"
+#include "Misc/FeedbackContext.h"
 #endif // WITH_EDITOR
 
 #define LOCTEXT_NAMESPACE "AnimRetargetHelpers"
@@ -74,7 +79,7 @@ ERetargetSourceAssetStatus CheckRetargetSourceAssetDataImpl(const AssetType* InA
 	return ERetargetSourceAssetStatus::NoRetargetDataSet;
 }
 
-void CheckRetargetSourceAssetData(bool bFixAssets, const TArray<FString>& IncludedPaths, const TArray<FString>& ExcludedPaths)
+void CheckRetargetSourceAssetData(bool bFixAssets, bool bWantsFullScan, const TArray<FString>& IncludedPaths, const TArray<FString>& ExcludedPaths)
 {
 	TArray<FAssetData> Assets;
 
@@ -118,41 +123,89 @@ void CheckRetargetSourceAssetData(bool bFixAssets, const TArray<FString>& Includ
 
 	UE_LOG(LogAnimation, Log, TEXT("Check Retarget Source Assets after filtering exclusions : [%d] assets."), Assets.Num());
 
+	PRAGMA_DISABLE_DEPRECATION_WARNINGS
+	const FName SearchableName = GET_MEMBER_NAME_CHECKED(UAnimSequence, RetargetSourceAsset);
+	PRAGMA_ENABLE_DEPRECATION_WARNINGS
+	const FString AssetNone(TEXT("None"));
+
 	const int32 NumAssets = Assets.Num();
+	TArray<UPackage*> PackagesToUnload;
+	PackagesToUnload.Reserve(NumAssets);
+
+	FScopedSlowTask Feedback(NumAssets, LOCTEXT("ScanningRetargetSourceAssets", "Scanning Retarget Source Assets..."));
+	Feedback.MakeDialog(true);
+
 	for (int32 Idx = 0; Idx < NumAssets; Idx++)
 	{
-		if (UAnimSequence* AnimSequence = Cast<UAnimSequence>(Assets[Idx].GetAsset()))
+		FAssetData& AssetData = Assets[Idx];
+
+		Feedback.EnterProgressFrame(1, 
+			FText::Format(LOCTEXT("ScanningRetargetSourceAssets", "Assets Left [{0}] - {1}"), 
+				NumAssets - Idx,
+				FText::FromString((AssetData.PackagePath.ToString() + FString(TEXT("/")) + AssetData.AssetName.ToString()))));
+
+		// Fast scan will skip assets that does not have a retarget source asset set in the metadata
+		// With full scan we can check if an anim sequence has retarget transforms but no asset set for some reason (but it is slower)
+		if (!bWantsFullScan)
 		{
+			const FString RetargetSourceAssetFile = AssetData.GetTagValueRef<FString>(SearchableName);
+			if (RetargetSourceAssetFile.IsEmpty() || RetargetSourceAssetFile == AssetNone)
+			{
+				continue;
+			}
+		}
+
+		UPackage* Package = FindPackage(NULL, *AssetData.PackageName.ToString());
+		bool bWasPackageFullyLoaded = Package && Package->IsFullyLoaded();
+
+		if (UAnimSequence* AnimSequence = Cast<UAnimSequence>(AssetData.GetAsset()))
+		{
+			Package = AnimSequence->GetPackage();
 			const ERetargetSourceAssetStatus Status = Private::CheckRetargetSourceAssetDataImpl(AnimSequence);
-			if (bFixAssets)
+			if (bFixAssets && Status == ERetargetSourceAssetStatus::RetargetSourceMissing)
 			{
-				if (Status == ERetargetSourceAssetStatus::RetargetSourceMissing)
-				{
-					AnimSequence->Modify();
-					PRAGMA_DISABLE_DEPRECATION_WARNINGS
-					AnimSequence->RetargetSourceAsset.Reset();
-					PRAGMA_ENABLE_DEPRECATION_WARNINGS
-					AnimSequence->RetargetSourceAssetReferencePose.Empty();
-					AnimSequence->MarkPackageDirty();
-				}
+				AnimSequence->Modify();
+				PRAGMA_DISABLE_DEPRECATION_WARNINGS
+				AnimSequence->RetargetSourceAsset.Reset();
+				PRAGMA_ENABLE_DEPRECATION_WARNINGS
+				AnimSequence->RetargetSourceAssetReferencePose.Empty();
+				AnimSequence->MarkPackageDirty();
+				bWasPackageFullyLoaded = true; // avoid the package getting unloaded if it has been fixed
 			}
 		}
-		else if (UPoseAsset* PoseAsset = Cast<UPoseAsset>(Assets[Idx].GetAsset()))
+		else if (UPoseAsset* PoseAsset = Cast<UPoseAsset>(AssetData.GetAsset()))
 		{
+			Package = PoseAsset->GetPackage();
 			const ERetargetSourceAssetStatus Status = Private::CheckRetargetSourceAssetDataImpl(PoseAsset);
-			if (bFixAssets)
+			if (bFixAssets && Status == ERetargetSourceAssetStatus::RetargetSourceMissing)
 			{
-				if (Status == ERetargetSourceAssetStatus::RetargetSourceMissing)
-				{
-					PoseAsset->Modify();
-					PRAGMA_DISABLE_DEPRECATION_WARNINGS
-					PoseAsset->RetargetSourceAsset.Reset();
-					PRAGMA_ENABLE_DEPRECATION_WARNINGS
-					PoseAsset->RetargetSourceAssetReferencePose.Empty();
-					PoseAsset->MarkPackageDirty();
-				}
+				PoseAsset->Modify();
+				PRAGMA_DISABLE_DEPRECATION_WARNINGS
+				PoseAsset->RetargetSourceAsset.Reset();
+				PRAGMA_ENABLE_DEPRECATION_WARNINGS
+				PoseAsset->RetargetSourceAssetReferencePose.Empty();
+				PoseAsset->MarkPackageDirty();
+				bWasPackageFullyLoaded = true; // avoid the package getting unloaded if it has been fixed
 			}
 		}
+
+		if (!bWasPackageFullyLoaded && Package && Package->IsFullyLoaded())
+		{
+			PackagesToUnload.Add(Package);
+		}
+
+		const bool bCancelProcess = GWarn->ReceivedUserCancel();
+
+		if (((Idx % 100) == 0 || (Idx == NumAssets - 1) || bCancelProcess) && PackagesToUnload.Num() > 0)
+		{
+			UPackageTools::UnloadPackages({ PackagesToUnload });
+			PackagesToUnload.Reset();
+		}
+
+		if (bCancelProcess)
+		{
+			break;
+		}
 	}
 }
 
@@ -172,18 +225,24 @@ static FAutoConsoleCommand CheckRetargetSourceAssetDataCmd(
 
 			TArray<FString> IncludedPaths;
 			TArray<FString> ExcludedPaths;
-
+			const FString FixAssetsFlag("FixAssets");
+			const FString FullScanFlag("FullScan");
 			bool bWantsFix = false;
+			bool bWantsFullScan = false;
 			FString SwitchValue;
 
 			const int32 NumArgs = Args.Num();
 			for (int32 i = 0; i < NumArgs; i++)
 			{
 				const FString& Arg = Args[i];
-				if (Arg.ToLower() == FString("FixAssets"))
+				if (Arg.ToLower() == FixAssetsFlag)
 				{
 					bWantsFix = true;
 				}
+				else if (Arg.ToLower() == FullScanFlag)
+				{
+					bWantsFullScan = true;
+				}
 				else if (FParse::Value(*Arg, *IncludedPathsSwitch, SwitchValue))
 				{
 					SwitchValue.ParseIntoArray(IncludedPaths, ParamDelims, 3);
@@ -198,7 +257,7 @@ static FAutoConsoleCommand CheckRetargetSourceAssetDataCmd(
 				}
 			}
 
-			CheckRetargetSourceAssetData(bWantsFix, IncludedPaths, ExcludedPaths);
+			CheckRetargetSourceAssetData(bWantsFix, bWantsFullScan, IncludedPaths, ExcludedPaths);
 		}
 	));
 } // end namespace Private
