Change Task dispatch to ExecuteOnGameThread() for game thread work when modifying virutal texture pools.
Enquing to a ticker is cleaner than using tasks for this case. Tasks can lead to recursive calls to FlushRenderingCommands().

diff --git a/Engine/Source/Runtime/Engine/Private/VT/VirtualTexturePoolConfig.cpp b/Engine/Source/Runtime/Engine/Private/VT/VirtualTexturePoolConfig.cpp
--- a/Engine/Source/Runtime/Engine/Private/VT/VirtualTexturePoolConfig.cpp
+++ b/Engine/Source/Runtime/Engine/Private/VT/VirtualTexturePoolConfig.cpp
@@ -2,7 +2,7 @@
 
 #include "VT/VirtualTexturePoolConfig.h"
 
-#include "Async/TaskGraphInterfaces.h"
+#include "Containers/Ticker.h"
 #include "HAL/IConsoleManager.h"
 #include "RenderingThread.h"
 #include "VT/VirtualTextureRecreate.h"
@@ -195,74 +195,48 @@ void UVirtualTexturePoolConfig::PostEditChangeProperty(struct FPropertyChangedEv
 
 
 /**
- * A task to modify the virtual texture physical pool configs.
+ * Modify the virtual texture physical pool configs.
  * This needs queuing up on the game thread so that the virtual texture UObjects can be reinitialized correctly.
  */
-class FAddOrModifyTransientPoolConfigsTask
+static void AddOrModifyTransientPoolConfigs_GameThread(TArray<FVirtualTextureSpacePoolConfig> const& InConfigs)
 {
-	TArray<FVirtualTextureSpacePoolConfig> Configs;
+	check(IsInGameThread());
 
-public:
-	FAddOrModifyTransientPoolConfigsTask(TArray<FVirtualTextureSpacePoolConfig>& InConfigs)
-		: Configs(MoveTemp(InConfigs))
-	{
-	}
+	TArray<int32, TInlineAllocator<8>> ChangedConfigs;
 
-	FORCEINLINE TStatId GetStatId() const
+	UVirtualTexturePoolConfig* PoolConfig = GetMutableDefault<UVirtualTexturePoolConfig>();
+	for (int32 ConfigIndex = 0; ConfigIndex < InConfigs.Num(); ++ConfigIndex)
 	{
-		RETURN_QUICK_DECLARE_CYCLE_STAT(FAddOrModifyTransientPoolConfigsTask, STATGROUP_TaskGraphTasks);
-	}
-
-	ENamedThreads::Type GetDesiredThread()
-	{
-		return ENamedThreads::GameThread;
+		if (PoolConfig->AddOrModifyTransientPoolConfig(InConfigs[ConfigIndex]))
+		{
+			ChangedConfigs.Add(ConfigIndex);
+		}
 	}
 
-	static ESubsequentsMode::Type GetSubsequentsMode()
+	// Defer virtual texture recreation until all changes to config are done.
+	for (int32 ConfigIndex : ChangedConfigs)
 	{
-		return ESubsequentsMode::FireAndForget;
+		VirtualTexture::Recreate(InConfigs[ConfigIndex].Formats);
 	}
 
-	void DoTask(ENamedThreads::Type CurrentThread, const FGraphEventRef& MyCompletionGraphEvent)
+#if WITH_EDITOR
+	// Notify user about changes.
+	if (ChangedConfigs.Num())
 	{
-		check(IsInGameThread());
-
-		TArray<int32, TInlineAllocator<8>> ChangedConfigs;
-
-		UVirtualTexturePoolConfig* PoolConfig = GetMutableDefault<UVirtualTexturePoolConfig>();
-		for (int32 ConfigIndex = 0; ConfigIndex < Configs.Num(); ++ConfigIndex)
-		{
-			if (PoolConfig->AddOrModifyTransientPoolConfig(Configs[ConfigIndex]))
-			{
-				ChangedConfigs.Add(ConfigIndex);
-			}
-		}
-
-		// Defer virtual texture recreation until all changes to config are done.
-		for (int32 ConfigIndex : ChangedConfigs)
+		static TWeakPtr<class SNotificationItem> NotificationHandle;
+		if (!NotificationHandle.IsValid())
 		{
-			VirtualTexture::Recreate(Configs[ConfigIndex].Formats);
-		}
+			FNotificationInfo Info(LOCTEXT("PoolResizeNotify", "Resizing Virtual Texture Pools."));
+			Info.SubText = LOCTEXT("PoolResizeNotifySubtext", "Size changes are not saved by default. To keep changes, copy from 'Transient Pools' to 'Fixed Pools' in the Virtual Texture Pool settings.");
+			Info.ExpireDuration = 8.0f;
+			Info.HyperlinkText = LOCTEXT("OpenSettings", "Open Project Settings");
+			Info.Hyperlink = FSimpleDelegate::CreateStatic([]() { FGlobalTabmanager::Get()->TryInvokeTab(FName("ProjectSettings")); });
 
-#if WITH_EDITOR
-		// Notify user about changes.
-		if (ChangedConfigs.Num())
-		{
-			static TWeakPtr<class SNotificationItem> NotificationHandle;
-			if (!NotificationHandle.IsValid())
-			{
-				FNotificationInfo Info(LOCTEXT("PoolResizeNotify", "Resizing Virtual Texture Pools."));
-				Info.SubText = LOCTEXT("PoolResizeNotifySubtext", "Size changes are not saved by default. To keep changes, copy from 'Transient Pools' to 'Fixed Pools' in the Virtual Texture Pool settings.");
-				Info.ExpireDuration = 8.0f;
-				Info.HyperlinkText = LOCTEXT("OpenSettings", "Open Project Settings");
-				Info.Hyperlink = FSimpleDelegate::CreateStatic([]() { FGlobalTabmanager::Get()->TryInvokeTab(FName("ProjectSettings")); });
-			
-				NotificationHandle = FSlateNotificationManager::Get().AddNotification(Info);
-			}
+			NotificationHandle = FSlateNotificationManager::Get().AddNotification(Info);
 		}
-#endif
 	}
-};
+#endif
+}
 
 void VirtualTexturePool::FindPoolConfig(TEnumAsByte<EPixelFormat> const* InFormats, int32 InNumLayers, int32 InTileSize, FVirtualTextureSpacePoolConfig& OutConfig)
 {
@@ -273,7 +247,11 @@ void VirtualTexturePool::FindPoolConfig(TEnumAsByte<EPixelFormat> const* InForma
 void VirtualTexturePool::AddOrModifyTransientPoolConfigs_RenderThread(TArray<FVirtualTextureSpacePoolConfig>& InConfigs)
 {
 	check(IsInRenderingThread());
-	TGraphTask<FAddOrModifyTransientPoolConfigsTask>::CreateTask(nullptr).ConstructAndDispatchWhenReady(InConfigs);
+
+	ExecuteOnGameThread(UE_SOURCE_LOCATION, [Configs = MoveTemp(InConfigs)]()
+	{
+		AddOrModifyTransientPoolConfigs_GameThread(Configs);
+	});
 }
 
 float VirtualTexturePool::GetPoolSizeScale()
