Fix inconsistent tooltip for Percent_FloatFloat. It should say "Modulo (A % B)" like the other Percent_ functions.

#rnx

diff --git a/Engine/Source/Runtime/Engine/Classes/Kismet/KismetMathLibrary.h b/Engine/Source/Runtime/Engine/Classes/Kismet/KismetMathLibrary.h
--- a/Engine/Source/Runtime/Engine/Classes/Kismet/KismetMathLibrary.h
+++ b/Engine/Source/Runtime/Engine/Classes/Kismet/KismetMathLibrary.h
@@ -514,6 +514,7 @@ class UKismetMathLibrary : public UBlueprintFunctionLibrary
 	UFUNCTION(BlueprintPure, meta = (DisplayName = "int * float", CompactNodeTitle = "*", Keywords = "* multiply"), Category = "Math|Float")
 	static ENGINE_API double Multiply_IntFloat(int32 A, double B);
 
+	/** Modulo (A % B) */
 	UFUNCTION(BlueprintPure, CustomThunk, meta = (DisplayName = "% (Float)", CompactNodeTitle = "%", Keywords = "% modulus"), Category = "Math|Float")
 	static ENGINE_API double Percent_FloatFloat(double A, double B = 1.f);
 
