Fixed an issue where convex shapes were not getting a correct ImpactNormal when run on the physics thread.
Also cleared out the face index for convex shapes so that the results always match (Sweep/Ray GT/PT).
Also updated FindGeometryOpposingNormal on all implicit objects to return OriginalNormal instead of a hardcoded vector if the face index is invalid. This matches the default behavior of FImplicitObject.

diff --git a/Engine/Source/Runtime/Engine/Private/Collision/CollisionConversions.cpp b/Engine/Source/Runtime/Engine/Private/Collision/CollisionConversions.cpp
--- a/Engine/Source/Runtime/Engine/Private/Collision/CollisionConversions.cpp
+++ b/Engine/Source/Runtime/Engine/Private/Collision/CollisionConversions.cpp
@@ -438,6 +438,12 @@ EConvertQueryResult ConvertQueryImpactHitImp(const UWorld* World, const THitLoca
 	{
 		OutResult.FaceIndex = GetTriangleMeshExternalFaceIndex(*GTShape, InternalFaceIndex);
 	}
+	else if (GeomType == ECollisionShapeType::Convex)
+	{
+		// Convex gets a face index set so that an impact normal can be set for sweeps.
+		// This isn't a true face index that matches to anything a user would expect, so clear this away.
+		OutResult.FaceIndex = INDEX_NONE;
+	}
 
 	return EConvertQueryResult::Valid;
 }
@@ -464,7 +470,7 @@ EConvertQueryResult ConvertTraceResults(bool& OutHasValidBlockingHit, const UWor
 		HitType& Hit = Hits[i];
 		if(GetDistance(Hit) <= MaxDistance)
 		{
-			if constexpr (std::is_same_v<HitType, FHitSweep>)
+			if constexpr (std::is_same_v<HitType, FHitSweep> || std::is_same_v<HitType, ChaosInterface::FPTSweepHit>)
 			{
 				if (!HadInitialOverlap(Hit))
 				{
@@ -497,7 +503,7 @@ template <typename Hit>
 EConvertQueryResult ConvertTraceResults(bool& OutHasValidBlockingHit, const UWorld* World, int32 NumHits, Hit* Hits, float CheckLength, const FCollisionFilterData& QueryFilter, FHitResult& OutHit, const FVector& StartLoc, const FVector& EndLoc, const FPhysicsGeometry* Geom, const FTransform& QueryTM, float MaxDistance, bool bReturnFaceIndex, bool bReturnPhysMat)
 {
 	const FVector Dir = (EndLoc - StartLoc).GetSafeNormal();
-	if constexpr (std::is_same_v<Hit, FHitSweep>)
+	if constexpr (std::is_same_v<Hit, FHitSweep> || std::is_same_v<Hit, ChaosInterface::FPTSweepHit>)
 	{
 		if (!HadInitialOverlap(Hits[0]))
 		{
diff --git a/Engine/Source/Runtime/Experimental/Chaos/Private/Chaos/HeightField.cpp b/Engine/Source/Runtime/Experimental/Chaos/Private/Chaos/HeightField.cpp
--- a/Engine/Source/Runtime/Experimental/Chaos/Private/Chaos/HeightField.cpp
+++ b/Engine/Source/Runtime/Experimental/Chaos/Private/Chaos/HeightField.cpp
@@ -2461,7 +2461,7 @@ namespace Chaos
 			return Normal;
 		}
 
-		return FVec3(0, 0, 1);
+		return OriginalNormal;
 	}
 
 	void FHeightField::CalcBounds()
diff --git a/Engine/Source/Runtime/Experimental/Chaos/Private/Chaos/TriangleMeshImplicitObject.cpp b/Engine/Source/Runtime/Experimental/Chaos/Private/Chaos/TriangleMeshImplicitObject.cpp
--- a/Engine/Source/Runtime/Experimental/Chaos/Private/Chaos/TriangleMeshImplicitObject.cpp
+++ b/Engine/Source/Runtime/Experimental/Chaos/Private/Chaos/TriangleMeshImplicitObject.cpp
@@ -1872,7 +1872,11 @@ int32 FTriangleMeshImplicitObject::FindMostOpposingFaceScaled(const FVec3& Posit
 
 FVec3 FTriangleMeshImplicitObject::FindGeometryOpposingNormal(const FVec3& DenormDir, int32 FaceIndex, const FVec3& OriginalNormal) const
 {
-	return GetFaceNormal(FaceIndex);
+	if (CHAOS_ENSURE(FaceIndex != INDEX_NONE))
+	{
+		return GetFaceNormal(FaceIndex);
+	}
+	return OriginalNormal;
 }
 
 template <typename IdxType>
diff --git a/Engine/Source/Runtime/Experimental/Chaos/Public/Chaos/Convex.h b/Engine/Source/Runtime/Experimental/Chaos/Public/Chaos/Convex.h
--- a/Engine/Source/Runtime/Experimental/Chaos/Public/Chaos/Convex.h
+++ b/Engine/Source/Runtime/Experimental/Chaos/Public/Chaos/Convex.h
@@ -559,7 +559,7 @@ namespace Chaos
 				const FPlaneType& OpposingFace = GetFaces()[FaceIndex];
 				return OpposingFace.Normal();
 			}
-			return FVec3(0.f, 0.f, 1.f);
+			return OriginalNormal;
 		}
 
 		CHAOS_API virtual int32 FindClosestFaceAndVertices(const FVec3& Position, TArray<FVec3>& FaceVertices, FReal SearchDist = 0.01f) const override;
diff --git a/Engine/Source/Runtime/Experimental/Chaos/Public/Chaos/ImplicitObjectScaled.h b/Engine/Source/Runtime/Experimental/Chaos/Public/Chaos/ImplicitObjectScaled.h
--- a/Engine/Source/Runtime/Experimental/Chaos/Public/Chaos/ImplicitObjectScaled.h
+++ b/Engine/Source/Runtime/Experimental/Chaos/Public/Chaos/ImplicitObjectScaled.h
@@ -861,7 +861,7 @@ public:
 		if (NormalLength == 0)
 		{
 			CHAOS_ENSURE(false);
-			Normal = FVec3(0, 0, 1);
+			Normal = OriginalNormal;
 		}
 
 		return Normal;
