[Actor/LastRenderTime]
 - Fix race on last render time

#rnx

diff --git a/Engine/Source/Runtime/Engine/Classes/GameFramework/Actor.h b/Engine/Source/Runtime/Engine/Classes/GameFramework/Actor.h
--- a/Engine/Source/Runtime/Engine/Classes/GameFramework/Actor.h
+++ b/Engine/Source/Runtime/Engine/Classes/GameFramework/Actor.h
@@ -17,6 +17,8 @@
 #include "Net/Core/Misc/NetSubObjectRegistry.h"
 #include "Engine/ReplicatedState.h"
 
+#include <atomic>
+
 #if WITH_EDITOR
 #include "Folder.h"
 #endif
@@ -129,7 +131,23 @@ enum class UE_DEPRECATED(5.0, "EActorGridPlacement is deprecated.") EActorGridPl
 /** Helper struct that allows UPrimitiveComponent and FPrimitiveSceneInfo write to the Actor's LastRenderTime member */
 struct FActorLastRenderTime
 {
-	float LastRenderTime;
+	UE_DEPRECATED(5.6, "Use GetLastRenderTime/SetLastRenderTime instead.")
+	std::atomic<float> LastRenderTime;
+
+	void SetLastRenderTime(float InLastRenderTime)
+	{
+		PRAGMA_DISABLE_DEPRECATION_WARNINGS
+		LastRenderTime.store(InLastRenderTime, std::memory_order_relaxed);
+		PRAGMA_ENABLE_DEPRECATION_WARNINGS
+	}
+
+	float GetLastRenderTime() const
+	{
+		PRAGMA_DISABLE_DEPRECATION_WARNINGS
+		return LastRenderTime.load(std::memory_order_relaxed);
+		PRAGMA_ENABLE_DEPRECATION_WARNINGS
+	}
+
 	std::atomic_int32_t NumAlwaysVisibleComponents = 0;
 
 private:
@@ -4577,7 +4595,7 @@ private:
 
 inline void FActorLastRenderTime::Set(AActor* InActor, float LastRenderTime)
 {
-	InActor->LastRenderTime.LastRenderTime = LastRenderTime;
+	InActor->LastRenderTime.SetLastRenderTime(LastRenderTime);
 }
 
 inline FActorLastRenderTime* FActorLastRenderTime::GetPtr(AActor* InActor)
diff --git a/Engine/Source/Runtime/Engine/Private/Actor.cpp b/Engine/Source/Runtime/Engine/Private/Actor.cpp
--- a/Engine/Source/Runtime/Engine/Private/Actor.cpp
+++ b/Engine/Source/Runtime/Engine/Private/Actor.cpp
@@ -2173,7 +2173,7 @@ float AActor::GetLastRenderTime() const
 			return World->GetTimeSeconds();
 		}
 	}
-	return LastRenderTime.LastRenderTime;
+	return LastRenderTime.GetLastRenderTime();
 }
 
 void AActor::SetOwner(AActor* NewOwner)
diff --git a/Engine/Source/Runtime/Engine/Private/PrimitiveSceneInfoData.cpp b/Engine/Source/Runtime/Engine/Private/PrimitiveSceneInfoData.cpp
--- a/Engine/Source/Runtime/Engine/Private/PrimitiveSceneInfoData.cpp
+++ b/Engine/Source/Runtime/Engine/Private/PrimitiveSceneInfoData.cpp
@@ -18,6 +18,6 @@ void FPrimitiveSceneInfoData::SetLastRenderTime(float InLastRenderTime, bool bUp
 
 	if (OwnerLastRenderTimePtr)
 	{
-		OwnerLastRenderTimePtr->LastRenderTime = InLastRenderTime;
+		OwnerLastRenderTimePtr->SetLastRenderTime(InLastRenderTime);
 	}
 }
\ No newline at end of file
