[ALT]
 - Fix crash when async loading in editor is turned on since we can't possibly have a mesh compiling if it is still being loaded

#rnx

diff --git a/Engine/Source/Runtime/Engine/Private/StaticMesh.cpp b/Engine/Source/Runtime/Engine/Private/StaticMesh.cpp
--- a/Engine/Source/Runtime/Engine/Private/StaticMesh.cpp
+++ b/Engine/Source/Runtime/Engine/Private/StaticMesh.cpp
@@ -6293,13 +6293,17 @@ void UStaticMesh::Serialize(FArchive& Ar)
 	}
 
 #if WITH_EDITOR
-	// Since UPROPERTY are accessed directly by offset during serialization instead of using accessors, 
-	// the protection put in place to automatically finish compilation if a locked property is accessed will not work. 
-	// We have no choice but to force finish the compilation here to avoid potential race conditions between 
-	// async compilation and the serialization.
-	FStaticMeshCompilingManager::FFinishCompilationOptions FinishCompilationOptions;
-	FinishCompilationOptions.bIncludeDependentMeshes = Ar.IsLoading(); // Wait for dependent meshes if we're about to be editing ourselves
-	FStaticMeshCompilingManager::Get().FinishCompilation({this}, FinishCompilationOptions);
+	// Don't try to finish compilation while still in the loader as we can't possibly have our mesh being compiled if it has not been loaded yet.
+	if (!IsInAsyncLoadingThread())
+	{
+		// Since UPROPERTY are accessed directly by offset during serialization instead of using accessors, 
+		// the protection put in place to automatically finish compilation if a locked property is accessed will not work. 
+		// We have no choice but to force finish the compilation here to avoid potential race conditions between 
+		// async compilation and the serialization.
+		FStaticMeshCompilingManager::FFinishCompilationOptions FinishCompilationOptions;
+		FinishCompilationOptions.bIncludeDependentMeshes = Ar.IsLoading(); // Wait for dependent meshes if we're about to be editing ourselves
+		FStaticMeshCompilingManager::Get().FinishCompilation({this}, FinishCompilationOptions);
+	}
 #endif
 
 	{
