Added a physics query handler. This allows redirection of queries for eventual use with distributed worlds.
At the moment, the framework is in place but nothing will trigger the pathway to use the default query handler (later review).

diff --git a/Engine/Source/Runtime/Engine/Classes/Engine/World.h b/Engine/Source/Runtime/Engine/Classes/Engine/World.h
--- a/Engine/Source/Runtime/Engine/Classes/Engine/World.h
+++ b/Engine/Source/Runtime/Engine/Classes/Engine/World.h
@@ -21,6 +21,7 @@
 #include "EngineDefines.h"
 #include "Engine/PendingNetGame.h"
 #include "Engine/LatentActionManager.h"
+#include "Physics/PhysicsQueryHandler.h"
 #include "Physics/PhysicsInterfaceDeclares.h"
 #include "Particles/WorldPSCPool.h"
 #include "Containers/SortedMap.h"
@@ -901,6 +902,10 @@ class ENGINE_API UWorld final : public UObject, public FNetworkNotify
 	UPROPERTY(Transient)
 	TObjectPtr<class UPhysicsCollisionHandler>				PhysicsCollisionHandler;
 
+	/** Instance of this world's game-specific physics query handler */
+	UPROPERTY(Transient)
+	TObjectPtr<class UPhysicsQueryHandler>				    PhysicsQueryHandler;
+
 	/** Array of any additional objects that need to be referenced by this world, to make sure they aren't GC'd */
 	UPROPERTY(Transient)
 	TArray<TObjectPtr<UObject>>							ExtraReferencedObjects;
diff --git a/Engine/Source/Runtime/Engine/Private/Collision/SceneQuery.cpp b/Engine/Source/Runtime/Engine/Private/Collision/SceneQuery.cpp
--- a/Engine/Source/Runtime/Engine/Private/Collision/SceneQuery.cpp
+++ b/Engine/Source/Runtime/Engine/Private/Collision/SceneQuery.cpp
@@ -5,6 +5,8 @@
 #include "CollisionDebugDrawingPublic.h"
 #include "Physics/Experimental/PhysScene_Chaos.h"
 #include "Physics/PhysicsInterfaceUtils.h"
+#include "Physics/PhysicsQueryHandler.h"
+#include "Physics/SceneQueryData.h"
 #include "Collision/CollisionConversions.h"
 #include "PhysicsEngine/ClusterUnionComponent.h"
 #include "PhysicsEngine/CollisionQueryFilterCallback.h"
@@ -166,6 +168,37 @@ TGeomSQInputs<TGeom> GeomToSQInputs(const TGeom& Geom, const FQuat& Rot)
 	}
 }
 
+bool ToCollisionShape(const FPhysicsGeometryCollection& GeometryCollection, FCollisionShape& OutCollisionShape)
+{
+	ECollisionShapeType ShapeType = GeometryCollection.GetType();
+	switch (ShapeType)
+	{
+		case ECollisionShapeType::Box:
+		{
+			const Chaos::TBox<Chaos::FReal, 3>& Box = GeometryCollection.GetBoxGeometry();
+			OutCollisionShape = FCollisionShape::MakeBox(Box.Extents() * 0.5f);
+			return true;
+		}
+		case ECollisionShapeType::Sphere:
+		{
+			const Chaos::TSphere<Chaos::FReal, 3>& Sphere = GeometryCollection.GetSphereGeometry();
+			OutCollisionShape = FCollisionShape::MakeSphere(Sphere.GetRadius());
+			return true;
+		}
+		case ECollisionShapeType::Capsule:
+		{
+			const Chaos::FCapsule& Capsule = GeometryCollection.GetCapsuleGeometry();
+			OutCollisionShape = FCollisionShape::MakeCapsule(Capsule.GetRadius(), Capsule.GetHeight() * 0.5f);
+			return true;
+		}
+		default:
+		{
+			UE_LOG(LogChaos, Warning, TEXT("PhysicsQueryHandler: Invalid shape type. Shape type %d is not yet supported"), (int)ShapeType);
+		}
+	}
+	return false;
+}
+
 template <typename InHitType, ESweepOrRay InGeometryQuery, ESingleMultiOrTest InSingleMultiOrTest>
 struct TSQTraits
 {
@@ -357,12 +390,7 @@ struct TSQTraits
 
 };
 
-enum class EThreadQueryContext
-{
-	GTData,		//use interpolated GT data
-	PTDataWithGTObjects,	//use pt data, but convert back to GT when possible
-	PTOnlyData,	//use only the PT data and don't try to convert anything back to GT
-};
+using EThreadQueryContext = Chaos::EThreadQueryContext;
 
 EThreadQueryContext GetThreadQueryContext(const Chaos::FPhysicsSolver& Solver)
 {
@@ -689,8 +717,109 @@ bool TSceneCastCommonImp(const UWorld* World, typename Traits::TOutHits& OutHits
 	return bReturnResult;
 }
 
-template <typename Traits, typename PTTraits, typename TGeomInputs, typename TAccelContainer = FDefaultAccelContainer>
-bool TSceneCastCommon(const UWorld* World, typename Traits::TOutHits& OutHits, const TGeomInputs& GeomInputs, const FVector Start, const FVector End, ECollisionChannel TraceChannel, const struct FCollisionQueryParams& Params, const struct FCollisionResponseParams& ResponseParams, const struct FCollisionObjectQueryParams& ObjectParams, const TAccelContainer& AccelContainer = FDefaultAccelContainer{})
+// This selects the GT/PT traits based upon the given thread context.
+template <typename HitType, typename PTHitType, ESweepOrRay SweepOrRay, ESingleMultiOrTest SingleMultiOrTest, typename GeomInputsType, typename OutHitsType, typename AccelContainerType>
+bool TraceCommonImp(const EThreadQueryContext ThreadContext, const UWorld* World, OutHitsType& OutHits, const GeomInputsType& GeomInputs, const FVector Start, const FVector End, ECollisionChannel TraceChannel, const struct FCollisionQueryParams& Params, const struct FCollisionResponseParams& ResponseParams, const struct FCollisionObjectQueryParams& ObjectParams, const AccelContainerType& AccelContainer)
+{
+	using namespace ChaosInterface;
+	using Traits = TSQTraits<HitType, SweepOrRay, SingleMultiOrTest>;
+	using PTTraits = TSQTraits<PTHitType, SweepOrRay, SingleMultiOrTest>;
+	if (ThreadContext == EThreadQueryContext::GTData)
+	{
+		return TSceneCastCommonImp<Traits, GeomInputsType>(World, OutHits, GeomInputs, Start, End, TraceChannel, Params, ResponseParams, ObjectParams, AccelContainer);
+	}
+	else
+	{
+		return TSceneCastCommonImp<PTTraits, GeomInputsType>(World, OutHits, GeomInputs, Start, End, TraceChannel, Params, ResponseParams, ObjectParams, AccelContainer);
+	}
+}
+
+// This converts ESweepOrRay to the two different GT/PT Hit types
+template <ESweepOrRay SweepOrRay, ESingleMultiOrTest SingleMultiOrTest, typename GeomInputsType, typename OutHitsType, typename AccelContainerType>
+bool TraceCommonImp(const EThreadQueryContext ThreadContext, const UWorld* World, OutHitsType& OutHits, const GeomInputsType& GeomInputs, const FVector Start, const FVector End, ECollisionChannel TraceChannel, const struct FCollisionQueryParams& Params, const struct FCollisionResponseParams& ResponseParams, const struct FCollisionObjectQueryParams& ObjectParams, const AccelContainerType& AccelContainer)
+{
+	using namespace ChaosInterface;
+	if constexpr (SweepOrRay == ESweepOrRay::Raycast)
+	{
+		return TraceCommonImp<FRaycastHit, FPTRaycastHit, SweepOrRay, SingleMultiOrTest>(ThreadContext, World, OutHits, GeomInputs, Start, End, TraceChannel, Params, ResponseParams, ObjectParams, AccelContainer);
+	}
+	else
+	{
+		return TraceCommonImp<FSweepHit, FPTSweepHit, SweepOrRay, SingleMultiOrTest>(ThreadContext, World, OutHits, GeomInputs, Start, End, TraceChannel, Params, ResponseParams, ObjectParams, AccelContainer);
+	}
+}
+
+template <ESingleMultiOrTest SingleMultiOrTest, typename GeomInputsType, typename OutHitsType>
+bool RaycastWithQueryHandler(UPhysicsQueryHandler& QueryHandler, const EThreadQueryContext ThreadContext, const UWorld* World, OutHitsType& OutHits, const GeomInputsType& GeomInputs, const FVector Start, const FVector End, const Chaos::FCommonQueryData& CommonData)
+{
+	static_assert(std::is_same_v<GeomInputsType, FRaycastSQAdditionalInputs>, "Invalid GeomInputs passed to Raycast Query.");
+
+	Chaos::FRayQueryData RayData
+	{
+		.Start = Start,
+		.End = End,
+	};
+	if constexpr (SingleMultiOrTest == ESingleMultiOrTest::Test)
+	{
+		return QueryHandler.RaycastTest(ThreadContext, World, RayData, CommonData, OutHits);
+	}
+	else if constexpr (SingleMultiOrTest == ESingleMultiOrTest::Single)
+	{
+		return QueryHandler.RaycastSingle(ThreadContext, World, RayData, CommonData, OutHits);
+	}
+	else
+	{
+		static_assert(SingleMultiOrTest == ESingleMultiOrTest::Multi);
+		return QueryHandler.RaycastMulti(ThreadContext, World, RayData, CommonData, OutHits);
+	}
+}
+
+template <ESingleMultiOrTest SingleMultiOrTest, typename OutHitsType>
+bool SweepWithQueryHandler(UPhysicsQueryHandler& QueryHandler, const EThreadQueryContext ThreadContext, const UWorld* World, OutHitsType& OutHits, const FCollisionShape& CollisionShape, const FQuat& GeomRot, const FVector Start, const FVector End, const Chaos::FCommonQueryData& CommonData)
+{
+	Chaos::FSweepQueryData SweepData
+	{
+		.Start = Start,
+		.End = End,
+		.CollisionShape = CollisionShape,
+		.GeomRot = GeomRot,
+	};
+	if constexpr (SingleMultiOrTest == ESingleMultiOrTest::Test)
+	{
+		return QueryHandler.SweepTest(ThreadContext, World, SweepData, CommonData, OutHits);
+	}
+
+	else if constexpr (SingleMultiOrTest == ESingleMultiOrTest::Single)
+	{
+		return QueryHandler.SweepSingle(ThreadContext, World, SweepData, CommonData, OutHits);
+	}
+	else
+	{
+		static_assert(SingleMultiOrTest == ESingleMultiOrTest::Multi);
+		return QueryHandler.SweepMulti(ThreadContext, World, SweepData, CommonData, OutHits);
+	}
+}
+
+template <ESingleMultiOrTest SingleMultiOrTest, typename OutHitsType>
+bool SweepWithQueryHandler(UPhysicsQueryHandler& QueryHandler, const EThreadQueryContext ThreadContext, const UWorld* World, OutHitsType& OutHits, const FGeomSQAdditionalInputs& GeomInputs, const FVector Start, const FVector End, const Chaos::FCommonQueryData& CommonData)
+{
+	return SweepWithQueryHandler<SingleMultiOrTest>(QueryHandler, ThreadContext, World, OutHits, GeomInputs.CollisionShape, *GeomInputs.GetGeometryOrientation(), Start, End, CommonData);
+}
+
+template <ESingleMultiOrTest SingleMultiOrTest, typename OutHitsType>
+bool SweepWithQueryHandler(UPhysicsQueryHandler& QueryHandler, const EThreadQueryContext ThreadContext, const UWorld* World, OutHitsType& OutHits, const FGeomCollectionSQAdditionalInputs& GeomInputs, const FVector Start, const FVector End, const Chaos::FCommonQueryData& CommonData)
+{
+	const FQuat& GeomRot = GeomInputs.GeomRot;
+	FCollisionShape CollisionShape;
+	if (ToCollisionShape(GeomInputs.Collection, CollisionShape))
+	{
+		return SweepWithQueryHandler<SingleMultiOrTest>(QueryHandler, ThreadContext, World, OutHits, CollisionShape, GeomRot, Start, End, CommonData);
+	}
+	return false;
+}
+
+template <ESweepOrRay SweepOrRay, ESingleMultiOrTest SingleMultiOrTest, typename OutHitsType, typename GeomInputsType, typename AccelContainerType = FDefaultAccelContainer>
+bool TSceneCastCommon(const UWorld* World, OutHitsType& OutHits, const GeomInputsType& GeomInputs, const FVector Start, const FVector End, ECollisionChannel TraceChannel, const struct FCollisionQueryParams& Params, const struct FCollisionResponseParams& ResponseParams, const struct FCollisionObjectQueryParams& ObjectParams, const AccelContainerType& AccelContainer = FDefaultAccelContainer{})
 {
 	if ((World == NULL) || (World->GetPhysicsScene() == NULL))
 	{
@@ -698,13 +827,36 @@ bool TSceneCastCommon(const UWorld* World, typename Traits::TOutHits& OutHits, c
 	}
 
 	const EThreadQueryContext ThreadContext = GetThreadQueryContext(*World->GetPhysicsScene()->GetSolver());
-	if(ThreadContext == EThreadQueryContext::GTData)
+	// If there's an acceleration structure override, always do a local query as the query handler is only needed for scene queries.
+	if constexpr (AccelContainerType::HasAccelerationStructureOverride())
 	{
-		return TSceneCastCommonImp<Traits, TGeomInputs>(World, OutHits, GeomInputs, Start, End, TraceChannel, Params, ResponseParams, ObjectParams, AccelContainer);
+		return TraceCommonImp<SweepOrRay, SingleMultiOrTest>(ThreadContext, World, OutHits, GeomInputs, Start, End, TraceChannel, Params, ResponseParams, ObjectParams, AccelContainer);
 	}
 	else
 	{
-		return TSceneCastCommonImp<PTTraits, TGeomInputs>(World, OutHits, GeomInputs, Start, End, TraceChannel, Params, ResponseParams, ObjectParams, AccelContainer);
+		UPhysicsQueryHandler* QueryHandler = World->PhysicsQueryHandler;
+		if (QueryHandler == nullptr)
+		{
+			return TraceCommonImp<SweepOrRay, SingleMultiOrTest>(ThreadContext, World, OutHits, GeomInputs, Start, End, TraceChannel, Params, ResponseParams, ObjectParams, AccelContainer);
+		}
+		else
+		{
+			Chaos::FCommonQueryData CommonData
+			{
+				.TraceChannel = TraceChannel,
+				.Params = Params,
+				.ResponseParams = ResponseParams,
+				.ObjectParams = ObjectParams,
+			};
+			if constexpr (SweepOrRay == ESweepOrRay::Raycast)
+			{
+				return RaycastWithQueryHandler<SingleMultiOrTest>(*QueryHandler, ThreadContext, World, OutHits, GeomInputs, Start, End, CommonData);
+			}
+			else
+			{
+				return SweepWithQueryHandler<SingleMultiOrTest>(*QueryHandler, ThreadContext, World, OutHits, GeomInputs, Start, End, CommonData);
+			}
+		}
 	}
 }
 
@@ -719,11 +871,9 @@ bool FGenericPhysicsInterface::RaycastTest(const UWorld* World, const FVector St
 	SCOPE_CYCLE_COUNTER(STAT_Collision_RaycastAny);
 	CSV_SCOPED_TIMING_STAT(SceneQuery, RaycastTest);
 
-	using TCastTraits = TSQTraits<FHitRaycast, ESweepOrRay::Raycast, ESingleMultiOrTest::Test>;
-	using TPTCastTraits = TSQTraits<FPTRaycastHit, ESweepOrRay::Raycast, ESingleMultiOrTest::Test>;
 	FHitResult DummyHit(NoInit);
 
-	return TSceneCastCommon<TCastTraits, TPTCastTraits>(World, DummyHit, FRaycastSQAdditionalInputs(), Start, End, TraceChannel, Params, ResponseParams, ObjectParams);
+	return TSceneCastCommon<ESweepOrRay::Raycast, ESingleMultiOrTest::Test>(World, DummyHit, FRaycastSQAdditionalInputs(), Start, End, TraceChannel, Params, ResponseParams, ObjectParams);
 }
 
 template<typename TAccel>
@@ -735,11 +885,9 @@ bool FGenericRaycastPhysicsInterfaceUsingSpatialAcceleration<TAccel>::RaycastTes
 	SCOPE_CYCLE_COUNTER(STAT_Collision_RaycastAny);
 	CSV_SCOPED_TIMING_STAT(SceneQuery, RaycastTest);
 
-	using TCastTraits = TSQTraits<FHitRaycast, ESweepOrRay::Raycast, ESingleMultiOrTest::Test>;
-	using TPTCastTraits = TSQTraits<FPTRaycastHit, ESweepOrRay::Raycast, ESingleMultiOrTest::Test>;
 	FHitResult DummyHit(NoInit);
 
-	return TSceneCastCommon<TCastTraits, TPTCastTraits>(World, DummyHit, FRaycastSQAdditionalInputs(), Start, End, TraceChannel, Params, ResponseParams, ObjectParams, FOverrideAccelContainer{ Accel });
+	return TSceneCastCommon<ESweepOrRay::Raycast, ESingleMultiOrTest::Test>(World, DummyHit, FRaycastSQAdditionalInputs(), Start, End, TraceChannel, Params, ResponseParams, ObjectParams, FOverrideAccelContainer{ Accel });
 }
 
 bool FGenericPhysicsInterface::RaycastSingle(const UWorld* World, struct FHitResult& OutHit, const FVector Start, const FVector End, ECollisionChannel TraceChannel, const struct FCollisionQueryParams& Params, const struct FCollisionResponseParams& ResponseParams, const struct FCollisionObjectQueryParams& ObjectParams)
@@ -750,9 +898,7 @@ bool FGenericPhysicsInterface::RaycastSingle(const UWorld* World, struct FHitRes
 	SCOPE_CYCLE_COUNTER(STAT_Collision_RaycastSingle);
 	CSV_SCOPED_TIMING_STAT(SceneQuery, RaycastSingle);
 
-	using TCastTraits = TSQTraits<FRaycastHit, ESweepOrRay::Raycast, ESingleMultiOrTest::Single>;
-	using TPTCastTraits = TSQTraits<FPTRaycastHit, ESweepOrRay::Raycast, ESingleMultiOrTest::Single>;
-	return TSceneCastCommon<TCastTraits, TPTCastTraits>(World, OutHit, FRaycastSQAdditionalInputs(), Start, End, TraceChannel, Params, ResponseParams, ObjectParams);
+	return TSceneCastCommon<ESweepOrRay::Raycast, ESingleMultiOrTest::Single>(World, OutHit, FRaycastSQAdditionalInputs(), Start, End, TraceChannel, Params, ResponseParams, ObjectParams);
 }
 
 template<typename TAccel>
@@ -764,9 +910,7 @@ bool FGenericRaycastPhysicsInterfaceUsingSpatialAcceleration<TAccel>::RaycastSin
 	SCOPE_CYCLE_COUNTER(STAT_Collision_RaycastSingle);
 	CSV_SCOPED_TIMING_STAT(SceneQuery, RaycastSingle);
 
-	using TCastTraits = TSQTraits<FRaycastHit, ESweepOrRay::Raycast, ESingleMultiOrTest::Single>;
-	using TPTCastTraits = TSQTraits<FPTRaycastHit, ESweepOrRay::Raycast, ESingleMultiOrTest::Single>;
-	return TSceneCastCommon<TCastTraits, TPTCastTraits>(World, OutHit, FRaycastSQAdditionalInputs(), Start, End, TraceChannel, Params, ResponseParams, ObjectParams, FOverrideAccelContainer{ Accel });
+	return TSceneCastCommon<ESweepOrRay::Raycast, ESingleMultiOrTest::Single>(World, OutHit, FRaycastSQAdditionalInputs(), Start, End, TraceChannel, Params, ResponseParams, ObjectParams, FOverrideAccelContainer{ Accel });
 }
 
 
@@ -778,9 +922,7 @@ bool FGenericPhysicsInterface::RaycastMulti(const UWorld* World, TArray<struct F
 	SCOPE_CYCLE_COUNTER(STAT_Collision_RaycastMultiple);
 	CSV_SCOPED_TIMING_STAT(SceneQuery, RaycastMultiple);
 
-	using TCastTraits = TSQTraits<FHitRaycast, ESweepOrRay::Raycast, ESingleMultiOrTest::Multi>;
-	using TPTCastTraits = TSQTraits<FPTRaycastHit, ESweepOrRay::Raycast, ESingleMultiOrTest::Multi>;
-	return TSceneCastCommon<TCastTraits, TPTCastTraits> (World, OutHits, FRaycastSQAdditionalInputs(), Start, End, TraceChannel, Params, ResponseParams, ObjectParams);
+	return TSceneCastCommon<ESweepOrRay::Raycast, ESingleMultiOrTest::Multi>(World, OutHits, FRaycastSQAdditionalInputs(), Start, End, TraceChannel, Params, ResponseParams, ObjectParams);
 }
 
 template<typename TAccel>
@@ -792,9 +934,7 @@ bool FGenericRaycastPhysicsInterfaceUsingSpatialAcceleration<TAccel>::RaycastMul
 	SCOPE_CYCLE_COUNTER(STAT_Collision_RaycastMultiple);
 	CSV_SCOPED_TIMING_STAT(SceneQuery, RaycastMultiple);
 
-	using TCastTraits = TSQTraits<FHitRaycast, ESweepOrRay::Raycast, ESingleMultiOrTest::Multi>;
-	using TPTCastTraits = TSQTraits<FPTRaycastHit, ESweepOrRay::Raycast, ESingleMultiOrTest::Multi>;
-	return TSceneCastCommon<TCastTraits, TPTCastTraits>(World, OutHits, FRaycastSQAdditionalInputs(), Start, End, TraceChannel, Params, ResponseParams, ObjectParams, FOverrideAccelContainer{ Accel });
+	return TSceneCastCommon<ESweepOrRay::Raycast, ESingleMultiOrTest::Multi>(World, OutHits, FRaycastSQAdditionalInputs(), Start, End, TraceChannel, Params, ResponseParams, ObjectParams, FOverrideAccelContainer{ Accel });
 }
 
 //////////////////////////////////////////////////////////////////////////
@@ -808,11 +948,9 @@ bool FGenericPhysicsInterface::GeomSweepTest(const UWorld* World, const struct F
 	SCOPE_CYCLE_COUNTER(STAT_Collision_GeomSweepAny);
 	CSV_SCOPED_TIMING_STAT(SceneQuery, GeomSweepTest);
 
-	using TCastTraits = TSQTraits<FHitSweep, ESweepOrRay::Sweep, ESingleMultiOrTest::Test>;
-	using TPTCastTraits = TSQTraits<FPTSweepHit, ESweepOrRay::Sweep, ESingleMultiOrTest::Test>;
 	FHitResult DummyHit(NoInit);
 
-	return TSceneCastCommon<TCastTraits, TPTCastTraits>(World, DummyHit, FGeomSQAdditionalInputs(CollisionShape, Rot), Start, End, TraceChannel, Params, ResponseParams, ObjectParams);
+	return TSceneCastCommon<ESweepOrRay::Sweep, ESingleMultiOrTest::Test>(World, DummyHit, FGeomSQAdditionalInputs(CollisionShape, Rot), Start, End, TraceChannel, Params, ResponseParams, ObjectParams);
 }
 
 template<typename TAccel, typename TGeom>
@@ -824,11 +962,9 @@ bool FGenericGeomPhysicsInterfaceUsingSpatialAcceleration<TAccel, TGeom>::GeomSw
 	SCOPE_CYCLE_COUNTER(STAT_Collision_GeomSweepAny);
 	CSV_SCOPED_TIMING_STAT(SceneQuery, GeomSweepTest);
 
-	using TCastTraits = TSQTraits<FHitSweep, ESweepOrRay::Sweep, ESingleMultiOrTest::Test>;
-	using TPTCastTraits = TSQTraits<FPTSweepHit, ESweepOrRay::Sweep, ESingleMultiOrTest::Test>;
 	FHitResult DummyHit(NoInit);
 
-	return TSceneCastCommon<TCastTraits, TPTCastTraits>(World, DummyHit, GeomToSQInputs(InGeom, Rot), Start, End, TraceChannel, Params, ResponseParams, ObjectParams, FOverrideAccelContainer{ Accel });
+	return TSceneCastCommon<ESweepOrRay::Sweep, ESingleMultiOrTest::Test>(World, DummyHit, GeomToSQInputs(InGeom, Rot), Start, End, TraceChannel, Params, ResponseParams, ObjectParams, FOverrideAccelContainer{ Accel });
 }
 
 template<>
@@ -840,9 +976,7 @@ bool FGenericPhysicsInterface::GeomSweepSingle(const UWorld* World, const FPhysi
 	SCOPE_CYCLE_COUNTER(STAT_Collision_GeomSweepSingle);
 	CSV_SCOPED_TIMING_STAT(SceneQuery, GeomSweepSingle);
 
-	using TCastTraits = TSQTraits<FSweepHit, ESweepOrRay::Sweep, ESingleMultiOrTest::Single>;
-	using TPTCastTraits = TSQTraits<FPTSweepHit, ESweepOrRay::Sweep, ESingleMultiOrTest::Single>;
-	return TSceneCastCommon<TCastTraits, TPTCastTraits>(World, OutHit, FGeomCollectionSQAdditionalInputs(InGeom, Rot), Start, End, TraceChannel, Params, ResponseParams, ObjectParams);
+	return TSceneCastCommon<ESweepOrRay::Sweep, ESingleMultiOrTest::Single>(World, OutHit, FGeomCollectionSQAdditionalInputs(InGeom, Rot), Start, End, TraceChannel, Params, ResponseParams, ObjectParams);
 }
 
 template<>
@@ -854,9 +988,7 @@ bool FGenericPhysicsInterface::GeomSweepSingle(const UWorld* World, const struct
 	SCOPE_CYCLE_COUNTER(STAT_Collision_GeomSweepSingle);
 	CSV_SCOPED_TIMING_STAT(SceneQuery, GeomSweepSingle);
 
-	using TCastTraits = TSQTraits<FSweepHit, ESweepOrRay::Sweep, ESingleMultiOrTest::Single>;
-	using TPTCastTraits = TSQTraits<FPTSweepHit, ESweepOrRay::Sweep, ESingleMultiOrTest::Single>;
-	return TSceneCastCommon<TCastTraits, TPTCastTraits>(World, OutHit, FGeomSQAdditionalInputs(CollisionShape, Rot), Start, End, TraceChannel, Params, ResponseParams, ObjectParams);
+	return TSceneCastCommon<ESweepOrRay::Sweep, ESingleMultiOrTest::Single>(World, OutHit, FGeomSQAdditionalInputs(CollisionShape, Rot), Start, End, TraceChannel, Params, ResponseParams, ObjectParams);
 }
 
 template<>
@@ -874,9 +1006,7 @@ bool FGenericPhysicsInterface::GeomSweepMulti(const UWorld* World, const FPhysic
 	SCOPE_CYCLE_COUNTER(STAT_Collision_GeomSweepMultiple);
 	CSV_SCOPED_TIMING_STAT(SceneQuery, GeomSweepMultiple);
 
-	using TCastTraits = TSQTraits<FHitSweep, ESweepOrRay::Sweep, ESingleMultiOrTest::Multi>;
-	using TPTCastTraits = TSQTraits<FPTSweepHit, ESweepOrRay::Sweep, ESingleMultiOrTest::Multi>;
-	return TSceneCastCommon<TCastTraits, TPTCastTraits>(World, OutHits, FGeomCollectionSQAdditionalInputs(InGeom, InGeomRot), Start, End, TraceChannel, Params, ResponseParams, ObjectParams);
+	return TSceneCastCommon<ESweepOrRay::Sweep, ESingleMultiOrTest::Multi>(World, OutHits, FGeomCollectionSQAdditionalInputs(InGeom, InGeomRot), Start, End, TraceChannel, Params, ResponseParams, ObjectParams);
 }
 
 template<>
@@ -888,9 +1018,7 @@ bool FGenericPhysicsInterface::GeomSweepMulti(const UWorld* World, const FCollis
 	SCOPE_CYCLE_COUNTER(STAT_Collision_GeomSweepMultiple);
 	CSV_SCOPED_TIMING_STAT(SceneQuery, GeomSweepMultiple);
 
-	using TCastTraits = TSQTraits<FHitSweep, ESweepOrRay::Sweep, ESingleMultiOrTest::Multi>;
-	using TPTCastTraits = TSQTraits<FPTSweepHit, ESweepOrRay::Sweep, ESingleMultiOrTest::Multi>;
-	return TSceneCastCommon<TCastTraits, TPTCastTraits>(World, OutHits, FGeomSQAdditionalInputs(InGeom, InGeomRot), Start, End, TraceChannel, Params, ResponseParams, ObjectParams);
+	return TSceneCastCommon<ESweepOrRay::Sweep, ESingleMultiOrTest::Multi>(World, OutHits, FGeomSQAdditionalInputs(InGeom, InGeomRot), Start, End, TraceChannel, Params, ResponseParams, ObjectParams);
 }
 
 template<>
@@ -902,18 +1030,9 @@ bool FGenericPhysicsInterface::GeomSweepMulti(const UWorld* World, const FPhysic
 //////////////////////////////////////////////////////////////////////////
 // GEOM OVERLAP
 
-namespace EQueryInfo
-{
-	//This is used for templatizing code based on the info we're trying to get out.
-	enum Type
-	{
-		GatherAll,		//get all data and actually return it
-		IsBlocking,		//is any of the data blocking? only return a bool so don't bother collecting
-		IsAnything		//is any of the data blocking or touching? only return a bool so don't bother collecting
-	};
-}
+using EQueryInfo = Chaos::EQueryInfo;
 
-template <typename TOverlapHit, EQueryInfo::Type InfoType, typename TCollisionAnalyzerType, typename TAccelContainer>
+template <typename TOverlapHit, EQueryInfo InfoType, typename TCollisionAnalyzerType, typename TAccelContainer>
 bool GeomOverlapMultiImp(const UWorld* World, const FPhysicsGeometry& Geom, const TCollisionAnalyzerType& CollisionAnalyzerType, const FTransform& GeomPose, TArray<FOverlapResult>& OutOverlaps, ECollisionChannel TraceChannel, const struct FCollisionQueryParams& Params, const struct FCollisionResponseParams& ResponseParams, const struct FCollisionObjectQueryParams& ObjectParams, const TAccelContainer& AccelContainer)
 {
 	using namespace ChaosInterface;
@@ -1063,8 +1182,55 @@ bool GeomOverlapMultiImp(const UWorld* World, const FPhysicsGeometry& Geom, cons
 	return bHaveBlockingHit;
 }
 
-template <EQueryInfo::Type InfoType, typename TCollisionAnalyzerType, typename TAccelContainer = FDefaultAccelContainer>
-bool GeomOverlapMultiHelper(const UWorld* World, const FPhysicsGeometry& Geom, const TCollisionAnalyzerType& CollisionAnalyzerType, const FTransform& GeomPose, TArray<FOverlapResult>& OutOverlaps, ECollisionChannel TraceChannel, const struct FCollisionQueryParams& Params, const struct FCollisionResponseParams& ResponseParams, const struct FCollisionObjectQueryParams& ObjectParams, const TAccelContainer& AccelContainer = FDefaultAccelContainer{})
+template <EQueryInfo InfoType, typename GeomType, typename AccelContainerType = FDefaultAccelContainer>
+bool GeomOverlapMultiHelper(const EThreadQueryContext ThreadContext, const UWorld* World, const GeomType& InGeom, const FTransform& GeomPose, TArray<FOverlapResult>& OutOverlaps, ECollisionChannel TraceChannel, const struct FCollisionQueryParams& Params, const struct FCollisionResponseParams& ResponseParams, const struct FCollisionObjectQueryParams& ObjectParams, const AccelContainerType& AccelContainer = FDefaultAccelContainer{})
+{
+	using namespace ChaosInterface;
+
+	TGeomSQInputs<GeomType> SQInputs = GeomToSQInputs(InGeom, GeomPose.GetRotation());
+	const FPhysicsGeometry& Geom = *SQInputs.GetGeometry();
+	auto& CollisionAnalyzerType = *SQInputs.GetCollisionShape();
+	if (ThreadContext == EThreadQueryContext::GTData)
+	{
+		return GeomOverlapMultiImp<FOverlapHit, InfoType>(World, Geom, CollisionAnalyzerType, GeomPose, OutOverlaps, TraceChannel, Params, ResponseParams, ObjectParams, AccelContainer);
+	}
+	else
+	{
+		return GeomOverlapMultiImp<FPTOverlapHit, InfoType>(World, Geom, CollisionAnalyzerType, GeomPose, OutOverlaps, TraceChannel, Params, ResponseParams, ObjectParams, AccelContainer);
+	}
+}
+
+template <EQueryInfo InfoType>
+bool GeomOverlapMultiHelperWithQueryHandler(UPhysicsQueryHandler& QueryHandler, const EThreadQueryContext ThreadContext, const UWorld* World, const FCollisionShape& CollisionShape, const FTransform& GeomPose, TArray<FOverlapResult>& OutOverlaps, ECollisionChannel TraceChannel, const struct FCollisionQueryParams& Params, const struct FCollisionResponseParams& ResponseParams, const struct FCollisionObjectQueryParams& ObjectParams)
+{
+	Chaos::FCommonQueryData CommonData
+	{
+		.TraceChannel = TraceChannel,
+		.Params = Params,
+		.ResponseParams = ResponseParams,
+		.ObjectParams = ObjectParams,
+	};
+	Chaos::FOverlapQueryData OverlapData
+	{
+		.CollisionShape = CollisionShape,
+		.GeomPose = GeomPose,
+	};
+	return QueryHandler.Overlap(InfoType, ThreadContext, World, OverlapData, CommonData, OutOverlaps);
+}
+
+template <EQueryInfo InfoType>
+bool GeomOverlapMultiHelperWithQueryHandler(UPhysicsQueryHandler& QueryHandler, const EThreadQueryContext ThreadContext, const UWorld* World, const FPhysicsGeometryCollection& InCollisionShape, const FTransform& GeomPose, TArray<FOverlapResult>& OutOverlaps, ECollisionChannel TraceChannel, const struct FCollisionQueryParams& Params, const struct FCollisionResponseParams& ResponseParams, const struct FCollisionObjectQueryParams& ObjectParams)
+{
+	FCollisionShape CollisionShape;
+	if (ToCollisionShape(InCollisionShape, CollisionShape))
+	{
+		return GeomOverlapMultiHelperWithQueryHandler<InfoType>(QueryHandler, ThreadContext, World, CollisionShape, GeomPose, OutOverlaps, TraceChannel, Params, ResponseParams, ObjectParams);
+	}
+	return false;
+}
+
+template <EQueryInfo InfoType, typename GeomType, typename AccelContainerType = FDefaultAccelContainer>
+bool GeomOverlapMultiHelper(const UWorld* World, const GeomType& InGeom, const FTransform& GeomPose, TArray<FOverlapResult>& OutOverlaps, ECollisionChannel TraceChannel, const struct FCollisionQueryParams& Params, const struct FCollisionResponseParams& ResponseParams, const struct FCollisionObjectQueryParams& ObjectParams, const AccelContainerType& AccelContainer = FDefaultAccelContainer{})
 {
 	using namespace ChaosInterface;
 
@@ -1074,13 +1240,21 @@ bool GeomOverlapMultiHelper(const UWorld* World, const FPhysicsGeometry& Geom, c
 	}
 
 	const EThreadQueryContext ThreadContext = GetThreadQueryContext(*World->GetPhysicsScene()->GetSolver());
-	if (ThreadContext == EThreadQueryContext::GTData)
+	if constexpr (AccelContainerType::HasAccelerationStructureOverride())
 	{
-		return GeomOverlapMultiImp<FOverlapHit, InfoType, TCollisionAnalyzerType>(World, Geom, CollisionAnalyzerType, GeomPose, OutOverlaps, TraceChannel, Params, ResponseParams, ObjectParams, AccelContainer);
+		return GeomOverlapMultiHelper<InfoType>(ThreadContext, World, InGeom, GeomPose, OutOverlaps, TraceChannel, Params, ResponseParams, ObjectParams, AccelContainer);
 	}
 	else
 	{
-		return GeomOverlapMultiImp<FPTOverlapHit, InfoType, TCollisionAnalyzerType>(World, Geom, CollisionAnalyzerType, GeomPose, OutOverlaps, TraceChannel, Params, ResponseParams, ObjectParams, AccelContainer);
+		UPhysicsQueryHandler* QueryHandler = World->PhysicsQueryHandler;
+		if (QueryHandler == nullptr)
+		{
+			return GeomOverlapMultiHelper<InfoType>(ThreadContext, World, InGeom, GeomPose, OutOverlaps, TraceChannel, Params, ResponseParams, ObjectParams, AccelContainer);
+		}
+		else
+		{
+			return GeomOverlapMultiHelperWithQueryHandler<InfoType>(*QueryHandler, ThreadContext, World, InGeom, GeomPose, OutOverlaps, TraceChannel, Params, ResponseParams, ObjectParams);
+		}
 	}
 }
 
@@ -1092,8 +1266,7 @@ bool FGenericPhysicsInterface::GeomOverlapBlockingTest(const UWorld* World, cons
 
 	TArray<FOverlapResult> Overlaps;	//needed only for template shared code
 	FTransform GeomTransform(Rot, Pos);
-	FPhysicsShapeAdapter Adaptor(GeomTransform.GetRotation(), CollisionShape);
-	return GeomOverlapMultiHelper<EQueryInfo::IsBlocking>(World, Adaptor.GetGeometry(), CollisionShape, Adaptor.GetGeomPose(GeomTransform.GetTranslation()), Overlaps, TraceChannel, Params, ResponseParams, ObjectParams);
+	return GeomOverlapMultiHelper<EQueryInfo::IsBlocking>(World, CollisionShape, GeomTransform, Overlaps, TraceChannel, Params, ResponseParams, ObjectParams);
 }
 
 template<typename TAccel, typename TGeom>
@@ -1104,8 +1277,8 @@ bool FGenericGeomPhysicsInterfaceUsingSpatialAcceleration<TAccel, TGeom>::GeomOv
 	CSV_SCOPED_TIMING_STAT(SceneQuery, GeomOverlapBlocking);
 
 	TArray<FOverlapResult> Overlaps;	//needed only for template shared code
-	TGeomSQInputs<TGeom> SQInputs = GeomToSQInputs(InGeom, Rot);
-	return GeomOverlapMultiHelper<EQueryInfo::IsBlocking>(World, *SQInputs.GetGeometry(), *SQInputs.GetCollisionShape(), FTransform{Rot, Pos}, Overlaps, TraceChannel, Params, ResponseParams, ObjectParams, FOverrideAccelContainer{Accel});
+	FTransform GeomTransform{ Rot, Pos };
+	return GeomOverlapMultiHelper<EQueryInfo::IsBlocking>(World, InGeom, GeomTransform, Overlaps, TraceChannel, Params, ResponseParams, ObjectParams, FOverrideAccelContainer{ Accel });
 }
 
 bool FGenericPhysicsInterface::GeomOverlapAnyTest(const UWorld* World, const struct FCollisionShape& CollisionShape, const FVector& Pos, const FQuat& Rot, ECollisionChannel TraceChannel, const struct FCollisionQueryParams& Params, const struct FCollisionResponseParams& ResponseParams, const struct FCollisionObjectQueryParams& ObjectParams)
@@ -1116,8 +1289,7 @@ bool FGenericPhysicsInterface::GeomOverlapAnyTest(const UWorld* World, const str
 
 	TArray<FOverlapResult> Overlaps;	//needed only for template shared code
 	FTransform GeomTransform(Rot, Pos);
-	FPhysicsShapeAdapter Adaptor(GeomTransform.GetRotation(), CollisionShape);
-	return GeomOverlapMultiHelper<EQueryInfo::IsAnything>(World, Adaptor.GetGeometry(), CollisionShape, Adaptor.GetGeomPose(GeomTransform.GetTranslation()), Overlaps, TraceChannel, Params, ResponseParams, ObjectParams);
+	return GeomOverlapMultiHelper<EQueryInfo::IsAnything>(World, CollisionShape, GeomTransform, Overlaps, TraceChannel, Params, ResponseParams, ObjectParams);
 }
 
 template<typename TAccel, typename TGeom>
@@ -1128,8 +1300,8 @@ bool FGenericGeomPhysicsInterfaceUsingSpatialAcceleration<TAccel, TGeom>::GeomOv
 	CSV_SCOPED_TIMING_STAT(SceneQuery, GeomOverlapAny);
 
 	TArray<FOverlapResult> Overlaps;	//needed only for template shared code
-	TGeomSQInputs<TGeom> SQInputs = GeomToSQInputs(InGeom, Rot);
-	return GeomOverlapMultiHelper<EQueryInfo::IsAnything>(World, *SQInputs.GetGeometry(), *SQInputs.GetCollisionShape(), FTransform{ Rot, Pos }, Overlaps, TraceChannel, Params, ResponseParams, ObjectParams, FOverrideAccelContainer{ Accel });
+	FTransform GeomTransform{ Rot, Pos };
+	return GeomOverlapMultiHelper<EQueryInfo::IsAnything>(World, InGeom, GeomTransform, Overlaps, TraceChannel, Params, ResponseParams, ObjectParams, FOverrideAccelContainer{ Accel });
 }
 
 template<>
@@ -1140,7 +1312,7 @@ bool FGenericPhysicsInterface::GeomOverlapMulti(const UWorld* World, const FPhys
 	CSV_SCOPED_TIMING_STAT(SceneQuery, GeomOverlapMultiple);
 
 	FTransform GeomTransform(InRotation, InPosition);
-	return GeomOverlapMultiHelper<EQueryInfo::GatherAll>(World, InGeom.GetGeometry(), InGeom, GeomTransform, OutOverlaps, TraceChannel, Params, ResponseParams, ObjectParams);
+	return GeomOverlapMultiHelper<EQueryInfo::GatherAll>(World, InGeom, GeomTransform, OutOverlaps, TraceChannel, Params, ResponseParams, ObjectParams);
 }
 
 template<>
@@ -1151,8 +1323,7 @@ bool FGenericPhysicsInterface::GeomOverlapMulti(const UWorld* World, const FColl
 	CSV_SCOPED_TIMING_STAT(SceneQuery, GeomOverlapMultiple);
 
 	FTransform GeomTransform(InRotation, InPosition);
-	FPhysicsShapeAdapter Adaptor(GeomTransform.GetRotation(), InGeom);
-	return GeomOverlapMultiHelper<EQueryInfo::GatherAll>(World, Adaptor.GetGeometry(), InGeom, Adaptor.GetGeomPose(GeomTransform.GetTranslation()), OutOverlaps, TraceChannel, Params, ResponseParams, ObjectParams);
+	return GeomOverlapMultiHelper<EQueryInfo::GatherAll>(World, InGeom, GeomTransform, OutOverlaps, TraceChannel, Params, ResponseParams, ObjectParams);
 }
 
 template<>
@@ -1170,9 +1341,7 @@ bool FGenericGeomPhysicsInterfaceUsingSpatialAcceleration<TAccel, TGeom>::GeomSw
 	SCOPE_CYCLE_COUNTER(STAT_Collision_GeomSweepSingle);
 	CSV_SCOPED_TIMING_STAT(SceneQuery, GeomSweepSingle);
 
-	using TCastTraits = TSQTraits<FSweepHit, ESweepOrRay::Sweep, ESingleMultiOrTest::Single>;
-	using TPTCastTraits = TSQTraits<FPTSweepHit, ESweepOrRay::Sweep, ESingleMultiOrTest::Single>;
-	return TSceneCastCommon<TCastTraits, TPTCastTraits>(World, OutHit, GeomToSQInputs(InGeom, Rot), Start, End, TraceChannel, Params, ResponseParams, ObjectParams, FOverrideAccelContainer{ Accel });
+	return TSceneCastCommon<ESweepOrRay::Sweep, ESingleMultiOrTest::Single>(World, OutHit, GeomToSQInputs(InGeom, Rot), Start, End, TraceChannel, Params, ResponseParams, ObjectParams, FOverrideAccelContainer{ Accel });
 }
 
 template<typename TAccel, typename TGeom>
@@ -1184,9 +1353,7 @@ bool FGenericGeomPhysicsInterfaceUsingSpatialAcceleration<TAccel, TGeom>::GeomSw
 	SCOPE_CYCLE_COUNTER(STAT_Collision_GeomSweepMultiple);
 	CSV_SCOPED_TIMING_STAT(SceneQuery, GeomSweepMultiple);
 
-	using TCastTraits = TSQTraits<FHitSweep, ESweepOrRay::Sweep, ESingleMultiOrTest::Multi>;
-	using TPTCastTraits = TSQTraits<FPTSweepHit, ESweepOrRay::Sweep, ESingleMultiOrTest::Multi>;
-	return TSceneCastCommon<TCastTraits, TPTCastTraits>(World, OutHits, GeomToSQInputs(InGeom, InGeomRot), Start, End, TraceChannel, Params, ResponseParams, ObjectParams, FOverrideAccelContainer{ Accel });
+	return TSceneCastCommon<ESweepOrRay::Sweep, ESingleMultiOrTest::Multi>(World, OutHits, GeomToSQInputs(InGeom, InGeomRot), Start, End, TraceChannel, Params, ResponseParams, ObjectParams, FOverrideAccelContainer{ Accel });
 }
 
 template<typename TAccel, typename TGeom>
@@ -1196,8 +1363,8 @@ bool FGenericGeomPhysicsInterfaceUsingSpatialAcceleration<TAccel, TGeom>::GeomOv
 	SCOPE_CYCLE_COUNTER(STAT_Collision_GeomOverlapMultiple);
 	CSV_SCOPED_TIMING_STAT(SceneQuery, GeomOverlapMultiple);
 
-	TGeomSQInputs<TGeom> SQInputs = GeomToSQInputs(InGeom, InRotation);
-	return GeomOverlapMultiHelper<EQueryInfo::GatherAll>(World, *SQInputs.GetGeometry(), *SQInputs.GetCollisionShape(), FTransform{ InRotation, InPosition }, OutOverlaps, TraceChannel, Params, ResponseParams, ObjectParams, FOverrideAccelContainer{ Accel });
+	FTransform GeomTransform{ InRotation, InPosition };
+	return GeomOverlapMultiHelper<EQueryInfo::GatherAll>(World, InGeom, GeomTransform, OutOverlaps, TraceChannel, Params, ResponseParams, ObjectParams, FOverrideAccelContainer{ Accel });
 }
 
 template struct FGenericRaycastPhysicsInterfaceUsingSpatialAcceleration<Chaos::IDefaultChaosSpatialAcceleration>;
@@ -1228,13 +1395,94 @@ namespace Chaos::Private
 
 		if (QueryRadius > UE_SMALL_NUMBER)
 		{
-			using PTTraits = TSQTraits<ChaosInterface::FPTSweepHit, ESweepOrRay::Sweep, ESingleMultiOrTest::Multi>;
-			return TSceneCastCommonImp<PTTraits, FGeomSQAdditionalInputs>(World, OutHits, FGeomSQAdditionalInputs(FCollisionShape::MakeSphere(QueryRadius), FQuat::Identity), Start, End, TraceChannel, Params, ResponseParams, ObjectParams, FDefaultAccelContainer{});
+			return TSceneCastCommon<ESweepOrRay::Sweep, ESingleMultiOrTest::Multi>(World, OutHits, FGeomSQAdditionalInputs(FCollisionShape::MakeSphere(QueryRadius), FQuat::Identity), Start, End, TraceChannel, Params, ResponseParams, ObjectParams, FDefaultAccelContainer{});
 		}
 		else
 		{
-			using PTTraits = TSQTraits<ChaosInterface::FPTRaycastHit, ESweepOrRay::Raycast, ESingleMultiOrTest::Multi>;
-			return TSceneCastCommonImp<PTTraits, FRaycastSQAdditionalInputs>(World, OutHits, FRaycastSQAdditionalInputs(), Start, End, TraceChannel, Params, ResponseParams, ObjectParams, FDefaultAccelContainer{});
+			return TSceneCastCommon<ESweepOrRay::Raycast, ESingleMultiOrTest::Multi>(World, OutHits, FRaycastSQAdditionalInputs(), Start, End, TraceChannel, Params, ResponseParams, ObjectParams, FDefaultAccelContainer{});
 		}
 	};
+
+	template <typename OverlapHitType, EQueryInfo InfoType>
+	bool OverlapImpl(const UWorld* World, const Chaos::FOverlapQueryData& OverlapData, const FCommonQueryData& CommonData, TArray<FOverlapResult>& OutOverlaps)
+	{
+		FPhysicsShapeAdapter Adaptor(OverlapData.GeomPose.GetRotation(), OverlapData.CollisionShape);
+		const FPhysicsGeometry& Geom = Adaptor.GetGeometry();
+		return GeomOverlapMultiImp<OverlapHitType, InfoType>(World, Geom, OverlapData.CollisionShape, OverlapData.GeomPose, OutOverlaps, CommonData.TraceChannel, CommonData.Params, CommonData.ResponseParams, CommonData.ObjectParams, FDefaultAccelContainer{});
+	}
+
+	template <EQueryInfo InfoType, typename ... ArgTypes>
+	bool ForwardThreadContextType_Overlap(const EThreadQueryContext ThreadContext, ArgTypes&& ... Args)
+	{
+		using namespace ChaosInterface;
+		if (ThreadContext == EThreadQueryContext::GTData)
+		{
+			return OverlapImpl<FOverlapHit, InfoType>(std::forward<ArgTypes>(Args)...);
+		}
+		else
+		{
+			return OverlapImpl<FPTOverlapHit, InfoType>(std::forward<ArgTypes>(Args)...);
+		}
+	}
+
+	bool FQueryInterface_Internal::Overlap(const EQueryInfo InfoType, const EThreadQueryContext ThreadContext, const UWorld* World, const FOverlapQueryData& OverlapData, const FCommonQueryData& CommonData, TArray<FOverlapResult>& OutOverlaps)
+	{
+		if (InfoType == EQueryInfo::GatherAll)
+		{
+			return ForwardThreadContextType_Overlap<EQueryInfo::GatherAll>(ThreadContext, World, OverlapData, CommonData, OutOverlaps);
+		}
+		else if (InfoType == EQueryInfo::IsBlocking)
+		{
+			return ForwardThreadContextType_Overlap<EQueryInfo::IsBlocking>(ThreadContext, World, OverlapData, CommonData, OutOverlaps);
+		}
+		else if (InfoType == EQueryInfo::IsAnything)
+		{
+			return ForwardThreadContextType_Overlap<EQueryInfo::IsAnything>(ThreadContext, World, OverlapData, CommonData, OutOverlaps);
+		}
+		return false;
+	}
+
+	template <ESingleMultiOrTest SingleMultiOrTest, typename HitResultsType>
+	bool RaycastCommon(const EThreadQueryContext ThreadContext, const UWorld* World, const Chaos::FRayQueryData& RayData, const Chaos::FCommonQueryData& CommonData, HitResultsType& OutHits)
+	{
+		FRaycastSQAdditionalInputs GeomInputs;
+		return TraceCommonImp<ESweepOrRay::Raycast, SingleMultiOrTest>(ThreadContext, World, OutHits, GeomInputs, RayData.Start, RayData.End, CommonData.TraceChannel, CommonData.Params, CommonData.ResponseParams, CommonData.ObjectParams, FDefaultAccelContainer{});
+	}
+
+	bool FQueryInterface_Internal::RaycastTest(const EThreadQueryContext ThreadContext, const UWorld* World, const FRayQueryData& RayData, const FCommonQueryData& CommonData, FHitResult& OutHits)
+	{
+		return RaycastCommon<ESingleMultiOrTest::Test>(ThreadContext, World, RayData, CommonData, OutHits);
+	}
+
+	bool FQueryInterface_Internal::RaycastSingle(const EThreadQueryContext ThreadContext, const UWorld* World, const FRayQueryData& RayData, const FCommonQueryData& CommonData, FHitResult& OutHits)
+	{
+		return RaycastCommon<ESingleMultiOrTest::Single>(ThreadContext, World, RayData, CommonData, OutHits);
+	}
+
+	bool FQueryInterface_Internal::RaycastMulti(const EThreadQueryContext ThreadContext, const UWorld* World, const FRayQueryData& RayData, const FCommonQueryData& CommonData, TArray<FHitResult>& OutHits)
+	{
+		return RaycastCommon<ESingleMultiOrTest::Multi>(ThreadContext, World, RayData, CommonData, OutHits);
+	}
+
+	template <ESingleMultiOrTest SingleMultiOrTest, typename HitResultsType>
+	bool SweepCommon(const EThreadQueryContext ThreadContext, const UWorld* World, const Chaos::FSweepQueryData& SweepData, const Chaos::FCommonQueryData& CommonData, HitResultsType& OutHits)
+	{
+		FGeomSQAdditionalInputs GeomInputs(SweepData.CollisionShape, SweepData.GeomRot);
+		return TraceCommonImp<ESweepOrRay::Sweep, SingleMultiOrTest>(ThreadContext, World, OutHits, GeomInputs, SweepData.Start, SweepData.End, CommonData.TraceChannel, CommonData.Params, CommonData.ResponseParams, CommonData.ObjectParams, FDefaultAccelContainer{});
+	}
+
+	bool FQueryInterface_Internal::SweepTest(const EThreadQueryContext ThreadContext, const UWorld* World, const FSweepQueryData& SweepData, const FCommonQueryData& CommonData, FHitResult& OutHits)
+	{
+		return SweepCommon<ESingleMultiOrTest::Test>(ThreadContext, World, SweepData, CommonData, OutHits);
+	}
+
+	bool FQueryInterface_Internal::SweepSingle(const EThreadQueryContext ThreadContext, const UWorld* World, const FSweepQueryData& SweepData, const FCommonQueryData& CommonData, FHitResult& OutHits)
+	{
+		return SweepCommon<ESingleMultiOrTest::Single>(ThreadContext, World, SweepData, CommonData, OutHits);
+	}
+
+	bool FQueryInterface_Internal::SweepMulti(const EThreadQueryContext ThreadContext, const UWorld* World, const FSweepQueryData& SweepData, const FCommonQueryData& CommonData, TArray<FHitResult>& OutHits)
+	{
+		return SweepCommon<ESingleMultiOrTest::Multi>(ThreadContext, World, SweepData, CommonData, OutHits);
+	}
 }
diff --git a/Engine/Source/Runtime/Engine/Private/PhysicsEngine/DefaultPhysicsQueryHandler.cpp b/Engine/Source/Runtime/Engine/Private/PhysicsEngine/DefaultPhysicsQueryHandler.cpp
new file mode 100644
--- /dev/null
+++ b/Engine/Source/Runtime/Engine/Private/PhysicsEngine/DefaultPhysicsQueryHandler.cpp
@@ -0,0 +1,40 @@
+// Copyright Epic Games, Inc. All Rights Reserved.
+
+#include "Physics/DefaultPhysicsQueryHandler.h"
+
+#include "Physics/GenericPhysicsInterface.h"
+
+bool UDefaultPhysicsQueryHandler::Overlap(const EQueryInfo InfoType, const EThreadQueryContext ThreadContext, const UWorld* World, const FOverlapQueryData& OverlapData, const FCommonQueryData& CommonData, TArray<FOverlapResult>& OutOverlaps)
+{
+	return Chaos::Private::FQueryInterface_Internal::Overlap(InfoType, ThreadContext, World, OverlapData, CommonData, OutOverlaps);
+}
+
+bool UDefaultPhysicsQueryHandler::RaycastTest(const EThreadQueryContext ThreadContext, const UWorld* World, const FRayQueryData& RayData, const FCommonQueryData& CommonData, FHitResult& OutHits)
+{
+	return Chaos::Private::FQueryInterface_Internal::RaycastTest(ThreadContext, World, RayData, CommonData, OutHits);
+}
+
+bool UDefaultPhysicsQueryHandler::RaycastSingle(const EThreadQueryContext ThreadContext, const UWorld* World, const FRayQueryData& RayData, const FCommonQueryData& CommonData, FHitResult& OutHits)
+{
+	return Chaos::Private::FQueryInterface_Internal::RaycastSingle(ThreadContext, World, RayData, CommonData, OutHits);
+}
+
+bool UDefaultPhysicsQueryHandler::RaycastMulti(const EThreadQueryContext ThreadContext, const UWorld* World, const FRayQueryData& RayData, const FCommonQueryData& CommonData, TArray<FHitResult>& OutHits)
+{
+	return Chaos::Private::FQueryInterface_Internal::RaycastMulti(ThreadContext, World, RayData, CommonData, OutHits);
+}
+
+bool UDefaultPhysicsQueryHandler::SweepTest(const EThreadQueryContext ThreadContext, const UWorld* World, const FSweepQueryData& SweepData, const FCommonQueryData& CommonData, FHitResult& OutHits)
+{
+	return Chaos::Private::FQueryInterface_Internal::SweepTest(ThreadContext, World, SweepData, CommonData, OutHits);
+}
+
+bool UDefaultPhysicsQueryHandler::SweepSingle(const EThreadQueryContext ThreadContext, const UWorld* World, const FSweepQueryData& SweepData, const FCommonQueryData& CommonData, FHitResult& OutHits)
+{
+	return Chaos::Private::FQueryInterface_Internal::SweepSingle(ThreadContext, World, SweepData, CommonData, OutHits);
+}
+
+bool UDefaultPhysicsQueryHandler::SweepMulti(const EThreadQueryContext ThreadContext, const UWorld* World, const FSweepQueryData& SweepData, const FCommonQueryData& CommonData, TArray<FHitResult>& OutHits)
+{
+	return Chaos::Private::FQueryInterface_Internal::SweepMulti(ThreadContext, World, SweepData, CommonData, OutHits);
+}
diff --git a/Engine/Source/Runtime/Engine/Private/PhysicsEngine/PhysicsQueryHandler.cpp b/Engine/Source/Runtime/Engine/Private/PhysicsEngine/PhysicsQueryHandler.cpp
new file mode 100644
--- /dev/null
+++ b/Engine/Source/Runtime/Engine/Private/PhysicsEngine/PhysicsQueryHandler.cpp
@@ -0,0 +1,38 @@
+// Copyright Epic Games, Inc. All Rights Reserved.
+
+#include "Physics/PhysicsQueryHandler.h"
+
+bool UPhysicsQueryHandler::Overlap(const EQueryInfo InfoType, const EThreadQueryContext ThreadContext, const UWorld* World, const FOverlapQueryData& OverlapData, const FCommonQueryData& CommonData, TArray<FOverlapResult>& OutOverlaps)
+{
+	return false;
+}
+
+bool UPhysicsQueryHandler::RaycastTest(const EThreadQueryContext ThreadContext, const UWorld* World, const FRayQueryData& RayData, const FCommonQueryData& CommonData, FHitResult& OutHits)
+{
+	return false;
+}
+
+bool UPhysicsQueryHandler::RaycastSingle(const EThreadQueryContext ThreadContext, const UWorld* World, const FRayQueryData& RayData, const FCommonQueryData& CommonData, FHitResult& OutHits)
+{
+	return false;
+}
+
+bool UPhysicsQueryHandler::RaycastMulti(const EThreadQueryContext ThreadContext, const UWorld* World, const FRayQueryData& RayData, const FCommonQueryData& CommonData, TArray<FHitResult>& OutHits)
+{
+	return false;
+}
+
+bool UPhysicsQueryHandler::SweepTest(const EThreadQueryContext ThreadContext, const UWorld* World, const FSweepQueryData& SweepData, const FCommonQueryData& CommonData, FHitResult& OutHits)
+{
+	return false;
+}
+
+bool UPhysicsQueryHandler::SweepSingle(const EThreadQueryContext ThreadContext, const UWorld* World, const FSweepQueryData& SweepData, const FCommonQueryData& CommonData, FHitResult& OutHits)
+{
+	return false;
+}
+
+bool UPhysicsQueryHandler::SweepMulti(const EThreadQueryContext ThreadContext, const UWorld* World, const FSweepQueryData& RaySweepData, const FCommonQueryData& CommonData, TArray<FHitResult>& OutHits)
+{
+	return false;
+}
diff --git a/Engine/Source/Runtime/Engine/Public/Physics/DefaultPhysicsQueryHandler.h b/Engine/Source/Runtime/Engine/Public/Physics/DefaultPhysicsQueryHandler.h
new file mode 100644
--- /dev/null
+++ b/Engine/Source/Runtime/Engine/Public/Physics/DefaultPhysicsQueryHandler.h
@@ -0,0 +1,25 @@
+// Copyright Epic Games, Inc. All Rights Reserved.
+
+#pragma once
+
+#include "PhysicsQueryHandler.h"
+
+#include "DefaultPhysicsQueryHandler.generated.h"
+
+UCLASS(hidecategories=(Object), MinimalAPI)
+// This does default handling for physics queries (as if there was no handler).
+class UDefaultPhysicsQueryHandler : public UPhysicsQueryHandler
+{
+public:
+	GENERATED_BODY()
+
+	virtual bool Overlap(const EQueryInfo InfoType, const EThreadQueryContext ThreadContext, const UWorld* World, const FOverlapQueryData& OverlapData, const FCommonQueryData& CommonData, TArray<FOverlapResult>& OutOverlaps) override;
+
+	virtual bool RaycastTest(const EThreadQueryContext ThreadContext, const UWorld* World, const FRayQueryData& RayData, const FCommonQueryData& CommonData, FHitResult& OutHits) override;
+	virtual bool RaycastSingle(const EThreadQueryContext ThreadContext, const UWorld* World, const FRayQueryData& RayData, const FCommonQueryData& CommonData, FHitResult& OutHits) override;
+	virtual bool RaycastMulti(const EThreadQueryContext ThreadContext, const UWorld* World, const FRayQueryData& RayData, const FCommonQueryData& CommonData, TArray<FHitResult>& OutHits) override;
+
+	virtual bool SweepTest(const EThreadQueryContext ThreadContext, const UWorld* World, const FSweepQueryData& SweepData, const FCommonQueryData& CommonData, FHitResult& OutHits) override;
+	virtual bool SweepSingle(const EThreadQueryContext ThreadContext, const UWorld* World, const FSweepQueryData& SweepData, const FCommonQueryData& CommonData, FHitResult& OutHits) override;
+	virtual bool SweepMulti(const EThreadQueryContext ThreadContext, const UWorld* World, const FSweepQueryData& SweepData, const FCommonQueryData& CommonData, TArray<FHitResult>& OutHits) override;
+};
diff --git a/Engine/Source/Runtime/Engine/Public/Physics/GenericPhysicsInterface.h b/Engine/Source/Runtime/Engine/Public/Physics/GenericPhysicsInterface.h
--- a/Engine/Source/Runtime/Engine/Public/Physics/GenericPhysicsInterface.h
+++ b/Engine/Source/Runtime/Engine/Public/Physics/GenericPhysicsInterface.h
@@ -11,6 +11,17 @@
 
 class UWorld;
 
+namespace Chaos
+{
+	enum class EQueryInfo : uint8;
+	enum class EThreadQueryContext;
+
+	struct FCommonQueryData;
+	struct FOverlapQueryData;
+	struct FRayQueryData;
+	struct FSweepQueryData;
+}
+
 struct FGenericPhysicsInterface
 {
 	/** Trace a ray against the world and return if a blocking hit is found */
@@ -108,4 +119,18 @@ namespace Chaos::Private
 		*/
 		static ENGINE_API bool SpherecastMulti(const UWorld* World, float QueryRadius, TArray<FHitResult>& OutHits, FVector Start, FVector End, ECollisionChannel TraceChannel, const FCollisionQueryParams& Params, const FCollisionResponseParams& ResponseParams, const FCollisionObjectQueryParams& ObjectParams = FCollisionObjectQueryParams::DefaultObjectQueryParam);
 	};
+
+	// This is for internal use so that the QueryInterface API can resume queries.
+	struct FQueryInterface_Internal
+	{
+		static bool Overlap(const EQueryInfo InfoType, const EThreadQueryContext ThreadContext, const UWorld* World, const FOverlapQueryData& OverlapData, const FCommonQueryData& CommonData, TArray<FOverlapResult>& OutOverlaps);
+
+		static bool RaycastTest(const EThreadQueryContext ThreadContext, const UWorld* World, const FRayQueryData& RayData, const FCommonQueryData& CommonData, FHitResult& OutHits);
+		static bool RaycastSingle(const EThreadQueryContext ThreadContext, const UWorld* World, const FRayQueryData& RayData, const FCommonQueryData& CommonData, FHitResult& OutHits);
+		static bool RaycastMulti(const EThreadQueryContext ThreadContext, const UWorld* World, const FRayQueryData& RayData, const FCommonQueryData& CommonData, TArray<FHitResult>& OutHits);
+
+		static bool SweepTest(const EThreadQueryContext ThreadContext, const UWorld* World, const FSweepQueryData& SweepData, const FCommonQueryData& CommonData, FHitResult& OutHits);
+		static bool SweepSingle(const EThreadQueryContext ThreadContext, const UWorld* World, const FSweepQueryData& SweepData, const FCommonQueryData& CommonData, FHitResult& OutHits);
+		static bool SweepMulti(const EThreadQueryContext ThreadContext, const UWorld* World, const FSweepQueryData& SweepData, const FCommonQueryData& CommonData, TArray<FHitResult>& OutHits);
+	};
 }
diff --git a/Engine/Source/Runtime/Engine/Public/Physics/PhysicsQueryHandler.h b/Engine/Source/Runtime/Engine/Public/Physics/PhysicsQueryHandler.h
new file mode 100644
--- /dev/null
+++ b/Engine/Source/Runtime/Engine/Public/Physics/PhysicsQueryHandler.h
@@ -0,0 +1,39 @@
+// Copyright Epic Games, Inc. All Rights Reserved.
+
+#pragma once
+
+#include "CoreMinimal.h"
+#include "UObject/ObjectMacros.h"
+#include "UObject/Object.h"
+#include "Physics/SceneQueryData.h"
+
+#include "PhysicsQueryHandler.generated.h"
+
+struct FOverlapResult;
+struct FHitResult;
+
+UCLASS(hidecategories=(Object), MinimalAPI)
+// Allows custom handling of physics queries. Primarily for handling network/server redirection.
+class UPhysicsQueryHandler : public UObject
+{
+	GENERATED_BODY()
+
+public:
+	using EQueryInfo = Chaos::EQueryInfo;
+	using EThreadQueryContext = Chaos::EThreadQueryContext;
+
+	using FCommonQueryData = Chaos::FCommonQueryData;
+	using FOverlapQueryData = Chaos::FOverlapQueryData;
+	using FRayQueryData = Chaos::FRayQueryData;
+	using FSweepQueryData = Chaos::FSweepQueryData;
+
+	virtual bool Overlap(const EQueryInfo InfoType, const EThreadQueryContext ThreadContext, const UWorld* World, const FOverlapQueryData& OverlapData, const FCommonQueryData& CommonData, TArray<FOverlapResult>& OutOverlaps);
+
+	virtual bool RaycastTest(const EThreadQueryContext ThreadContext, const UWorld* World, const FRayQueryData& RayData, const FCommonQueryData& CommonData, FHitResult& OutHits);
+	virtual bool RaycastSingle(const EThreadQueryContext ThreadContext, const UWorld* World, const FRayQueryData& RayData, const FCommonQueryData& CommonData, FHitResult& OutHits);
+	virtual bool RaycastMulti(const EThreadQueryContext ThreadContext, const UWorld* World, const FRayQueryData& RayData, const FCommonQueryData& CommonData, TArray<FHitResult>& OutHits);
+
+	virtual bool SweepTest(const EThreadQueryContext ThreadContext, const UWorld* World, const FSweepQueryData& SweepData, const FCommonQueryData& CommonData, FHitResult& OutHits);
+	virtual bool SweepSingle(const EThreadQueryContext ThreadContext, const UWorld* World, const FSweepQueryData& SweepData, const FCommonQueryData& CommonData, FHitResult& OutHits);
+	virtual bool SweepMulti(const EThreadQueryContext ThreadContext, const UWorld* World, const FSweepQueryData& RaySweepData, const FCommonQueryData& CommonData, TArray<FHitResult>& OutHits);
+};
diff --git a/Engine/Source/Runtime/Engine/Public/Physics/SceneQueryData.h b/Engine/Source/Runtime/Engine/Public/Physics/SceneQueryData.h
new file mode 100644
--- /dev/null
+++ b/Engine/Source/Runtime/Engine/Public/Physics/SceneQueryData.h
@@ -0,0 +1,53 @@
+// Copyright Epic Games, Inc. All Rights Reserved.
+
+#pragma once
+
+#include "Engine/EngineTypes.h"
+#include "Math/MathFwd.h"
+#include "CollisionShape.h"
+#include "CollisionQueryParams.h"
+
+namespace Chaos
+{
+	enum class EQueryInfo : uint8
+	{
+		GatherAll,		//get all data and actually return it
+		IsBlocking,		//is any of the data blocking? only return a bool so don't bother collecting
+		IsAnything		//is any of the data blocking or touching? only return a bool so don't bother collecting
+	};
+
+	enum class EThreadQueryContext
+	{
+		GTData,		//use interpolated GT data
+		PTDataWithGTObjects,	//use pt data, but convert back to GT when possible
+		PTOnlyData,	//use only the PT data and don't try to convert anything back to GT
+	};
+
+	struct FCommonQueryData
+	{
+		ECollisionChannel TraceChannel;
+		FCollisionQueryParams Params;
+		FCollisionResponseParams ResponseParams;
+		FCollisionObjectQueryParams ObjectParams;
+	};
+
+	struct FOverlapQueryData
+	{
+		FCollisionShape CollisionShape;
+		FTransform GeomPose;
+	};
+
+	struct FRayQueryData
+	{
+		FVector Start;
+		FVector End;
+	};
+
+	struct FSweepQueryData
+	{
+		FVector Start;
+		FVector End;
+		FCollisionShape CollisionShape;
+		FQuat GeomRot;
+	};
+}
