Material Parameter HLOD Modifier:
Sets custom material parameters during the HLOD Bake process
Reverts the MPC values back to defaults.
Version 1 only supports scalar parameters

#rnx

diff --git a/Engine/Source/Runtime/Engine/Private/WorldPartition/HLOD/MaterialParameterCollectionHLODModifier.cpp b/Engine/Source/Runtime/Engine/Private/WorldPartition/HLOD/MaterialParameterCollectionHLODModifier.cpp
new file mode 100644
--- /dev/null
+++ b/Engine/Source/Runtime/Engine/Private/WorldPartition/HLOD/MaterialParameterCollectionHLODModifier.cpp
@@ -0,0 +1,64 @@
+// Copyright Epic Games, Inc. All Rights Reserved.
+
+#include "WorldPartition/HLOD/MaterialParameterCollectionHLODModifier.h"
+#include "Engine/World.h"
+#include "Materials/MaterialParameterCollection.h"
+
+
+UMaterialParameterCollectionHLODModifier::UMaterialParameterCollectionHLODModifier(const FObjectInitializer& ObjectInitializer)
+	: Super(ObjectInitializer)
+{
+}
+
+
+
+bool UMaterialParameterCollectionHLODModifier::CanModifyHLOD(TSubclassOf<UHLODBuilder> InHLODBuilderClass) const
+{
+	return true;
+}
+
+void UMaterialParameterCollectionHLODModifier::BeginHLODBuild(const FHLODBuildContext& InHLODBuildContext)
+{
+	//Cache off the build context, all of these seems to be executed in one frame so it should be fine to dereference in EndHLODBuild
+	HLODBuildContext = &InHLODBuildContext;
+	if (MPC != nullptr)
+	{
+		if (UMaterialParameterCollectionInstance* Instance = HLODBuildContext->World->GetParameterCollectionInstance(MPC))
+		{
+			ApplyScalarParameterOverrides(Instance);
+		}
+	}
+}
+
+void UMaterialParameterCollectionHLODModifier::EndHLODBuild(TArray<UActorComponent*>& InOutComponents)
+{
+	if (MPC != nullptr)
+	{
+		if (UMaterialParameterCollectionInstance* Instance = HLODBuildContext->World->GetParameterCollectionInstance(MPC))
+		{
+			RestoreScalarParameterValues(Instance);
+		}
+	}
+}
+
+void UMaterialParameterCollectionHLODModifier::ApplyScalarParameterOverrides(UMaterialParameterCollectionInstance* Instance)
+{
+	for (int32 ParameterIndex = 0; ParameterIndex < ScalarParameters.Num(); ParameterIndex++)
+	{
+		const FHLODModifierScalarParameter& Parameter = ScalarParameters[ParameterIndex];
+		float CurScalarVal;
+		Instance->GetScalarParameterValue(Parameter.ParameterName, CurScalarVal);
+		Instance->SetScalarParameterValue(Parameter.ParameterName, Parameter.OverrideValue);
+		CachedScalarParameterValues.Add(CurScalarVal);
+	}
+}
+
+void UMaterialParameterCollectionHLODModifier::RestoreScalarParameterValues(UMaterialParameterCollectionInstance* Instance)
+{
+	// Revert the overrides in the reverse order, this supports a stack style parameter overrrides, so we can have the same param multiple times in the stack, so only the last value is applied but the original value is reverted to
+	for (int32 ParameterIndex = ScalarParameters.Num() - 1; ParameterIndex >= 0; ParameterIndex--)
+	{
+		const FHLODModifierScalarParameter& Parameter = ScalarParameters[ParameterIndex];
+		Instance->SetScalarParameterValue(Parameter.ParameterName, CachedScalarParameterValues[ParameterIndex]);
+	}
+}
\ No newline at end of file
diff --git a/Engine/Source/Runtime/Engine/Public/WorldPartition/HLOD/MaterialParameterCollectionHLODModifier.h b/Engine/Source/Runtime/Engine/Public/WorldPartition/HLOD/MaterialParameterCollectionHLODModifier.h
new file mode 100644
--- /dev/null
+++ b/Engine/Source/Runtime/Engine/Public/WorldPartition/HLOD/MaterialParameterCollectionHLODModifier.h
@@ -0,0 +1,54 @@
+// Copyright Epic Games, Inc. All Rights Reserved.
+
+#pragma once
+
+#include "WorldPartition/HLOD/HLODModifier.h"
+#include "WorldPartition/HLOD/HLODBuilder.h"
+#include "Materials/MaterialParameterCollectionInstance.h"
+#include "MaterialParameterCollectionHLODModifier.generated.h"
+
+class UMaterialParameterCollection;
+
+/** A scalar parameter */
+USTRUCT()
+struct FHLODModifierScalarParameter
+{
+	GENERATED_USTRUCT_BODY()
+
+	FHLODModifierScalarParameter()
+	: OverrideValue(0.0f)
+	{
+		ParameterName = FName(TEXT("Scalar"));
+	}
+
+	UPROPERTY(EditAnywhere, Category=Parameter)
+	FName ParameterName;
+
+	UPROPERTY(EditAnywhere, Category=Parameter)
+	float OverrideValue;
+};
+
+
+UCLASS(BlueprintType, Blueprintable)
+class UMaterialParameterCollectionHLODModifier : public UWorldPartitionHLODModifier
+{
+	GENERATED_UCLASS_BODY()
+
+	virtual bool CanModifyHLOD(TSubclassOf<UHLODBuilder> InHLODBuilderClass) const;
+	virtual void BeginHLODBuild(const FHLODBuildContext& InHLODBuildContext);
+	virtual void EndHLODBuild(TArray<UActorComponent*>& InOutComponents);
+
+
+	UPROPERTY(EditAnywhere, Category=Material, DisplayName="Material Parameter Collection")
+	TObjectPtr<UMaterialParameterCollection> MPC;
+
+	UPROPERTY(EditAnywhere, Category=Material, Meta = (TitleProperty = "ParameterName"))
+	TArray<FHLODModifierScalarParameter> ScalarParameters;
+
+private:
+	void ApplyScalarParameterOverrides(UMaterialParameterCollectionInstance* Instance);
+	void RestoreScalarParameterValues(UMaterialParameterCollectionInstance* Instance);
+
+	const FHLODBuildContext* HLODBuildContext;
+	TArray<float> CachedScalarParameterValues;
+};
