Fix crash when receiving NMT_JoinSplit with unique net ID V2

diff --git a/Engine/Source/Runtime/Engine/Private/World.cpp b/Engine/Source/Runtime/Engine/Private/World.cpp
--- a/Engine/Source/Runtime/Engine/Private/World.cpp
+++ b/Engine/Source/Runtime/Engine/Private/World.cpp
@@ -6854,8 +6854,8 @@ void UWorld::NotifyControlMessage(UNetConnection* Connection, uint8 MessageType,
 
 					UE_LOG(LogNet, Log, TEXT("Join splitscreen request: %s userId: %s parentUserId: %s"),
 						NewRequestURL,
-						SplitRequestUniqueIdRepl.IsValid() ? *SplitRequestUniqueIdRepl->ToString() : TEXT("Invalid"),
-						Connection->PlayerId.IsValid() ? *Connection->PlayerId->ToString() : TEXT("Invalid"));
+						SplitRequestUniqueIdRepl.IsValid() ? *SplitRequestUniqueIdRepl.ToString() : TEXT("Invalid"),
+						Connection->PlayerId.IsValid() ? *Connection->PlayerId.ToString() : TEXT("Invalid"));
 
 					// Compromise for passing splitscreen playercount through to gameplay login code,
 					// without adding a lot of extra unnecessary complexity throughout the login code.
