First Person: Changed custom output node log category from Warning to Display.

diff --git a/Engine/Source/Runtime/Engine/Private/Materials/HLSLMaterialTranslator.cpp b/Engine/Source/Runtime/Engine/Private/Materials/HLSLMaterialTranslator.cpp
--- a/Engine/Source/Runtime/Engine/Private/Materials/HLSLMaterialTranslator.cpp
+++ b/Engine/Source/Runtime/Engine/Private/Materials/HLSLMaterialTranslator.cpp
@@ -1078,7 +1078,7 @@ uint64 FHLSLMaterialTranslator::CompileMaterialAttributesCustomOutputs(uint64 Cu
 
 		if ((CustomOutputNodeBitmask & FMaterialAttributeDefinitionMap::GetBitmask(CustomAttribute.AttributeID)) != 0)
 		{
-			UE_LOG(LogMaterial, Warning, TEXT("The material sets custom output '%s' using both MaterialAttributes and a custom output node! The value set via MaterialAttributes will be ignored."), *CustomAttribute.AttributeName);
+			UE_LOG(LogMaterial, Display, TEXT("The material sets custom output '%s' using both MaterialAttributes and a custom output node! The value set via MaterialAttributes will be ignored."), *CustomAttribute.AttributeName);
 			continue;
 		}
 
