Added in ULevelStreaming::FindStreamingLevel an ensure that the provided level pointer is valid.

diff --git a/Engine/Source/Runtime/Engine/Private/LevelStreaming.cpp b/Engine/Source/Runtime/Engine/Private/LevelStreaming.cpp
--- a/Engine/Source/Runtime/Engine/Private/LevelStreaming.cpp
+++ b/Engine/Source/Runtime/Engine/Private/LevelStreaming.cpp
@@ -1337,13 +1337,16 @@ void ULevelStreaming::RemoveLevelAnnotation(const ULevel* Level)
 ULevelStreaming* ULevelStreaming::FindStreamingLevel(const ULevel* Level)
 {
 	ULevelStreaming* FoundLevelStreaming = nullptr;
-	if (Level && Level->OwningWorld && !Level->IsPersistentLevel())
+	if (Level && ensure(IsValid(Level)))
 	{
-		FLevelAnnotation LevelAnnotation = LevelAnnotations.GetAnnotation(Level);
-		if (LevelAnnotation.LevelStreaming)
+		if (Level->OwningWorld && !Level->IsPersistentLevel())
 		{
-			check(LevelAnnotation.LevelStreaming->GetLoadedLevel() == Level);
-			FoundLevelStreaming = LevelAnnotation.LevelStreaming;
+			FLevelAnnotation LevelAnnotation = LevelAnnotations.GetAnnotation(Level);
+			if (LevelAnnotation.LevelStreaming)
+			{
+				check(LevelAnnotation.LevelStreaming->GetLoadedLevel() == Level);
+				FoundLevelStreaming = LevelAnnotation.LevelStreaming;
+			}
 		}
 	}
 
