Fix leftover comment

diff --git a/Engine/Source/Runtime/Engine/Private/SceneTexturesConfig.cpp b/Engine/Source/Runtime/Engine/Private/SceneTexturesConfig.cpp
--- a/Engine/Source/Runtime/Engine/Private/SceneTexturesConfig.cpp
+++ b/Engine/Source/Runtime/Engine/Private/SceneTexturesConfig.cpp
@@ -175,7 +175,7 @@ void GetSceneColorFormatAndCreateFlags(ERHIFeatureLevel::Type FeatureLevel, bool
 	SceneColorCreateFlags |= sRGBFlag;
 }
 
-static ETextureCreateFlags GetSceneDepthStencilCreateFlags(uint32 NumSamples, bool bKeepDepthContent, bool bMemorylessMSAA, EShaderPlatform ShaderPlatform, ETextureCreateFlags ExtraSceneDepthCreateFlags) //ericado
+static ETextureCreateFlags GetSceneDepthStencilCreateFlags(uint32 NumSamples, bool bKeepDepthContent, bool bMemorylessMSAA, EShaderPlatform ShaderPlatform, ETextureCreateFlags ExtraSceneDepthCreateFlags)
 {
 	ETextureCreateFlags DepthCreateFlags = TexCreate_DepthStencilTargetable | TexCreate_ShaderResource | TexCreate_InputAttachmentRead | ExtraSceneDepthCreateFlags;
 	// We can't discard the SceneDepth MSAA target if full depth prepass is enabled because it will be used in the base pass later on.
