Removed includes to speed up compiles

diff --git a/Engine/Plugins/Animation/AnimationWarping/Source/Runtime/Private/BoneControllers/AnimNode_FootPlacement.cpp b/Engine/Plugins/Animation/AnimationWarping/Source/Runtime/Private/BoneControllers/AnimNode_FootPlacement.cpp
--- a/Engine/Plugins/Animation/AnimationWarping/Source/Runtime/Private/BoneControllers/AnimNode_FootPlacement.cpp
+++ b/Engine/Plugins/Animation/AnimationWarping/Source/Runtime/Private/BoneControllers/AnimNode_FootPlacement.cpp
@@ -2,6 +2,7 @@
 
 #include "BoneControllers/AnimNode_FootPlacement.h"
 #include "Animation/AnimNodeFunctionRef.h"
+#include "Animation/AnimTrace.h"
 #include "AnimationRuntime.h"
 #include "Animation/AnimInstanceProxy.h"
 #include "Components/SkeletalMeshComponent.h"
diff --git a/Engine/Plugins/Animation/AnimationWarping/Source/Runtime/Private/BoneControllers/AnimNode_OffsetRootBone.cpp b/Engine/Plugins/Animation/AnimationWarping/Source/Runtime/Private/BoneControllers/AnimNode_OffsetRootBone.cpp
--- a/Engine/Plugins/Animation/AnimationWarping/Source/Runtime/Private/BoneControllers/AnimNode_OffsetRootBone.cpp
+++ b/Engine/Plugins/Animation/AnimationWarping/Source/Runtime/Private/BoneControllers/AnimNode_OffsetRootBone.cpp
@@ -5,6 +5,8 @@
 #include "Animation/AnimInstanceProxy.h"
 #include "Animation/AnimNodeFunctionRef.h"
 #include "Animation/AnimRootMotionProvider.h"
+#include "Components/SkeletalMeshComponent.h"
+#include "Engine/World.h"
 #include "HAL/IConsoleManager.h"
 #include "VisualLogger/VisualLogger.h"
 
diff --git a/Engine/Plugins/Animation/BlendStack/Source/Runtime/Private/AnimNode_BlendStack.cpp b/Engine/Plugins/Animation/BlendStack/Source/Runtime/Private/AnimNode_BlendStack.cpp
--- a/Engine/Plugins/Animation/BlendStack/Source/Runtime/Private/AnimNode_BlendStack.cpp
+++ b/Engine/Plugins/Animation/BlendStack/Source/Runtime/Private/AnimNode_BlendStack.cpp
@@ -9,6 +9,7 @@
 #include "Animation/AnimNode_Inertialization.h"
 #include "Animation/AnimPoseSearchProvider.h"
 #include "Animation/AnimSequence.h"
+#include "Animation/AnimTrace.h"
 #include "Animation/BlendSpace.h"
 #include "BlendStack/AnimNode_BlendStackInput.h"
 #include "BlendStackAnimEventsFilterScope.h"
diff --git a/Engine/Plugins/Animation/ControlRig/Source/ControlRig/Private/AnimNode_ControlRig.cpp b/Engine/Plugins/Animation/ControlRig/Source/ControlRig/Private/AnimNode_ControlRig.cpp
--- a/Engine/Plugins/Animation/ControlRig/Source/ControlRig/Private/AnimNode_ControlRig.cpp
+++ b/Engine/Plugins/Animation/ControlRig/Source/ControlRig/Private/AnimNode_ControlRig.cpp
@@ -1,6 +1,7 @@
 // Copyright Epic Games, Inc. All Rights Reserved.
 
 #include "AnimNode_ControlRig.h"
+#include "Animation/AnimTrace.h"
 #include "ControlRig.h"
 #include "ControlRigComponent.h"
 #include "Components/SkeletalMeshComponent.h"
diff --git a/Engine/Plugins/Animation/ControlRig/Source/ControlRig/Private/ControlRigGizmoLibrary.cpp b/Engine/Plugins/Animation/ControlRig/Source/ControlRig/Private/ControlRigGizmoLibrary.cpp
--- a/Engine/Plugins/Animation/ControlRig/Source/ControlRig/Private/ControlRigGizmoLibrary.cpp
+++ b/Engine/Plugins/Animation/ControlRig/Source/ControlRig/Private/ControlRigGizmoLibrary.cpp
@@ -1,6 +1,7 @@
 // Copyright Epic Games, Inc. All Rights Reserved.
 
 #include "ControlRigGizmoLibrary.h"
+#include "Materials/Material.h"
 
 #include UE_INLINE_GENERATED_CPP_BY_NAME(ControlRigGizmoLibrary)
 
diff --git a/Engine/Plugins/Animation/ControlRig/Source/ControlRig/Private/Rigs/RigHierarchy.cpp b/Engine/Plugins/Animation/ControlRig/Source/ControlRig/Private/Rigs/RigHierarchy.cpp
--- a/Engine/Plugins/Animation/ControlRig/Source/ControlRig/Private/Rigs/RigHierarchy.cpp
+++ b/Engine/Plugins/Animation/ControlRig/Source/ControlRig/Private/Rigs/RigHierarchy.cpp
@@ -1,8 +1,10 @@
 // Copyright Epic Games, Inc. All Rights Reserved.
 
 #include "Rigs/RigHierarchy.h"
+#include "Async/TaskGraphInterfaces.h"
 #include "ControlRig.h"
 
+#include "Engine/World.h"
 #include "Rigs/RigHierarchyElements.h"
 #include "Rigs/RigHierarchyController.h"
 #include "ModularRigRuleManager.h"
diff --git a/Engine/Plugins/Animation/ControlRig/Source/ControlRig/Public/ControlRigAnimInstance.h b/Engine/Plugins/Animation/ControlRig/Source/ControlRig/Public/ControlRigAnimInstance.h
--- a/Engine/Plugins/Animation/ControlRig/Source/ControlRig/Public/ControlRigAnimInstance.h
+++ b/Engine/Plugins/Animation/ControlRig/Source/ControlRig/Public/ControlRigAnimInstance.h
@@ -4,6 +4,7 @@
 
 #include "Animation/AnimInstance.h"
 #include "Animation/AnimInstanceProxy.h"
+#include "Animation/AttributesRuntime.h"
 #include "ControlRigAnimInstance.generated.h"
 
 struct FMeshPoseBoneIndex;
diff --git a/Engine/Plugins/Animation/ControlRig/Source/ControlRig/Public/ControlRigGizmoLibrary.h b/Engine/Plugins/Animation/ControlRig/Source/ControlRig/Public/ControlRigGizmoLibrary.h
--- a/Engine/Plugins/Animation/ControlRig/Source/ControlRig/Public/ControlRigGizmoLibrary.h
+++ b/Engine/Plugins/Animation/ControlRig/Source/ControlRig/Public/ControlRigGizmoLibrary.h
@@ -6,11 +6,11 @@
 #include "UObject/ObjectMacros.h"
 #include "UObject/ScriptMacros.h"
 #include "Engine/StaticMesh.h"
-#include "Materials/Material.h"
 
 #include "ControlRigGizmoLibrary.generated.h"
 
 class UControlRigShapeLibrary;
+class UMaterial;
 
 USTRUCT(BlueprintType, meta = (DisplayName = "Shape"))
 struct CONTROLRIG_API FControlRigShapeDefinition
@@ -76,4 +76,8 @@ private:
 	const TArray<FName> GetUpdatedNameList(bool bReset = false);
 
 	TArray<FName> NameList;
-};
\ No newline at end of file
+};
+
+#if UE_ENABLE_INCLUDE_ORDER_DEPRECATED_IN_5_6
+#include "Materials/Material.h"
+#endif
diff --git a/Engine/Plugins/Animation/ControlRig/Source/ControlRigEditor/Private/ControlRigGizmoLibraryFactory.cpp b/Engine/Plugins/Animation/ControlRig/Source/ControlRigEditor/Private/ControlRigGizmoLibraryFactory.cpp
--- a/Engine/Plugins/Animation/ControlRig/Source/ControlRigEditor/Private/ControlRigGizmoLibraryFactory.cpp
+++ b/Engine/Plugins/Animation/ControlRig/Source/ControlRigEditor/Private/ControlRigGizmoLibraryFactory.cpp
@@ -3,6 +3,7 @@
 #include "ControlRigGizmoLibraryFactory.h"
 #include "AssetTypeCategories.h"
 #include "Engine/StaticMesh.h"
+#include "Materials/Material.h"
 
 #include UE_INLINE_GENERATED_CPP_BY_NAME(ControlRigGizmoLibraryFactory)
 
diff --git a/Engine/Plugins/Animation/ControlRig/Source/ControlRigEditor/Private/ControlRigThumbnailRenderer.cpp b/Engine/Plugins/Animation/ControlRig/Source/ControlRigEditor/Private/ControlRigThumbnailRenderer.cpp
--- a/Engine/Plugins/Animation/ControlRig/Source/ControlRigEditor/Private/ControlRigThumbnailRenderer.cpp
+++ b/Engine/Plugins/Animation/ControlRig/Source/ControlRigEditor/Private/ControlRigThumbnailRenderer.cpp
@@ -1,6 +1,7 @@
 // Copyright Epic Games, Inc. All Rights Reserved.
 
 #include "ControlRigThumbnailRenderer.h"
+#include "Materials/Material.h"
 #include "ThumbnailHelpers.h"
 #include "Engine/StaticMesh.h"
 #include "Components/StaticMeshComponent.h"
diff --git a/Engine/Plugins/Animation/ControlRig/Source/ControlRigEditor/Private/Editor/SControlRigAnimAttributeView.cpp b/Engine/Plugins/Animation/ControlRig/Source/ControlRigEditor/Private/Editor/SControlRigAnimAttributeView.cpp
--- a/Engine/Plugins/Animation/ControlRig/Source/ControlRigEditor/Private/Editor/SControlRigAnimAttributeView.cpp
+++ b/Engine/Plugins/Animation/ControlRig/Source/ControlRigEditor/Private/Editor/SControlRigAnimAttributeView.cpp
@@ -1,6 +1,7 @@
 // Copyright Epic Games, Inc. All Rights Reserved.
 
 #include "Editor/SControlRigAnimAttributeView.h"
+#include "Components/SkeletalMeshComponent.h"
 #include "Editor/ControlRigEditor.h"
 #include "ControlRigBlueprint.h"
 #include "Engine/Console.h"
diff --git a/Engine/Plugins/Animation/DeformerGraph/Source/OptimusCore/Private/ControlRig/RigUnit_Optimus.cpp b/Engine/Plugins/Animation/DeformerGraph/Source/OptimusCore/Private/ControlRig/RigUnit_Optimus.cpp
--- a/Engine/Plugins/Animation/DeformerGraph/Source/OptimusCore/Private/ControlRig/RigUnit_Optimus.cpp
+++ b/Engine/Plugins/Animation/DeformerGraph/Source/OptimusCore/Private/ControlRig/RigUnit_Optimus.cpp
@@ -1,6 +1,8 @@
 ﻿// Copyright Epic Games, Inc. All Rights Reserved.
 #include "RigUnit_Optimus.h"
 
+#include "Components/SkeletalMeshComponent.h"
+
 #if WITH_EDITOR
 #include "RigVMModel/RigVMController.h"
 #endif 
diff --git a/Engine/Plugins/Animation/DeformerGraph/Source/OptimusCore/Private/DataInterfaces/OptimusDataInterfaceSkinWeightsAsVertexMask.cpp b/Engine/Plugins/Animation/DeformerGraph/Source/OptimusCore/Private/DataInterfaces/OptimusDataInterfaceSkinWeightsAsVertexMask.cpp
--- a/Engine/Plugins/Animation/DeformerGraph/Source/OptimusCore/Private/DataInterfaces/OptimusDataInterfaceSkinWeightsAsVertexMask.cpp
+++ b/Engine/Plugins/Animation/DeformerGraph/Source/OptimusCore/Private/DataInterfaces/OptimusDataInterfaceSkinWeightsAsVertexMask.cpp
@@ -5,6 +5,7 @@
 #include "OptimusDataDomain.h"
 #include "OptimusDataTypeRegistry.h"
 #include "OptimusDeformerInstance.h"
+#include "ShaderCompilerCore.h"
 #include "ShaderParameterMetadataBuilder.h"
 #include "SkeletalMeshDeformerHelpers.h"
 #include "SkeletalRenderPublic.h"
diff --git a/Engine/Plugins/Animation/DeformerGraph/Source/OptimusCore/Private/OptimusDeformerDynamicInstanceManager.cpp b/Engine/Plugins/Animation/DeformerGraph/Source/OptimusCore/Private/OptimusDeformerDynamicInstanceManager.cpp
--- a/Engine/Plugins/Animation/DeformerGraph/Source/OptimusCore/Private/OptimusDeformerDynamicInstanceManager.cpp
+++ b/Engine/Plugins/Animation/DeformerGraph/Source/OptimusCore/Private/OptimusDeformerDynamicInstanceManager.cpp
@@ -2,6 +2,7 @@
 
 #include "OptimusDeformerDynamicInstanceManager.h"
 
+#include "Components/SkeletalMeshComponent.h"
 #include "ControlRig.h"
 #include "IControlRigObjectBinding.h"
 #include "OptimusDeformerInstance.h"
diff --git a/Engine/Plugins/Animation/IKRig/Source/IKRig/Private/AnimNodes/AnimNode_IKRig.cpp b/Engine/Plugins/Animation/IKRig/Source/IKRig/Private/AnimNodes/AnimNode_IKRig.cpp
--- a/Engine/Plugins/Animation/IKRig/Source/IKRig/Private/AnimNodes/AnimNode_IKRig.cpp
+++ b/Engine/Plugins/Animation/IKRig/Source/IKRig/Private/AnimNodes/AnimNode_IKRig.cpp
@@ -2,6 +2,8 @@
 
 #include "AnimNodes/AnimNode_IKRig.h"
 
+#include "Animation/AnimInstance.h"
+#include "Animation/AnimTrace.h"
 #include "Rig/IKRigProcessor.h"
 #include "Rig/IKRigDefinition.h"
 #include "Rig/Solvers/IKRigSolver.h"
diff --git a/Engine/Plugins/Animation/IKRig/Source/IKRig/Private/AnimNodes/AnimNode_RetargetPoseFromMesh.cpp b/Engine/Plugins/Animation/IKRig/Source/IKRig/Private/AnimNodes/AnimNode_RetargetPoseFromMesh.cpp
--- a/Engine/Plugins/Animation/IKRig/Source/IKRig/Private/AnimNodes/AnimNode_RetargetPoseFromMesh.cpp
+++ b/Engine/Plugins/Animation/IKRig/Source/IKRig/Private/AnimNodes/AnimNode_RetargetPoseFromMesh.cpp
@@ -3,6 +3,7 @@
 #include "AnimNodes/AnimNode_RetargetPoseFromMesh.h"
 
 #include "Animation/AnimCurveUtils.h"
+#include "Animation/AnimInstance.h"
 #include "Animation/AnimInstanceProxy.h"
 #include "Components/SkeletalMeshComponent.h"
 #include "Engine/SkeletalMesh.h"
diff --git a/Engine/Plugins/Animation/PerformanceCaptureCore/Source/PerformanceCaptureCore/Private/CaptureCharacter.cpp b/Engine/Plugins/Animation/PerformanceCaptureCore/Source/PerformanceCaptureCore/Private/CaptureCharacter.cpp
--- a/Engine/Plugins/Animation/PerformanceCaptureCore/Source/PerformanceCaptureCore/Private/CaptureCharacter.cpp
+++ b/Engine/Plugins/Animation/PerformanceCaptureCore/Source/PerformanceCaptureCore/Private/CaptureCharacter.cpp
@@ -2,6 +2,7 @@
 
 #include "CaptureCharacter.h"
 
+#include "Components/SkeletalMeshComponent.h"
 #include "RetargetAnimInstance.h"
 #include "RetargetComponent.h"
 
diff --git a/Engine/Plugins/Animation/PerformanceCaptureCore/Source/PerformanceCaptureCore/Private/CapturePerformer.cpp b/Engine/Plugins/Animation/PerformanceCaptureCore/Source/PerformanceCaptureCore/Private/CapturePerformer.cpp
--- a/Engine/Plugins/Animation/PerformanceCaptureCore/Source/PerformanceCaptureCore/Private/CapturePerformer.cpp
+++ b/Engine/Plugins/Animation/PerformanceCaptureCore/Source/PerformanceCaptureCore/Private/CapturePerformer.cpp
@@ -1,6 +1,7 @@
 // Copyright Epic Games, Inc. All Rights Reserved.
 
 #include "CapturePerformer.h"
+#include "Components/SkeletalMeshComponent.h"
 #include "PerformerComponent.h"
 
 // Sets default values
diff --git a/Engine/Plugins/Animation/PerformanceCaptureCore/Source/PerformanceCaptureCore/Private/PerformerComponent.cpp b/Engine/Plugins/Animation/PerformanceCaptureCore/Source/PerformanceCaptureCore/Private/PerformerComponent.cpp
--- a/Engine/Plugins/Animation/PerformanceCaptureCore/Source/PerformanceCaptureCore/Private/PerformerComponent.cpp
+++ b/Engine/Plugins/Animation/PerformanceCaptureCore/Source/PerformanceCaptureCore/Private/PerformerComponent.cpp
@@ -1,6 +1,7 @@
 // Copyright Epic Games, Inc. All Rights Reserved.
 
 #include "PerformerComponent.h"
+#include "Components/SkeletalMeshComponent.h"
 #include "LiveLinkInstance.h"
 
 // Sets default values for this component's properties
diff --git a/Engine/Plugins/Animation/PerformanceCaptureCore/Source/PerformanceCaptureCore/Private/RetargetComponent.cpp b/Engine/Plugins/Animation/PerformanceCaptureCore/Source/PerformanceCaptureCore/Private/RetargetComponent.cpp
--- a/Engine/Plugins/Animation/PerformanceCaptureCore/Source/PerformanceCaptureCore/Private/RetargetComponent.cpp
+++ b/Engine/Plugins/Animation/PerformanceCaptureCore/Source/PerformanceCaptureCore/Private/RetargetComponent.cpp
@@ -1,6 +1,7 @@
 // Copyright Epic Games, Inc. All Rights Reserved.
 
 #include "RetargetComponent.h"
+#include "Components/SkeletalMeshComponent.h"
 #include "LiveLinkInstance.h"
 #include "RetargetAnimInstance.h"
 
diff --git a/Engine/Plugins/Editor/DataValidation/Source/DataValidation/Private/EditorValidator_Material.cpp b/Engine/Plugins/Editor/DataValidation/Source/DataValidation/Private/EditorValidator_Material.cpp
--- a/Engine/Plugins/Editor/DataValidation/Source/DataValidation/Private/EditorValidator_Material.cpp
+++ b/Engine/Plugins/Editor/DataValidation/Source/DataValidation/Private/EditorValidator_Material.cpp
@@ -1,10 +1,11 @@
-﻿// Copyright Epic Games, Inc. All Rights Reserved.
+// Copyright Epic Games, Inc. All Rights Reserved.
 
 #include "EditorValidator_Material.h"
 #include "AssetCompilingManager.h"
 #include "DataDrivenShaderPlatformInfo.h"
 #include "DetailWidgetRow.h"
 #include "EditorValidatorSubsystem.h"
+#include "MaterialShared.h"
 #include "Materials/Material.h"
 #include "Materials/MaterialInstance.h"
 #include "Misc/DataValidation.h"
diff --git a/Engine/Plugins/Experimental/AnimNextAnimGraph/Source/AnimNextAnimGraph/Private/EvaluationVM/Tasks/DeadBlending.cpp b/Engine/Plugins/Experimental/AnimNextAnimGraph/Source/AnimNextAnimGraph/Private/EvaluationVM/Tasks/DeadBlending.cpp
--- a/Engine/Plugins/Experimental/AnimNextAnimGraph/Source/AnimNextAnimGraph/Private/EvaluationVM/Tasks/DeadBlending.cpp
+++ b/Engine/Plugins/Experimental/AnimNextAnimGraph/Source/AnimNextAnimGraph/Private/EvaluationVM/Tasks/DeadBlending.cpp
@@ -2,6 +2,7 @@
 
 #include "EvaluationVM/Tasks/DeadBlending.h"
 
+#include "Animation/AnimTrace.h"
 #include "EvaluationVM/EvaluationVM.h"
 #include "EvaluationVM/KeyframeState.h"
 
diff --git a/Engine/Plugins/Experimental/AnimNextAnimGraph/Source/AnimNextAnimGraph/Private/Graph/AnimGraph/AnimNode_AnimNextGraph.cpp b/Engine/Plugins/Experimental/AnimNextAnimGraph/Source/AnimNextAnimGraph/Private/Graph/AnimGraph/AnimNode_AnimNextGraph.cpp
--- a/Engine/Plugins/Experimental/AnimNextAnimGraph/Source/AnimNextAnimGraph/Private/Graph/AnimGraph/AnimNode_AnimNextGraph.cpp
+++ b/Engine/Plugins/Experimental/AnimNextAnimGraph/Source/AnimNextAnimGraph/Private/Graph/AnimGraph/AnimNode_AnimNextGraph.cpp
@@ -1,6 +1,7 @@
 // Copyright Epic Games, Inc. All Rights Reserved.
 
 #include "Graph/AnimGraph/AnimNode_AnimNextGraph.h"
+#include "Animation/AnimInstance.h"
 #include "DataRegistry.h"
 #include "ReferencePose.h"
 #include "Components/SkeletalMeshComponent.h"
diff --git a/Engine/Plugins/Experimental/Animation/RigMapper/Source/RigMapper/Private/AnimNode_RigMapper.cpp b/Engine/Plugins/Experimental/Animation/RigMapper/Source/RigMapper/Private/AnimNode_RigMapper.cpp
--- a/Engine/Plugins/Experimental/Animation/RigMapper/Source/RigMapper/Private/AnimNode_RigMapper.cpp
+++ b/Engine/Plugins/Experimental/Animation/RigMapper/Source/RigMapper/Private/AnimNode_RigMapper.cpp
@@ -2,6 +2,8 @@
 
 #include "AnimNode_RigMapper.h"
 
+#include "Animation/AnimInstance.h"
+#include "Components/SkeletalMeshComponent.h"
 #include "UObject/SoftObjectPath.h"
 #include "RigMapperDefinition.h"
 #include "Animation/AnimCurveUtils.h"
diff --git a/Engine/Plugins/Experimental/ChaosModularVehicle/Source/ChaosModularVehicleEngine/Private/ChaosModularVehicle/AnimNode_ModularVehicleController.cpp b/Engine/Plugins/Experimental/ChaosModularVehicle/Source/ChaosModularVehicleEngine/Private/ChaosModularVehicle/AnimNode_ModularVehicleController.cpp
--- a/Engine/Plugins/Experimental/ChaosModularVehicle/Source/ChaosModularVehicleEngine/Private/ChaosModularVehicle/AnimNode_ModularVehicleController.cpp
+++ b/Engine/Plugins/Experimental/ChaosModularVehicle/Source/ChaosModularVehicleEngine/Private/ChaosModularVehicle/AnimNode_ModularVehicleController.cpp
@@ -1,6 +1,7 @@
 // Copyright Epic Games, Inc. All Rights Reserved.
 
 #include "ChaosModularVehicle/AnimNode_ModularVehicleController.h"
+#include "Animation/AnimTrace.h"
 #include "AnimationRuntime.h"
 #include "Animation/AnimStats.h"
 #include "SimModule/SimulationModuleBase.h"
diff --git a/Engine/Plugins/Experimental/ChaosVehiclesPlugin/Source/ChaosVehicles/Private/AnimNode_StageCoachWheelController.cpp b/Engine/Plugins/Experimental/ChaosVehiclesPlugin/Source/ChaosVehicles/Private/AnimNode_StageCoachWheelController.cpp
--- a/Engine/Plugins/Experimental/ChaosVehiclesPlugin/Source/ChaosVehicles/Private/AnimNode_StageCoachWheelController.cpp
+++ b/Engine/Plugins/Experimental/ChaosVehiclesPlugin/Source/ChaosVehicles/Private/AnimNode_StageCoachWheelController.cpp
@@ -1,6 +1,7 @@
 // Copyright Epic Games, Inc. All Rights Reserved.
 
 #include "AnimNode_StageCoachWheelController.h"
+#include "Animation/AnimTrace.h"
 #include "AnimationRuntime.h"
 #include "Animation/AnimStats.h"
 #include "WheeledVehiclePawn.h"
diff --git a/Engine/Plugins/Experimental/ChaosVehiclesPlugin/Source/ChaosVehicles/Private/AnimNode_WheelController.cpp b/Engine/Plugins/Experimental/ChaosVehiclesPlugin/Source/ChaosVehicles/Private/AnimNode_WheelController.cpp
--- a/Engine/Plugins/Experimental/ChaosVehiclesPlugin/Source/ChaosVehicles/Private/AnimNode_WheelController.cpp
+++ b/Engine/Plugins/Experimental/ChaosVehiclesPlugin/Source/ChaosVehicles/Private/AnimNode_WheelController.cpp
@@ -1,6 +1,7 @@
 // Copyright Epic Games, Inc. All Rights Reserved.
 
 #include "AnimNode_WheelController.h"
+#include "Animation/AnimTrace.h"
 #include "AnimationRuntime.h"
 #include "Animation/AnimStats.h"
 #include "WheeledVehiclePawn.h"
diff --git a/Engine/Plugins/Experimental/DynamicMaterial/Source/DynamicMaterialEditor/Private/UI/Widgets/Editor/SlotEditor/SDMMaterialStage.cpp b/Engine/Plugins/Experimental/DynamicMaterial/Source/DynamicMaterialEditor/Private/UI/Widgets/Editor/SlotEditor/SDMMaterialStage.cpp
--- a/Engine/Plugins/Experimental/DynamicMaterial/Source/DynamicMaterialEditor/Private/UI/Widgets/Editor/SlotEditor/SDMMaterialStage.cpp
+++ b/Engine/Plugins/Experimental/DynamicMaterial/Source/DynamicMaterialEditor/Private/UI/Widgets/Editor/SlotEditor/SDMMaterialStage.cpp
@@ -16,6 +16,7 @@
 #include "DynamicMaterialModule.h"
 #include "Engine/Texture.h"
 #include "Framework/Application/SlateApplication.h"
+#include "Materials/Material.h"
 #include "SAssetDropTarget.h"
 #include "Styling/StyleColors.h"
 #include "UI/Menus/DMMaterialStageMenus.h"
diff --git a/Engine/Plugins/Experimental/PhysicsControl/Source/PhysicsControl/Private/AnimNode_RigidBodyWithControl.cpp b/Engine/Plugins/Experimental/PhysicsControl/Source/PhysicsControl/Private/AnimNode_RigidBodyWithControl.cpp
--- a/Engine/Plugins/Experimental/PhysicsControl/Source/PhysicsControl/Private/AnimNode_RigidBodyWithControl.cpp
+++ b/Engine/Plugins/Experimental/PhysicsControl/Source/PhysicsControl/Private/AnimNode_RigidBodyWithControl.cpp
@@ -1,6 +1,7 @@
 // Copyright Epic Games, Inc. All Rights Reserved.
 
 #include "AnimNode_RigidBodyWithControl.h"
+#include "Animation/AnimInstance.h"
 #include "PhysicsControlLog.h"
 #include "AnimationRuntime.h"
 #include "Animation/AnimInstanceProxy.h"
diff --git a/Engine/Plugins/FX/Niagara/Source/Niagara/Classes/NiagaraDataSet.h b/Engine/Plugins/FX/Niagara/Source/Niagara/Classes/NiagaraDataSet.h
--- a/Engine/Plugins/FX/Niagara/Source/Niagara/Classes/NiagaraDataSet.h
+++ b/Engine/Plugins/FX/Niagara/Source/Niagara/Classes/NiagaraDataSet.h
@@ -3,7 +3,6 @@
 
 #include "CoreMinimal.h"
 #include "NiagaraCommon.h"
-#include "NiagaraCompileHashVisitor.h"
 #include "NiagaraDataSetCompiledData.h"
 #include "Containers/DynamicRHIResourceArray.h"
 #include "RHI.h"
@@ -593,3 +592,7 @@ namespace NiagaraDataSetPrivate
 		return DataBuffer->GetNumInstances();
 	}
 }
+
+#if UE_ENABLE_INCLUDE_ORDER_DEPRECATED_IN_5_6
+#include "NiagaraCompileHashVisitor.h"
+#endif
diff --git a/Engine/Plugins/FX/Niagara/Source/Niagara/Private/NiagaraComponentLeakDetector.cpp b/Engine/Plugins/FX/Niagara/Source/Niagara/Private/NiagaraComponentLeakDetector.cpp
--- a/Engine/Plugins/FX/Niagara/Source/Niagara/Private/NiagaraComponentLeakDetector.cpp
+++ b/Engine/Plugins/FX/Niagara/Source/Niagara/Private/NiagaraComponentLeakDetector.cpp
@@ -5,6 +5,7 @@
 #include "NiagaraEffectType.h"
 #include "NiagaraSystem.h"
 #include "NiagaraWorldManager.h"
+#include "UObject/UObjectIterator.h"
 
 #if WITH_NIAGARA_LEAK_DETECTOR
 
diff --git a/Engine/Plugins/FX/Niagara/Source/Niagara/Private/NiagaraDataChannel.cpp b/Engine/Plugins/FX/Niagara/Source/Niagara/Private/NiagaraDataChannel.cpp
--- a/Engine/Plugins/FX/Niagara/Source/Niagara/Private/NiagaraDataChannel.cpp
+++ b/Engine/Plugins/FX/Niagara/Source/Niagara/Private/NiagaraDataChannel.cpp
@@ -2,6 +2,7 @@
 
 #include "NiagaraDataChannel.h"
 
+#include "Misc/LargeWorldRenderPosition.h"
 #include "NiagaraWorldManager.h"
 #include "NiagaraDataChannelCommon.h"
 #include "NiagaraDataChannelHandler.h"
diff --git a/Engine/Plugins/FX/Niagara/Source/Niagara/Private/NiagaraDataChannel_Islands.cpp b/Engine/Plugins/FX/Niagara/Source/Niagara/Private/NiagaraDataChannel_Islands.cpp
--- a/Engine/Plugins/FX/Niagara/Source/Niagara/Private/NiagaraDataChannel_Islands.cpp
+++ b/Engine/Plugins/FX/Niagara/Source/Niagara/Private/NiagaraDataChannel_Islands.cpp
@@ -3,6 +3,7 @@
 #include "NiagaraDataChannel_Islands.h"
 
 #include "Engine/AssetManager.h"
+#include "Misc/LargeWorldRenderPosition.h"
 #include "UObject/UObjectGlobals.h"
 
 #include "NiagaraFunctionLibrary.h"
diff --git a/Engine/Plugins/FX/Niagara/Source/Niagara/Private/NiagaraDataInterfaceLandscape.cpp b/Engine/Plugins/FX/Niagara/Source/Niagara/Private/NiagaraDataInterfaceLandscape.cpp
--- a/Engine/Plugins/FX/Niagara/Source/Niagara/Private/NiagaraDataInterfaceLandscape.cpp
+++ b/Engine/Plugins/FX/Niagara/Source/Niagara/Private/NiagaraDataInterfaceLandscape.cpp
@@ -3,6 +3,7 @@
 #include "NiagaraDataInterfaceLandscape.h"
 
 #include "Containers/ResourceArray.h"
+#include "Misc/LargeWorldRenderPosition.h"
 #include "NiagaraShaderParametersBuilder.h"
 #include "GlobalRenderResources.h"
 #include "NiagaraStats.h"
diff --git a/Engine/Plugins/FX/Niagara/Source/Niagara/Private/NiagaraDataInterfaceSparseVolumeTexture.cpp b/Engine/Plugins/FX/Niagara/Source/Niagara/Private/NiagaraDataInterfaceSparseVolumeTexture.cpp
--- a/Engine/Plugins/FX/Niagara/Source/Niagara/Private/NiagaraDataInterfaceSparseVolumeTexture.cpp
+++ b/Engine/Plugins/FX/Niagara/Source/Niagara/Private/NiagaraDataInterfaceSparseVolumeTexture.cpp
@@ -1,6 +1,7 @@
 // Copyright Epic Games, Inc. All Rights Reserved.
 
 #include "NiagaraDataInterfaceSparseVolumeTexture.h"
+#include "NiagaraCompileHashVisitor.h"
 #include "SparseVolumeTexture/SparseVolumeTexture.h"
 #include "NiagaraComputeExecutionContext.h"
 #include "NiagaraCustomVersion.h"
diff --git a/Engine/Plugins/FX/Niagara/Source/Niagara/Private/NiagaraSystem.cpp b/Engine/Plugins/FX/Niagara/Source/Niagara/Private/NiagaraSystem.cpp
--- a/Engine/Plugins/FX/Niagara/Source/Niagara/Private/NiagaraSystem.cpp
+++ b/Engine/Plugins/FX/Niagara/Source/Niagara/Private/NiagaraSystem.cpp
@@ -3,6 +3,7 @@
 
 #include "NiagaraSystem.h"
 
+#include "Containers/StridedView.h"
 #include "EngineAnalytics.h"
 #include "Components/PrimitiveComponent.h"
 #include "NiagaraSystemImpl.h"
diff --git a/Engine/Plugins/FX/Niagara/Source/Niagara/Public/NiagaraSystemInstanceController.h b/Engine/Plugins/FX/Niagara/Source/Niagara/Public/NiagaraSystemInstanceController.h
--- a/Engine/Plugins/FX/Niagara/Source/Niagara/Public/NiagaraSystemInstanceController.h
+++ b/Engine/Plugins/FX/Niagara/Source/Niagara/Public/NiagaraSystemInstanceController.h
@@ -34,6 +34,7 @@ class UMaterialInterface;
 class UTextureRenderTarget;
 class FNiagaraSystemInstanceController;
 class FNiagaraSceneProxy;
+class FStreamingTextureLevelContext;
 
 using FNiagaraSystemInstanceControllerPtr = TSharedPtr<FNiagaraSystemInstanceController, ESPMode::ThreadSafe>;
 using FNiagaraSystemInstanceControllerConstPtr = TSharedPtr<const FNiagaraSystemInstanceController, ESPMode::ThreadSafe>;
diff --git a/Engine/Plugins/FX/Niagara/Source/Niagara/Public/NiagaraWorldManager.h b/Engine/Plugins/FX/Niagara/Source/Niagara/Public/NiagaraWorldManager.h
--- a/Engine/Plugins/FX/Niagara/Source/Niagara/Public/NiagaraWorldManager.h
+++ b/Engine/Plugins/FX/Niagara/Source/Niagara/Public/NiagaraWorldManager.h
@@ -16,7 +16,6 @@
 #include "NiagaraSystemSimulation.h"
 #include "NiagaraSystemInstance.h"
 #include "GlobalDistanceFieldParameters.h"
-#include "NiagaraDataInterfaceSkeletalMesh.h"
 #include "NiagaraComponentPool.h"
 #include "NiagaraEffectType.h"
 #include "NiagaraScalabilityManager.h"
@@ -468,3 +467,7 @@ void FNiagaraWorldManager::ForAllWorldManagers(TAction Func)
 		}
 	}
 }
+
+#if UE_ENABLE_INCLUDE_ORDER_DEPRECATED_IN_5_6
+#include "NiagaraDataInterfaceSkeletalMesh.h"
+#endif
diff --git a/Engine/Plugins/FX/Niagara/Source/NiagaraEditor/Private/NiagaraEditorUtilities.cpp b/Engine/Plugins/FX/Niagara/Source/NiagaraEditor/Private/NiagaraEditorUtilities.cpp
--- a/Engine/Plugins/FX/Niagara/Source/NiagaraEditor/Private/NiagaraEditorUtilities.cpp
+++ b/Engine/Plugins/FX/Niagara/Source/NiagaraEditor/Private/NiagaraEditorUtilities.cpp
@@ -7,6 +7,7 @@
 #include "ContentBrowserModule.h"
 #include "EdGraphSchema_Niagara.h"
 #include "GeneralProjectSettings.h"
+#include "NiagaraCompileHashVisitor.h"
 #include "Styling/AppStyle.h"
 #include "IContentBrowserSingleton.h"
 #include "INiagaraEditorTypeUtilities.h"
diff --git a/Engine/Plugins/FX/Niagara/Source/NiagaraEditor/Private/NiagaraNode.cpp b/Engine/Plugins/FX/Niagara/Source/NiagaraEditor/Private/NiagaraNode.cpp
--- a/Engine/Plugins/FX/Niagara/Source/NiagaraEditor/Private/NiagaraNode.cpp
+++ b/Engine/Plugins/FX/Niagara/Source/NiagaraEditor/Private/NiagaraNode.cpp
@@ -1,6 +1,7 @@
 // Copyright Epic Games, Inc. All Rights Reserved.
 
 #include "NiagaraNode.h"
+#include "NiagaraCompileHashVisitor.h"
 #include "NiagaraGraph.h"
 #include "EdGraphSchema_Niagara.h"
 #include "NiagaraHlslTranslator.h"
diff --git a/Engine/Plugins/FX/Niagara/Source/NiagaraEditor/Private/NiagaraNodeEmitter.cpp b/Engine/Plugins/FX/Niagara/Source/NiagaraEditor/Private/NiagaraNodeEmitter.cpp
--- a/Engine/Plugins/FX/Niagara/Source/NiagaraEditor/Private/NiagaraNodeEmitter.cpp
+++ b/Engine/Plugins/FX/Niagara/Source/NiagaraEditor/Private/NiagaraNodeEmitter.cpp
@@ -1,6 +1,7 @@
 // Copyright Epic Games, Inc. All Rights Reserved.
 
 #include "NiagaraNodeEmitter.h"
+#include "NiagaraCompileHashVisitor.h"
 #include "NiagaraSystem.h"
 #include "NiagaraEmitter.h"
 #include "NiagaraEditorUtilities.h"
diff --git a/Engine/Plugins/FX/Niagara/Source/NiagaraEditor/Private/NiagaraNodeOp.cpp b/Engine/Plugins/FX/Niagara/Source/NiagaraEditor/Private/NiagaraNodeOp.cpp
--- a/Engine/Plugins/FX/Niagara/Source/NiagaraEditor/Private/NiagaraNodeOp.cpp
+++ b/Engine/Plugins/FX/Niagara/Source/NiagaraEditor/Private/NiagaraNodeOp.cpp
@@ -1,6 +1,7 @@
 // Copyright Epic Games, Inc. All Rights Reserved.
 
 #include "NiagaraNodeOp.h"
+#include "NiagaraCompileHashVisitor.h"
 #include "NiagaraHlslTranslator.h"
 #include "GraphEditorSettings.h"
 
diff --git a/Engine/Plugins/FX/Niagara/Source/NiagaraEditor/Private/NiagaraScriptVariable.cpp b/Engine/Plugins/FX/Niagara/Source/NiagaraEditor/Private/NiagaraScriptVariable.cpp
--- a/Engine/Plugins/FX/Niagara/Source/NiagaraEditor/Private/NiagaraScriptVariable.cpp
+++ b/Engine/Plugins/FX/Niagara/Source/NiagaraEditor/Private/NiagaraScriptVariable.cpp
@@ -5,6 +5,7 @@
 #include "EdGraphSchema_Niagara.h"
 #include "INiagaraEditorTypeUtilities.h"
 #include "NiagaraCommon.h"
+#include "NiagaraCompileHashVisitor.h"
 #include "NiagaraCustomVersion.h"
 #include "NiagaraEditorModule.h"
 #include "NiagaraEditorUtilities.h"
diff --git a/Engine/Plugins/FX/Niagara/Source/NiagaraEditor/Private/Toolkits/NiagaraSystemToolkit.cpp b/Engine/Plugins/FX/Niagara/Source/NiagaraEditor/Private/Toolkits/NiagaraSystemToolkit.cpp
--- a/Engine/Plugins/FX/Niagara/Source/NiagaraEditor/Private/Toolkits/NiagaraSystemToolkit.cpp
+++ b/Engine/Plugins/FX/Niagara/Source/NiagaraEditor/Private/Toolkits/NiagaraSystemToolkit.cpp
@@ -8,6 +8,7 @@
 #include "NiagaraEmitterHandle.h"
 #include "NiagaraObjectSelection.h"
 #include "UObject/Linker.h"
+#include "UObject/UObjectIterator.h"
 #include "ViewModels/NiagaraSystemViewModel.h"
 #include "ViewModels/NiagaraEmitterHandleViewModel.h"
 #include "ViewModels/NiagaraEmitterViewModel.h"
diff --git a/Engine/Plugins/FX/Niagara/Source/NiagaraEditor/Private/Widgets/SNiagaraSystemViewport.cpp b/Engine/Plugins/FX/Niagara/Source/NiagaraEditor/Private/Widgets/SNiagaraSystemViewport.cpp
--- a/Engine/Plugins/FX/Niagara/Source/NiagaraEditor/Private/Widgets/SNiagaraSystemViewport.cpp
+++ b/Engine/Plugins/FX/Niagara/Source/NiagaraEditor/Private/Widgets/SNiagaraSystemViewport.cpp
@@ -4,6 +4,7 @@
 
 #include "AdvancedPreviewScene.h"
 #include "CanvasItem.h"
+#include "DrawDebugHelpers.h"
 #include "EditorViewportCommands.h"
 #include "EngineUtils.h"
 #include "ImageUtils.h"
diff --git a/Engine/Plugins/FX/Niagara/Source/NiagaraEditor/Public/NiagaraGraph.h b/Engine/Plugins/FX/Niagara/Source/NiagaraEditor/Public/NiagaraGraph.h
--- a/Engine/Plugins/FX/Niagara/Source/NiagaraEditor/Public/NiagaraGraph.h
+++ b/Engine/Plugins/FX/Niagara/Source/NiagaraEditor/Public/NiagaraGraph.h
@@ -14,9 +14,10 @@ class UNiagaraNodeStaticSwitch;
 class UNiagaraParameterDefinitions;
 class UNiagaraScriptVariable;
 class UNiagaraHierarchyRoot;
-struct FSynchronizeWithParameterDefinitionsArgs;
+struct FNiagaraCompileHashVisitorDebugInfo;
 struct FNiagaraScriptVariableData;
 struct FNiagaraStaticVariableSearchContext;
+struct FSynchronizeWithParameterDefinitionsArgs;
 
 /** This is the type of action that occurred on a given Niagara graph. Note that this should follow from EEdGraphActionType, leaving some slop for growth. */
 enum ENiagaraGraphActionType
diff --git a/Engine/Plugins/PCG/Source/PCG/Private/Elements/PCGStaticMeshSpawner.cpp b/Engine/Plugins/PCG/Source/PCG/Private/Elements/PCGStaticMeshSpawner.cpp
--- a/Engine/Plugins/PCG/Source/PCG/Private/Elements/PCGStaticMeshSpawner.cpp
+++ b/Engine/Plugins/PCG/Source/PCG/Private/Elements/PCGStaticMeshSpawner.cpp
@@ -24,6 +24,7 @@
 #include "Components/InstancedStaticMeshComponent.h"
 #include "Engine/StaticMesh.h"
 #include "Materials/MaterialInterface.h"
+#include "ShaderCompilerCore.h"
 #include "UObject/Package.h"
 
 #include UE_INLINE_GENERATED_CPP_BY_NAME(PCGStaticMeshSpawner)
diff --git a/Engine/Plugins/Runtime/HairStrands/Source/HairStrandsCore/Private/Niagara/NiagaraDataInterfaceHairStrands.cpp b/Engine/Plugins/Runtime/HairStrands/Source/HairStrandsCore/Private/Niagara/NiagaraDataInterfaceHairStrands.cpp
--- a/Engine/Plugins/Runtime/HairStrands/Source/HairStrandsCore/Private/Niagara/NiagaraDataInterfaceHairStrands.cpp
+++ b/Engine/Plugins/Runtime/HairStrands/Source/HairStrandsCore/Private/Niagara/NiagaraDataInterfaceHairStrands.cpp
@@ -1,6 +1,7 @@
 // Copyright Epic Games, Inc. All Rights Reserved.
 
 #include "Niagara/NiagaraDataInterfaceHairStrands.h"
+#include "NiagaraCompileHashVisitor.h"
 #include "NiagaraShader.h"
 #include "NiagaraComponent.h"
 #include "NiagaraGpuComputeDispatchInterface.h"
diff --git a/Engine/Plugins/Runtime/HairStrands/Source/HairStrandsCore/Private/Niagara/NiagaraDataInterfacePressureGrid.cpp b/Engine/Plugins/Runtime/HairStrands/Source/HairStrandsCore/Private/Niagara/NiagaraDataInterfacePressureGrid.cpp
--- a/Engine/Plugins/Runtime/HairStrands/Source/HairStrandsCore/Private/Niagara/NiagaraDataInterfacePressureGrid.cpp
+++ b/Engine/Plugins/Runtime/HairStrands/Source/HairStrandsCore/Private/Niagara/NiagaraDataInterfacePressureGrid.cpp
@@ -1,6 +1,7 @@
 // Copyright Epic Games, Inc. All Rights Reserved.
 
 #include "Niagara/NiagaraDataInterfacePressureGrid.h"
+#include "NiagaraCompileHashVisitor.h"
 #include "NiagaraShader.h"
 #include "NiagaraComponent.h"
 #include "NiagaraGpuComputeDispatchInterface.h"
diff --git a/Engine/Plugins/Runtime/HairStrands/Source/HairStrandsCore/Private/Niagara/NiagaraDataInterfaceVelocityGrid.cpp b/Engine/Plugins/Runtime/HairStrands/Source/HairStrandsCore/Private/Niagara/NiagaraDataInterfaceVelocityGrid.cpp
--- a/Engine/Plugins/Runtime/HairStrands/Source/HairStrandsCore/Private/Niagara/NiagaraDataInterfaceVelocityGrid.cpp
+++ b/Engine/Plugins/Runtime/HairStrands/Source/HairStrandsCore/Private/Niagara/NiagaraDataInterfaceVelocityGrid.cpp
@@ -1,6 +1,7 @@
 // Copyright Epic Games, Inc. All Rights Reserved.
 
 #include "Niagara/NiagaraDataInterfaceVelocityGrid.h"
+#include "NiagaraCompileHashVisitor.h"
 #include "NiagaraShader.h"
 #include "NiagaraComponent.h"
 #include "NiagaraGpuComputeDispatchInterface.h"
diff --git a/Engine/Plugins/Runtime/HairStrands/Source/HairStrandsCore/Public/GroomAsset.h b/Engine/Plugins/Runtime/HairStrands/Source/HairStrandsCore/Public/GroomAsset.h
--- a/Engine/Plugins/Runtime/HairStrands/Source/HairStrandsCore/Public/GroomAsset.h
+++ b/Engine/Plugins/Runtime/HairStrands/Source/HairStrandsCore/Public/GroomAsset.h
@@ -3,6 +3,7 @@
 
 #pragma once
 
+#include "PSOPrecache.h"
 #include "UObject/Object.h"
 #include "UObject/ObjectMacros.h"
 #include "HairDescription.h"
diff --git a/Engine/Plugins/Runtime/RigVM/Source/RigVM/Private/RigVMHost.cpp b/Engine/Plugins/Runtime/RigVM/Source/RigVM/Private/RigVMHost.cpp
--- a/Engine/Plugins/Runtime/RigVM/Source/RigVM/Private/RigVMHost.cpp
+++ b/Engine/Plugins/Runtime/RigVM/Source/RigVM/Private/RigVMHost.cpp
@@ -6,6 +6,7 @@
 #include "RigVMCore/RigVMNativized.h"
 #include "RigVMObjectVersion.h"
 #include "RigVMTypeUtils.h"
+#include "SceneManagement.h"
 #include "UObject/UE5MainStreamObjectVersion.h"
 #include "UObject/UObjectIterator.h"
 #include "UObject/ObjectSaveContext.h"
diff --git a/Engine/Plugins/Runtime/RigVM/Source/RigVM/Public/RigVMHost.h b/Engine/Plugins/Runtime/RigVM/Source/RigVM/Public/RigVMHost.h
--- a/Engine/Plugins/Runtime/RigVM/Source/RigVM/Public/RigVMHost.h
+++ b/Engine/Plugins/Runtime/RigVM/Source/RigVM/Public/RigVMHost.h
@@ -9,7 +9,6 @@
 #include "RigVMCore/RigVMGraphFunctionHost.h"
 #include "Engine/AssetUserData.h"
 #include "Interfaces/Interface_AssetUserData.h"
-#include "SceneManagement.h"
 #if WITH_EDITOR
 #include "RigVMCore/RigVMDebugInfo.h"
 #endif
@@ -560,3 +559,7 @@ private:
 
 	int32& Bracket;
 };
+
+#if UE_ENABLE_INCLUDE_ORDER_DEPRECATED_IN_5_6
+#include "SceneManagement.h"
+#endif
diff --git a/Engine/Plugins/Runtime/USDCore/Source/USDUtilities/Private/USDSkeletalDataConversion.cpp b/Engine/Plugins/Runtime/USDCore/Source/USDUtilities/Private/USDSkeletalDataConversion.cpp
--- a/Engine/Plugins/Runtime/USDCore/Source/USDUtilities/Private/USDSkeletalDataConversion.cpp
+++ b/Engine/Plugins/Runtime/USDCore/Source/USDUtilities/Private/USDSkeletalDataConversion.cpp
@@ -2,6 +2,7 @@
 
 #include "USDSkeletalDataConversion.h"
 
+#include "Materials/Material.h"
 #include "UnrealUSDWrapper.h"
 #include "USDAttributeUtils.h"
 #include "USDClassesModule.h"
diff --git a/Engine/Source/Editor/MaterialEditor/Public/MaterialEditingLibrary.h b/Engine/Source/Editor/MaterialEditor/Public/MaterialEditingLibrary.h
--- a/Engine/Source/Editor/MaterialEditor/Public/MaterialEditingLibrary.h
+++ b/Engine/Source/Editor/MaterialEditor/Public/MaterialEditingLibrary.h
@@ -11,6 +11,7 @@
 
 class UMaterialFunction;
 class UMaterialInstance;
+class FMaterialUpdateContext;
 class UMaterialInstanceConstant;
 class URuntimeVirtualTexture;
 
diff --git a/Engine/Source/Editor/UnrealEd/Private/CookOnTheFlyServer.cpp b/Engine/Source/Editor/UnrealEd/Private/CookOnTheFlyServer.cpp
--- a/Engine/Source/Editor/UnrealEd/Private/CookOnTheFlyServer.cpp
+++ b/Engine/Source/Editor/UnrealEd/Private/CookOnTheFlyServer.cpp
@@ -98,6 +98,7 @@
 #include "Logging/MessageLog.h"
 #include "Logging/TokenizedMessage.h"
 #include "LooseFilesCookArtifactReader.h"
+#include "MaterialShared.h"
 #include "Materials/Material.h"
 #include "Materials/MaterialInterface.h"
 #include "MeshCardRepresentation.h"
diff --git a/Engine/Source/Editor/UnrealEd/Private/Cooker/OnDemandShaderCompilation.cpp b/Engine/Source/Editor/UnrealEd/Private/Cooker/OnDemandShaderCompilation.cpp
--- a/Engine/Source/Editor/UnrealEd/Private/Cooker/OnDemandShaderCompilation.cpp
+++ b/Engine/Source/Editor/UnrealEd/Private/Cooker/OnDemandShaderCompilation.cpp
@@ -7,6 +7,7 @@
 #include "AssetRegistry/IAssetRegistry.h"
 #include "AssetRegistry/AssetData.h"
 #include "LandscapeMaterialInstanceConstant.h"
+#include "MaterialShared.h"
 #include "UObject/CoreRedirects.h"
 #include "PackageTools.h"
 #include "WorldPartition/WorldPartitionHelpers.h"
diff --git a/Engine/Source/Editor/UnrealEd/Private/Cooker/OnDemandShaderCompilation.h b/Engine/Source/Editor/UnrealEd/Private/Cooker/OnDemandShaderCompilation.h
--- a/Engine/Source/Editor/UnrealEd/Private/Cooker/OnDemandShaderCompilation.h
+++ b/Engine/Source/Editor/UnrealEd/Private/Cooker/OnDemandShaderCompilation.h
@@ -5,6 +5,7 @@
 #include "UObject/StrongObjectPtr.h"
 #include "Materials/MaterialInterface.h"
 
+class FMaterialShaderMap;
 class UMaterialInterface;
 
 namespace UE::Cook
diff --git a/Engine/Source/Runtime/AIModule/Private/GameplayDebugger/GameplayDebuggerCategory_EQS.cpp b/Engine/Source/Runtime/AIModule/Private/GameplayDebugger/GameplayDebuggerCategory_EQS.cpp
--- a/Engine/Source/Runtime/AIModule/Private/GameplayDebugger/GameplayDebuggerCategory_EQS.cpp
+++ b/Engine/Source/Runtime/AIModule/Private/GameplayDebugger/GameplayDebuggerCategory_EQS.cpp
@@ -6,6 +6,7 @@
 #if WITH_GAMEPLAY_DEBUGGER_MENU
 
 #include "GameFramework/PlayerController.h"
+#include "GlobalRenderResources.h"
 #include "CanvasItem.h"
 #include "Engine/Canvas.h"
 #include "EnvironmentQuery/EnvQueryManager.h"
diff --git a/Engine/Source/Runtime/AnimGraphRuntime/Private/AnimGraphRuntimeModule.cpp b/Engine/Source/Runtime/AnimGraphRuntime/Private/AnimGraphRuntimeModule.cpp
--- a/Engine/Source/Runtime/AnimGraphRuntime/Private/AnimGraphRuntimeModule.cpp
+++ b/Engine/Source/Runtime/AnimGraphRuntime/Private/AnimGraphRuntimeModule.cpp
@@ -1,6 +1,7 @@
 // Copyright Epic Games, Inc. All Rights Reserved.
 
 #include "Modules/ModuleInterface.h"
+#include "HAL/IConsoleManager.h"
 #include "Modules/ModuleManager.h"
 #include "BoneControllers/AnimNode_AnimDynamics.h"
 #include "UObject/UObjectIterator.h"
diff --git a/Engine/Source/Runtime/AnimGraphRuntime/Private/AnimNodes/AnimNode_TwoWayBlend.cpp b/Engine/Source/Runtime/AnimGraphRuntime/Private/AnimNodes/AnimNode_TwoWayBlend.cpp
--- a/Engine/Source/Runtime/AnimGraphRuntime/Private/AnimNodes/AnimNode_TwoWayBlend.cpp
+++ b/Engine/Source/Runtime/AnimGraphRuntime/Private/AnimNodes/AnimNode_TwoWayBlend.cpp
@@ -3,6 +3,7 @@
 #include "AnimNodes/AnimNode_TwoWayBlend.h"
 #include "Animation/AnimInstanceProxy.h"
 #include "Animation/AnimStats.h"
+#include "Animation/AnimTrace.h"
 #include "AnimationRuntime.h"
 
 #include UE_INLINE_GENERATED_CPP_BY_NAME(AnimNode_TwoWayBlend)
diff --git a/Engine/Source/Runtime/AnimGraphRuntime/Private/BoneControllers/AnimNode_AnimDynamics.cpp b/Engine/Source/Runtime/AnimGraphRuntime/Private/BoneControllers/AnimNode_AnimDynamics.cpp
--- a/Engine/Source/Runtime/AnimGraphRuntime/Private/BoneControllers/AnimNode_AnimDynamics.cpp
+++ b/Engine/Source/Runtime/AnimGraphRuntime/Private/BoneControllers/AnimNode_AnimDynamics.cpp
@@ -1,6 +1,7 @@
 // Copyright Epic Games, Inc. All Rights Reserved.
 
 #include "BoneControllers/AnimNode_AnimDynamics.h"
+#include "Animation/AnimInstance.h"
 #include "Components/SkeletalMeshComponent.h"
 #include "GameFramework/WorldSettings.h"
 #include "Animation/AnimInstanceProxy.h"
diff --git a/Engine/Source/Runtime/AnimGraphRuntime/Private/BoneControllers/AnimNode_RigidBody.cpp b/Engine/Source/Runtime/AnimGraphRuntime/Private/BoneControllers/AnimNode_RigidBody.cpp
--- a/Engine/Source/Runtime/AnimGraphRuntime/Private/BoneControllers/AnimNode_RigidBody.cpp
+++ b/Engine/Source/Runtime/AnimGraphRuntime/Private/BoneControllers/AnimNode_RigidBody.cpp
@@ -1,6 +1,7 @@
 // Copyright Epic Games, Inc. All Rights Reserved.
 
 #include "BoneControllers/AnimNode_RigidBody.h"
+#include "Animation/AnimInstance.h"
 #include "AnimationRuntime.h"
 #include "Animation/AnimInstanceProxy.h"
 
diff --git a/Engine/Source/Runtime/AnimGraphRuntime/Private/BoneControllers/AnimNode_SkeletalControlBase.cpp b/Engine/Source/Runtime/AnimGraphRuntime/Private/BoneControllers/AnimNode_SkeletalControlBase.cpp
--- a/Engine/Source/Runtime/AnimGraphRuntime/Private/BoneControllers/AnimNode_SkeletalControlBase.cpp
+++ b/Engine/Source/Runtime/AnimGraphRuntime/Private/BoneControllers/AnimNode_SkeletalControlBase.cpp
@@ -1,6 +1,8 @@
 // Copyright Epic Games, Inc. All Rights Reserved.
 
 #include "BoneControllers/AnimNode_SkeletalControlBase.h"
+#include "Animation/AnimInstance.h"
+#include "Animation/AnimTrace.h"
 #include "Engine/SkeletalMesh.h"
 #include "Animation/AnimInstanceProxy.h"
 #include "Animation/AnimStats.h"
diff --git a/Engine/Source/Runtime/AnimGraphRuntime/Public/BoneControllers/AnimNode_AnimDynamics.h b/Engine/Source/Runtime/AnimGraphRuntime/Public/BoneControllers/AnimNode_AnimDynamics.h
--- a/Engine/Source/Runtime/AnimGraphRuntime/Public/BoneControllers/AnimNode_AnimDynamics.h
+++ b/Engine/Source/Runtime/AnimGraphRuntime/Public/BoneControllers/AnimNode_AnimDynamics.h
@@ -15,6 +15,7 @@
 
 class UAnimInstance;
 class USkeletalMeshComponent;
+template <class T> class TAutoConsoleVariable;
 
 extern TAutoConsoleVariable<int32> CVarEnableDynamics;
 extern ANIMGRAPHRUNTIME_API TAutoConsoleVariable<int32> CVarLODThreshold;
diff --git a/Engine/Source/Runtime/Engine/Classes/Components/SkinnedMeshComponent.h b/Engine/Source/Runtime/Engine/Classes/Components/SkinnedMeshComponent.h
--- a/Engine/Source/Runtime/Engine/Classes/Components/SkinnedMeshComponent.h
+++ b/Engine/Source/Runtime/Engine/Classes/Components/SkinnedMeshComponent.h
@@ -16,7 +16,6 @@
 #include "LODSyncInterface.h"
 #include "BoneContainer.h"
 #include "ClothingSystemRuntimeTypes.h"
-#include "Animation/SkinWeightProfile.h"
 #include "StateStream/SkinnedMeshStateStreamHandle.h"
 #include "SkinnedMeshComponent.generated.h"
 
@@ -2158,4 +2157,8 @@ FVector3f GetTypedSkinnedVertexPosition(
 	const FSkinWeightVertexBuffer& SkinWeightVertexBuffer,
 	const int32 VertIndex,
 	const TArray<FMatrix44f> & RefToLocals = TArray<FMatrix44f>()
-);
\ No newline at end of file
+);
+
+#if UE_ENABLE_INCLUDE_ORDER_DEPRECATED_IN_5_6
+#include "Animation/SkinWeightProfile.h"
+#endif
diff --git a/Engine/Source/Runtime/Engine/Classes/Engine/SkinnedAsset.h b/Engine/Source/Runtime/Engine/Classes/Engine/SkinnedAsset.h
--- a/Engine/Source/Runtime/Engine/Classes/Engine/SkinnedAsset.h
+++ b/Engine/Source/Runtime/Engine/Classes/Engine/SkinnedAsset.h
@@ -12,7 +12,6 @@
 #include "Interfaces/Interface_AsyncCompilation.h"
 #include "ReferenceSkeleton.h"
 #include "UObject/PerPlatformProperties.h"
-#include "PSOPrecache.h"
 #include "SkinnedAsset.generated.h"
 
 struct FSkeletalMaterial;
@@ -26,8 +25,11 @@ class UMeshDeformer;
 class UMorphTarget;
 class UPhysicsAsset;
 class USkeleton;
+struct FPSOPrecacheVertexFactoryDataPerMaterialIndex;
 struct FSkinnedAssetAsyncBuildTask;
 enum class ESkeletalMeshVertexFlags : uint8;
+typedef TArray<FPSOPrecacheVertexFactoryDataPerMaterialIndex, TInlineAllocator<4> > FPSOPrecacheVertexFactoryDataPerMaterialIndexList;
+
 enum class ESkinnedAssetAsyncPropertyLockType
 {
 	None = 0,
@@ -235,7 +237,7 @@ public:
 	ENGINE_API virtual void PostLoad() override;
 	//~ End UObject Interface
 		
-	ENGINE_API FPSOPrecacheVertexFactoryDataPerMaterialIndexList GetVertexFactoryTypesPerMaterialIndex(USkinnedMeshComponent* SkinnedMeshComponent, int32 MinLODIndex, bool bCPUSkin, ERHIFeatureLevel::Type FeatureLevel);
+	ENGINE_API void GetVertexFactoryTypesPerMaterialIndex(FPSOPrecacheVertexFactoryDataPerMaterialIndexList& OutList, USkinnedMeshComponent* SkinnedMeshComponent, int32 MinLODIndex, bool bCPUSkin, ERHIFeatureLevel::Type FeatureLevel);
 
 	/** Helper function for resource tracking, construct a string using the skinned asset's path name and LOD index . */
 	static ENGINE_API FString GetLODPathName(const USkinnedAsset* Mesh, int32 LODIndex);
@@ -353,3 +355,6 @@ private:
 	static ENGINE_API TArray<FSkeletalMaterial>& GetSkeletalMaterialDummyArray();
 };
 
+#if UE_ENABLE_INCLUDE_ORDER_DEPRECATED_IN_5_6
+#include "PSOPrecache.h"
+#endif
diff --git a/Engine/Source/Runtime/Engine/Private/Animation/AnimCompressionTypes.cpp b/Engine/Source/Runtime/Engine/Private/Animation/AnimCompressionTypes.cpp
--- a/Engine/Source/Runtime/Engine/Private/Animation/AnimCompressionTypes.cpp
+++ b/Engine/Source/Runtime/Engine/Private/Animation/AnimCompressionTypes.cpp
@@ -4,6 +4,7 @@
 #include "Animation/AnimSequence.h"
 #include "AnimationUtils.h"
 #include "Animation/AnimSequenceDecompressionContext.h"
+#include "HAL/IConsoleManager.h"
 #include "Interfaces/ITargetPlatform.h"
 #include "Animation/AnimBoneCompressionCodec.h"
 #include "Animation/AnimBoneCompressionSettings.h"
@@ -13,6 +14,7 @@
 #include "AnimationCompression.h"
 #include "AnimationRuntime.h"
 #include "BonePose.h"
+#include "Misc/CommandLine.h"
 #include "UObject/FortniteMainBranchObjectVersion.h"
 #include "Animation/AnimSequenceHelpers.h"
 #include "Animation/AnimData/AnimDataModel.h"
diff --git a/Engine/Source/Runtime/Engine/Private/Animation/AnimInstance.cpp b/Engine/Source/Runtime/Engine/Private/Animation/AnimInstance.cpp
--- a/Engine/Source/Runtime/Engine/Private/Animation/AnimInstance.cpp
+++ b/Engine/Source/Runtime/Engine/Private/Animation/AnimInstance.cpp
@@ -6,6 +6,7 @@
 
 #include "Animation/AnimInstance.h"
 #include "Animation/AnimSequence.h"
+#include "Animation/AnimTrace.h"
 #include "GameFramework/Pawn.h"
 #include "Animation/AnimStats.h"
 #include "UObject/Package.h"
diff --git a/Engine/Source/Runtime/Engine/Private/Animation/AnimInstanceProxy.cpp b/Engine/Source/Runtime/Engine/Private/Animation/AnimInstanceProxy.cpp
--- a/Engine/Source/Runtime/Engine/Private/Animation/AnimInstanceProxy.cpp
+++ b/Engine/Source/Runtime/Engine/Private/Animation/AnimInstanceProxy.cpp
@@ -5,6 +5,7 @@
 #include "Animation/AnimNodeFunctionRef.h"
 #include "Animation/AnimSlotEvaluationPose.h"
 #include "Animation/AnimNotifies/AnimNotify.h"
+#include "Animation/AnimTrace.h"
 #include "Animation/PoseAsset.h"
 #include "Animation/AnimNotifies/AnimNotifyState.h"
 #include "Animation/AnimStats.h"
@@ -3271,9 +3272,9 @@ void FAnimInstanceProxy::ResetDynamics()
 	ResetDynamics(ETeleportType::ResetPhysics);
 }
 
-#if ANIM_TRACE_ENABLED
 void FAnimInstanceProxy::TraceMontageEvaluationData(const FAnimationUpdateContext& InContext, const FName& InSlotName)
 {
+#if ANIM_TRACE_ENABLED
 	for (const FMontageEvaluationState& MontageEvaluationState : GetMontageEvaluationData())
 	{
 		if (MontageEvaluationState.Montage != nullptr && MontageEvaluationState.Montage->IsValidSlot(InSlotName))
@@ -3294,8 +3295,8 @@ void FAnimInstanceProxy::TraceMontageEvaluationData(const FAnimationUpdateContex
 			}
 		}
 	}
-}
 #endif
+}
 
 TArray<const FAnimNode_AssetPlayerBase*> FAnimInstanceProxy::GetInstanceAssetPlayers(const FName& GraphName) const
 {
diff --git a/Engine/Source/Runtime/Engine/Private/Animation/AnimNodeBase.cpp b/Engine/Source/Runtime/Engine/Private/Animation/AnimNodeBase.cpp
--- a/Engine/Source/Runtime/Engine/Private/Animation/AnimNodeBase.cpp
+++ b/Engine/Source/Runtime/Engine/Private/Animation/AnimNodeBase.cpp
@@ -5,6 +5,7 @@
 #include "Animation/AnimAttributes.h"
 #include "Animation/AnimNodeFunctionRef.h"
 #include "Animation/AnimSubsystem_Base.h"
+#include "Animation/AnimTrace.h"
 
 #include UE_INLINE_GENERATED_CPP_BY_NAME(AnimNodeBase)
 
diff --git a/Engine/Source/Runtime/Engine/Private/Animation/AnimNode_ApplyMeshSpaceAdditive.cpp b/Engine/Source/Runtime/Engine/Private/Animation/AnimNode_ApplyMeshSpaceAdditive.cpp
--- a/Engine/Source/Runtime/Engine/Private/Animation/AnimNode_ApplyMeshSpaceAdditive.cpp
+++ b/Engine/Source/Runtime/Engine/Private/Animation/AnimNode_ApplyMeshSpaceAdditive.cpp
@@ -1,7 +1,9 @@
 // Copyright Epic Games, Inc. All Rights Reserved.
 
 #include "Animation/AnimNode_ApplyMeshSpaceAdditive.h"
+#include "Animation/AnimInstance.h"
 #include "Animation/AnimInstanceProxy.h"
+#include "Animation/AnimTrace.h"
 #include "Animation/AnimationPoseData.h"
 #include "Animation/AnimStats.h"
 #include "Animation/ExposedValueHandler.h"
diff --git a/Engine/Source/Runtime/Engine/Private/Animation/AnimNode_DeadBlending.cpp b/Engine/Source/Runtime/Engine/Private/Animation/AnimNode_DeadBlending.cpp
--- a/Engine/Source/Runtime/Engine/Private/Animation/AnimNode_DeadBlending.cpp
+++ b/Engine/Source/Runtime/Engine/Private/Animation/AnimNode_DeadBlending.cpp
@@ -4,6 +4,7 @@
 #include "Animation/AnimInstanceProxy.h"
 #include "Animation/AnimRootMotionProvider.h"
 #include "Animation/AnimNode_SaveCachedPose.h"
+#include "Animation/AnimTrace.h"
 #include "Animation/BlendProfile.h"
 #include "Algo/MaxElement.h"
 #include "Components/SkeletalMeshComponent.h"
diff --git a/Engine/Source/Runtime/Engine/Private/Animation/AnimNode_Inertialization.cpp b/Engine/Source/Runtime/Engine/Private/Animation/AnimNode_Inertialization.cpp
--- a/Engine/Source/Runtime/Engine/Private/Animation/AnimNode_Inertialization.cpp
+++ b/Engine/Source/Runtime/Engine/Private/Animation/AnimNode_Inertialization.cpp
@@ -5,6 +5,7 @@
 #include "Animation/AnimRootMotionProvider.h"
 #include "Animation/AnimNode_SaveCachedPose.h"
 #include "Animation/AnimStats.h"
+#include "Animation/AnimTrace.h"
 #include "Animation/BlendProfile.h"
 #include "Algo/MaxElement.h"
 #include "Components/SkeletalMeshComponent.h"
diff --git a/Engine/Source/Runtime/Engine/Private/Animation/AnimNode_LinkedAnimGraph.cpp b/Engine/Source/Runtime/Engine/Private/Animation/AnimNode_LinkedAnimGraph.cpp
--- a/Engine/Source/Runtime/Engine/Private/Animation/AnimNode_LinkedAnimGraph.cpp
+++ b/Engine/Source/Runtime/Engine/Private/Animation/AnimNode_LinkedAnimGraph.cpp
@@ -5,6 +5,7 @@
 #include "Animation/AnimNode_Inertialization.h"
 #include "Animation/AnimNode_LinkedInputPose.h"
 #include "Animation/AnimNode_Root.h"
+#include "Animation/AnimTrace.h"
 #include "Animation/BlendProfile.h"
 #include "Components/SkeletalMeshComponent.h"
 #include "Animation/ExposedValueHandler.h"
diff --git a/Engine/Source/Runtime/Engine/Private/Animation/AnimNode_SaveCachedPose.cpp b/Engine/Source/Runtime/Engine/Private/Animation/AnimNode_SaveCachedPose.cpp
--- a/Engine/Source/Runtime/Engine/Private/Animation/AnimNode_SaveCachedPose.cpp
+++ b/Engine/Source/Runtime/Engine/Private/Animation/AnimNode_SaveCachedPose.cpp
@@ -3,6 +3,7 @@
 #include "Animation/AnimNode_SaveCachedPose.h"
 #include "Animation/AnimStats.h"
 #include "Animation/AnimInstanceProxy.h"
+#include "Animation/AnimTrace.h"
 
 #include UE_INLINE_GENERATED_CPP_BY_NAME(AnimNode_SaveCachedPose)
 
diff --git a/Engine/Source/Runtime/Engine/Private/Animation/AnimNode_SequencePlayer.cpp b/Engine/Source/Runtime/Engine/Private/Animation/AnimNode_SequencePlayer.cpp
--- a/Engine/Source/Runtime/Engine/Private/Animation/AnimNode_SequencePlayer.cpp
+++ b/Engine/Source/Runtime/Engine/Private/Animation/AnimNode_SequencePlayer.cpp
@@ -4,6 +4,7 @@
 
 #include "Animation/AnimMontage.h"
 #include "Animation/AnimInstanceProxy.h"
+#include "Animation/AnimTrace.h"
 #include "Animation/AnimationPoseData.h"
 #include "Animation/AnimPoseSearchProvider.h"
 #include "Animation/ExposedValueHandler.h"
diff --git a/Engine/Source/Runtime/Engine/Private/Animation/AnimNode_StateMachine.cpp b/Engine/Source/Runtime/Engine/Private/Animation/AnimNode_StateMachine.cpp
--- a/Engine/Source/Runtime/Engine/Private/Animation/AnimNode_StateMachine.cpp
+++ b/Engine/Source/Runtime/Engine/Private/Animation/AnimNode_StateMachine.cpp
@@ -7,6 +7,7 @@
 #include "Animation/AnimNode_TransitionPoseEvaluator.h"
 #include "Animation/AnimNode_Inertialization.h"
 #include "Animation/AnimStats.h"
+#include "Animation/AnimTrace.h"
 #include "Animation/BlendProfile.h"
 #include "Animation/AnimNode_LinkedAnimLayer.h"
 #include "Animation/ActiveStateMachineScope.h"
diff --git a/Engine/Source/Runtime/Engine/Private/Animation/AnimNode_TransitionPoseEvaluator.cpp b/Engine/Source/Runtime/Engine/Private/Animation/AnimNode_TransitionPoseEvaluator.cpp
--- a/Engine/Source/Runtime/Engine/Private/Animation/AnimNode_TransitionPoseEvaluator.cpp
+++ b/Engine/Source/Runtime/Engine/Private/Animation/AnimNode_TransitionPoseEvaluator.cpp
@@ -2,6 +2,7 @@
 
 #include "Animation/AnimNode_TransitionPoseEvaluator.h"
 #include "Animation/AnimInstanceProxy.h"
+#include "Animation/AnimTrace.h"
 
 #include UE_INLINE_GENERATED_CPP_BY_NAME(AnimNode_TransitionPoseEvaluator)
 
diff --git a/Engine/Source/Runtime/Engine/Private/Animation/AnimSingleNodeInstanceProxy.cpp b/Engine/Source/Runtime/Engine/Private/Animation/AnimSingleNodeInstanceProxy.cpp
--- a/Engine/Source/Runtime/Engine/Private/Animation/AnimSingleNodeInstanceProxy.cpp
+++ b/Engine/Source/Runtime/Engine/Private/Animation/AnimSingleNodeInstanceProxy.cpp
@@ -4,6 +4,7 @@
 #include "Animation/AnimMontageEvaluationState.h"
 #include "Animation/AnimComposite.h"
 #include "Animation/AnimSequence.h"
+#include "Animation/AnimTrace.h"
 #include "Animation/BlendSpace.h"
 #include "Animation/PoseAsset.h"
 #include "Animation/AnimStreamable.h"
diff --git a/Engine/Source/Runtime/Engine/Private/Animation/AnimSubsystem_PropertyAccess.cpp b/Engine/Source/Runtime/Engine/Private/Animation/AnimSubsystem_PropertyAccess.cpp
--- a/Engine/Source/Runtime/Engine/Private/Animation/AnimSubsystem_PropertyAccess.cpp
+++ b/Engine/Source/Runtime/Engine/Private/Animation/AnimSubsystem_PropertyAccess.cpp
@@ -1,6 +1,7 @@
 // Copyright Epic Games, Inc. All Rights Reserved.
 
 #include "Animation/AnimSubsystem_PropertyAccess.h"
+#include "Animation/AnimInstance.h"
 #include "Animation/AnimInstanceProxy.h"
 
 #include UE_INLINE_GENERATED_CPP_BY_NAME(AnimSubsystem_PropertyAccess)
diff --git a/Engine/Source/Runtime/Engine/Private/Animation/AnimSync.cpp b/Engine/Source/Runtime/Engine/Private/Animation/AnimSync.cpp
--- a/Engine/Source/Runtime/Engine/Private/Animation/AnimSync.cpp
+++ b/Engine/Source/Runtime/Engine/Private/Animation/AnimSync.cpp
@@ -3,6 +3,7 @@
 #include "Animation/AnimSync.h"
 #include "Animation/AnimInstanceProxy.h"
 #include "Animation/AnimStats.h"
+#include "Animation/AnimTrace.h"
 #include "AnimationUtils.h"
 #include "Animation/MirrorDataTable.h"
 #include "AnimationRuntime.h"
diff --git a/Engine/Source/Runtime/Engine/Private/Animation/AnimSyncScope.cpp b/Engine/Source/Runtime/Engine/Private/Animation/AnimSyncScope.cpp
--- a/Engine/Source/Runtime/Engine/Private/Animation/AnimSyncScope.cpp
+++ b/Engine/Source/Runtime/Engine/Private/Animation/AnimSyncScope.cpp
@@ -3,6 +3,7 @@
 #include "Animation/AnimSyncScope.h"
 #include "Animation/AnimInstanceProxy.h"
 #include "Animation/AnimNodeBase.h"
+#include "Animation/AnimTrace.h"
 
 IMPLEMENT_ANIMGRAPH_MESSAGE(UE::Anim::FAnimSyncGroupScope);
 
diff --git a/Engine/Source/Runtime/Engine/Private/Animation/BonePose.cpp b/Engine/Source/Runtime/Engine/Private/Animation/BonePose.cpp
--- a/Engine/Source/Runtime/Engine/Private/Animation/BonePose.cpp
+++ b/Engine/Source/Runtime/Engine/Private/Animation/BonePose.cpp
@@ -6,6 +6,8 @@
 #include "AnimEncoding.h"
 #include "Animation/AnimSequenceHelpers.h"
 #include "AutoRTFM/AutoRTFM.h"
+#include "HAL/IConsoleManager.h"
+
 #if INTEL_ISPC
 #include "BonePose.ispc.generated.h"
 
diff --git a/Engine/Source/Runtime/Engine/Private/Animation/Skeleton.cpp b/Engine/Source/Runtime/Engine/Private/Animation/Skeleton.cpp
--- a/Engine/Source/Runtime/Engine/Private/Animation/Skeleton.cpp
+++ b/Engine/Source/Runtime/Engine/Private/Animation/Skeleton.cpp
@@ -8,6 +8,7 @@
 
 #include "AnimationSequenceCompiler.h"
 #include "Animation/AnimData/AnimDataModel.h"
+#include "Components/SkinnedMeshComponent.h"
 #include "UObject/AssetRegistryTagsContext.h"
 #include "UObject/LinkerLoad.h"
 #include "Engine/AssetUserData.h"
diff --git a/Engine/Source/Runtime/Engine/Private/Components/SkinnedMeshComponent.cpp b/Engine/Source/Runtime/Engine/Private/Components/SkinnedMeshComponent.cpp
--- a/Engine/Source/Runtime/Engine/Private/Components/SkinnedMeshComponent.cpp
+++ b/Engine/Source/Runtime/Engine/Private/Components/SkinnedMeshComponent.cpp
@@ -650,7 +650,8 @@ void USkinnedMeshComponent::CollectPSOPrecacheData(const FPSOPrecacheParams& Bas
 	int32 MinLODIndex = ComputeMinLOD();
 	bool bCPUSkin = bRenderStatic || ShouldCPUSkin();
 
-	FPSOPrecacheVertexFactoryDataPerMaterialIndexList VFsPerMaterials = GetSkinnedAsset()->GetVertexFactoryTypesPerMaterialIndex(this, MinLODIndex, bCPUSkin, FeatureLevel);
+	FPSOPrecacheVertexFactoryDataPerMaterialIndexList VFsPerMaterials;
+	GetSkinnedAsset()->GetVertexFactoryTypesPerMaterialIndex(VFsPerMaterials, this, MinLODIndex, bCPUSkin, FeatureLevel);
 	bool bAnySectionCastsShadows = GetSkinnedAsset()->GetResourceForRendering()->AnyRenderSectionCastsShadows(MinLODIndex);
 
 	FPSOPrecacheParams PrecachePSOParams = BasePrecachePSOParams;
diff --git a/Engine/Source/Runtime/Engine/Private/SkeletalMeshSceneProxy.cpp b/Engine/Source/Runtime/Engine/Private/SkeletalMeshSceneProxy.cpp
--- a/Engine/Source/Runtime/Engine/Private/SkeletalMeshSceneProxy.cpp
+++ b/Engine/Source/Runtime/Engine/Private/SkeletalMeshSceneProxy.cpp
@@ -5,6 +5,7 @@
 #include "Engine/SkeletalMesh.h"
 #include "EngineUtils.h"
 #include "MaterialDomain.h"
+#include "MaterialShared.h"
 #include "Materials/Material.h"
 #include "PhysicsEngine/PhysicsAsset.h"
 #include "PhysicsEngine/SkeletalBodySetup.h"
diff --git a/Engine/Source/Runtime/Engine/Private/SkinnedAsset.cpp b/Engine/Source/Runtime/Engine/Private/SkinnedAsset.cpp
--- a/Engine/Source/Runtime/Engine/Private/SkinnedAsset.cpp
+++ b/Engine/Source/Runtime/Engine/Private/SkinnedAsset.cpp
@@ -4,6 +4,7 @@
 #include "Engine/SkinnedAssetAsyncCompileUtils.h"
 #include "Engine/SkinnedAssetCommon.h"
 #include "Engine/World.h"
+#include "PSOPrecache.h"
 #include "SkinnedAssetCompiler.h"
 #include "Animation/AnimStats.h"
 #include "Materials/MaterialInterface.h"
@@ -152,7 +153,8 @@ void USkinnedAsset::PostLoad()
 		ERHIFeatureLevel::Type FeatureLevel = GetWorld() ? GetWorld()->GetFeatureLevel() : GMaxRHIFeatureLevel;
 		int32 MinLODIndex = GetMinLodIdx();
 		
-		FPSOPrecacheVertexFactoryDataPerMaterialIndexList VFsPerMaterials = GetVertexFactoryTypesPerMaterialIndex(nullptr, MinLODIndex, bCPUSkin, FeatureLevel);
+		FPSOPrecacheVertexFactoryDataPerMaterialIndexList VFsPerMaterials;
+		GetVertexFactoryTypesPerMaterialIndex(VFsPerMaterials, nullptr, MinLODIndex, bCPUSkin, FeatureLevel);
 		bool bAnySectionCastsShadows = GetResourceForRendering()->AnyRenderSectionCastsShadows(MinLODIndex);
 
 		// Precache using default material & PSO precache params but mark movable
@@ -176,11 +178,9 @@ void USkinnedAsset::PostLoad()
 // TODO: move this to a StaticMeshResources.h?
 extern void InitStaticMeshVertexFactoryComponents(const FStaticMeshVertexBuffers& VertexBuffers, FLocalVertexFactory* VertexFactory, int32 LightMapCoordinateIndex, bool bOverrideColorVertexBuffer, FLocalVertexFactory::FDataType& OutData);
 
-FPSOPrecacheVertexFactoryDataPerMaterialIndexList USkinnedAsset::GetVertexFactoryTypesPerMaterialIndex(
+void USkinnedAsset::GetVertexFactoryTypesPerMaterialIndex(FPSOPrecacheVertexFactoryDataPerMaterialIndexList& OutList,
 	USkinnedMeshComponent* SkinnedMeshComponent, int32 MinLODIndex, bool bCPUSkin, ERHIFeatureLevel::Type FeatureLevel)
 {
-	FPSOPrecacheVertexFactoryDataPerMaterialIndexList VFTypesPerMaterialIndex;
-
 	FSkeletalMeshRenderData* SkelMeshRenderData = GetResourceForRendering();
 	for (int32 LODIndex = MinLODIndex; LODIndex < SkelMeshRenderData->LODRenderData.Num(); LODIndex++)
 	{
@@ -202,11 +202,11 @@ FPSOPrecacheVertexFactoryDataPerMaterialIndexList USkinnedAsset::GetVertexFactor
 				MaterialIndex = FMath::Clamp(MaterialIndex, 0, GetNumMaterials());
 			}
 
-			FPSOPrecacheVertexFactoryDataPerMaterialIndex* VFsPerMaterial = VFTypesPerMaterialIndex.FindByPredicate(
+			FPSOPrecacheVertexFactoryDataPerMaterialIndex* VFsPerMaterial = OutList.FindByPredicate(
 				[MaterialIndex](const FPSOPrecacheVertexFactoryDataPerMaterialIndex& Other) { return Other.MaterialIndex == MaterialIndex; });
 			if (VFsPerMaterial == nullptr)
 			{
-				VFsPerMaterial = &VFTypesPerMaterialIndex.AddDefaulted_GetRef();
+				VFsPerMaterial = &OutList.AddDefaulted_GetRef();
 				VFsPerMaterial->MaterialIndex = MaterialIndex;
 			}
 
@@ -238,8 +238,6 @@ FPSOPrecacheVertexFactoryDataPerMaterialIndexList USkinnedAsset::GetVertexFactor
 			}
 		}
 	}
-
-	return VFTypesPerMaterialIndex;
 }
 
 #if WITH_EDITOR
diff --git a/Engine/Source/Runtime/Engine/Public/Animation/AnimInstanceProxy.h b/Engine/Source/Runtime/Engine/Public/Animation/AnimInstanceProxy.h
--- a/Engine/Source/Runtime/Engine/Public/Animation/AnimInstanceProxy.h
+++ b/Engine/Source/Runtime/Engine/Public/Animation/AnimInstanceProxy.h
@@ -9,7 +9,6 @@
 #include "Animation/AnimationAsset.h"
 #include "Animation/AnimBlueprint.h"
 #include "Animation/AnimClassInterface.h"
-#include "Animation/AnimTrace.h"
 #include "Animation/AnimSync.h"
 #include "Animation/AnimInertializationRequest.h"
 
@@ -20,6 +19,7 @@ class UBlendProfile;
 class UBlendSpace;
 class FTokenizedMessage;
 namespace EMessageSeverity { enum Type : int; }
+struct FAnimationEvaluationContext;
 struct FAnimationPoseData;
 struct FAnimBlueprintDebugData_NodeVisit;
 struct FAnimBlueprintDebugData_AttributeRecord;
@@ -34,7 +34,10 @@ struct FAnimNode_StateMachine;
 struct FAnimTickRecord;
 struct FBakedAnimationStateMachine;
 struct FCompactPose;
+struct FMontageActiveSlotTracker;
 struct FMontageEvaluationState;
+struct FNativeStateBinding;
+struct FNativeTransitionBinding;
 struct FNodeDebugData;
 struct FPoseContext;
 struct FPoseSnapshot;
@@ -63,6 +66,9 @@ extern const FName NAME_Evaluate;
 extern const FName NAME_Update;
 extern const FName NAME_AnimGraph;
 
+DECLARE_DELEGATE_RetVal(bool, FCanTakeTransition);
+DECLARE_DELEGATE_ThreeParams(FOnGraphStateChanged, const struct FAnimNode_StateMachine&, int32, int32);
+
 UENUM()
 namespace EDrawDebugItemType
 {
@@ -549,10 +555,8 @@ public:
 	/** Get the transform of the actor we are running on */
 	const FTransform& GetActorTransform() const { return ActorTransform; }
 
-#if ANIM_TRACE_ENABLED
 	// Trace montage debug data for the specified slot
 	ENGINE_API void TraceMontageEvaluationData(const FAnimationUpdateContext& InContext, const FName& InSlotName);
-#endif
 
 	/** Get the debug data for this instance's anim bp */
 	ENGINE_API FAnimBlueprintDebugData* GetAnimBlueprintDebugData() const;
@@ -1196,3 +1200,7 @@ private:
 
 	uint8 bUseMainInstanceMontageEvaluationData : 1;
 };
+
+#if UE_ENABLE_INCLUDE_ORDER_DEPRECATED_IN_5_6
+#include "Animation/AnimTrace.h"
+#endif
diff --git a/Engine/Source/Runtime/Engine/Public/AnimationRuntime.h b/Engine/Source/Runtime/Engine/Public/AnimationRuntime.h
--- a/Engine/Source/Runtime/Engine/Public/AnimationRuntime.h
+++ b/Engine/Source/Runtime/Engine/Public/AnimationRuntime.h
@@ -15,17 +15,18 @@
 #include "Animation/AnimCurveTypes.h"
 #include "Animation/AnimSequenceBase.h"
 #include "Animation/SkeletonRemapping.h"
-#include "Components/SkinnedMeshComponent.h"
 #include "BonePose.h"
 #include "Containers/ArrayView.h"
 
 class UBlendSpace;
 class USkeletalMeshComponent;
 class UMirrorDataTable;
+class UMorphTarget;
 struct FA2CSPose;
 struct FA2Pose;
 struct FInputBlendPose;
 struct FAnimationPoseData;
+using FMorphTargetWeightMap = TMap<const UMorphTarget*, int32>;
 
 namespace UE { namespace Anim { struct FStackAttributeContainer; }}
 
@@ -655,3 +656,7 @@ public:
 };
 
 ENUM_CLASS_FLAGS(FAnimationRuntime::EBlendPosesPerBoneFilterFlags);
+
+#if UE_ENABLE_INCLUDE_ORDER_DEPRECATED_IN_5_6
+#include "Components/SkinnedMeshComponent.h"
+#endif
diff --git a/Engine/Source/Runtime/Engine/Public/DynamicMeshBuilder.h b/Engine/Source/Runtime/Engine/Public/DynamicMeshBuilder.h
--- a/Engine/Source/Runtime/Engine/Public/DynamicMeshBuilder.h
+++ b/Engine/Source/Runtime/Engine/Public/DynamicMeshBuilder.h
@@ -6,11 +6,12 @@
 
 #pragma once
 
+#include "Components.h"
 #include "CoreMinimal.h"
 #include "PackedNormal.h"
 #include "HitProxies.h"
+#include "RenderResource.h"
 #include "RenderUtils.h"
-#include "LocalVertexFactory.h"
 #include "RenderMath.h"
 
 class FMaterialRenderProxy;
@@ -237,3 +238,7 @@ public:
 
 	ENGINE_API virtual void InitRHI(FRHICommandListBase& RHICmdList) override;
 };
+
+#if UE_ENABLE_INCLUDE_ORDER_DEPRECATED_IN_5_6
+#include "LocalVertexFactory.h"
+#endif
diff --git a/Engine/Source/Runtime/Engine/Public/Materials/MaterialInterface.h b/Engine/Source/Runtime/Engine/Public/Materials/MaterialInterface.h
--- a/Engine/Source/Runtime/Engine/Public/Materials/MaterialInterface.h
+++ b/Engine/Source/Runtime/Engine/Public/Materials/MaterialInterface.h
@@ -20,7 +20,6 @@
 #include "SceneTypes.h"
 #include "Engine/BlendableInterface.h"
 #include "Materials/MaterialLayersFunctions.h"
-#include "Materials/MaterialIRModule.h"
 #include "Interfaces/Interface_AssetUserData.h"
 #include "MaterialSceneTextureId.h"
 #include "Materials/MaterialRelevance.h"
@@ -1275,3 +1274,7 @@ ENGINE_API FMaterialResource* FindOrCreateMaterialResource(TArray<FMaterialResou
 	UMaterialInstance* OwnerMaterialInstance,
 	ERHIFeatureLevel::Type InFeatureLevel,
 	EMaterialQualityLevel::Type QualityLevel);
+
+#if UE_ENABLE_INCLUDE_ORDER_DEPRECATED_IN_5_6
+#include "Materials/MaterialIRModule.h"
+#endif
diff --git a/Engine/Source/Runtime/Landscape/Private/LandscapeEditLayers.cpp b/Engine/Source/Runtime/Landscape/Private/LandscapeEditLayers.cpp
--- a/Engine/Source/Runtime/Landscape/Private/LandscapeEditLayers.cpp
+++ b/Engine/Source/Runtime/Landscape/Private/LandscapeEditLayers.cpp
@@ -29,6 +29,7 @@ LandscapeEditLayers.cpp: Landscape editing layers mode
 #include "Application/SlateApplicationBase.h"
 #include "Shader.h"
 #include "GlobalShader.h"
+#include "ShaderCompilerCore.h"
 #include "ShaderParameterUtils.h"
 #include "EngineModule.h"
 #include "Engine/TextureRenderTarget2D.h"
diff --git a/Engine/Source/Runtime/UMG/Private/Components/WidgetComponent.cpp b/Engine/Source/Runtime/UMG/Private/Components/WidgetComponent.cpp
--- a/Engine/Source/Runtime/UMG/Private/Components/WidgetComponent.cpp
+++ b/Engine/Source/Runtime/UMG/Private/Components/WidgetComponent.cpp
@@ -4,6 +4,7 @@
 #include "Engine/GameInstance.h"
 #include "GameFramework/Pawn.h"
 #include "GameFramework/PlayerController.h"
+#include "LocalVertexFactory.h"
 #include "Materials/Material.h"
 #include "PrimitiveViewRelevance.h"
 #include "PrimitiveSceneProxy.h"
