[HWRT] Added checks to detect requests to build FRayTracingGeometry while a previous request has not been processed yet.
- also shifted responsibility of updating RayTracingBuildRequestIndex from FRayTracingGeometry to FRayTracingGeometryManager since the manager is also responsible for resetting it back to INDEX_NONE when the build is kicked off.

diff --git a/Engine/Source/Runtime/Engine/Private/Rendering/RayTracingGeometryManager.cpp b/Engine/Source/Runtime/Engine/Private/Rendering/RayTracingGeometryManager.cpp
--- a/Engine/Source/Runtime/Engine/Private/Rendering/RayTracingGeometryManager.cpp
+++ b/Engine/Source/Runtime/Engine/Private/Rendering/RayTracingGeometryManager.cpp
@@ -169,6 +169,7 @@ static float GetInitialBuildPriority(ERTAccelerationStructureBuildPriority InBui
 
 FRayTracingGeometryManager::BuildRequestIndex FRayTracingGeometryManager::RequestBuildAccelerationStructure(FRayTracingGeometry* InGeometry, ERTAccelerationStructureBuildPriority InPriority, EAccelerationStructureBuildMode InBuildMode)
 {
+	check(InGeometry->RayTracingBuildRequestIndex == INDEX_NONE);
 	FBuildRequest Request;
 	Request.BuildPriority = GetInitialBuildPriority(InPriority);
 	Request.Owner = InGeometry;
@@ -181,6 +182,8 @@ FRayTracingGeometryManager::BuildRequestIndex FRayTracingGeometryManager::Reques
 	INC_DWORD_STAT(STAT_RayTracingPendingBuilds);
 	INC_DWORD_STAT_BY(STAT_RayTracingPendingBuildPrimitives, InGeometry->Initializer.TotalPrimitiveCount);
 
+	InGeometry->RayTracingBuildRequestIndex = RequestIndex;
+	
 	return RequestIndex;
 }
 
@@ -1184,7 +1187,7 @@ void FRayTracingGeometryManager::ProcessBuildRequests(FRHIComputeCommandList& In
 
 void FRayTracingGeometryManager::SetupBuildParams(const FBuildRequest& InBuildRequest, TArray<FRayTracingGeometryBuildParams>& InBuildParams, bool bRemoveFromRequestArray)
 {
-	check(InBuildRequest.RequestIndex != INDEX_NONE && InBuildRequest.Owner->RayTracingBuildRequestIndex != INDEX_NONE);
+	check(InBuildRequest.RequestIndex != INDEX_NONE && InBuildRequest.RequestIndex == InBuildRequest.Owner->RayTracingBuildRequestIndex);
 	checkf(InBuildRequest.Owner->GetRHI() != nullptr, TEXT("Build request for FRayTracingGeometry without valid RHI. Was the FRayTracingGeometry evicted or released without calling RemoveBuildRequest()?"));
 
 	FRayTracingGeometryBuildParams BuildParam;
diff --git a/Engine/Source/Runtime/RenderCore/Private/RayTracingGeometry.cpp b/Engine/Source/Runtime/RenderCore/Private/RayTracingGeometry.cpp
--- a/Engine/Source/Runtime/RenderCore/Private/RayTracingGeometry.cpp
+++ b/Engine/Source/Runtime/RenderCore/Private/RayTracingGeometry.cpp
@@ -73,7 +73,7 @@ void FRayTracingGeometry::RequestBuildIfNeeded(FRHICommandListBase& RHICmdList,
 {
 	if (GetRequiresBuild())
 	{
-		RayTracingBuildRequestIndex = GRayTracingGeometryManager->RequestBuildAccelerationStructure(this, InBuildPriority);
+		GRayTracingGeometryManager->RequestBuildAccelerationStructure(this, InBuildPriority);
 		SetRequiresBuild(false);
 	}
 }
@@ -174,7 +174,7 @@ void FRayTracingGeometry::CreateRayTracingGeometry(FRHICommandListBase& RHICmdLi
 			{
 				if (RayTracingGeometryRHI)
 				{
-					RayTracingBuildRequestIndex = GRayTracingGeometryManager->RequestBuildAccelerationStructure(this, InBuildPriority);
+					GRayTracingGeometryManager->RequestBuildAccelerationStructure(this, InBuildPriority);
 				}
 				SetRequiresBuild(false);
 			}
@@ -187,7 +187,7 @@ void FRayTracingGeometry::CreateRayTracingGeometry(FRHICommandListBase& RHICmdLi
 		{
 			if (RayTracingGeometryRHI && RayTracingGeometryRHI->IsCompressed())
 			{
-				RayTracingBuildRequestIndex = GRayTracingGeometryManager->RequestBuildAccelerationStructure(this, InBuildPriority);
+				GRayTracingGeometryManager->RequestBuildAccelerationStructure(this, InBuildPriority);
 			}
 
 			SetRequiresBuild(false);
