Fix for duplicated  LOCTEXT in AnimRetargetHelper

#rnx

diff --git a/Engine/Source/Runtime/Engine/Private/Animation/AnimRetargetHelpers.cpp b/Engine/Source/Runtime/Engine/Private/Animation/AnimRetargetHelpers.cpp
--- a/Engine/Source/Runtime/Engine/Private/Animation/AnimRetargetHelpers.cpp
+++ b/Engine/Source/Runtime/Engine/Private/Animation/AnimRetargetHelpers.cpp
@@ -140,7 +140,7 @@ void CheckRetargetSourceAssetData(bool bFixAssets, bool bWantsFullScan, const TA
 		FAssetData& AssetData = Assets[Idx];
 
 		Feedback.EnterProgressFrame(1, 
-			FText::Format(LOCTEXT("ScanningRetargetSourceAssets", "Assets Left [{0}] - {1}"), 
+			FText::Format(LOCTEXT("ScanningRetargetSourceAssetsProgress", "Assets Left [{0}] - {1}"), 
 				NumAssets - Idx,
 				FText::FromString((AssetData.PackagePath.ToString() + FString(TEXT("/")) + AssetData.AssetName.ToString()))));
 
