[Engine]
* Fixed non unity compile error in AnimNodeAlphaOptions.h

diff --git a/Engine/Source/Runtime/Engine/Classes/Animation/AnimNodeAlphaOptions.h b/Engine/Source/Runtime/Engine/Classes/Animation/AnimNodeAlphaOptions.h
--- a/Engine/Source/Runtime/Engine/Classes/Animation/AnimNodeAlphaOptions.h
+++ b/Engine/Source/Runtime/Engine/Classes/Animation/AnimNodeAlphaOptions.h
@@ -2,6 +2,7 @@
 
 #pragma once
 
+#include "Animation/AnimInstance.h"
 #include "Animation/AnimInstanceProxy.h"
 #include "Animation/AnimNodeBase.h"
 #include "Animation/InputScaleBias.h"
