Use structured logging for a couple of AnimMontage PostLoad messages.

diff --git a/Engine/Source/Runtime/Engine/Private/Animation/AnimMontage.cpp b/Engine/Source/Runtime/Engine/Private/Animation/AnimMontage.cpp
--- a/Engine/Source/Runtime/Engine/Private/Animation/AnimMontage.cpp
+++ b/Engine/Source/Runtime/Engine/Private/Animation/AnimMontage.cpp
@@ -21,6 +21,7 @@
 #include "Animation/AnimTrace.h"
 #include "Animation/ActiveMontageInstanceScope.h"
 #include "Animation/AnimationSettings.h"
+#include "Logging/StructuredLog.h"
 
 #include UE_INLINE_GENERATED_CPP_BY_NAME(AnimMontage)
 
@@ -427,7 +428,9 @@ void UAnimMontage::PostLoad()
 		const float CurrentCalculatedLength = CalculateSequenceLength();
 		if(!FMath::IsNearlyEqual(CurrentCalculatedLength, GetPlayLength(), UE_KINDA_SMALL_NUMBER))		
 		{
-			UE_LOG(LogAnimMontage, Display, TEXT("UAnimMontage::PostLoad: The actual sequence length for %s does not match the length stored in the asset, please resave the asset."), *GetFullName());
+			UE_LOGFMT_NSLOC(LogAnimMontage, Display, "AnimMontage", "PostLoad_SequenceLengthMismatch",
+				"UAnimMontage::PostLoad: The actual sequence length for {Asset} does not match the length stored in the asset. Please resave the asset.",
+				("Asset", UE::FAssetLog(this)));
 			SetCompositeLength(CurrentCalculatedLength);
 		}
 
@@ -436,8 +439,14 @@ void UAnimMontage::PostLoad()
 		{
 			if(AnimSegment.IsPlayLengthOutOfDate())
 			{
-				UE_LOG(LogAnimation, Warning, TEXT("AnimMontage (%s) contains a Segment for Slot (%s) for which the playable length %f is out-of-sync with the represented AnimationSequence its length %f (%s). Please up-date the segment and resave."), *GetFullName(), *SlotIter->SlotName.ToString(), (AnimSegment.AnimEndTime - AnimSegment.AnimStartTime), AnimSegment.GetAnimReference()->GetPlayLength(),
-					*AnimSegment.GetAnimReference()->GetFullName());
+				UE_LOGFMT_NSLOC(LogAnimation, Warning, "AnimMontage", "PostLoad_InvalidSectionLength",
+					"AnimMontage ({MontageAsset}) contains a Segment for Slot ({SlotName}) for which the playable length {PlayableLength} "
+					"is out-of-sync with the represented AnimationSequence its length {AnimSequenceLength} ({AnimSequenceAsset}). Please up-date the segment and resave.",
+					("MontageAsset", UE::FAssetLog(this)),
+					("SlotName", SlotIter->SlotName),
+					("PlayableLength", AnimSegment.AnimEndTime - AnimSegment.AnimStartTime),
+					("AnimSequenceLength", AnimSegment.GetAnimReference()->GetPlayLength()),
+					("AnimSequenceAsset", UE::FAssetLog(AnimSegment.GetAnimReference())));
 			}
 		}
 #endif
