Moving CleanedUpState flag reset to UWorld::InitWorld() to accomodate more cases of World reuse
Reset/Restore CleanedupState flag when components need to be registered during save since it's a supported flow

diff --git a/Engine/Source/Editor/UnrealEd/Private/EditorWorldUtils.cpp b/Engine/Source/Editor/UnrealEd/Private/EditorWorldUtils.cpp
--- a/Engine/Source/Editor/UnrealEd/Private/EditorWorldUtils.cpp
+++ b/Engine/Source/Editor/UnrealEd/Private/EditorWorldUtils.cpp
@@ -61,7 +61,6 @@ void FScopedEditorWorld::Init(UWorld* InWorld, const UWorld::InitializationValue
 
 	// Initialize the world
 	World->WorldType = InWorldType;
-	World->ResetCleanedUpState(); // We reset the cleanedup state to be able to reuse the world
 	World->InitWorld(InInitializationValues);
 	World->PersistentLevel->UpdateModelComponents();
 	World->UpdateWorldComponents(true /*bRerunConstructionScripts*/, false /*bCurrentLevelOnly*/);
diff --git a/Engine/Source/Runtime/Engine/Private/World.cpp b/Engine/Source/Runtime/Engine/Private/World.cpp
--- a/Engine/Source/Runtime/Engine/Private/World.cpp
+++ b/Engine/Source/Runtime/Engine/Private/World.cpp
@@ -1584,6 +1584,11 @@ void UWorld::PreSaveRoot(FObjectPreSaveRootContext ObjectSaveContext)
 	// Update components and keep track off whether we need to clean them up afterwards.
 	if(!PersistentLevel->bAreComponentsCurrentlyRegistered)
 	{
+#if WITH_EDITOR
+		// Even if the World has been CleanedUp we find ourselves in a state where we'll add the components back again (and clean it up afterwards), so we reset the CleanedUp state to allow registration of components to occur.
+		check((bIsWorldInitialized && !IsCleanedUp()) || (IsCleanedUp() && !bIsWorldInitialized)); // since we'll be using bIsWorldInitialized to ensure we restore the CleanedUp state properly we check we can rely on it's value
+		ResetCleanedUpState();
+#endif
 		PersistentLevel->UpdateLevelComponents(true);
 		ObjectSaveContext.SetCleanupRequired(true);
 	}
@@ -1600,6 +1605,13 @@ void UWorld::PostSaveRoot( FObjectPostSaveRootContext ObjectSaveContext )
 	if( ObjectSaveContext.IsCleanupRequired() )
 	{
 		PersistentLevel->ClearLevelComponents();
+#if WITH_EDITOR
+		if (!bIsWorldInitialized)
+		{
+			// If we resetted the CleanUp state to accommodate save we restore it here. 
+			CleanupWorldTag = ++CleanupWorldGlobalTag;
+		}
+#endif
 	}
 
 #if WITH_EDITOR
@@ -2128,6 +2140,9 @@ void UWorld::InitWorld(const InitializationValues IVS)
 	// Reset flags in case of world reuse
 	bIsLevelStreamingFrozen = false;
 	bShouldForceUnloadStreamingLevels = false;
+#if WITH_EDITOR
+	ResetCleanedUpState();
+#endif
 
 	FCoreUObjectDelegates::GetPostGarbageCollect().AddUObject(this, &UWorld::OnPostGC);
 
