Level Instance: backout the change that made level instances editor-only as it broke dependencies gathering.
Resaved all actors that were marked as editor-only referencing cooked worlds.

diff --git a/Engine/Source/Runtime/Engine/Private/LevelInstance/LevelInstanceActor.cpp b/Engine/Source/Runtime/Engine/Private/LevelInstance/LevelInstanceActor.cpp
--- a/Engine/Source/Runtime/Engine/Private/LevelInstance/LevelInstanceActor.cpp
+++ b/Engine/Source/Runtime/Engine/Private/LevelInstance/LevelInstanceActor.cpp
@@ -28,7 +28,6 @@ ALevelInstance::ALevelInstance()
 	RootComponent->Mobility = EComponentMobility::Static;
 
 #if WITH_EDITORONLY_DATA
-	bIsEditorOnlyActor = true;
 	DesiredRuntimeBehavior = ELevelInstanceRuntimeBehavior::Partitioned;
 #endif
 }
@@ -378,21 +377,6 @@ bool ALevelInstance::EditorCanAttachFrom(const AActor* InChild, FText& OutReason
 	return false;
 }
 
-bool ALevelInstance::IsEditorOnly() const
-{
-	if (IsRunningCookCommandlet() && GetLevel() && !GetLevel()->bIsPartitioned)
-	{
-		return false;
-	}
-
-	if (DesiredRuntimeBehavior == ELevelInstanceRuntimeBehavior::Partitioned)
-	{
-		return Super::IsEditorOnly();
-	}
-
-	return false;
-}
-
 FBox ALevelInstance::GetComponentsBoundingBox(bool bNonColliding, bool bIncludeFromChildActors) const
 {
 	FBox Box = Super::GetComponentsBoundingBox(bNonColliding, bIncludeFromChildActors);
diff --git a/Engine/Source/Runtime/Engine/Private/PackedLevelActor/PackedLevelActor.cpp b/Engine/Source/Runtime/Engine/Private/PackedLevelActor/PackedLevelActor.cpp
--- a/Engine/Source/Runtime/Engine/Private/PackedLevelActor/PackedLevelActor.cpp
+++ b/Engine/Source/Runtime/Engine/Private/PackedLevelActor/PackedLevelActor.cpp
@@ -22,7 +22,6 @@ APackedLevelActor::APackedLevelActor()
 #if WITH_EDITORONLY_DATA
 	// Packed Level Instances don't support level streaming or sub actors
 	DesiredRuntimeBehavior = ELevelInstanceRuntimeBehavior::None;
-	bIsEditorOnlyActor = false;
 #endif
 }
 
diff --git a/Engine/Source/Runtime/Engine/Private/WorldPartition/WorldPartitionStreamingGeneration.cpp b/Engine/Source/Runtime/Engine/Private/WorldPartition/WorldPartitionStreamingGeneration.cpp
--- a/Engine/Source/Runtime/Engine/Private/WorldPartition/WorldPartitionStreamingGeneration.cpp
+++ b/Engine/Source/Runtime/Engine/Private/WorldPartition/WorldPartitionStreamingGeneration.cpp
@@ -785,31 +785,25 @@ class FWorldPartitionStreamingGenerator
 
 					*OutActor = Actor;
 				}
+
+				return !Actor->IsEditorOnly() || (bIsPIE && Actor->IsEditorOnlyLoadedInPIE());
 			}
 
-			return true;
+			return !InActorDescInstance->GetActorIsEditorOnly() || (bIsPIE && InActorDescInstance->GetActorIsEditorOnlyLoadedInPIE());
 		};
 
 		// Register the actor descriptor view
-		auto RegisterActorDescView = [this, &OutActorDescViewMap, &OutContainerInstances](FStreamingGenerationActorDescView&& InActorDescView, TSet<FGuid>* OutEditorOnlyActorDescSet = nullptr, AActor* InActor = nullptr)
+		auto RegisterActorDescView = [this, &OutActorDescViewMap, &OutContainerInstances](FStreamingGenerationActorDescView&& InActorDescView)
 		{
 			if (InActorDescView.IsChildContainerInstance())
 			{
 				OutContainerInstances.Add(InActorDescView);
 			}
-
-			const bool bActorIsEditorOnly = InActor ? InActor->IsEditorOnly() : InActorDescView.GetActorIsEditorOnly();
-			const bool bActorIsNonPIEEditorOnly = InActorDescView.GetActorIsEditorOnlyLoadedInPIE();
-
-			if (!bActorIsEditorOnly || (bIsPIE && bActorIsNonPIEEditorOnly))
+			else
 			{
 				const FGuid ActorGuid = InActorDescView.GetGuid();
 				OutActorDescViewMap.Emplace(ActorGuid, MoveTemp(InActorDescView));
 			}
-			else if (OutEditorOnlyActorDescSet)
-			{
-				OutEditorOnlyActorDescSet->Add(InActorDescView.GetGuid());
-			}
 		};
 
 		for (FStreamingGenerationContainerInstanceCollection::TConstIterator<> Iterator(&InActorDescCollection); Iterator; ++Iterator)
@@ -843,13 +837,17 @@ class FWorldPartitionStreamingGenerator
 					{
 						// Dirty, unsaved actor for PIE
 						TUniquePtr<FStreamingGenerationUnsavedDirtyActorDescInstance>& UnsavedDirtyRef = OutUnsavedDirtyInstances.Add_GetRef(FStreamingGenerationUnsavedDirtyActorDescInstance::Create(Actor, InActorDescCollection));
-						RegisterActorDescView(FStreamingGenerationActorDescView(OutActorDescViewMap, UnsavedDirtyRef.Get(), true), &OutEditorOnlyActorDescSet, Actor);
+						RegisterActorDescView(FStreamingGenerationActorDescView(OutActorDescViewMap, UnsavedDirtyRef.Get(), true));
 						continue;
 					}
 				}
 
 				// Non-dirty actor
-				RegisterActorDescView(FStreamingGenerationActorDescView(OutActorDescViewMap, *Iterator), &OutEditorOnlyActorDescSet);
+				RegisterActorDescView(FStreamingGenerationActorDescView(OutActorDescViewMap, *Iterator));
+			}
+			else
+			{
+				OutEditorOnlyActorDescSet.Add(Iterator->GetGuid());
 			}
 		}
 
@@ -885,7 +883,9 @@ class FWorldPartitionStreamingGenerator
 				// Here, FindHandlingContainer is used to make sure that the actor is handled by the collection
 				// The main reason is that UWorldPartition::CheckForErrors currently builds a collection per ActorDescContainer of the WorldPartition.
 				// This is probably a limitation introduced by ContentBundles. 
-				if (IsValid(Actor) && ((Actor->IsPackageExternal() && Actor->IsMainPackageActor()) || Actor->HasAllFlags(RF_Transient | RF_NonPIEDuplicateTransient))
+				if (IsValid(Actor)
+					&& ((Actor->IsPackageExternal() && Actor->IsMainPackageActor()) || Actor->HasAllFlags(RF_Transient | RF_NonPIEDuplicateTransient))
+					&& (!Actor->IsEditorOnly() || (bIsPIE && Actor->IsEditorOnlyLoadedInPIE()))
 					&& InActorDescCollection.FindHandlingContainerInstance(Actor)
 					&& !InActorDescCollection.GetActorDescInstance(Actor->GetActorGuid()))
 				{
diff --git a/Engine/Source/Runtime/Engine/Public/LevelInstance/LevelInstanceActor.h b/Engine/Source/Runtime/Engine/Public/LevelInstance/LevelInstanceActor.h
--- a/Engine/Source/Runtime/Engine/Public/LevelInstance/LevelInstanceActor.h
+++ b/Engine/Source/Runtime/Engine/Public/LevelInstance/LevelInstanceActor.h
@@ -109,7 +109,6 @@ public:
 	ENGINE_API virtual bool GetSoftReferencedContentObjects(TArray<FSoftObjectPath>& SoftObjects) const override;
 	ENGINE_API virtual bool OpenAssetEditor() override;
 	ENGINE_API virtual bool EditorCanAttachFrom(const AActor* InChild, FText& OutReason) const override;	
-	ENGINE_API virtual bool IsEditorOnly() const override;
 	ENGINE_API virtual bool IsUserManaged() const override;
 	ENGINE_API virtual bool ShouldExport() override;
 	ENGINE_API virtual bool SupportsSubRootSelection() const override { return true; }
