[NetPhysics] First implementation of Physics Replication LOD
- Interface that allows for custom LOD solutions
- A base implementation of an LOD solution and the use of it inside physics replication, disabled by default.
- Project Settings and CVars for settings
- API to register focal points / particles in the LOD via AActor, UPrimitiveComponent or directly to the game thread or physics thread side of the LOD
- Option to register autonomous proxy in LOD via the NetworkPhysicsSettingsComponent

diff --git a/Engine/Source/Runtime/Engine/Classes/Components/PrimitiveComponent.h b/Engine/Source/Runtime/Engine/Classes/Components/PrimitiveComponent.h
--- a/Engine/Source/Runtime/Engine/Classes/Components/PrimitiveComponent.h
+++ b/Engine/Source/Runtime/Engine/Classes/Components/PrimitiveComponent.h
@@ -2865,6 +2865,14 @@ public:
 	 */
 	ENGINE_API void SetRigidBodyReplicatedTarget(FRigidBodyState& UpdatedState, const FName BoneName = NAME_None, int32 ServerFrame = 0, int32 ServerHandle = 0);
 
+	/** Register this components physics object as a focal particle in Physics Repliocation LOD */
+	UFUNCTION(BlueprintCallable, Category = Replication)
+	ENGINE_API void RegisterAsFocalPointInPhysicsReplicationLOD() const;
+
+	/** Unregister this components physics object from being a focal particle in Physics Repliocation LOD */
+	UFUNCTION(BlueprintCallable, Category = Replication)
+	ENGINE_API void UnregisterAsFocalPointInPhysicsReplicationLOD() const;
+
 protected:
 	ENGINE_API virtual bool CanBeUsedInPhysicsReplication(const FName BoneName = NAME_None) const { return true; }
 
diff --git a/Engine/Source/Runtime/Engine/Classes/GameFramework/Actor.h b/Engine/Source/Runtime/Engine/Classes/GameFramework/Actor.h
--- a/Engine/Source/Runtime/Engine/Classes/GameFramework/Actor.h
+++ b/Engine/Source/Runtime/Engine/Classes/GameFramework/Actor.h
@@ -881,6 +881,14 @@ public:
 	UFUNCTION(BlueprintCallable, Category = Replication)
 	ENGINE_API EPhysicsReplicationMode GetPhysicsReplicationMode();
 
+	/** Register this actors root components physics object as a focal particle in Physics Repliocation LOD */
+	UFUNCTION(BlueprintCallable, Category = Replication)
+	ENGINE_API void RegisterAsFocalPointInPhysicsReplicationLOD() const;
+
+	/** Unregister this actors root components physics object from being a focal particle in Physics Repliocation LOD */
+	UFUNCTION(BlueprintCallable, Category = Replication)
+	ENGINE_API void UnregisterAsFocalPointInPhysicsReplicationLOD() const;
+
 	/** Can this body trigger a resimulation when Physics Prediction is enabled */
 	UFUNCTION(BlueprintCallable, Category = Physics)
 	ENGINE_API bool CanTriggerResimulation() const;
diff --git a/Engine/Source/Runtime/Engine/Classes/PhysicsEngine/PhysicsSettings.h b/Engine/Source/Runtime/Engine/Classes/PhysicsEngine/PhysicsSettings.h
--- a/Engine/Source/Runtime/Engine/Classes/PhysicsEngine/PhysicsSettings.h
+++ b/Engine/Source/Runtime/Engine/Classes/PhysicsEngine/PhysicsSettings.h
@@ -95,7 +95,7 @@ struct FPhysicsReplicationResimulationSettings
 	UPROPERTY(config, EditDefaultsOnly, Category = "Replication", Meta = (InlineEditConditionToggle))
 	bool bEnableResimulationErrorPositionThreshold;
 
-	/** Distance in centimeters before a state discrepancy triggers a resimulation */
+	/** Enable positional error threshold to trigger resimulation and set distance in centimeters before a state discrepancy triggers a resimulation */
 	UPROPERTY(EditAnywhere, Category = "Replication", meta = (editcondition = "bEnableResimulationErrorPositionThreshold"))
 	float ResimulationErrorPositionThreshold;
 
@@ -103,7 +103,7 @@ struct FPhysicsReplicationResimulationSettings
 	UPROPERTY(config, EditDefaultsOnly, Category = "Replication", Meta = (InlineEditConditionToggle))
 	bool bEnableResimulationErrorRotationThreshold;
 
-	/** Rotation difference in degrees before a state discrepancy triggers a resimulation */
+	/** Enable rotational error threshold to trigger resimulation and set rotation difference in degrees before a state discrepancy triggers a resimulation */
 	UPROPERTY(EditAnywhere, Category = "Replication", meta = (editcondition = "bEnableResimulationErrorRotationThreshold"))
 	float ResimulationErrorRotationThreshold;
 
@@ -111,7 +111,7 @@ struct FPhysicsReplicationResimulationSettings
 	UPROPERTY(config, EditDefaultsOnly, Category = "Replication", Meta = (InlineEditConditionToggle))
 	bool bEnableResimulationErrorLinearVelocityThreshold;
 
-	/** Velocity difference in centimeters / second before a state discrepancy triggers a resimulation */
+	/** Enable linear velocity error threshold to trigger resimulation and set velocity difference in centimeters / second before a state discrepancy triggers a resimulation */
 	UPROPERTY(EditAnywhere, Category = "Replication", meta = (editcondition = "bEnableResimulationErrorLinearVelocityThreshold"))
 	float ResimulationErrorLinearVelocityThreshold;
 
@@ -119,7 +119,7 @@ struct FPhysicsReplicationResimulationSettings
 	UPROPERTY(config, EditDefaultsOnly, Category = "Replication", Meta = (InlineEditConditionToggle))
 	bool bEnableResimulationErrorAngularVelocityThreshold;
 
-	/** Velocity degrees / second before a state discrepancy triggers a resimulation */
+	/** Enable angular velocity error threshold to trigger resimulation and set velocity degrees / second before a state discrepancy triggers a resimulation */
 	UPROPERTY(EditAnywhere, Category = "Replication", meta = (editcondition = "bEnableResimulationErrorAngularVelocityThreshold"))
 	float ResimulationErrorAngularVelocityThreshold;
 
@@ -135,6 +135,43 @@ struct FPhysicsReplicationResimulationSettings
 	{ }
 };
 
+/** Default settings for physics replication LOD */
+USTRUCT()
+struct FPhysicsReplicationLODSettings
+{
+	GENERATED_BODY();
+
+	/** Enable physics replication LOD */
+	UPROPERTY(config, EditDefaultsOnly, Category = "Replication")
+	bool bEnablePhysicsReplicationLOD;
+
+	/** Within this distance from Focal Points in LOD, physics replicated objects will use EPhysicsReplicationMode::Resimulation
+	* Resimulation is fully forward predicted physics replication, for a seamless LOD transition this should be equal to or lower than DistanceForFullPrediction, read that property for information. */
+	UPROPERTY(EditAnywhere, Category = "Replication")
+	float DistanceForResimulationMode;
+
+	/** Within this distance from Focal Points in LOD, physics replicated objects will be in the clients forward predicted timeline
+	* The local player is typically forward predicted, while physics replicated objects are replicated without forward prediction.
+	* Physics replication LOD gradually forward predicts the physics replication when objects get closer to focal points in LOD.
+	* Read the TimeOverDistance property for details on how forward prediction is gradually applied up until the distance of DistanceForFullPrediction. */
+	UPROPERTY(EditAnywhere, Category = "Replication")
+	float DistanceForFullPrediction;
+
+	/** Time/Distance alignment value, used as Distance * TimeOverDistance = Time
+	* Time is how far behind the current timeline replication should run
+	* Starts outside of DistanceForFullPrediction and gets clamped by the received states timeline
+	* Lower value = Slower timeline transition which produces a larger LOD radius.
+	* EXAMPLE: At 500cm away an object should be 37,5ms behind the clients timeline if DistanceFullPrediction is 250cm and TimeOverDistance is 0.15ms/cm. (500cm - 250cm) * 0.15ms/cm = 37,5ms */
+	UPROPERTY(EditAnywhere, Category = "Replication")
+	float TimeOverDistance;
+
+	FPhysicsReplicationLODSettings()
+		: bEnablePhysicsReplicationLOD(false)
+		, DistanceForResimulationMode(200.0f)
+		, DistanceForFullPrediction(250.0f)
+		, TimeOverDistance(0.15f)
+	{ }
+};
 
 /** Physics Prediction Settings */
 USTRUCT()
@@ -171,6 +208,10 @@ struct FPhysicsPredictionSettings
 	UPROPERTY(EditAnywhere, Category = "Replication", meta = (DisplayName = "Default Resimulation Replication Settings (EPhysicsReplicationMode::Resimulation)"))
 	FPhysicsReplicationResimulationSettings ResimulationSettings;
 
+	/** Default settings for physics replication LOD. */
+	UPROPERTY(EditAnywhere, Category = "Replication", meta = (DisplayName = "Default settings for physics replication LOD"))
+	FPhysicsReplicationLODSettings PhysicsReplicationLODSettings;
+
 	FPhysicsPredictionSettings()
 		: bEnablePhysicsResimulation_DEPRECATED(false)
 		, ResimulationErrorThreshold_DEPRECATED(10.0f)
diff --git a/Engine/Source/Runtime/Engine/Private/ActorReplication.cpp b/Engine/Source/Runtime/Engine/Private/ActorReplication.cpp
--- a/Engine/Source/Runtime/Engine/Private/ActorReplication.cpp
+++ b/Engine/Source/Runtime/Engine/Private/ActorReplication.cpp
@@ -1196,6 +1196,23 @@ EPhysicsReplicationMode AActor::GetPhysicsReplicationMode()
 	return PhysicsReplicationMode;
 }
 
+
+void AActor::RegisterAsFocalPointInPhysicsReplicationLOD() const
+{
+	if (UPrimitiveComponent* RootPrimComp = Cast<UPrimitiveComponent>(RootComponent))
+	{
+		RootPrimComp->RegisterAsFocalPointInPhysicsReplicationLOD();
+	}
+}
+
+void AActor::UnregisterAsFocalPointInPhysicsReplicationLOD() const
+{
+	if (UPrimitiveComponent* RootPrimComp = Cast<UPrimitiveComponent>(RootComponent))
+	{
+		RootPrimComp->UnregisterAsFocalPointInPhysicsReplicationLOD();
+	}
+}
+
 float AActor::GetResimulationThreshold() const
 {
 	return UPhysicsSettings::Get()->PhysicsPrediction.ResimulationSettings.ResimulationErrorPositionThreshold;
diff --git a/Engine/Source/Runtime/Engine/Private/PhysicsEngine/Experimental/PhysScene_Chaos.cpp b/Engine/Source/Runtime/Engine/Private/PhysicsEngine/Experimental/PhysScene_Chaos.cpp
--- a/Engine/Source/Runtime/Engine/Private/PhysicsEngine/Experimental/PhysScene_Chaos.cpp
+++ b/Engine/Source/Runtime/Engine/Private/PhysicsEngine/Experimental/PhysScene_Chaos.cpp
@@ -11,6 +11,7 @@
 #include "Components/SkeletalMeshComponent.h"
 #include "Physics/PhysicsInterfaceUtils.h"
 #include "PhysicsReplication.h"
+#include "PhysicsReplicationLOD.h"
 #include "Physics/PhysicsReplicationCache.h"
 #include "PhysicsEngine/ClusterUnionComponent.h"
 #include "PhysicsEngine/ConstraintInstance.h"
@@ -533,6 +534,7 @@ FPhysScene_Chaos::FPhysScene_Chaos(AActor* InSolverActor
 #endif
 	)
 	, PhysicsReplication(nullptr)
+	, PhysicsReplicationLOD(nullptr)
 	, SolverActor(InSolverActor)
 	, LastEventDispatchTime(Chaos::FReal(-1))
 	, LastBreakEventDispatchTime(Chaos::FReal(-1))
@@ -570,9 +572,12 @@ FPhysScene_Chaos::FPhysScene_Chaos(AActor* InSolverActor
 	Scene.GetSolver()->GetEvolution()->GetParticles().AddArray(BodyInstances.Get());
 #endif
 
-	// Create replication manager
+	// Create physics replication manager
 	CreatePhysicsReplication();
 
+	// Create physics replication LOD manager
+	CreatePhysicsReplicationLOD();
+
 	FPhysicsDelegates::OnPhysSceneInit.Broadcast(this);
 
 	ChaosEventRelay = NewObject<UChaosEventRelay>();
@@ -590,6 +595,7 @@ FPhysScene_Chaos::~FPhysScene_Chaos()
 	
 	// Make sure physics replication is cleared before we're fully destructed
 	PhysicsReplication.Reset();
+	PhysicsReplicationLOD.Reset();
 	PhysicsReplicationCache.Reset();
 
 	FPhysicsDelegates::OnPhysSceneTerm.Broadcast(this);
@@ -831,18 +837,43 @@ IPhysicsReplication* FPhysScene_Chaos::GetPhysicsReplication()
 IPhysicsReplication* FPhysScene_Chaos::CreatePhysicsReplication()
 {
 	// Create replication manager
-	PhysicsReplication
-		= PhysicsReplicationFactory.IsValid()
-		? PhysicsReplicationFactory->CreatePhysicsReplication(this)
-		: MakeUnique<FPhysicsReplication>(this);
+	PhysicsReplication = nullptr;
+	if (PhysicsReplicationFactory.IsValid())
+	{
+		PhysicsReplication = PhysicsReplicationFactory->CreatePhysicsReplication(this);
+	}
+
+	// Check that replication manager was created, else create the default implementation
+	if (PhysicsReplication == nullptr)
+	{
+		PhysicsReplication = MakeUnique<FPhysicsReplication>(this);
+	}
 
 	// Return ptr to the new physics rep
 	return PhysicsReplication.Get();
 }
 
-void FPhysScene_Chaos::SetPhysicsReplication(IPhysicsReplication* InPhysicsReplication)
+IPhysicsReplicationLOD* FPhysScene_Chaos::GetPhysicsReplicationLOD()
+{
+	return PhysicsReplicationLOD.Get();
+}
+
+IPhysicsReplicationLOD* FPhysScene_Chaos::CreatePhysicsReplicationLOD()
 {
-	PhysicsReplication = TUniquePtr<IPhysicsReplication>(InPhysicsReplication);
+	PhysicsReplicationLOD = nullptr;
+	if (PhysicsReplicationFactory.IsValid())
+	{
+		PhysicsReplicationLOD = PhysicsReplicationFactory->CreatePhysicsReplicationLOD(this);
+	}
+
+	// Check that replication LOD manager was created, else create the default implementation
+	if (PhysicsReplicationLOD == nullptr)
+	{
+		PhysicsReplicationLOD = MakeUnique<FPhysicsReplicationLOD>(this);
+	}
+
+	// Return ptr to the new physics replication LOD 
+	return PhysicsReplicationLOD.Get();
 }
 
 void FPhysScene_Chaos::AddReferencedObjects(FReferenceCollector& Collector)
diff --git a/Engine/Source/Runtime/Engine/Private/PhysicsEngine/NetworkPhysicsSettingsComponent.cpp b/Engine/Source/Runtime/Engine/Private/PhysicsEngine/NetworkPhysicsSettingsComponent.cpp
--- a/Engine/Source/Runtime/Engine/Private/PhysicsEngine/NetworkPhysicsSettingsComponent.cpp
+++ b/Engine/Source/Runtime/Engine/Private/PhysicsEngine/NetworkPhysicsSettingsComponent.cpp
@@ -31,6 +31,32 @@ UNetworkPhysicsSettingsComponent::UNetworkPhysicsSettingsComponent()
 	bAutoActivate = true;
 }
 
+void UNetworkPhysicsSettingsComponent::OnRegister()
+{
+	Super::OnRegister();
+
+	if (AActor* Owner = GetOwner())
+	{
+		if (UPrimitiveComponent* RootPrimComp = Cast<UPrimitiveComponent>(Owner->GetRootComponent()))
+		{
+			RootPrimComp->OnComponentPhysicsStateChanged.AddUniqueDynamic(this, &ThisClass::OnComponentPhysicsStateChanged);
+		}
+	}
+}
+
+void UNetworkPhysicsSettingsComponent::OnUnregister()
+{
+	Super::OnUnregister();
+
+	if (AActor* Owner = GetOwner())
+	{
+		if (UPrimitiveComponent* RootPrimComp = Cast<UPrimitiveComponent>(Owner->GetRootComponent()))
+		{
+			RootPrimComp->OnComponentPhysicsStateChanged.RemoveDynamic(this, &ThisClass::OnComponentPhysicsStateChanged);
+		}
+	}
+}
+
 void UNetworkPhysicsSettingsComponent::InitializeComponent()
 {
 	Super::InitializeComponent();
@@ -56,11 +82,15 @@ void UNetworkPhysicsSettingsComponent::InitializeComponent()
 							{
 								FNetworkPhysicsSettingsAsyncInput* AsyncInput = NetworkPhysicsSettings_Internal->GetProducerInputData_External();
 								AsyncInput->PhysicsObject = PhysicsObject;
-								AsyncInput->Settings.GeneralSettings = GeneralSettings;
-								AsyncInput->Settings.DefaultReplicationSettings = DefaultReplicationSettings;
-								AsyncInput->Settings.ResimulationSettings = ResimulationSettings;
-								AsyncInput->Settings.PredictiveInterpolationSettings = PredictiveInterpolationSettings;
-								AsyncInput->Settings.NetworkPhysicsComponentSettings = NetworkPhysicsComponentSettings;
+								FNetworkPhysicsSettingsAsync Settings
+								{
+									.GeneralSettings = GeneralSettings,
+									.DefaultReplicationSettings = DefaultReplicationSettings,
+									.PredictiveInterpolationSettings = PredictiveInterpolationSettings,
+									.ResimulationSettings = ResimulationSettings,
+									.NetworkPhysicsComponentSettings = NetworkPhysicsComponentSettings,
+								};
+								AsyncInput->Settings = Settings;
 							}
 						}
 					}
@@ -112,6 +142,36 @@ void UNetworkPhysicsSettingsComponent::BeginPlay()
 			Owner->SetPhysicsReplicationMode(RepMode);
 		}
 	}
+	RegisterInPhysicsReplicationLOD();
+}
+
+void UNetworkPhysicsSettingsComponent::OnComponentPhysicsStateChanged(UPrimitiveComponent* ChangedComponent, EComponentPhysicsStateChange StateChange)
+{
+	if (StateChange == EComponentPhysicsStateChange::Created)
+	{
+		if (Chaos::FConstPhysicsObjectHandle PhysicsObject = ChangedComponent->GetPhysicsObjectByName(NAME_None))
+		{
+			FNetworkPhysicsSettingsAsyncInput* AsyncInput = NetworkPhysicsSettings_Internal->GetProducerInputData_External();
+			AsyncInput->PhysicsObject = PhysicsObject;
+		}
+		RegisterInPhysicsReplicationLOD();
+	}
+}
+
+void UNetworkPhysicsSettingsComponent::RegisterInPhysicsReplicationLOD()
+{
+	if (!GeneralSettings.bFocalParticleInPhysicsReplicationLOD)
+	{
+		return;
+	}
+
+	if (AActor* Owner = GetOwner())
+	{
+		if (Owner->GetLocalRole() == ENetRole::ROLE_AutonomousProxy)
+		{
+			Owner->RegisterAsFocalPointInPhysicsReplicationLOD();
+		}
+	}
 }
 
 UNetworkPhysicsSettingsComponent* UNetworkPhysicsSettingsComponent::GetSettingsForActor(AActor* Owner)
@@ -122,19 +182,36 @@ UNetworkPhysicsSettingsComponent* UNetworkPhysicsSettingsComponent::GetSettingsF
 
 #pragma region // FNetworkPhysicsSettingsComponentAsync
 
-void FNetworkPhysicsSettingsComponentAsync::OnPostInitialize_Internal()
+void FNetworkPhysicsSettingsComponentAsync::OnPreSimulate_Internal()
+{
+	ConsumeAsyncInput();
+};
+
+void FNetworkPhysicsSettingsComponentAsync::ConsumeAsyncInput()
 {
-	// Receive settings data on PT from GT
+	// Receive data on PT from GT
 	if (const FNetworkPhysicsSettingsAsyncInput* AsyncInput = GetConsumerInput_Internal())
 	{
-		Settings = AsyncInput->Settings;
-		
-		if (Chaos::FPBDRigidsSolver* RigidsSolver = static_cast<Chaos::FPBDRigidsSolver*>(GetSolver()))
+		if (AsyncInput->PhysicsObject.IsSet())
 		{
-			if (FPhysicsReplicationAsync* PhysRep = RigidsSolver->GetPhysicsReplication())
-			{
-				PhysRep->RegisterSettings(AsyncInput->PhysicsObject, Settings);
-			}
+			PhysicsObject = *AsyncInput->PhysicsObject;
+			RegisterSettingsInPhysicsReplication();
+		}
+
+		if (AsyncInput->Settings.IsSet())
+		{
+			Settings = *AsyncInput->Settings;
+		}
+	}
+}
+
+void FNetworkPhysicsSettingsComponentAsync::RegisterSettingsInPhysicsReplication()
+{
+	if (Chaos::FPBDRigidsSolver* RigidsSolver = static_cast<Chaos::FPBDRigidsSolver*>(GetSolver()))
+	{
+		if (IPhysicsReplicationAsync* PhysRep = RigidsSolver->GetPhysicsReplication_Internal())
+		{
+			PhysRep->RegisterSettings(PhysicsObject, Settings);
 		}
 	}
 }
diff --git a/Engine/Source/Runtime/Engine/Private/PhysicsEngine/PhysicsInterfaceUtils.cpp b/Engine/Source/Runtime/Engine/Private/PhysicsEngine/PhysicsInterfaceUtils.cpp
--- a/Engine/Source/Runtime/Engine/Private/PhysicsEngine/PhysicsInterfaceUtils.cpp
+++ b/Engine/Source/Runtime/Engine/Private/PhysicsEngine/PhysicsInterfaceUtils.cpp
@@ -103,16 +103,3 @@ FCollisionFilterData CreateQueryFilterData(const uint8 MyChannel, const bool bTr
 		return CreateTraceQueryFilterData(MyChannel, bTraceComplex, InCollisionResponseContainer, QueryParam);
 	}
 }
-
-//
-// NOTE: Once the old Create/Destroy methods are deprecated, remove the default implementations
-//
-
-#include "PhysicsReplication.h" // Remove when the deprecated Create() function is removed
-TUniquePtr<IPhysicsReplication> IPhysicsReplicationFactory::CreatePhysicsReplication(FPhysScene* OwningPhysScene)
-{
-PRAGMA_DISABLE_DEPRECATION_WARNINGS
-	UE_LOG(LogPhysics, Warning, TEXT("\nIPhysicsReplicationFactory::Create has been deprecated in favor of IPhysicsReplicationFactory::CreatePhysicsReplication. Use the new CreatePhysicsReplication method instead moving forward, as the Create function will be removed in future versions. Please update your code to the new API before upgrading to the next release, otherwise your project will no longer compile."));
-	return TUniquePtr<IPhysicsReplication>(Create(OwningPhysScene));
-PRAGMA_ENABLE_DEPRECATION_WARNINGS
-}
diff --git a/Engine/Source/Runtime/Engine/Private/PhysicsEngine/PhysicsReplication.cpp b/Engine/Source/Runtime/Engine/Private/PhysicsEngine/PhysicsReplication.cpp
--- a/Engine/Source/Runtime/Engine/Private/PhysicsEngine/PhysicsReplication.cpp
+++ b/Engine/Source/Runtime/Engine/Private/PhysicsEngine/PhysicsReplication.cpp
@@ -5,6 +5,7 @@
 =============================================================================*/
 
 #include "PhysicsReplication.h"
+#include "PhysicsReplicationLOD.h"
 #include "Engine/World.h"
 #include "Components/SkeletalMeshComponent.h"
 #include "DrawDebugHelpers.h"
@@ -871,12 +872,9 @@ void FPhysicsReplicationAsync::FetchObjectSettings(Chaos::FConstPhysicsObjectHan
 
 void FPhysicsReplicationAsync::OnPostInitialize_Internal()
 {
-	Chaos::FPBDRigidsSolver* RigidsSolver = static_cast<Chaos::FPBDRigidsSolver*>(GetSolver());
-	
-	if (ensure(RigidsSolver))
-	{
-		RigidsSolver->SetPhysicsReplication(this);
-	}
+	Chaos::FPBDRigidsSolver& RigidsSolver = GetSolver()->CastChecked();
+
+	RigidsSolver.SetPhysicsReplication_Internal(this);
 }
 
 void FPhysicsReplicationAsync::OnPreSimulate_Internal()
@@ -1110,7 +1108,10 @@ void FPhysicsReplicationAsync::UpdateAsyncTarget(const FPhysicsRepAsyncInputData
 		// Update waiting state
 		Target->UpdateWaiting(Input.ServerFrame);
 
-		if (Input.RepMode == EPhysicsReplicationMode::PredictiveInterpolation)
+		// Apply full Replication LOD on received target
+		ApplyPhysicsReplicationLOD(Input.PhysicsObject, *Target, EPhysicsReplicationLODFlags::LODFlag_All);
+
+		if (Target->RepMode == EPhysicsReplicationMode::PredictiveInterpolation)
 		{
 #if !(UE_BUILD_SHIPPING || UE_BUILD_TEST)
 			if (PhysicsReplicationCVars::PredictiveInterpolationCVars::bDrawDebugTargets)
@@ -1124,6 +1125,8 @@ void FPhysicsReplicationAsync::UpdateAsyncTarget(const FPhysicsRepAsyncInputData
 			Target->PrevPosTarget = Input.TargetState.Position;
 			Target->PrevRotTarget = Input.TargetState.Quaternion;
 		
+			// TickCount is 0 by default at this point and when LOD is used, TickCount will be 0 if no LOD alignment was performed, in this case perform the normal target alignment
+			if (Target->TickCount == 0)
 			{
 				/** Target Alignment Feature
 				* With variable network conditions state inputs from the server can arrive both later or earlier than expected.
@@ -1171,6 +1174,14 @@ void FPhysicsReplicationAsync::CacheResimInteractions()
 {
 	if(!PhysicsReplicationCVars::ResimulationCVars::bDisableReplicationOnInteraction)
 	{
+		ParticlesInResimIslands.Empty();
+		return;
+	}
+
+	if (UsePhysicsReplicationLOD())
+	{
+		// This will be handled by the LOD system
+		ParticlesInResimIslands.Empty();
 		return;
 	}
 
@@ -1181,8 +1192,10 @@ void FPhysicsReplicationAsync::CacheResimInteractions()
 	}
 
 	ParticlesInResimIslands.Empty(FMath::CeilToInt(static_cast<float>(ParticlesInResimIslands.Num()) * 0.9f));
+
 	Chaos::Private::FPBDIslandManager& IslandManager = RigidsSolver->GetEvolution()->GetIslandManager();
-	Chaos::FWritePhysicsObjectInterface_Internal Interface = Chaos::FPhysicsObjectInternalInterface::GetWrite();
+	Chaos::FReadPhysicsObjectInterface_Internal Interface = Chaos::FPhysicsObjectInternalInterface::GetRead();
+
 	for (auto Itr = ObjectToTarget.CreateIterator(); Itr; ++Itr)
 	{
 		FReplicatedPhysicsTargetAsync& Target = Itr.Value();
@@ -1218,48 +1231,102 @@ void FPhysicsReplicationAsync::ApplyTargetStatesAsync(const float DeltaSeconds,
 		}
 	}
 
+	/** Helper function to remove replicated target*/
+	auto RemoveTargetHelper = [this](TMap<Chaos::FConstPhysicsObjectHandle, FReplicatedPhysicsTargetAsync>::TIterator Itr, FGeometryParticleHandle* Handle)
+	{
+		if (Handle)
+		{
+			ReplicatedParticleIDs.Remove(Handle->ParticleID());
+		}
+		Itr.RemoveCurrent();
+	};
+
 	// PhysicsObject flow
 	Chaos::FWritePhysicsObjectInterface_Internal Interface = Chaos::FPhysicsObjectInternalInterface::GetWrite();
-	for (auto Itr = ObjectToTarget.CreateIterator(); Itr; ++Itr)
+	for (TMap<Chaos::FConstPhysicsObjectHandle, FReplicatedPhysicsTargetAsync>::TIterator Itr = ObjectToTarget.CreateIterator(); Itr; ++Itr)
 	{
 		bool bRemoveItr = true; // Remove current cached replication target unless replication logic tells us to store it for next tick
-		FParticleID ParticleID;
 
 		Chaos::FConstPhysicsObjectHandle& POHandle = Itr.Key();
-		if (FGeometryParticleHandle* Handle = Interface.GetParticle(POHandle))
+		FGeometryParticleHandle* Handle = Interface.GetParticle(POHandle);
+		if (!Handle)
 		{
-			FReplicatedPhysicsTargetAsync& Target = Itr.Value();
+			RemoveTargetHelper(Itr, nullptr);
+			continue;
+		}
 
-			if (FPBDRigidParticleHandle* RigidHandle = Handle->CastToRigidParticle())
-			{
-				ParticleID = RigidHandle->ParticleID();
+		FPBDRigidParticleHandle* RigidHandle = Handle->CastToRigidParticle();
+		if (!RigidHandle)
+		{
+			RemoveTargetHelper(Itr, Handle);
+			continue;
+		}
 
-				// Cache custom settings for this object if there are any
-				FetchObjectSettings(POHandle);
+		FReplicatedPhysicsTargetAsync& Target = Itr.Value();
 
-				const EPhysicsReplicationMode RepMode = Target.IsWaiting() ? Target.RepModeOverride : Target.RepMode;
-				switch (RepMode)
-				{
-					case EPhysicsReplicationMode::Default:
-						bRemoveItr = DefaultReplication(RigidHandle, Target, DeltaSeconds);
-						break;
+		// Cache custom settings for this object if there are any
+		FetchObjectSettings(POHandle);
 
-					case EPhysicsReplicationMode::PredictiveInterpolation:
-						bRemoveItr = PredictiveInterpolation(RigidHandle, Target, DeltaSeconds);
-						break;
+		// Apply limited Replication LOD
+		ApplyPhysicsReplicationLOD(POHandle, Target, EPhysicsReplicationLODFlags::LODFlag_IslandCheck);
 
-					case EPhysicsReplicationMode::Resimulation:
-						bRemoveItr = ResimulationReplication(RigidHandle, Target, DeltaSeconds);
-						break;
-				}
-				Target.TickCount++;
-			}
-		}
+		const EPhysicsReplicationMode RepMode = Target.IsWaiting() ? Target.RepModeOverride : Target.RepMode;
+		switch (RepMode)
+		{
+			case EPhysicsReplicationMode::Default:
+				bRemoveItr = DefaultReplication(RigidHandle, Target, DeltaSeconds);
+				break;
 
+			case EPhysicsReplicationMode::PredictiveInterpolation:
+				bRemoveItr = PredictiveInterpolation(RigidHandle, Target, DeltaSeconds);
+				break;
+
+			case EPhysicsReplicationMode::Resimulation:
+				bRemoveItr = ResimulationReplication(RigidHandle, Target, DeltaSeconds);
+				break;
+		}
+		Target.TickCount++;
+			
 		if (bRemoveItr)
 		{
-			ReplicatedParticleIDs.Remove(ParticleID);
-			Itr.RemoveCurrent();
+			RemoveTargetHelper(Itr, RigidHandle);
+		}
+	}
+}
+
+void FPhysicsReplicationAsync::ApplyPhysicsReplicationLOD(Chaos::FConstPhysicsObjectHandle PhysicsObjectHandle, FReplicatedPhysicsTargetAsync& Target, const uint32 LODFLags)
+{
+	Chaos::FPBDRigidsSolver& RigidsSolver = GetSolver()->CastChecked();
+
+	IPhysicsReplicationLODAsync* PhysRepLod = RigidsSolver.GetPhysicsReplicationLOD_Internal();
+	if (!PhysRepLod || !PhysRepLod->IsEnabled())
+	{
+		return;
+	}
+
+	FPhysicsRepLodData* LodData = PhysRepLod->GetLODData_Internal(PhysicsObjectHandle, LODFLags);
+	if (LodData && LodData->DataAssigned)
+	{
+		// Apply recommended replication mode
+		Target.RepMode = LodData->ReplicationMode;
+
+		if (Target.RepMode == EPhysicsReplicationMode::PredictiveInterpolation)
+		{
+			const bool bShouldSleep = (Target.TargetState.Flags & ERigidBodyFlags::Sleeping) != 0;
+			int32 TargetSourceClientFrame = (Target.ServerFrame - Target.FrameOffset);
+
+			// If we use Predicitve Interpolation and we should not sleep and the aligned frame from LOD is ahead of the target, perform LOD aligment extrapolation
+			if (!bShouldSleep && LodData->AlignedFrame > TargetSourceClientFrame)
+			{
+				// Calculate how far to forward predict and extrapolate target by that amount
+				const int32 FullPredictionFrames = RigidsSolver.GetCurrentFrame() - TargetSourceClientFrame;
+				const float FullPredictionTime = (FullPredictionFrames * GetDeltaTime_Internal());
+				const float AlignedPredictionTime = FullPredictionTime - LodData->AlignedTime;
+				FPhysicsReplicationAsync::ExtrapolateTarget(Target, AlignedPredictionTime);
+
+				// Update tick count based on LOD alignment
+				Target.TickCount = (LodData->AlignedFrame + Target.FrameOffset) - Target.ServerFrame;  
+			}
 		}
 	}
 }
@@ -1574,11 +1641,6 @@ bool FPhysicsReplicationAsync::PredictiveInterpolation(Chaos::FPBDRigidParticleH
 		return true;
 	}
 
-	if (Target.IsWaiting())
-	{
-		return false;
-	}
-
 	Chaos::FPBDRigidsSolver* RigidsSolver = static_cast<Chaos::FPBDRigidsSolver*>(GetSolver());
 	if (RigidsSolver == nullptr)
 	{
@@ -1652,6 +1714,12 @@ bool FPhysicsReplicationAsync::PredictiveInterpolation(Chaos::FPBDRigidParticleH
 		return bClearTarget;
 	};
 
+	// If waiting on an up to date state, early out but allow target clearing since we might not receive a new state if target is already set to sleep for example
+	if (Target.IsWaiting())
+	{
+		return EndReplicationHelper(Target, true);
+	}
+
 	// If target velocity is low enough, check the distance from the current position to the source position of our target to see if it's low enough to early out of replication
 	const bool bXCanEarlyOut = (PhysicsReplicationCVars::PredictiveInterpolationCVars::bEarlyOutWithVelocity || Target.TargetState.LinVel.SizeSquared() < UE_KINDA_SMALL_NUMBER) &&
 		(Target.PrevPosTarget - Handle->GetX()).SizeSquared() < PhysicsReplicationCVars::PredictiveInterpolationCVars::EarlyOutDistanceSqr;
@@ -1973,7 +2041,12 @@ bool FPhysicsReplicationAsync::PredictiveInterpolation(Chaos::FPBDRigidParticleH
 void FPhysicsReplicationAsync::ExtrapolateTarget(FReplicatedPhysicsTargetAsync& Target, const int32 ExtrapolateFrames, const float DeltaSeconds)
 {
 	const float ExtrapolationTime = DeltaSeconds * static_cast<float>(ExtrapolateFrames);
+	FPhysicsReplicationAsync::ExtrapolateTarget(Target, ExtrapolationTime);
+}
 
+/** Static function to extrapolate a target for N Seconds */
+void FPhysicsReplicationAsync::ExtrapolateTarget(FReplicatedPhysicsTargetAsync& Target, const float ExtrapolationTime)
+{
 	// Extrapolate target position
 	Target.TargetState.Position = Target.TargetState.Position + Target.TargetState.LinVel * ExtrapolationTime;
 
@@ -2001,6 +2074,11 @@ bool FPhysicsReplicationAsync::ResimulationReplication(Chaos::FPBDRigidParticleH
 		return true;
 	}
 
+	if (Target.ServerFrame <= 0)
+	{
+		return true;
+	}
+
 	const int32 LocalFrame = Target.ServerFrame - Target.FrameOffset;
 
 	if (LocalFrame >= RewindData->CurrentFrame() || LocalFrame < RewindData->GetEarliestFrame_Internal())
@@ -2160,6 +2238,11 @@ bool FPhysicsReplicationAsync::ResimulationReplication(Chaos::FPBDRigidParticleH
 			RigidsSolver->GetEvolution()->ApplySleepOnConnectedParticles(Handle);
 		}
 	}
+	else if (Target.IsWaiting()) 
+	{
+		// Don't clear the target if we are waiting for a specific target frame and not sleeping
+		bClearTarget = false;
+	}
 
 	return bClearTarget;
 }
@@ -2170,6 +2253,14 @@ FName FPhysicsReplicationAsync::GetFNameForStatId() const
 	return StaticName;
 }
 
+bool FPhysicsReplicationAsync::UsePhysicsReplicationLOD()
+{
+	Chaos::FPBDRigidsSolver& RigidsSolver = GetSolver()->CastChecked();
+
+	IPhysicsReplicationLODAsync* PhysRepLod = RigidsSolver.GetPhysicsReplicationLOD_Internal();
+	return PhysRepLod && PhysRepLod->IsEnabled();
+}
+
 #pragma endregion // FPhysicsReplicationAsync
 
 
diff --git a/Engine/Source/Runtime/Engine/Private/PhysicsEngine/PhysicsReplicationLOD.cpp b/Engine/Source/Runtime/Engine/Private/PhysicsEngine/PhysicsReplicationLOD.cpp
new file mode 100644
--- /dev/null
+++ b/Engine/Source/Runtime/Engine/Private/PhysicsEngine/PhysicsReplicationLOD.cpp
@@ -0,0 +1,369 @@
+// Copyright Epic Games, Inc. All Rights Reserved.
+
+#include "PhysicsReplicationLOD.h"
+#include "Physics/Experimental/PhysScene_Chaos.h"
+#include "PBDRigidsSolver.h"
+#include "Chaos/PhysicsObjectInternalInterface.h"
+#include <Components/PrimitiveComponent.h>
+#include "Chaos/ParticleDirtyFlags.h"
+#include "Chaos/DebugDrawQueue.h"
+
+
+namespace PhysicsReplicationLODCVars
+{
+	bool bEnabled = false;
+	static FAutoConsoleVariableRef CVarEnabled(TEXT("p.ReplicationLOD.Enabled"), bEnabled, TEXT("Override a disabled PhysicsReplicationLOD to become enabled, can't disable an LOD that is set enabled via other means like project settings or API call."));
+
+	float DistanceForResimulationMode = -1.0f;
+	static FAutoConsoleVariableRef CVarDistanceForResimulationMode(TEXT("p.ReplicationLOD.DistanceForResimulationMode"), DistanceForResimulationMode, TEXT("Override distance where within resimulation will be recommended as the replication mode by LOD."));
+
+	float DistanceForFullPrediction = -1.0f;
+	static FAutoConsoleVariableRef CVarDistanceForFullPrediction(TEXT("p.ReplicationLOD.DistanceForFullPrediction"), DistanceForFullPrediction, TEXT("Override the distance where within full forward prediction to clients current timeline will be recommended."));
+
+	float TimeOverDistance = -1.0f;
+	static FAutoConsoleVariableRef CVaTimeOverDistance(TEXT("p.ReplicationLOD.TimeOverDistance"), TimeOverDistance, TEXT("Time/Distance LOD alignment value in milliseconds/centimeter. Used as Distance * TimeOverDistance = Time, where Time is how far behind the current timeline replication should run. Start outside of DistanceFullPrediction and clamped by the received states timeline so we don't extrapolate backwards."));
+
+	bool bDrawDebugEnabled = false;
+	static FAutoConsoleVariableRef CVarDrawDebug(TEXT("p.ReplicationLOD.DrawDebug.Enabled"), bDrawDebugEnabled, TEXT(""));
+
+	float DrawDebugWorstLatency = 300.0f;
+	static FAutoConsoleVariableRef CVarDrawDebugWorstLatency(TEXT("p.ReplicationLOD.DrawDebug.WorstLatency"), DrawDebugWorstLatency, TEXT("Used to calculate color gradient in debug draw from forward predicted to worst expected latency."));
+}
+
+// ----------- Game Thread API -----------
+
+FPhysicsReplicationLOD::FPhysicsReplicationLOD(FPhysScene* InPhysicsScene)
+	: bEnabled(false)
+	, PhysScene(InPhysicsScene)
+	, PhysicsReplicationLODAsync(nullptr)
+{
+	// Apply project settings
+	FPhysicsReplicationLODSettings& LODSettings = UPhysicsSettings::Get()->PhysicsPrediction.PhysicsReplicationLODSettings;
+	bEnabled = LODSettings.bEnablePhysicsReplicationLOD;
+
+	if (Chaos::FPhysicsSolverBase* Solver = PhysScene->GetSolver())
+	{
+		// Create physics thread instance and cache project settings there
+		PhysicsReplicationLODAsync = Solver->CreateAndRegisterSimCallbackObject_External<FPhysicsReplicationLODAsync>();
+		PhysicsReplicationLODAsync->bEnabled = LODSettings.bEnablePhysicsReplicationLOD;
+		PhysicsReplicationLODAsync->DefaultSettings = LODSettings;
+	}
+}
+
+FPhysicsReplicationLOD::~FPhysicsReplicationLOD()
+{
+	if (PhysicsReplicationLODAsync)
+	{
+		if (Chaos::FPhysicsSolverBase* Solver = PhysScene->GetSolver())
+		{
+			Solver->UnregisterAndFreeSimCallbackObject_External(PhysicsReplicationLODAsync);
+		}
+	}
+}
+
+void FPhysicsReplicationLOD::SetEnabled(const bool bInEnabled)
+{
+	bEnabled = bInEnabled;
+	
+	// Marshal over the enabled setting to the physics thread
+	FPhysicsReplicationLODAsyncInput* AsyncInput = PhysicsReplicationLODAsync->GetProducerInputData_External();
+	AsyncInput->bEnabled = bEnabled;
+}
+
+bool FPhysicsReplicationLOD::IsEnabled() const
+{
+	return bEnabled || PhysicsReplicationLODCVars::bEnabled;
+}
+
+void FPhysicsReplicationLOD::RegisterFocalPoint_External(const UPrimitiveComponent* Component, FName BoneName)
+{
+	if (Component)
+	{
+		RegisterFocalPoint_External(Component->GetPhysicsObjectByName(BoneName));
+	}
+}
+
+void FPhysicsReplicationLOD::UnregisterFocalPoint_External(const UPrimitiveComponent* Component, FName BoneName)
+{
+	if (Component)
+	{
+		UnregisterFocalPoint_External(Component->GetPhysicsObjectByName(BoneName));
+	}
+}
+
+void FPhysicsReplicationLOD::RegisterFocalPoint_External(Chaos::FConstPhysicsObjectHandle PhysicsObject)
+{
+	if (PhysicsObject)
+	{
+		// Marshal registration over to physics thread
+		FPhysicsReplicationLODAsyncInput* AsyncInput = PhysicsReplicationLODAsync->GetProducerInputData_External();
+		AsyncInput->PhysicsObjectsToRegister.Add(PhysicsObject);
+	}
+}
+
+void FPhysicsReplicationLOD::UnregisterFocalPoint_External(Chaos::FConstPhysicsObjectHandle PhysicsObject)
+{
+	if (PhysicsObject)
+	{
+		// Marshal deregistration over to physics thread
+		FPhysicsReplicationLODAsyncInput* AsyncInput = PhysicsReplicationLODAsync->GetProducerInputData_External();
+		AsyncInput->PhysicsObjectsToUnregister.Add(PhysicsObject);
+	}
+}
+
+IPhysicsReplicationLODAsync* FPhysicsReplicationLOD::GetPhysicsReplicationLOD_Internal()
+{
+	return PhysicsReplicationLODAsync;
+}
+
+
+
+// ----------- Physics Thread API -----------
+
+bool FPhysicsReplicationLODAsync::IsEnabled() const
+{
+	return (bEnabled || PhysicsReplicationLODCVars::bEnabled); 
+}
+
+void FPhysicsReplicationLODAsync::OnPostInitialize_Internal()
+{
+	Chaos::FPBDRigidsSolver& RigidsSolver = GetSolver()->CastChecked();
+	
+	// Register this physics replication LOD in the solver
+	RigidsSolver.SetPhysicsReplicationLOD_Internal(this);
+}
+
+// ProcessInputs_Internal gets called before OnPreSimulate_Internal
+void FPhysicsReplicationLODAsync::ProcessInputs_Internal(int32 PhysicsStep)
+{
+	// Process incoming marshaled data from game thread
+	ConsumeAsyncInput();
+
+	// Called here instead of in OnPreSimulate_Internal for execution order, since Physics Replication might query the LOD system inside its OnPreSimulate_Internal 
+	CacheParticlesInFocalIslands();
+}
+
+void FPhysicsReplicationLODAsync::CacheParticlesInFocalIslands()
+{
+	Chaos::FPBDRigidsSolver& RigidsSolver = GetSolver()->CastChecked();
+
+	ParticlesInFocalIslands.Reset();
+
+	Chaos::Private::FPBDIslandManager& IslandManager = RigidsSolver.GetEvolution()->GetIslandManager();
+	Chaos::FReadPhysicsObjectInterface_Internal Interface = Chaos::FPhysicsObjectInternalInterface::GetRead();
+
+	const int32 FocalParticlesIndexMax = FocalParticles.Num() - 1;
+	for (int32 i = FocalParticlesIndexMax; i >= 0; i--)
+	{
+		Chaos::FConstPhysicsObjectHandle FocalParticle = FocalParticles[i];
+		Chaos::FGeometryParticleHandle* FocalParticleHandle = Interface.GetParticle(FocalParticle);
+		
+		if (!FocalParticleHandle)
+		{
+			// If focal particle is no longer valid, remove it
+			FocalParticles.RemoveAt(i);
+			continue;
+		}
+
+		// Get a list of particles from the same island as a focal particle is in, i.e. particles interacting with a resim particle
+		for (const Chaos::FGeometryParticleHandle* InteractParticle : IslandManager.FindParticlesInIslands(IslandManager.FindParticleIslands(FocalParticleHandle)))
+		{
+			if (InteractParticle->GetParticleType() != Chaos::EParticleType::Static)
+			{
+				// Add all particles that are in the same island as a focal particle to an array
+				ParticlesInFocalIslands.Add(InteractParticle->GetHandleIdx());
+			}
+		}
+	}
+}
+
+void FPhysicsReplicationLODAsync::ConsumeAsyncInput()
+{
+	if (const FPhysicsReplicationLODAsyncInput* AsyncInput = GetConsumerInput_Internal())
+	{
+		if (AsyncInput->bEnabled.IsSet())
+		{
+			bEnabled = *AsyncInput->bEnabled;
+		}
+
+		Chaos::FReadPhysicsObjectInterface_Internal Interface = Chaos::FPhysicsObjectInternalInterface::GetRead();
+		for (Chaos::FConstPhysicsObjectHandle PhysicsObject : AsyncInput->PhysicsObjectsToRegister)
+		{
+			RegisterFocalPoint_Internal(PhysicsObject);
+		}
+
+		for (Chaos::FConstPhysicsObjectHandle PhysicsObject : AsyncInput->PhysicsObjectsToUnregister)
+		{
+			UnregisterFocalPoint_Internal(PhysicsObject);
+		}
+	}
+}
+
+void FPhysicsReplicationLODAsync::OnPreSimulate_Internal()
+{
+#if !(UE_BUILD_SHIPPING || UE_BUILD_TEST)
+	if (PhysicsReplicationLODCVars::bDrawDebugEnabled)
+	{
+		Chaos::FReadPhysicsObjectInterface_Internal Interface = Chaos::FPhysicsObjectInternalInterface::GetRead();
+		for (Chaos::FConstPhysicsObjectHandle FocalParticle : FocalParticles)
+		{
+			if (Chaos::FGeometryParticleHandle* FocalParticleHandle = Interface.GetParticle(FocalParticle))
+			{
+				FColor DebugColor = IsEnabled() ? FLinearColor::MakeFromHSV8(0, 150, 255).ToFColor(false) : FColor::Black;
+				DebugColor.A = 255;
+				const float DistanceFullPrediction = PhysicsReplicationLODCVars::DistanceForFullPrediction > 0.0f ? PhysicsReplicationLODCVars::DistanceForFullPrediction : DefaultSettings.DistanceForFullPrediction;
+				Chaos::FDebugDrawQueue::GetInstance().DrawDebugSphere(FocalParticleHandle->GetX(), DistanceFullPrediction, 8, DebugColor, false, -1.0f, 0, 1.0f);
+			}
+		}
+	}
+#endif
+}
+
+void FPhysicsReplicationLODAsync::OnPhysicsObjectUnregistered_Internal(Chaos::FConstPhysicsObjectHandle PhysicsObject)
+{
+	UnregisterFocalPoint_Internal(PhysicsObject);
+}
+
+void FPhysicsReplicationLODAsync::RegisterFocalPoint_Internal(Chaos::FConstPhysicsObjectHandle PhysicsObject)
+{
+	if (!PhysicsObject)
+	{
+		return;
+	}
+
+	bool bIsRegistered = false;
+	for (int32 i = 0; i < FocalParticles.Num(); i++)
+	{
+		if (FocalParticles[i] == PhysicsObject)
+		{
+			bIsRegistered = true;
+			break;
+		}
+	}
+
+	if (!bIsRegistered)
+	{
+		FocalParticles.Add(PhysicsObject);
+	}
+}
+
+void FPhysicsReplicationLODAsync::UnregisterFocalPoint_Internal(Chaos::FConstPhysicsObjectHandle PhysicsObject)
+{
+	int32 NumRemoved = FocalParticles.Remove(PhysicsObject);
+}
+
+FPhysicsRepLodData* FPhysicsReplicationLODAsync::GetLODData_Internal(Chaos::FConstPhysicsObjectHandle PhysicsObject, const uint32 LODFlags)
+{
+	LodData.Reset();
+
+	if (FocalParticles.Num() == 0)
+	{
+		return nullptr;
+	}
+
+	Chaos::FReadPhysicsObjectInterface_Internal Interface = Chaos::FPhysicsObjectInternalInterface::GetRead();
+	const Chaos::FGeometryParticleHandle* ParticleHandle = Interface.GetParticle(PhysicsObject);
+	if (!ParticleHandle)
+	{
+		return nullptr;
+	}
+
+	// If PhysicsObject is a focal particle, early out
+	for (Chaos::FConstPhysicsObjectHandle FocalParticle : FocalParticles)
+	{
+		if (PhysicsObject == FocalParticle)
+		{
+			return nullptr;
+		}
+	}
+	
+	if (PerformIslandLOD(*ParticleHandle, LODFlags))
+	{
+		return &LodData;
+	}
+	
+	if (PerformDistanceLOD(*ParticleHandle, LODFlags))
+	{
+		return &LodData;
+	}
+	
+	return nullptr;
+}
+
+bool FPhysicsReplicationLODAsync::PerformIslandLOD(const Chaos::FGeometryParticleHandle& ParticleHandle, const uint32 LODFlags)
+{
+	if ((LODFlags & EPhysicsReplicationLODFlags::LODFlag_IslandCheck) == 0)
+	{
+		return false;
+	}
+
+	// Check if particle is in the same island as a focal particle
+	if (ParticlesInFocalIslands.Contains(ParticleHandle.GetHandleIdx()))
+	{
+		Chaos::FPBDRigidsSolver& RigidsSolver = GetSolver()->CastChecked();
+
+		LodData.DataAssigned = true;
+		LodData.ReplicationMode = EPhysicsReplicationMode::Resimulation;
+		LodData.AlignedFrame = RigidsSolver.GetCurrentFrame();
+		LodData.AlignedTime = 0.0f;
+	}
+	return LodData.DataAssigned;
+}
+
+bool FPhysicsReplicationLODAsync::PerformDistanceLOD(const Chaos::FGeometryParticleHandle& ParticleHandle, const uint32 LODFlags)
+{
+	if ((LODFlags & EPhysicsReplicationLODFlags::LODFlag_DistanceCheck) == 0)
+	{
+		return false;
+	}
+
+	Chaos::FPBDRigidsSolver& RigidsSolver = GetSolver()->CastChecked();
+	Chaos::FReadPhysicsObjectInterface_Internal Interface = Chaos::FPhysicsObjectInternalInterface::GetRead();
+
+	const float DistanceForResimulationMode = PhysicsReplicationLODCVars::DistanceForResimulationMode > 0.0f ? PhysicsReplicationLODCVars::DistanceForResimulationMode : DefaultSettings.DistanceForResimulationMode;
+	const float DistanceFullPrediction = PhysicsReplicationLODCVars::DistanceForFullPrediction > 0.0f ? PhysicsReplicationLODCVars::DistanceForFullPrediction : DefaultSettings.DistanceForFullPrediction;
+	const float TimeOverDistance = PhysicsReplicationLODCVars::TimeOverDistance > 0.0f ? PhysicsReplicationLODCVars::TimeOverDistance : DefaultSettings.TimeOverDistance;
+
+	for (Chaos::FConstPhysicsObjectHandle FocalParticle : FocalParticles)
+	{
+		Chaos::FGeometryParticleHandle* FocalParticleHandle = Interface.GetParticle(FocalParticle);
+		if (!FocalParticleHandle)
+		{
+			continue;
+		}
+
+		// Check distance towards focal particles
+		const Chaos::FVec3 PosOffset = FocalParticleHandle->GetX() - ParticleHandle.GetX();
+		const float Distance = PosOffset.Size();
+		float TimeAlignment = FMath::Max((Distance - DistanceFullPrediction), 0.0f) * TimeOverDistance;
+		TimeAlignment = (TimeAlignment * 0.001f); // Convert from ms to s
+		const float TickAlignment = TimeAlignment / RigidsSolver.GetAsyncDeltaTime();
+		const int32 TickAlignmentDiscrete = FMath::Floor(TickAlignment);
+
+#if !(UE_BUILD_SHIPPING || UE_BUILD_TEST)
+		if (PhysicsReplicationLODCVars::bDrawDebugEnabled)
+		{
+			PhysicsReplicationLODCVars::DrawDebugWorstLatency = FMath::Max(PhysicsReplicationLODCVars::DrawDebugWorstLatency, 1.0f); // Clamp to above 1ms, to not divide by 0 or negative value
+			float LodAlpha = (TimeAlignment * 1000.0f) / PhysicsReplicationLODCVars::DrawDebugWorstLatency; // Temporary solution for debug draw to use NumPredictedFrames
+			LodAlpha = FMath::Min(LodAlpha, 1.0f); // Clamp to 1 or lower
+
+			FColor DebugColor = FLinearColor::MakeFromHSV8((uint8)FMath::CeilToInt32((255.0f * 0.8f) * LodAlpha), 150, 255).ToFColor(false);
+			DebugColor.A = 255;
+			Chaos::FDebugDrawQueue::GetInstance().DrawDebugDirectionalArrow(ParticleHandle.GetX(), ParticleHandle.GetX() + PosOffset, 1.0f, DebugColor, false, 5, 0, 1.0f);
+		}
+#endif
+
+		// Update LodData and keep the highest LOD data
+		LodData.AlignedTime = FMath::Max(LodData.AlignedTime, TimeAlignment);
+		LodData.AlignedFrame = FMath::Max(LodData.AlignedFrame, RigidsSolver.GetCurrentFrame() - TickAlignmentDiscrete);
+		LodData.DataAssigned = true;
+
+		if (Distance <= DistanceForResimulationMode)
+		{
+			LodData.ReplicationMode = EPhysicsReplicationMode::Resimulation;
+			break; // Highest LOD reached, don't check further focal particles
+		}
+	}
+
+	return LodData.DataAssigned;
+}
diff --git a/Engine/Source/Runtime/Engine/Private/PrimitiveComponentPhysics.cpp b/Engine/Source/Runtime/Engine/Private/PrimitiveComponentPhysics.cpp
--- a/Engine/Source/Runtime/Engine/Private/PrimitiveComponentPhysics.cpp
+++ b/Engine/Source/Runtime/Engine/Private/PrimitiveComponentPhysics.cpp
@@ -69,6 +69,34 @@ void UPrimitiveComponent::SetRigidBodyReplicatedTarget(FRigidBodyState& UpdatedS
 	}
 }
 
+void UPrimitiveComponent::RegisterAsFocalPointInPhysicsReplicationLOD() const
+{
+	if (UWorld* World = GetWorld())
+	{
+		if (FPhysScene* PhysScene = World->GetPhysicsScene())
+		{
+			if (IPhysicsReplicationLOD* PhysicsReplicationLOD = PhysScene->GetPhysicsReplicationLOD())
+			{
+				PhysicsReplicationLOD->RegisterFocalPoint_External(this);
+			}
+		}
+	}
+}
+
+void UPrimitiveComponent::UnregisterAsFocalPointInPhysicsReplicationLOD() const
+{
+	if (UWorld* World = GetWorld())
+	{
+		if (FPhysScene* PhysScene = World->GetPhysicsScene())
+		{
+			if (IPhysicsReplicationLOD* PhysicsReplicationLOD = PhysScene->GetPhysicsReplicationLOD())
+			{
+				PhysicsReplicationLOD->UnregisterFocalPoint_External(this);
+			}
+		}
+	}
+}
+
 bool UPrimitiveComponent::GetRigidBodyState(FRigidBodyState& OutState, FName BoneName)
 {
 	// If we have a BodyInstance, use it
diff --git a/Engine/Source/Runtime/Engine/Public/Physics/Experimental/PhysInterface_Chaos.h b/Engine/Source/Runtime/Engine/Public/Physics/Experimental/PhysInterface_Chaos.h
--- a/Engine/Source/Runtime/Engine/Public/Physics/Experimental/PhysInterface_Chaos.h
+++ b/Engine/Source/Runtime/Engine/Public/Physics/Experimental/PhysInterface_Chaos.h
@@ -21,7 +21,6 @@ static int32 NextConstraintIdValue = 0;
 class FPhysInterface_Chaos;
 struct FBodyInstance;
 struct FPhysxUserData;
-class IPhysicsReplicationFactory;
 struct FConstraintDrive;
 struct FLinearDriveConstraint;
 struct FAngularDriveConstraint;
diff --git a/Engine/Source/Runtime/Engine/Public/Physics/Experimental/PhysScene_Chaos.h b/Engine/Source/Runtime/Engine/Public/Physics/Experimental/PhysScene_Chaos.h
--- a/Engine/Source/Runtime/Engine/Public/Physics/Experimental/PhysScene_Chaos.h
+++ b/Engine/Source/Runtime/Engine/Public/Physics/Experimental/PhysScene_Chaos.h
@@ -28,6 +28,7 @@ class UPrimitiveComponent;
 
 class AdvanceOneTimeStepTask;
 class IPhysicsReplication;
+class IPhysicsReplicationLOD;
 class FPhysicsReplicationCache;
 class FPhysInterface_Chaos;
 class FChaosSolversModule;
@@ -169,8 +170,9 @@ public:
 
 	ENGINE_API IPhysicsReplication* CreatePhysicsReplication();
 
-	UE_DEPRECATED(5.3, "If possible, avoid directly setting physics replication at runtime - this function will take ownership of the IPhysicsReplication's lifetime. Instead, specify a PhysicsReplication factory and if you must change the replication class at runtime run RecreatePhysicsReplication after having set the factory.")
-	ENGINE_API void SetPhysicsReplication(IPhysicsReplication* InPhysicsReplication);
+	ENGINE_API IPhysicsReplicationLOD* GetPhysicsReplicationLOD();
+
+	ENGINE_API IPhysicsReplicationLOD* CreatePhysicsReplicationLOD();
 
 	ENGINE_API virtual void AddReferencedObjects(FReferenceCollector& Collector) override;
 	/** Given a solver object, returns its associated component. */
@@ -388,9 +390,12 @@ private:
 	ENGINE_API void HandleRemovalEvents(const Chaos::FRemovalEventData& Event);
 	ENGINE_API void HandleCrumblingEvents(const Chaos::FCrumblingEventData& Event);
 
-	/** Replication manager that updates physics bodies towards replicated physics state */
+	/** Physics Replication manager which updates physics particles state to replicate towards networked target states */
 	TUniquePtr<IPhysicsReplication> PhysicsReplication;
 
+	/** Physics Replication "Level of Detail" that recommends physics replication mode and how much forward prediction should be applied for replicated physics particles */
+	TUniquePtr<IPhysicsReplicationLOD> PhysicsReplicationLOD;
+
 #if CHAOS_WITH_PAUSABLE_SOLVER
 	/** Callback that checks the status of the world settings for this scene before pausing/unpausing its solver. */
 	ENGINE_API void OnUpdateWorldPause();
diff --git a/Engine/Source/Runtime/Engine/Public/Physics/NetworkPhysicsSettingsComponent.h b/Engine/Source/Runtime/Engine/Public/Physics/NetworkPhysicsSettingsComponent.h
--- a/Engine/Source/Runtime/Engine/Public/Physics/NetworkPhysicsSettingsComponent.h
+++ b/Engine/Source/Runtime/Engine/Public/Physics/NetworkPhysicsSettingsComponent.h
@@ -87,6 +87,10 @@ struct FNetworkPhysicsSettings
 	// Override the EPhysicsReplicationMode for Actors with ENetRole::ROLE_SimulatedProxy.
 	UPROPERTY(config, EditDefaultsOnly, Category = "Overrides", Meta = (EditCondition = "bOverrideSimProxyRepMode"))
 	EPhysicsReplicationMode SimProxyRepMode = EPhysicsReplicationMode::PredictiveInterpolation;
+	
+	// Register this actors Autonomous Proxy as a focal point / focal particle in Physics Replication LOD.
+	UPROPERTY(config, EditDefaultsOnly, Category = "Overrides")
+	bool bFocalParticleInPhysicsReplicationLOD = false;
 };
 
 USTRUCT()
@@ -448,6 +452,8 @@ class UNetworkPhysicsSettingsComponent : public UActorComponent
 public:
 	ENGINE_API UNetworkPhysicsSettingsComponent();
 
+	virtual void OnRegister() override;
+	virtual void OnUnregister() override;
 	virtual void InitializeComponent() override;
 	virtual void UninitializeComponent() override;
 
@@ -456,6 +462,13 @@ public:
 	/** Get the settings internal to the PhysicsThread (Only access construct on the Physics Thread) */
 	FNetworkPhysicsSettingsComponentAsync* GetNetworkPhysicsSettings_Internal() const { return NetworkPhysicsSettings_Internal; };
 
+private:
+	UFUNCTION()
+	void OnComponentPhysicsStateChanged(UPrimitiveComponent* ChangedComponent, EComponentPhysicsStateChange StateChange);
+	
+	/** Perform logic to register this actor into physics replication LOD, depending on settings */
+	void RegisterInPhysicsReplicationLOD();
+
 public:
 	UPROPERTY(EditDefaultsOnly, Category = "Networked Physics Settings")
 	FNetworkPhysicsSettings GeneralSettings;
@@ -511,23 +524,30 @@ struct FNetworkPhysicsSettingsAsync
 
 struct FNetworkPhysicsSettingsAsyncInput : public Chaos::FSimCallbackInput
 {
-	Chaos::FConstPhysicsObjectHandle PhysicsObject;
-	FNetworkPhysicsSettingsAsync Settings;
+	TOptional<Chaos::FConstPhysicsObjectHandle> PhysicsObject;
+	TOptional<FNetworkPhysicsSettingsAsync> Settings;
 
 	void Reset()
 	{
-		Settings = FNetworkPhysicsSettingsAsync();
+		PhysicsObject.Reset();
+		Settings.Reset();
 	}
 };
 
 class FNetworkPhysicsSettingsComponentAsync : public Chaos::TSimCallbackObject<FNetworkPhysicsSettingsAsyncInput>
 {
 public:
-	virtual void OnPostInitialize_Internal() override;
-	virtual void OnPreSimulate_Internal() override { };
+	virtual void OnPreSimulate_Internal() override;
+
+private:
+	void ConsumeAsyncInput();
+	void RegisterSettingsInPhysicsReplication();
 
 public:
 	FNetworkPhysicsSettingsAsync Settings;
+
+private:
+	Chaos::FConstPhysicsObjectHandle PhysicsObject;
 };
 
 #pragma endregion // FNetworkPhysicsSettingsComponentAsync
diff --git a/Engine/Source/Runtime/Engine/Public/Physics/PhysicsInterfaceUtils.h b/Engine/Source/Runtime/Engine/Public/Physics/PhysicsInterfaceUtils.h
--- a/Engine/Source/Runtime/Engine/Public/Physics/PhysicsInterfaceUtils.h
+++ b/Engine/Source/Runtime/Engine/Public/Physics/PhysicsInterfaceUtils.h
@@ -7,6 +7,7 @@
 #include "Physics/PhysicsInterfaceTypes.h"
 #include "PhysicsInterfaceUtilsCore.h"
 #include "PhysicsReplicationInterface.h"
+#include "PhysicsReplicationLODInterface.h"
 #include "Misc/CoreMiscDefines.h"
 
 class FPhysScene_PhysX;
@@ -30,18 +31,11 @@ struct FConstraintBrokenDelegateData
 	int32 ConstraintIndex;
 };
 
-class FPhysicsReplication;
-
 /** Interface for the creation of customized physics replication.*/
 class IPhysicsReplicationFactory
 {
 public:
 
-	// NOTE: Once the old Create/Destroy methods are deprecated, remove the default implementation!
-	ENGINE_API virtual TUniquePtr<IPhysicsReplication> CreatePhysicsReplication(FPhysScene* OwningPhysScene);
-
-	UE_DEPRECATED(5.3, "Use CreatePhysicsReplication (which gives a replication interface) instead")
-	virtual FPhysicsReplication* Create(FPhysScene* OwningPhysScene) { return nullptr; }
-	UE_DEPRECATED(5.3, "No longer used - PhysScene owns the physics replication ptr once it is created. If you need a custom deleter, override CreatePhysicsReplication and specify the deleter in TUniquePtr.")
-	virtual void Destroy(FPhysicsReplication* PhysicsReplication) { }
+	ENGINE_API virtual TUniquePtr<IPhysicsReplication> CreatePhysicsReplication(FPhysScene* OwningPhysScene) { return nullptr; }
+	ENGINE_API virtual TUniquePtr<IPhysicsReplicationLOD> CreatePhysicsReplicationLOD(FPhysScene* OwningPhysScene) { return nullptr; }
 };
diff --git a/Engine/Source/Runtime/Engine/Public/PhysicsReplication.h b/Engine/Source/Runtime/Engine/Public/PhysicsReplication.h
--- a/Engine/Source/Runtime/Engine/Public/PhysicsReplication.h
+++ b/Engine/Source/Runtime/Engine/Public/PhysicsReplication.h
@@ -193,7 +193,8 @@ public:
 	}
 };
 
-class FPhysicsReplicationAsync : public Chaos::TSimCallbackObject<
+class FPhysicsReplicationAsync : public IPhysicsReplicationAsync,
+	public Chaos::TSimCallbackObject<
 	FPhysicsReplicationAsyncInput,
 	Chaos::FSimCallbackNoOutput,
 	Chaos::ESimCallbackOptions::Presimulate | Chaos::ESimCallbackOptions::PhysicsObjectUnregister>
@@ -212,7 +213,7 @@ class FPhysicsReplicationAsync : public Chaos::TSimCallbackObject<
 	virtual bool ResimulationReplication(Chaos::FPBDRigidParticleHandle* Handle, FReplicatedPhysicsTargetAsync& Target, const float DeltaSeconds);
 
 public:
-	virtual void RegisterSettings(Chaos::FConstPhysicsObjectHandle PhysicsObject, FNetworkPhysicsSettingsAsync InSettings);
+	virtual void RegisterSettings(Chaos::FConstPhysicsObjectHandle PhysicsObject, FNetworkPhysicsSettingsAsync InSettings) override;
 
 private:
 	float LatencyOneWay;
@@ -232,7 +233,13 @@ private:
 	void CacheResimInteractions();
 	// Sets SettingsCurrent to either the objects custom settings or to the default settings
 	void FetchObjectSettings(Chaos::FConstPhysicsObjectHandle PhysicsObject); 
+	bool UsePhysicsReplicationLOD();
+	void ApplyPhysicsReplicationLOD(Chaos::FConstPhysicsObjectHandle PhysicsObjectHandle, FReplicatedPhysicsTargetAsync& Target, const uint32 LODFlags);
+	
+	/** Static function to extrapolate a target for N ticks using X DeltaSeconds */
 	static void ExtrapolateTarget(FReplicatedPhysicsTargetAsync& Target, const int32 ExtrapolateFrames, const float DeltaSeconds);
+	/** Static function to extrapolate a target for N Seconds */
+	static void ExtrapolateTarget(FReplicatedPhysicsTargetAsync& Target, const float ExtrapolationTime);
 
 public:
 	void Setup(FRigidBodyErrorCorrection ErrorCorrection)
diff --git a/Engine/Source/Runtime/Engine/Public/PhysicsReplicationInterface.h b/Engine/Source/Runtime/Engine/Public/PhysicsReplicationInterface.h
--- a/Engine/Source/Runtime/Engine/Public/PhysicsReplicationInterface.h
+++ b/Engine/Source/Runtime/Engine/Public/PhysicsReplicationInterface.h
@@ -3,11 +3,13 @@
 #pragma once
 
 #include "CoreMinimal.h"
+#include "Chaos/PhysicsObject.h"
 
 class UPrimitiveComponent;
 struct FRigidBodyState;
+struct FNetworkPhysicsSettingsAsync;
 
-class IPhysicsReplication
+class IPhysicsReplication // Game Thread API
 {
 public:
 	virtual ~IPhysicsReplication() { }
@@ -18,3 +20,11 @@ public:
 
 	virtual void RemoveReplicatedTarget(UPrimitiveComponent* Component) = 0;
 };
+
+class IPhysicsReplicationAsync // Physics Thread API
+{
+public:
+	virtual ~IPhysicsReplicationAsync() { }
+
+	virtual void RegisterSettings(Chaos::FConstPhysicsObjectHandle PhysicsObject, FNetworkPhysicsSettingsAsync InSettings) = 0;
+};
diff --git a/Engine/Source/Runtime/Engine/Public/PhysicsReplicationLOD.h b/Engine/Source/Runtime/Engine/Public/PhysicsReplicationLOD.h
new file mode 100644
--- /dev/null
+++ b/Engine/Source/Runtime/Engine/Public/PhysicsReplicationLOD.h
@@ -0,0 +1,130 @@
+// Copyright Epic Games, Inc. All Rights Reserved.
+
+#pragma once
+
+#include "Engine/EngineTypes.h"
+#include "PhysicsReplicationLODInterface.h"
+#include "Chaos/SimCallbackObject.h"
+#include "Physics/PhysicsInterfaceDeclares.h"
+#include "Physics/PhysicsInterfaceUtils.h"
+#include "PhysicsEngine/PhysicsSettings.h"
+
+
+enum EPhysicsReplicationLODFlags : uint32
+{
+	LODFlag_None = 0,
+	LODFlag_IslandCheck = (1 << 0),
+	LODFlag_DistanceCheck = (1 << 1),
+	LODFlag_All = UINT_MAX,
+};
+
+class UPrimitiveComponent;
+class FPhysicsReplicationLODAsync;
+
+class FPhysicsReplicationLOD : public IPhysicsReplicationLOD
+{
+public:
+	FPhysicsReplicationLOD(FPhysScene* PhysScene);
+
+	virtual ~FPhysicsReplicationLOD();
+
+public: // IPhysicsReplicationLOD Functions
+
+	virtual void SetEnabled(const bool bInEnabled) override;
+	virtual bool IsEnabled() const override;
+
+	/** Register a components physics object as a Focal Particle in LOD */
+	virtual void RegisterFocalPoint_External(const UPrimitiveComponent* Component, FName BoneName = NAME_None) override;
+	/** Unregister a components physics object as a Focal Particle in LOD */
+	virtual void UnregisterFocalPoint_External(const UPrimitiveComponent* Component, FName BoneName = NAME_None) override;
+
+	/** Register PhysicsObject as a Focal Particle in LOD */
+	virtual void RegisterFocalPoint_External(Chaos::FConstPhysicsObjectHandle PhysicsObject) override;
+	/** Unregister PhysicsObject as a Focal Particle in LOD */
+	virtual void UnregisterFocalPoint_External(Chaos::FConstPhysicsObjectHandle PhysicsObject) override;
+	
+	/** Get the PhysicsThread instance of the PhysicsReplicationLOD */
+	virtual IPhysicsReplicationLODAsync* GetPhysicsReplicationLOD_Internal() override;
+
+private:
+	bool bEnabled;
+	FPhysScene* PhysScene;
+	FPhysicsReplicationLODAsync* PhysicsReplicationLODAsync;
+};
+
+
+
+
+// ------------------------------- Async -------------------------------
+
+
+struct FPhysicsReplicationLODAsyncInput : public Chaos::FSimCallbackInput
+{
+	TOptional<bool> bEnabled;
+	TArray<Chaos::FConstPhysicsObjectHandle> PhysicsObjectsToRegister;
+	TArray<Chaos::FConstPhysicsObjectHandle> PhysicsObjectsToUnregister;
+
+	void Reset()
+	{
+		bEnabled.Reset();
+		PhysicsObjectsToRegister.Reset();
+		PhysicsObjectsToUnregister.Reset();
+	}
+};
+
+
+class FPhysicsReplicationLODAsync : public IPhysicsReplicationLODAsync
+	, public Chaos::TSimCallbackObject<
+	FPhysicsReplicationLODAsyncInput,
+	Chaos::FSimCallbackNoOutput,
+	Chaos::ESimCallbackOptions::Presimulate | Chaos::ESimCallbackOptions::Rewind | Chaos::ESimCallbackOptions::PhysicsObjectUnregister>
+{
+friend FPhysicsReplicationLOD;
+
+public:
+
+	FPhysicsReplicationLODAsync() 
+	: bEnabled(false)
+	, DefaultSettings()
+	, LodData()
+	{ 
+		FocalParticles.Reset(1);
+		ParticlesInFocalIslands.Reset(1);
+	}
+
+	virtual ~FPhysicsReplicationLODAsync() { }
+
+private: // TSimCallbackObject Functions
+	virtual void OnPostInitialize_Internal() override;
+	virtual void ProcessInputs_Internal(int32 PhysicsStep) override; // Called before OnPreSimulate_Internal()
+	virtual void OnPreSimulate_Internal() override;
+	virtual void OnPhysicsObjectUnregistered_Internal(Chaos::FConstPhysicsObjectHandle PhysicsObject) override;
+
+
+public: // IPhysicsReplicationLODAsync Functions
+	virtual bool IsEnabled() const override;
+
+	/** Register PhysicsObject as a Focal Particle in LOD */
+	virtual void RegisterFocalPoint_Internal(Chaos::FConstPhysicsObjectHandle PhysicsObject) override;
+	/** Unregister PhysicsObject as a Focal Particle in LOD */
+	virtual void UnregisterFocalPoint_Internal(Chaos::FConstPhysicsObjectHandle PhysicsObject) override;
+
+	/** Receive the LOD data for @param PhysicsObject, based on its relation to registered Focal Particles in LOD.
+	* NOTE: @param PhysicsObject will no be manipulated by the LOD, the returning FPhysicsRepLodData has data and recommendation that can be used by the caller.
+	* @param LODFlags is of type EPhysicsReplicationLODFlags */
+	virtual FPhysicsRepLodData* GetLODData_Internal(Chaos::FConstPhysicsObjectHandle PhysicsObject, const uint32 LODFlags = EPhysicsReplicationLODFlags::LODFlag_All) override;
+
+private:
+	void ConsumeAsyncInput();
+	void CacheParticlesInFocalIslands();
+
+	bool PerformIslandLOD(const Chaos::FGeometryParticleHandle& ParticleHandle, const uint32 LODFlags);
+	bool PerformDistanceLOD(const Chaos::FGeometryParticleHandle& ParticleHandle, const uint32 LODFlags);
+
+private:
+	bool bEnabled;
+	FPhysicsReplicationLODSettings DefaultSettings;
+	TArray<Chaos::FConstPhysicsObjectHandle> FocalParticles;
+	TArray<int32> ParticlesInFocalIslands;
+	FPhysicsRepLodData LodData;
+};
diff --git a/Engine/Source/Runtime/Engine/Public/PhysicsReplicationLODInterface.h b/Engine/Source/Runtime/Engine/Public/PhysicsReplicationLODInterface.h
new file mode 100644
--- /dev/null
+++ b/Engine/Source/Runtime/Engine/Public/PhysicsReplicationLODInterface.h
@@ -0,0 +1,65 @@
+// Copyright Epic Games, Inc. All Rights Reserved.
+
+#pragma once
+
+#include "CoreMinimal.h"
+#include "Chaos/PhysicsObject.h"
+#include "Engine/EngineTypes.h"
+
+class UPrimitiveComponent;
+class IPhysicsReplicationLODAsync;
+
+class IPhysicsReplicationLOD // Game Thread API
+{
+public:
+	virtual ~IPhysicsReplicationLOD() { }
+
+	virtual void SetEnabled(const bool bInEnabled) = 0;
+	virtual bool IsEnabled() const = 0;
+
+	virtual void RegisterFocalPoint_External(const UPrimitiveComponent* Component, FName BoneName = NAME_None) = 0;
+	virtual void UnregisterFocalPoint_External(const UPrimitiveComponent* Component, FName BoneName = NAME_None) = 0;
+	virtual void RegisterFocalPoint_External(Chaos::FConstPhysicsObjectHandle PhysicsObject) = 0;
+	virtual void UnregisterFocalPoint_External(Chaos::FConstPhysicsObjectHandle PhysicsObject) = 0;
+
+	virtual IPhysicsReplicationLODAsync* GetPhysicsReplicationLOD_Internal() = 0;
+};
+
+struct FPhysicsRepLodData
+{
+	FPhysicsRepLodData()
+	: DataAssigned(false)
+	, AlignedTime(0.0f)
+	, AlignedFrame(0)
+	, ReplicationMode(EPhysicsReplicationMode::PredictiveInterpolation)
+	, bKinematic(false)
+	{};
+
+	bool DataAssigned;
+	float AlignedTime;
+	int32 AlignedFrame;
+	EPhysicsReplicationMode ReplicationMode;
+	bool bKinematic;
+
+	void Reset()
+	{
+		DataAssigned = false;
+		AlignedTime = 0.0f;
+		AlignedFrame = 0;
+		ReplicationMode = EPhysicsReplicationMode::PredictiveInterpolation;
+		bKinematic = false;
+	}
+};
+
+class IPhysicsReplicationLODAsync // Physics Thread API
+{
+public:
+	virtual ~IPhysicsReplicationLODAsync() { }
+
+	virtual bool IsEnabled() const = 0;
+
+	virtual void RegisterFocalPoint_Internal(Chaos::FConstPhysicsObjectHandle PhysicsObject) = 0;
+	virtual void UnregisterFocalPoint_Internal(Chaos::FConstPhysicsObjectHandle PhysicsObject) = 0;
+
+	virtual FPhysicsRepLodData* GetLODData_Internal(Chaos::FConstPhysicsObjectHandle PhysicsObject, const uint32 LODFlags = UINT_MAX) = 0;
+};
diff --git a/Engine/Source/Runtime/Experimental/Chaos/Public/Chaos/Framework/PhysicsSolverBase.h b/Engine/Source/Runtime/Experimental/Chaos/Public/Chaos/Framework/PhysicsSolverBase.h
--- a/Engine/Source/Runtime/Experimental/Chaos/Public/Chaos/Framework/PhysicsSolverBase.h
+++ b/Engine/Source/Runtime/Experimental/Chaos/Public/Chaos/Framework/PhysicsSolverBase.h
@@ -23,6 +23,8 @@
 
 class FChaosSolversModule;
 class FPhysicsReplicationAsync;
+class IPhysicsReplicationAsync;
+class IPhysicsReplicationLODAsync;
 
 DECLARE_MULTICAST_DELEGATE_OneParam(FSolverPreAdvance, Chaos::FReal);
 DECLARE_MULTICAST_DELEGATE_OneParam(FSolverPreBuffer, Chaos::FReal);
@@ -459,14 +461,36 @@ namespace Chaos
 			return bUseCollisionResimCache;
 		}
 
-		void SetPhysicsReplication(FPhysicsReplicationAsync* InPhysicsReplication)
+		UE_DEPRECATED(5.6, "Use SetPhysicsReplication_Internal")
+			void SetPhysicsReplication(FPhysicsReplicationAsync* InPhysicsReplication)
 		{
-			PhysicsReplication = InPhysicsReplication;
+			PhysicsReplication_DEPRECATED = InPhysicsReplication;
 		}
 
-		FPhysicsReplicationAsync* GetPhysicsReplication()
+		UE_DEPRECATED(5.6, "Use GetPhysicsReplication_Internal which return the interface instead of the implementation")
+			FPhysicsReplicationAsync* GetPhysicsReplication()
 		{
-			return PhysicsReplication;
+			return PhysicsReplication_DEPRECATED;
+		}
+
+		void SetPhysicsReplication_Internal(IPhysicsReplicationAsync* InPhysicsReplication)
+		{
+			PhysicsReplicationAsync = InPhysicsReplication;
+		}
+
+		IPhysicsReplicationAsync* GetPhysicsReplication_Internal()
+		{
+			return PhysicsReplicationAsync;
+		}
+
+		void SetPhysicsReplicationLOD_Internal(IPhysicsReplicationLODAsync* InPhysicsReplicationLOD)
+		{
+			PhysicsReplicationLODAsync = InPhysicsReplicationLOD;
+		}
+
+		IPhysicsReplicationLODAsync* GetPhysicsReplicationLOD_Internal()
+		{
+			return PhysicsReplicationLODAsync;
 		}
 
 		//Used as helper for GT to go from unique idx back to gt particle
@@ -890,7 +914,9 @@ namespace Chaos
 
 		TUniquePtr<FRewindData> MRewindData;
 		TUniquePtr<IRewindCallback> MRewindCallback;
-		FPhysicsReplicationAsync* PhysicsReplication;
+		FPhysicsReplicationAsync* PhysicsReplication_DEPRECATED;
+		IPhysicsReplicationAsync* PhysicsReplicationAsync;
+		IPhysicsReplicationLODAsync* PhysicsReplicationLODAsync;
 
 		bool bUseCollisionResimCache;
 
