[HWRT] Support streaming of FStaticMeshRayTracingProxyLOD with bOwnsBuffers == true.

Context:
- Streaming logic in FRayTracingGeometryManager only handled streaming of offline built BVH not other associated data such as VB/IB.
- Different asset types might serialize different data/layouts so responsibility of managing that shouldn't be on FRayTracingGeometryManager.

Changes:
- new FRayTracingStreamableAsset interface that asset types can implement to support streaming related to ray tracing geometry.
- modified FRayTracingGeometryManager to issue stream in/out requets via FRayTracingStreamableAsset
- FStaticMeshRayTracingProxyLOD implements FRayTracingStreamableAsset interface.
- Modified serialization of FStaticMeshRayTracingProxyLOD to track offset and size of offline BVH and buffers in serialized data blob.

diff --git a/Engine/Source/Runtime/Core/Private/UObject/FortniteMainBranchObjectVersion.cpp b/Engine/Source/Runtime/Core/Private/UObject/FortniteMainBranchObjectVersion.cpp
--- a/Engine/Source/Runtime/Core/Private/UObject/FortniteMainBranchObjectVersion.cpp
+++ b/Engine/Source/Runtime/Core/Private/UObject/FortniteMainBranchObjectVersion.cpp
@@ -13,7 +13,7 @@ TMap<FGuid, FGuid> FFortniteMainBranchObjectVersion::GetSystemGuids()
 	SystemGuids.Add(DevGuids.NIAGARASHADERMAP_DERIVEDDATA_VER, FGuid("C417A50271B3427D98785388266E3B66"));
 	SystemGuids.Add(DevGuids.Niagara_LatestScriptCompileVersion, FGuid("393081597FA7234C92F00869DE262D73"));
 	SystemGuids.Add(DevGuids.SkeletalMeshDerivedDataVersion, FGuid("B9AEAA3EE7AC4FCFB3860E8F3E06DD75"));
-	SystemGuids.Add(DevGuids.STATICMESH_DERIVEDDATA_VER, FGuid("250DE3CCCEE241528830D27086F51CDE"));
+	SystemGuids.Add(DevGuids.STATICMESH_DERIVEDDATA_VER, FGuid("941D874D601A435EA52B595F6164C31A"));
 	SystemGuids.Add(DevGuids.MaterialTranslationDDCVersion, FGuid("D318C2D188414B7284D232E41F77BB74"));
 	return SystemGuids;
 }
diff --git a/Engine/Source/Runtime/Engine/Private/Rendering/RayTracingGeometryManager.cpp b/Engine/Source/Runtime/Engine/Private/Rendering/RayTracingGeometryManager.cpp
--- a/Engine/Source/Runtime/Engine/Private/Rendering/RayTracingGeometryManager.cpp
+++ b/Engine/Source/Runtime/Engine/Private/Rendering/RayTracingGeometryManager.cpp
@@ -12,7 +12,7 @@
 #include "RayTracingGeometry.h"
 #include "RenderUtils.h"
 
-#include "Serialization/MemoryReader.h"
+#include "Rendering/RayTracingStreamableAsset.h"
 
 #include "Math/UnitConversion.h"
 
@@ -354,19 +354,15 @@ void FRayTracingGeometryManager::ReleaseRayTracingGeometryHandle(RayTracingGeome
 	DEC_DWORD_STAT(STAT_RayTracingGeometryCount);
 }
 
-void FRayTracingGeometryManager::SetRayTracingGeometryStreamingData(const FRayTracingGeometry* Geometry, FByteBulkData& BulkData, uint32 Offset, uint32 Size)
+void FRayTracingGeometryManager::SetRayTracingGeometryStreamingData(const FRayTracingGeometry* Geometry, FRayTracingStreamableAsset& StreamableAsset)
 {
 	FScopeLock ScopeLock(&MainCS);
 
 	checkf(RegisteredGeometries.IsValidIndex(Geometry->RayTracingGeometryHandle), TEXT("SetRayTracingGeometryStreamingData(...) can only be used with FRayTracingGeometry that has been registered with FRayTracingGeometryManager."));
 
-	FRayTracingAccelerationStructureOfflineMetadata OfflineMetadata = RHIGetRayTracingGeometryOfflineMetadata(Geometry->Initializer.OfflineDataHeader);
-	check(OfflineMetadata.SerializedSize <= Size);
-
 	FRegisteredGeometry& RegisteredGeometry = RegisteredGeometries[Geometry->RayTracingGeometryHandle];
-	RegisteredGeometry.StreamableData = &BulkData;
-	RegisteredGeometry.StreamableDataOffset = Offset + OfflineMetadata.SerializedOffset;
-	RegisteredGeometry.StreamableDataSize = OfflineMetadata.SerializedSize;
+	RegisteredGeometry.StreamableAsset = &StreamableAsset;
+	RegisteredGeometry.StreamableDataSize = StreamableAsset.GetRequestSize();
 }
 
 
@@ -885,8 +881,6 @@ bool FRayTracingGeometryManager::RequestRayTracingGeometryStreamIn(FRHICommandLi
 		}
 	}
 
-	FByteBulkData* StreamableData = RegisteredGeometry.StreamableData;
-
 	TResourceArray<uint8> RawData;
 	RawData.SetAllowCPUAccess(true);
 
@@ -898,29 +892,8 @@ bool FRayTracingGeometryManager::RequestRayTracingGeometryStreamIn(FRHICommandLi
 
 		RegisteredGeometry.Status = FRegisteredGeometry::FStatus::StreamedIn;
 	}
-	else if (StreamableData->IsBulkDataLoaded())
-	{
-		{
-			const uint8* Ptr = (const uint8*)StreamableData->LockReadOnly();
-
-			FMemoryView MemView(Ptr + RegisteredGeometry.StreamableDataOffset, RegisteredGeometry.StreamableDataSize);
-
-			FMemoryReaderView MemReader(MemView, true);
-			RawData.BulkSerialize(MemReader);
-			StreamableData->Unlock();
-		}
-
-		if (!RawData.IsEmpty())
-		{
-			OfflineData = &RawData;
-		}
-
-		RegisteredGeometry.Status = FRegisteredGeometry::FStatus::StreamedIn;
-	}
 	else
 	{
-		checkf(StreamableData->CanLoadFromDisk(), TEXT("Bulk data is not loaded and cannot be loaded from disk!"));
-		check(!StreamableData->IsStoredCompressedOnDisk()); // We do not support compressed Bulkdata for this system. Limitation of the streaming request/bulk data
 
 		if (NumStreamingRequests >= GRayTracingStreamingMaxPendingRequests)
 		{
@@ -938,12 +911,7 @@ bool FRayTracingGeometryManager::RequestRayTracingGeometryStreamIn(FRHICommandLi
 		StreamingRequest.GeometryHandle = Geometry->RayTracingGeometryHandle;
 		StreamingRequest.RequestBuffer = FIoBuffer(RegisteredGeometry.StreamableDataSize); // TODO: Use FIoBuffer::Wrap with preallocated memory
 
-		// TODO: We're currently using a single batch per request so we can individually cancel and wait on requests.
-		// This isn't ideal and should be revisited in the future.
-		TRACE_IOSTORE_METADATA_SCOPE_TAG("RaytracingGeometry");
-		FBulkDataBatchRequest::FScatterGatherBuilder Batch = FBulkDataBatchRequest::ScatterGather(1);
-		Batch.Read(*StreamableData, RegisteredGeometry.StreamableDataOffset, RegisteredGeometry.StreamableDataSize);
-		Batch.Issue(StreamingRequest.RequestBuffer, AIOP_Low, [](FBulkDataRequest::EStatus) {}, StreamingRequest.Request);
+		RegisteredGeometry.StreamableAsset->IssueRequest(StreamingRequest.Request, StreamingRequest.RequestBuffer);
 
 		RegisteredGeometry.Status = FRegisteredGeometry::FStatus::Streaming;
 	}
@@ -1008,56 +976,16 @@ void FRayTracingGeometryManager::ProcessCompletedStreamingRequests(FRHICommandLi
 			if (!PendingRequest.Request.IsOk())
 			{
 				// Retry if IO request failed for some reason
-
-				FByteBulkData* StreamableData = RegisteredGeometry.StreamableData;
-
-				TRACE_IOSTORE_METADATA_SCOPE_TAG("RaytracingGeometry (Retry)");
-				FBulkDataBatchRequest::FScatterGatherBuilder Batch = FBulkDataBatchRequest::ScatterGather(1);
-				Batch.Read(*StreamableData, RegisteredGeometry.StreamableDataOffset, RegisteredGeometry.StreamableDataSize);
-				Batch.Issue(PendingRequest.RequestBuffer, AIOP_Low, [](FBulkDataRequest::EStatus) {}, PendingRequest.Request);
+				RegisteredGeometry.StreamableAsset->IssueRequest(PendingRequest.Request, PendingRequest.RequestBuffer);
 						
 				// TODO: Could other requests already be completed?
 				break;
 			}
 			else
 			{
-				{
-					RegisteredGeometry.Geometry->RawData.SetNum(RegisteredGeometry.StreamableDataSize);
-
-					FMemoryReaderView Ar(PendingRequest.RequestBuffer.GetView(), /*bIsPersistent=*/ true);
-					Ar.Serialize(RegisteredGeometry.Geometry->RawData.GetData(), RegisteredGeometry.StreamableDataSize);
-					check(!Ar.IsError());
-				}
-
-				{
-					FRHIResourceReplaceBatcher Batcher(RHICmdList, 1);
-					FRayTracingGeometryInitializer IntermediateInitializer = RegisteredGeometry.Geometry->Initializer;
-					IntermediateInitializer.Type = ERayTracingGeometryInitializerType::StreamingSource;
-
-					if (!RegisteredGeometry.Geometry->RawData.IsEmpty())
-					{
-						IntermediateInitializer.OfflineData = &RegisteredGeometry.Geometry->RawData;
-					}
-
-					FRayTracingGeometryRHIRef IntermediateRayTracingGeometry = RHICmdList.CreateRayTracingGeometry(IntermediateInitializer);
-
-					RegisteredGeometry.Geometry->SetRequiresBuild(IntermediateInitializer.OfflineData == nullptr || IntermediateRayTracingGeometry->IsCompressed());
-
-					RegisteredGeometry.Geometry->InitRHIForStreaming(IntermediateRayTracingGeometry, Batcher);
-
-					// When Batcher goes out of scope it will add commands to copy the BLAS buffers on RHI thread.
-					// We need to do it before we build the current geometry (also on RHI thread).
-				}
+				RegisteredGeometry.StreamableAsset->InitWithStreamedData(RHICmdList, PendingRequest.RequestBuffer.GetView());
 
 				RegisteredGeometry.Status = FRegisteredGeometry::FStatus::StreamedIn;
-
-				if (!RegisteredGeometry.Geometry->GetRequiresBuild())
-				{
-					// only need to request here if no build will be requested since build path already requests update as necessary
-					RequestUpdateCachedRenderState(RegisteredGeometry.Geometry->GroupHandle);
-				}
-
-				RegisteredGeometry.Geometry->RequestBuildIfNeeded(RHICmdList, ERTAccelerationStructureBuildPriority::Normal);
 			}
 
 			PendingRequest.Reset();
diff --git a/Engine/Source/Runtime/Engine/Private/StaticMesh.cpp b/Engine/Source/Runtime/Engine/Private/StaticMesh.cpp
--- a/Engine/Source/Runtime/Engine/Private/StaticMesh.cpp
+++ b/Engine/Source/Runtime/Engine/Private/StaticMesh.cpp
@@ -81,6 +81,7 @@
 #include "Logging/MessageLog.h"
 #include "Misc/UObjectToken.h"
 #include "Streaming/StaticMeshUpdate.h"
+#include "ProfilingDebugging/IoStoreTrace.h"
 
 #include UE_INLINE_GENERATED_CPP_BY_NAME(StaticMesh)
 
@@ -1621,10 +1622,10 @@ void FStaticMeshRayTracingProxyLOD::InitResources(UStaticMesh* Parent, int32 LOD
 	if (bOwnsBuffers)
 	{
 		BeginInitResource(OwnerName, IndexBuffer);
+
 		BeginInitResource(OwnerName, &VertexBuffers->StaticMeshVertexBuffer);
 		BeginInitResource(OwnerName, &VertexBuffers->PositionVertexBuffer);
-
-		// todo: support ColorVertexBuffer
+		BeginInitResource(OwnerName, &VertexBuffers->ColorVertexBuffer);
 	}
 
 	if (Parent && !Parent->bSupportRayTracing)
@@ -1635,6 +1636,14 @@ void FStaticMeshRayTracingProxyLOD::InitResources(UStaticMesh* Parent, int32 LOD
 #if RHI_RAYTRACING
 	if (IsRayTracingAllowed() && RayTracingGeometry != nullptr && bOwnsRayTracingGeometry)
 	{
+		if(RayTracingGeometry->RawDataHeader.IsValid())
+		{
+			FRayTracingAccelerationStructureOfflineMetadata OfflineMetadata = RHIGetRayTracingGeometryOfflineMetadata(RayTracingGeometry->RawDataHeader);
+
+			OfflineBVHOffset += OfflineMetadata.SerializedOffset;
+			OfflineBVHSize = OfflineMetadata.SerializedSize;
+		}
+
 		ENQUEUE_RENDER_COMMAND(InitStaticMeshRayTracingGeometry)(
 			[this, DebugName = (Parent ? Parent->GetFName() : NAME_None), OwnerName](FRHICommandListImmediate& RHICmdList)
 			{
@@ -1722,32 +1731,36 @@ void FStaticMeshRayTracingProxyLOD::Serialize(FArchive& Ar, UObject* Owner, int3
 		RayTracingGeometry = new FRayTracingGeometry();
 	}
 
+	// TODO: Support DDC streaming in editor
+	bool bTmpBuffersInlined = !Ar.IsCooking();
+	Ar << bTmpBuffersInlined;
+	bBuffersInlined = bTmpBuffersInlined;
+
 	if (bBuffersInlined)
 	{
-		SerializeBuffersMetaData(Ar);
 		SerializeBuffers(Ar, OwnerStaticMesh, 0);
+		SerializeRayTracingGeometry(Ar, OwnerStaticMesh, 0);
+		SerializeMetaData(Ar);
 	}
 	else
 	{
-		SerializeBuffersMetaData(Ar);
-
 #if WITH_EDITOR
-		// Serialize buffers to FByteBulkData when saving to DDC
-		// (don't need to do it again when cooking)
-		if (Ar.IsSaving() && !Ar.IsCooking())
+		if (Ar.IsSaving())
 		{
+			checkf(Ar.IsCooking(), TEXT("Only expected to serialize buffers to FByteBulkData when cooking."));
+
 			TArray<uint8> TmpBuff;
-			uint32 BuffersSize = 0;
 
 			{
 				FMemoryWriter MemWriter(TmpBuff, /*bIsPersistent=*/ true);
 				MemWriter.SetSavePackageData(Ar.GetSavePackageData());
 				MemWriter.SetByteSwapping(Ar.IsByteSwapping());
 
-				BuffersSize = SerializeBuffers(MemWriter, OwnerStaticMesh, 0);
+				SerializeBuffers(MemWriter, OwnerStaticMesh, 0);
+				SerializeRayTracingGeometry(MemWriter, OwnerStaticMesh, 0);
 			}
 
-			if(BuffersSize > 0)
+			if(!TmpBuff.IsEmpty())
 			{
 				StreamableData.Lock(LOCK_READ_WRITE);
 
@@ -1761,45 +1774,68 @@ void FStaticMeshRayTracingProxyLOD::Serialize(FArchive& Ar, UObject* Owner, int3
 #endif
 		
 		StreamableData.Serialize(Ar, Owner, Index);
+
+		SerializeMetaData(Ar);
 	}
 }
 
-uint32 FStaticMeshRayTracingProxyLOD::SerializeBuffers(FArchive& Ar, UStaticMesh* OwnerStaticMesh, uint8 InStripFlags)
+void FStaticMeshRayTracingProxyLOD::SerializeBuffers(FArchive& Ar, UStaticMesh* OwnerStaticMesh, uint8 InStripFlags)
 {
-	uint32 OutSize = 0;
-
 	if (bOwnsBuffers)
 	{
+		const int64 BeforeOffset = Ar.Tell();
+
 		VertexBuffers->PositionVertexBuffer.Serialize(Ar, false);
 		VertexBuffers->StaticMeshVertexBuffer.Serialize(Ar, false);
 		VertexBuffers->ColorVertexBuffer.Serialize(Ar, false);
 
 		IndexBuffer->Serialize(Ar, false);
 
-		OutSize += VertexBuffers->PositionVertexBuffer.GetNumVertices() * VertexBuffers->PositionVertexBuffer.GetStride();
-		OutSize += VertexBuffers->StaticMeshVertexBuffer.GetResourceSize();
-		OutSize += VertexBuffers->ColorVertexBuffer.GetNumVertices() * VertexBuffers->ColorVertexBuffer.GetStride();
+		const int64 AfterOffset = Ar.Tell();
 
-		OutSize += IndexBuffer->GetIndexDataSize();
+		if (Ar.IsSaving())
+		{
+			BuffersSize = AfterOffset - BeforeOffset;
+		}
 	}
+}
 
+void FStaticMeshRayTracingProxyLOD::SerializeRayTracingGeometry(FArchive& Ar, UStaticMesh* OwnerStaticMesh, uint8 InStripFlags)
+{
 	if (bOwnsRayTracingGeometry)
 	{
-		const uint32 RawDataSize = RayTracingGeometry->RawData.Num() * RayTracingGeometry->RawData.GetTypeSize();
+		// logic similar to TArray::BulkSerialize(...)
+		// but want to track specific offset of RawData to be able to stream-in specific chunks based on FRayTracingAccelerationStructureOfflineMetadata
 
-		// serialize RawData as a memory blob to be able to stream-in specific chunks based on FRayTracingAccelerationStructureOfflineMetadata
-		Ar.Serialize(RayTracingGeometry->RawData.GetData(), RawDataSize);
+		uint32 RawDataNum = RayTracingGeometry->RawData.Num();
+		Ar << RawDataNum;
 
-		OutSize += RawDataSize;
-	}
+		if (Ar.IsLoading())
+		{
+			RayTracingGeometry->RawData.Empty(RawDataNum);
+			RayTracingGeometry->RawData.AddUninitialized(RawDataNum);
+		}
+
+		const int64 BeforeOffset = Ar.Tell();
+
+		Ar.Serialize(RayTracingGeometry->RawData.GetData(), RawDataNum * RayTracingGeometry->RawData.GetTypeSize());
 
-	return OutSize;
+		const int64 AfterOffset = Ar.Tell();
+
+		if (Ar.IsSaving())
+		{
+			OfflineBVHOffset = BeforeOffset;
+			OfflineBVHSize = AfterOffset - BeforeOffset;
+		}
+	}
 }
 
-void FStaticMeshRayTracingProxyLOD::SerializeBuffersMetaData(FArchive& Ar)
+void FStaticMeshRayTracingProxyLOD::SerializeMetaData(FArchive& Ar)
 {
 	if (bOwnsBuffers)
 	{
+		Ar << BuffersSize;
+
 		VertexBuffers->PositionVertexBuffer.SerializeMetaData(Ar);
 		VertexBuffers->StaticMeshVertexBuffer.SerializeMetaData(Ar);
 		VertexBuffers->ColorVertexBuffer.SerializeMetaData(Ar);
@@ -1809,10 +1845,138 @@ void FStaticMeshRayTracingProxyLOD::SerializeBuffersMetaData(FArchive& Ar)
 
 	if (bOwnsRayTracingGeometry)
 	{
+		Ar << OfflineBVHOffset;
+		Ar << OfflineBVHSize;
+
 		Ar << RayTracingGeometry->RawDataHeader;
 	}
 }
 
+uint32 FStaticMeshRayTracingProxyLOD::GetRequestSize() const
+{
+	const uint32 RequestSize = BuffersSize + OfflineBVHSize;
+	check(RequestSize < StreamableData.GetBulkDataSize());
+	return RequestSize;
+}
+
+void FStaticMeshRayTracingProxyLOD::IssueRequest(FBulkDataBatchRequest& Request, FIoBuffer& RequestBuffer)
+{
+	checkf(!StreamableData.IsBulkDataLoaded(), TEXT("Bulk data should not be loaded yet!"));
+	checkf(StreamableData.CanLoadFromDisk(), TEXT("Bulk data is not loaded and cannot be loaded from disk!"));
+	check(!StreamableData.IsStoredCompressedOnDisk()); // We do not support compressed Bulkdata for this system. Limitation of the streaming request/bulk data
+
+	const bool bRequestBuffers = BuffersSize > 0;
+	const bool bRequestOfflineBVH = OfflineBVHSize > 0;
+
+	uint32 BatchSize = 0;
+	BatchSize += bRequestBuffers ? 1 : 0;
+	BatchSize += bRequestOfflineBVH ? 1 : 0;
+
+	// TODO: We're currently using a single batch per request so we can individually cancel and wait on requests.
+	// This isn't ideal and should be revisited in the future.
+	TRACE_IOSTORE_METADATA_SCOPE_TAG("StaticMeshRayTracingProxyLOD");
+	FBulkDataBatchRequest::FScatterGatherBuilder Batch = FBulkDataBatchRequest::ScatterGather(BatchSize);
+	if (bRequestBuffers)
+	{
+		Batch.Read(StreamableData, 0, BuffersSize);
+	}
+	if (bRequestOfflineBVH)
+	{
+		Batch.Read(StreamableData, OfflineBVHOffset, OfflineBVHSize);
+	}
+	Batch.Issue(RequestBuffer, AIOP_Low, [](FBulkDataRequest::EStatus) {}, Request);
+}
+
+void FStaticMeshRayTracingProxyLOD::InitWithStreamedData(FRHICommandListBase& RHICmdList, FMemoryView StreamedData)
+{
+	FMemoryReaderView Ar(StreamedData, /*bIsPersistent=*/ true);
+
+	SerializeBuffers(Ar, /*OwnerStaticMesh*/nullptr, /*InStripFlags*/0);
+
+	if (bOwnsRayTracingGeometry && OfflineBVHSize > 0)
+	{
+		RayTracingGeometry->RawData.Empty(OfflineBVHSize);
+		RayTracingGeometry->RawData.AddUninitialized(OfflineBVHSize);
+		Ar.Serialize(RayTracingGeometry->RawData.GetData(), OfflineBVHSize);
+	}
+	
+	check(!Ar.IsError());
+
+	if (bOwnsBuffers)
+	{
+		FStaticMeshVertexBuffers& VBs = *VertexBuffers;
+		FBufferRHIRef TangentsVertexBuffer = VBs.StaticMeshVertexBuffer.CreateTangentsRHIBuffer(RHICmdList);
+		FBufferRHIRef TexCoordVertexBuffer = VBs.StaticMeshVertexBuffer.CreateTexCoordRHIBuffer(RHICmdList);
+		FBufferRHIRef PositionVertexBuffer = VBs.PositionVertexBuffer.CreateRHIBuffer(RHICmdList);
+		FBufferRHIRef ColorVertexBuffer = VBs.ColorVertexBuffer.CreateRHIBuffer(RHICmdList);
+		FBufferRHIRef IndexBufferRHI = IndexBuffer->CreateRHIBuffer(RHICmdList);
+
+		FRHIResourceReplaceBatcher Batcher(RHICmdList, 1);
+
+		VBs.StaticMeshVertexBuffer.InitRHIForStreaming(TangentsVertexBuffer, TexCoordVertexBuffer, Batcher);
+		VBs.PositionVertexBuffer.InitRHIForStreaming(PositionVertexBuffer, Batcher);
+		VBs.ColorVertexBuffer.InitRHIForStreaming(ColorVertexBuffer, Batcher);
+		IndexBuffer->InitRHIForStreaming(IndexBufferRHI, Batcher);
+	}
+
+#if RHI_RAYTRACING
+	if (bOwnsRayTracingGeometry)
+	{
+		{
+			FRHIResourceReplaceBatcher Batcher(RHICmdList, 1);
+			FRayTracingGeometryInitializer IntermediateInitializer = RayTracingGeometry->Initializer;
+			IntermediateInitializer.Type = ERayTracingGeometryInitializerType::StreamingSource;
+
+			if (!RayTracingGeometry->RawData.IsEmpty())
+			{
+				IntermediateInitializer.OfflineData = &RayTracingGeometry->RawData;
+			}
+
+			FRayTracingGeometryRHIRef IntermediateRayTracingGeometry = RHICmdList.CreateRayTracingGeometry(IntermediateInitializer);
+
+			RayTracingGeometry->SetRequiresBuild(IntermediateInitializer.OfflineData == nullptr || IntermediateRayTracingGeometry->IsCompressed());
+
+			RayTracingGeometry->InitRHIForStreaming(IntermediateRayTracingGeometry, Batcher);
+
+			// When Batcher goes out of scope it will add commands to copy the BLAS buffers on RHI thread.
+			// We need to do it before we build the current geometry (also on RHI thread).
+		}
+
+		if (!RayTracingGeometry->GetRequiresBuild())
+		{
+			// only need to request here if no build will be requested since build path already requests update as necessary
+			GRayTracingGeometryManager->RequestUpdateCachedRenderState(RayTracingGeometry->GroupHandle);
+		}
+
+		RayTracingGeometry->RequestBuildIfNeeded(RHICmdList, ERTAccelerationStructureBuildPriority::Normal);
+	}
+#endif // RHI_RAYTRACING
+
+	if (!GRHIGlobals.RayTracing.SupportsShaders && !GRHIGlobals.RayTracing.RequiresInlineRayTracingSBT)
+	{
+		// TODO: stream out VB/IB
+	}
+}
+
+void FStaticMeshRayTracingProxyLOD::ReleaseForStreaming(FRHIResourceReplaceBatcher& Batcher)
+{
+#if RHI_RAYTRACING
+	if (bOwnsRayTracingGeometry)
+	{
+		RayTracingGeometry->ReleaseRHIForStreaming(Batcher);
+	}
+#endif // RHI_RAYTRACING
+
+	if (bOwnsBuffers)
+	{
+		VertexBuffers->PositionVertexBuffer.ReleaseRHIForStreaming(Batcher);
+		VertexBuffers->StaticMeshVertexBuffer.ReleaseRHIForStreaming(Batcher);
+		VertexBuffers->ColorVertexBuffer.ReleaseRHIForStreaming(Batcher);
+
+		IndexBuffer->ReleaseRHIForStreaming(Batcher);
+	}
+}
+
 /*------------------------------------------------------------------------------
 	FStaticMeshRayTracingProxy
 ------------------------------------------------------------------------------*/
@@ -2376,9 +2540,7 @@ void FStaticMeshRenderData::InitResources(ERHIFeatureLevel::Type InFeatureLevel,
 						{
 							((FRayTracingGeometryManager*)GRayTracingGeometryManager)->SetRayTracingGeometryStreamingData(
 								RayTracingLOD.RayTracingGeometry,
-								RayTracingLOD.StreamableData,
-								0,
-								RayTracingLOD.StreamableData.GetBulkDataSize());
+								RayTracingLOD);
 						}
 					}
 				}
diff --git a/Engine/Source/Runtime/Engine/Public/Rendering/RayTracingGeometryManager.h b/Engine/Source/Runtime/Engine/Public/Rendering/RayTracingGeometryManager.h
--- a/Engine/Source/Runtime/Engine/Public/Rendering/RayTracingGeometryManager.h
+++ b/Engine/Source/Runtime/Engine/Public/Rendering/RayTracingGeometryManager.h
@@ -12,6 +12,7 @@
 
 class FPrimitiveSceneProxy;
 class UStaticMesh;
+class FRayTracingStreamableAsset;
 
 class FRayTracingGeometryManager : public IRayTracingGeometryManager
 {
@@ -35,7 +36,7 @@ public:
 
 	ENGINE_API virtual void RefreshRegisteredGeometry(RayTracingGeometryHandle Handle) override;
 
-	void SetRayTracingGeometryStreamingData(const FRayTracingGeometry* Geometry, FByteBulkData& BulkData, uint32 Offset, uint32 Size);
+	void SetRayTracingGeometryStreamingData(const FRayTracingGeometry* Geometry, FRayTracingStreamableAsset& StreamableAsset);
 	void SetRayTracingGeometryGroupCurrentFirstLODIndex(FRHICommandListBase& RHICmdList, RayTracing::GeometryGroupHandle Handle, uint8 CurrentFirstLODIdx);
 
 	ENGINE_API virtual void PreRender() override;
@@ -108,8 +109,7 @@ private:
 		uint64 LastReferencedFrame = 0;
 		uint32 Size = 0;
 
-		FByteBulkData* StreamableData = nullptr;
-		uint32 StreamableDataOffset = 0;
+		FRayTracingStreamableAsset* StreamableAsset = nullptr;
 		uint32 StreamableDataSize = 0;
 
 		int16 StreamingRequestIndex = INDEX_NONE;
diff --git a/Engine/Source/Runtime/Engine/Public/Rendering/RayTracingStreamableAsset.h b/Engine/Source/Runtime/Engine/Public/Rendering/RayTracingStreamableAsset.h
new file mode 100644
--- /dev/null
+++ b/Engine/Source/Runtime/Engine/Public/Rendering/RayTracingStreamableAsset.h
@@ -0,0 +1,20 @@
+// Copyright Epic Games, Inc. All Rights Reserved.
+#pragma once
+
+#include "Memory/MemoryFwd.h"
+#include "CoreMinimal.h"
+
+class FBulkDataBatchRequest;
+class FIoBuffer;
+class FRHICommandListBase;
+class FRHIResourceReplaceBatcher;
+
+class FRayTracingStreamableAsset
+{
+public:
+	virtual uint32 GetRequestSize() const = 0;
+
+	virtual void IssueRequest(FBulkDataBatchRequest& Request, FIoBuffer& RequestBuffer) = 0;
+	virtual void InitWithStreamedData(FRHICommandListBase& RHICmdList, FMemoryView StreamedData) = 0;
+	virtual void ReleaseForStreaming(FRHIResourceReplaceBatcher& Batcher) = 0;
+};
diff --git a/Engine/Source/Runtime/Engine/Public/StaticMeshResources.h b/Engine/Source/Runtime/Engine/Public/StaticMeshResources.h
--- a/Engine/Source/Runtime/Engine/Public/StaticMeshResources.h
+++ b/Engine/Source/Runtime/Engine/Public/StaticMeshResources.h
@@ -31,6 +31,7 @@
 #include "Rendering/PositionVertexBuffer.h"
 #include "Rendering/StaticMeshVertexDataInterface.h"
 #include "Rendering/NaniteInterface.h"
+#include "Rendering/RayTracingStreamableAsset.h"
 #include "RenderTransform.h"
 #include "Templates/UniquePtr.h"
 #include "Serialization/BulkData.h"
@@ -687,7 +688,7 @@ struct FStaticMeshVertexFactories
 using FStaticMeshLODResourcesArray = TIndirectArray<FStaticMeshLODResources>;
 using FStaticMeshVertexFactoriesArray = TArray<FStaticMeshVertexFactories>;
 
-struct FStaticMeshRayTracingProxyLOD
+struct FStaticMeshRayTracingProxyLOD : public FRayTracingStreamableAsset
 {
 	FRayTracingGeometry* RayTracingGeometry = nullptr;
 
@@ -697,20 +698,24 @@ struct FStaticMeshRayTracingProxyLOD
 	FRawStaticIndexBuffer* IndexBuffer = nullptr;
 
 	FByteBulkData StreamableData;
+	uint32 BuffersSize = 0;
+	uint32 OfflineBVHOffset = 0;
+	uint32 OfflineBVHSize = 0;
 
 	bool bOwnsRayTracingGeometry : 1 = true;
 	bool bOwnsBuffers : 1 = true;
 	bool bBuffersInlined : 1 = false;
 
-	ENGINE_API ~FStaticMeshRayTracingProxyLOD();
+	ENGINE_API virtual ~FStaticMeshRayTracingProxyLOD();
 
 	void InitResources(UStaticMesh* Owner, int32 LODIndex);
 	void ReleaseResources();
 
 	void Serialize(FArchive& Ar, UObject* Owner, int32 Index);
 
-	uint32 SerializeBuffers(FArchive& Ar, UStaticMesh* OwnerStaticMesh, uint8 InStripFlags);
-	void SerializeBuffersMetaData(FArchive& Ar);
+	void SerializeBuffers(FArchive& Ar, UStaticMesh* OwnerStaticMesh, uint8 InStripFlags);
+	void SerializeRayTracingGeometry(FArchive& Ar, UStaticMesh* OwnerStaticMesh, uint8 InStripFlags);
+	void SerializeMetaData(FArchive& Ar);
 
 	/** Discard loaded data. Used when a streaming request is cancelled */
 	void DiscardCPUData();
@@ -718,6 +723,12 @@ struct FStaticMeshRayTracingProxyLOD
 #if RHI_RAYTRACING
 	void SetupRayTracingGeometryInitializer(FRayTracingGeometryInitializer& Initializer, const FDebugName& DebugName, const FName& OwnerName) const;
 #endif // RHI_RAYTRACING
+
+	// FRayTracingStreamableAsset
+	ENGINE_API virtual uint32 GetRequestSize() const override;
+	ENGINE_API virtual void IssueRequest(FBulkDataBatchRequest& Request, FIoBuffer& RequestBuffer) override;
+	ENGINE_API virtual void InitWithStreamedData(FRHICommandListBase& RHICmdList, FMemoryView StreamedData) override;
+	ENGINE_API virtual void ReleaseForStreaming(FRHIResourceReplaceBatcher& Batcher) override;
 };
 
 using FStaticMeshRayTracingProxyLODArray = TIndirectArray<FStaticMeshRayTracingProxyLOD>;
