[NetPhysics] Fixup for Physics Replication LOD

diff --git a/Engine/Source/Runtime/Engine/Private/PhysicsEngine/NetworkPhysicsSettingsComponent.cpp b/Engine/Source/Runtime/Engine/Private/PhysicsEngine/NetworkPhysicsSettingsComponent.cpp
--- a/Engine/Source/Runtime/Engine/Private/PhysicsEngine/NetworkPhysicsSettingsComponent.cpp
+++ b/Engine/Source/Runtime/Engine/Private/PhysicsEngine/NetworkPhysicsSettingsComponent.cpp
@@ -80,17 +80,19 @@ void UNetworkPhysicsSettingsComponent::InitializeComponent()
 						{
 							if (Chaos::FConstPhysicsObjectHandle PhysicsObject = RootPrimComp->GetPhysicsObjectByName(NAME_None))
 							{
-								FNetworkPhysicsSettingsAsyncInput* AsyncInput = NetworkPhysicsSettings_Internal->GetProducerInputData_External();
-								AsyncInput->PhysicsObject = PhysicsObject;
-								FNetworkPhysicsSettingsAsync Settings
+								if (FNetworkPhysicsSettingsAsyncInput* AsyncInput = NetworkPhysicsSettings_Internal->GetProducerInputData_External())
 								{
-									.GeneralSettings = GeneralSettings,
-									.DefaultReplicationSettings = DefaultReplicationSettings,
-									.PredictiveInterpolationSettings = PredictiveInterpolationSettings,
-									.ResimulationSettings = ResimulationSettings,
-									.NetworkPhysicsComponentSettings = NetworkPhysicsComponentSettings,
-								};
-								AsyncInput->Settings = Settings;
+									AsyncInput->PhysicsObject = PhysicsObject;
+									FNetworkPhysicsSettingsAsync Settings
+									{
+										.GeneralSettings = GeneralSettings,
+										.DefaultReplicationSettings = DefaultReplicationSettings,
+										.PredictiveInterpolationSettings = PredictiveInterpolationSettings,
+										.ResimulationSettings = ResimulationSettings,
+										.NetworkPhysicsComponentSettings = NetworkPhysicsComponentSettings,
+									};
+									AsyncInput->Settings = Settings;
+								}
 							}
 						}
 					}
@@ -151,9 +153,15 @@ void UNetworkPhysicsSettingsComponent::OnComponentPhysicsStateChanged(UPrimitive
 	{
 		if (Chaos::FConstPhysicsObjectHandle PhysicsObject = ChangedComponent->GetPhysicsObjectByName(NAME_None))
 		{
-			FNetworkPhysicsSettingsAsyncInput* AsyncInput = NetworkPhysicsSettings_Internal->GetProducerInputData_External();
-			AsyncInput->PhysicsObject = PhysicsObject;
+			if (NetworkPhysicsSettings_Internal)
+			{
+				if (FNetworkPhysicsSettingsAsyncInput* AsyncInput = NetworkPhysicsSettings_Internal->GetProducerInputData_External())
+				{
+					AsyncInput->PhysicsObject = PhysicsObject;
+				}
+			}
 		}
+
 		RegisterInPhysicsReplicationLOD();
 	}
 }
diff --git a/Engine/Source/Runtime/Engine/Private/PhysicsEngine/PhysicsReplicationLOD.cpp b/Engine/Source/Runtime/Engine/Private/PhysicsEngine/PhysicsReplicationLOD.cpp
--- a/Engine/Source/Runtime/Engine/Private/PhysicsEngine/PhysicsReplicationLOD.cpp
+++ b/Engine/Source/Runtime/Engine/Private/PhysicsEngine/PhysicsReplicationLOD.cpp
@@ -66,8 +66,13 @@ void FPhysicsReplicationLOD::SetEnabled(const bool bInEnabled)
 	bEnabled = bInEnabled;
 	
 	// Marshal over the enabled setting to the physics thread
-	FPhysicsReplicationLODAsyncInput* AsyncInput = PhysicsReplicationLODAsync->GetProducerInputData_External();
-	AsyncInput->bEnabled = bEnabled;
+	if (PhysicsReplicationLODAsync)
+	{
+		if (FPhysicsReplicationLODAsyncInput* AsyncInput = PhysicsReplicationLODAsync->GetProducerInputData_External())
+		{
+			AsyncInput->bEnabled = bEnabled;
+		}
+	}
 }
 
 bool FPhysicsReplicationLOD::IsEnabled() const
@@ -93,21 +98,25 @@ void FPhysicsReplicationLOD::UnregisterFocalPoint_External(const UPrimitiveCompo
 
 void FPhysicsReplicationLOD::RegisterFocalPoint_External(Chaos::FConstPhysicsObjectHandle PhysicsObject)
 {
-	if (PhysicsObject)
+	if (PhysicsObject && PhysicsReplicationLODAsync)
 	{
 		// Marshal registration over to physics thread
-		FPhysicsReplicationLODAsyncInput* AsyncInput = PhysicsReplicationLODAsync->GetProducerInputData_External();
-		AsyncInput->PhysicsObjectsToRegister.Add(PhysicsObject);
+		if (FPhysicsReplicationLODAsyncInput* AsyncInput = PhysicsReplicationLODAsync->GetProducerInputData_External())
+		{
+			AsyncInput->PhysicsObjectsToRegister.Add(PhysicsObject);
+		}
 	}
 }
 
 void FPhysicsReplicationLOD::UnregisterFocalPoint_External(Chaos::FConstPhysicsObjectHandle PhysicsObject)
 {
-	if (PhysicsObject)
+	if (PhysicsObject && PhysicsReplicationLODAsync)
 	{
 		// Marshal deregistration over to physics thread
-		FPhysicsReplicationLODAsyncInput* AsyncInput = PhysicsReplicationLODAsync->GetProducerInputData_External();
-		AsyncInput->PhysicsObjectsToUnregister.Add(PhysicsObject);
+		if (FPhysicsReplicationLODAsyncInput* AsyncInput = PhysicsReplicationLODAsync->GetProducerInputData_External())
+		{
+			AsyncInput->PhysicsObjectsToUnregister.Add(PhysicsObject);
+		}
 	}
 }
 
@@ -211,8 +220,8 @@ void FPhysicsReplicationLODAsync::OnPreSimulate_Internal()
 			{
 				FColor DebugColor = IsEnabled() ? FLinearColor::MakeFromHSV8(0, 150, 255).ToFColor(false) : FColor::Black;
 				DebugColor.A = 255;
-				const float DistanceFullPrediction = PhysicsReplicationLODCVars::DistanceForFullPrediction > 0.0f ? PhysicsReplicationLODCVars::DistanceForFullPrediction : DefaultSettings.DistanceForFullPrediction;
-				Chaos::FDebugDrawQueue::GetInstance().DrawDebugSphere(FocalParticleHandle->GetX(), DistanceFullPrediction, 8, DebugColor, false, -1.0f, 0, 1.0f);
+				const float DistanceForFullPrediction = PhysicsReplicationLODCVars::DistanceForFullPrediction > 0.0f ? PhysicsReplicationLODCVars::DistanceForFullPrediction : DefaultSettings.DistanceForFullPrediction;
+				Chaos::FDebugDrawQueue::GetInstance().DrawDebugSphere(FocalParticleHandle->GetX(), DistanceForFullPrediction, 8, DebugColor, false, -1.0f, 0, 1.0f);
 			}
 		}
 	}
@@ -321,7 +330,7 @@ bool FPhysicsReplicationLODAsync::PerformDistanceLOD(const Chaos::FGeometryParti
 	Chaos::FReadPhysicsObjectInterface_Internal Interface = Chaos::FPhysicsObjectInternalInterface::GetRead();
 
 	const float DistanceForResimulationMode = PhysicsReplicationLODCVars::DistanceForResimulationMode > 0.0f ? PhysicsReplicationLODCVars::DistanceForResimulationMode : DefaultSettings.DistanceForResimulationMode;
-	const float DistanceFullPrediction = PhysicsReplicationLODCVars::DistanceForFullPrediction > 0.0f ? PhysicsReplicationLODCVars::DistanceForFullPrediction : DefaultSettings.DistanceForFullPrediction;
+	const float DistanceForFullPrediction = PhysicsReplicationLODCVars::DistanceForFullPrediction > 0.0f ? PhysicsReplicationLODCVars::DistanceForFullPrediction : DefaultSettings.DistanceForFullPrediction;
 	const float TimeOverDistance = PhysicsReplicationLODCVars::TimeOverDistance > 0.0f ? PhysicsReplicationLODCVars::TimeOverDistance : DefaultSettings.TimeOverDistance;
 
 	for (Chaos::FConstPhysicsObjectHandle FocalParticle : FocalParticles)
@@ -335,7 +344,7 @@ bool FPhysicsReplicationLODAsync::PerformDistanceLOD(const Chaos::FGeometryParti
 		// Check distance towards focal particles
 		const Chaos::FVec3 PosOffset = FocalParticleHandle->GetX() - ParticleHandle.GetX();
 		const float Distance = PosOffset.Size();
-		float TimeAlignment = FMath::Max((Distance - DistanceFullPrediction), 0.0f) * TimeOverDistance;
+		float TimeAlignment = FMath::Max((Distance - DistanceForFullPrediction), 0.0f) * TimeOverDistance;
 		TimeAlignment = (TimeAlignment * 0.001f); // Convert from ms to s
 		const float TickAlignment = TimeAlignment / RigidsSolver.GetAsyncDeltaTime();
 		const int32 TickAlignmentDiscrete = FMath::Floor(TickAlignment);
