[NetPhysics] Simulated proxy improvements
- Reduced end-to-end latency for sim-proxies to receive inputs
- Input Decay settings
- Input Decay for modular vehicles

diff --git a/Engine/Plugins/Experimental/ChaosModularVehicle/Source/ChaosModularVehicleEngine/Private/ChaosModularVehicle/ChaosSimModuleManagerAsyncCallback.cpp b/Engine/Plugins/Experimental/ChaosModularVehicle/Source/ChaosModularVehicleEngine/Private/ChaosModularVehicle/ChaosSimModuleManagerAsyncCallback.cpp
--- a/Engine/Plugins/Experimental/ChaosModularVehicle/Source/ChaosModularVehicleEngine/Private/ChaosModularVehicle/ChaosSimModuleManagerAsyncCallback.cpp
+++ b/Engine/Plugins/Experimental/ChaosModularVehicle/Source/ChaosModularVehicleEngine/Private/ChaosModularVehicle/ChaosSimModuleManagerAsyncCallback.cpp
@@ -304,6 +304,11 @@ void FNetworkModularVehicleInputs::MergeData(const FNetworkPhysicsData& FromData
 	VehicleInputs.Container.Merge(FromInput.VehicleInputs.Container);
 }
 
+void FNetworkModularVehicleInputs::DecayData(float DecayAmount)
+{
+	VehicleInputs.Container.Decay(DecayAmount);
+}
+
 bool FNetworkModularVehicleStates::NetSerialize(FArchive& Ar, class UPackageMap* Map, bool& bOutSuccess)
 {
 	using namespace UE::Net;
diff --git a/Engine/Plugins/Experimental/ChaosModularVehicle/Source/ChaosModularVehicleEngine/Public/ChaosModularVehicle/ChaosSimModuleManagerAsyncCallback.h b/Engine/Plugins/Experimental/ChaosModularVehicle/Source/ChaosModularVehicleEngine/Public/ChaosModularVehicle/ChaosSimModuleManagerAsyncCallback.h
--- a/Engine/Plugins/Experimental/ChaosModularVehicle/Source/ChaosModularVehicleEngine/Public/ChaosModularVehicle/ChaosSimModuleManagerAsyncCallback.h
+++ b/Engine/Plugins/Experimental/ChaosModularVehicle/Source/ChaosModularVehicleEngine/Public/ChaosModularVehicle/ChaosSimModuleManagerAsyncCallback.h
@@ -91,6 +91,8 @@ struct CHAOSMODULARVEHICLEENGINE_API FNetworkModularVehicleInputs : public FNetw
 	/** Merge data when multiple inputs happen an same simulation tick*/
 	virtual void MergeData(const FNetworkPhysicsData& FromData) override;
 
+	/**  Decay data during resimulation by DecayAmount which increases over resimulation frames from 0.0 -> 1.0 when the input is being reused */
+	virtual void DecayData(float DecayAmount) override;
 };
 
 template<>
diff --git a/Engine/Source/Runtime/Engine/Private/PhysicsEngine/NetworkPhysicsComponent.cpp b/Engine/Source/Runtime/Engine/Private/PhysicsEngine/NetworkPhysicsComponent.cpp
--- a/Engine/Source/Runtime/Engine/Private/PhysicsEngine/NetworkPhysicsComponent.cpp
+++ b/Engine/Source/Runtime/Engine/Private/PhysicsEngine/NetworkPhysicsComponent.cpp
@@ -23,7 +23,9 @@ namespace PhysicsReplicationCVars
 	namespace ResimulationCVars
 	{
 		int32 RedundantInputs = 2;
-		static FAutoConsoleVariableRef CVarResimRedundantInputs(TEXT("np2.Resim.RedundantInputs"), RedundantInputs, TEXT("How many extra inputs to send with each unreliable network message, to account for packetloss."));
+		static FAutoConsoleVariableRef CVarResimRedundantInputs(TEXT("np2.Resim.RedundantInputs"), RedundantInputs, TEXT("How many extra inputs to send with each unreliable network message, to account for packetloss. From owning client to server and back to owning client."));
+		int32 RedundantRemoteInputs = 1;
+		static FAutoConsoleVariableRef CVarResimRedundantRemoteInputs(TEXT("np2.Resim.RedundantRemoteInputs"), RedundantRemoteInputs, TEXT("How many extra inputs to send with each unreliable network message, to account for packetloss. From server to remote clients."));
 		int32 RedundantStates = 0;
 		static FAutoConsoleVariableRef CVarResimRedundantStates(TEXT("np2.Resim.RedundantStates"), RedundantStates, TEXT("How many extra states to send with each unreliable network message, to account for packetloss."));
 		bool bAllowRewindToClosestState = true;
@@ -152,6 +154,7 @@ FNetworkPhysicsRewindDataProxy& FNetworkPhysicsRewindDataProxy::operator=(const
 
 #if UE_WITH_IRIS
 UE_NET_IMPLEMENT_NAMED_STRUCT_LASTRESORT_NETSERIALIZER_AND_REGISTRY_DELEGATES(NetworkPhysicsRewindDataInputProxy);
+UE_NET_IMPLEMENT_NAMED_STRUCT_LASTRESORT_NETSERIALIZER_AND_REGISTRY_DELEGATES(NetworkPhysicsRewindDataRemoteInputProxy);
 UE_NET_IMPLEMENT_NAMED_STRUCT_LASTRESORT_NETSERIALIZER_AND_REGISTRY_DELEGATES(NetworkPhysicsRewindDataStateProxy);
 UE_NET_IMPLEMENT_NAMED_STRUCT_LASTRESORT_NETSERIALIZER_AND_REGISTRY_DELEGATES(NetworkPhysicsRewindDataImportantInputProxy);
 UE_NET_IMPLEMENT_NAMED_STRUCT_LASTRESORT_NETSERIALIZER_AND_REGISTRY_DELEGATES(NetworkPhysicsRewindDataImportantStateProxy);
@@ -162,6 +165,11 @@ bool FNetworkPhysicsRewindDataInputProxy::NetSerialize(FArchive& Ar, class UPack
 	return NetSerializeBase(Ar, Map, bOutSuccess, [this]() { return Owner->ReplicatedInputs.History->CreateNew(); });
 }
 
+bool FNetworkPhysicsRewindDataRemoteInputProxy::NetSerialize(FArchive& Ar, class UPackageMap* Map, bool& bOutSuccess)
+{
+	return NetSerializeBase(Ar, Map, bOutSuccess, [this]() { return Owner->ReplicatedRemoteInputs.History->CreateNew(); });
+}
+
 bool FNetworkPhysicsRewindDataStateProxy::NetSerialize(FArchive& Ar, class UPackageMap* Map, bool& bOutSuccess)
 {
 	return NetSerializeBase(Ar, Map, bOutSuccess, [this]() { return Owner->ReplicatedStates.History->CreateNew(); });
@@ -427,11 +435,15 @@ void UNetworkPhysicsComponent::InitPhysics()
 {
 	if (const IConsoleVariable* CVarRedundantInputs = IConsoleManager::Get().FindConsoleVariable(TEXT("np2.Resim.RedundantInputs")))
 	{
-		SetNumberOfInputsToNetwork(CVarRedundantInputs->GetInt());
+		SetNumberOfInputsToNetwork(CVarRedundantInputs->GetInt() + 1);
+	}
+	if (const IConsoleVariable* CVarRedundantRemoteInputs = IConsoleManager::Get().FindConsoleVariable(TEXT("np2.Resim.RedundantRemoteInputs")))
+	{
+		SetNumberOfRemoteInputsToNetwork(CVarRedundantRemoteInputs->GetInt() + 1);
 	}
 	if (const IConsoleVariable* CVarRedundantStates = IConsoleManager::Get().FindConsoleVariable(TEXT("np2.Resim.RedundantStates")))
 	{
-		SetNumberOfStatesToNetwork(CVarRedundantStates->GetInt());
+		SetNumberOfStatesToNetwork(CVarRedundantStates->GetInt() + 1);
 	}
 
 	if (AActor* Owner = GetOwner())
@@ -475,8 +487,9 @@ void UNetworkPhysicsComponent::InitializeComponent()
 		UNetworkPhysicsSettingsComponent* SettingsComponent = Owner->FindComponentByClass<UNetworkPhysicsSettingsComponent>();
 		if (SettingsComponent)
 		{
-			SetNumberOfInputsToNetwork(SettingsComponent->NetworkPhysicsComponentSettings.GetRedundantInputs());
-			SetNumberOfStatesToNetwork(SettingsComponent->NetworkPhysicsComponentSettings.GetRedundantStates());
+			SetNumberOfInputsToNetwork(SettingsComponent->NetworkPhysicsComponentSettings.GetRedundantInputs() + 1);
+			SetNumberOfRemoteInputsToNetwork(SettingsComponent->NetworkPhysicsComponentSettings.GetRedundantRemoteInputs() + 1);
+			SetNumberOfStatesToNetwork(SettingsComponent->NetworkPhysicsComponentSettings.GetRedundantStates() + 1);
 			bEnableUnreliableFlow = SettingsComponent->NetworkPhysicsComponentSettings.GetEnableUnreliableFlow();
 			bEnableReliableFlow = SettingsComponent->NetworkPhysicsComponentSettings.GetEnableReliableFlow();
 			bValidateDataOnGameThread = SettingsComponent->NetworkPhysicsComponentSettings.GetValidateDataOnGameThread();
@@ -485,6 +498,10 @@ void UNetworkPhysicsComponent::InitializeComponent()
 			{
 				ReplicatedInputs.History->ResizeDataHistory(InputsToNetwork);
 			}
+			if (ReplicatedRemoteInputs.History)
+			{
+				ReplicatedRemoteInputs.History->ResizeDataHistory(InputsToNetworkRemote);
+			}
 			if (ReplicatedStates.History)
 			{
 				ReplicatedStates.History->ResizeDataHistory(StatesToNetwork);
@@ -538,13 +555,24 @@ void UNetworkPhysicsComponent::GetLifetimeReplicatedProps(TArray< FLifetimePrope
 {
 	Super::GetLifetimeReplicatedProps(OutLifetimeProps);
 
-	FDoRepLifetimeParams ReplicatedParams;
-	ReplicatedParams.Condition = COND_None;
-	ReplicatedParams.RepNotifyCondition = REPNOTIFY_Always; //REPNOTIFY_OnChanged;
-	ReplicatedParams.bIsPushBased = true;
+	FDoRepLifetimeParams ReplicatedParamsOwner;
+	ReplicatedParamsOwner.Condition = COND_OwnerOnly;
+	ReplicatedParamsOwner.RepNotifyCondition = REPNOTIFY_Always;
+	ReplicatedParamsOwner.bIsPushBased = true;
+
+	FDoRepLifetimeParams ReplicatedParamsRemote;
+	ReplicatedParamsRemote.Condition = COND_SkipOwner;
+	ReplicatedParamsRemote.RepNotifyCondition = REPNOTIFY_Always;
+	ReplicatedParamsRemote.bIsPushBased = true;
+
+	FDoRepLifetimeParams ReplicatedParamsAll;
+	ReplicatedParamsAll.Condition = COND_None;
+	ReplicatedParamsAll.RepNotifyCondition = REPNOTIFY_Always;
+	ReplicatedParamsAll.bIsPushBased = true;
 
-	DOREPLIFETIME_WITH_PARAMS_FAST(UNetworkPhysicsComponent, ReplicatedInputs, ReplicatedParams);
-	DOREPLIFETIME_WITH_PARAMS_FAST(UNetworkPhysicsComponent, ReplicatedStates, ReplicatedParams);
+	DOREPLIFETIME_WITH_PARAMS_FAST(UNetworkPhysicsComponent, ReplicatedInputs, ReplicatedParamsOwner);
+	DOREPLIFETIME_WITH_PARAMS_FAST(UNetworkPhysicsComponent, ReplicatedRemoteInputs, ReplicatedParamsRemote);
+	DOREPLIFETIME_WITH_PARAMS_FAST(UNetworkPhysicsComponent, ReplicatedStates, ReplicatedParamsAll);
 }
 
 // Called every Game Thread frame
@@ -570,7 +598,15 @@ void UNetworkPhysicsComponent::NetworkMarshaledData()
 			{
 				if (bIsServer)
 				{
-					// Server sends inputs through property replication
+					if (IsLocallyControlled())
+					{
+						// Send inputs to remote clients after getting marshaled from PT if server is the one controlling the component
+						ReplicatedRemoteInputs.History->SetRecordDataIncremental(true); // Only record data that is newer than already cached data
+						ReplicatedInputs.History->CopyAllData(*ReplicatedRemoteInputs.History, /*bIncludeUnimportant*/ true, /*bIncludeImportant*/ true);
+						MARK_PROPERTY_DIRTY_FROM_NAME(UNetworkPhysicsComponent, ReplicatedRemoteInputs, this);
+					}
+
+					// Server sends inputs through property replication to owning client
 					MARK_PROPERTY_DIRTY_FROM_NAME(UNetworkPhysicsComponent, ReplicatedInputs, this);
 				}
 				else if (IsLocallyControlled())
@@ -729,6 +765,25 @@ void UNetworkPhysicsComponent::OnRep_SetReplicatedInputs()
 	}
 }
 
+void UNetworkPhysicsComponent::OnRep_SetReplicatedRemoteInputs()
+{
+	if (!NetworkPhysicsComponent_Internal || !InputHelper || !ReplicatedRemoteInputs.History)
+	{
+		return ;
+	}
+
+	if (FAsyncNetworkPhysicsComponentInput* AsyncInput = NetworkPhysicsComponent_Internal->GetProducerInputData_External())
+	{
+		if (!AsyncInput->InputData)
+		{
+			AsyncInput->InputData = InputHelper->CreateUniqueRewindHistory(InputsToNetwork);
+		}
+
+		AsyncInput->InputData->ResetFast();
+		ReplicatedRemoteInputs.History->CopyAllData(*AsyncInput->InputData.Get(), /*bIncludeUnimportant*/ true, /*bIncludeImportant*/ true);
+	}
+}
+
 /** DEPRECATED UE 5.4*/
 void UNetworkPhysicsComponent::ServerReceiveInputsDatas_Implementation(const FNetworkPhysicsRewindDataInputProxy& ClientInputs)
 {
@@ -757,6 +812,11 @@ void UNetworkPhysicsComponent::ServerReceiveInputData_Implementation(const FNetw
 
 		AsyncInput->InputData->ResetFast();
 		ClientInputs.History->CopyAllData(*AsyncInput->InputData.Get(), /*bIncludeUnimportant*/ true, /*bIncludeImportant*/ true);
+
+		// Send received inputs to remote clients
+		ReplicatedRemoteInputs.History->SetRecordDataIncremental(true); // Only record data that is newer than already cached data
+		ClientInputs.History->CopyAllData(*ReplicatedRemoteInputs.History, /*bIncludeUnimportant*/ true, /*bIncludeImportant*/ true);
+		MARK_PROPERTY_DIRTY_FROM_NAME(UNetworkPhysicsComponent, ReplicatedRemoteInputs, this);
 	}
 }
 
@@ -1714,13 +1774,19 @@ const float FAsyncNetworkPhysicsComponent::GetCurrentInputDecay(const FNetworkPh
 		return 0.0f;
 	}
 
+	const FNetworkPhysicsSettingsNetworkPhysicsComponent& ComponentSettings = GetComponentSettings();
+	const FRuntimeFloatCurve& InputDecayCurve = ComponentSettings.GetInputDecayCurve();
+
 	const float NumPredictedInputs = RewindData->CurrentFrame() - PhysicsData->LocalFrame; // Number of frames we have used the same PhysicsData for during resim
 	const float MaxPredictedInputs = RewindData->GetLatestFrame() - 1 - PhysicsData->LocalFrame; // Max number of frames PhysicsData registered frame until end of resim
 
 	// Linear decay
 	const float PredictionAlpha = MaxPredictedInputs > 0 ? (NumPredictedInputs / MaxPredictedInputs) : 0.0f;
 
-	return PredictionAlpha;
+	// Get decay from curve
+	const float InputDecay = InputDecayCurve.GetRichCurveConst()->Eval(PredictionAlpha);
+
+	return InputDecay;
 }
 
 void FAsyncNetworkPhysicsComponent::RegisterDataHistoryInRewindData()
diff --git a/Engine/Source/Runtime/Engine/Private/PhysicsEngine/NetworkPhysicsSettingsComponent.cpp b/Engine/Source/Runtime/Engine/Private/PhysicsEngine/NetworkPhysicsSettingsComponent.cpp
--- a/Engine/Source/Runtime/Engine/Private/PhysicsEngine/NetworkPhysicsSettingsComponent.cpp
+++ b/Engine/Source/Runtime/Engine/Private/PhysicsEngine/NetworkPhysicsSettingsComponent.cpp
@@ -12,6 +12,7 @@
 #include "PBDRigidsSolver.h"
 #include "PhysicsReplication.h"
 #include "Components/PrimitiveComponent.h"
+#include "Chaos/PhysicsObjectInternalInterface.h"
 
 namespace PhysicsReplicationCVars
 {
@@ -184,12 +185,17 @@ void UNetworkPhysicsSettingsComponent::RegisterInPhysicsReplicationLOD()
 
 UNetworkPhysicsSettingsComponent* UNetworkPhysicsSettingsComponent::GetSettingsForActor(AActor* Owner)
 {
-	UNetworkPhysicsSettingsComponent** Value = ObjectToSettings_External.Find(Owner);
+	UNetworkPhysicsSettingsComponent** Value = UNetworkPhysicsSettingsComponent::ObjectToSettings_External.Find(Owner);
 	return Value ? *Value : nullptr;
 }
 
 #pragma region // FNetworkPhysicsSettingsComponentAsync
 
+FNetworkPhysicsSettingsComponentAsync::FNetworkPhysicsSettingsComponentAsync() : TSimCallbackObject()
+	, Settings()
+	, PhysicsObject(nullptr)
+{}
+
 void FNetworkPhysicsSettingsComponentAsync::OnPreSimulate_Internal()
 {
 	ConsumeAsyncInput();
diff --git a/Engine/Source/Runtime/Engine/Public/Physics/NetworkPhysicsComponent.h b/Engine/Source/Runtime/Engine/Public/Physics/NetworkPhysicsComponent.h
--- a/Engine/Source/Runtime/Engine/Public/Physics/NetworkPhysicsComponent.h
+++ b/Engine/Source/Runtime/Engine/Public/Physics/NetworkPhysicsComponent.h
@@ -394,6 +394,27 @@ struct TStructOpsTypeTraits<FNetworkPhysicsRewindDataInputProxy> : public TStruc
 	};
 };
 
+/**
+ * Struct suitable for use as a replicated property to replicate input rewind history
+ */
+USTRUCT()
+struct FNetworkPhysicsRewindDataRemoteInputProxy : public FNetworkPhysicsRewindDataProxy
+{
+	GENERATED_BODY()
+		
+	bool NetSerialize(FArchive& Ar, class UPackageMap* Map, bool& bOutSuccess);
+};
+
+template<>
+struct TStructOpsTypeTraits<FNetworkPhysicsRewindDataRemoteInputProxy> : public TStructOpsTypeTraitsBase2<FNetworkPhysicsRewindDataRemoteInputProxy>
+{
+	enum
+	{
+		WithNetSerializer = true,
+		WithIdenticalViaEquality = true
+	};
+};
+
 /**
  * Struct suitable for use as a replicated property to replicate state rewind history
  */
@@ -790,18 +811,26 @@ protected :
 	UPROPERTY(Transient)
 	FNetworkPhysicsRewindDataImportantStateProxy ReplicatedImportantState;
 
-	// repnotify for the inputs on the client
+	// repnotify for inputs on owner client
 	UFUNCTION()
 	ENGINE_API void OnRep_SetReplicatedInputs();
 
+	// repnotify for inputs on remote clients
+	UFUNCTION()
+	ENGINE_API void OnRep_SetReplicatedRemoteInputs();
+
 	// repnotify for the states on the client
 	UFUNCTION()
 	ENGINE_API void OnRep_SetReplicatedStates();
 
-	// replicated physics inputs
+	// replicated physics inputs for owner client
 	UPROPERTY(Transient, ReplicatedUsing = OnRep_SetReplicatedInputs)
 	FNetworkPhysicsRewindDataInputProxy ReplicatedInputs;
 
+	// replicated physics inputs for remote clients
+	UPROPERTY(Transient, ReplicatedUsing = OnRep_SetReplicatedRemoteInputs)
+	FNetworkPhysicsRewindDataRemoteInputProxy ReplicatedRemoteInputs;
+
 	// replicated physics states 
 	UPROPERTY(Transient, ReplicatedUsing = OnRep_SetReplicatedStates)
 	FNetworkPhysicsRewindDataStateProxy ReplicatedStates;
@@ -816,16 +845,24 @@ private:
 	/** Setup data type and history in async input and trigger the creation of data history in the async component */
 	ENGINE_API void CreateAsyncDataHistory();
 
-	/** Set number of inputs to send over the network with each message, clamped to 1 as minimum */
-	void SetNumberOfInputsToNetwork(uint16 RedundantInputs)
+	/** Set number of inputs to send over the network with each message, clamped to 1 as minimum 
+	* From owning client to server and server to owning client. */
+	void SetNumberOfInputsToNetwork(uint16 NumInputs)
 	{
-		InputsToNetwork = FMath::Max((uint16)(1 + RedundantInputs), (uint16)1);
+		InputsToNetwork = FMath::Max((uint16)(NumInputs), (uint16)1);
+	}
+
+	/** Set number of inputs to send over the network with each message, clamped to 1 as minimum.
+	* From server to remote clients. */
+	void SetNumberOfRemoteInputsToNetwork(uint16 NumInputs)
+	{
+		InputsToNetworkRemote = FMath::Max((uint16)(NumInputs), (uint16)1);
 	}
 
 	/** Set number of states to send over the network with each message, clamped to 1 as minimum */
-	void SetNumberOfStatesToNetwork(uint16 RedundantStates)
+	void SetNumberOfStatesToNetwork(uint16 NumInputs)
 	{
-		StatesToNetwork = FMath::Max((uint16)(1 + RedundantStates), (uint16)1);
+		StatesToNetwork = FMath::Max((uint16)(NumInputs), (uint16)1);
 	}
 
 public:
@@ -843,6 +880,7 @@ private:
 
 	friend FNetworkPhysicsCallback;
 	friend struct FNetworkPhysicsRewindDataInputProxy;
+	friend struct FNetworkPhysicsRewindDataRemoteInputProxy;
 	friend struct FNetworkPhysicsRewindDataStateProxy;
 	friend struct FNetworkPhysicsRewindDataImportantInputProxy;
 	friend struct FNetworkPhysicsRewindDataImportantStateProxy;
@@ -862,10 +900,13 @@ private:
 	// Helper for the creation of state data and history with correct derived type
 	TUniquePtr<FNetworkPhysicsDataHelper> StateHelper;
 
-	// Send last N number of inputs each replication call to patch up holes due to packet loss
+	// Send last N number of inputs each replication call from owning client to server and server to owning client
 	uint16 InputsToNetwork = 3;
 
-	// Send last N number of states each replication call to patch up holes due to packet loss
+	// Send last N number of inputs each replication call from server to remote clients
+	uint16 InputsToNetworkRemote = 2;
+
+	// Send last N number of states each replication call from server to remote clients
 	uint16 StatesToNetwork = 1;
 
 	// Actor component that will be used to fill the histories
@@ -989,6 +1030,9 @@ FORCEINLINE void UNetworkPhysicsComponent::CreateDataHistory(UActorComponent* Hi
 	ReplicatedInputs.History = MakeUnique<TNetRewindHistory<typename PhysicsTraits::InputsType>>(InputsToNetwork);
 	ReplicatedInputs.Owner = this;
 
+	ReplicatedRemoteInputs.History = MakeUnique<TNetRewindHistory<typename PhysicsTraits::InputsType>>(InputsToNetworkRemote);
+	ReplicatedRemoteInputs.Owner = this;
+
 	ReplicatedStates.History = MakeUnique<TNetRewindHistory<typename PhysicsTraits::StatesType>>(StatesToNetwork);
 	ReplicatedStates.Owner = this;
 
@@ -1011,6 +1055,9 @@ FORCEINLINE void UNetworkPhysicsComponent::CreateInputHistory(UActorComponent* H
 	ReplicatedInputs.History = MakeUnique<TNetRewindHistory<InputsType>>(InputsToNetwork);
 	ReplicatedInputs.Owner = this;
 
+	ReplicatedRemoteInputs.History = MakeUnique<TNetRewindHistory<InputsType>>(InputsToNetworkRemote);
+	ReplicatedRemoteInputs.Owner = this;
+
 	ReplicatedImportantInput.History = MakeUnique<TNetRewindHistory<InputsType>>(1);
 	ReplicatedImportantInput.Owner = this;
 
diff --git a/Engine/Source/Runtime/Engine/Public/Physics/NetworkPhysicsSettingsComponent.h b/Engine/Source/Runtime/Engine/Public/Physics/NetworkPhysicsSettingsComponent.h
--- a/Engine/Source/Runtime/Engine/Public/Physics/NetworkPhysicsSettingsComponent.h
+++ b/Engine/Source/Runtime/Engine/Public/Physics/NetworkPhysicsSettingsComponent.h
@@ -9,6 +9,7 @@
 #include "CoreMinimal.h"
 #include "Components/ActorComponent.h"
 #include "Chaos/SimCallbackObject.h"
+#include "Curves/CurveFloat.h"
 
 #include "NetworkPhysicsSettingsComponent.generated.h"
 
@@ -68,6 +69,7 @@ namespace PhysicsReplicationCVars
 		extern bool bAllowInputExtrapolation;
 		extern bool bValidateDataOnGameThread;
 		extern int32 RedundantInputs;
+		extern int32 RedundantRemoteInputs;
 		extern int32 RedundantStates;
 		extern bool bCompareStateToTriggerRewind;
 		extern bool bCompareInputToTriggerRewind;
@@ -251,7 +253,7 @@ struct FNetworkPhysicsSettingsResimulation
 {
 	FNetworkPhysicsSettingsResimulation()
 	: bOverrideResimulationErrorThreshold_DEPRECATED(0)
-	, ResimulationErrorThreshold_DEPRECATED(10.0f)
+	, ResimulationErrorThreshold_DEPRECATED(10)
 	{};
 
 	GENERATED_BODY()
@@ -265,7 +267,7 @@ struct FNetworkPhysicsSettingsResimulation
 	uint32 ResimulationErrorThreshold_DEPRECATED;
 
 	UE_DEPRECATED(5.5, "GetResimulationErrorThreshold has been renamed, please use GetResimulationErrorPositionThreshold.")
-	uint32 GetResimulationErrorThreshold(uint32 DefaultValue) { return GetResimulationErrorPositionThreshold(DefaultValue); }
+	uint32 GetResimulationErrorThreshold(uint32 DefaultValue) { return (uint32)FMath::CeilToInt(GetResimulationErrorPositionThreshold(DefaultValue)); }
 
 	UPROPERTY(config, EditDefaultsOnly, Category = "Overrides", Meta = (InlineEditConditionToggle))
 	uint32 bOverrideResimulationErrorPositionThreshold : 1;
@@ -348,21 +350,36 @@ struct FNetworkPhysicsSettingsResimulation
 USTRUCT()
 struct FNetworkPhysicsSettingsNetworkPhysicsComponent
 {
+	FNetworkPhysicsSettingsNetworkPhysicsComponent()
+	: InputDecayCurve()
+	{
+		InputDecayCurve.GetRichCurve()->AddKey(0.0f, 0.0f);
+		InputDecayCurve.GetRichCurve()->AddKey(0.5f, 1.0f);
+		InputDecayCurve.GetRichCurve()->AddKey(1.0f, 1.0f);
+	};
+
 	GENERATED_BODY()
 
 	UPROPERTY(config, EditDefaultsOnly, Category = "Overrides", Meta = (InlineEditConditionToggle))
 	uint32 bOverrideRedundantInputs : 1;
-	// Overrides CVar: np2.Resim.RedundantInputs -- How many extra inputs to send with each unreliable network message, to account for packetloss.
+	// Overrides CVar: np2.Resim.RedundantInputs -- How many extra inputs to send with each unreliable network message, to account for packetloss. From owning client to server and server to owning client.
 	UPROPERTY(config, EditDefaultsOnly, Category = "Overrides", Meta = (EditCondition = "bOverrideRedundantInputs"))
-	uint16 RedundantInputs = PhysicsReplicationCVars::ResimulationCVars::RedundantInputs;
-	const uint16 GetRedundantInputs() const { return bOverrideRedundantInputs ? RedundantInputs : PhysicsReplicationCVars::ResimulationCVars::RedundantInputs; }
+	uint16 RedundantInputs = (uint16)PhysicsReplicationCVars::ResimulationCVars::RedundantInputs;
+	const uint16 GetRedundantInputs() const { return bOverrideRedundantInputs ? RedundantInputs : (uint16)PhysicsReplicationCVars::ResimulationCVars::RedundantInputs; }
+
+	UPROPERTY(config, EditDefaultsOnly, Category = "Overrides", Meta = (InlineEditConditionToggle))
+	uint32 bOverrideRedundantRemoteInputs : 1;
+	// Overrides CVar: np2.Resim.RedundantRemoteInputs -- How many extra inputs to send with each unreliable network message, to account for packetloss. From Server to remote clients.
+	UPROPERTY(config, EditDefaultsOnly, Category = "Overrides", Meta = (EditCondition = "bOverrideRedundantRemoteInputs"))
+	uint16 RedundantRemoteInputs = (uint16)PhysicsReplicationCVars::ResimulationCVars::RedundantRemoteInputs;
+	const uint16 GetRedundantRemoteInputs() const { return bOverrideRedundantRemoteInputs ? RedundantRemoteInputs : (uint16)PhysicsReplicationCVars::ResimulationCVars::RedundantRemoteInputs; }
 
 	UPROPERTY(config, EditDefaultsOnly, Category = "Overrides", Meta = (InlineEditConditionToggle))
 	uint32 bOverrideRedundantStates : 1;
-	// Overrides CVar: np2.Resim.RedundantStates -- How many extra states to send with each unreliable network message, to account for packetloss.
+	// Overrides CVar: np2.Resim.RedundantStates -- How many extra states to send with each unreliable network message, to account for packetloss. From Server to remote clients.
 	UPROPERTY(config, EditDefaultsOnly, Category = "Overrides", Meta = (EditCondition = "bOverrideRedundantStates"))
-	uint16 RedundantStates = PhysicsReplicationCVars::ResimulationCVars::RedundantStates;
-	const uint16 GetRedundantStates() const { return bOverrideRedundantStates ? RedundantStates : PhysicsReplicationCVars::ResimulationCVars::RedundantStates; }
+	uint16 RedundantStates = (uint16)PhysicsReplicationCVars::ResimulationCVars::RedundantStates;
+	const uint16 GetRedundantStates() const { return bOverrideRedundantStates ? RedundantStates : (uint16)PhysicsReplicationCVars::ResimulationCVars::RedundantStates; }
 
 	UPROPERTY(config, EditDefaultsOnly, Category = "Overrides", Meta = (InlineEditConditionToggle))
 	uint32 bOverrideCompareStateToTriggerRewind : 1;
@@ -426,6 +443,13 @@ struct FNetworkPhysicsSettingsNetworkPhysicsComponent
 	UPROPERTY(config, EditDefaultsOnly, Category = "Overrides", Meta = (EditCondition = "bOverrideApplySimProxyInputAtRuntime"))
 	bool bApplySimProxyInputAtRuntime = PhysicsReplicationCVars::ResimulationCVars::bApplySimProxyInputAtRuntime;
 	const bool GetApplySimProxyInputAtRuntime() const { return bOverrideApplySimProxyInputAtRuntime ? bApplySimProxyInputAtRuntime : PhysicsReplicationCVars::ResimulationCVars::bApplySimProxyInputAtRuntime; }
+
+	/* Curve for input decay during resimulation if input is being reused. 
+	* XAxis = Alpha value from 0.0 to 1.0 where 0 is the start of reusing input and 1 is the last time we will reuse the input this resimulation.
+	* YAxis = The Input Decay value from 0.0 to 1.0 (as a percentage where 1.0 = 100% decay) for the given Alpha. */
+	UPROPERTY(config, EditDefaultsOnly, Category = "Overrides", Meta = (XAxisName="LerpAlpha", YAxisName="DecayValue"))
+	FRuntimeFloatCurve InputDecayCurve;
+	const FRuntimeFloatCurve& GetInputDecayCurve() const { return InputDecayCurve; }
 };
 
 /*
@@ -496,11 +520,12 @@ public:
 private:
 	FNetworkPhysicsSettingsComponentAsync* NetworkPhysicsSettings_Internal;
 
+
+	/* --- Static API --- */
+private:
 	// Game Thread map of settings component per actor
 	static TMap<AActor*, UNetworkPhysicsSettingsComponent*> ObjectToSettings_External;
 
-
-	/* --- Static API --- */
 public:
 	/** Get the settings component for a specified actor  */
 	static UNetworkPhysicsSettingsComponent* GetSettingsForActor(AActor* Owner);
@@ -537,6 +562,9 @@ struct FNetworkPhysicsSettingsAsyncInput : public Chaos::FSimCallbackInput
 class FNetworkPhysicsSettingsComponentAsync : public Chaos::TSimCallbackObject<FNetworkPhysicsSettingsAsyncInput>
 {
 public:
+	FNetworkPhysicsSettingsComponentAsync();
+	~FNetworkPhysicsSettingsComponentAsync() {};
+
 	virtual void OnPreSimulate_Internal() override;
 
 private:
diff --git a/Engine/Source/Runtime/Experimental/Chaos/Public/RewindData.h b/Engine/Source/Runtime/Experimental/Chaos/Public/RewindData.h
--- a/Engine/Source/Runtime/Experimental/Chaos/Public/RewindData.h
+++ b/Engine/Source/Runtime/Experimental/Chaos/Public/RewindData.h
@@ -69,9 +69,13 @@ struct FBaseRewindHistory
 
 	/** Record data into the history buffer at a given time */
 	FORCEINLINE virtual bool RecordData(const int32 RecordFrame, const void* HistoryData) { return true; }
+
 	UE_DEPRECATED(5.4, "Deprecated, use RecordData() instead")
 	FORCEINLINE virtual bool RecordDatas(const int32 RecordFrame, const void* HistoryDatas) { return RecordData(RecordFrame, HistoryDatas); }
-	
+
+	/** Set if this history should only allow overriding of data if it has a higher value than current recorded data */
+	FORCEINLINE virtual void SetRecordDataIncremental(const bool bInIncremental) { }
+
 	/** Copy all data from local history into into @param OutHistory
 	* @param bIncludeUnimportant = If to copy unimportant data entries
 	* @param bIncludeImportant = If to copy important data entries */
@@ -138,13 +142,13 @@ template<typename DataType>
 struct TDataRewindHistory : public FBaseRewindHistory
 {
 	FORCEINLINE TDataRewindHistory(const int32 FrameCount, const bool bIsHistoryLocal) :
-		bIsLocalHistory(bIsHistoryLocal), DataHistory(), LatestFrame(INDEX_NONE), CurrentFrame(0), CurrentIndex(0), NumFrames(FrameCount)
+		bIsLocalHistory(bIsHistoryLocal), bIncremental(false), DataHistory(), LatestFrame(INDEX_NONE), CurrentFrame(0), CurrentIndex(0), NumFrames(FrameCount)
 	{
 		DataHistory.SetNum(NumFrames);
 	}
 
 	FORCEINLINE TDataRewindHistory(const int32 FrameCount) :
-		DataHistory(), LatestFrame(INDEX_NONE), CurrentFrame(0), CurrentIndex(0), NumFrames(FrameCount)
+		bIsLocalHistory(false), bIncremental(false), DataHistory(), LatestFrame(INDEX_NONE), CurrentFrame(0), CurrentIndex(0), NumFrames(FrameCount)
 	{
 		DataHistory.SetNum(NumFrames);
 	}
@@ -310,18 +314,31 @@ public :
 	UE_DEPRECATED(5.4, "Deprecated, use EvalData() instead")
 	FORCEINLINE bool EvalDatas(const int32 EvalFrame) { return EvalData(EvalFrame); }
 
-	/** Record the data from the buffer at a specific frame */
 	FORCEINLINE virtual bool RecordData(const int32 RecordFrame, const void* HistoryData) override
 	{
 		PRAGMA_DISABLE_DEPRECATION_WARNINGS // TODO: Change to LoadData() in UE 5.6 and remove deprecation pragma
 		LoadDatas(RecordFrame);
 		PRAGMA_ENABLE_DEPRECATION_WARNINGS
-		DataHistory[CurrentIndex] = *static_cast<const DataType*>(HistoryData);
+
+		const DataType* HistoryDataCast = static_cast<const DataType*>(HistoryData);
+
+		// If incremental recording, early out if currently cached data is newer than the data to record
+		if (bIncremental && DataHistory[CurrentIndex].LocalFrame >= HistoryDataCast->LocalFrame)
+		{
+			return false;
+		}
+
+		DataHistory[CurrentIndex] = *HistoryDataCast;
 
 		LatestFrame = FMath::Max(LatestFrame, DataHistory[CurrentIndex].LocalFrame);
 		return true;
 	}
 
+	FORCEINLINE virtual void SetRecordDataIncremental(const bool bInIncremental) override
+	{
+		bIncremental = bInIncremental;
+	}
+
 	/** Current data that is being loaded/recorded*/
 	DataType& GetCurrentData() { return DataHistory[CurrentIndex]; }
 	UE_DEPRECATED(5.4, "Deprecated, use GetCurrentData() instead")
@@ -413,6 +430,9 @@ protected :
 	/** Check if the history is on the local/remote client*/
 	bool bIsLocalHistory;
 
+	/** If this history only record data that is of a higher frame value than previous recorded frame on the same index */
+	bool bIncremental;
+
 	/** Data buffer holding the history */
 	TArray<DataType> DataHistory;
 
diff --git a/Engine/Source/Runtime/Experimental/ChaosVehicles/ChaosVehiclesCore/Private/SimModule/ModuleInput.cpp b/Engine/Source/Runtime/Experimental/ChaosVehicles/ChaosVehiclesCore/Private/SimModule/ModuleInput.cpp
--- a/Engine/Source/Runtime/Experimental/ChaosVehicles/ChaosVehiclesCore/Private/SimModule/ModuleInput.cpp
+++ b/Engine/Source/Runtime/Experimental/ChaosVehicles/ChaosVehiclesCore/Private/SimModule/ModuleInput.cpp
@@ -116,7 +116,30 @@ void FModuleInputValue::Merge(const FModuleInputValue& From)
 		checkf(false, TEXT("Unsupported value type for module input value!"));
 		break;
 	}
+}
+
+void FModuleInputValue::Decay(const float DecayAmount)
+{
+	if (!ShouldApplyInputDecay())
+	{
+		return;
+	}
 
+	switch (GetValueType())
+	{
+	case EModuleInputValueType::MBoolean:
+		// Don't decay boolean
+		break;
+	case EModuleInputValueType::MAxis1D:
+	case EModuleInputValueType::MAxis2D:
+	case EModuleInputValueType::MAxis3D:
+		// Reduce input value by DecayAmount
+		Value = (1.0f - DecayAmount) * Value;
+		break;
+	default:
+		checkf(false, TEXT("Unsupported value type for module input value!"));
+		break;
+	}
 }
 
 FString FModuleInputValue::ToString() const
@@ -204,6 +227,7 @@ void FModuleInputContainer::Initialize(TArray<FModuleInputSetup>& SetupData, FIn
 	for (FModuleInputSetup& Setup : SetupData)
 	{
 		int Index = AddInput(Setup.Type, Setup.InputModifierClass);
+		InputValues[Index].SetApplyInputDecay(Setup.bApplyInputDecay);
 		NameMapOut.Add(Setup.Name, Index);
 	}
 }
@@ -227,10 +251,12 @@ void FModuleInputContainer::Serialize(FArchive& Ar, UPackageMap* Map, bool& bOut
 
 	// Array of bits to mark which inputs to serialize or not
 	FNetBitArray InputsBitArray(Number);
+	FNetBitArray InputDecayBitArray(Number);
 
 	if (Ar.IsLoading())
 	{
 		Ar.SerializeBits(InputsBitArray.GetData(), Number);
+		Ar.SerializeBits(InputDecayBitArray.GetData(), Number);
 
 		InputValues.Reset(Number);
 		for (uint32 I = 0; I < Number; I++)
@@ -241,6 +267,7 @@ void FModuleInputContainer::Serialize(FArchive& Ar, UPackageMap* Map, bool& bOut
 			if (bHasNonZeroValue)
 			{
 				Value.Serialize(Ar, Map, bOutSuccess);
+				Value.SetApplyInputDecay(InputDecayBitArray.IsBitSet(I));
 			}
 			InputValues.EmplaceAt(I, Value);
 		}
@@ -253,9 +280,14 @@ void FModuleInputContainer::Serialize(FArchive& Ar, UPackageMap* Map, bool& bOut
 			if (InputValues[I].IsNonZero())
 			{
 				InputsBitArray.SetBit(I);
+				if (InputValues[I].ShouldApplyInputDecay())
+				{
+					InputDecayBitArray.SetBit(I);
+				}
 			}
 		}
 		Ar.SerializeBits(InputsBitArray.GetData(), Number);
+		Ar.SerializeBits(InputDecayBitArray.GetData(), Number);
 
 		for (uint32 I = 0; I < Number; I++)
 		{
@@ -298,6 +330,13 @@ void FModuleInputContainer::Merge(const FModuleInputContainer& From)
 	}
 }
 
+void FModuleInputContainer::Decay(const float DecayAmount)
+{
+	for (int I = 0; I < InputValues.Num(); I++)
+	{
+		InputValues[I].Decay(DecayAmount);
+	}
+}
 
 void FInputInterface::SetValue(const FName& InName, const FModuleInputValue& InValue)
 {
diff --git a/Engine/Source/Runtime/Experimental/ChaosVehicles/ChaosVehiclesCore/Public/SimModule/ModuleInput.h b/Engine/Source/Runtime/Experimental/ChaosVehicles/ChaosVehiclesCore/Public/SimModule/ModuleInput.h
--- a/Engine/Source/Runtime/Experimental/ChaosVehicles/ChaosVehiclesCore/Public/SimModule/ModuleInput.h
+++ b/Engine/Source/Runtime/Experimental/ChaosVehicles/ChaosVehiclesCore/Public/SimModule/ModuleInput.h
@@ -328,9 +328,15 @@ public:
 
 	void Merge(const FModuleInputValue& From);
 
+	void Decay(const float DecayAmount);
+
 	// Type sensitive debug stringify
 	FString ToString() const;
 
+	/** During physics resimulation, apply decay for this input while it's being extrapolated */
+	void SetApplyInputDecay(const bool bInApplyInputDecay) { bApplyInputDecay = bInApplyInputDecay; }
+	bool ShouldApplyInputDecay() { return  bApplyInputDecay; }
+
 protected:
 	void SetMagnitude(float NewSize);
 
@@ -339,6 +345,9 @@ protected:
 
 	UPROPERTY()
 	EModuleInputValueType ValueType = EModuleInputValueType::MBoolean;
+
+	UPROPERTY()
+	bool bApplyInputDecay;
 };
 
 // Supported getter specializations
@@ -473,6 +482,10 @@ struct CHAOSVEHICLESCORE_API FModuleInputSetup
 
 	UPROPERTY(EditAnywhere, Category = VehicleInput)
 	TSubclassOf<UDefaultModularVehicleInputModifier> InputModifierClass;
+
+	/** During physics resimulation, apply decay for this input while it's being extrapolated */
+	UPROPERTY(EditAnywhere, Category = VehicleInput)
+	bool bApplyInputDecay;
 };
 
 class CHAOSVEHICLESCORE_API FScopedModuleInputInitializer
@@ -548,6 +561,9 @@ public:
 
 	void Merge(const FModuleInputContainer& From);
 
+	/**  Decay input during resimulation by DecayAmount which increases over resimulation frames from 0.0 -> 1.0 when the input is being reused */
+	void Decay(const float DecayAmount);
+
 private:
 	UPROPERTY()
 	TArray<FModuleInputValue> InputValues;
