Fix SCW resume behaviour bug where under heaviest load could drop to 1 worker and not recover.

diff --git a/Engine/Source/Runtime/Engine/Private/ShaderCompiler/ShaderCompilerThreadRunnable.cpp b/Engine/Source/Runtime/Engine/Private/ShaderCompiler/ShaderCompilerThreadRunnable.cpp
--- a/Engine/Source/Runtime/Engine/Private/ShaderCompiler/ShaderCompilerThreadRunnable.cpp
+++ b/Engine/Source/Runtime/Engine/Private/ShaderCompiler/ShaderCompilerThreadRunnable.cpp
@@ -1071,9 +1071,9 @@ void FShaderCompileThreadRunnable::CheckMemoryLimitViolation()
 			// This approach suspends workers from 100% to 50% and then resumes them back up to 75%.
 			if (LimitInfo.MemoryUsed < LimitInfo.MemoryLimit / 2)
 			{
-				// Number of workers to resume should be a lower bound (i.e. division by 2 can result in zero),
-				// in which case we accept that we cannot allow the last suspended worker to be resumed.
-				const int32 NumWorkersToResume = GetNumberOfAvailableWorkers() / 2;
+				//	Number of workers to resume is half the currently inactive workers, with 1 added before the divide to always ensure at least one worker
+				//	is attempted to be woken.
+				const int32 NumWorkersToResume = ((GetNumberOfWorkers() - GetNumberOfAvailableWorkers())+1) / 2;
 				const int32 NumResumedWorkers = ResumeSuspendedWorkers(NumWorkersToResume);
 				if (NumResumedWorkers > 0)
 				{
