World Partition: hook UWorldPartition::Uninitialize call to ULevel::OnCleanupLevel instead of calling directly from there, to be able to determine if it comes from a direct call or from normal world cleanup.
During normal world cleanup, ignore actors unregistrations from world partition to fix issues during engine exit (and faster world cleanups when changing maps in the editor), since the world is going to be garbage collected.
Before this change, UWorldPartition::Uninitialize was only called when bCleanupResources was set, but looking at code that passes false to it (e.g. FSeamlessTravelHandler::Tick), the world partition should be uninitialized regardless of that flag since it's replaced and garbage collected.

diff --git a/Engine/Source/Runtime/Engine/Classes/Engine/Level.h b/Engine/Source/Runtime/Engine/Classes/Engine/Level.h
--- a/Engine/Source/Runtime/Engine/Classes/Engine/Level.h
+++ b/Engine/Source/Runtime/Engine/Classes/Engine/Level.h
@@ -671,6 +671,7 @@ public:
 	DECLARE_MULTICAST_DELEGATE_OneParam(FLevelTransformEvent, const FTransform&);
 	FLevelTransformEvent OnApplyLevelTransform;
 
+	// Event on level cleanup
 	DECLARE_MULTICAST_DELEGATE(FLevelCleanupEvent);
 	FLevelCleanupEvent OnCleanupLevel;
 
diff --git a/Engine/Source/Runtime/Engine/Private/Level.cpp b/Engine/Source/Runtime/Engine/Private/Level.cpp
--- a/Engine/Source/Runtime/Engine/Private/Level.cpp
+++ b/Engine/Source/Runtime/Engine/Private/Level.cpp
@@ -528,14 +528,6 @@ void ULevel::CleanupLevel(bool bCleanupResources, bool bUnloadFromEditor)
 {
 	OnCleanupLevel.Broadcast();
 
-	if (bCleanupResources)
-	{
-		if (UWorldPartition* WorldPartition = GetWorldPartition(); WorldPartition && WorldPartition->IsInitialized())
-		{
-			WorldPartition->Uninitialize();
-		}
-	}
-
 	bool bTrashPackage = !ULevelStreaming::ShouldReuseUnloadedButStillAroundLevels(this);
 #if WITH_EDITOR
 	if (bForcePackageTrashingAtCleanup)
diff --git a/Engine/Source/Runtime/Engine/Private/WorldPartition/WorldPartition.cpp b/Engine/Source/Runtime/Engine/Private/WorldPartition/WorldPartition.cpp
--- a/Engine/Source/Runtime/Engine/Private/WorldPartition/WorldPartition.cpp
+++ b/Engine/Source/Runtime/Engine/Private/WorldPartition/WorldPartition.cpp
@@ -858,6 +858,15 @@ void UWorldPartition::Initialize(UWorld* InWorld, const FTransform& InTransform)
 	FWorldPartitionEvents::BroadcastWorldPartitionInitialized(World, this);
 }
 
+void UWorldPartition::OnCleanupLevel()
+{
+#if WITH_EDITOR
+	FWorldPartitionLoadingContext::FNull LoadingContext;
+#endif
+
+	Uninitialize();
+}
+
 void UWorldPartition::Uninitialize()
 {
 	TRACE_CPUPROFILER_EVENT_SCOPE(UWorldPartition::Uninitialize);
@@ -1096,6 +1105,8 @@ void UWorldPartition::RegisterDelegates()
 
 		World->GetSubsystem<UWorldPartitionSubsystem>()->OnStreamingStateUpdated().AddUObject(this, &UWorldPartition::OnStreamingStateUpdated);
 	}
+
+	GetTypedOuter<ULevel>()->OnCleanupLevel.AddUObject(this, &UWorldPartition::OnCleanupLevel);
 }
 
 void UWorldPartition::UnregisterDelegates()
@@ -1147,6 +1158,8 @@ void UWorldPartition::UnregisterDelegates()
 
 		World->GetSubsystem<UWorldPartitionSubsystem>()->OnStreamingStateUpdated().RemoveAll(this);
 	}
+
+	GetTypedOuter<ULevel>()->OnCleanupLevel.RemoveAll(this);
 }
 
 #if !UE_BUILD_SHIPPING
diff --git a/Engine/Source/Runtime/Engine/Private/WorldPartition/WorldPartitionHandle.cpp b/Engine/Source/Runtime/Engine/Private/WorldPartition/WorldPartitionHandle.cpp
--- a/Engine/Source/Runtime/Engine/Private/WorldPartition/WorldPartitionHandle.cpp
+++ b/Engine/Source/Runtime/Engine/Private/WorldPartition/WorldPartitionHandle.cpp
@@ -77,14 +77,18 @@ void FWorldPartitionLoadingContext::UnloadAndUnregisterActor(FWorldPartitionActo
 */
 FWorldPartitionLoadingContext::IContext::IContext()
 {
-	check(ActiveContext == &DefaultContext);
-	ActiveContext = this;
+	if (ActiveContext == &DefaultContext)
+	{
+		ActiveContext = this;
+	}
 }
 
 FWorldPartitionLoadingContext::IContext::~IContext()
 {
-	check(ActiveContext == this);
-	ActiveContext = &DefaultContext;
+	if (ActiveContext == this)
+	{
+		ActiveContext = &DefaultContext;
+	}
 }
 
 /**
diff --git a/Engine/Source/Runtime/Engine/Public/WorldPartition/WorldPartition.h b/Engine/Source/Runtime/Engine/Public/WorldPartition/WorldPartition.h
--- a/Engine/Source/Runtime/Engine/Public/WorldPartition/WorldPartition.h
+++ b/Engine/Source/Runtime/Engine/Public/WorldPartition/WorldPartition.h
@@ -552,6 +552,9 @@ private:
 	bool bShouldCheckEnableStreamingWarning;
 #endif
 
+private:
+	void OnCleanupLevel();
+
 public:
 	UE_DEPRECATED(5.4, "Use GetActorDescContainerInstance")
 	UActorDescContainer* GetActorDescContainer() const { return nullptr; }
