Reduce inlining to save stack space in AnimBP graph graversal to avoid stack overflows

diff --git a/Engine/Source/Runtime/AnimGraphRuntime/Private/BoneControllers/AnimNode_SkeletalControlBase.cpp b/Engine/Source/Runtime/AnimGraphRuntime/Private/BoneControllers/AnimNode_SkeletalControlBase.cpp
--- a/Engine/Source/Runtime/AnimGraphRuntime/Private/BoneControllers/AnimNode_SkeletalControlBase.cpp
+++ b/Engine/Source/Runtime/AnimGraphRuntime/Private/BoneControllers/AnimNode_SkeletalControlBase.cpp
@@ -19,14 +19,6 @@
 #endif
 DEFINE_LOG_CATEGORY(LogSkeletalControlBase);
 
-// Macro to disable inlining of function calls within the statement that follows
-// e.g: UE_DONT_INLINE_CALL ThisFunctionCallWontBeInlined(Foo->ThisIsntInlinedEither());
-#if defined(__clang__) && defined(__has_cpp_attribute) && __has_cpp_attribute(clang::noinline)
-#define UE_DONT_INLINE_CALL [[clang::noinline]]
-#else
-#define UE_DONT_INLINE_CALL
-#endif
-
 /////////////////////////////////////////////////////
 // FAnimNode_SkeletalControlBase
 
diff --git a/Engine/Source/Runtime/Core/Public/Misc/CoreMiscDefines.h b/Engine/Source/Runtime/Core/Public/Misc/CoreMiscDefines.h
--- a/Engine/Source/Runtime/Core/Public/Misc/CoreMiscDefines.h
+++ b/Engine/Source/Runtime/Core/Public/Misc/CoreMiscDefines.h
@@ -74,6 +74,16 @@
 	#define FORCEINLINE_DEBUGGABLE FORCEINLINE
 #endif
 
+// Macro to disable inlining of function calls within the statement that follows
+// e.g: UE_DONT_INLINE_CALL ThisFunctionCallWontBeInlined(Foo->ThisIsntInlinedEither());
+#if !defined(UE_DONT_INLINE_CALL)
+	#if defined(__clang__) && defined(__has_cpp_attribute) && __has_cpp_attribute(clang::noinline)
+		#define UE_DONT_INLINE_CALL [[clang::noinline]]
+	#else
+		#define UE_DONT_INLINE_CALL
+	#endif
+#endif
+
 
 #if STATS
 	#define CLOCK_CYCLES(Timer)   {Timer -= FPlatformTime::Cycles();}
diff --git a/Engine/Source/Runtime/Engine/Private/Animation/AnimNodeBase.cpp b/Engine/Source/Runtime/Engine/Private/Animation/AnimNodeBase.cpp
--- a/Engine/Source/Runtime/Engine/Private/Animation/AnimNodeBase.cpp
+++ b/Engine/Source/Runtime/Engine/Private/Animation/AnimNodeBase.cpp
@@ -223,7 +223,7 @@ void FPoseLinkBase::Initialize(const FAnimationInitializeContext& InContext)
 		FAnimationInitializeContext LinkContext(InContext);
 		LinkContext.SetNodeId(LinkID);
 		TRACE_SCOPED_ANIM_NODE(LinkContext);
-		LinkedNode->Initialize_AnyThread(LinkContext);
+		UE_DONT_INLINE_CALL LinkedNode->Initialize_AnyThread(LinkContext);
 	}
 }
 
@@ -289,7 +289,7 @@ void FPoseLinkBase::CacheBones(const FAnimationCacheBonesContext& InContext)
 
 	if (LinkedNode != nullptr)
 	{
-		LinkedNode->CacheBones_AnyThread(InContext);
+		UE_DONT_INLINE_CALL LinkedNode->CacheBones_AnyThread(InContext);
 	}
 }
 
@@ -340,7 +340,7 @@ void FPoseLinkBase::Update(const FAnimationUpdateContext& InContext)
 			UE::Anim::FNodeFunctionCaller::BecomeRelevant(LinkContext, *LinkedNode);
 			UE::Anim::FNodeFunctionCaller::Update(LinkContext, *LinkedNode);
 		}
-		LinkedNode->Update_AnyThread(LinkContext);
+		UE_DONT_INLINE_CALL LinkedNode->Update_AnyThread(LinkContext);
 	}
 }
 
@@ -389,7 +389,7 @@ void FPoseLink::Evaluate(FPoseContext& Output)
 		{
 			Output.SetNodeId(LinkID);
 			TRACE_SCOPED_ANIM_NODE(Output);
-			LinkedNode->Evaluate_AnyThread(Output);
+			UE_DONT_INLINE_CALL LinkedNode->Evaluate_AnyThread(Output);
 			TRACE_ANIM_NODE_BLENDABLE_ATTRIBUTES(Output, SourceID, LinkID);
 		}
 
@@ -461,7 +461,7 @@ void FComponentSpacePoseLink::EvaluateComponentSpace(FComponentSpacePoseContext&
 		{
 			Output.SetNodeId(LinkID);
 			TRACE_SCOPED_ANIM_NODE(Output);
-			LinkedNode->EvaluateComponentSpace_AnyThread(Output);
+			UE_DONT_INLINE_CALL LinkedNode->EvaluateComponentSpace_AnyThread(Output);
 			TRACE_ANIM_NODE_BLENDABLE_ATTRIBUTES(Output, SourceID, LinkID);
 		}
 
diff --git a/Engine/Source/Runtime/Engine/Public/BonePose.h b/Engine/Source/Runtime/Engine/Public/BonePose.h
--- a/Engine/Source/Runtime/Engine/Public/BonePose.h
+++ b/Engine/Source/Runtime/Engine/Public/BonePose.h
@@ -493,7 +493,7 @@ struct FCSPose
 	* But the blending is done in Local Space. Also we need to refresh any Children they have
 	* that has been previously converted to Component Space.
 	*/
-	void LocalBlendCSBoneTransforms(TConstArrayView<FBoneTransform> BoneTransforms, float Alpha);
+	FORCEINLINE_DEBUGGABLE void LocalBlendCSBoneTransforms(TConstArrayView<FBoneTransform> BoneTransforms, float Alpha);
 
 	/** This function convert component space to local space to OutPose 
 	 *
@@ -503,7 +503,7 @@ struct FCSPose
 	 *
 	 * If a part of chain hasn't been converted, it will trigger ensure. 
 	 */
-	static void ConvertComponentPosesToLocalPoses(FCSPose<PoseType>&& InPose, PoseType& OutPose);
+	FORCEINLINE_DEBUGGABLE static void ConvertComponentPosesToLocalPoses(FCSPose<PoseType>&& InPose, PoseType& OutPose);
 	
 	/** This function convert component space to local space to OutPose 
 	 *
@@ -512,7 +512,7 @@ struct FCSPose
 	 *
 	 * If a part of chain hasn't been converted, it will trigger ensure. 
 	 */
-	static void ConvertComponentPosesToLocalPoses(const FCSPose<PoseType>& InPose, PoseType& OutPose);
+	FORCEINLINE_DEBUGGABLE static void ConvertComponentPosesToLocalPoses(const FCSPose<PoseType>& InPose, PoseType& OutPose);
 
 	/**
 	 * This function convert component space to local space to OutPose
@@ -525,7 +525,7 @@ struct FCSPose
 	 *
 	 * However it is more expensive as a result. 
 	 */
-	static void ConvertComponentPosesToLocalPosesSafe(FCSPose<PoseType>& InPose, PoseType& OutPose);
+	FORCEINLINE_DEBUGGABLE static void ConvertComponentPosesToLocalPosesSafe(FCSPose<PoseType>& InPose, PoseType& OutPose);
 protected:
 	PoseType Pose;
 
