Fix UWorldPartitionLevelStreamingPolicy::GetSubObject to properly resolve EditorPaths at runtime

diff --git a/Engine/Source/Runtime/Engine/Private/WorldPartition/WorldPartitionLevelStreamingPolicy.cpp b/Engine/Source/Runtime/Engine/Private/WorldPartition/WorldPartitionLevelStreamingPolicy.cpp
--- a/Engine/Source/Runtime/Engine/Private/WorldPartition/WorldPartitionLevelStreamingPolicy.cpp
+++ b/Engine/Source/Runtime/Engine/Private/WorldPartition/WorldPartitionLevelStreamingPolicy.cpp
@@ -218,7 +218,8 @@ bool UWorldPartitionLevelStreamingPolicy::ConvertEditorPathToRuntimePath(const F
 		{
 			FString OutSubObjectString;
 			const UObject* OutLevelMountPointContext = nullptr;
-			const FName* CellName = FindCellNameForSubObject(SubObjectString, true, OutSubObjectString, OutLevelMountPointContext);
+			UObject* OutCellOuter = nullptr;
+			const FName* CellName = FindCellNameForSubObject(SubObjectString, true, OutSubObjectString, OutLevelMountPointContext, OutCellOuter);
 			FString SubPathString = SubObjectContext + TEXT(".") + OutSubObjectString;
 
 			if (!CellName)
@@ -273,22 +274,21 @@ UObject* UWorldPartitionLevelStreamingPolicy::GetSubObject(const TCHAR* SubObjec
 {
 	TRACE_CPUPROFILER_EVENT_SCOPE(UWorldPartitionLevelStreamingPolicy::GetSubObject);
 
-	// Support for subobjects such as Actor.Component
-	FString SubObjectName;
-	FString SubObjectContext;
-	if (!FString(SubObjectPath).Split(TEXT("."), &SubObjectContext, &SubObjectName))
-	{
-		SubObjectContext = SubObjectPath;
-	}
+	const FString SrcPath = UWorld::RemovePIEPrefix(SubObjectPath);
 
-	const FString SrcPath = UWorld::RemovePIEPrefix(*SubObjectContext);
-	if (const UWorldPartitionRuntimeLevelStreamingCell* Cell = FindCellForSubObject(*SrcPath))
+	FString OutSubObjectString;
+	const UObject* OutLevelMountPointContext = nullptr;
+	UObject* OutCellOuter = nullptr;
+	if (const FName* CellName = FindCellNameForSubObject(SrcPath, /*bInResolveContainers=*/true, OutSubObjectString, OutLevelMountPointContext, OutCellOuter))
 	{
-		if (UWorldPartitionLevelStreamingDynamic* LevelStreaming = Cell->GetLevelStreaming())
+		if (const UWorldPartitionRuntimeLevelStreamingCell* Cell = (UWorldPartitionRuntimeLevelStreamingCell*)StaticFindObject(UWorldPartitionRuntimeLevelStreamingCell::StaticClass(), OutCellOuter, *(CellName->ToString())))
 		{
-			if (LevelStreaming->GetLoadedLevel())
+			if (UWorldPartitionLevelStreamingDynamic* LevelStreaming = Cell->GetLevelStreaming())
 			{
-				return StaticFindObject(UObject::StaticClass(), LevelStreaming->GetLoadedLevel(), SubObjectPath);
+				if (LevelStreaming->GetLoadedLevel())
+				{
+					return StaticFindObject(UObject::StaticClass(), LevelStreaming->GetLoadedLevel(), *OutSubObjectString);
+				}
 			}
 		}
 	}
@@ -402,9 +402,10 @@ void UWorldPartitionLevelStreamingPolicy::DrawRuntimeCellsDetails(UCanvas* Canva
 	Offset.Y = MaxPosY;
 }
 
-const FName* UWorldPartitionLevelStreamingPolicy::FindCellNameForSubObject(const FString& InSubObjectString, bool bInResolveContainers, FString& OutSubPathString, const UObject*& OutLevelMountPointContext) const
+const FName* UWorldPartitionLevelStreamingPolicy::FindCellNameForSubObject(const FString& InSubObjectString, bool bInResolveContainers, FString& OutSubPathString, const UObject*& OutLevelMountPointContext, UObject*& OutCellOuter) const
 {
 	OutLevelMountPointContext = nullptr;
+	OutCellOuter = nullptr;
 
 	FString SubObjectString;
 	FString SubObjectContext(InSubObjectString);
@@ -415,6 +416,7 @@ const FName* UWorldPartitionLevelStreamingPolicy::FindCellNameForSubObject(const
 
 	if (const FName* CellName = SubObjectsToCellRemapping.Find(*SubObjectContext))
 	{
+		OutCellOuter = GetOuterUWorldPartition()->RuntimeHash;
 		return CellName;
 	}
 
@@ -424,6 +426,7 @@ const FName* UWorldPartitionLevelStreamingPolicy::FindCellNameForSubObject(const
 		if (const FName* CellName = ExternalStreamingObjectPtr->SubObjectsToCellRemapping.Find(*SubObjectContext))
 		{
 			OutLevelMountPointContext = ExternalStreamingObjectPtr->GetLevelMountPointContextObject();
+			OutCellOuter = ExternalStreamingObject->Get();
 			return CellName;
 		}
 	}
@@ -433,7 +436,7 @@ const FName* UWorldPartitionLevelStreamingPolicy::FindCellNameForSubObject(const
 	{
 		if (ContainerResolver.ResolveContainerPath(InSubObjectString, OutResolvedSubPathString))
 		{
-			return FindCellNameForSubObject(OutResolvedSubPathString, false, OutSubPathString, OutLevelMountPointContext);
+			return FindCellNameForSubObject(OutResolvedSubPathString, false, OutSubPathString, OutLevelMountPointContext, OutCellOuter);
 		}
 
 		for (const TWeakObjectPtr<URuntimeHashExternalStreamingObjectBase>& ExternalStreamingObject : ExternalStreamingObjects)
@@ -442,42 +445,11 @@ const FName* UWorldPartitionLevelStreamingPolicy::FindCellNameForSubObject(const
 			{
 				if (ExternalStreamingObject.Get()->ContainerResolver.ResolveContainerPath(InSubObjectString, OutResolvedSubPathString))
 				{
-					return FindCellNameForSubObject(OutResolvedSubPathString, false, OutSubPathString, OutLevelMountPointContext);
+					return FindCellNameForSubObject(OutResolvedSubPathString, false, OutSubPathString, OutLevelMountPointContext, OutCellOuter);
 				}
 			}
 		}
 	}
 			
-	return nullptr;
-}
-
-const UWorldPartitionRuntimeLevelStreamingCell* UWorldPartitionLevelStreamingPolicy::FindCellForSubObject(FName SubObjectName) const
-{
-	FName CellName;
-	UObject* CellPackage = nullptr;
-	if (const FName* FoundCell = SubObjectsToCellRemapping.Find(SubObjectName))
-	{
-		CellName = *FoundCell;
-		CellPackage = GetOuterUWorldPartition()->RuntimeHash;
-	}
-
-	if (!CellPackage)
-	{
-		if (const TWeakObjectPtr<URuntimeHashExternalStreamingObjectBase>* ExternalStreamingObject = SubObjectsToExternalStreamingObjectsRemapping.Find(SubObjectName); ExternalStreamingObject && ExternalStreamingObject->IsValid())
-		{
-			if (const FName* FoundCell = ExternalStreamingObject->Get()->SubObjectsToCellRemapping.Find(SubObjectName))
-			{
-				CellName = *FoundCell;
-				CellPackage = ExternalStreamingObject->Get();
-			}
-		}
-	}
-
-	if (CellPackage)
-	{
-		check(CellName.IsValid());
-		return (UWorldPartitionRuntimeLevelStreamingCell*)StaticFindObject(UWorldPartitionRuntimeLevelStreamingCell::StaticClass(), CellPackage, *(CellName.ToString()));
-	}
-
 	return nullptr;
 }
\ No newline at end of file
diff --git a/Engine/Source/Runtime/Engine/Public/WorldPartition/WorldPartitionLevelStreamingPolicy.h b/Engine/Source/Runtime/Engine/Public/WorldPartition/WorldPartitionLevelStreamingPolicy.h
--- a/Engine/Source/Runtime/Engine/Public/WorldPartition/WorldPartitionLevelStreamingPolicy.h
+++ b/Engine/Source/Runtime/Engine/Public/WorldPartition/WorldPartitionLevelStreamingPolicy.h
@@ -50,8 +50,7 @@ protected:
 	void ForEachActiveRuntimeCell(TFunctionRef<void(const UWorldPartitionRuntimeCell*)> Func) const;
 
 private:
-	const FName* FindCellNameForSubObject(const FString& InSubObjectString, bool bInResolveContainers, FString& OutSubObjectString, const UObject*& OutLevelMountPointContext) const;
-	const UWorldPartitionRuntimeLevelStreamingCell* FindCellForSubObject(FName SubObjectName) const;
+	const FName* FindCellNameForSubObject(const FString& InSubObjectString, bool bInResolveContainers, FString& OutSubObjectString, const UObject*& OutLevelMountPointContext, UObject*& OutCellOuter) const;
 
 	UPROPERTY()
 	FTopLevelAssetPath SourceWorldAssetPath;
