Cache the value for ULevel::IsMapBuildDataOwner() since the rare cases of it's usage inside ULevel::BeginDestroy() are not going to work due to annotations being released.

diff --git a/Engine/Source/Runtime/Engine/Classes/Engine/Level.h b/Engine/Source/Runtime/Engine/Classes/Engine/Level.h
--- a/Engine/Source/Runtime/Engine/Classes/Engine/Level.h
+++ b/Engine/Source/Runtime/Engine/Classes/Engine/Level.h
@@ -600,6 +600,9 @@ private:
 	/** Whether the level will force trash its packages in CleanupLevel (including object packages of objects outered to the level) (editor-only) */
 	uint8										bForcePackageTrashingAtCleanup:1;
 #endif
+
+	/** Whether the level owns it's MapBuildData and associated resources and must release them (LevelInstance will use their owner level MapBuildData) */
+	uint8										bIsMapBuildDataOwner:1;
 	
 public:
 	/** The below variables are used temporarily while making a level visible.				*/
diff --git a/Engine/Source/Runtime/Engine/Private/Level.cpp b/Engine/Source/Runtime/Engine/Private/Level.cpp
--- a/Engine/Source/Runtime/Engine/Private/Level.cpp
+++ b/Engine/Source/Runtime/Engine/Private/Level.cpp
@@ -448,6 +448,7 @@ ULevel::ULevel( const FObjectInitializer& ObjectInitializer )
 	IncrementalUnregisterComponentState = EIncrementalUnregisterComponentState::Init;
 	CurrentActorIndexForIncrementalUpdate = 0;
 	CurrentActorIndexForUnregisterComponents = 0;
+	bIsMapBuildDataOwner = true;
 }
 
 void ULevel::Initialize(const FURL& InURL)
@@ -3363,6 +3364,7 @@ void ULevel::OnLevelLoaded()
 	if (ULevel* OwningLevel = ULevelInstanceSubsystem::GetOwningLevel(this, true))
 	{
 		MapBuildData = OwningLevel->MapBuildData;
+		bIsMapBuildDataOwner = (OwningLevel == this);
 	}
 
 	// 1. Cook commandlet does it's own UWorldPartition::Initialize call in FWorldPartitionCookPackageSplitter::GetGenerateList
@@ -3817,7 +3819,7 @@ UPackage* ULevel::CreateMapBuildDataPackage() const
 
 bool ULevel::IsMapBuildDataOwner() const
 {
-	return ULevelInstanceSubsystem::GetOwningLevel(this, true) == this;
+	return bIsMapBuildDataOwner;
 }
 
 UMapBuildDataRegistry* ULevel::GetOrCreateMapBuildData()
@@ -3847,6 +3849,7 @@ UMapBuildDataRegistry* ULevel::GetOrCreateMapBuildData()
 			if (ULevelInstanceSubsystem::GetOwningLevel(Level, true) == this)
 			{
 				Level->MapBuildData = MapBuildData;
+				Level->bIsMapBuildDataOwner = false;
 			}
 		}
 		MarkPackageDirty();
