[Engine]
* Fixed non unity/pch compile error

diff --git a/Engine/Source/Runtime/Engine/Private/StaticMesh.cpp b/Engine/Source/Runtime/Engine/Private/StaticMesh.cpp
--- a/Engine/Source/Runtime/Engine/Private/StaticMesh.cpp
+++ b/Engine/Source/Runtime/Engine/Private/StaticMesh.cpp
@@ -13,6 +13,7 @@
 #include "Misc/ScopedSlowTask.h"
 #include "Engine/CollisionProfile.h"
 #include "Math/ScaleRotationTranslationMatrix.h"
+#include "Serialization/MemoryReader.h"
 #include "UObject/AssetRegistryTagsContext.h"
 #include "UObject/EditorObjectVersion.h"
 #include "PhysicsEngine/BoxElem.h"
