Added non-shipping tracking for ResetAsyncTrace Delegates - The delegate dispatch step inside ResetAsyncTraces is now Timed, and if it's above a defined threshold, we dump some key info about those delegates to help diagnose rare performance spikes

#tests Tested against Dedicated Server Win64 Cooked Data
[RN] minor foundation/core

diff --git a/Engine/Config/BaseDeviceProfiles.ini b/Engine/Config/BaseDeviceProfiles.ini
--- a/Engine/Config/BaseDeviceProfiles.ini
+++ b/Engine/Config/BaseDeviceProfiles.ini
@@ -198,10 +198,12 @@ BaseProfileName=Windows
 [WindowsClient DeviceProfile]
 DeviceType=Windows
 BaseProfileName=Windows
++CVars=AsyncTraceDelegateHitchThresholdMS=0.75
 
 [WindowsServer DeviceProfile]
 DeviceType=Windows
 BaseProfileName=Windows
++CVars=AsyncTraceDelegateHitchThresholdMS=7.0
 
 [WindowsCookedEditor DeviceProfile]
 DeviceType=Windows
@@ -333,6 +335,7 @@ bIsVisibleForAssets=True
 ; Optimize Slate signed distance field text for performance
 +CVars=UI.SlateSDFText.RasterizationMode=SdfApproximation
 +CVars=UI.SlateSDFText.ResolutionLevel=1
++CVars=AsyncTraceDelegateHitchThresholdMS=0.75
 
 ; PreviewAllowlistCVars and PreviewDenyListCVars are arrays of cvars (or prefixes) that are included or excluded from being applied in editor preview
 ; If any PreviewAllowlistCVars is set, cvars are denied by default.
@@ -854,6 +857,7 @@ bIsVisibleForAssets=True
 ; PreviewAllowlistCVars and PreviewDenyListCVars are arrays of cvars (or prefixes) that are included or excluded from being applied in editor preview.
 ; If any PreviewAllowlistCVars is set, cvars are denied by default.
 PreviewAllowlistCVars=sg.
++CVars=AsyncTraceDelegateHitchThresholdMS=1.5
 
 [Android_Preview_OpenGL DeviceProfile]
 DeviceType=Android
@@ -1263,12 +1267,14 @@ DeviceType=Linux
 BaseProfileName=Linux
 MeshLODSettings=
 TextureLODSettings=
++CVars=AsyncTraceDelegateHitchThresholdMS=7.0
 
 [LinuxArm64Server DeviceProfile]
 DeviceType=Linux
 BaseProfileName=Linux
 MeshLODSettings=
 TextureLODSettings=
++CVars=AsyncTraceDelegateHitchThresholdMS=7.0
 
 [LinuxCookedEditor DeviceProfile]
 DeviceType=Linux
diff --git a/Engine/Source/Runtime/Engine/Classes/Engine/World.h b/Engine/Source/Runtime/Engine/Classes/Engine/World.h
--- a/Engine/Source/Runtime/Engine/Classes/Engine/World.h
+++ b/Engine/Source/Runtime/Engine/Classes/Engine/World.h
@@ -1465,6 +1465,10 @@ private:
 	/** The state of async tracing - abstracted into its own object for easier reference */
 	FWorldAsyncTraceState AsyncTraceState;
 
+#if !UE_BUILD_SHIPPING
+	double AsyncDelegateHitchLoggingLastTimestamp = 0.0;
+#endif //!UE_BUILD_SHIPPING
+
 #if WITH_EDITOR
 	/**	Objects currently being debugged in Kismet	*/
 	FBlueprintToDebuggedObjectMap BlueprintObjectsBeingDebugged;
diff --git a/Engine/Source/Runtime/Engine/Private/Collision/WorldCollisionAsync.cpp b/Engine/Source/Runtime/Engine/Private/Collision/WorldCollisionAsync.cpp
--- a/Engine/Source/Runtime/Engine/Private/Collision/WorldCollisionAsync.cpp
+++ b/Engine/Source/Runtime/Engine/Private/Collision/WorldCollisionAsync.cpp
@@ -13,8 +13,10 @@
 #include "Misc/Fork.h"
 #include "Physics/Experimental/PhysInterface_Chaos.h"
 #include "ProfilingDebugging/CsvProfiler.h"
+#include "ProfilingDebugging/ScopedTimers.h"
 
 CSV_DEFINE_CATEGORY(WorldCollision, true);
+DEFINE_LOG_CATEGORY_STATIC(LogWorldCollision, Log, All);
 
 /**
  * Async trace functions
@@ -46,6 +48,21 @@ namespace AsyncTraceCVars
 	{
 		return RunAsyncTraceOnWorkerThread != 0 && (FApp::ShouldUseThreadingForPerformance() || FForkProcessHelper::IsForkedMultithreadInstance());
 	}
+
+	static float AsyncTraceDelegateHitchThresholdMS = -1.0f;
+	static FAutoConsoleVariableRef CVarAsyncTraceDelegateHitchThresholdMS(
+		TEXT("AsyncTraceDelegateHitchThresholdMS"),
+		AsyncTraceDelegateHitchThresholdMS,
+		TEXT("During ResetAsyncTrace, if the delegate dispatch step takes longer than this time (in Milliseconds) then output a list of the delegates to the log\n")
+		TEXT("Value of < 0 will deactivate this functionality"),
+		ECVF_Default);
+
+	static float AsyncTraceDelegateLoggingIntervalSeconds = 30.0f;
+	static FAutoConsoleVariableRef CVarAsyncTraceDelegateLoggingIntervalSeconds(
+		TEXT("AsyncTraceDelegateLoggingIntervalSeconds"),
+		AsyncTraceDelegateLoggingIntervalSeconds,
+		TEXT("If AsyncTraceDelegateHitchThresholdMS is > 0, then this value defines the minimum time in between the hitch logging, so we avoid LogSpam\n"),
+		ECVF_Default);
 }
 
 namespace
@@ -513,18 +530,59 @@ void UWorld::ResetAsyncTrace()
 	WaitForAllAsyncTraceTasks();
 	DataBufferExecuted.bAsyncTasksCompleted = true;
 
+#if !UE_BUILD_SHIPPING
+	FAutoScopedDurationTimer Timer;
+#endif //!UE_BUILD_SHIPPING
+
 	// do run delegates before starting next round
 	for (int32 Idx = 0; Idx != DataBufferExecuted.NumQueuedTraceData; ++Idx)
 	{
-		auto& TraceData = FBufferIndexPair(Idx).DatumLookupChecked(DataBufferExecuted.TraceData);
+		FTraceDatum& TraceData = FBufferIndexPair(Idx).DatumLookupChecked(DataBufferExecuted.TraceData);
 		TraceData.Delegate.ExecuteIfBound(FTraceHandle(TraceData.FrameNumber, Idx), TraceData);
 	}
 
 	for (int32 Idx = 0; Idx != DataBufferExecuted.NumQueuedOverlapData; ++Idx)
 	{
-		auto& TraceData = FBufferIndexPair(Idx).DatumLookupChecked(DataBufferExecuted.OverlapData);
+		FOverlapDatum& TraceData = FBufferIndexPair(Idx).DatumLookupChecked(DataBufferExecuted.OverlapData);
 		TraceData.Delegate.ExecuteIfBound(FTraceHandle(TraceData.FrameNumber, Idx), TraceData);
 	}
+
+#if !UE_BUILD_SHIPPING
+	Timer.Stop();
+
+	//If the delegate dispatch has taken longer than the defined threshold, log data about which delegates were fired
+	const float AsyncTraceDelegateThresholdTimeSeconds = AsyncTraceCVars::AsyncTraceDelegateHitchThresholdMS / 1000.0f;
+	if(AsyncTraceDelegateThresholdTimeSeconds > 0.0f && Timer.GetTime() > AsyncTraceDelegateThresholdTimeSeconds && FPlatformTime::Seconds() > AsyncDelegateHitchLoggingLastTimestamp + AsyncTraceCVars::AsyncTraceDelegateLoggingIntervalSeconds)
+	{
+		AsyncDelegateHitchLoggingLastTimestamp = FPlatformTime::Seconds();
+		UE_LOG(LogWorldCollision , Log, TEXT("ResetAsyncTrace has exceeded budget [Time:%f (s) Budget: %f (s) Over: %f] - Dumping Trace Delegates"), Timer.GetTime(), AsyncTraceDelegateThresholdTimeSeconds, Timer.GetTime() - AsyncTraceDelegateThresholdTimeSeconds);
+
+#if CSV_PROFILER
+		const FCsvProfiler* CSVProfiler = FCsvProfiler::Get();
+
+		if (CSVProfiler && CSVProfiler->IsCapturing())
+		{
+			UE_LOG(LogWorldCollision, Log, TEXT("ResetAsyncTrace for CSVFrame %d"), CSVProfiler->GetCaptureFrameNumber());
+		}
+#endif //CSV_PROFILER
+
+		for (int32 Idx = 0; Idx != DataBufferExecuted.NumQueuedTraceData; ++Idx)
+		{
+			const FTraceDatum& TraceData = FBufferIndexPair(Idx).DatumLookupChecked(DataBufferExecuted.TraceData);
+			const FName TraceTag = TraceData.CollisionParams.CollisionQueryParam.TraceTag;
+			const UObject* DelegateObject = TraceData.Delegate.GetUObject();
+			UE_LOG(LogWorldCollision, Log, TEXT("--  Trace  -- TraceTag: %s Object: %s"), *TraceTag.ToString(), DelegateObject ? *DelegateObject->GetName() : TEXT("nullptr"));
+		}
+
+		for (int32 Idx = 0; Idx != DataBufferExecuted.NumQueuedOverlapData; ++Idx)
+		{
+			const FOverlapDatum& TraceData = FBufferIndexPair(Idx).DatumLookupChecked(DataBufferExecuted.OverlapData);
+			const FName TraceTag = TraceData.CollisionParams.CollisionQueryParam.TraceTag;
+			const UObject* DelegateObject = TraceData.Delegate.GetUObject();
+			UE_LOG(LogWorldCollision, Log, TEXT("-- Overlap -- TraceTag: %s Object: %s"), *TraceTag.ToString(), DelegateObject ? *DelegateObject->GetName() : TEXT("nullptr"));
+		}
+	}
+#endif //!UE_BUILD_SHIPPING
 }
 
 void UWorld::FinishAsyncTrace()
