[TSAN/Material]
 - Fix stats race

#rnx

diff --git a/Engine/Source/Runtime/Engine/Private/Materials/MaterialShared.cpp b/Engine/Source/Runtime/Engine/Private/Materials/MaterialShared.cpp
--- a/Engine/Source/Runtime/Engine/Private/Materials/MaterialShared.cpp
+++ b/Engine/Source/Runtime/Engine/Private/Materials/MaterialShared.cpp
@@ -3706,7 +3706,7 @@ void FMaterial::GetOutputPrecision(EMaterialFloatPrecisionMode FloatPrecisionMod
 	}
 }
 
-TRACE_DECLARE_INT_COUNTER(Shaders_OnDemandShaderRequests, TEXT("Shaders/OnDemandShaderRequests"));
+TRACE_DECLARE_ATOMIC_INT_COUNTER(Shaders_OnDemandShaderRequests, TEXT("Shaders/OnDemandShaderRequests"));
 bool FMaterial::TryGetShaders(const FMaterialShaderTypes& InTypes, const FVertexFactoryType* InVertexFactoryType, FMaterialShaders& OutShaders) const
 {
 	// TRACE_CPUPROFILER_EVENT_SCOPE(FMaterial::TryGetShaders); // <-- disabled by default due to verbosity (hundreds of calls per frame)
