Removed world partition post asynchronous update streaming state GT task and replaced it by a manual call once the asynchronous update task completes (when wp.Runtime.UpdateStreaming.EnableAsyncUpdate = 1). Changed the asynchronous task to use ETaskPriority::Normal instead of ETaskPriority::BackgroundHigh and remove wp.Runtime.TaskPriorities.AsyncUpdateStreamingStateTask.

diff --git a/Engine/Source/Runtime/Engine/Private/WorldPartition/WorldPartitionStreamingPolicy.cpp b/Engine/Source/Runtime/Engine/Private/WorldPartition/WorldPartitionStreamingPolicy.cpp
--- a/Engine/Source/Runtime/Engine/Private/WorldPartition/WorldPartitionStreamingPolicy.cpp
+++ b/Engine/Source/Runtime/Engine/Private/WorldPartition/WorldPartitionStreamingPolicy.cpp
@@ -63,20 +63,6 @@ FAutoConsoleVariableRef UWorldPartitionStreamingPolicy::CVarAsyncUpdateStreaming
 	TEXT("Set to enable asynchronous World Partition UpdateStreamingState."),
 	ECVF_Default);
 
-namespace UE::Private::WorldPartition
-{
-	static FAutoConsoleTaskPriority CPrio_UpdateStreamingStateAsyncTaskPriority(
-		TEXT("wp.Runtime.TaskPriorities.AsyncUpdateStreamingStateTask"),
-		TEXT("Task and thread priority for world partition asynchronous UpdateStreamingState task."),
-		ENamedThreads::BackgroundThreadPriority, // Run on background threads by default...
-		ENamedThreads::HighTaskPriority, // .. at high task priority
-		ENamedThreads::NormalTaskPriority // if we have to run on a normal thread, run at normal task priority
-	);
-}
-
-DECLARE_CYCLE_STAT(TEXT("WorldPartition_AsyncUpdateStreamingState"), WPStreamingPolicy_AsyncUpdateStreamingState, STATGROUP_TaskGraphTasks);
-DECLARE_CYCLE_STAT(TEXT("WorldPartition_AsyncPostUpdateStreamingState"), WPStreamingPolicy_AsyncPostUpdateStreamingState, STATGROUP_TaskGraphTasks);
-
 UWorldPartitionStreamingPolicy::UWorldPartitionStreamingPolicy(const FObjectInitializer& ObjectInitializer)
 	: Super(ObjectInitializer) 
 	, WorldPartition(nullptr)
@@ -293,6 +279,10 @@ void UWorldPartitionStreamingPolicy::UpdateStreamingState()
 	if (WaitForAsyncUpdateStreamingState())
 	{
 		TRACE_CPUPROFILER_EVENT_SCOPE(UWorldPartitionStreamingPolicy::UpdateTargetStateFromAsyncTask);
+
+		check(AsyncUpdateStreamingStateTask.IsCompleted());
+		PostUpdateStreamingStateInternal_GameThread(AsyncTaskTargetState);
+
 		// Update Target State using asynchronous task results
 		// Filter result (AsyncTaskTargetState) using CurrentState as asynchronous tasks started with a snapshot version of ActivatedCells/LoadedCells
 		for (const UWorldPartitionRuntimeCell* Cell : AsyncTaskTargetState.ToActivateCells)
@@ -313,8 +303,7 @@ void UWorldPartitionStreamingPolicy::UpdateStreamingState()
 		// Reset everything related to last asynchronous tasks
 		check(AsyncUpdateTaskState == EAsyncUpdateTaskState::Started);
 		AsyncUpdateTaskState = EAsyncUpdateTaskState::None;
-		AsyncUpdateStreamingStateTask = nullptr;
-		AsyncPostUpdateStreamingStateTask = nullptr;
+		AsyncUpdateStreamingStateTask = UE::Tasks::TTask<void>();
 		AsyncTaskCurrentState.Reset();
 		AsyncTaskTargetState.Reset();
 	}
@@ -721,14 +710,16 @@ void UWorldPartitionStreamingPolicy::OnPreChangeStreamingContent()
 
 bool UWorldPartitionStreamingPolicy::WaitForAsyncUpdateStreamingState()
 {
-	if (AsyncPostUpdateStreamingStateTask.IsValid())
+	if (AsyncUpdateStreamingStateTask.IsValid())
 	{
-		// Wait for completion
-		TRACE_CPUPROFILER_EVENT_SCOPE(UWorldPartitionStreamingPolicy::WaitForAsyncUpdateStreamingState);
 		check(AsyncUpdateTaskState == EAsyncUpdateTaskState::Started);
-		AsyncPostUpdateStreamingStateTask->Wait();
-		check(AsyncUpdateStreamingStateTask.IsValid());
-		check(AsyncUpdateStreamingStateTask->IsComplete());
+		// Wait for completion
+		if (!AsyncUpdateStreamingStateTask.IsCompleted())
+		{
+			QUICK_SCOPE_CYCLE_COUNTER(STAT_WorldPartitionStreamingPolicy_WaitForAsyncUpdateStreamingState);
+			AsyncUpdateStreamingStateTask.Wait();
+			check(AsyncUpdateStreamingStateTask.IsCompleted());
+		}
 		return true;
 	}
 	return false;
@@ -769,12 +760,11 @@ void UWorldPartitionStreamingPolicy::OnStreamingStateUpdated()
 			// Create async tasks
 			TRACE_CPUPROFILER_EVENT_SCOPE(UWorldPartitionStreamingPolicy::CreateAndDispatchAsyncTasks);
 			FUpdateStreamingStateParams InputParams = FUpdateStreamingStateParams(this, AsyncTaskCurrentState).SetRequiredWorldDataLayersEffectiveStatesCopy(true);
-			AsyncUpdateStreamingStateTask = FFunctionGraphTask::CreateAndDispatchWhenReady([InputParams = MoveTemp(InputParams), this]()
+			AsyncUpdateStreamingStateTask = UE::Tasks::Launch(UE_SOURCE_LOCATION, [InputParams = MoveTemp(InputParams), this]()
 			{
-				TRACE_CPUPROFILER_EVENT_SCOPE(UWorldPartitionStreamingPolicy::AsyncUpdateStreamingState);
+				QUICK_SCOPE_CYCLE_COUNTER(STAT_WorldPartitionStreamingPolicy_AsyncUpdateStreamingState);
 				UWorldPartitionStreamingPolicy::UpdateStreamingStateInternal(InputParams, AsyncTaskTargetState);
-			}, GET_STATID(WPStreamingPolicy_AsyncUpdateStreamingState), NULL, UE::Private::WorldPartition::CPrio_UpdateStreamingStateAsyncTaskPriority.Get());
-			AsyncPostUpdateStreamingStateTask = FFunctionGraphTask::CreateAndDispatchWhenReady([this]() { PostUpdateStreamingStateInternal_GameThread(AsyncTaskTargetState); }, GET_STATID(WPStreamingPolicy_AsyncPostUpdateStreamingState), AsyncUpdateStreamingStateTask, ENamedThreads::GameThread);
+			}, UE::Tasks::ETaskPriority::Normal);
 			AsyncUpdateTaskState = EAsyncUpdateTaskState::Started;
 		}
 	}
diff --git a/Engine/Source/Runtime/Engine/Public/WorldPartition/WorldPartitionStreamingPolicy.h b/Engine/Source/Runtime/Engine/Public/WorldPartition/WorldPartitionStreamingPolicy.h
--- a/Engine/Source/Runtime/Engine/Public/WorldPartition/WorldPartitionStreamingPolicy.h
+++ b/Engine/Source/Runtime/Engine/Public/WorldPartition/WorldPartitionStreamingPolicy.h
@@ -236,10 +236,10 @@ private:
 	uint32 UpdateStreamingSourcesHash;
 	uint32 UpdateStreamingStateCounter;
 
-	// Asynchronous update tasks
+	// Asynchronous update task
 	EAsyncUpdateTaskState AsyncUpdateTaskState;
-	FGraphEventRef AsyncUpdateStreamingStateTask;
-	FGraphEventRef AsyncPostUpdateStreamingStateTask;
+	UE::Tasks::TTask<void> AsyncUpdateStreamingStateTask;
+
 	int32 AsyncShouldSkipUpdateCounter;
 
 	mutable TOptional<TSet<FName>> CachedServerDisallowStreamingOutDataLayers;
