Fixed up the physics query handler to have a query shape to make it easier to handle convex shapes going forward.

diff --git a/Engine/Source/Runtime/Engine/Private/Collision/SceneQuery.cpp b/Engine/Source/Runtime/Engine/Private/Collision/SceneQuery.cpp
--- a/Engine/Source/Runtime/Engine/Private/Collision/SceneQuery.cpp
+++ b/Engine/Source/Runtime/Engine/Private/Collision/SceneQuery.cpp
@@ -144,31 +144,37 @@ struct FRaycastSQAdditionalInputs
 };
 
 
-template<typename TGeom>
-using TGeomSQInputs = std::conditional_t<std::is_same_v<TGeom, FPhysicsGeometryCollection>, FGeomCollectionSQAdditionalInputs,
-	std::conditional_t<std::is_same_v<TGeom, FCollisionShape>, FGeomSQAdditionalInputs,
-	std::conditional_t<std::is_same_v<TGeom, FPhysicsGeometry>, FPhysicsGeometrySQAdditionalInputs,
+template<typename GeomType>
+using TGeomSQInputs = std::conditional_t<std::is_same_v<GeomType, FPhysicsGeometryCollection>, FGeomCollectionSQAdditionalInputs,
+	std::conditional_t<std::is_same_v<GeomType, FCollisionShape>, FGeomSQAdditionalInputs,
+	std::conditional_t<std::is_same_v<GeomType, FPhysicsGeometry>, FPhysicsGeometrySQAdditionalInputs,
 	void>>>;
 
-template<typename TGeom>
-TGeomSQInputs<TGeom> GeomToSQInputs(const TGeom& Geom, const FQuat& Rot)
+template<typename GeomType>
+TGeomSQInputs<GeomType> GeomToSQInputs(const GeomType& Geom, const FQuat& Rot)
 {
-	static_assert(!std::is_same_v<TGeomSQInputs<TGeom>, void>, "Invalid geometry passed to SQ.");
-	if constexpr (std::is_same_v<TGeom, FPhysicsGeometryCollection>)
+	static_assert(!std::is_same_v<TGeomSQInputs<GeomType>, void>, "Invalid geometry passed to SQ.");
+	if constexpr (std::is_same_v<GeomType, FPhysicsGeometryCollection>)
 	{
 		return FGeomCollectionSQAdditionalInputs(Geom, Rot);
 	}
-	else if constexpr (std::is_same_v<TGeom, FPhysicsGeometry>)
+	else if constexpr (std::is_same_v<GeomType, FPhysicsGeometry>)
 	{
 		return FPhysicsGeometrySQAdditionalInputs(Geom, Rot);
 	}
-	else if constexpr (std::is_same_v<TGeom, FCollisionShape>)
+	else if constexpr (std::is_same_v<GeomType, FCollisionShape>)
 	{
 		return FGeomSQAdditionalInputs(Geom, Rot);
 	}
 }
 
-bool ToCollisionShape(const FPhysicsGeometryCollection& GeometryCollection, FCollisionShape& OutCollisionShape)
+bool BuildQueryShape(const FCollisionShape& CollisionShape, Chaos::FQueryShape& OutQueryShape)
+{
+	OutQueryShape.CollisionShape = CollisionShape;
+	return true;
+}
+
+bool BuildQueryShape(const FPhysicsGeometryCollection& GeometryCollection, Chaos::FQueryShape& OutQueryShape)
 {
 	ECollisionShapeType ShapeType = GeometryCollection.GetType();
 	switch (ShapeType)
@@ -176,19 +182,19 @@ bool ToCollisionShape(const FPhysicsGeometryCollection& GeometryCollection, FCol
 		case ECollisionShapeType::Box:
 		{
 			const Chaos::TBox<Chaos::FReal, 3>& Box = GeometryCollection.GetBoxGeometry();
-			OutCollisionShape = FCollisionShape::MakeBox(Box.Extents() * 0.5f);
+			OutQueryShape.CollisionShape = FCollisionShape::MakeBox(Box.Extents() * 0.5f);
 			return true;
 		}
 		case ECollisionShapeType::Sphere:
 		{
 			const Chaos::TSphere<Chaos::FReal, 3>& Sphere = GeometryCollection.GetSphereGeometry();
-			OutCollisionShape = FCollisionShape::MakeSphere(Sphere.GetRadius());
+			OutQueryShape.CollisionShape = FCollisionShape::MakeSphere(Sphere.GetRadius());
 			return true;
 		}
 		case ECollisionShapeType::Capsule:
 		{
 			const Chaos::FCapsule& Capsule = GeometryCollection.GetCapsuleGeometry();
-			OutCollisionShape = FCollisionShape::MakeCapsule(Capsule.GetRadius(), Capsule.GetHeight() * 0.5f);
+			OutQueryShape.CollisionShape = FCollisionShape::MakeCapsule(Capsule.GetRadius(), Capsule.GetHeight() * 0.5f);
 			return true;
 		}
 		default:
@@ -204,41 +210,41 @@ struct TSQTraits
 {
 	static const ESingleMultiOrTest SingleMultiOrTest = InSingleMultiOrTest;
 	static const ESweepOrRay GeometryQuery = InGeometryQuery;
-	using THitType = InHitType;
-	using TOutHits = std::conditional_t<InSingleMultiOrTest == ESingleMultiOrTest::Multi, TArray<FHitResult>, FHitResult>;
-	using THitBuffer = std::conditional_t<InSingleMultiOrTest == ESingleMultiOrTest::Multi, FDynamicHitBuffer<InHitType>, FSingleHitBuffer<InHitType>>;
+	using HitType = InHitType;
+	using OutHitsType = std::conditional_t<InSingleMultiOrTest == ESingleMultiOrTest::Multi, TArray<FHitResult>, FHitResult>;
+	using HitBufferType = std::conditional_t<InSingleMultiOrTest == ESingleMultiOrTest::Multi, FDynamicHitBuffer<InHitType>, FSingleHitBuffer<InHitType>>;
 
 	// GetNumHits - multi
 	template <ESingleMultiOrTest T = SingleMultiOrTest>
-	static typename TEnableIf<T == ESingleMultiOrTest::Multi, int32>::Type GetNumHits(const THitBuffer& HitBuffer)
+	static typename TEnableIf<T == ESingleMultiOrTest::Multi, int32>::Type GetNumHits(const HitBufferType& HitBuffer)
 	{
 		return HitBuffer.GetNumHits();
 	}
 
 	// GetNumHits - single/test
 	template <ESingleMultiOrTest T = SingleMultiOrTest>
-	static typename TEnableIf<T != ESingleMultiOrTest::Multi, int32>::Type GetNumHits(const THitBuffer& HitBuffer)
+	static typename TEnableIf<T != ESingleMultiOrTest::Multi, int32>::Type GetNumHits(const HitBufferType& HitBuffer)
 	{
 		return GetHasBlock(HitBuffer) ? 1 : 0;
 	}
 
 	//GetHits - multi
 	template <ESingleMultiOrTest T = SingleMultiOrTest>
-	static typename TEnableIf<T == ESingleMultiOrTest::Multi, THitType*>::Type GetHits(THitBuffer& HitBuffer)
+	static typename TEnableIf<T == ESingleMultiOrTest::Multi, HitType*>::Type GetHits(HitBufferType& HitBuffer)
 	{
 		return HitBuffer.GetHits();
 	}
 
 	//GetHits - single/test
 	template <ESingleMultiOrTest T = SingleMultiOrTest>
-	static typename TEnableIf<T != ESingleMultiOrTest::Multi, THitType*>::Type GetHits(THitBuffer& HitBuffer)
+	static typename TEnableIf<T != ESingleMultiOrTest::Multi, HitType*>::Type GetHits(HitBufferType& HitBuffer)
 	{
 		return GetBlock(HitBuffer);
 	}
 
 	//SceneTrace - ray
-	template <typename TAccelContainer, typename TGeomInputs, ESweepOrRay T = GeometryQuery>
-	static typename TEnableIf<T == ESweepOrRay::Raycast, void>::Type SceneTrace(const TAccelContainer& Container, const TGeomInputs& GeomInputs, const FVector& Dir, float DeltaMag, const FTransform& StartTM, THitBuffer& HitBuffer, EHitFlags OutputFlags, EQueryFlags QueryFlags, const FCollisionFilterData& FilterData, const FCollisionQueryParams& Params, ICollisionQueryFilterCallbackBase* QueryCallback)
+	template <typename AccelContainerType, typename GeomInputsType, ESweepOrRay T = GeometryQuery>
+	static typename TEnableIf<T == ESweepOrRay::Raycast, void>::Type SceneTrace(const AccelContainerType& Container, const GeomInputsType& GeomInputs, const FVector& Dir, float DeltaMag, const FTransform& StartTM, HitBufferType& HitBuffer, EHitFlags OutputFlags, EQueryFlags QueryFlags, const FCollisionFilterData& FilterData, const FCollisionQueryParams& Params, ICollisionQueryFilterCallbackBase* QueryCallback)
 	{
 		using namespace ChaosInterface;
 		FQueryFilterData QueryFilterData = MakeQueryFilterData(FilterData, QueryFlags, Params);
@@ -250,8 +256,8 @@ struct TSQTraits
 	}
 
 	//SceneTrace - sweep
-	template <typename TAccelContainer, typename TGeomInputs, ESweepOrRay T = GeometryQuery>
-	static typename TEnableIf<T == ESweepOrRay::Sweep, void>::Type SceneTrace(const TAccelContainer& Container, const TGeomInputs& GeomInputs, const FVector& Dir, float DeltaMag, const FTransform& StartTM, THitBuffer& HitBuffer, EHitFlags OutputFlags, EQueryFlags QueryFlags, const FCollisionFilterData& FilterData, const FCollisionQueryParams& Params, ICollisionQueryFilterCallbackBase* QueryCallback)
+	template <typename AccelContainerType, typename GeomInputsType, ESweepOrRay T = GeometryQuery>
+	static typename TEnableIf<T == ESweepOrRay::Sweep, void>::Type SceneTrace(const AccelContainerType& Container, const GeomInputsType& GeomInputs, const FVector& Dir, float DeltaMag, const FTransform& StartTM, HitBufferType& HitBuffer, EHitFlags OutputFlags, EQueryFlags QueryFlags, const FCollisionFilterData& FilterData, const FCollisionQueryParams& Params, ICollisionQueryFilterCallbackBase* QueryCallback)
 	{
 		using namespace ChaosInterface;
 		FQueryFilterData QueryFilterData = MakeQueryFilterData(FilterData, QueryFlags, Params);
@@ -294,8 +300,8 @@ struct TSQTraits
 		DrawTraces(World, Start, End, PGeom, GeomRotation, Hits);
 	}
 
-	template <typename TGeomInputs>
-	static void CaptureTraces(const UWorld* World, const FVector& Start, const FVector& End, const TGeomInputs& GeomInputs, ECollisionChannel TraceChannel, const FCollisionQueryParams& Params, const FCollisionResponseParams& ResponseParams, const FCollisionObjectQueryParams& ObjectParams, const TArray<FHitResult>& Hits, bool bHaveBlockingHit, double StartTime)
+	template <typename GeomInputsType>
+	static void CaptureTraces(const UWorld* World, const FVector& Start, const FVector& End, const GeomInputsType& GeomInputs, ECollisionChannel TraceChannel, const FCollisionQueryParams& Params, const FCollisionResponseParams& ResponseParams, const FCollisionObjectQueryParams& ObjectParams, const TArray<FHitResult>& Hits, bool bHaveBlockingHit, double StartTime)
 	{
 #if ENABLE_COLLISION_ANALYZER
 		ECAQueryMode::Type QueryMode = IsMulti() ? ECAQueryMode::Multi : (IsSingle() ? ECAQueryMode::Single : ECAQueryMode::Test);
@@ -310,8 +316,8 @@ struct TSQTraits
 #endif
 	}
 
-	template <typename TGeomInputs>
-	static void CaptureTraces(const UWorld* World, const FVector& Start, const FVector& End, const TGeomInputs& GeomInputs, ECollisionChannel TraceChannel, const FCollisionQueryParams& Params, const FCollisionResponseParams& ResponseParams, const FCollisionObjectQueryParams& ObjectParams, const FHitResult& Hit, bool bHaveBlockingHit, double StartTime)
+	template <typename GeomInputsType>
+	static void CaptureTraces(const UWorld* World, const FVector& Start, const FVector& End, const GeomInputsType& GeomInputs, ECollisionChannel TraceChannel, const FCollisionQueryParams& Params, const FCollisionResponseParams& ResponseParams, const FCollisionObjectQueryParams& ObjectParams, const FHitResult& Hit, bool bHaveBlockingHit, double StartTime)
 	{
 		TArray<FHitResult> Hits;
 		if (bHaveBlockingHit)
@@ -430,10 +436,10 @@ struct FDefaultAccelContainer
 	}
 };
 
-template<typename TAccel>
+template<typename AccelType>
 struct FOverrideAccelContainer
 {
-	explicit FOverrideAccelContainer(const TAccel& InSpatialAcceleration)
+	explicit FOverrideAccelContainer(const AccelType& InSpatialAcceleration)
 		: SpatialAcceleration(InSpatialAcceleration)
 	{}
 
@@ -442,17 +448,17 @@ struct FOverrideAccelContainer
 		return true;
 	}
 
-	const TAccel& GetSpatialAcceleration() const
+	const AccelType& GetSpatialAcceleration() const
 	{
 		return SpatialAcceleration;
 	}
 
 private:
-	const TAccel& SpatialAcceleration;
+	const AccelType& SpatialAcceleration;
 };
 
-template <typename Traits, typename TGeomInputs, typename TAccelContainer>
-bool TSceneCastCommonImpWithRetryRequest(const UWorld* World, typename Traits::TOutHits& OutHits, const TGeomInputs& GeomInputs, const FVector Start, const FVector End, ECollisionChannel TraceChannel, const struct FCollisionQueryParams& Params, const struct FCollisionResponseParams& ResponseParams, const struct FCollisionObjectQueryParams& ObjectParams, const TAccelContainer& AccelContainer, bool& bOutRequestRetry, FCollisionQueryParams& OutRetryParams)
+template <typename Traits, typename GeomInputsType, typename AccelContainerType>
+bool TSceneCastCommonImpWithRetryRequest(const UWorld* World, typename Traits::OutHitsType& OutHits, const GeomInputsType& GeomInputs, const FVector Start, const FVector End, ECollisionChannel TraceChannel, const struct FCollisionQueryParams& Params, const struct FCollisionResponseParams& ResponseParams, const struct FCollisionObjectQueryParams& ObjectParams, const AccelContainerType& AccelContainer, bool& bOutRequestRetry, FCollisionQueryParams& OutRetryParams)
 {
 	using namespace ChaosInterface;
 	bOutRequestRetry = false;
@@ -487,7 +493,7 @@ bool TSceneCastCommonImpWithRetryRequest(const UWorld* World, typename Traits::T
 			QueryCallback.bIgnoreTouches = true;
 		}
 
-		typename Traits::THitBuffer HitBufferSync;
+		typename Traits::HitBufferType HitBufferSync;
 
 		bool bBlockingHit = false;
 		const FVector Dir = DeltaMag > 0.f ? (Delta / DeltaMag) : FVector(1, 0, 0);
@@ -501,7 +507,7 @@ bool TSceneCastCommonImpWithRetryRequest(const UWorld* World, typename Traits::T
 			FScopedSQHitchRepeater<decltype(HitBufferSync)> HitchRepeater(HitBufferSync, QueryCallback, FHitchDetectionInfo(Start, End, TraceChannel, Params));
 			do
 			{
-				if constexpr (TAccelContainer::HasAccelerationStructureOverride())
+				if constexpr (AccelContainerType::HasAccelerationStructureOverride())
 				{
 					Traits::SceneTrace(AccelContainer.GetSpatialAcceleration(), GeomInputs, Dir, DeltaMag, StartTM, HitchRepeater.GetBuffer(), Traits::GetHitFlags(), Traits::GetQueryFlags(), Filter, Params, &QueryCallback);
 				}
@@ -687,8 +693,8 @@ bool TSceneCastCommonImpWithRetryRequest(const UWorld* World, typename Traits::T
 	return bHaveBlockingHit;
 }
 
-template <typename Traits, typename TGeomInputs, typename TAccelContainer>
-bool TSceneCastCommonImp(const UWorld* World, typename Traits::TOutHits& OutHits, const TGeomInputs& GeomInputs, const FVector Start, const FVector End, ECollisionChannel TraceChannel, const struct FCollisionQueryParams& Params, const struct FCollisionResponseParams& ResponseParams, const struct FCollisionObjectQueryParams& ObjectParams, const TAccelContainer& AccelContainer)
+template <typename Traits, typename GeomInputsType, typename AccelContainerType>
+bool TSceneCastCommonImp(const UWorld* World, typename Traits::OutHitsType& OutHits, const GeomInputsType& GeomInputs, const FVector Start, const FVector End, ECollisionChannel TraceChannel, const struct FCollisionQueryParams& Params, const struct FCollisionResponseParams& ResponseParams, const struct FCollisionObjectQueryParams& ObjectParams, const AccelContainerType& AccelContainer)
 {
 	bool bRequestRetry = true;
 	bool bReturnResult = false;
@@ -697,14 +703,14 @@ bool TSceneCastCommonImp(const UWorld* World, typename Traits::TOutHits& OutHits
 	{
 		constexpr bool bIsRetryQuery = false;
 		CVD_TRACE_SCOPED_SCENE_QUERY_HELPER(World, GeomInputs.GetGeometry(), FTransform(GeomInputs.GetGeometryOrientation() ? *GeomInputs.GetGeometryOrientation() : FQuat::Identity, Start), End, TraceChannel, Params, ResponseParams, ObjectParams, Traits::IsSweep() ? EChaosVDSceneQueryType::Sweep : EChaosVDSceneQueryType::RayCast, static_cast<EChaosVDSceneQueryMode>(Traits::SingleMultiOrTest), bIsRetryQuery);
-		bReturnResult = TSceneCastCommonImpWithRetryRequest<Traits, TGeomInputs, TAccelContainer>(World, OutHits, GeomInputs, Start, End, TraceChannel, Params, ResponseParams, ObjectParams, AccelContainer, bRequestRetry, RetryParams);
+		bReturnResult = TSceneCastCommonImpWithRetryRequest<Traits, GeomInputsType, AccelContainerType>(World, OutHits, GeomInputs, Start, End, TraceChannel, Params, ResponseParams, ObjectParams, AccelContainer, bRequestRetry, RetryParams);
 	}
 
 	int InfiniteLoopProtection = 10;
 	while (bRequestRetry && InfiniteLoopProtection > 0)
 	{
 		CVD_TRACE_SCOPED_SCENE_QUERY_HELPER(World, GeomInputs.GetGeometry(), FTransform(GeomInputs.GetGeometryOrientation() ? *GeomInputs.GetGeometryOrientation() : FQuat::Identity, Start), End, TraceChannel, Params, ResponseParams, ObjectParams, Traits::IsSweep() ? EChaosVDSceneQueryType::Sweep : EChaosVDSceneQueryType::RayCast, static_cast<EChaosVDSceneQueryMode>(Traits::SingleMultiOrTest), bRequestRetry);
-		bReturnResult = TSceneCastCommonImpWithRetryRequest<Traits, TGeomInputs, TAccelContainer>(World, OutHits, GeomInputs, Start, End, TraceChannel, RetryParams, ResponseParams, ObjectParams, AccelContainer, bRequestRetry, RetryParams);
+		bReturnResult = TSceneCastCommonImpWithRetryRequest<Traits, GeomInputsType, AccelContainerType>(World, OutHits, GeomInputs, Start, End, TraceChannel, RetryParams, ResponseParams, ObjectParams, AccelContainer, bRequestRetry, RetryParams);
 		InfiniteLoopProtection--;
 	}
 
@@ -775,20 +781,19 @@ bool RaycastWithQueryHandler(UPhysicsQueryHandler& QueryHandler, const EThreadQu
 }
 
 template <ESingleMultiOrTest SingleMultiOrTest, typename OutHitsType>
-bool SweepWithQueryHandler(UPhysicsQueryHandler& QueryHandler, const EThreadQueryContext ThreadContext, const UWorld* World, OutHitsType& OutHits, const FCollisionShape& CollisionShape, const FQuat& GeomRot, const FVector Start, const FVector End, const Chaos::FCommonQueryData& CommonData)
+bool SweepWithQueryHandler(UPhysicsQueryHandler& QueryHandler, const EThreadQueryContext ThreadContext, const UWorld* World, OutHitsType& OutHits, const Chaos::FQueryShape& QueryShape, const FQuat& GeomRot, const FVector Start, const FVector End, const Chaos::FCommonQueryData& CommonData)
 {
 	Chaos::FSweepQueryData SweepData
 	{
 		.Start = Start,
 		.End = End,
-		.CollisionShape = CollisionShape,
+		.QueryShape = QueryShape,
 		.GeomRot = GeomRot,
 	};
 	if constexpr (SingleMultiOrTest == ESingleMultiOrTest::Test)
 	{
 		return QueryHandler.SweepTest(ThreadContext, World, SweepData, CommonData, OutHits);
 	}
-
 	else if constexpr (SingleMultiOrTest == ESingleMultiOrTest::Single)
 	{
 		return QueryHandler.SweepSingle(ThreadContext, World, SweepData, CommonData, OutHits);
@@ -803,17 +808,21 @@ bool SweepWithQueryHandler(UPhysicsQueryHandler& QueryHandler, const EThreadQuer
 template <ESingleMultiOrTest SingleMultiOrTest, typename OutHitsType>
 bool SweepWithQueryHandler(UPhysicsQueryHandler& QueryHandler, const EThreadQueryContext ThreadContext, const UWorld* World, OutHitsType& OutHits, const FGeomSQAdditionalInputs& GeomInputs, const FVector Start, const FVector End, const Chaos::FCommonQueryData& CommonData)
 {
-	return SweepWithQueryHandler<SingleMultiOrTest>(QueryHandler, ThreadContext, World, OutHits, GeomInputs.CollisionShape, *GeomInputs.GetGeometryOrientation(), Start, End, CommonData);
+	Chaos::FQueryShape QueryShape;
+	if (BuildQueryShape(GeomInputs.CollisionShape, QueryShape))
+	{
+		return SweepWithQueryHandler<SingleMultiOrTest>(QueryHandler, ThreadContext, World, OutHits, QueryShape, *GeomInputs.GetGeometryOrientation(), Start, End, CommonData);
+	}
+	return false;
 }
 
 template <ESingleMultiOrTest SingleMultiOrTest, typename OutHitsType>
 bool SweepWithQueryHandler(UPhysicsQueryHandler& QueryHandler, const EThreadQueryContext ThreadContext, const UWorld* World, OutHitsType& OutHits, const FGeomCollectionSQAdditionalInputs& GeomInputs, const FVector Start, const FVector End, const Chaos::FCommonQueryData& CommonData)
 {
-	const FQuat& GeomRot = GeomInputs.GeomRot;
-	FCollisionShape CollisionShape;
-	if (ToCollisionShape(GeomInputs.Collection, CollisionShape))
+	Chaos::FQueryShape QueryShape;
+	if (BuildQueryShape(GeomInputs.Collection, QueryShape))
 	{
-		return SweepWithQueryHandler<SingleMultiOrTest>(QueryHandler, ThreadContext, World, OutHits, CollisionShape, GeomRot, Start, End, CommonData);
+		return SweepWithQueryHandler<SingleMultiOrTest>(QueryHandler, ThreadContext, World, OutHits, QueryShape, *GeomInputs.GetGeometryOrientation(), Start, End, CommonData);
 	}
 	return false;
 }
@@ -876,8 +885,8 @@ bool FGenericPhysicsInterface::RaycastTest(const UWorld* World, const FVector St
 	return TSceneCastCommon<ESweepOrRay::Raycast, ESingleMultiOrTest::Test>(World, DummyHit, FRaycastSQAdditionalInputs(), Start, End, TraceChannel, Params, ResponseParams, ObjectParams);
 }
 
-template<typename TAccel>
-bool FGenericRaycastPhysicsInterfaceUsingSpatialAcceleration<TAccel>::RaycastTest(const TAccel& Accel, const UWorld* World, const FVector Start, const FVector End, ECollisionChannel TraceChannel, const FCollisionQueryParams& Params, const FCollisionResponseParams& ResponseParams, const FCollisionObjectQueryParams& ObjectParams)
+template<typename AccelType>
+bool FGenericRaycastPhysicsInterfaceUsingSpatialAcceleration<AccelType>::RaycastTest(const AccelType& Accel, const UWorld* World, const FVector Start, const FVector End, ECollisionChannel TraceChannel, const FCollisionQueryParams& Params, const FCollisionResponseParams& ResponseParams, const FCollisionObjectQueryParams& ObjectParams)
 {
 	using namespace ChaosInterface;
 
@@ -901,8 +910,8 @@ bool FGenericPhysicsInterface::RaycastSingle(const UWorld* World, struct FHitRes
 	return TSceneCastCommon<ESweepOrRay::Raycast, ESingleMultiOrTest::Single>(World, OutHit, FRaycastSQAdditionalInputs(), Start, End, TraceChannel, Params, ResponseParams, ObjectParams);
 }
 
-template<typename TAccel>
-bool FGenericRaycastPhysicsInterfaceUsingSpatialAcceleration<TAccel>::RaycastSingle(const TAccel& Accel, const UWorld* World, struct FHitResult& OutHit, const FVector Start, const FVector End, ECollisionChannel TraceChannel, const FCollisionQueryParams& Params, const FCollisionResponseParams& ResponseParams, const FCollisionObjectQueryParams& ObjectParams)
+template<typename AccelType>
+bool FGenericRaycastPhysicsInterfaceUsingSpatialAcceleration<AccelType>::RaycastSingle(const AccelType& Accel, const UWorld* World, struct FHitResult& OutHit, const FVector Start, const FVector End, ECollisionChannel TraceChannel, const FCollisionQueryParams& Params, const FCollisionResponseParams& ResponseParams, const FCollisionObjectQueryParams& ObjectParams)
 {
 	using namespace ChaosInterface;
 
@@ -925,8 +934,8 @@ bool FGenericPhysicsInterface::RaycastMulti(const UWorld* World, TArray<struct F
 	return TSceneCastCommon<ESweepOrRay::Raycast, ESingleMultiOrTest::Multi>(World, OutHits, FRaycastSQAdditionalInputs(), Start, End, TraceChannel, Params, ResponseParams, ObjectParams);
 }
 
-template<typename TAccel>
-bool FGenericRaycastPhysicsInterfaceUsingSpatialAcceleration<TAccel>::RaycastMulti(const TAccel& Accel, const UWorld* World, TArray<struct FHitResult>& OutHits, const FVector& Start, const FVector& End, ECollisionChannel TraceChannel, const struct FCollisionQueryParams& Params, const struct FCollisionResponseParams& ResponseParams, const struct FCollisionObjectQueryParams& ObjectParams)
+template<typename AccelType>
+bool FGenericRaycastPhysicsInterfaceUsingSpatialAcceleration<AccelType>::RaycastMulti(const AccelType& Accel, const UWorld* World, TArray<struct FHitResult>& OutHits, const FVector& Start, const FVector& End, ECollisionChannel TraceChannel, const struct FCollisionQueryParams& Params, const struct FCollisionResponseParams& ResponseParams, const struct FCollisionObjectQueryParams& ObjectParams)
 {
 	using namespace ChaosInterface;
 
@@ -953,8 +962,8 @@ bool FGenericPhysicsInterface::GeomSweepTest(const UWorld* World, const struct F
 	return TSceneCastCommon<ESweepOrRay::Sweep, ESingleMultiOrTest::Test>(World, DummyHit, FGeomSQAdditionalInputs(CollisionShape, Rot), Start, End, TraceChannel, Params, ResponseParams, ObjectParams);
 }
 
-template<typename TAccel, typename TGeom>
-bool FGenericGeomPhysicsInterfaceUsingSpatialAcceleration<TAccel, TGeom>::GeomSweepTest(const TAccel& Accel, const UWorld* World, const TGeom& InGeom, const FQuat& Rot, FVector Start, FVector End, ECollisionChannel TraceChannel, const struct FCollisionQueryParams& Params, const struct FCollisionResponseParams& ResponseParams, const struct FCollisionObjectQueryParams& ObjectParams)
+template<typename AccelType, typename GeomType>
+bool FGenericGeomPhysicsInterfaceUsingSpatialAcceleration<AccelType, GeomType>::GeomSweepTest(const AccelType& Accel, const UWorld* World, const GeomType& InGeom, const FQuat& Rot, FVector Start, FVector End, ECollisionChannel TraceChannel, const struct FCollisionQueryParams& Params, const struct FCollisionResponseParams& ResponseParams, const struct FCollisionObjectQueryParams& ObjectParams)
 {
 	using namespace ChaosInterface;
 
@@ -1032,8 +1041,8 @@ bool FGenericPhysicsInterface::GeomSweepMulti(const UWorld* World, const FPhysic
 
 using EQueryInfo = Chaos::EQueryInfo;
 
-template <typename TOverlapHit, EQueryInfo InfoType, typename TCollisionAnalyzerType, typename TAccelContainer>
-bool GeomOverlapMultiImp(const UWorld* World, const FPhysicsGeometry& Geom, const TCollisionAnalyzerType& CollisionAnalyzerType, const FTransform& GeomPose, TArray<FOverlapResult>& OutOverlaps, ECollisionChannel TraceChannel, const struct FCollisionQueryParams& Params, const struct FCollisionResponseParams& ResponseParams, const struct FCollisionObjectQueryParams& ObjectParams, const TAccelContainer& AccelContainer)
+template <typename OverlapHitType, EQueryInfo InfoType, typename CollisionAnalyzerType, typename AccelContainerType>
+bool GeomOverlapMultiImp(const UWorld* World, const FPhysicsGeometry& Geom, const CollisionAnalyzerType& CollisionAnalyzerShape, const FTransform& GeomPose, TArray<FOverlapResult>& OutOverlaps, ECollisionChannel TraceChannel, const struct FCollisionQueryParams& Params, const struct FCollisionResponseParams& ResponseParams, const struct FCollisionObjectQueryParams& ObjectParams, const AccelContainerType& AccelContainer)
 {
 	using namespace ChaosInterface;
 
@@ -1066,7 +1075,7 @@ bool GeomOverlapMultiImp(const UWorld* World, const FPhysicsGeometry& Geom, cons
 		{
 			QueryFlags = QueryFlags | EQueryFlags::SkipNarrowPhase;
 		}
-		FDynamicHitBuffer<TOverlapHit> OverlapBuffer;
+		FDynamicHitBuffer<OverlapHitType> OverlapBuffer;
 
 		FQueryDebugParams DebugParams;
 #if !(UE_BUILD_TEST || UE_BUILD_SHIPPING)
@@ -1082,7 +1091,7 @@ bool GeomOverlapMultiImp(const UWorld* World, const FPhysicsGeometry& Geom, cons
 				FScopedSQHitchRepeater<typename TRemoveReference<decltype(OverlapBuffer)>::Type> HitchRepeater(OverlapBuffer, QueryCallback, FHitchDetectionInfo(GeomPose, TraceChannel, Params));
 				do
 				{
-					if constexpr (TAccelContainer::HasAccelerationStructureOverride())
+					if constexpr (AccelContainerType::HasAccelerationStructureOverride())
 					{
 						LowLevelOverlap(AccelContainer.GetSpatialAcceleration(), Geom, GeomPose, HitchRepeater.GetBuffer(), QueryFlags, Filter, MakeQueryFilterData(Filter, QueryFlags, Params), &QueryCallback, DebugParams);
 					}
@@ -1108,7 +1117,7 @@ bool GeomOverlapMultiImp(const UWorld* World, const FPhysicsGeometry& Geom, cons
 
 					// This block is only necessary on the game thread when dealing with cluster unions.
 					// TODO: Is there a way to get this to generalize better to the PT as well? Right now it depends on GT functions on the cluster union component.
-					if constexpr (std::is_base_of_v<ChaosInterface::FActorShape, TOverlapHit>)
+					if constexpr (std::is_base_of_v<ChaosInterface::FActorShape, OverlapHitType>)
 					{
 						auto DoClusterUnionOverlap = [&GeomPose, &Geom, TraceChannel, &Params, &ResponseParams, &ObjectParams](const FOverlapResult& OriginalOverlap, TArray<FOverlapResult>& NewOverlaps)
 						{
@@ -1175,7 +1184,7 @@ bool GeomOverlapMultiImp(const UWorld* World, const FPhysicsGeometry& Geom, cons
 		// Determine query mode ('single' doesn't really exist for overlaps)
 		ECAQueryMode::Type QueryMode = (InfoType == EQueryInfo::GatherAll) ? ECAQueryMode::Multi : ECAQueryMode::Test;
 
-		CAPTUREGEOMOVERLAP(World, CollisionAnalyzerType, GeomPose, QueryMode, TraceChannel, Params, ResponseParams, ObjectParams, OutOverlaps);
+		CAPTUREGEOMOVERLAP(World, CollisionAnalyzerShape, GeomPose, QueryMode, TraceChannel, Params, ResponseParams, ObjectParams, OutOverlaps);
 	}
 #endif // ENABLE_COLLISION_ANALYZER
 
@@ -1189,19 +1198,19 @@ bool GeomOverlapMultiHelper(const EThreadQueryContext ThreadContext, const UWorl
 
 	TGeomSQInputs<GeomType> SQInputs = GeomToSQInputs(InGeom, GeomPose.GetRotation());
 	const FPhysicsGeometry& Geom = *SQInputs.GetGeometry();
-	auto& CollisionAnalyzerType = *SQInputs.GetCollisionShape();
+	auto& CollisionAnalyzerShape = *SQInputs.GetCollisionShape();
 	if (ThreadContext == EThreadQueryContext::GTData)
 	{
-		return GeomOverlapMultiImp<FOverlapHit, InfoType>(World, Geom, CollisionAnalyzerType, GeomPose, OutOverlaps, TraceChannel, Params, ResponseParams, ObjectParams, AccelContainer);
+		return GeomOverlapMultiImp<FOverlapHit, InfoType>(World, Geom, CollisionAnalyzerShape, GeomPose, OutOverlaps, TraceChannel, Params, ResponseParams, ObjectParams, AccelContainer);
 	}
 	else
 	{
-		return GeomOverlapMultiImp<FPTOverlapHit, InfoType>(World, Geom, CollisionAnalyzerType, GeomPose, OutOverlaps, TraceChannel, Params, ResponseParams, ObjectParams, AccelContainer);
+		return GeomOverlapMultiImp<FPTOverlapHit, InfoType>(World, Geom, CollisionAnalyzerShape, GeomPose, OutOverlaps, TraceChannel, Params, ResponseParams, ObjectParams, AccelContainer);
 	}
 }
 
 template <EQueryInfo InfoType>
-bool GeomOverlapMultiHelperWithQueryHandler(UPhysicsQueryHandler& QueryHandler, const EThreadQueryContext ThreadContext, const UWorld* World, const FCollisionShape& CollisionShape, const FTransform& GeomPose, TArray<FOverlapResult>& OutOverlaps, ECollisionChannel TraceChannel, const struct FCollisionQueryParams& Params, const struct FCollisionResponseParams& ResponseParams, const struct FCollisionObjectQueryParams& ObjectParams)
+bool GeomOverlapMultiHelperWithQueryHandler(UPhysicsQueryHandler& QueryHandler, const EThreadQueryContext ThreadContext, const UWorld* World, const Chaos::FQueryShape& QueryShape, const FTransform& GeomPose, TArray<FOverlapResult>& OutOverlaps, ECollisionChannel TraceChannel, const struct FCollisionQueryParams& Params, const struct FCollisionResponseParams& ResponseParams, const struct FCollisionObjectQueryParams& ObjectParams)
 {
 	Chaos::FCommonQueryData CommonData
 	{
@@ -1212,19 +1221,30 @@ bool GeomOverlapMultiHelperWithQueryHandler(UPhysicsQueryHandler& QueryHandler,
 	};
 	Chaos::FOverlapQueryData OverlapData
 	{
-		.CollisionShape = CollisionShape,
+		.QueryShape = QueryShape,
 		.GeomPose = GeomPose,
 	};
 	return QueryHandler.Overlap(InfoType, ThreadContext, World, OverlapData, CommonData, OutOverlaps);
 }
 
 template <EQueryInfo InfoType>
-bool GeomOverlapMultiHelperWithQueryHandler(UPhysicsQueryHandler& QueryHandler, const EThreadQueryContext ThreadContext, const UWorld* World, const FPhysicsGeometryCollection& InCollisionShape, const FTransform& GeomPose, TArray<FOverlapResult>& OutOverlaps, ECollisionChannel TraceChannel, const struct FCollisionQueryParams& Params, const struct FCollisionResponseParams& ResponseParams, const struct FCollisionObjectQueryParams& ObjectParams)
+bool GeomOverlapMultiHelperWithQueryHandler(UPhysicsQueryHandler& QueryHandler, const EThreadQueryContext ThreadContext, const UWorld* World, const FCollisionShape& CollisionShape, const FTransform& GeomPose, TArray<FOverlapResult>& OutOverlaps, ECollisionChannel TraceChannel, const struct FCollisionQueryParams& Params, const struct FCollisionResponseParams& ResponseParams, const struct FCollisionObjectQueryParams& ObjectParams)
 {
-	FCollisionShape CollisionShape;
-	if (ToCollisionShape(InCollisionShape, CollisionShape))
+	Chaos::FQueryShape QueryShape;
+	if (BuildQueryShape(CollisionShape, QueryShape))
 	{
-		return GeomOverlapMultiHelperWithQueryHandler<InfoType>(QueryHandler, ThreadContext, World, CollisionShape, GeomPose, OutOverlaps, TraceChannel, Params, ResponseParams, ObjectParams);
+		return GeomOverlapMultiHelperWithQueryHandler<InfoType>(QueryHandler, ThreadContext, World, QueryShape, GeomPose, OutOverlaps, TraceChannel, Params, ResponseParams, ObjectParams);
+	}
+	return false;
+}
+
+template <EQueryInfo InfoType>
+bool GeomOverlapMultiHelperWithQueryHandler(UPhysicsQueryHandler& QueryHandler, const EThreadQueryContext ThreadContext, const UWorld* World, const FPhysicsGeometryCollection& CollisionShape, const FTransform& GeomPose, TArray<FOverlapResult>& OutOverlaps, ECollisionChannel TraceChannel, const struct FCollisionQueryParams& Params, const struct FCollisionResponseParams& ResponseParams, const struct FCollisionObjectQueryParams& ObjectParams)
+{
+	Chaos::FQueryShape QueryShape;
+	if (BuildQueryShape(CollisionShape, QueryShape))
+	{
+		return GeomOverlapMultiHelperWithQueryHandler<InfoType>(QueryHandler, ThreadContext, World, QueryShape, GeomPose, OutOverlaps, TraceChannel, Params, ResponseParams, ObjectParams);
 	}
 	return false;
 }
@@ -1269,8 +1289,8 @@ bool FGenericPhysicsInterface::GeomOverlapBlockingTest(const UWorld* World, cons
 	return GeomOverlapMultiHelper<EQueryInfo::IsBlocking>(World, CollisionShape, GeomTransform, Overlaps, TraceChannel, Params, ResponseParams, ObjectParams);
 }
 
-template<typename TAccel, typename TGeom>
-bool FGenericGeomPhysicsInterfaceUsingSpatialAcceleration<TAccel, TGeom>::GeomOverlapBlockingTest(const TAccel& Accel, const UWorld* World, const TGeom& InGeom, const FVector& Pos, const FQuat& Rot, ECollisionChannel TraceChannel, const struct FCollisionQueryParams& Params, const struct FCollisionResponseParams& ResponseParams, const struct FCollisionObjectQueryParams& ObjectParams)
+template<typename AccelType, typename GeomType>
+bool FGenericGeomPhysicsInterfaceUsingSpatialAcceleration<AccelType, GeomType>::GeomOverlapBlockingTest(const AccelType& Accel, const UWorld* World, const GeomType& InGeom, const FVector& Pos, const FQuat& Rot, ECollisionChannel TraceChannel, const struct FCollisionQueryParams& Params, const struct FCollisionResponseParams& ResponseParams, const struct FCollisionObjectQueryParams& ObjectParams)
 {
 	SCOPE_CYCLE_COUNTER(STAT_Collision_SceneQueryTotal);
 	SCOPE_CYCLE_COUNTER(STAT_Collision_GeomOverlapBlocking);
@@ -1292,8 +1312,8 @@ bool FGenericPhysicsInterface::GeomOverlapAnyTest(const UWorld* World, const str
 	return GeomOverlapMultiHelper<EQueryInfo::IsAnything>(World, CollisionShape, GeomTransform, Overlaps, TraceChannel, Params, ResponseParams, ObjectParams);
 }
 
-template<typename TAccel, typename TGeom>
-bool FGenericGeomPhysicsInterfaceUsingSpatialAcceleration<TAccel, TGeom>::GeomOverlapAnyTest(const TAccel& Accel, const UWorld* World, const TGeom& InGeom, const FVector& Pos, const FQuat& Rot, ECollisionChannel TraceChannel, const struct FCollisionQueryParams& Params, const struct FCollisionResponseParams& ResponseParams, const struct FCollisionObjectQueryParams& ObjectParams)
+template<typename AccelType, typename GeomType>
+bool FGenericGeomPhysicsInterfaceUsingSpatialAcceleration<AccelType, GeomType>::GeomOverlapAnyTest(const AccelType& Accel, const UWorld* World, const GeomType& InGeom, const FVector& Pos, const FQuat& Rot, ECollisionChannel TraceChannel, const struct FCollisionQueryParams& Params, const struct FCollisionResponseParams& ResponseParams, const struct FCollisionObjectQueryParams& ObjectParams)
 {
 	SCOPE_CYCLE_COUNTER(STAT_Collision_SceneQueryTotal);
 	SCOPE_CYCLE_COUNTER(STAT_Collision_GeomOverlapAny);
@@ -1332,8 +1352,8 @@ bool FGenericPhysicsInterface::GeomOverlapMulti(const UWorld* World, const FPhys
 	return GeomOverlapMulti(World, FChaosEngineInterface::GetGeometryCollection(InGeom), InPosition, InRotation, OutOverlaps, TraceChannel, Params, ResponseParams, ObjectParams);
 }
 
-template<typename TAccel, typename TGeom>
-bool FGenericGeomPhysicsInterfaceUsingSpatialAcceleration<TAccel, TGeom>::GeomSweepSingle(const TAccel& Accel, const UWorld* World, const TGeom& InGeom, const FQuat& Rot, FHitResult& OutHit, FVector Start, FVector End, ECollisionChannel TraceChannel, const FCollisionQueryParams& Params, const FCollisionResponseParams& ResponseParams, const FCollisionObjectQueryParams& ObjectParams)
+template<typename AccelType, typename GeomType>
+bool FGenericGeomPhysicsInterfaceUsingSpatialAcceleration<AccelType, GeomType>::GeomSweepSingle(const AccelType& Accel, const UWorld* World, const GeomType& InGeom, const FQuat& Rot, FHitResult& OutHit, FVector Start, FVector End, ECollisionChannel TraceChannel, const FCollisionQueryParams& Params, const FCollisionResponseParams& ResponseParams, const FCollisionObjectQueryParams& ObjectParams)
 {
 	using namespace ChaosInterface;
 
@@ -1344,8 +1364,8 @@ bool FGenericGeomPhysicsInterfaceUsingSpatialAcceleration<TAccel, TGeom>::GeomSw
 	return TSceneCastCommon<ESweepOrRay::Sweep, ESingleMultiOrTest::Single>(World, OutHit, GeomToSQInputs(InGeom, Rot), Start, End, TraceChannel, Params, ResponseParams, ObjectParams, FOverrideAccelContainer{ Accel });
 }
 
-template<typename TAccel, typename TGeom>
-bool FGenericGeomPhysicsInterfaceUsingSpatialAcceleration<TAccel, TGeom>::GeomSweepMulti(const TAccel& Accel, const UWorld* World, const TGeom& InGeom, const FQuat& InGeomRot, TArray<FHitResult>& OutHits, FVector Start, FVector End, ECollisionChannel TraceChannel, const FCollisionQueryParams& Params, const FCollisionResponseParams& ResponseParams, const FCollisionObjectQueryParams& ObjectParams)
+template<typename AccelType, typename GeomType>
+bool FGenericGeomPhysicsInterfaceUsingSpatialAcceleration<AccelType, GeomType>::GeomSweepMulti(const AccelType& Accel, const UWorld* World, const GeomType& InGeom, const FQuat& InGeomRot, TArray<FHitResult>& OutHits, FVector Start, FVector End, ECollisionChannel TraceChannel, const FCollisionQueryParams& Params, const FCollisionResponseParams& ResponseParams, const FCollisionObjectQueryParams& ObjectParams)
 {
 	using namespace ChaosInterface;
 
@@ -1356,8 +1376,8 @@ bool FGenericGeomPhysicsInterfaceUsingSpatialAcceleration<TAccel, TGeom>::GeomSw
 	return TSceneCastCommon<ESweepOrRay::Sweep, ESingleMultiOrTest::Multi>(World, OutHits, GeomToSQInputs(InGeom, InGeomRot), Start, End, TraceChannel, Params, ResponseParams, ObjectParams, FOverrideAccelContainer{ Accel });
 }
 
-template<typename TAccel, typename TGeom>
-bool FGenericGeomPhysicsInterfaceUsingSpatialAcceleration<TAccel, TGeom>::GeomOverlapMulti(const TAccel& Accel, const UWorld* World, const TGeom& InGeom, const FVector& InPosition, const FQuat& InRotation, TArray<FOverlapResult>& OutOverlaps, ECollisionChannel TraceChannel, const FCollisionQueryParams& Params, const FCollisionResponseParams& ResponseParams, const FCollisionObjectQueryParams& ObjectParams)
+template<typename AccelType, typename GeomType>
+bool FGenericGeomPhysicsInterfaceUsingSpatialAcceleration<AccelType, GeomType>::GeomOverlapMulti(const AccelType& Accel, const UWorld* World, const GeomType& InGeom, const FVector& InPosition, const FQuat& InRotation, TArray<FOverlapResult>& OutOverlaps, ECollisionChannel TraceChannel, const FCollisionQueryParams& Params, const FCollisionResponseParams& ResponseParams, const FCollisionObjectQueryParams& ObjectParams)
 {
 	SCOPE_CYCLE_COUNTER(STAT_Collision_SceneQueryTotal);
 	SCOPE_CYCLE_COUNTER(STAT_Collision_GeomOverlapMultiple);
@@ -1406,9 +1426,12 @@ namespace Chaos::Private
 	template <typename OverlapHitType, EQueryInfo InfoType>
 	bool OverlapImpl(const UWorld* World, const Chaos::FOverlapQueryData& OverlapData, const FCommonQueryData& CommonData, TArray<FOverlapResult>& OutOverlaps)
 	{
-		FPhysicsShapeAdapter Adaptor(OverlapData.GeomPose.GetRotation(), OverlapData.CollisionShape);
+		const FQuat& Rotation = OverlapData.GeomPose.GetRotation();
+		const FCollisionShape& CollisionShape = OverlapData.QueryShape.CollisionShape;
+		FPhysicsShapeAdapter Adaptor(Rotation, CollisionShape);
 		const FPhysicsGeometry& Geom = Adaptor.GetGeometry();
-		return GeomOverlapMultiImp<OverlapHitType, InfoType>(World, Geom, OverlapData.CollisionShape, OverlapData.GeomPose, OutOverlaps, CommonData.TraceChannel, CommonData.Params, CommonData.ResponseParams, CommonData.ObjectParams, FDefaultAccelContainer{});
+
+		return GeomOverlapMultiImp<OverlapHitType, InfoType>(World, Geom, CollisionShape, OverlapData.GeomPose, OutOverlaps, CommonData.TraceChannel, CommonData.Params, CommonData.ResponseParams, CommonData.ObjectParams, FDefaultAccelContainer{});
 	}
 
 	template <EQueryInfo InfoType, typename ... ArgTypes>
@@ -1467,7 +1490,7 @@ namespace Chaos::Private
 	template <ESingleMultiOrTest SingleMultiOrTest, typename HitResultsType>
 	bool SweepCommon(const EThreadQueryContext ThreadContext, const UWorld* World, const Chaos::FSweepQueryData& SweepData, const Chaos::FCommonQueryData& CommonData, HitResultsType& OutHits)
 	{
-		FGeomSQAdditionalInputs GeomInputs(SweepData.CollisionShape, SweepData.GeomRot);
+		FGeomSQAdditionalInputs GeomInputs(SweepData.QueryShape.CollisionShape, SweepData.GeomRot);
 		return TraceCommonImp<ESweepOrRay::Sweep, SingleMultiOrTest>(ThreadContext, World, OutHits, GeomInputs, SweepData.Start, SweepData.End, CommonData.TraceChannel, CommonData.Params, CommonData.ResponseParams, CommonData.ObjectParams, FDefaultAccelContainer{});
 	}
 
diff --git a/Engine/Source/Runtime/Engine/Public/Physics/SceneQueryData.h b/Engine/Source/Runtime/Engine/Public/Physics/SceneQueryData.h
--- a/Engine/Source/Runtime/Engine/Public/Physics/SceneQueryData.h
+++ b/Engine/Source/Runtime/Engine/Public/Physics/SceneQueryData.h
@@ -31,9 +31,14 @@ namespace Chaos
 		FCollisionObjectQueryParams ObjectParams;
 	};
 
-	struct FOverlapQueryData
+	struct FQueryShape
 	{
 		FCollisionShape CollisionShape;
+	};
+
+	struct FOverlapQueryData
+	{
+		FQueryShape QueryShape;
 		FTransform GeomPose;
 	};
 
@@ -47,7 +52,7 @@ namespace Chaos
 	{
 		FVector Start;
 		FVector End;
-		FCollisionShape CollisionShape;
+		FQueryShape QueryShape;
 		FQuat GeomRot;
 	};
 }
