AnimNext: Add blueprint setter for bEnableAnimation
#rnx

diff --git a/Engine/Source/Runtime/Engine/Classes/Components/SkeletalMeshComponent.h b/Engine/Source/Runtime/Engine/Classes/Components/SkeletalMeshComponent.h
--- a/Engine/Source/Runtime/Engine/Classes/Components/SkeletalMeshComponent.h
+++ b/Engine/Source/Runtime/Engine/Classes/Components/SkeletalMeshComponent.h
@@ -974,6 +974,12 @@ public:
 	UFUNCTION(BlueprintCallable, Category = "Components|SkeletalMesh", meta = (Keywords = "AnimBlueprint", DisplayName = "Set Anim Instance Class"))
 	ENGINE_API virtual void SetAnimInstanceClass(class UClass* NewClass);
 
+	/**
+	* Whether the built-in animation of this component should run when the component ticks. See USkeletalMeshComponent::bEnableAnimation for more details.
+	*/
+	UFUNCTION(BlueprintCallable, Category = "Components|SkeletalMesh", meta = (Keywords = "Animation"))
+	ENGINE_API void SetEnableAnimation(bool bNewEnableAnimation);
+
 	/**
 	 * Returns the animation instance that is driving the class (if available). This is typically an instance of
 	 * the class set as AnimBlueprintGeneratedClass (generated by an animation blueprint)
diff --git a/Engine/Source/Runtime/Engine/Private/Components/SkeletalMeshComponent.cpp b/Engine/Source/Runtime/Engine/Private/Components/SkeletalMeshComponent.cpp
--- a/Engine/Source/Runtime/Engine/Private/Components/SkeletalMeshComponent.cpp
+++ b/Engine/Source/Runtime/Engine/Private/Components/SkeletalMeshComponent.cpp
@@ -3261,6 +3261,12 @@ void USkeletalMeshComponent::SetAnimInstanceClass(class UClass* NewClass)
 	}
 }
 
+void USkeletalMeshComponent::SetEnableAnimation(bool bNewEnableAnimation)
+{
+	bEnableAnimation = bNewEnableAnimation;
+	PrimaryComponentTick.bRunOnAnyThread = !bEnableAnimation;
+}
+
 UAnimInstance* USkeletalMeshComponent::GetAnimInstance() const
 {
 	return AnimScriptInstance;
