Made StaticMeshComponent use push based replication.

#rnx

diff --git a/Engine/Source/Runtime/Engine/Private/Components/StaticMeshComponent.cpp b/Engine/Source/Runtime/Engine/Private/Components/StaticMeshComponent.cpp
--- a/Engine/Source/Runtime/Engine/Private/Components/StaticMeshComponent.cpp
+++ b/Engine/Source/Runtime/Engine/Private/Components/StaticMeshComponent.cpp
@@ -24,6 +24,7 @@
 #include "StaticMeshComponentLODInfo.h"
 #include "StaticMeshResources.h"
 #include "StaticMeshSceneProxy.h"
+#include "Net/Core/PushModel/PushModel.h"
 #include "Net/UnrealNetwork.h"
 #include "Logging/MessageLog.h"
 #include "Misc/UObjectToken.h"
@@ -251,9 +252,9 @@ void UStaticMeshComponent::GetLifetimeReplicatedProps(TArray< FLifetimeProperty
 {
 	Super::GetLifetimeReplicatedProps(OutLifetimeProps);
 
-	PRAGMA_DISABLE_DEPRECATION_WARNINGS
-	DOREPLIFETIME( UStaticMeshComponent, StaticMesh );
-	PRAGMA_ENABLE_DEPRECATION_WARNINGS
+	FDoRepLifetimeParams Params;
+	Params.bIsPushBased = true;
+	DOREPLIFETIME_WITH_PARAMS_FAST(UStaticMeshComponent, StaticMesh, Params);
 }
 
 /// @endcond
@@ -2329,6 +2330,7 @@ bool UStaticMeshComponent::AllowsAsyncPhysicsStateDestruction() const
 
 void UStaticMeshComponent::SetStaticMeshInternal(UStaticMesh* NewMesh)
 {
+	MARK_PROPERTY_DIRTY_FROM_NAME(UStaticMeshComponent, StaticMesh, this);
 	StaticMesh = NewMesh;
 
 #if WITH_STATE_STREAM_ACTOR
