Modify changes made in CL 35191113 to fix issue with skeletal mesh assets being stomped after set in a user construction script.

New fix for UE-218988: the SkeletalMesh was being updated correctly in FActorComponentInstanceData::ApplyToComponent by FComponentPropertyReader's constructor (which serializes tagged properties). The SkeletalMeshAsset, however, was not updated as this directly updated the SkeletalMesh property rather than going through SetSkeletalMeshAsset which would typically updated SkeletalMeshAsset.
Now it isn't possible to set SkeletalMeshAsset anymore because if we try to call SetSkeletalMeshAsset with the same asset as SkeletalMesh it will be ignored and early return. So the solution here is to introduce a new function to force the SkeletalMeshAsset to refresh to the correct value based on the current SkeletalMesh.

#tests Tested editing the skel mesh component on an actor instance in a level. Also tested construction script logic which dynamically sets the skeletal mesh based on another variable (e.g. bool setting to one asset when true, and another when false). Alternating between the two on the same asset worked as expected. Preflight had unrelated warnings and errors (ran two which had separate unrelated errors).

#rnx

diff --git a/Engine/Source/Runtime/Engine/Classes/Components/SkeletalMeshComponent.h b/Engine/Source/Runtime/Engine/Classes/Components/SkeletalMeshComponent.h
--- a/Engine/Source/Runtime/Engine/Classes/Components/SkeletalMeshComponent.h
+++ b/Engine/Source/Runtime/Engine/Classes/Components/SkeletalMeshComponent.h
@@ -336,6 +336,9 @@ public:
 	UFUNCTION(BlueprintPure, Category = "Components|SkeletalMesh")
 	ENGINE_API USkeletalMesh* GetSkeletalMeshAsset() const;
 
+	/** Sets the SkeletalMeshAsset to the current SkeletalMesh. */
+	void RefreshSkeletalMeshAsset();
+
 	/** Gets the shared bone container used between all owned anim instances. Creates it on the first call. */
 	ENGINE_API TSharedPtr<struct FBoneContainer> GetSharedRequiredBones();
 
diff --git a/Engine/Source/Runtime/Engine/Private/Components/SkeletalMeshComponent.cpp b/Engine/Source/Runtime/Engine/Private/Components/SkeletalMeshComponent.cpp
--- a/Engine/Source/Runtime/Engine/Private/Components/SkeletalMeshComponent.cpp
+++ b/Engine/Source/Runtime/Engine/Private/Components/SkeletalMeshComponent.cpp
@@ -380,6 +380,15 @@ USkeletalMesh* USkeletalMeshComponent::GetSkeletalMeshAsset() const
 	return Cast<USkeletalMesh>(GetSkinnedAsset());
 }
 
+void USkeletalMeshComponent::RefreshSkeletalMeshAsset()
+{
+#if WITH_EDITORONLY_DATA
+	PRAGMA_DISABLE_DEPRECATION_WARNINGS
+	SkeletalMeshAsset = GetSkeletalMeshAsset();
+	PRAGMA_ENABLE_DEPRECATION_WARNINGS
+#endif
+}
+
 TSharedPtr<FBoneContainer> USkeletalMeshComponent::GetSharedRequiredBones()
 {
 	if (!SharedRequiredBones)
@@ -480,12 +489,8 @@ void USkeletalMeshComponent::PostLoad()
 		bOverrideMinLod = true;
 	}
 
-#if WITH_EDITORONLY_DATA
 	// Update property alias
-	PRAGMA_DISABLE_DEPRECATION_WARNINGS
-	SkeletalMeshAsset = GetSkeletalMeshAsset();
-	PRAGMA_ENABLE_DEPRECATION_WARNINGS
-#endif
+	RefreshSkeletalMeshAsset();
 }
 
 void USkeletalMeshComponent::PostInitProperties()
diff --git a/Engine/Source/Runtime/Engine/Private/Components/SkeletalMeshComponentInstanceData.cpp b/Engine/Source/Runtime/Engine/Private/Components/SkeletalMeshComponentInstanceData.cpp
--- a/Engine/Source/Runtime/Engine/Private/Components/SkeletalMeshComponentInstanceData.cpp
+++ b/Engine/Source/Runtime/Engine/Private/Components/SkeletalMeshComponentInstanceData.cpp
@@ -7,7 +7,6 @@
 
 FSkeletalMeshComponentInstanceData::FSkeletalMeshComponentInstanceData(const USkeletalMeshComponent* SourceComponent)
 	: FSceneComponentInstanceData(SourceComponent)
-	, SkeletalMeshAsset(nullptr)
 	, bUpdateAnimationInEditor(0)
 	, bUpdateClothInEditor(0)
 {
@@ -18,7 +17,6 @@ FSkeletalMeshComponentInstanceData::FSkeletalMeshComponentInstanceData(const USk
 		|| SourceComponent->CreationMethod == EComponentCreationMethod::UserConstructionScript;
 	if (bIsBlueprintCreatedComponent)
 	{
-		SkeletalMeshAsset = SourceComponent->GetSkeletalMeshAsset();
 		bUpdateAnimationInEditor = SourceComponent->GetUpdateAnimationInEditor();
 		bUpdateClothInEditor = SourceComponent->GetUpdateClothInEditor();
 	}
@@ -27,7 +25,7 @@ FSkeletalMeshComponentInstanceData::FSkeletalMeshComponentInstanceData(const USk
 
 bool FSkeletalMeshComponentInstanceData::ContainsData() const
 {
-	return bUpdateAnimationInEditor || bUpdateClothInEditor || SkeletalMeshAsset;
+	return bUpdateAnimationInEditor || bUpdateClothInEditor;
 }
 
 void FSkeletalMeshComponentInstanceData::ApplyToComponent(UActorComponent* Component, const ECacheApplyPhase CacheApplyPhase)
@@ -35,16 +33,9 @@ void FSkeletalMeshComponentInstanceData::ApplyToComponent(UActorComponent* Compo
 	Super::ApplyToComponent(Component, CacheApplyPhase);
 
 #if WITH_EDITOR
-	const bool bIsBlueprintCreatedComponent = Component->CreationMethod == EComponentCreationMethod::SimpleConstructionScript
-		|| Component->CreationMethod == EComponentCreationMethod::UserConstructionScript;
-	USkeletalMeshComponent* SkeletalMesh = Cast<USkeletalMeshComponent>(Component);
-
-	if (SkeletalMesh && bIsBlueprintCreatedComponent)
+	if (USkeletalMeshComponent * SkeletalMesh = Cast<USkeletalMeshComponent>(Component))
 	{
-		if (SkeletalMeshAsset)
-		{
-			SkeletalMesh->SetSkeletalMeshAsset(SkeletalMeshAsset);
-		}
+		SkeletalMesh->RefreshSkeletalMeshAsset();
 		SkeletalMesh->SetUpdateAnimationInEditor(bUpdateAnimationInEditor);
 		SkeletalMesh->SetUpdateClothInEditor(bUpdateClothInEditor);
 	}
diff --git a/Engine/Source/Runtime/Engine/Private/Components/SkeletalMeshComponentInstanceData.h b/Engine/Source/Runtime/Engine/Private/Components/SkeletalMeshComponentInstanceData.h
--- a/Engine/Source/Runtime/Engine/Private/Components/SkeletalMeshComponentInstanceData.h
+++ b/Engine/Source/Runtime/Engine/Private/Components/SkeletalMeshComponentInstanceData.h
@@ -6,7 +6,6 @@
 #include "SkeletalMeshComponentInstanceData.generated.h"
 
 class USkeletalMeshComponent;
-class USkeletalMesh;
 
 /** Saves internal SkeletalMesh (transient) state that gets lost at reconstruction for Blueprint created components. */
 USTRUCT()
@@ -20,7 +19,6 @@ struct ENGINE_API FSkeletalMeshComponentInstanceData : public FSceneComponentIns
 	virtual bool ContainsData() const override;
 	virtual void ApplyToComponent(UActorComponent* Component, const ECacheApplyPhase CacheApplyPhase) override;
 
-	USkeletalMesh* SkeletalMeshAsset;
 	uint8 bUpdateAnimationInEditor : 1;
 	uint8 bUpdateClothInEditor : 1;
 };
