- Allow the OnlineEngineInterface ClassName to be configured.
- Still support the V1/V2 split bool, but log a warning to prompt people to move to the new ClassName config entry
- Configure the V1 class in BaseEngine.ini

diff --git a/Engine/Config/BaseEngine.ini b/Engine/Config/BaseEngine.ini
--- a/Engine/Config/BaseEngine.ini
+++ b/Engine/Config/BaseEngine.ini
@@ -336,6 +336,9 @@ NearClipPlane=10.0
 ; Default World Partition HLOD layer
 DefaultWorldPartitionHLODLayer=/Engine/Maps/Templates/HLODs/HLODLayer_Instanced.HLODLayer_Instanced
 
+[/Script/Engine.OnlineEngineInterface]
+ClassName=/Script/OnlineSubsystemUtils.OnlineEngineInterfaceImpl
+
 [CoreRedirects]
 
 ; All redirects prior to UE 4.10 are now in FCoreRedirects::RegisterNativeRedirects, new ones should be added to the bottom of this section
diff --git a/Engine/Source/Runtime/Engine/Private/OnlineEngineInterface.cpp b/Engine/Source/Runtime/Engine/Private/OnlineEngineInterface.cpp
--- a/Engine/Source/Runtime/Engine/Private/OnlineEngineInterface.cpp
+++ b/Engine/Source/Runtime/Engine/Private/OnlineEngineInterface.cpp
@@ -6,6 +6,8 @@
 
 #include UE_INLINE_GENERATED_CPP_BY_NAME(OnlineEngineInterface)
 
+DEFINE_LOG_CATEGORY_STATIC(LogOnlineEngine, Log, All);
+
 UOnlineEngineInterface* UOnlineEngineInterface::Singleton = nullptr;
 
 UOnlineEngineInterface::UOnlineEngineInterface(const FObjectInitializer& ObjectInitializer)
@@ -17,18 +19,25 @@ UOnlineEngineInterface* UOnlineEngineInterface::Get()
 {
 	if (!Singleton)
 	{
+		FString OnlineEngineInterfaceClassName;
+		GConfig->GetString(TEXT("/Script/Engine.OnlineEngineInterface"), TEXT("ClassName"), OnlineEngineInterfaceClassName, GEngineIni);
+
+		// To not break licensees using this, prefer this if it is present, and warn. Remove in 5.7
 		bool bUseOnlineServicesV2 = false;
-		GConfig->GetBool(TEXT("/Script/Engine.OnlineEngineInterface"), TEXT("bUseOnlineServicesV2"), bUseOnlineServicesV2, GEngineIni);
-		// Proper interface class hard coded here to emphasize the fact that this is not expected to change much, any need to do so should go through the OGS team first
-		UClass* OnlineEngineInterfaceClass = nullptr;
-		if (bUseOnlineServicesV2)
+		if (GConfig->GetBool(TEXT("/Script/Engine.OnlineEngineInterface"), TEXT("bUseOnlineServicesV2"), bUseOnlineServicesV2, GEngineIni))
 		{
-			OnlineEngineInterfaceClass = StaticLoadClass(UOnlineEngineInterface::StaticClass(), NULL, TEXT("/Script/OnlineSubsystemUtils.OnlineServicesEngineInterfaceImpl"), NULL, LOAD_Quiet, NULL);
+			const TCHAR* V1ClassName = TEXT("/Script/OnlineSubsystemUtils.OnlineEngineInterfaceImpl");
+			const TCHAR* V2ClassName = TEXT("/Script/OnlineSubsystemUtils.OnlineServicesEngineInterfaceImpl");
+			OnlineEngineInterfaceClassName = bUseOnlineServicesV2 ? V2ClassName : V1ClassName;
+			UE_LOG(LogOnlineEngine, Warning, TEXT("bUseOnlineServicesV2 is deprecated, please instead configure [/Script/Engine.OnlineEngineInterface]:ClassName=%s"), *OnlineEngineInterfaceClassName);
 		}
-		else
+		
+		UClass* OnlineEngineInterfaceClass = nullptr;
+		if (!OnlineEngineInterfaceClassName.IsEmpty())
 		{
-			OnlineEngineInterfaceClass = StaticLoadClass(UOnlineEngineInterface::StaticClass(), NULL, TEXT("/Script/OnlineSubsystemUtils.OnlineEngineInterfaceImpl"), NULL, LOAD_Quiet, NULL);
+			OnlineEngineInterfaceClass = StaticLoadClass(UOnlineEngineInterface::StaticClass(), NULL, *OnlineEngineInterfaceClassName, NULL, LOAD_Quiet, NULL);
 		}
+		
 		if (!OnlineEngineInterfaceClass)
 		{
 			// Default to the no op class if necessary
