Added selectable & has per instance hit proxies on the ISM descriptors.

#rnx

diff --git a/Engine/Source/Runtime/Engine/Private/ISMPartition/ISMComponentDescriptor.cpp b/Engine/Source/Runtime/Engine/Private/ISMPartition/ISMComponentDescriptor.cpp
--- a/Engine/Source/Runtime/Engine/Private/ISMPartition/ISMComponentDescriptor.cpp
+++ b/Engine/Source/Runtime/Engine/Private/ISMPartition/ISMComponentDescriptor.cpp
@@ -121,6 +121,8 @@ void FISMComponentDescriptorBase::InitFrom(const UStaticMeshComponent* Template,
 	bHiddenInGame = Template->bHiddenInGame;
 	bIsEditorOnly = Template->bIsEditorOnly;
 	bVisible = Template->GetVisibleFlag();
+	bSelectable = Template->bSelectable;
+	bHasPerInstanceHitProxies = Template->bHasPerInstanceHitProxies;
 	bVisibleInRayTracing = Template->bVisibleInRayTracing;
 	bEvaluateWorldPositionOffset = Template->bEvaluateWorldPositionOffset;
 	WorldPositionOffsetDisableDistance = Template->WorldPositionOffsetDisableDistance;
@@ -246,6 +248,8 @@ bool FISMComponentDescriptorBase::operator==(const FISMComponentDescriptorBase&
 	bIsEditorOnly == Other.bIsEditorOnly &&
 	bVisible == Other.bVisible &&
 	bVisibleInRayTracing == Other.bVisibleInRayTracing &&
+	bSelectable == Other.bSelectable &&
+	bHasPerInstanceHitProxies == Other.bHasPerInstanceHitProxies &&
 	bEvaluateWorldPositionOffset == Other.bEvaluateWorldPositionOffset &&
 	bReverseCulling == Other.bReverseCulling &&
 	bUseGpuLodSelection == Other.bUseGpuLodSelection &&
@@ -375,6 +379,8 @@ void FISMComponentDescriptorBase::InitComponent(UInstancedStaticMeshComponent* I
 	ISMComponent->bIsEditorOnly = bIsEditorOnly;
 	ISMComponent->SetVisibleFlag(bVisible);
 	ISMComponent->bVisibleInRayTracing = bVisibleInRayTracing;
+	ISMComponent->bSelectable = bSelectable;
+	ISMComponent->bHasPerInstanceHitProxies = bHasPerInstanceHitProxies;
 	ISMComponent->bEvaluateWorldPositionOffset = bEvaluateWorldPositionOffset;
 	ISMComponent->bReverseCulling = bReverseCulling;
 	ISMComponent->bUseGpuLodSelection = bUseGpuLodSelection;
diff --git a/Engine/Source/Runtime/Engine/Public/ISMPartition/ISMComponentDescriptor.h b/Engine/Source/Runtime/Engine/Public/ISMPartition/ISMComponentDescriptor.h
--- a/Engine/Source/Runtime/Engine/Public/ISMPartition/ISMComponentDescriptor.h
+++ b/Engine/Source/Runtime/Engine/Public/ISMPartition/ISMComponentDescriptor.h
@@ -168,6 +168,12 @@ public:
 
 	UPROPERTY(EditAnywhere, Category = "Component Settings")
 	uint8 bVisible : 1;
+
+	UPROPERTY(EditAnywhere, Category = "Component Settings")
+	uint8 bSelectable : 1;
+
+	UPROPERTY(EditAnywhere, Category = "Component Settings")
+	uint8 bHasPerInstanceHitProxies : 1;
 	
 	UPROPERTY(EditAnywhere, Category = "Component Settings")
 	uint8 bEvaluateWorldPositionOffset : 1;
