FDurationTimer: Add a threadsafe version of FDurationTimer using std::atomic. Change ShaderCompilerCookStats to std::atomic since they can be read/written from multiple threads when asynchronous loading is used in editor.
#rnx

diff --git a/Engine/Source/Runtime/Core/Public/ProfilingDebugging/ScopedTimers.h b/Engine/Source/Runtime/Core/Public/ProfilingDebugging/ScopedTimers.h
--- a/Engine/Source/Runtime/Core/Public/ProfilingDebugging/ScopedTimers.h
+++ b/Engine/Source/Runtime/Core/Public/ProfilingDebugging/ScopedTimers.h
@@ -11,6 +11,19 @@
 #include "Misc/OutputDeviceRedirector.h"
 #include "HAL/ThreadSafeCounter64.h"
 
+#include <atomic>
+
+/**
+ * Implements the same behavior as std::atomic<double>.fetch_add(double Delta), which is only available in c++20.
+ * This function will be deprecated and callsites will replace it with std::atomic<double>.fetch_add when versions
+ * earlier than c++20 are deprecated.
+ * 
+ * Atomically adds Delta to Value, using the given MemoryOrder for the ReadModifyWrite of the new value.
+ * @return The value before the finally successful addition.
+ */
+double AtomicDoubleFetchAdd(std::atomic<double>& Value, double Delta,
+	std::memory_order MemoryOrder = std::memory_order_seq_cst);
+
 /**
  * Utility stopwatch class for tracking the duration of some action (tracks 
  * time in seconds and adds it to the specified variable on destruction).
@@ -87,6 +100,65 @@ private:
 	double AccumulatorValue;
 };
 
+
+/**
+ * Utility stopwatch class for tracking the duration of some action (tracks time in seconds and adds it to the
+ * specified variable on destruction), when the storage is std::atomic<double>.
+ */
+class FDurationAtomicTimer
+{
+public:
+	explicit FDurationAtomicTimer(std::atomic<double>& AccumulatorIn,
+		std::memory_order InMemoryOrder = std::memory_order_relaxed)
+		: StartTime(FPlatformTime::Seconds())
+		, Accumulator(AccumulatorIn)
+		, MemoryOrder(InMemoryOrder)
+	{}
+
+	double Start()
+	{
+		StartTime = FPlatformTime::Seconds();
+		return StartTime;
+	}
+
+	double Stop()
+	{
+		double StopTime = FPlatformTime::Seconds();
+		AtomicDoubleFetchAdd(Accumulator, StopTime - StartTime, MemoryOrder);
+		StartTime = StopTime;
+
+		return StopTime;
+	}
+
+protected:
+	/** Start time, captured in ctor. */
+	double StartTime;
+	/** Time variable to update. */
+	std::atomic<double>& Accumulator;
+	/** Memory order specified in constructor and passed to fetch_add. */
+	std::memory_order MemoryOrder;
+};
+
+/**
+ * Utility class for tracking the duration of a scoped action (the user doesn't have to call Start() and
+ * Stop() manually), when the storage is std::atomic<double>.
+ */
+class FScopedDurationAtomicTimer : public FDurationAtomicTimer
+{
+public:
+	explicit FScopedDurationAtomicTimer(std::atomic<double>& AccumulatorIn,
+		std::memory_order InMemoryOrder = std::memory_order_relaxed)
+		: FDurationAtomicTimer(AccumulatorIn, InMemoryOrder)
+	{
+	}
+
+	/** Dtor, updating seconds with time delta. */
+	~FScopedDurationAtomicTimer()
+	{
+		Stop();
+	}
+};
+
 /**
  * Utility class for tracking the duration of a scoped action to an accumulator in a thread-safe fashion.
  * Can accumulate into a 32bit or 64bit counter.
@@ -240,6 +312,29 @@ private:
 	double& TotalTime;
 };
 
+
+inline double AtomicDoubleFetchAdd(std::atomic<double>& Value, double Delta,
+	std::memory_order MemoryOrder)
+{
+	double Expected = Value.load(std::memory_order_relaxed);
+	double Desired;
+	std::memory_order ReadModifyWriteMemoryOrder =
+		MemoryOrder == std::memory_order_relaxed ? std::memory_order_relaxed
+		: MemoryOrder == std::memory_order_acq_rel ? std::memory_order_acq_rel
+		: std::memory_order_seq_cst;
+
+	bool bExchangeSucceeded = false;
+	do
+	{
+		Desired = Expected + Delta;
+		bExchangeSucceeded = Value.compare_exchange_weak(Expected, Desired,
+			/* success order */ ReadModifyWriteMemoryOrder,
+			/* failure order */ std::memory_order_relaxed);
+	} while (!bExchangeSucceeded);
+
+	return Expected;
+}
+
 #if NO_LOGGING
 #define UE_SCOPED_TIMER(Title, Category, Verbosity)
 #else
diff --git a/Engine/Source/Runtime/Engine/Private/ShaderCompiler/ShaderCompiler.cpp b/Engine/Source/Runtime/Engine/Private/ShaderCompiler/ShaderCompiler.cpp
--- a/Engine/Source/Runtime/Engine/Private/ShaderCompiler/ShaderCompiler.cpp
+++ b/Engine/Source/Runtime/Engine/Private/ShaderCompiler/ShaderCompiler.cpp
@@ -550,20 +550,20 @@ namespace ShaderCompiler
 #if ENABLE_COOK_STATS
 namespace ShaderCompilerCookStats
 {
-	static double BlockingTimeSec = 0.0;
-	static double GlobalBeginCompileShaderTimeSec = 0.0;
-	static int32 GlobalBeginCompileShaderCalls = 0;
-	static double ProcessAsyncResultsTimeSec = 0.0;
-	double AsyncCompileTimeSec = 0.0;
+	static std::atomic<double> BlockingTimeSec = 0.0;
+	static std::atomic<double> GlobalBeginCompileShaderTimeSec = 0.0;
+	static std::atomic<int32> GlobalBeginCompileShaderCalls = 0;
+	static std::atomic<double> ProcessAsyncResultsTimeSec = 0.0;
+	std::atomic<double> AsyncCompileTimeSec = 0.0;
 
 	static FCookStatsManager::FAutoRegisterCallback RegisterCookStats([](FCookStatsManager::AddStatFuncRef AddStat)
 		{
 			AddStat(TEXT("ShaderCompiler"), FCookStatsManager::CreateKeyValueArray(
-				TEXT("BlockingTimeSec"), BlockingTimeSec,
-				TEXT("AsyncCompileTimeSec"), AsyncCompileTimeSec,
-				TEXT("GlobalBeginCompileShaderTimeSec"), GlobalBeginCompileShaderTimeSec,
-				TEXT("GlobalBeginCompileShaderCalls"), GlobalBeginCompileShaderCalls,
-				TEXT("ProcessAsyncResultsTimeSec"), ProcessAsyncResultsTimeSec
+				TEXT("BlockingTimeSec"), BlockingTimeSec.load(std::memory_order_relaxed),
+				TEXT("AsyncCompileTimeSec"), AsyncCompileTimeSec.load(std::memory_order_relaxed),
+				TEXT("GlobalBeginCompileShaderTimeSec"), GlobalBeginCompileShaderTimeSec.load(std::memory_order_relaxed),
+				TEXT("GlobalBeginCompileShaderCalls"), GlobalBeginCompileShaderCalls.load(std::memory_order_relaxed),
+				TEXT("ProcessAsyncResultsTimeSec"), ProcessAsyncResultsTimeSec.load(std::memory_order_relaxed)
 			));
 		});
 }
@@ -1667,7 +1667,7 @@ void FShaderCompilingManager::BlockOnShaderMapCompletion(const TArray<int32>& Sh
 
 	TRACE_CPUPROFILER_EVENT_SCOPE(FShaderCompilingManager::BlockOnShaderMapCompletion);
 
-	COOK_STAT(FScopedDurationTimer BlockingTimer(ShaderCompilerCookStats::BlockingTimeSec));
+	COOK_STAT(FScopedDurationAtomicTimer BlockingTimer(ShaderCompilerCookStats::BlockingTimeSec));
 	if (bAllowAsynchronousShaderCompiling)
 	{
 		// Calculate how many shader jobs there are total to provide the slow task with the correct amount of work.
@@ -1800,7 +1800,7 @@ void FShaderCompilingManager::BlockOnAllShaderMapCompletion(TMap<int32, FShaderM
 
 	TRACE_CPUPROFILER_EVENT_SCOPE(FShaderCompilingManager::BlockOnAllShaderMapCompletion);
 
-	COOK_STAT(FScopedDurationTimer BlockingTimer(ShaderCompilerCookStats::BlockingTimeSec));
+	COOK_STAT(FScopedDurationAtomicTimer BlockingTimer(ShaderCompilerCookStats::BlockingTimeSec));
 	if (bAllowAsynchronousShaderCompiling)
 	{
 		// Calculate how many shader jobs there are total to provide the slow task with the correct amount of work.
@@ -2660,7 +2660,7 @@ void FShaderCompilingManager::ProcessAsyncResults(float TimeSlice, bool bBlockOn
 
 	TRACE_CPUPROFILER_EVENT_SCOPE(FShaderCompilingManager::ProcessAsyncResults)
 
-	COOK_STAT(FScopedDurationTimer Timer(ShaderCompilerCookStats::ProcessAsyncResultsTimeSec));
+	COOK_STAT(FScopedDurationAtomicTimer Timer(ShaderCompilerCookStats::ProcessAsyncResultsTimeSec));
 	check(IsInGameThread());
 
 	const double StartTime = FPlatformTime::Seconds();
@@ -3038,8 +3038,8 @@ void GlobalBeginCompileShader(
 	LLM_SCOPE_BYTAG(ShaderCompiler);
 
 	TRACE_CPUPROFILER_EVENT_SCOPE(GlobalBeginCompileShader);
-	COOK_STAT(ShaderCompilerCookStats::GlobalBeginCompileShaderCalls++);
-	COOK_STAT(FScopedDurationTimer DurationTimer(ShaderCompilerCookStats::GlobalBeginCompileShaderTimeSec));
+	COOK_STAT(ShaderCompilerCookStats::GlobalBeginCompileShaderCalls.fetch_add(1, std::memory_order_relaxed));
+	COOK_STAT(FScopedDurationAtomicTimer DurationTimer(ShaderCompilerCookStats::GlobalBeginCompileShaderTimeSec));
 
 	const EShaderPlatform ShaderPlatform = EShaderPlatform(Target.Platform);
 	const FName ShaderFormatName = LegacyShaderPlatformToShaderFormat(ShaderPlatform);
diff --git a/Engine/Source/Runtime/Engine/Private/ShaderCompiler/ShaderCompilerPrivate.h b/Engine/Source/Runtime/Engine/Private/ShaderCompiler/ShaderCompilerPrivate.h
--- a/Engine/Source/Runtime/Engine/Private/ShaderCompiler/ShaderCompilerPrivate.h
+++ b/Engine/Source/Runtime/Engine/Private/ShaderCompiler/ShaderCompilerPrivate.h
@@ -11,6 +11,8 @@
 #include "DerivedDataCache.h"
 #include "ProfilingDebugging/CookStats.h"
 
+#include <atomic>
+
 LLM_DECLARE_TAG(ShaderCompiler);
 
 class FGlobalShaderMap;
@@ -42,7 +44,7 @@ namespace GlobalShaderCookStats
 
 namespace ShaderCompilerCookStats
 {
-	extern double AsyncCompileTimeSec;
+	extern std::atomic<double> AsyncCompileTimeSec;
 }
 
 /** Helper functions for logging more debug info */
diff --git a/Engine/Source/Runtime/Engine/Private/ShaderCompiler/ShaderCompilerThreadRunnable.cpp b/Engine/Source/Runtime/Engine/Private/ShaderCompiler/ShaderCompilerThreadRunnable.cpp
--- a/Engine/Source/Runtime/Engine/Private/ShaderCompiler/ShaderCompilerThreadRunnable.cpp
+++ b/Engine/Source/Runtime/Engine/Private/ShaderCompiler/ShaderCompilerThreadRunnable.cpp
@@ -16,6 +16,7 @@
 #include "ProfilingDebugging/CookStats.h"
 #include "ProfilingDebugging/DiagnosticTable.h"
 #include "ProfilingDebugging/LoadTimeTracker.h"
+#include "ProfilingDebugging/ScopedTimers.h"
 #include "ProfilingDebugging/StallDetector.h"
 
 
@@ -450,7 +451,7 @@ void FShaderCompileThreadRunnable::PushCompletedJobsToManager()
 			const double ElapsedTime = FPlatformTime::Seconds() - CurrentWorkerInfo.StartTime;
 
 			Manager->WorkersBusyTime += ElapsedTime;
-			COOK_STAT(ShaderCompilerCookStats::AsyncCompileTimeSec += ElapsedTime);
+			COOK_STAT(AtomicDoubleFetchAdd(ShaderCompilerCookStats::AsyncCompileTimeSec, ElapsedTime, std::memory_order_relaxed));
 
 
 			CurrentWorkerInfo.FinishTime = FPlatformTime::Seconds();
@@ -792,7 +793,7 @@ int32 FShaderCompileThreadRunnable::ReadAvailableResults()
 void FShaderCompileThreadRunnable::CompileDirectlyThroughDll()
 {
 	// If we aren't compiling through workers, so we can just track the serial time here.
-	COOK_STAT(FScopedDurationTimer CompileTimer (ShaderCompilerCookStats::AsyncCompileTimeSec));
+	COOK_STAT(FScopedDurationAtomicTimer CompileTimer(ShaderCompilerCookStats::AsyncCompileTimeSec));
 
 	FScopeLock WorkerScopeLock(&WorkerInfosLock);
 	for (int32 WorkerIndex = 0; WorkerIndex < WorkerInfos.Num(); WorkerIndex++)
@@ -1237,7 +1238,7 @@ int32 FShaderCompileThreadRunnable::CompilingLoop()
 					}
 
 					Manager->WorkersBusyTime += ElapsedTime;
-					COOK_STAT(ShaderCompilerCookStats::AsyncCompileTimeSec += ElapsedTime);
+					COOK_STAT(AtomicDoubleFetchAdd(ShaderCompilerCookStats::AsyncCompileTimeSec, ElapsedTime, std::memory_order_relaxed));
 				}, CompileTasks);
 			// num active threads is up to the task system; as long as we return non-zero
 			// this will indicate to the caller that pending jobs remain
