[EdGraph] Moved title text coloring from SGraphNode to EdGraphNode and called it from SGraphNode instead to allow override.
#rnx

diff --git a/Engine/Source/Editor/GraphEditor/Private/SGraphNode.cpp b/Engine/Source/Editor/GraphEditor/Private/SGraphNode.cpp
--- a/Engine/Source/Editor/GraphEditor/Private/SGraphNode.cpp
+++ b/Engine/Source/Editor/GraphEditor/Private/SGraphNode.cpp
@@ -661,12 +661,7 @@ FSlateColor SGraphNode::GetNodeTitleIconColor() const
 
 FLinearColor SGraphNode::GetNodeTitleTextColor() const
 {
-	FLinearColor ReturnTextColor = FLinearColor::White;
-	if(!GraphNode->IsNodeEnabled() || GraphNode->IsDisplayAsDisabledForced() || GraphNode->IsNodeUnrelated())
-	{
-		ReturnTextColor *= FLinearColor(1.0f, 1.0f, 1.0f, 0.3f); 
-	}
-	return ReturnTextColor;
+	return GetNodeObj()->GetNodeTitleTextColor();
 }
 
 FSlateColor SGraphNode::GetNodeCommentColor() const
diff --git a/Engine/Source/Runtime/Engine/Classes/EdGraph/EdGraphNode.h b/Engine/Source/Runtime/Engine/Classes/EdGraph/EdGraphNode.h
--- a/Engine/Source/Runtime/Engine/Classes/EdGraph/EdGraphNode.h
+++ b/Engine/Source/Runtime/Engine/Classes/EdGraph/EdGraphNode.h
@@ -759,6 +759,11 @@ public:
 	 */
 	ENGINE_API virtual FLinearColor GetNodeCommentColor() const;
 
+	/**
+	 * Gets the  color of a node's title text
+	 */
+	ENGINE_API virtual FLinearColor GetNodeTitleTextColor() const;
+
 	/**
 	 * Gets the draw color of a node's body tine
 	 */
diff --git a/Engine/Source/Runtime/Engine/Private/EdGraph/EdGraphNode.cpp b/Engine/Source/Runtime/Engine/Private/EdGraph/EdGraphNode.cpp
--- a/Engine/Source/Runtime/Engine/Private/EdGraph/EdGraphNode.cpp
+++ b/Engine/Source/Runtime/Engine/Private/EdGraph/EdGraphNode.cpp
@@ -579,6 +579,16 @@ FLinearColor UEdGraphNode::GetNodeTitleColor() const
 	return FLinearColor(0.4f, 0.62f, 1.0f);
 }
 
+FLinearColor UEdGraphNode::GetNodeTitleTextColor() const
+{
+	FLinearColor ReturnTextColor = FLinearColor::White;
+	if (!IsNodeEnabled() || IsDisplayAsDisabledForced() || IsNodeUnrelated())
+	{
+		ReturnTextColor *= FLinearColor(1.0f, 1.0f, 1.0f, 0.3f);
+	}
+	return ReturnTextColor;
+}
+
 FLinearColor UEdGraphNode::GetNodeCommentColor() const
 {
 	return FLinearColor::White;
