Fix attached component to a simulated component by socket missing collisions

The current implementation calls GetBodyInstance with the socket name as a parameter. On skeletal meshes, this will always fail since the method expects the bone name instead, causing welding to not run. This change gets the bone name from the socket to call the function correctly so welding will work correctly in this case.

diff --git a/Engine/Source/Runtime/Engine/Private/PrimitiveComponentPhysics.cpp b/Engine/Source/Runtime/Engine/Private/PrimitiveComponentPhysics.cpp
--- a/Engine/Source/Runtime/Engine/Private/PrimitiveComponentPhysics.cpp
+++ b/Engine/Source/Runtime/Engine/Private/PrimitiveComponentPhysics.cpp
@@ -13,6 +13,8 @@
 #include "UObject/UObjectThreadContext.h"
 #include "PhysicsEngine/PhysicsObjectExternalInterface.h"
 #include "PhysicsProxy/SingleParticlePhysicsProxy.h"
+#include "Components/SkeletalMeshComponent.h"
+#include "Engine/SkeletalMeshSocket.h"
 
 //////////////// PRIMITIVECOMPONENT ///////////////
 
@@ -982,7 +984,21 @@ bool UPrimitiveComponent::WeldToImplementation(USceneComponent * InParent, FName
 
 	if (RootComponent)
 	{
-		if (FBodyInstance* RootBI = RootComponent->GetBodyInstance(SocketName, false))
+		FBodyInstance* RootBI = RootComponent->GetBodyInstance(SocketName, false);
+		FName BoneName = SocketName;
+
+		//If initial lookup fails, check if we need to lookup the bone from the socket instead
+		if (!RootBI && RootComponent->IsA<USkeletalMeshComponent>())
+		{
+			USkeletalMeshComponent* RootSkeletal = Cast<USkeletalMeshComponent>(RootComponent);
+			if (const USkeletalMeshSocket* Socket = RootSkeletal->GetSocketByName(SocketName))
+			{
+				BoneName = Socket->BoneName;
+				RootBI = RootComponent->GetBodyInstance(BoneName, false);
+			}
+		}
+
+		if (RootBI)
 		{
 			if (BI->WeldParent == RootBI)	//already welded so stop
 			{
@@ -995,7 +1011,7 @@ bool UPrimitiveComponent::WeldToImplementation(USceneComponent * InParent, FName
 			//Child always inherits from root
 
 			//if root is kinematic simply set child to be kinematic and we're done
-			if ((RootComponent->IsSimulatingPhysics(SocketName) == false) && (bWeldToKinematicParent == false))
+			if ((RootComponent->IsSimulatingPhysics(BoneName) == false) && (bWeldToKinematicParent == false))
 			{
 				void* OriginalWeldParent = nullptr;
 
@@ -1008,7 +1024,7 @@ bool UPrimitiveComponent::WeldToImplementation(USceneComponent * InParent, FName
 						void* Old = FPlatformAtomics::InterlockedExchangePtr((void**)&BI->WeldParent, OriginalWeldParent);
 						ensure(nullptr == Old);
 					});
-				
+
 				SetSimulatePhysics(false);
 				return false;	//return false because we need to continue with regular body initialization
 			}
