Physics changes to support physics in ControlRig by allowing access to the solver settings.

diff --git a/Engine/Source/Runtime/Engine/Private/PhysicsEngine/ImmediatePhysics/ImmediatePhysicsChaos/ImmediatePhysicsSimulation_Chaos.cpp b/Engine/Source/Runtime/Engine/Private/PhysicsEngine/ImmediatePhysics/ImmediatePhysicsChaos/ImmediatePhysicsSimulation_Chaos.cpp
--- a/Engine/Source/Runtime/Engine/Private/PhysicsEngine/ImmediatePhysics/ImmediatePhysicsChaos/ImmediatePhysicsSimulation_Chaos.cpp
+++ b/Engine/Source/Runtime/Engine/Private/PhysicsEngine/ImmediatePhysics/ImmediatePhysicsChaos/ImmediatePhysicsSimulation_Chaos.cpp
@@ -626,12 +626,13 @@ namespace ImmediatePhysics_Chaos
 	// Currently sets up potential collision with ActorHandle and all dynamics
 	void FSimulation::AddToCollidingPairs(FActorHandle* ActorHandle)
 	{
+		bool bCouldBeDynamic = ActorHandle->CouldBeDynamic();
 		using namespace Chaos;
 		FGeometryParticleHandle* Particle0 = ActorHandle->GetParticle();
 		for (FActorHandle* OtherActorHandle : Implementation->ActorHandles)
 		{
 			FGeometryParticleHandle* Particle1 = OtherActorHandle->GetParticle();
-			if ((OtherActorHandle != ActorHandle) && OtherActorHandle->CouldBeDynamic())
+			if ((OtherActorHandle != ActorHandle) && (bCouldBeDynamic || OtherActorHandle->CouldBeDynamic()))
 			{
 				Implementation->PotentiallyCollidingPairs.Emplace(FParticlePair(Particle0, Particle1));
 			}
@@ -887,7 +888,11 @@ namespace ImmediatePhysics_Chaos
 		Implementation->Evolution.SetRewindVelocities(bRewindVelocities);
 	}
 
-	void FSimulation::Simulate(FReal InDeltaTime, FReal MaxStepTime, int32 MaxSubSteps, const FVector& InGravity)
+	void FSimulation::Simulate(
+		FReal InDeltaTime, FReal MaxStepTime, int32 MaxSubSteps, const FVector& InGravity,
+		Chaos::FPBDJointSolverSettings* JointSolverSettings, 
+		Chaos::FPBDCollisionSolverSettings* CollisionSolverSettings,
+		Chaos::FCollisionDetectorSettings* CollisionDetectorSettings)
 	{
 		SCOPE_CYCLE_COUNTER(STAT_ImmediateSimulate_Chaos);
 		CVD_SCOPE_TRACE_SOLVER_FRAME(FSimulation, *this);
@@ -905,18 +910,24 @@ namespace ImmediatePhysics_Chaos
 		// Calculate number of steps to run
 		int32 NumSteps = FMath::Clamp(FMath::RoundToInt32(DeltaTime / StepTime), 1, MaxSubSteps);
 
-		// TEMP: overrides
+		// These are overrides, so only apply if > 0. 
+		SetSolverSettings(
+			ChaosImmediate_Evolution_FixedStepTime,
+			ChaosImmediate_Collision_CullDistance,
+			ChaosImmediate_Collision_MaxDepenetrationVelocity,
+			ChaosImmediate_Joint_UseLinearSolver,
+			ChaosImmediate_Evolution_PositionIterations,
+			ChaosImmediate_Evolution_VelocityIterations,
+			ChaosImmediate_Evolution_ProjectionIterations,
+			ChaosImmediate_Collision_UseManifolds);
+
+		if (JointSolverSettings)
+		{
+			Implementation->Joints.LinearConstraints.SetSettings(*JointSolverSettings);
+			Implementation->Joints.NonLinearConstraints.SetSettings(*JointSolverSettings);
+		}
+		else
 		{
-			SetSolverSettings(
-				ChaosImmediate_Evolution_FixedStepTime,
-				ChaosImmediate_Collision_CullDistance,
-				ChaosImmediate_Collision_MaxDepenetrationVelocity,
-				ChaosImmediate_Joint_UseLinearSolver,
-				ChaosImmediate_Evolution_PositionIterations,
-				ChaosImmediate_Evolution_VelocityIterations,
-				ChaosImmediate_Evolution_ProjectionIterations,
-				ChaosImmediate_Collision_UseManifolds);
-
 			FPBDJointSolverSettings JointsSettings = Implementation->Joints.LinearConstraints.GetSettings();
 			JointsSettings.SwingTwistAngleTolerance = ChaosImmediate_Joint_SwingTwistAngleTolerance;
 			JointsSettings.PositionTolerance = ChaosImmediate_Joint_PositionTolerance;
@@ -947,7 +958,14 @@ namespace ImmediatePhysics_Chaos
 			JointsSettings.bSortEnabled = true;
 			Implementation->Joints.LinearConstraints.SetSettings(JointsSettings);
 			Implementation->Joints.NonLinearConstraints.SetSettings(JointsSettings);
+		}
 
+		if (CollisionSolverSettings)
+		{
+			Implementation->Collisions.SetSolverSettings(*CollisionSolverSettings);
+		}
+		else
+		{
 			Implementation->Collisions.SetPositionFrictionIterations(ChaosImmediate_Collision_NumPositionFrictionIterations);
 			Implementation->Collisions.SetVelocityFrictionIterations(ChaosImmediate_Collision_NumVelocityFrictionIterations);
 			Implementation->Collisions.SetPositionShockPropagationIterations(ChaosImmediate_Collision_NumPositionShockPropagationIterations);
@@ -958,20 +976,27 @@ namespace ImmediatePhysics_Chaos
 			Implementation->Collisions.SetCollisionsEnabled(ChaosImmediate_Collision_Enabled != 0);
 
 			Implementation->Evolution.SetConstraintContainerPriority(Implementation->Collisions.GetContainerId(), ChaosImmediate_Collision_Priority);
+		}
 
+		if (CollisionDetectorSettings)
+		{
+			Implementation->Collisions.SetDetectorSettings(*CollisionDetectorSettings);
+		}
+		else
+		{
 			FCollisionDetectorSettings DetectorSettings = Implementation->Collisions.GetDetectorSettings();
 			DetectorSettings.bAllowManifoldReuse = false;
 			DetectorSettings.bDeferNarrowPhase = (ChaosImmediate_Collision_DeferNarrowPhase != 0);;
 			Implementation->Collisions.SetDetectorSettings(DetectorSettings);
+		}
 
-			if (ChaosImmediate_Evolution_StepTime > 0)
-			{
-				StepTime = ChaosImmediate_Evolution_StepTime;
-			}
-			if (ChaosImmediate_Evolution_NumSteps > 0)
-			{
-				NumSteps = ChaosImmediate_Evolution_NumSteps;
-			}
+		if (ChaosImmediate_Evolution_StepTime > 0)
+		{
+			StepTime = ChaosImmediate_Evolution_StepTime;
+		}
+		if (ChaosImmediate_Evolution_NumSteps > 0)
+		{
+			NumSteps = ChaosImmediate_Evolution_NumSteps;
 		}
 
 		// Fixed timestep mode DT (Only used if > 0.0f)
diff --git a/Engine/Source/Runtime/Engine/Public/Physics/ImmediatePhysics/ImmediatePhysicsChaos/ImmediatePhysicsSimulation_Chaos.h b/Engine/Source/Runtime/Engine/Public/Physics/ImmediatePhysics/ImmediatePhysicsChaos/ImmediatePhysicsSimulation_Chaos.h
--- a/Engine/Source/Runtime/Engine/Public/Physics/ImmediatePhysics/ImmediatePhysicsChaos/ImmediatePhysicsSimulation_Chaos.h
+++ b/Engine/Source/Runtime/Engine/Public/Physics/ImmediatePhysics/ImmediatePhysicsChaos/ImmediatePhysicsSimulation_Chaos.h
@@ -12,6 +12,9 @@
 namespace Chaos
 {
 	class FPBDJointSettings;
+	class FPBDJointSolverSettings;
+	class FPBDCollisionSolverSettings;
+	class FCollisionDetectorSettings;
 }
 
 namespace ImmediatePhysics_Chaos
@@ -81,9 +84,22 @@ namespace ImmediatePhysics_Chaos
 		 */
 		ENGINE_API void SetRewindVelocities(bool bRewindVelocities);
 
-		/** Advance the simulation by DeltaTime */
-		ENGINE_API void Simulate(FReal DeltaTime, FReal MaxStepTime, int32 MaxSubSteps, const FVector& InGravity);
-		void Simulate_AssumesLocked(FReal DeltaTime, FReal MaxStepTime, int32 MaxSubSteps, const FVector& InGravity) { Simulate(DeltaTime, MaxStepTime, MaxSubSteps, InGravity); }
+		/** 
+		 * Advance the simulation by DeltaTime. If settings are passed in they will be used, 
+		 * otherwise they will be taken from CVars 
+		 */
+		ENGINE_API void Simulate(
+			FReal DeltaTime, FReal MaxStepTime, int32 MaxSubSteps, const FVector& InGravity,
+			Chaos::FPBDJointSolverSettings* JointSolverSettings = nullptr, 
+			Chaos::FPBDCollisionSolverSettings* CollisionSolverSettings = nullptr,
+			Chaos::FCollisionDetectorSettings* CollisionDetectorSettings = nullptr);
+		void Simulate_AssumesLocked(FReal DeltaTime, FReal MaxStepTime, int32 MaxSubSteps, const FVector& InGravity,
+			Chaos::FPBDJointSolverSettings* JointSolverSettings = nullptr,
+			Chaos::FPBDCollisionSolverSettings* CollisionSolverSettings = nullptr,
+			Chaos::FCollisionDetectorSettings* CollisionDetectorSettings = nullptr)
+		{
+			Simulate(DeltaTime, MaxStepTime, MaxSubSteps, InGravity, JointSolverSettings, CollisionSolverSettings, CollisionDetectorSettings);
+		}
 
 		ENGINE_API void InitSimulationSpace(
 			const FTransform& Transform);
diff --git a/Engine/Source/Runtime/Experimental/Chaos/Public/Chaos/Collision/PBDCollisionSolverSettings.h b/Engine/Source/Runtime/Experimental/Chaos/Public/Chaos/Collision/PBDCollisionSolverSettings.h
--- a/Engine/Source/Runtime/Experimental/Chaos/Public/Chaos/Collision/PBDCollisionSolverSettings.h
+++ b/Engine/Source/Runtime/Experimental/Chaos/Public/Chaos/Collision/PBDCollisionSolverSettings.h
@@ -11,8 +11,6 @@ namespace Chaos
 	class FPBDCollisionSolverSettings
 	{
 	public:
-		UE_NONCOPYABLE(FPBDCollisionSolverSettings);
-
 		FPBDCollisionSolverSettings()
 			: MaxPushOutVelocity(0)
 			, DepenetrationVelocity(0)
diff --git a/Engine/Source/Runtime/Experimental/Chaos/Public/Chaos/PBDCollisionConstraints.h b/Engine/Source/Runtime/Experimental/Chaos/Public/Chaos/PBDCollisionConstraints.h
--- a/Engine/Source/Runtime/Experimental/Chaos/Public/Chaos/PBDCollisionConstraints.h
+++ b/Engine/Source/Runtime/Experimental/Chaos/Public/Chaos/PBDCollisionConstraints.h
@@ -299,6 +299,11 @@ public:
 
 	const FPBDCollisionSolverSettings& GetSolverSettings() const { return SolverSettings; }
 
+	void SetSolverSettings(const FPBDCollisionSolverSettings& InSettings)
+	{
+		SolverSettings = InSettings;
+	}
+
 	const FCollisionDetectorSettings& GetDetectorSettings() const { return DetectorSettings; }
 
 	void SetDetectorSettings(const FCollisionDetectorSettings& InSettings)
diff --git a/Engine/Source/Runtime/Experimental/Chaos/Public/Chaos/PBDJointConstraintTypes.h b/Engine/Source/Runtime/Experimental/Chaos/Public/Chaos/PBDJointConstraintTypes.h
--- a/Engine/Source/Runtime/Experimental/Chaos/Public/Chaos/PBDJointConstraintTypes.h
+++ b/Engine/Source/Runtime/Experimental/Chaos/Public/Chaos/PBDJointConstraintTypes.h
@@ -144,7 +144,9 @@ namespace Chaos
 		TVector<EJointMotionType, 3> LinearMotionTypes;
 		FReal LinearLimit;
 
+		// Order is twist, swing1, swing2 
 		TVector<EJointMotionType, 3> AngularMotionTypes;
+		// The values here are in radians
 		FVec3 AngularLimits;
 
 		bool bSoftLinearLimitsEnabled;
