Postpone execution of asynchronous physics state creation/destruction tasks until UWorld::UpdateLevelStreaming has finished to avoid physics resource contention.

diff --git a/Engine/Source/Runtime/Engine/Private/PhysicsEngine/Experimental/PhysScene_AsyncPhysicsStateJobQueue.cpp b/Engine/Source/Runtime/Engine/Private/PhysicsEngine/Experimental/PhysScene_AsyncPhysicsStateJobQueue.cpp
--- a/Engine/Source/Runtime/Engine/Private/PhysicsEngine/Experimental/PhysScene_AsyncPhysicsStateJobQueue.cpp
+++ b/Engine/Source/Runtime/Engine/Private/PhysicsEngine/Experimental/PhysScene_AsyncPhysicsStateJobQueue.cpp
@@ -7,8 +7,15 @@
 #include "PhysicsEngine/BodySetup.h"
 #include "Engine/World.h"
 
+FPhysScene_AsyncPhysicsStateJobQueue::FPhysScene_AsyncPhysicsStateJobQueue(FPhysScene* InPhysicScene)
+	: PhysScene(InPhysicScene)
+{
+	PhysScene->GetOwningWorld()->OnAllLevelsChanged().AddRaw(this, &FPhysScene_AsyncPhysicsStateJobQueue::OnUpdateLevelStreamingDone);
+}
+
 FPhysScene_AsyncPhysicsStateJobQueue::~FPhysScene_AsyncPhysicsStateJobQueue()
 {
+	PhysScene->GetOwningWorld()->OnAllLevelsChanged().RemoveAll(this);
 	Launch();
 	WaitForAsyncTasksExecution();
 	OnAsyncTasksExecuted();
@@ -59,6 +66,7 @@ void FPhysScene_AsyncPhysicsStateJobQueue::OnAsyncTasksExecuted()
 	}
 	PipedJobs.Reset();
 	PipeTasks.Reset();
+	IsReadyToRunAsyncTasksEvent.Reset();
 }
 
 bool FPhysScene_AsyncPhysicsStateJobQueue::AreAsyncTasksExecuted() const
@@ -91,8 +99,27 @@ void FPhysScene_AsyncPhysicsStateJobQueue::Launch()
 	// Move Jobs to the async queue
 	FJobs& NewJobs = PipedJobs.Emplace_GetRef();
 	NewJobs = MoveTemp(PendingJobs);
+	if (!IsReadyToRunAsyncTasksEvent)
+	{
+		IsReadyToRunAsyncTasksEvent.Emplace(UE_SOURCE_LOCATION);
+	}
 	// Launch a new task in the pipe
-	PipeTasks.Emplace(Pipe.Launch(UE_SOURCE_LOCATION, [this, NewJobs = NewJobs.Get()]() { NewJobs->Execute(); }, UE::Tasks::ETaskPriority::BackgroundHigh));
+	PipeTasks.Emplace(Pipe.Launch(UE_SOURCE_LOCATION, [this, NewJobs = NewJobs.Get()]() { NewJobs->Execute(); }, *IsReadyToRunAsyncTasksEvent, UE::Tasks::ETaskPriority::BackgroundHigh));
+}
+
+void FPhysScene_AsyncPhysicsStateJobQueue::OnUpdateLevelStreamingDone()
+{
+	TriggerIsReadyToRunAsyncTasksEvent();
+	
+}
+
+void FPhysScene_AsyncPhysicsStateJobQueue::TriggerIsReadyToRunAsyncTasksEvent()
+{
+	if (IsReadyToRunAsyncTasksEvent)
+	{
+		IsReadyToRunAsyncTasksEvent->Trigger();
+		IsReadyToRunAsyncTasksEvent.Reset();
+	}
 }
 
 void FPhysScene_AsyncPhysicsStateJobQueue::WaitForAsyncTasksExecution()
@@ -100,6 +127,7 @@ void FPhysScene_AsyncPhysicsStateJobQueue::WaitForAsyncTasksExecution()
 	TRACE_CPUPROFILER_EVENT_SCOPE(FPhysScene_AsyncTaskQueue::WaitForAsyncTasksExecution);
 	check(IsInGameThread());
 
+	TriggerIsReadyToRunAsyncTasksEvent();
 	UE::Tasks::Wait(PipeTasks);
 }
 
diff --git a/Engine/Source/Runtime/Engine/Public/Physics/Experimental/PhysScene_AsyncPhysicsStateJobQueue.h b/Engine/Source/Runtime/Engine/Public/Physics/Experimental/PhysScene_AsyncPhysicsStateJobQueue.h
--- a/Engine/Source/Runtime/Engine/Public/Physics/Experimental/PhysScene_AsyncPhysicsStateJobQueue.h
+++ b/Engine/Source/Runtime/Engine/Public/Physics/Experimental/PhysScene_AsyncPhysicsStateJobQueue.h
@@ -12,10 +12,7 @@ class UActorComponent;
 
 struct FPhysScene_AsyncPhysicsStateJobQueue
 {
-	FPhysScene_AsyncPhysicsStateJobQueue(FPhysScene* InPhysicScene)
-		: PhysScene(InPhysicScene)
-	{}
-
+	FPhysScene_AsyncPhysicsStateJobQueue(FPhysScene* InPhysicScene);
 	~FPhysScene_AsyncPhysicsStateJobQueue();
 
 	enum class EJobType
@@ -86,11 +83,14 @@ private:
 	void WaitForAsyncTasksExecution();
 	bool AreAsyncTasksExecuted() const;
 	void OnAsyncTasksExecuted();
+	void OnUpdateLevelStreamingDone();
+	void TriggerIsReadyToRunAsyncTasksEvent();
 
 	FPhysScene* PhysScene;
 	using FJobs = TUniquePtr<FJobSet>;
 	FJobs PendingJobs;
 	TArray<FJobs> PipedJobs;
 	TArray<UE::Tasks::FTask> PipeTasks;
+	TOptional<UE::Tasks::FTaskEvent> IsReadyToRunAsyncTasksEvent;
 	UE::Tasks::FPipe Pipe{ TEXT("FPhysScene::AsyncCreatePhysicStateQueue") };
 };
\ No newline at end of file
