World Partition: add missing support to dump to streaming generation state log for RuntimeHashSetExternalStreamingObject.

diff --git a/Engine/Source/Runtime/Engine/Private/WorldPartition/RuntimeHashSet/WorldPartitionRuntimeHashSet.cpp b/Engine/Source/Runtime/Engine/Private/WorldPartition/RuntimeHashSet/WorldPartitionRuntimeHashSet.cpp
--- a/Engine/Source/Runtime/Engine/Private/WorldPartition/RuntimeHashSet/WorldPartitionRuntimeHashSet.cpp
+++ b/Engine/Source/Runtime/Engine/Private/WorldPartition/RuntimeHashSet/WorldPartitionRuntimeHashSet.cpp
@@ -126,6 +126,53 @@ void URuntimeHashSetExternalStreamingObject::DestroyPartitionsSpatialIndex() con
 	}
 }
 
+#if WITH_EDITOR
+void URuntimeHashSetExternalStreamingObject::DumpStateLog(FHierarchicalLogArchive& Ar)
+{
+	Super::DumpStateLog(Ar);
+
+	TArray<const UWorldPartitionRuntimeCell*> StreamingCells;
+	for (const FRuntimePartitionStreamingData& StreamingData : RuntimeStreamingData)
+	{
+		auto HandleStreamingCell = [&StreamingCells](UWorldPartitionRuntimeCell* StreamingCell)
+		{
+			if (!StreamingCell->IsAlwaysLoaded() || !IsRunningCookCommandlet())
+			{
+				StreamingCells.Add(StreamingCell);
+			}
+		};
+
+		for (UWorldPartitionRuntimeCell* StreamingCell : StreamingData.SpatiallyLoadedCells)
+		{
+			HandleStreamingCell(StreamingCell);
+		}
+
+		for (UWorldPartitionRuntimeCell* StreamingCell : StreamingData.NonSpatiallyLoadedCells)
+		{
+			HandleStreamingCell(StreamingCell);
+		}
+	}
+				
+	StreamingCells.Sort([this](const UWorldPartitionRuntimeCell& A, const UWorldPartitionRuntimeCell& B)
+	{
+		if (A.IsAlwaysLoaded() ==  B.IsAlwaysLoaded())
+		{
+			return A.GetFName().LexicalLess(B.GetFName());
+		}
+
+		return A.IsAlwaysLoaded() > B.IsAlwaysLoaded();
+	});
+
+	for (const UWorldPartitionRuntimeCell* StreamingCell : StreamingCells)
+	{
+		FHierarchicalLogArchive::FIndentScope CellIndentScope = Ar.PrintfIndent(TEXT("Content of Cell %s (%s)"), *StreamingCell->GetDebugName(), *StreamingCell->GetName());
+		StreamingCell->DumpStateLog(Ar);
+	}
+
+	Ar.Printf(TEXT(""));
+}
+#endif
+
 void URuntimeHashSetExternalStreamingObject::AddReferencedObjects(UObject* InThis, FReferenceCollector& Collector)
 {
 #if WITH_EDITOR
diff --git a/Engine/Source/Runtime/Engine/Public/WorldPartition/RuntimeHashSet/WorldPartitionRuntimeHashSet.h b/Engine/Source/Runtime/Engine/Public/WorldPartition/RuntimeHashSet/WorldPartitionRuntimeHashSet.h
--- a/Engine/Source/Runtime/Engine/Public/WorldPartition/RuntimeHashSet/WorldPartitionRuntimeHashSet.h
+++ b/Engine/Source/Runtime/Engine/Public/WorldPartition/RuntimeHashSet/WorldPartitionRuntimeHashSet.h
@@ -123,6 +123,10 @@ class ENGINE_API URuntimeHashSetExternalStreamingObject : public URuntimeHashExt
 {
 	GENERATED_BODY()
 
+#if WITH_EDITOR
+	virtual void DumpStateLog(FHierarchicalLogArchive& Ar) override;
+#endif
+
 public:
 	//~ Begin UObject Interface
 	static void AddReferencedObjects(UObject* InThis, FReferenceCollector& Collector);
