[Shaders] Log job diagnostics in local and distributed shader controllers on job stalls.

The stalls are detected when a 64-bit hash over the distributed task inputs and local worker states have not changed within the period of time, specified by CVar "r.ShaderCompiler.WorkerHeartbeat".
This also replaces the old "r.ShaderCompiler.DistributedControllerTimeout" CVar with the new one.
An exerpt of an example output may look like this:

LogShaderCompilers: Warning: No shader compile worker state change in 900.12 seconds
LogShaderCompilers: Display: ======= ShaderCompileWorker Diagnostics =======
LogShaderCompilers: Display: Worker [1/16]: bAvailable=1, bComplete=0, bIssuedTasksToWorker=1, bLaunchedWorker=1
  Job [1/5]: Global/FLensDistUVGenerationVS/0:/Plugin/ShaderLoopPlugin/Private/LoopShader.usf|MainVS (PermutationId=0)
  Job [2/5]: Global/FHoldoutCompositeSSRInputShader/0:/Plugin/HoldoutComposite/Private/HoldoutCompositeSSRInput.usf|MainPS (PermutationId=0)
  Job [3/5]: Global/FHoldoutCompositeGammaShader/0:/Plugin/HoldoutComposite/Private/HoldoutCompositeGamma.usf|MainCS (PermutationId=0)
  Job [4/5]: Global/FWaterQuadTreeBucketCountsCS/1:/Plugin/Water/Private/WaterQuadTreeDraws.usf|MainCS (PermutationId=1)
  Job [5/5]: Global/FWaterQuadTreeBucketPrefixSumCS/1:/Plugin/Water/Private/WaterQuadTreeDraws.usf|MainCS (PermutationId=1)

diff --git a/Engine/Source/Runtime/Engine/Private/ShaderCompiler/ShaderCompilerDistributed.cpp b/Engine/Source/Runtime/Engine/Private/ShaderCompiler/ShaderCompilerDistributed.cpp
--- a/Engine/Source/Runtime/Engine/Private/ShaderCompiler/ShaderCompilerDistributed.cpp
+++ b/Engine/Source/Runtime/Engine/Private/ShaderCompiler/ShaderCompilerDistributed.cpp
@@ -23,13 +23,6 @@ namespace DistributedShaderCompilerVariables
 		TEXT("Smaller number of shaders will compile locally."),
 		ECVF_Default);
 
-	static int32 GDistributedControllerTimeout = 15 * 60;
-	static FAutoConsoleVariableRef CVarDistributedControllerTimeout(
-		TEXT("r.ShaderCompiler.DistributedControllerTimeout"),
-		GDistributedControllerTimeout,
-		TEXT("Maximum number of seconds we expect to pass between getting distributed controller complete a task (this is used to detect problems with the distribution controllers).")
-	);
-
 	static int32 GDistributedJobDescriptionLevel = 1;
 	static FAutoConsoleVariableRef CVarDistributedJobDescriptionLevel(
 		TEXT("r.ShaderCompiler.DistributedJobDescriptionLevel"),
@@ -71,7 +64,6 @@ public:
 FShaderCompileDistributedThreadRunnable_Interface::FShaderCompileDistributedThreadRunnable_Interface(class FShaderCompilingManager* InManager, IDistributedBuildController& InController)
 	: FShaderCompileThreadRunnableBase(InManager)
 	, NumDispatchedJobs(0)
-	, LastTimeTaskCompleted(FPlatformTime::Seconds())
 	, bIsHung(false)
 	, CachedController(InController)
 {
@@ -282,6 +274,32 @@ TArray<FString> FShaderCompileDistributedThreadRunnable_Interface::GetDependency
 	return Dependencies;
 }
 
+static void LogShaderCompileWorkerDistributedDiagnostics(const TSparseArray<FDistributedShaderCompilerTask*>& InDispatchedTasks)
+{
+	UE_LOG(LogShaderCompilers, Display, TEXT("======= ShaderCompileWorker-Distributed Diagnostics ======="));
+
+	FString JobDiagnostics;
+
+	for (int32 TaskIndex = 0; TaskIndex < InDispatchedTasks.Num(); ++TaskIndex)
+	{
+		FDistributedShaderCompilerTask* Task = InDispatchedTasks[TaskIndex];
+		if (!Task->ShaderJobs.IsEmpty())
+		{
+			JobDiagnostics.Empty();
+
+			for (int32 JobIndex = 0; JobIndex < Task->ShaderJobs.Num(); ++JobIndex)
+			{
+				Task->ShaderJobs[JobIndex]->AppendDiagnostics(JobDiagnostics, JobIndex, Task->ShaderJobs.Num(), TEXT("  "));
+			}
+
+			UE_LOG(
+				LogShaderCompilers, Display, TEXT("Task [%d/%d]:\n%s"),
+				TaskIndex + 1, InDispatchedTasks.Num(), *JobDiagnostics
+			);
+		}
+	}
+}
+
 int32 FShaderCompileDistributedThreadRunnable_Interface::CompilingLoop()
 {
 	TArray<FShaderCommonCompileJobPtr> PendingJobs;
@@ -306,12 +324,6 @@ int32 FShaderCompileDistributedThreadRunnable_Interface::CompilingLoop()
 		}
 	}
 
-	// if we don't have any dispatched jobs, reset the completion timer
-	if (DispatchedTasks.Num() == 0)
-	{
-		LastTimeTaskCompleted = FPlatformTime::Seconds();
-	}
-
 	if (PendingJobs.Num() > 0)
 	{
 		// Increase the batch size when more jobs are queued/in flight.
@@ -415,19 +427,32 @@ int32 FShaderCompileDistributedThreadRunnable_Interface::CompilingLoop()
 		}
 	}
 
+	TMemoryHasher<FXxHash64Builder, FXxHash64> WorkerStateHasher;
+	bool bHasAnyJobs = false;
+
 	for (auto Iter = DispatchedTasks.CreateIterator(); Iter; ++Iter)
 	{
 		bool bOutputFileReadFailed = true;
 
 		FDistributedShaderCompilerTask* Task = *Iter;
-		if (!Task->Future.IsReady())
+
+		bHasAnyJobs = bHasAnyJobs || !Task->ShaderJobs.IsEmpty();
+
+		// Add jobs input hashes to the current state hash
+		bool bIsTaskReady = Task->Future.IsReady();
+		WorkerStateHasher << bIsTaskReady;
+		for (const FShaderCommonCompileJobPtr& Job : Task->ShaderJobs)
+		{
+			WorkerStateHasher << Job->InputHash;
+		}
+
+		if (!bIsTaskReady)
 		{
 			continue;
 		}
 
 		FDistributedBuildTaskResult Result = Task->Future.Get();
 		NumDispatchedJobs -= Task->ShaderJobs.Num();
-		LastTimeTaskCompleted = FPlatformTime::Seconds();
 
 		if (Result.ReturnCode != 0)
 		{
@@ -506,16 +531,10 @@ int32 FShaderCompileDistributedThreadRunnable_Interface::CompilingLoop()
 	// Yield for a short while to stop this thread continuously polling the disk.
 	FPlatformProcess::Sleep(0.01f);
 
-	// normally we expect to have at least one task in 5 minutes, although there could be edge cases
-	double TimeSinceLastCompletedTask = FPlatformTime::Seconds() - LastTimeTaskCompleted;
-	if (TimeSinceLastCompletedTask >= DistributedShaderCompilerVariables::GDistributedControllerTimeout)
+	// Check if shader jobs have not changed in too long
+	if (!WorkerStateHeartbeat(bHasAnyJobs ? WorkerStateHasher.Finalize().Hash : 0))
 	{
-		if (!bIsHung)
-		{
-			UE_LOG(LogShaderCompilers, Warning, TEXT("Distributed compilation controller didn't receive a completed task in %f seconds!"), TimeSinceLastCompletedTask);
-			bIsHung = true;
-			// TODO: resubmit the hung jobs
-		}
+		LogShaderCompileWorkerDistributedDiagnostics(DispatchedTasks);
 	}
 
 	// Return true if there is more work to be done.
diff --git a/Engine/Source/Runtime/Engine/Private/ShaderCompiler/ShaderCompilerThreadRunnable.cpp b/Engine/Source/Runtime/Engine/Private/ShaderCompiler/ShaderCompilerThreadRunnable.cpp
--- a/Engine/Source/Runtime/Engine/Private/ShaderCompiler/ShaderCompilerThreadRunnable.cpp
+++ b/Engine/Source/Runtime/Engine/Private/ShaderCompiler/ShaderCompilerThreadRunnable.cpp
@@ -43,6 +43,13 @@ static FAutoConsoleVariableRef CVarShaderCompilerMemoryLimit(
 	TEXT("By default 0, effectively disabling the limitation. If this is non-zero, it must be greater than or equal to 1024 since shader compilation must be granted at least 1024 MiB of memory in total."),
 	ECVF_ReadOnly);
 
+int32 GShaderWorkerStateChangeHeartbeat = 15 * 60;
+static TAutoConsoleVariable<int32> CVarShaderWorkerStateChangeHeartbeat(
+	TEXT("r.ShaderCompiler.WorkerHeartbeat"),
+	GShaderWorkerStateChangeHeartbeat,
+	TEXT("Number of seconds until an unchanged state of compile workers will dump their state to diagnose hung shader compile jobs. Default is 15 * 60 seconds."),
+	ECVF_Default);
+
 extern bool GShaderCompilerDumpWorkerDiagnostics;
 
 // Configuration to retry shader compile through workers after a worker has been abandoned
@@ -155,13 +162,38 @@ struct FShaderCompileWorkerInfo
 };
 
 FShaderCompileThreadRunnableBase::FShaderCompileThreadRunnableBase(FShaderCompilingManager* InManager)
-	: Manager(InManager)
+	: WorkerStateHash(0)
+	, WorkerStateChangeTimestamp(-1.0)
+	, Manager(InManager)
 	, Thread(nullptr)
 	, MinPriorityIndex(0)
 	, MaxPriorityIndex(NumShaderCompileJobPriorities - 1)
 	, bForceFinish(false)
 {
 }
+
+bool FShaderCompileThreadRunnableBase::WorkerStateHeartbeat(uint64 InWorkerStateHash)
+{
+	// Reset worker state hash if it has changed or requested to reset (value of 0)
+	const double CurrentTimestamp = FPlatformTime::Seconds();
+	if (InWorkerStateHash == 0 || InWorkerStateHash != WorkerStateHash || WorkerStateChangeTimestamp < 0.0)
+	{
+		WorkerStateHash = InWorkerStateHash;
+		WorkerStateChangeTimestamp = CurrentTimestamp;
+	}
+
+	// Report warning when hearbeat failed and reset timer
+	const double ElapsedTimeSinceWorkerStateChanged = CurrentTimestamp - WorkerStateChangeTimestamp;
+	if (ElapsedTimeSinceWorkerStateChanged > (double)GShaderWorkerStateChangeHeartbeat)
+	{
+		UE_LOG(LogShaderCompilers, Warning, TEXT("No shader compile worker state change in %.2f seconds"), ElapsedTimeSinceWorkerStateChanged);
+		WorkerStateChangeTimestamp = CurrentTimestamp;
+		return false;
+	}
+
+	return true;
+}
+
 void FShaderCompileThreadRunnableBase::StartThread()
 {
 	if (Manager->bAllowAsynchronousShaderCompiling && !FPlatformProperties::RequiresCookedData())
@@ -1180,8 +1212,53 @@ FShaderCompileMemoryUsage FShaderCompileThreadRunnable::GetExternalWorkerMemoryU
 	return MemoryUsage;
 }
 
+static void LogShaderCompileWorkerDiagnostics(const TArray<TUniquePtr<FShaderCompileWorkerInfo>>& InWorkerInfos)
+{
+	UE_LOG(LogShaderCompilers, Display, TEXT("======= ShaderCompileWorker Diagnostics ======="));
+
+	FString JobDiagnostics;
+
+	for (int32 WorkerIndex = 0; WorkerIndex < InWorkerInfos.Num(); ++WorkerIndex)
+	{
+		const TUniquePtr<FShaderCompileWorkerInfo>& WorkerInfo = InWorkerInfos[WorkerIndex];
+		if (!WorkerInfo->QueuedJobs.IsEmpty())
+		{
+			JobDiagnostics.Empty();
+
+			for (int32 JobIndex = 0; JobIndex < WorkerInfo->QueuedJobs.Num(); ++JobIndex)
+			{
+				WorkerInfo->QueuedJobs[JobIndex]->AppendDiagnostics(JobDiagnostics, JobIndex, WorkerInfo->QueuedJobs.Num(), TEXT("  "));
+			}
+
+			UE_LOG(
+				LogShaderCompilers, Display, TEXT("Worker [%d/%d]: bAvailable=%d, bComplete=%d, bIssuedTasksToWorker=%d, bLaunchedWorker=%d\n%s"),
+				WorkerIndex + 1, InWorkerInfos.Num(), (int32)WorkerInfo->bAvailable, (int32)WorkerInfo->bComplete, (int32)WorkerInfo->bIssuedTasksToWorker, (int32)WorkerInfo->bLaunchedWorker, *JobDiagnostics
+			);
+		}
+	}
+}
+
 int32 FShaderCompileThreadRunnable::CompilingLoop()
 {
+	// Generate hash over all worker states to detect hung shader compile jobs
+	TMemoryHasher<FXxHash64Builder, FXxHash64> WorkerStateHasher;
+	bool bHasAnyJobs = false;
+
+	for (const TUniquePtr<FShaderCompileWorkerInfo>& WorkerInfo : WorkerInfos)
+	{
+		WorkerStateHasher << WorkerInfo->bAvailable << WorkerInfo->bComplete << WorkerInfo->bIssuedTasksToWorker << WorkerInfo->bLaunchedWorker;
+		for (const auto& Job : WorkerInfo->QueuedJobs)
+		{
+			bHasAnyJobs = true;
+			WorkerStateHasher << Job->InputHash;
+		}
+	}
+
+	if (!WorkerStateHeartbeat(bHasAnyJobs ? WorkerStateHasher.Finalize().Hash : 0))
+	{
+		LogShaderCompileWorkerDiagnostics(WorkerInfos);
+	}
+
 	if (!Manager->bAllowCompilingThroughWorkers && CVarCompileParallelInProcess.GetValueOnAnyThread())
 	{
 		int32 NumJobs = Manager->GetNumPendingJobs();
diff --git a/Engine/Source/Runtime/Engine/Public/ShaderCompiler.h b/Engine/Source/Runtime/Engine/Public/ShaderCompiler.h
--- a/Engine/Source/Runtime/Engine/Public/ShaderCompiler.h
+++ b/Engine/Source/Runtime/Engine/Public/ShaderCompiler.h
@@ -137,6 +137,12 @@ class FShaderCompileThreadRunnableBase : public FRunnable
 {
 	friend class FShaderCompilingManager;
 
+	/** 64-bit hash value of the worker states to detect hung shader compile jobs. */
+	uint64 WorkerStateHash;
+
+	/** Timestamp of the last point in time the worker states have changed. */
+	double WorkerStateChangeTimestamp;
+
 protected:
 	/** The manager for this thread */
 	class FShaderCompilingManager* Manager;
@@ -163,6 +169,12 @@ protected:
 		return TEXT("ShaderCompilingThread");
 	}
 
+	/**
+	Returns false if the new worker state hash hasn't changed in a certain period of time. If false, a warning is printed.
+	An input hash of 0 will always succeed and effectively reset the timer. A value of 0 should be used when no jobs are pending.
+	*/
+	bool WorkerStateHeartbeat(uint64 InWorkerStateHash);
+
 public:
 	FShaderCompileThreadRunnableBase(class FShaderCompilingManager* InManager);
 	virtual ~FShaderCompileThreadRunnableBase()
@@ -343,9 +355,6 @@ class FShaderCompileDistributedThreadRunnable_Interface : public FShaderCompileT
 
 	TSparseArray<class FDistributedShaderCompilerTask*> DispatchedTasks;
 
-	/** Last time we received a task back. */
-	double LastTimeTaskCompleted;
-
 	/** Whether we consider this controller hung / out of order. */
 	bool bIsHung;
 
diff --git a/Engine/Source/Runtime/RenderCore/Private/ShaderCore.cpp b/Engine/Source/Runtime/RenderCore/Private/ShaderCore.cpp
--- a/Engine/Source/Runtime/RenderCore/Private/ShaderCore.cpp
+++ b/Engine/Source/Runtime/RenderCore/Private/ShaderCore.cpp
@@ -4195,6 +4195,14 @@ FStringView FShaderCompileJob::GetFinalSourceView() const
 	}
 }
 
+void FShaderCompileJob::AppendDiagnostics(FString& OutDiagnostics, int32 InJobIndex, int32 InNumJobs, const TCHAR* Indentation) const
+{
+	OutDiagnostics += FString::Printf(
+		TEXT("%sJob [%d/%d]: %s (PermutationId=%d)\n"),
+		Indentation == nullptr ? TEXT("") : Indentation, InJobIndex + 1, InNumJobs, *Input.GenerateShaderName(), Key.PermutationId
+	);
+}
+
 FShaderPipelineCompileJob::FShaderPipelineCompileJob(int32 NumStages)
 	: FShaderCommonCompileJob(Type, 0u, 0u, EShaderCompileJobPriority::Num)
 {
@@ -4268,3 +4276,17 @@ void FShaderPipelineCompileJob::AppendDebugName(FStringBuilderBase& OutName) con
 		OutName << TEXT("\n");
 	}
 }
+
+void FShaderPipelineCompileJob::AppendDiagnostics(FString& OutDiagnostics, int32 InJobIndex, int32 InNumJobs, const TCHAR* Indentation) const
+{
+	FString NewIndentation = Indentation == nullptr ? TEXT("") : Indentation;
+	OutDiagnostics += FString::Printf(
+		TEXT("%sPipeline Job [%d/%d] (PermutationId=%d):\n"),
+		*NewIndentation, InJobIndex + 1, InNumJobs, Key.PermutationId
+	);
+	NewIndentation += TEXT("  ");
+	for (int32 StageJobIndex = 0; StageJobIndex < StageJobs.Num(); ++StageJobIndex)
+	{
+		StageJobs[StageJobIndex]->AppendDiagnostics(OutDiagnostics, StageJobIndex, StageJobs.Num(), *NewIndentation);
+	}
+}
diff --git a/Engine/Source/Runtime/RenderCore/Public/ShaderCompilerJobTypes.h b/Engine/Source/Runtime/RenderCore/Public/ShaderCompilerJobTypes.h
--- a/Engine/Source/Runtime/RenderCore/Public/ShaderCompilerJobTypes.h
+++ b/Engine/Source/Runtime/RenderCore/Public/ShaderCompilerJobTypes.h
@@ -156,6 +156,9 @@ public:
 	/** Serializes (and deserializes) the output for caching purposes. */
 	virtual void SerializeOutput(FShaderCacheSerializeContext& Ctx) {}
 
+	/** Generates a diagnostics string for this compile job suitable for the output log. */
+	virtual void AppendDiagnostics(FString& OutDiagnostics, int32 InJobIndex, int32 InNumJobs, const TCHAR* Indentation = nullptr) const = 0;
+
 	FShaderCompileJob* GetSingleShaderJob();
 	const FShaderCompileJob* GetSingleShaderJob() const;
 	FShaderPipelineCompileJob* GetShaderPipelineJob();
@@ -281,6 +284,8 @@ public:
 
 	RENDERCORE_API FStringView GetFinalSourceView() const;
 
+	virtual RENDERCORE_API void AppendDiagnostics(FString& OutDiagnostics, int32 InJobIndex, int32 InNumJobs, const TCHAR* Indentation = nullptr) const override final;
+
 	FShaderCompileJob() : FShaderCommonCompileJob(Type, 0u, 0u, EShaderCompileJobPriority::Num)
 	{}
 
@@ -331,6 +336,7 @@ public:
 
 	virtual RENDERCORE_API void OnComplete() override;
 	virtual RENDERCORE_API void AppendDebugName(FStringBuilderBase& OutName) const override;
+	virtual RENDERCORE_API void AppendDiagnostics(FString& OutDiagnostics, int32 InJobIndex, int32 InNumJobs, const TCHAR* Indentation = nullptr) const override final;
 
 	RENDERCORE_API FShaderPipelineCompileJob(int32 NumStages);
 	RENDERCORE_API FShaderPipelineCompileJob(uint32 InHash, uint32 InId, EShaderCompileJobPriority InPriroity, const FShaderPipelineCompileJobKey& InKey);
