Provide VTable a constructor for  UMaterial so materials that are structed with a FVTableHelper do not cause garbage collection to spin forever.  This change ensures that ReleasedByRT is set to true when this type of constructor is hit, allowing the material to be garbage collected.

diff --git a/Engine/Source/Runtime/Engine/Private/Materials/Material.cpp b/Engine/Source/Runtime/Engine/Private/Materials/Material.cpp
--- a/Engine/Source/Runtime/Engine/Private/Materials/Material.cpp
+++ b/Engine/Source/Runtime/Engine/Private/Materials/Material.cpp
@@ -1024,9 +1024,28 @@ bool UMaterial::IsDefaultMaterial() const
 	return bDefault;
 }
 
+UMaterial::UMaterial()
+	: Super()
+	, ReleasedByRT(true)
+{
+	SetInitialValues();
+}
+
 UMaterial::UMaterial(const FObjectInitializer& ObjectInitializer)
 	: Super(ObjectInitializer)
 	, ReleasedByRT(true)
+{
+	SetInitialValues();
+}
+
+UMaterial::UMaterial(FVTableHelper& Helper)
+	: Super(Helper)
+	, ReleasedByRT(true)
+{
+	SetInitialValues();
+}
+
+void UMaterial::SetInitialValues()
 {
 	BlendMode = BLEND_Opaque;
 	ShadingModel = MSM_DefaultLit;
diff --git a/Engine/Source/Runtime/Engine/Public/Materials/Material.h b/Engine/Source/Runtime/Engine/Public/Materials/Material.h
--- a/Engine/Source/Runtime/Engine/Public/Materials/Material.h
+++ b/Engine/Source/Runtime/Engine/Public/Materials/Material.h
@@ -1160,6 +1160,9 @@ private:
 #endif // WITH_EDITORONLY_DATA
 public:
 
+	ENGINE_API UMaterial();
+	ENGINE_API UMaterial(FVTableHelper& Helper);
+
 	//~ Begin UMaterialInterface Interface.
 	ENGINE_API virtual UMaterial* GetMaterial() override;
 	ENGINE_API virtual const UMaterial* GetMaterial() const override;
@@ -1393,6 +1396,7 @@ public:
 	
 private:
 
+	void SetInitialValues();
 
 	/** Handles setting up an annotation for this object if a flag has changed value */
 	void MarkUsageFlagDirty(EMaterialUsage Usage, bool CurrentValue, bool NewValue);
