Fix some members of FLightRenderParameters not being initialized by FDirectionalLightSceneProxy.

diff --git a/Engine/Source/Runtime/Engine/Private/Components/DirectionalLightComponent.cpp b/Engine/Source/Runtime/Engine/Private/Components/DirectionalLightComponent.cpp
--- a/Engine/Source/Runtime/Engine/Private/Components/DirectionalLightComponent.cpp
+++ b/Engine/Source/Runtime/Engine/Private/Components/DirectionalLightComponent.cpp
@@ -413,6 +413,8 @@ public:
 		LightParameters.SourceRadius = FMath::Sin( 0.5f * FMath::DegreesToRadians( LightSourceAngle ) );
 		LightParameters.SoftSourceRadius = FMath::Sin( 0.5f * FMath::DegreesToRadians( LightSourceSoftAngle ) );
 		LightParameters.SourceLength = 0.0f;
+		LightParameters.RectLightBarnCosAngle = 0.0f;
+		LightParameters.RectLightBarnLength = -2.0f;
 		LightParameters.RectLightAtlasUVOffset = FVector2f::ZeroVector;
 		LightParameters.RectLightAtlasUVScale = FVector2f::ZeroVector;
 		LightParameters.RectLightAtlasMaxLevel = FLightRenderParameters::GetRectLightAtlasInvalidMIPLevel();
