Actors that are being migrated will no longer have EndPLay() and BeginPlay() called on them during migration

diff --git a/Engine/Source/Runtime/Engine/Private/Actor.cpp b/Engine/Source/Runtime/Engine/Private/Actor.cpp
--- a/Engine/Source/Runtime/Engine/Private/Actor.cpp
+++ b/Engine/Source/Runtime/Engine/Private/Actor.cpp
@@ -973,6 +973,25 @@ void AActor::Serialize(FArchive& Ar)
 	{		
 		FActorInstanceGuid::Serialize(Ar, *this);
 	}
+
+#if UE_WITH_REMOTE_OBJECT_HANDLE
+	if (Ar.IsMigratingRemoteObjects())
+	{
+		Ar << CreationTime;
+
+		float CurrentLifespan = 0.0f; 
+		if (Ar.IsSaving())
+		{
+			CurrentLifespan = GetLifeSpan();
+		}
+		Ar << CurrentLifespan;
+		if (Ar.IsLoading())
+		{
+			// When migrating actors are in a state after BeginPlay has been called but we need to call SetLifeSpan to initialize local timer
+			SetLifeSpan(CurrentLifespan);
+		}
+	}
+#endif // UE_WITH_REMOTE_OBJECT_HANDLE
 }
 
 void AActor::PostLoad()
@@ -1139,6 +1158,9 @@ void AActor::PostLoadSubobjects(FObjectInstancingGraph* OuterInstanceGraph)
 }
 
 #if UE_WITH_REMOTE_OBJECT_HANDLE
+
+static USceneComponent* FixupNativeActorComponents(AActor* Actor);
+
 void AActor::PostMigrate(FRemoteServerId DestinationServerId)
 {
 	Super::PostMigrate(DestinationServerId);
@@ -1148,14 +1170,89 @@ void AActor::PostMigrate(FRemoteServerId DestinationServerId)
 	{
 		// Remove the actor from the world on this server as it is about to be migrated to another server.
 		// Migrated actors are unique to the server instance that currently owns their memory and are not replicated to other servers.
-		Destroy(false, false);
+		if (UWorld* const World = GetWorld())
+		{
+			ClearComponentOverlaps();
+			UE::Net::Private::FNetPropertyConditionManager::Get().NotifyObjectDestroyed(this);
+
+			// Notify net drivers that this actor has been destroyed.
+			if (FWorldContext* Context = GEngine->GetWorldContextFromWorld(World))
+			{
+				for (FNamedNetDriver& Driver : Context->ActiveNetDrivers)
+				{
+					if (Driver.NetDriver != nullptr && Driver.NetDriver->ShouldReplicateActor(this))
+					{
+						Driver.NetDriver->NotifyActorDestroyed(this);
+					}
+				}
+			}
+
+			// Remove the actor from the actor list.
+			World->RemoveActor(this, /*bShouldModifyLevel*/ false);
+		}
+
+		// Clean up the actor's components.
+		UnregisterAllComponents();
+
+		// Mark the actor and its direct components as pending kill.
+		MarkAsGarbage();
+		MarkComponentsAsGarbage();
+
+		// Unregister the actor's tick function
+		RegisterAllActorTickFunctions(/*bRegisterTickFunctions*/ false, /*bIncludeComponents*/ true);
 	}
 	else
 	{
 		// Add the actor that was migrated from another server to this server's world
 		GetLevel()->TryAddActorToList(this, /*bAddUnique*/false);
-		PostSpawnInitialize(GetTransform(), GetOwner(), nullptr, false, true, false);
-		GetWorld()->AddNetworkActor(this);
+
+		UWorld* const World = GetWorld();
+		checkf(World != nullptr, TEXT("Migrated Actor %s (%s) doesn't have a World"), *GetPathName(), *FRemoteObjectId(this).ToString());
+
+		// Set the actor's world transform if it has a native rootcomponent.
+		USceneComponent* const SceneRootComponent = FixupNativeActorComponents(this);
+		FTransform FinalRootComponentTransform = GetTransform();
+		if (SceneRootComponent != nullptr)
+		{
+			// Respect any non-default transform value that the root component may have received from the archetype that's owned
+			// by the native CDO, so the final transform might not always necessarily equate to the passed-in UserSpawnTransform.
+			const FTransform RootTransform(SceneRootComponent->GetRelativeRotation(), SceneRootComponent->GetRelativeLocation(), SceneRootComponent->GetRelativeScale3D());
+			FinalRootComponentTransform = RootTransform * GetTransform();
+			SceneRootComponent->SetWorldTransform(FinalRootComponentTransform, false, nullptr, ETeleportType::ResetPhysics);
+		}
+
+		RegisterAllComponents();
+
+		bHasFinishedSpawning = true;
+
+		PreInitializeComponents();
+		InitializeComponents();
+		PostInitializeComponents();
+
+		BuildReplicatedComponentsInfo();
+
+#if UE_WITH_IRIS
+		BeginReplication();
+#endif // UE_WITH_IRIS
+
+		RegisterAllActorTickFunctions(true, false);
+
+		TInlineComponentArray<UActorComponent*> Components;
+		GetComponents(Components);
+
+		for (UActorComponent* Component : Components)
+		{
+			if (Component->IsRegistered())
+			{
+				Component->RegisterAllComponentTickFunctions(true);
+			}
+		}
+
+		ActorHasBegunPlay = EActorBeginPlayState::HasBegunPlay;
+
+		UpdateInitialOverlaps(/*bFromLevelStreaming*/ false);
+		
+		World->AddNetworkActor(this);
 	}
 }
 #endif // UE_WITH_REMOTE_OBJECT_HANDLE
diff --git a/Engine/Source/Runtime/Engine/Private/Components/ActorComponent.cpp b/Engine/Source/Runtime/Engine/Private/Components/ActorComponent.cpp
--- a/Engine/Source/Runtime/Engine/Private/Components/ActorComponent.cpp
+++ b/Engine/Source/Runtime/Engine/Private/Components/ActorComponent.cpp
@@ -2896,6 +2896,27 @@ void UActorComponent::Serialize(FArchive& Ar)
 			}
 		}
 	}
+
+#if UE_WITH_REMOTE_OBJECT_HANDLE
+	if (Ar.IsMigratingRemoteObjects())
+	{
+		uint8 ComponentFlags = 0;
+
+		if (Ar.IsSaving())
+		{
+			ComponentFlags |= bHasBeenCreated	<< 0;
+			ComponentFlags |= bHasBegunPlay		<< 1;
+		}
+
+		Ar << ComponentFlags;
+
+		if (Ar.IsLoading())
+		{
+			bHasBeenCreated	= (ComponentFlags >> 0) & 1;
+			bHasBegunPlay	= (ComponentFlags >> 1) & 1;
+		}
+	}
+#endif // UE_WITH_REMOTE_OBJECT_HANDLE
 }
 
 void UActorComponent::PostApplyToComponent()
