Do not load default material for default instances of UVolumetricCloudComponent

Otherwise, it loads material which might not be used and holds it in memory forever.

#rnx

diff --git a/Engine/Source/Runtime/Engine/Private/Components/VolumetricCloudComponent.cpp b/Engine/Source/Runtime/Engine/Private/Components/VolumetricCloudComponent.cpp
--- a/Engine/Source/Runtime/Engine/Private/Components/VolumetricCloudComponent.cpp
+++ b/Engine/Source/Runtime/Engine/Private/Components/VolumetricCloudComponent.cpp
@@ -57,15 +57,22 @@ UVolumetricCloudComponent::UVolumetricCloudComponent(const FObjectInitializer& O
 	, bRenderInMainPass(true)
 	, VolumetricCloudSceneProxy(nullptr)
 {
-	static ConstructorHelpers::FObjectFinder<UMaterialInterface> VolumetricCloudDefaultMaterialRef(TEXT("/Engine/EngineSky/VolumetricClouds/m_SimpleVolumetricCloud_Inst.m_SimpleVolumetricCloud_Inst"));
-	Material = VolumetricCloudDefaultMaterialRef.Object;
+	bool bUseDefaultMaterial = !HasAnyFlags(RF_ClassDefaultObject | RF_DefaultSubObject);
+	if (bUseDefaultMaterial && GetOuter() != nullptr)
+	{
+		bUseDefaultMaterial = !GetOuter()->HasAnyFlags(RF_ClassDefaultObject);
+	}
+	if (bUseDefaultMaterial)
+	{
+		static ConstructorHelpers::FObjectFinder<UMaterialInterface> VolumetricCloudDefaultMaterialRef(TEXT("/Engine/EngineSky/VolumetricClouds/m_SimpleVolumetricCloud_Inst.m_SimpleVolumetricCloud_Inst"));
+		Material = VolumetricCloudDefaultMaterialRef.Object;
+	}
 }
 
 UVolumetricCloudComponent::~UVolumetricCloudComponent()
 {
 }
 
-
 void UVolumetricCloudComponent::CreateRenderState_Concurrent(FRegisterComponentContext* Context)
 {
 	Super::CreateRenderState_Concurrent(Context);
