UDeletedObjectPlaceholder: expose required public functions for scene graph development and change create function to take a generic outer instead of a level.

diff --git a/Engine/Source/Runtime/Engine/Private/DeletedObjectPlaceholder.cpp b/Engine/Source/Runtime/Engine/Private/DeletedObjectPlaceholder.cpp
--- a/Engine/Source/Runtime/Engine/Private/DeletedObjectPlaceholder.cpp
+++ b/Engine/Source/Runtime/Engine/Private/DeletedObjectPlaceholder.cpp
@@ -4,7 +4,6 @@
 
 #if WITH_EDITOR
 
-#include "Engine/Level.h"
 #include "WorldPartition/ActorDescContainer.h"
 #include "WorldPartition/WorldPartition.h"
 #include "WorldPartition/DataLayer/ExternalDataLayerAsset.h"
@@ -16,17 +15,17 @@ bool UDeletedObjectPlaceholder::IsAsset() const
 	return IsPackageExternal() && !GetPackage()->HasAnyFlags(RF_Transient) && !HasAnyFlags(RF_ClassDefaultObject) && !GetPackage()->HasAnyPackageFlags(PKG_PlayInEditor) && !AActor::FindActorInPackage(GetExternalPackage());
 }
 
-UDeletedObjectPlaceholder* UDeletedObjectPlaceholder::Create(ULevel* InLevel, UPackage* InPackage, const UObject* InOriginalObject)
+UDeletedObjectPlaceholder* UDeletedObjectPlaceholder::Create(UObject* InOuter, UPackage* InPackage, const UObject* InOriginalObject)
 {
 	if (IsRunningCommandlet())
 	{
 		return nullptr;
 	}
-	check(InLevel);
+	check(InOuter);
 	check(InOriginalObject);
 	check(!UDeletedObjectPlaceholder::FindInPackage(InPackage));
 	const AActor* OriginalActor = Cast<AActor>(InOriginalObject);
-	UDeletedObjectPlaceholder* DeletedObjectPlaceholder = NewObject<UDeletedObjectPlaceholder>(InLevel, NAME_None, RF_Standalone | RF_Transactional | RF_Transient);
+	UDeletedObjectPlaceholder* DeletedObjectPlaceholder = NewObject<UDeletedObjectPlaceholder>(InOuter, NAME_None, RF_Standalone | RF_Transactional | RF_Transient);
 	DeletedObjectPlaceholder->SetExternalPackage(InPackage);
 	DeletedObjectPlaceholder->OriginalObject = InOriginalObject;
 	DeletedObjectPlaceholder->DisplayName = OriginalActor ? OriginalActor->GetActorLabel() : InOriginalObject->GetName();
diff --git a/Engine/Source/Runtime/Engine/Public/DeletedObjectPlaceholder.h b/Engine/Source/Runtime/Engine/Public/DeletedObjectPlaceholder.h
--- a/Engine/Source/Runtime/Engine/Public/DeletedObjectPlaceholder.h
+++ b/Engine/Source/Runtime/Engine/Public/DeletedObjectPlaceholder.h
@@ -29,9 +29,9 @@ public:
 	const FString& GetDisplayName() const { return DisplayName; }
 	const FExternalDataLayerUID GetExternalDataLayerUID() const { return FExternalDataLayerUID(ExternalDataLayerUID); }
 
-	static UDeletedObjectPlaceholder* Create(ULevel* InLevel, UPackage* InPackage, const UObject* InOriginalObject);
-	static UDeletedObjectPlaceholder* FindInPackage(UPackage* InPackage);
-	static UDeletedObjectPlaceholder* RemoveFromPackage(UPackage* InPackage);
+	static ENGINE_API UDeletedObjectPlaceholder* Create(UObject* InOuter, UPackage* InPackage, const UObject* InOriginalObject);
+	static ENGINE_API UDeletedObjectPlaceholder* FindInPackage(UPackage* InPackage);
+	static ENGINE_API UDeletedObjectPlaceholder* RemoveFromPackage(UPackage* InPackage);
 
 	DECLARE_MULTICAST_DELEGATE_OneParam(FObjectCreated, const UDeletedObjectPlaceholder*);
 	static ENGINE_API FObjectCreated OnObjectCreated;
