UActorDescContainer: added verbose log when registering actor descriptors used while debugging a recent issue.

diff --git a/Engine/Source/Runtime/Engine/Private/WorldPartition/ActorDescContainer.cpp b/Engine/Source/Runtime/Engine/Private/WorldPartition/ActorDescContainer.cpp
--- a/Engine/Source/Runtime/Engine/Private/WorldPartition/ActorDescContainer.cpp
+++ b/Engine/Source/Runtime/Engine/Private/WorldPartition/ActorDescContainer.cpp
@@ -383,6 +383,8 @@ void UActorDescContainer::RegisterActorDescriptor(FWorldPartitionActorDesc* Acto
 	ActorDesc->SetContainer(this);
 		
 	ActorsByName.Add(ActorDesc->GetActorName(), ActorsByGuid.FindChecked(ActorDesc->GetGuid()));
+
+	UE_LOG(LogWorldPartition, Verbose, TEXT("\tRegistered actor descriptor '%s'"), *ActorDesc->ToString(FWorldPartitionActorDesc::EToStringMode::Verbose));
 }
 
 void UActorDescContainer::UnregisterActorDescriptor(FWorldPartitionActorDesc* ActorDesc)
