Fix to calculate all active transition alphas prior to setting state weights

diff --git a/Engine/Source/Runtime/Engine/Private/Animation/AnimNode_StateMachine.cpp b/Engine/Source/Runtime/Engine/Private/Animation/AnimNode_StateMachine.cpp
--- a/Engine/Source/Runtime/Engine/Private/Animation/AnimNode_StateMachine.cpp
+++ b/Engine/Source/Runtime/Engine/Private/Animation/AnimNode_StateMachine.cpp
@@ -492,17 +492,20 @@ void FAnimNode_StateMachine::Update_AnyThread(const FAnimationUpdateContext& Con
 	if (ActiveTransitionArray.Num() > 0)
 	{
 		// Keep track of states that have blended out to avoid recalling their anim node state functions.
-		TSet<int32> BlendedOutStates;
+		TSet<int32, DefaultKeyFuncs<int32>, TInlineSetAllocator<4>> BlendedOutStates;
+		TArray<bool, TInlineAllocator<4>> FinishedTransitions;
+		FinishedTransitions.Init(false, ActiveTransitionArray.Num());
+
+		// Update the transitions prior to the states to ensure all transition alphas are correct
+		for (int32 Index = 0; Index < ActiveTransitionArray.Num(); ++Index)
+		{
+			ActiveTransitionArray[Index].Update(Context, CurrentState, /*out*/ FinishedTransitions[Index]);
+		}
 		
 		for (int32 Index = 0; Index < ActiveTransitionArray.Num(); ++Index)
 		{
 			// The custom graph will tick the needed states
-			bool bFinishedTrans = false;
-
-			// The custom graph will tick the needed states
-			ActiveTransitionArray[Index].Update(Context, CurrentState, /*out*/ bFinishedTrans);
-			
-			if (bFinishedTrans)
+			if (FinishedTransitions[Index])
 			{
 				// Trigger "Fully Blended Out" of any states that have any weight to them.
 				{
