Turn toast message into a warning message.

- This avoids toast spam which can lead to GPU crashes from creating too many swap chains.

diff --git a/Engine/Source/Runtime/Engine/Private/Materials/Material.cpp b/Engine/Source/Runtime/Engine/Private/Materials/Material.cpp
--- a/Engine/Source/Runtime/Engine/Private/Materials/Material.cpp
+++ b/Engine/Source/Runtime/Engine/Private/Materials/Material.cpp
@@ -1933,14 +1933,16 @@ bool UMaterial::SetMaterialUsage(bool &bNeedsRecompile, EMaterialUsage Usage, UM
 				#if WITH_EDITOR
 				else if (GEditor && !GEditor->IsPlaySessionInProgress())
 				{
-					FFormatNamedArguments Args;
-					Args.Add(TEXT("UsageName"), FText::FromString(GetUsageName(Usage)));
-					FNotificationInfo Info(FText::Format(LOCTEXT("CouldntSetMaterialUsage","Material didn't allow automatic setting of usage flag {UsageName} needed to render on this component, using Default Material instead."), Args));
-					Info.ExpireDuration = 5.0f;
-					Info.bUseSuccessFailIcons = true;
+					// Temporarily Turn off the toast to avoid spam which can result in GPU crashes from too many swap chains from being created.
+
+					//FFormatNamedArguments Args;
+					//Args.Add(TEXT("UsageName"), FText::FromString(GetUsageName(Usage)));
+					//FNotificationInfo Info(FText::Format(LOCTEXT("CouldntSetMaterialUsage","Material didn't allow automatic setting of usage flag {UsageName} needed to render on this component, using Default Material instead."), Args));
+					//Info.ExpireDuration = 5.0f;
+					//Info.bUseSuccessFailIcons = true;
 
 					// Give the user feedback as to why they are seeing the default material
-					FSlateNotificationManager::Get().AddNotification(Info);
+					//FSlateNotificationManager::Get().AddNotification(Info);
 				}
 				#endif
 
