Optimize how we validate textures in level to avoid scanning all primitives

#rnx

diff --git a/Engine/Source/Runtime/Engine/Private/World.cpp b/Engine/Source/Runtime/Engine/Private/World.cpp
--- a/Engine/Source/Runtime/Engine/Private/World.cpp
+++ b/Engine/Source/Runtime/Engine/Private/World.cpp
@@ -1220,24 +1220,29 @@ void UWorld::PostDuplicate(bool bDuplicateForPIE)
 				{
 					TRACE_CPUPROFILER_EVENT_SCOPE(ValidateTextureOverridesForPIE);
 
-					TSet<UMaterialInterface*> ProcessedMaterials;
-					for (TObjectIterator<UPrimitiveComponent> It; It; ++It)
+					TSet<UMaterialInterface*> ProcessedMaterials; 
+					TArray<UMaterialInterface*> Materials;
+					TArray<UPrimitiveComponent*> Components;
+					for (const TObjectPtr<AActor>& Actor : Level->Actors)
 					{
-						UPrimitiveComponent* Component = *It;
-						AActor* Owner = Component->GetOwner();
-						if (Owner != nullptr && !Owner->HasAnyFlags(RF_ClassDefaultObject) && Owner->IsInLevel(Level))
+						if (Actor != nullptr)
 						{
-							TArray<UMaterialInterface*> Materials;
-							Component->GetUsedMaterials(Materials);
-							for (UMaterialInterface* Material : Materials)
+							Components.Reset();
+							Actor->GetComponents<UPrimitiveComponent>(Components);
+							for (UPrimitiveComponent* Component : Components)
 							{
-								bool bIsAlreadyInSet = false;
-								ProcessedMaterials.FindOrAdd(Material, &bIsAlreadyInSet);
-								if (!bIsAlreadyInSet)
+								Materials.Reset();
+								Component->GetUsedMaterials(Materials);
+								for (UMaterialInterface* Material : Materials)
 								{
-									if (UMaterialInstance* MaterialInstance = Cast<UMaterialInstance>(Material))
+									bool bIsAlreadyInSet = false;
+									ProcessedMaterials.FindOrAdd(Material, &bIsAlreadyInSet);
+									if (!bIsAlreadyInSet)
 									{
-										MaterialInstance->ValidateTextureOverrides(FeatureLevel);
+										if (UMaterialInstance* MaterialInstance = Cast<UMaterialInstance>(Material))
+										{
+											MaterialInstance->ValidateTextureOverrides(FeatureLevel);
+										}
 									}
 								}
 							}
