[LocationVolume] Disable collisions on location volumes

diff --git a/Engine/Source/Runtime/Engine/Private/LocationVolume.cpp b/Engine/Source/Runtime/Engine/Private/LocationVolume.cpp
--- a/Engine/Source/Runtime/Engine/Private/LocationVolume.cpp
+++ b/Engine/Source/Runtime/Engine/Private/LocationVolume.cpp
@@ -28,6 +28,7 @@ ALocationVolume::ALocationVolume(const FObjectInitializer& ObjectInitializer)
 : Super(ObjectInitializer)
 {
 	GetBrushComponent()->SetGenerateOverlapEvents(false);
+	SetActorEnableCollision(false);
 
 #if WITH_EDITORONLY_DATA
 	bIsSpatiallyLoaded = false;
