Chaos::Enforce soft constraint on joints to not be treated as locked constraints.

diff --git a/Engine/Source/Runtime/Engine/Private/PhysicsEngine/ConstraintTypes.cpp b/Engine/Source/Runtime/Engine/Private/PhysicsEngine/ConstraintTypes.cpp
--- a/Engine/Source/Runtime/Engine/Private/PhysicsEngine/ConstraintTypes.cpp
+++ b/Engine/Source/Runtime/Engine/Private/PhysicsEngine/ConstraintTypes.cpp
@@ -77,7 +77,7 @@ bool ShouldSkipSoftLimits(float Stiffness, float Damping, float AverageMass)
 void FLinearConstraint::UpdateLinearLimit_AssumesLocked(const FPhysicsConstraintHandle& InConstraintRef, float AverageMass, float Scale) const
 {
 	const float UseLimit = FMath::Max(Limit * Scale, UE_KINDA_SMALL_NUMBER);	//physx doesn't ever want limit of 0
-	const bool bLockLimitSize = (UseLimit < RB_MinSizeToLockDOF);
+	const bool bLockLimitSize = (UseLimit < RB_MinSizeToLockDOF) && !bSoftConstraint;  //Joint should not be locked under soft constraint. 
 	
 	const bool bSkipSoft = bSoftConstraint && ShouldSkipSoftLimits(Stiffness, Damping, AverageMass);
 
diff --git a/Engine/Source/Runtime/Engine/Private/PhysicsEngine/ImmediatePhysics/ImmediatePhysicsChaos/ImmediatePhysicsJointHandle_Chaos.cpp b/Engine/Source/Runtime/Engine/Private/PhysicsEngine/ImmediatePhysics/ImmediatePhysicsChaos/ImmediatePhysicsJointHandle_Chaos.cpp
--- a/Engine/Source/Runtime/Engine/Private/PhysicsEngine/ImmediatePhysics/ImmediatePhysicsChaos/ImmediatePhysicsJointHandle_Chaos.cpp
+++ b/Engine/Source/Runtime/Engine/Private/PhysicsEngine/ImmediatePhysics/ImmediatePhysicsChaos/ImmediatePhysicsJointHandle_Chaos.cpp
@@ -198,7 +198,7 @@ namespace ImmediatePhysics_Chaos
 
 
 		// UE Disables Soft Limits when the Limit is less than some threshold. This is not necessary in Chaos but for now we also do it for parity's sake (See FLinearConstraint::UpdateLinearLimit_AssumesLocked).
-		if (JointSettings.LinearLimit < RB_MinSizeToLockDOF)
+		if (JointSettings.LinearLimit < RB_MinSizeToLockDOF && !Profile.LinearLimit.bSoftConstraint)
 		{
 			for (int32 AxisIndex = 0; AxisIndex < 3; ++AxisIndex)
 			{
