[Local Exposure] Expose soft thresholds in PPV.

diff --git a/Engine/Plugins/Cameras/GameplayCameras/Source/GameplayCameras/Private/Core/PostProcessSettingsCollection.cpp b/Engine/Plugins/Cameras/GameplayCameras/Source/GameplayCameras/Private/Core/PostProcessSettingsCollection.cpp
--- a/Engine/Plugins/Cameras/GameplayCameras/Source/GameplayCameras/Private/Core/PostProcessSettingsCollection.cpp
+++ b/Engine/Plugins/Cameras/GameplayCameras/Source/GameplayCameras/Private/Core/PostProcessSettingsCollection.cpp
@@ -128,6 +128,9 @@ void FPostProcessSettingsCollection::OverrideChanged(const FPostProcessSettings&
 		UE_SET_PP(LocalExposureDetailStrength);
 		UE_SET_PP(LocalExposureBlurredLuminanceBlend);
 		UE_SET_PP(LocalExposureBlurredLuminanceKernelSizePercent);
+		UE_SET_PP(LocalExposureUseSoftThresholds);
+		UE_SET_PP(LocalExposureHighlightThresholdSmoothness);
+		UE_SET_PP(LocalExposureShadowThresholdSmoothness);
 		UE_SET_PP(LocalExposureMiddleGreyBias);
 		UE_SET_PP(LensFlareIntensity);
 		UE_SET_PP(LensFlareTint);
@@ -431,6 +434,9 @@ void FPostProcessSettingsCollection::InternalLerpChanged(const FPostProcessSetti
 		UE_LERP_PP(LocalExposureDetailStrength);
 		UE_LERP_PP(LocalExposureBlurredLuminanceBlend);
 		UE_LERP_PP(LocalExposureBlurredLuminanceKernelSizePercent);
+		UE_SET_PP(LocalExposureUseSoftThresholds);
+		UE_LERP_PP(LocalExposureHighlightThresholdSmoothness);
+		UE_LERP_PP(LocalExposureShadowThresholdSmoothness);
 		UE_LERP_PP(LocalExposureMiddleGreyBias);
 		UE_LERP_PP(LensFlareIntensity);
 		UE_LERP_PP(LensFlareTint);
diff --git a/Engine/Shaders/Private/PostProcessHistogramCommon.ush b/Engine/Shaders/Private/PostProcessHistogramCommon.ush
--- a/Engine/Shaders/Private/PostProcessHistogramCommon.ush
+++ b/Engine/Shaders/Private/PostProcessHistogramCommon.ush
@@ -44,7 +44,9 @@ float2 LocalExposure_BilateralGridUVScale;
 float LocalExposure_HighlightThreshold;
 float LocalExposure_ShadowThreshold;
 uint LocalExposure_bUseSoftThreshold;
-float LocalExposure_SoftThresholdStrength;
+uint LocalExposure_bSoftThresholdUseSmoothStep;
+float LocalExposure_HighlightThresholdSmoothness;
+float LocalExposure_ShadowThresholdSmoothness;
 
 float CalculateEyeAdaptationLuminance(float3 Color)
 {
@@ -258,14 +260,31 @@ float CalculateLogLocalExposure(float LogLum, float BaseLogLum, float LogMiddleG
 
 	if(LocalExposure_bUseSoftThreshold != 0)
 	{
+		float Blend;
 		if (BaseCentered > 0)
 		{
-			ContrastScale = lerp(1, ContrastScale, saturate((BaseCentered - LocalExposure_HighlightThreshold) * LocalExposure_SoftThresholdStrength));
+			if (LocalExposure_bSoftThresholdUseSmoothStep)
+			{
+				Blend = smoothstep(LocalExposure_HighlightThreshold, LocalExposure_HighlightThreshold + 1.0f / (1.0f - LocalExposure_HighlightThresholdSmoothness), BaseCentered);
+			}
+			else
+			{
+				Blend = saturate((BaseCentered - LocalExposure_HighlightThreshold) * (1.0f - LocalExposure_HighlightThresholdSmoothness));
+			}
 		}
 		else
 		{
-			ContrastScale = lerp(1, ContrastScale, saturate((abs(BaseCentered) - LocalExposure_ShadowThreshold) * LocalExposure_SoftThresholdStrength));
+			if (LocalExposure_bSoftThresholdUseSmoothStep)
+			{
+				Blend = smoothstep(LocalExposure_ShadowThreshold, LocalExposure_ShadowThreshold + 1.0f / (1.0f - LocalExposure_ShadowThresholdSmoothness), abs(BaseCentered));
+			}
+			else
+			{
+				Blend = saturate((abs(BaseCentered) - LocalExposure_ShadowThreshold) * (1.0f - LocalExposure_ShadowThresholdSmoothness));
+			}
 		}
+
+		ContrastScale = lerp(1, ContrastScale, Blend);
 	}
 	else
 	{
diff --git a/Engine/Shaders/Private/PostProcessVisualizeLocalExposure.usf b/Engine/Shaders/Private/PostProcessVisualizeLocalExposure.usf
--- a/Engine/Shaders/Private/PostProcessVisualizeLocalExposure.usf
+++ b/Engine/Shaders/Private/PostProcessVisualizeLocalExposure.usf
@@ -94,14 +94,42 @@ float4 MainPS(noperspective float4 UVAndScreenPos : TEXCOORD0, float4 SvPosition
 		{
 			if (BaseCentered - LocalExposure_HighlightThreshold > 0)
 			{
-				OutValue = lerp(OutValue, float3(0,1,0), DiagonalStripesPattern(SvPosition.xy, 32.0f));
+				float ThresholdFade = 1.0f;
+
+				if (LocalExposure_bUseSoftThreshold)
+				{
+					if (LocalExposure_bSoftThresholdUseSmoothStep)
+					{
+						ThresholdFade = smoothstep(LocalExposure_HighlightThreshold, LocalExposure_HighlightThreshold + 1.0f / (1.0f - LocalExposure_HighlightThresholdSmoothness), BaseCentered);
+					}
+					else
+					{
+						ThresholdFade = saturate((BaseCentered - LocalExposure_HighlightThreshold) * (1.0f - LocalExposure_HighlightThresholdSmoothness));
+					}
+				}
+
+				OutValue = lerp(OutValue, float3(0,1,0), ThresholdFade * DiagonalStripesPattern(SvPosition.xy, 32.0f));
 			}
 		}
 		else
 		{
 			if (BaseCentered + LocalExposure_ShadowThreshold < 0)
 			{
-				OutValue = lerp(OutValue, float3(1,0,0), DiagonalStripesPattern(SvPosition.xy, 32.0f));
+				float ThresholdFade = 1.0f;
+
+				if (LocalExposure_bUseSoftThreshold)
+				{
+					if (LocalExposure_bSoftThresholdUseSmoothStep)
+					{
+						ThresholdFade = smoothstep(LocalExposure_ShadowThreshold, LocalExposure_ShadowThreshold + 1.0f / (1.0f - LocalExposure_ShadowThresholdSmoothness), abs(BaseCentered));
+					}
+					else
+					{
+						ThresholdFade = saturate((abs(BaseCentered) - LocalExposure_ShadowThreshold) * (1.0f - LocalExposure_ShadowThresholdSmoothness));
+					}
+				}
+
+				OutValue = lerp(OutValue, float3(1,0,0), ThresholdFade * DiagonalStripesPattern(SvPosition.xy, 32.0f));
 			}
 		}
 	}
diff --git a/Engine/Source/Runtime/Engine/Classes/Engine/Scene.h b/Engine/Source/Runtime/Engine/Classes/Engine/Scene.h
--- a/Engine/Source/Runtime/Engine/Classes/Engine/Scene.h
+++ b/Engine/Source/Runtime/Engine/Classes/Engine/Scene.h
@@ -932,6 +932,15 @@ struct FPostProcessSettings
 	UPROPERTY(EditAnywhere, BlueprintReadWrite, Category=Overrides, meta=(PinHiddenByDefault, InlineEditConditionToggle))
 	uint8 bOverride_LocalExposureBlurredLuminanceKernelSizePercent:1;
 
+	UPROPERTY(EditAnywhere, BlueprintReadWrite, Category=Overrides, meta=(PinHiddenByDefault, InlineEditConditionToggle))
+	uint8 bOverride_LocalExposureUseSoftThresholds:1;
+
+	UPROPERTY(EditAnywhere, BlueprintReadWrite, Category=Overrides, meta=(PinHiddenByDefault, InlineEditConditionToggle))
+	uint8 bOverride_LocalExposureHighlightThresholdSmoothness :1;
+
+	UPROPERTY(EditAnywhere, BlueprintReadWrite, Category=Overrides, meta=(PinHiddenByDefault, InlineEditConditionToggle))
+	uint8 bOverride_LocalExposureShadowThresholdSmoothness :1;
+
 	UPROPERTY(EditAnywhere, BlueprintReadWrite, Category=Overrides, meta=(PinHiddenByDefault, InlineEditConditionToggle))
 	uint8 bOverride_LocalExposureMiddleGreyBias:1;
 
@@ -1947,11 +1956,27 @@ struct FPostProcessSettings
 	UPROPERTY(interp, BlueprintReadWrite, Category="Lens|Local Exposure", meta=(UIMin = "0.0", UIMax = "100.0", editcondition = "bOverride_LocalExposureBlurredLuminanceKernelSizePercent", DisplayName = "Blurred Luminance Kernel Size Percent"))
 	float LocalExposureBlurredLuminanceKernelSizePercent;
 
+	/** Enables soft thresholds. */
+	UPROPERTY(EditAnywhere, BlueprintReadWrite, Category = "Lens|Local Exposure", AdvancedDisplay, meta = (editcondition = "bOverride_LocalExposureUseSoftThresholds", DisplayName = "Use Soft Thresholds (Experimental)"))
+	uint32 LocalExposureUseSoftThresholds : 1;
+
+	/** 
+	 * Smoothness of transition between regions considered highlights (only relevant when using Soft Thresholds).
+	*/
+	UPROPERTY(interp, BlueprintReadWrite, Category="Lens|Local Exposure", AdvancedDisplay, meta=(UIMin = "0.0", UIMax = "1.0", editcondition = "bOverride_LocalExposureHighlightThresholdSmoothness", DisplayName = "Highlight Threshold Smoothness"))
+	float LocalExposureHighlightThresholdSmoothness;
+
+	/**
+	 * Smoothness of transition between regions considered shadows (only relevant when using Soft Thresholds).
+	*/
+	UPROPERTY(interp, BlueprintReadWrite, Category="Lens|Local Exposure", AdvancedDisplay, meta=(UIMin = "0.0", UIMax = "1.0", editcondition = "bOverride_LocalExposureShadowThresholdSmoothness", DisplayName = "Shadow Threshold Smoothness"))
+	float LocalExposureShadowThresholdSmoothness;
+
 	/**
 	 * Logarithmic adjustment for the local exposure middle grey.
 	 * 0: no adjustment, -1:2x darker, -2:4x darker, 1:2x brighter, 2:4x brighter, ...
 	 */
-	UPROPERTY(interp, BlueprintReadWrite, Category = "Lens|Local Exposure", meta = (UIMin = "-15.0", UIMax = "15.0", editcondition = "bOverride_LocalExposureMiddleGreyBias", DisplayName = "Middle Grey Bias"))
+	UPROPERTY(interp, BlueprintReadWrite, Category = "Lens|Local Exposure", AdvancedDisplay, meta = (UIMin = "-15.0", UIMax = "15.0", editcondition = "bOverride_LocalExposureMiddleGreyBias", DisplayName = "Middle Grey Bias"))
 	float LocalExposureMiddleGreyBias;
 
 	/** Brightness scale of the image cased lens flares (linear) */
diff --git a/Engine/Source/Runtime/Engine/Private/PostProcessVolume.cpp b/Engine/Source/Runtime/Engine/Private/PostProcessVolume.cpp
--- a/Engine/Source/Runtime/Engine/Private/PostProcessVolume.cpp
+++ b/Engine/Source/Runtime/Engine/Private/PostProcessVolume.cpp
@@ -175,6 +175,13 @@ bool APostProcessVolume::CanEditChange(const FProperty* InProperty) const
 				return Settings.AutoExposureMethod == EAutoExposureMethod::AEM_Histogram;
 			}
 
+			// Parameters that are only used when using soft thhresholds
+			if (PropertyName == GET_MEMBER_NAME_STRING_CHECKED(FPostProcessSettings, LocalExposureHighlightThresholdSmoothness) ||
+				PropertyName == GET_MEMBER_NAME_STRING_CHECKED(FPostProcessSettings, LocalExposureShadowThresholdSmoothness))
+			{
+				return Settings.LocalExposureUseSoftThresholds != 0;
+			}
+
 			// Parameters that are only used for the Sum of Gaussian bloom / not the texture based fft bloom
 			if (PropertyName == GET_MEMBER_NAME_STRING_CHECKED(FPostProcessSettings, BloomThreshold) ||
 				PropertyName == GET_MEMBER_NAME_STRING_CHECKED(FPostProcessSettings, BloomSizeScale) ||
diff --git a/Engine/Source/Runtime/Engine/Private/Scene.cpp b/Engine/Source/Runtime/Engine/Private/Scene.cpp
--- a/Engine/Source/Runtime/Engine/Private/Scene.cpp
+++ b/Engine/Source/Runtime/Engine/Private/Scene.cpp
@@ -517,6 +517,9 @@ FPostProcessSettings::FPostProcessSettings()
 	LocalExposureDetailStrength = 1.0f;
 	LocalExposureBlurredLuminanceBlend = 0.6f;
 	LocalExposureBlurredLuminanceKernelSizePercent = 50.0f;
+	LocalExposureUseSoftThresholds = false;
+	LocalExposureHighlightThresholdSmoothness = 0.6f;
+	LocalExposureShadowThresholdSmoothness = 0.6f;
 	LocalExposureMiddleGreyBias = 0.0f;
 
 	// next value might get overwritten by r.DefaultFeature.LensFlare
diff --git a/Engine/Source/Runtime/Engine/Private/SceneView.cpp b/Engine/Source/Runtime/Engine/Private/SceneView.cpp
--- a/Engine/Source/Runtime/Engine/Private/SceneView.cpp
+++ b/Engine/Source/Runtime/Engine/Private/SceneView.cpp
@@ -1661,6 +1661,9 @@ void FSceneView::OverridePostProcessSettings(const FPostProcessSettings& Src, fl
 		LERP_PP(LocalExposureDetailStrength);
 		LERP_PP(LocalExposureBlurredLuminanceBlend);
 		LERP_PP(LocalExposureBlurredLuminanceKernelSizePercent);
+		SET_PP(LocalExposureUseSoftThresholds);
+		LERP_PP(LocalExposureHighlightThresholdSmoothness);
+		LERP_PP(LocalExposureShadowThresholdSmoothness);
 		LERP_PP(LocalExposureMiddleGreyBias);
 		LERP_PP(LensFlareIntensity);
 		LERP_PP(LensFlareTint);
diff --git a/Engine/Source/Runtime/Renderer/Private/PostProcess/PostProcessLocalExposure.cpp b/Engine/Source/Runtime/Renderer/Private/PostProcess/PostProcessLocalExposure.cpp
--- a/Engine/Source/Runtime/Renderer/Private/PostProcess/PostProcessLocalExposure.cpp
+++ b/Engine/Source/Runtime/Renderer/Private/PostProcess/PostProcessLocalExposure.cpp
@@ -13,16 +13,10 @@
 #include "ShaderCompilerCore.h"
 #include "DataDrivenShaderPlatformInfo.h"
 
-static TAutoConsoleVariable<bool> CVarLocalExposureSoftThreshold(
-	TEXT("r.LocalExposure.SoftThreshold"),
+static TAutoConsoleVariable<bool> CVarLocalExposureSoftThresholdSmoothstep(
+	TEXT("r.LocalExposure.SoftThreshold.Smoothstep"),
 	false,
-	TEXT("Whether to use soft thresholds."),
-	ECVF_RenderThreadSafe);
-
-static TAutoConsoleVariable<float> CVarLocalExposureSoftThresholdStrength(
-	TEXT("r.LocalExposure.SoftThreshold.Strength"),
-	1.0f,
-	TEXT(""),
+	TEXT("Whether to use smoothstep when calculating soft thresholds."),
 	ECVF_RenderThreadSafe);
 
 static TAutoConsoleVariable<float> CVarExposureFusionTargetLuminance(
@@ -272,8 +266,10 @@ FLocalExposureParameters GetLocalExposureParameters(const FViewInfo& View, FIntP
 	Parameters.BilateralGridUVScale = LocalExposureBilateralGridUVScale;
 	Parameters.HighlightThreshold = Settings.LocalExposureHighlightThreshold;
 	Parameters.ShadowThreshold = Settings.LocalExposureShadowThreshold;
-	Parameters.bUseSoftThreshold = CVarLocalExposureSoftThreshold.GetValueOnRenderThread() ? 1.0f : 0.0f;
-	Parameters.SoftThresholdStrength = CVarLocalExposureSoftThresholdStrength.GetValueOnRenderThread();
+	Parameters.bUseSoftThreshold = Settings.LocalExposureUseSoftThresholds;
+	Parameters.bSoftThresholdUseSmoothStep = CVarLocalExposureSoftThresholdSmoothstep.GetValueOnRenderThread();
+	Parameters.HighlightThresholdSmoothness = Settings.LocalExposureHighlightThresholdSmoothness;
+	Parameters.ShadowThresholdSmoothness = Settings.LocalExposureShadowThresholdSmoothness;
 	return Parameters;
 }
 
diff --git a/Engine/Source/Runtime/Renderer/Private/PostProcess/PostProcessLocalExposure.h b/Engine/Source/Runtime/Renderer/Private/PostProcess/PostProcessLocalExposure.h
--- a/Engine/Source/Runtime/Renderer/Private/PostProcess/PostProcessLocalExposure.h
+++ b/Engine/Source/Runtime/Renderer/Private/PostProcess/PostProcessLocalExposure.h
@@ -16,7 +16,9 @@ BEGIN_SHADER_PARAMETER_STRUCT(FLocalExposureParameters, )
 	SHADER_PARAMETER(float, HighlightThreshold)
 	SHADER_PARAMETER(float, ShadowThreshold)
 	SHADER_PARAMETER(uint32, bUseSoftThreshold)
-	SHADER_PARAMETER(float, SoftThresholdStrength)
+	SHADER_PARAMETER(uint32, bSoftThresholdUseSmoothStep)
+	SHADER_PARAMETER(float, HighlightThresholdSmoothness)
+	SHADER_PARAMETER(float, ShadowThresholdSmoothness)
 END_SHADER_PARAMETER_STRUCT()
 
 FLocalExposureParameters GetLocalExposureParameters(const FViewInfo& View, FIntPoint ViewRectSize, const FEyeAdaptationParameters& EyeAdaptationParameters);
