Fix a warning about the polymorphic struct FGraphSchemaSearchTextDebugInfo not having a virtual destructor (warning is from FBPContextMenuWeightDebugInfo, a child of FGraphSchemaSearchTextDebugInfo)

#rnx

diff --git a/Engine/Source/Runtime/Engine/Classes/EdGraph/EdGraphSchema.h b/Engine/Source/Runtime/Engine/Classes/EdGraph/EdGraphSchema.h
--- a/Engine/Source/Runtime/Engine/Classes/EdGraph/EdGraphSchema.h
+++ b/Engine/Source/Runtime/Engine/Classes/EdGraph/EdGraphSchema.h
@@ -693,6 +693,8 @@ struct FGraphSchemaSearchTextWeightInfo
 // Helper struct storing the breakdown of the weights assigned to the search text
 struct FGraphSchemaSearchTextDebugInfo
 {
+	virtual ~FGraphSchemaSearchTextDebugInfo() = default;
+	
 	float TotalWeight = 0.0f;			// Overall weight
 
 	float NodeTitleWeight = 0.0f;		// Weight for the node's title
