Improve and make readable Audio stream cache debug display

diff --git a/Engine/Source/Runtime/Engine/Private/AudioStreamingCache.cpp b/Engine/Source/Runtime/Engine/Private/AudioStreamingCache.cpp
--- a/Engine/Source/Runtime/Engine/Private/AudioStreamingCache.cpp
+++ b/Engine/Source/Runtime/Engine/Private/AudioStreamingCache.cpp
@@ -13,6 +13,8 @@ AudioStreaming.cpp: Implementation of audio streaming classes.
 #include "Engine/Engine.h"
 #include "Stats/StatsTrace.h"
 #include "ProfilingDebugging/IoStoreTrace.h"
+#include "Audio/AudioDebug.h"
+#include "Audio/AudioDebugDraw.h"
 
 DEFINE_LOG_CATEGORY(LogAudioStreamCaching);
 
@@ -283,6 +285,7 @@ void FAudioChunkCache::FCacheElement::UpdateDebugInfoLoadingBehavior()
 	{
 		DebugInfo.LoadingBehavior = ESoundWaveLoadingBehavior::Uninitialized;
 		DebugInfo.bLoadingBehaviorExternallyOverriden = false;
+		DebugInfo.bWaveDestroyed = true;
 	}
 }
 #endif // #if DEBUG_STREAM_CACHE
@@ -842,8 +845,12 @@ TArrayView<uint8> FAudioChunkCache::GetChunk(const FChunkKey& InKey, const TShar
 			{
 				int32 ChunkAudioDataSize = Chunk->AudioDataSize;
 #if DEBUG_STREAM_CACHE
-				FoundElement->DebugInfo.NumTotalChunks = FoundElement->GetNumChunks() - 1;
-				FoundElement->DebugInfo.TimeLoadStarted = FPlatformTime::Seconds();
+				const float SampleRate = InSoundWavePtr->GetSampleRate();
+				const int32 NumChunks = FoundElement->GetNumChunks();
+				const uint32 OffsetInFrames = Chunk->SeekOffsetInAudioFrames;
+				FoundElement->DebugInfo.OffsetInSeconds = OffsetInFrames != INDEX_NONE ? OffsetInFrames / SampleRate : -1.f;
+				FoundElement->DebugInfo.NumTotalChunks = NumChunks - 1; //(exclude chunk 0)
+				FoundElement->DebugInfo.TimeLoadStartedCycles = FPlatformTime::Cycles64();
 #endif
 				MemoryCounterBytes -= FoundElement->ChunkDataSize;
 
@@ -869,8 +876,7 @@ TArrayView<uint8> FAudioChunkCache::GetChunk(const FChunkKey& InKey, const TShar
 				FoundElement->ChunkDataSize = ChunkAudioDataSize;
 				FoundElement->bIsLoaded = true;
 #if DEBUG_STREAM_CACHE
-				FoundElement->DebugInfo.TimeToLoad = (FPlatformTime::Seconds() - FoundElement->DebugInfo.TimeLoadStarted) * 1000.0f;
-
+				FoundElement->DebugInfo.TimeToLoadMs = FPlatformTime::ToMilliseconds64(FPlatformTime::Cycles64() - FoundElement->DebugInfo.TimeLoadStartedCycles);
 #endif
 				// If this value is ever negative, it means that we're decrementing more than we're incrementing:
 				if (ensureMsgf(FoundElement->NumConsumers.GetValue() >= 0, TEXT("NumConsumers was negative for FoundElement. Reseting to 1")))
@@ -1507,6 +1513,18 @@ FAudioChunkCache::FCacheElement* FAudioChunkCache::InsertChunk(const FChunkKey&
 		CacheElement->Key = InKey;
 		CacheElement->SoundWaveWeakPtr = InSoundWavePtr;
 		TouchElement(CacheElement);
+#if DEBUG_STREAM_CACHE
+		if (InSoundWavePtr.IsValid())
+		{
+			CacheElement->DebugInfo.Format = InSoundWavePtr->GetRuntimeFormat();
+			const FStreamedAudioChunk& Chunk=InSoundWavePtr->GetChunk(InKey.ChunkIndex);
+			const float SampleRate = InSoundWavePtr->GetSampleRate();
+			const int32 NumChunks = InSoundWavePtr->GetNumChunks();
+			const uint32 OffsetInFrames = Chunk.SeekOffsetInAudioFrames;
+			CacheElement->DebugInfo.OffsetInSeconds = OffsetInFrames != INDEX_NONE ? OffsetInFrames / SampleRate : -1.f;
+			CacheElement->DebugInfo.NumTotalChunks = NumChunks - 1; //(exclude chunk 0)
+		}
+#endif //DEBUG_STREAM_CACHE		
 
 		// If we've got multiple chunks, we can not cache the least recent chunk
 		// without worrying about a circular dependency.
@@ -1722,7 +1740,7 @@ void FAudioChunkCache::KickOffAsyncLoad(FCacheElement* CacheElement, const FChun
 		}
 
 #if DEBUG_STREAM_CACHE
-		CacheElement->DebugInfo.TimeLoadStarted = FPlatformTime::Cycles64();
+		CacheElement->DebugInfo.TimeLoadStartedCycles = FPlatformTime::Cycles64();
 #endif
 
 		TFunction<void(bool)> OnLoadComplete = [OnLoadCompleted, CallbackThread, CacheElement, InKey, ChunkDataSize](bool bRequestFailed)
@@ -1750,7 +1768,8 @@ void FAudioChunkCache::KickOffAsyncLoad(FCacheElement* CacheElement, const FChun
 			}
 			
 #if DEBUG_STREAM_CACHE
-			CacheElement->DebugInfo.TimeToLoad = FPlatformTime::ToMilliseconds64(FPlatformTime::Cycles64() - CacheElement->DebugInfo.TimeLoadStarted);
+			const float TimeToLoad = FPlatformTime::ToMilliseconds64(FPlatformTime::Cycles64() - CacheElement->DebugInfo.TimeLoadStartedCycles);		
+			CacheElement->DebugInfo.TimeToLoadMs = TimeToLoad;
 #endif
 			EAudioChunkLoadResult ChunkLoadResult = bRequestFailed ? EAudioChunkLoadResult::Interrupted : EAudioChunkLoadResult::Completed;
 			ExecuteOnLoadCompleteCallback(ChunkLoadResult, OnLoadCompleted, CallbackThread);
@@ -1809,7 +1828,8 @@ void FAudioChunkCache::KickOffAsyncLoad(FCacheElement* CacheElement, const FChun
 			}
 			
 #if DEBUG_STREAM_CACHE
-			CacheElement->DebugInfo.TimeToLoad = FPlatformTime::ToMilliseconds64(FPlatformTime::Cycles64() - CacheElement->DebugInfo.TimeLoadStarted);
+			const float TimeToLoad = FPlatformTime::ToMilliseconds64(FPlatformTime::Cycles64() - CacheElement->DebugInfo.TimeLoadStartedCycles);		
+			CacheElement->DebugInfo.TimeToLoadMs = TimeToLoad;			
 #endif
 			const EAudioChunkLoadResult LoadResult = bWasCancelled ? EAudioChunkLoadResult::Interrupted : EAudioChunkLoadResult::Completed;
 			ExecuteOnLoadCompleteCallback(LoadResult, OnLoadCompleted, CallbackThread);
@@ -1818,7 +1838,7 @@ void FAudioChunkCache::KickOffAsyncLoad(FCacheElement* CacheElement, const FChun
 		};
 
 #if DEBUG_STREAM_CACHE
-		CacheElement->DebugInfo.TimeLoadStarted = FPlatformTime::Cycles64();
+		CacheElement->DebugInfo.TimeLoadStartedCycles = FPlatformTime::Cycles64();
 #endif
 
 		CacheElement->ReadRequest = nullptr;
@@ -1833,7 +1853,7 @@ void FAudioChunkCache::KickOffAsyncLoad(FCacheElement* CacheElement, const FChun
 
 #if DEBUG_STREAM_CACHE
 			UE_LOG(LogAudioStreamCaching, Verbose, TEXT("Loading Inlined Chunk: %s, %d, TimeToLoad=%2.2fms"), *InKey.SoundWaveName.ToString(),
-				InKey.ChunkIndex, FPlatformTime::ToMilliseconds64(FPlatformTime::Cycles64() - CacheElement->DebugInfo.TimeLoadStarted));
+				InKey.ChunkIndex, FPlatformTime::ToMilliseconds64(FPlatformTime::Cycles64() - CacheElement->DebugInfo.TimeLoadStartedCycles));
 
 			CacheElement->DebugInfo.bWasLoadedFromInlineChunk = true;
 			UE_LOG(LogAudioStreamCaching, VeryVerbose, TEXT("COPY+DISCARD %s - Bulk=0x%p"), *InKey.SoundWaveName.ToString(), &Chunk->BulkData);
@@ -2107,7 +2127,23 @@ TPair<int, int> FAudioChunkCache::DebugDisplayLegacy(UWorld* World, FViewport* V
 
 	float ColorLerpAmount = 0.0f;
 	const float ColorLerpStep = 0.04f;
-
+	
+	Audio::FTabularRenderHelper Helper;
+	Helper.AddCol(TEXT("Name"), 64);
+	Helper.AddCol(TEXT("Chunk"), 10);
+	Helper.AddCol(TEXT("Size(kb)"), 10);
+	Helper.AddCol(TEXT("Loadtime(ms)"),20);
+	Helper.AddCol(TEXT("Load-Behavior"), 20);
+	Helper.AddCol(TEXT("Format"), 8);
+	Helper.AddCol(TEXT("Refs"), 6);
+	Helper.AddCol(TEXT("Reqs"), 6);
+	Helper.AddCol(TEXT("Offset(secs)"), 20);
+	Helper.AddCol(TEXT("Inlined"),10);
+	Helper.AddCol(TEXT("Wasinline"),10);
+	Helper.AddCol(TEXT("Cachemiss"),10);
+	Helper.AddCol(TEXT("Stale"), 10);
+	Helper.AddCol(TEXT("Destroyed"),10);
+	
 	// More detailed info about individual chunks here:
 	while (CurrentElement != nullptr)
 	{
@@ -2119,23 +2155,30 @@ TPair<int, int> FAudioChunkCache::DebugDisplayLegacy(UWorld* World, FViewport* V
 
 		int32 NumTotalChunks = -1;
 		int32 NumTimesTouched = -1;
-		double TimeToLoad = -1.0;
-		float AveragePlaceInCache = -1.0f;
+		float TimeToLoad = -1.0;
 		ESoundWaveLoadingBehavior LoadingBehavior = ESoundWaveLoadingBehavior::Uninitialized;
 		bool bLoadingBehaviorExternallyOverriden = false;
 		bool bWasCacheMiss = false;
 		bool bIsStaleChunk = false;
 		bool bWasLoadedInlined = false;
+		bool bWasInlinedButUnloaded = false;
+		bool bWaveDestroyed = false;
+		float OffsetInSecs = 0;
+		FName Format;
 
 #if DEBUG_STREAM_CACHE
 		NumTotalChunks = CurrentElement->DebugInfo.NumTotalChunks;
 		NumTimesTouched = CurrentElement->DebugInfo.NumTimesTouched;
-		TimeToLoad = CurrentElement->DebugInfo.TimeToLoad;
-		AveragePlaceInCache = CurrentElement->DebugInfo.AverageLocationInCacheWhenNeeded;
+		TimeToLoad = CurrentElement->DebugInfo.TimeToLoadMs;
 		LoadingBehavior = CurrentElement->DebugInfo.LoadingBehavior;
 		bLoadingBehaviorExternallyOverriden = CurrentElement->DebugInfo.bLoadingBehaviorExternallyOverriden;
 		bWasCacheMiss = CurrentElement->DebugInfo.bWasCacheMiss;
 		bWasLoadedInlined = CurrentElement->DebugInfo.bWasLoadedFromInlineChunk;
+		bWasInlinedButUnloaded = CurrentElement->DebugInfo.bWasInlinedButUnloaded;
+		bWaveDestroyed = CurrentElement->DebugInfo.bWaveDestroyed;
+		OffsetInSecs = CurrentElement->DebugInfo.OffsetInSeconds;
+		Format = CurrentElement->DebugInfo.Format;
+		TimeToLoad = CurrentElement->DebugInfo.TimeToLoadMs;
 #endif
 
 #if WITH_EDITOR
@@ -2145,27 +2188,11 @@ TPair<int, int> FAudioChunkCache::DebugDisplayLegacy(UWorld* World, FViewport* V
 
 		const bool bWasTrimmed = CurrentElement->ChunkDataSize == 0;
 
-		FString ElementInfo = *FString::Printf(TEXT("%4i. Size: %6.2f KB   Chunk: %d of %d   Request Count: %d    Average Index: %6.2f  Number of Handles Retaining Chunk: %d     Chunk Load Time(in ms): %6.4fms      Loading Behavior: %s%s      Name: %s Notes: %s %s %s"),
-			Index,
-			CurrentElement->ChunkDataSize / 1024.0f,
-			CurrentElement->Key.ChunkIndex,
-			NumTotalChunks,
-			NumTimesTouched,
-			AveragePlaceInCache,
-			CurrentElement->NumConsumers.GetValue(),
-			TimeToLoad,
-			EnumToString(LoadingBehavior),
-			bLoadingBehaviorExternallyOverriden ? TEXT("*") : TEXT(""),
-			bWasTrimmed ? TEXT("TRIMMED CHUNK") : *CurrentElement->Key.SoundWaveName.ToString(),
-			bWasCacheMiss ? TEXT("(Cache Miss!)") : TEXT(""),
-			bIsStaleChunk ? TEXT("(Stale Chunk)") : TEXT(""),
-			bWasLoadedInlined ? TEXT("(Inlined)") : TEXT("")
-		);
-
 		// Since there's a lot of info here,
 		// Subtly fading the chunk info to gray seems to help as a visual indicator of how far down on the list things are.
 		ColorLerpAmount = FMath::Min(ColorLerpAmount + ColorLerpStep, 1.0f);
 		FLinearColor TextColor;
+		
 		if (bIsStaleChunk)
 		{
 			TextColor = FLinearColor::LerpUsingHSV(StaleChunkColor, FLinearColor::Gray, ColorLerpAmount);
@@ -2193,13 +2220,34 @@ TPair<int, int> FAudioChunkCache::DebugDisplayLegacy(UWorld* World, FViewport* V
 			TextColor = TrimmedChunkColor;
 		}
 
-		Canvas->DrawShadowedString(X, Y, *ElementInfo, UEngine::GetSmallFont(), TextColor);
-		Y += 12;
+		// Scale by load time.
+		const float T = FMath::GetMappedRangeValueClamped( FVector2D(0.f, 100.f), {0.f,1.f}, TimeToLoad);
+		FLinearColor LoadColor = FLinearColor::LerpUsingHSV(FColor::Green, FColor::Red, T);
+
+		using namespace Audio;
+		Helper.AddRow({
+			{CurrentElement->Key.SoundWaveName.ToString(), TextColor},											// name
+			{FString::Printf(TEXT("%d/%d"), CurrentElement->Key.ChunkIndex, NumTotalChunks), TextColor}, // chunk/count
+			{FString::Printf(TEXT("%u"),  CurrentElement->ChunkDataSize >> 10),TextColor},			    // size (kb)
+			{FString::Printf(TEXT("%.2f"), TimeToLoad), LoadColor},								        // load time (ms)
+			{ShortEnumString(EnumToString(LoadingBehavior)), TextColor },									    // load behavior
+			{Format.ToString(), TextColor },																    // format
+			{FString::Printf(TEXT("%d"), CurrentElement->NumConsumers.GetValue()), TextColor},		    // hard refs
+			{FString::Printf(TEXT("%d"), NumTimesTouched), TextColor},								    // requests
+			{FString::Printf(TEXT("%.2f"), OffsetInSecs), TextColor},								    // offset(secs)
+			{LexToString(bWasLoadedInlined), TextColor},													    // inlined?
+			{LexToString(bWasCacheMiss), TextColor},														    // cache miss?
+			{LexToString(bWasInlinedButUnloaded), TextColor},												    // inlinedButUnloaded?
+			{LexToString(bIsStaleChunk), TextColor},														    // stale?
+			{LexToString(bWaveDestroyed), TextColor},														    // destroyed?
+		});
 
 		CurrentElement = CurrentElement->LessRecentElement;
 		Index++;
 	}
 
+	Helper.Draw(Canvas,X,Y);
+	
 	return TPair<int, int>(X - InitialX, Y - InitialY);
 }
 
@@ -2274,16 +2322,20 @@ FString FAudioChunkCache::DebugPrint()
 		bool bWasCacheMiss = false;
 		bool bIsStaleChunk = false;
 		bool bWasLoadedInlined = false;
+		bool bWasInlinedButUnloaded = false;
+		bool bWaveDestroyed = false;
 
 #if DEBUG_STREAM_CACHE
 		NumTotalChunks = CurrentElement->DebugInfo.NumTotalChunks;
 		NumTimesTouched = CurrentElement->DebugInfo.NumTimesTouched;
-		TimeToLoad = CurrentElement->DebugInfo.TimeToLoad;
+		TimeToLoad = CurrentElement->DebugInfo.TimeToLoadMs;
 		AveragePlaceInCache = CurrentElement->DebugInfo.AverageLocationInCacheWhenNeeded;
 		LoadingBehavior = CurrentElement->DebugInfo.LoadingBehavior;
 		bLoadingBehaviorExternallyOverriden = CurrentElement->DebugInfo.bLoadingBehaviorExternallyOverriden;
 		bWasCacheMiss = CurrentElement->DebugInfo.bWasCacheMiss;
 		bWasLoadedInlined = CurrentElement->DebugInfo.bWasLoadedFromInlineChunk;
+		bWasInlinedButUnloaded = CurrentElement->DebugInfo.bWasInlinedButUnloaded;
+		bWaveDestroyed = CurrentElement->DebugInfo.bWaveDestroyed;
 #endif
 
 #if WITH_EDITOR
@@ -2293,7 +2345,8 @@ FString FAudioChunkCache::DebugPrint()
 
 		const bool bWasTrimmed = CurrentElement->ChunkDataSize == 0;
 
-		FString ElementInfo = *FString::Printf(TEXT("%4i.\t, %6.2f\t, %d of %d\t, %d\t, %6.2f\t, %d\t,  %6.4f\t, %s\t, %s%s, %s %s %s %s"),
+		using namespace Audio;
+		FString ElementInfo = *FString::Printf(TEXT("%4i.\t, %6.2f\t, %d of %d\t, %d\t, %6.2f\t, %d\t,  %6.4f\t, %s\t, %s%s, %s %s %s %s %s"),
 			Index,
 			CurrentElement->ChunkDataSize / 1024.0f,
 			CurrentElement->Key.ChunkIndex,
@@ -2303,12 +2356,14 @@ FString FAudioChunkCache::DebugPrint()
 			CurrentElement->NumConsumers.GetValue(),
 			TimeToLoad,
 			bWasTrimmed ? TEXT("TRIMMED CHUNK") : *CurrentElement->Key.SoundWaveName.ToString(),
-			EnumToString(LoadingBehavior),
+			ShortEnumString(EnumToString(LoadingBehavior)),
 			bLoadingBehaviorExternallyOverriden ? TEXT("*") : TEXT(""),
 			bWasCacheMiss ? TEXT("(Cache Miss!)") : TEXT(""),
 			bIsStaleChunk ? TEXT("(Stale Chunk)") : TEXT(""),
 			CurrentElement->IsLoadInProgress() ? TEXT("(Loading In Progress)") : TEXT(""),
-			bWasLoadedInlined ? TEXT("(Inlined)") : TEXT("(Disk)")
+			bWasInlinedButUnloaded ? TEXT("(Dumped!)") : bWasLoadedInlined ? TEXT("(Inlined)") : TEXT("(Disk)"),
+			bWaveDestroyed ? TEXT("(GCd!)") : TEXT("(Alive)")
+
 		);
 
 		if (!bWasTrimmed)
@@ -2763,7 +2818,7 @@ TPair<int, int> FAudioChunkCache::DebugVisualDisplay(UWorld* World, FViewport* V
 #if DEBUG_STREAM_CACHE
 		NumTotalChunks = CurrentElement->DebugInfo.NumTotalChunks;
 		NumTimesTouched = CurrentElement->DebugInfo.NumTimesTouched;
-		TimeToLoad = CurrentElement->DebugInfo.TimeToLoad;
+		TimeToLoad = CurrentElement->DebugInfo.TimeToLoadMs;
 		AveragePlaceInCache = CurrentElement->DebugInfo.AverageLocationInCacheWhenNeeded;
 		LoadingBehavior = CurrentElement->DebugInfo.LoadingBehavior;
 		bLoadingBehaviorExternallyOverriden = CurrentElement->DebugInfo.bLoadingBehaviorExternallyOverriden;
diff --git a/Engine/Source/Runtime/Engine/Public/Audio/AudioDebug.h b/Engine/Source/Runtime/Engine/Public/Audio/AudioDebug.h
--- a/Engine/Source/Runtime/Engine/Public/Audio/AudioDebug.h
+++ b/Engine/Source/Runtime/Engine/Public/Audio/AudioDebug.h
@@ -6,6 +6,19 @@
 #include "AudioDefines.h"		// For ENABLE_AUDIO_DEBUG
 #include "Audio.h"
 
+namespace Audio
+{
+	// enum shortener for logging.
+	static const TCHAR* ShortEnumString(const TCHAR* InString)
+	{
+		if (const TCHAR* DoubleColon = FCString::Strstr(InString, TEXT("::")))
+		{
+			return DoubleColon + 2;
+		}
+		return InString;
+	}
+}
+
 #if ENABLE_AUDIO_DEBUG
 
  // Forward Declarations
diff --git a/Engine/Source/Runtime/Engine/Public/Audio/AudioDebugDraw.h b/Engine/Source/Runtime/Engine/Public/Audio/AudioDebugDraw.h
new file mode 100644
--- /dev/null
+++ b/Engine/Source/Runtime/Engine/Public/Audio/AudioDebugDraw.h
@@ -0,0 +1,90 @@
+// Copyright Epic Games, Inc. All Rights Reserved.
+
+#pragma once
+
+#include "CanvasTypes.h"
+#include "Engine/Font.h"
+#include "Engine/Engine.h"
+
+namespace Audio
+{
+	class FTabularRenderHelper
+	{
+	public:
+		FTabularRenderHelper() = default;
+		
+		void Draw(FCanvas* Canvas, int32 X, int32 Y) const
+		{
+			DrawHeaders(Canvas,X,Y, HeaderColor);
+			Y+=Font->GetMaxCharHeight();
+			DrawRows(Canvas,X,Y, RowColor);
+		}
+
+		struct FItem { FString Text; FLinearColor Color=FColor::Green; };
+		void AddRow(const TArray<FItem>& Items)
+		{
+			for(int32 i = 0; i < Items.Num(); ++i)
+			{
+				if (Cols.IsValidIndex(i))
+				{
+					Cols[i].Items.Add(Items[i].Text);
+					Cols[i].Colors.Add(Items[i].Color);
+				}
+			}
+		}
+		void AddCol(FString Text, int32 WidthInChars=0)
+		{		
+			WidthInChars = !WidthInChars ? Text.Len() : WidthInChars;
+			Text=Text.LeftPad(WidthInChars);
+			int32 X,Y;
+			StringSize(Font, X, Y, *Text);
+			Cols.Add({Text, X});
+		}
+	private:
+		struct FCol
+		{
+			FString Header;
+			int32 Width = 32;
+			TArray<FString> Items;
+			TArray<FLinearColor> Colors;
+		};
+		TArray<FCol> Cols;
+		FColor HeaderColor = FColor::White;
+		FColor RowColor = FColor::Green;
+		const UFont* Font = UEngine::GetTinyFont();
+
+		void RightJustify(FCanvas* Canvas, const int32 X, const int32 Y, const int32 InterColumnOffset, TCHAR const* Text, FLinearColor const& Color) const
+		{
+			int32 SizeX, SizeY;
+			StringSize(Font, SizeX, SizeY, Text);
+			Canvas->DrawShadowedString(X + InterColumnOffset - SizeX, Y, Text, Font, Color);
+		}
+		void DrawHeaders(FCanvas* Canvas, int32 X, int32 Y,const FLinearColor& Color) const
+		{
+			for (const FCol& i : Cols)
+			{
+				RightJustify(Canvas,X,Y,i.Width,*i.Header, Color);
+				X += i.Width;
+			}
+		}
+		void DrawRows(FCanvas* Canvas, const int32 InX, const int32 InY,const FLinearColor& Color) const
+		{
+			int32 Line=0;
+			const int32 MaxY = Canvas->GetParentCanvasSize().Y;
+			for (int32 Y=InY; Y < MaxY; Y+=Font->GetMaxCharHeight(), ++Line)
+			{
+				// Draw row.
+				int32 X=InX;
+				for (const FCol& i : Cols)
+				{
+					const FLinearColor C = i.Colors.IsValidIndex(Line) ? i.Colors[Line] : Color;
+					if (i.Items.IsValidIndex(Line))
+					{
+						RightJustify(Canvas,X,Y,i.Width,*i.Items[Line], C);
+					}
+					X += i.Width;
+				}
+			}
+		}
+	};
+}
\ No newline at end of file
diff --git a/Engine/Source/Runtime/Engine/Public/AudioStreamingCache.h b/Engine/Source/Runtime/Engine/Public/AudioStreamingCache.h
--- a/Engine/Source/Runtime/Engine/Public/AudioStreamingCache.h
+++ b/Engine/Source/Runtime/Engine/Public/AudioStreamingCache.h
@@ -228,9 +228,9 @@ private:
 		// Number of times this chunk was requested during its time in the cache.
 		int32 NumTimesTouched;
 
-		uint64 TimeLoadStarted;
+		uint64 TimeLoadStartedCycles = 0;
 		// Amount of time spent loading the audio file.
-		float TimeToLoad;
+		float TimeToLoadMs;
 
 		// This is a cumulative moving average of a chunks location before it was 
 		float AverageLocationInCacheWhenNeeded;
@@ -239,16 +239,21 @@ private:
 		ESoundWaveLoadingBehavior LoadingBehavior;
 		bool bLoadingBehaviorExternallyOverriden;
 
+		// cached format.
+		FName Format;
+		float OffsetInSeconds=0.f;
+
 		// if true, 
 		bool bWasCacheMiss;
 		bool bWasLoadedFromInlineChunk = false;
 		bool bWasInlinedButUnloaded = false;
+		bool bWaveDestroyed = false;
 
 		FCacheElementDebugInfo()
 			: NumTotalChunks(0)
 			, NumTimesTouched(0)
-			, TimeLoadStarted(0)
-			, TimeToLoad(0.0)
+			, TimeLoadStartedCycles(0)
+			, TimeToLoadMs(0.0)
 			, AverageLocationInCacheWhenNeeded(0.0f)
 			, LoadingBehavior(ESoundWaveLoadingBehavior::Uninitialized)
 			, bWasCacheMiss(false)
@@ -259,13 +264,15 @@ private:
 		{
 			NumTotalChunks = 0;
 			NumTimesTouched = 0;
-			TimeLoadStarted = 0;
-			TimeToLoad = 0.0f;
+			TimeLoadStartedCycles = 0;
+			TimeToLoadMs = 0.0f;
 			LoadingBehavior = ESoundWaveLoadingBehavior::Uninitialized;
 			bWasCacheMiss = false;
 			AverageLocationInCacheWhenNeeded = 0.0f;
 			bWasLoadedFromInlineChunk = false;
 			bWasInlinedButUnloaded = false;
+			bWaveDestroyed = false;
+			Format = {};
 		}
 
 	};
