Fixing issue where changes to the game user settings config were being discarded before they could be saved to disk.

diff --git a/Engine/Source/Runtime/Engine/Private/GameUserSettings.cpp b/Engine/Source/Runtime/Engine/Private/GameUserSettings.cpp
--- a/Engine/Source/Runtime/Engine/Private/GameUserSettings.cpp
+++ b/Engine/Source/Runtime/Engine/Private/GameUserSettings.cpp
@@ -835,7 +835,7 @@ void UGameUserSettings::ReloadWindowPositions()
 	FProperty* WindowPositionsProperty = GetClass()->FindPropertyByName(GET_MEMBER_NAME_CHECKED(UGameUserSettings, WindowPositions));
 	if (WindowPositionsProperty != nullptr)
 	{
-		LoadConfigIni(true);
+		LoadConfigIni(false);
 		ReloadConfig(GetClass(), *GGameUserSettingsIni, UE::LCPF_None, WindowPositionsProperty);
 	}
 }
