Tentative fix for occasional crash occurring when changing scalability settings, possibly due uniform expression data not being recached.

Changing the scalability settings now always triggering updating materials uniform expression data when some setting has changed.

diff --git a/Engine/Source/Runtime/Engine/Private/UnrealEngine.cpp b/Engine/Source/Runtime/Engine/Private/UnrealEngine.cpp
--- a/Engine/Source/Runtime/Engine/Private/UnrealEngine.cpp
+++ b/Engine/Source/Runtime/Engine/Private/UnrealEngine.cpp
@@ -886,17 +886,14 @@ void ScalabilityCVarsSinkCallback()
 		return;
 	}
 
-	if (LocalScalabilityCVars.MaterialQualityLevel != GCachedScalabilityCVars.MaterialQualityLevel)
-	{
-		// Need to flush any in-flight uniform expression cache operations, when material quality changes.  The material quality level affects
-		// the result returned by calls to FMaterialRenderProxy::GetMaterialNoFallback, and in-flight operations will have been using a uniform
-		// buffer layout from the old version of the material, leading to an assert.
-		ENQUEUE_RENDER_COMMAND(UpdateDeferredCachedUniformExpressions)(
-			[](FRHICommandListImmediate& RHICmdList)
-			{
-				FMaterialRenderProxy::UpdateDeferredCachedUniformExpressions();
-			});
-	}
+	// Need to flush any in-flight uniform expression cache operations, when material quality changes.  The material quality level affects
+	// the result returned by calls to FMaterialRenderProxy::GetMaterialNoFallback, and in-flight operations will have been using a uniform
+	// buffer layout from the old version of the material, leading to an assert.
+	ENQUEUE_RENDER_COMMAND(UpdateDeferredCachedUniformExpressions)(
+		[](FRHICommandListImmediate& RHICmdList)
+		{
+			FMaterialRenderProxy::UpdateDeferredCachedUniformExpressions();
+		});
 
 	FlushRenderingCommands();
 
