[Shaders] Include current PID into job object name to fix warning when subprocesses instantate same object identifier.

#rnx

diff --git a/Engine/Source/Runtime/Engine/Private/ShaderCompiler/ShaderCompilerMemoryLimit.cpp b/Engine/Source/Runtime/Engine/Private/ShaderCompiler/ShaderCompilerMemoryLimit.cpp
--- a/Engine/Source/Runtime/Engine/Private/ShaderCompiler/ShaderCompilerMemoryLimit.cpp
+++ b/Engine/Source/Runtime/Engine/Private/ShaderCompiler/ShaderCompilerMemoryLimit.cpp
@@ -11,7 +11,7 @@
 #if PLATFORM_WINDOWS
 
 FWindowsResourceRestrictedJobObject::FWindowsResourceRestrictedJobObject(const FString& InJobName, int32 InitialJobMemoryLimitMiB) :
-	JobName(InJobName)
+	JobName(FString::Printf(TEXT("%s-%u"), *InJobName, FPlatformProcess::GetCurrentProcessId()))
 {
 	JobObject = ::CreateJobObject(nullptr, *JobName);
 	if (JobObject == nullptr)
