[AutomationTests] Add TestFalseExpr() macro, use it and TestTrueExpr() to clean up AutoRTFMEngineTests

Remove custom, duplicated TEST_CHECK_[TRUE|FALSE] macros with these.

diff --git a/Engine/Source/Runtime/Core/Public/Misc/AutomationTest.h b/Engine/Source/Runtime/Core/Public/Misc/AutomationTest.h
--- a/Engine/Source/Runtime/Core/Public/Misc/AutomationTest.h
+++ b/Engine/Source/Runtime/Core/Public/Misc/AutomationTest.h
@@ -2266,6 +2266,9 @@ public:
 		return TestFalse(*What, Value);
 	}
 
+	/** Macro version of above, uses the passed in expression as the description as well */
+	#define TestFalseExpr(Expression) TestFalse(TEXT(#Expression), Expression)
+
 	/**
 	 * Logs an error if the given object tests true when calling its IsValid member.
 	 *
diff --git a/Engine/Source/Runtime/Engine/Private/Tests/AutoRTFM/AutoRTFMActorComponentTests.cpp b/Engine/Source/Runtime/Engine/Private/Tests/AutoRTFM/AutoRTFMActorComponentTests.cpp
--- a/Engine/Source/Runtime/Engine/Private/Tests/AutoRTFM/AutoRTFMActorComponentTests.cpp
+++ b/Engine/Source/Runtime/Engine/Private/Tests/AutoRTFM/AutoRTFMActorComponentTests.cpp
@@ -140,28 +140,6 @@ private:
 	const int32 TestLine;
 };
 
-// Tests that EXPR evaluates to true. If EXPR evaluates to false, then an error
-// is raised and the function returns.
-#define TEST_CHECK_TRUE(EXPR) do \
-{ \
-	if (!(EXPR)) \
-	{ \
-		Fail(TEXT("'") TEXT(#EXPR) TEXT("' was not true"), TEXT(__FILE__), __LINE__); \
-		return; \
-	} \
-} while(false)
-
-// Tests that EXPR evaluates to false. If EXPR evaluates to true, then an error
-// is raised and the function returns.
-#define TEST_CHECK_FALSE(EXPR) do \
-{ \
-	if ((EXPR)) \
-	{ \
-		Fail(TEXT("'") TEXT(#EXPR) TEXT("' was not false"), TEXT(__FILE__), __LINE__); \
-		return; \
-	} \
-} while(false)
-
 // General tests for calling RegisterComponent() and UnregisterComponent() in transactions.
 // See: SOL-6709
 AUTORTFM_ACTOR_COMPONENT_TEST(RegisterComponent_UnregisterComponent)
@@ -180,8 +158,8 @@ AUTORTFM_ACTOR_COMPONENT_TEST(RegisterComponent_UnregisterComponent)
 			}
 		});
 
-	TEST_CHECK_TRUE(AutoRTFM::ETransactionResult::AbortedByRequest == Result);
-	TEST_CHECK_FALSE(Component->IsRegistered());
+	TestTrueExpr(AutoRTFM::ETransactionResult::AbortedByRequest == Result);
+	TestFalseExpr(Component->IsRegistered());
 
 	bool bWasRegistered = false;
 
@@ -191,8 +169,8 @@ AUTORTFM_ACTOR_COMPONENT_TEST(RegisterComponent_UnregisterComponent)
 			bWasRegistered = Component->IsRegistered();
 		});
 
-	TEST_CHECK_TRUE(bWasRegistered);
-	TEST_CHECK_TRUE(Component->IsRegistered());
+	TestTrueExpr(bWasRegistered);
+	TestTrueExpr(Component->IsRegistered());
 
 	Result = AutoRTFM::Transact([&]
 		{
@@ -200,15 +178,15 @@ AUTORTFM_ACTOR_COMPONENT_TEST(RegisterComponent_UnregisterComponent)
 			AutoRTFM::AbortTransaction();
 		});
 
-	TEST_CHECK_TRUE(AutoRTFM::ETransactionResult::AbortedByRequest == Result);
-	TEST_CHECK_TRUE(Component->IsRegistered());
+	TestTrueExpr(AutoRTFM::ETransactionResult::AbortedByRequest == Result);
+	TestTrueExpr(Component->IsRegistered());
 
 	AutoRTFM::Commit([&]
 		{
 			Component->UnregisterComponent();
 		});
 
-	TEST_CHECK_FALSE(Component->IsRegistered());
+	TestFalseExpr(Component->IsRegistered());
 }
 
 // Test aborting a call to Component::RegisterComponentWithWorld().
@@ -226,15 +204,15 @@ AUTORTFM_ACTOR_COMPONENT_TEST(RegisterComponentWithWorld)
 			AutoRTFM::AbortTransaction();
 		});
 
-	TEST_CHECK_TRUE(AutoRTFM::ETransactionResult::AbortedByRequest == Result);
-	TEST_CHECK_FALSE(Component->IsRegistered());
+	TestTrueExpr(AutoRTFM::ETransactionResult::AbortedByRequest == Result);
+	TestFalseExpr(Component->IsRegistered());
 
 	AutoRTFM::Commit([&]
 		{
 			Component->RegisterComponentWithWorld(World);
 		});
 
-	TEST_CHECK_TRUE(Component->IsRegistered());
+	TestTrueExpr(Component->IsRegistered());
 }
 
 // Test aborting a call to Component::WeldTo().
@@ -249,7 +227,7 @@ AUTORTFM_ACTOR_COMPONENT_TEST(WeldTo)
 	Component->BodyInstance = FBodyInstance();
 	Component->BodyInstance.bSimulatePhysics = 1;
 	Component->BodyInstance.WeldParent = &SomeInstance;
-	TEST_CHECK_TRUE(Component->IsWelded());
+	TestTrueExpr(Component->IsWelded());
 
 	UAutoRTFMTestBodySetup* BodySetup = NewObject<UAutoRTFMTestBodySetup>();
 	BodySetup->AggGeom.SphereElems.Add(FKSphereElem(1.0f));
@@ -265,17 +243,17 @@ AUTORTFM_ACTOR_COMPONENT_TEST(WeldTo)
 			AutoRTFM::AbortTransaction();
 		});
 
-	TEST_CHECK_TRUE(AutoRTFM::ETransactionResult::AbortedByRequest == Result);
-	TEST_CHECK_TRUE(Component->IsWelded());
-	TEST_CHECK_TRUE(&SomeInstance == Component->BodyInstance.WeldParent);
+	TestTrueExpr(AutoRTFM::ETransactionResult::AbortedByRequest == Result);
+	TestTrueExpr(Component->IsWelded());
+	TestTrueExpr(&SomeInstance == Component->BodyInstance.WeldParent);
 
 	AutoRTFM::Commit([&]
 		{
 			Component->WeldTo(Parent0);
 		});
 
-	TEST_CHECK_FALSE(Component->IsWelded());
-	TEST_CHECK_TRUE(nullptr == Component->BodyInstance.WeldParent);
+	TestFalseExpr(Component->IsWelded());
+	TestTrueExpr(nullptr == Component->BodyInstance.WeldParent);
 
 	Result = AutoRTFM::Transact([&]
 		{
@@ -283,15 +261,15 @@ AUTORTFM_ACTOR_COMPONENT_TEST(WeldTo)
 			AutoRTFM::AbortTransaction();
 		});
 
-	TEST_CHECK_TRUE(AutoRTFM::ETransactionResult::AbortedByRequest == Result);
-	TEST_CHECK_FALSE(Component->IsWelded());
+	TestTrueExpr(AutoRTFM::ETransactionResult::AbortedByRequest == Result);
+	TestFalseExpr(Component->IsWelded());
 
 	AutoRTFM::Commit([&]
 		{
 			Component->WeldTo(Parent1);
 		});
 
-	TEST_CHECK_FALSE(Component->IsWelded());
+	TestFalseExpr(Component->IsWelded());
 
 	Result = AutoRTFM::Transact([&]
 		{
@@ -299,15 +277,15 @@ AUTORTFM_ACTOR_COMPONENT_TEST(WeldTo)
 			AutoRTFM::AbortTransaction();
 		});
 
-	TEST_CHECK_TRUE(AutoRTFM::ETransactionResult::AbortedByRequest == Result);
-	TEST_CHECK_FALSE(Component->IsWelded());
+	TestTrueExpr(AutoRTFM::ETransactionResult::AbortedByRequest == Result);
+	TestFalseExpr(Component->IsWelded());
 
 	AutoRTFM::Commit([&]
 		{
 			Component->UnWeldFromParent();
 		});
 
-	TEST_CHECK_FALSE(Component->IsWelded());
+	TestFalseExpr(Component->IsWelded());
 }
 
 // Test calling Component->UnregisterComponent() on a Component with an event
@@ -320,7 +298,7 @@ AUTORTFM_ACTOR_COMPONENT_TEST(FSparseDelegate)
 	Component->RegisterComponent();
 	Component->OnComponentPhysicsStateChanged.AddDynamic(Object, &UAutoRTFMTestObject::OnComponentPhysicsStateChanged);
 
-	TEST_CHECK_FALSE(Object->bHitOnComponentPhysicsStateChanged);
+	TestFalseExpr(Object->bHitOnComponentPhysicsStateChanged);
 
 	AutoRTFM::ETransactionResult Result = AutoRTFM::Transact([&]
 		{
@@ -328,15 +306,15 @@ AUTORTFM_ACTOR_COMPONENT_TEST(FSparseDelegate)
 			AutoRTFM::AbortTransaction();
 		});
 
-	TEST_CHECK_TRUE(AutoRTFM::ETransactionResult::AbortedByRequest == Result);
-	TEST_CHECK_FALSE(Object->bHitOnComponentPhysicsStateChanged);
+	TestTrueExpr(AutoRTFM::ETransactionResult::AbortedByRequest == Result);
+	TestFalseExpr(Object->bHitOnComponentPhysicsStateChanged);
 
 	AutoRTFM::Commit([&]
 		{
 			Component->UnregisterComponent();
 		});
 
-	TEST_CHECK_TRUE(Object->bHitOnComponentPhysicsStateChanged);
+	TestTrueExpr(Object->bHitOnComponentPhysicsStateChanged);
 }
 
 AUTORTFM_ACTOR_COMPONENT_TEST(ChildActor)
@@ -357,16 +335,16 @@ AUTORTFM_ACTOR_COMPONENT_TEST(ChildActor)
 				AutoRTFM::AbortTransaction();
 			});
 
-		TEST_CHECK_TRUE(AutoRTFM::ETransactionResult::AbortedByRequest == Result);
-		TEST_CHECK_TRUE(nullptr != ChildActorComponent->GetChildActor());
+		TestTrueExpr(AutoRTFM::ETransactionResult::AbortedByRequest == Result);
+		TestTrueExpr(nullptr != ChildActorComponent->GetChildActor());
 
 		Result = AutoRTFM::Transact([&]
 			{
 				ChildActorComponent->DestroyChildActor();
 			});
 
-		TEST_CHECK_TRUE(AutoRTFM::ETransactionResult::Committed == Result);
-		TEST_CHECK_TRUE(nullptr == ChildActorComponent->GetChildActor());
+		TestTrueExpr(AutoRTFM::ETransactionResult::Committed == Result);
+		TestTrueExpr(nullptr == ChildActorComponent->GetChildActor());
 	}
 
 	AutoRTFM::ETransactionResult Result = AutoRTFM::Transact([&]
@@ -375,16 +353,16 @@ AUTORTFM_ACTOR_COMPONENT_TEST(ChildActor)
 			AutoRTFM::AbortTransaction();
 		});
 
-	TEST_CHECK_TRUE(AutoRTFM::ETransactionResult::AbortedByRequest == Result);
-	TEST_CHECK_TRUE(nullptr == ChildActorComponent->GetChildActor());
+	TestTrueExpr(AutoRTFM::ETransactionResult::AbortedByRequest == Result);
+	TestTrueExpr(nullptr == ChildActorComponent->GetChildActor());
 
 	Result = AutoRTFM::Transact([&]
 		{
 			ChildActorComponent->CreateChildActor();
 		});
 
-	TEST_CHECK_TRUE(AutoRTFM::ETransactionResult::Committed == Result);
-	TEST_CHECK_TRUE(nullptr != ChildActorComponent->GetChildActor());
+	TestTrueExpr(AutoRTFM::ETransactionResult::Committed == Result);
+	TestTrueExpr(nullptr != ChildActorComponent->GetChildActor());
 
 	Result = AutoRTFM::Transact([&]
 		{
@@ -392,16 +370,16 @@ AUTORTFM_ACTOR_COMPONENT_TEST(ChildActor)
 			AutoRTFM::AbortTransaction();
 		});
 
-	TEST_CHECK_TRUE(AutoRTFM::ETransactionResult::AbortedByRequest == Result);
-	TEST_CHECK_TRUE(nullptr != ChildActorComponent->GetChildActor());
+	TestTrueExpr(AutoRTFM::ETransactionResult::AbortedByRequest == Result);
+	TestTrueExpr(nullptr != ChildActorComponent->GetChildActor());
 
 	Result = AutoRTFM::Transact([&]
 		{
 			ChildActorComponent->DestroyChildActor();
 		});
 
-	TEST_CHECK_TRUE(AutoRTFM::ETransactionResult::Committed == Result);
-	TEST_CHECK_TRUE(nullptr == ChildActorComponent->GetChildActor());
+	TestTrueExpr(AutoRTFM::ETransactionResult::Committed == Result);
+	TestTrueExpr(nullptr == ChildActorComponent->GetChildActor());
 }
 
 // Test aborting a call to USkeletalMeshComponent::RegisterComponent() with an assigned skeletal
@@ -427,7 +405,7 @@ AUTORTFM_ACTOR_COMPONENT_TEST(USkeletalMeshComponent)
 			AutoRTFM::AbortTransaction();
 		});
 
-	TEST_CHECK_TRUE(AutoRTFM::ETransactionResult::AbortedByRequest == Result);
+	TestTrueExpr(AutoRTFM::ETransactionResult::AbortedByRequest == Result);
 }
 
 
@@ -441,7 +419,7 @@ AUTORTFM_ACTOR_COMPONENT_TEST(CreateComponentFromTemplate)
 			AutoRTFM::AbortTransaction();
 		});
 
-	TEST_CHECK_TRUE(AutoRTFM::ETransactionResult::AbortedByRequest == Result);
+	TestTrueExpr(AutoRTFM::ETransactionResult::AbortedByRequest == Result);
 }
 
 // Test aborting a call to UObject::GetArchetype().
@@ -454,7 +432,7 @@ AUTORTFM_ACTOR_COMPONENT_TEST(GetArchetype)
 			AutoRTFM::AbortTransaction();
 		});
 
-	TEST_CHECK_TRUE(AutoRTFM::ETransactionResult::AbortedByRequest == Result);
+	TestTrueExpr(AutoRTFM::ETransactionResult::AbortedByRequest == Result);
 }
 
 // Test aborting a call to FUObjectArray::CloseDisregardForGC().
@@ -468,7 +446,7 @@ AUTORTFM_ACTOR_COMPONENT_TEST(CloseDisregardForGC)
 			AutoRTFM::AbortTransaction();
 		});
 
-	TEST_CHECK_TRUE(AutoRTFM::ETransactionResult::AbortedByRequest == Result);
+	TestTrueExpr(AutoRTFM::ETransactionResult::AbortedByRequest == Result);
 }
 
 AUTORTFM_ACTOR_COMPONENT_TEST(WorldGetWorldSettings)
@@ -480,8 +458,8 @@ AUTORTFM_ACTOR_COMPONENT_TEST(WorldGetWorldSettings)
 			Settings = World->GetWorldSettings();
 		});
 
-	TEST_CHECK_TRUE(AutoRTFM::ETransactionResult::Committed == Result);
-	TEST_CHECK_TRUE(Settings);
+	TestTrueExpr(AutoRTFM::ETransactionResult::Committed == Result);
+	TestTrueExpr(Settings != nullptr);
 }
 
 AUTORTFM_ACTOR_COMPONENT_TEST(WorldProcessLevelStreamingVolumes)
@@ -491,7 +469,7 @@ AUTORTFM_ACTOR_COMPONENT_TEST(WorldProcessLevelStreamingVolumes)
 			World->ProcessLevelStreamingVolumes();
 		});
 
-	TEST_CHECK_TRUE(AutoRTFM::ETransactionResult::Committed == Result);
+	TestTrueExpr(AutoRTFM::ETransactionResult::Committed == Result);
 }
 
 AUTORTFM_ACTOR_COMPONENT_TEST(WorldBlockTillLevelStreamingCompleted)
@@ -502,22 +480,20 @@ AUTORTFM_ACTOR_COMPONENT_TEST(WorldBlockTillLevelStreamingCompleted)
 	LevelStreamingPersistent->SetWorldAsset(World);
 	World->AddStreamingLevel(LevelStreamingPersistent);
 
-	TEST_CHECK_TRUE(World->HasStreamingLevelsToConsider());
+	TestTrueExpr(World->HasStreamingLevelsToConsider());
 
 	AutoRTFM::ETransactionResult Result = AutoRTFM::Transact([&]
 		{
 			World->BlockTillLevelStreamingCompleted();
 		});
 
-	TEST_CHECK_TRUE(!World->HasStreamingLevelsToConsider());
+	TestFalseExpr(World->HasStreamingLevelsToConsider());
 
-	TEST_CHECK_TRUE(AutoRTFM::ETransactionResult::Committed == Result);
+	TestTrueExpr(AutoRTFM::ETransactionResult::Committed == Result);
 }
 
 }  // anonymous namespace
 
-#undef TEST_CHECK_FALSE
-#undef TEST_CHECK_TRUE
 #undef AUTORTFM_ACTOR_COMPONENT_TEST
 
 #endif //WITH_DEV_AUTOMATION_TESTS
diff --git a/Engine/Source/Runtime/Engine/Private/Tests/AutoRTFM/AutoRTFMGameplayTagTests.cpp b/Engine/Source/Runtime/Engine/Private/Tests/AutoRTFM/AutoRTFMGameplayTagTests.cpp
--- a/Engine/Source/Runtime/Engine/Private/Tests/AutoRTFM/AutoRTFMGameplayTagTests.cpp
+++ b/Engine/Source/Runtime/Engine/Private/Tests/AutoRTFM/AutoRTFMGameplayTagTests.cpp
@@ -1,23 +1,12 @@
 // Copyright Epic Games, Inc. All Rights Reserved.
 
 #include "Misc/AutomationTest.h"
-#include "Misc/App.h"
 #include "AutoRTFM/AutoRTFM.h"
 #include "GameplayTagsManager.h"
 #include "GameplayTagContainer.h"
 
 #if WITH_DEV_AUTOMATION_TESTS
 
-#define TEST_CHECK_TRUE(b) do \
-{ \
-	if (!(b)) \
-	{ \
-		FString String(FString::Printf(TEXT("FAILED: %s:%u"), TEXT(__FILE__), __LINE__)); \
-		ExecutionInfo.AddEvent(FAutomationEvent(EAutomationEventType::Info, String)); \
-		return false; \
-	} \
-} while(false)
-
 IMPLEMENT_SIMPLE_AUTOMATION_TEST(FAutoRTFMGameplayTagTests, "AutoRTFM + FGameplayTag", EAutomationTestFlags::EngineFilter | EAutomationTestFlags::ClientContext | EAutomationTestFlags::ServerContext | EAutomationTestFlags::CommandletContext)
 
 bool FAutoRTFMGameplayTagTests::RunTest(const FString & Parameters)
@@ -38,8 +27,8 @@ bool FAutoRTFMGameplayTagTests::RunTest(const FString & Parameters)
 			bResult = Tag.MatchesTag(Other);
 		});
 
-	TEST_CHECK_TRUE(AutoRTFM::ETransactionResult::Committed == Result);
-	TEST_CHECK_TRUE(!bResult);
+	TestTrueExpr(AutoRTFM::ETransactionResult::Committed == Result);
+	TestFalseExpr(bResult);
 
 	TArray<FGameplayTag> Parents;
 
@@ -48,20 +37,18 @@ bool FAutoRTFMGameplayTagTests::RunTest(const FString & Parameters)
 			UGameplayTagsManager::Get().ExtractParentTags(Tag, Parents);
 		});
 
-	TEST_CHECK_TRUE(AutoRTFM::ETransactionResult::Committed == Result);
-	TEST_CHECK_TRUE(Parents.IsEmpty());
+	TestTrueExpr(AutoRTFM::ETransactionResult::Committed == Result);
+	TestTrueExpr(Parents.IsEmpty());
 
 	Result = AutoRTFM::Transact([&]
 		{
 			bResult = UGameplayTagsManager::Get().RequestGameplayTagParents(Tag).IsEmpty();
 		});
 
-	TEST_CHECK_TRUE(AutoRTFM::ETransactionResult::Committed == Result);
-	TEST_CHECK_TRUE(bResult);
+	TestTrueExpr(AutoRTFM::ETransactionResult::Committed == Result);
+	TestTrueExpr(bResult);
 	
 	return true;
 }
 
-#undef TEST_CHECK_TRUE
-
 #endif //WITH_DEV_AUTOMATION_TESTS
diff --git a/Engine/Source/Runtime/Engine/Private/Tests/AutoRTFM/AutoRTFMNetDriverTest.cpp b/Engine/Source/Runtime/Engine/Private/Tests/AutoRTFM/AutoRTFMNetDriverTest.cpp
--- a/Engine/Source/Runtime/Engine/Private/Tests/AutoRTFM/AutoRTFMNetDriverTest.cpp
+++ b/Engine/Source/Runtime/Engine/Private/Tests/AutoRTFM/AutoRTFMNetDriverTest.cpp
@@ -1,24 +1,12 @@
 // Copyright Epic Games, Inc. All Rights Reserved.
 
 #include "Misc/AutomationTest.h"
-#include "Misc/App.h"
 #include "AutoRTFM/AutoRTFM.h"
-#include "AutoRTFMTestActor.h"
 #include "Engine/ActorChannel.h"
 #include "Engine/DemoNetDriver.h"
 
 #if WITH_DEV_AUTOMATION_TESTS
 
-#define TEST_CHECK_TRUE(b) do \
-{ \
-	if (!(b)) \
-	{ \
-		FString String(FString::Printf(TEXT("FAILED: %s:%u"), TEXT(__FILE__), __LINE__)); \
-		ExecutionInfo.AddEvent(FAutomationEvent(EAutomationEventType::Info, String)); \
-		return false; \
-	} \
-} while(false)
-
 IMPLEMENT_SIMPLE_AUTOMATION_TEST(FAutoRTFMNetDriverTests, "AutoRTFM + FTraceFilter", EAutomationTestFlags::EngineFilter | EAutomationTestFlags::ClientContext | EAutomationTestFlags::ServerContext | EAutomationTestFlags::CommandletContext)
 
 bool FAutoRTFMNetDriverTests::RunTest(const FString & Parameters)
@@ -48,20 +36,18 @@ bool FAutoRTFMNetDriverTests::RunTest(const FString & Parameters)
 			AutoRTFM::AbortTransaction();
 		});
 
-	TEST_CHECK_TRUE(AutoRTFM::ETransactionResult::AbortedByRequest == Result);
-	TEST_CHECK_TRUE(Description.IsEmpty());
+	TestTrueExpr(AutoRTFM::ETransactionResult::AbortedByRequest == Result);
+	TestTrueExpr(Description.IsEmpty());
 
 	Result = AutoRTFM::Transact([&]
 		{
 			Description = ActorChannel->Describe();
 		});
 
-	TEST_CHECK_TRUE(AutoRTFM::ETransactionResult::Committed == Result);
-	TEST_CHECK_TRUE(!Description.IsEmpty());
+	TestTrueExpr(AutoRTFM::ETransactionResult::Committed == Result);
+	TestFalseExpr(Description.IsEmpty());
 
 	return true;
 }
 
-#undef TEST_CHECK_TRUE
-
 #endif //WITH_DEV_AUTOMATION_TESTS
diff --git a/Engine/Source/Runtime/Engine/Private/Tests/AutoRTFM/AutoRTFMTraceFilterTests.cpp b/Engine/Source/Runtime/Engine/Private/Tests/AutoRTFM/AutoRTFMTraceFilterTests.cpp
--- a/Engine/Source/Runtime/Engine/Private/Tests/AutoRTFM/AutoRTFMTraceFilterTests.cpp
+++ b/Engine/Source/Runtime/Engine/Private/Tests/AutoRTFM/AutoRTFMTraceFilterTests.cpp
@@ -1,23 +1,12 @@
 // Copyright Epic Games, Inc. All Rights Reserved.
 
 #include "Misc/AutomationTest.h"
-#include "Misc/App.h"
 #include "TraceFilter.h"
 #include "AutoRTFM/AutoRTFM.h"
 #include "AutoRTFMTestObject.h"
 
 #if WITH_DEV_AUTOMATION_TESTS
 
-#define TEST_CHECK_TRUE(b) do \
-{ \
-	if (!(b)) \
-	{ \
-		FString String(FString::Printf(TEXT("FAILED: %s:%u"), TEXT(__FILE__), __LINE__)); \
-		ExecutionInfo.AddEvent(FAutomationEvent(EAutomationEventType::Info, String)); \
-		return false; \
-	} \
-} while(false)
-
 IMPLEMENT_SIMPLE_AUTOMATION_TEST(FAutoRTFMTraceFilterTests, "AutoRTFM + FTraceFilter", EAutomationTestFlags::EngineFilter | EAutomationTestFlags::ClientContext | EAutomationTestFlags::ServerContext | EAutomationTestFlags::CommandletContext)
 
 bool FAutoRTFMTraceFilterTests::RunTest(const FString & Parameters)
@@ -33,7 +22,7 @@ bool FAutoRTFMTraceFilterTests::RunTest(const FString & Parameters)
 		UAutoRTFMTestObject* Object = NewObject<UAutoRTFMTestObject>();
 
 		FTraceFilter::SetObjectIsTraceable(Object, false);
-		TEST_CHECK_TRUE(!FTraceFilter::IsObjectTraceable(Object));
+		TestFalseExpr(FTraceFilter::IsObjectTraceable(Object));
 
 		AutoRTFM::ETransactionResult Result = AutoRTFM::Transact([&]
 			{
@@ -41,16 +30,16 @@ bool FAutoRTFMTraceFilterTests::RunTest(const FString & Parameters)
 				AutoRTFM::AbortTransaction();
 			});
 
-		TEST_CHECK_TRUE(AutoRTFM::ETransactionResult::AbortedByRequest == Result);
-		TEST_CHECK_TRUE(!FTraceFilter::IsObjectTraceable(Object));
+		TestTrueExpr(AutoRTFM::ETransactionResult::AbortedByRequest == Result);
+		TestFalseExpr(FTraceFilter::IsObjectTraceable(Object));
 
 		Result = AutoRTFM::Transact([&]
 			{
 				FTraceFilter::SetObjectIsTraceable(Object, true);
 			});
 
-		TEST_CHECK_TRUE(AutoRTFM::ETransactionResult::Committed == Result);
-		TEST_CHECK_TRUE(FTraceFilter::IsObjectTraceable(Object));
+		TestTrueExpr(AutoRTFM::ETransactionResult::Committed == Result);
+		TestTrueExpr(FTraceFilter::IsObjectTraceable(Object));
 
 		Result = AutoRTFM::Transact([&]
 			{
@@ -62,7 +51,7 @@ bool FAutoRTFMTraceFilterTests::RunTest(const FString & Parameters)
 				AutoRTFM::AbortTransaction();
 			});
 
-		TEST_CHECK_TRUE(!FTraceFilter::IsObjectTraceable(Object));
+		TestFalseExpr(FTraceFilter::IsObjectTraceable(Object));
 
 		Result = AutoRTFM::Transact([&]
 			{
@@ -72,16 +61,16 @@ bool FAutoRTFMTraceFilterTests::RunTest(const FString & Parameters)
 					});
 			});
 
-		TEST_CHECK_TRUE(FTraceFilter::IsObjectTraceable(Object));
+		TestTrueExpr(FTraceFilter::IsObjectTraceable(Object));
 
 		FTraceFilter::SetObjectIsTraceable(Object, false);
-		TEST_CHECK_TRUE(!FTraceFilter::IsObjectTraceable(Object));
+		TestFalseExpr(FTraceFilter::IsObjectTraceable(Object));
 
 		UAutoRTFMTestObject* Other = NewObject<UAutoRTFMTestObject>();
 		UAutoRTFMTestObject* Another = NewObject<UAutoRTFMTestObject>();
 
 		FTraceFilter::SetObjectIsTraceable(Other, false);
-		TEST_CHECK_TRUE(!FTraceFilter::IsObjectTraceable(Other));
+		TestFalseExpr(FTraceFilter::IsObjectTraceable(Other));
 
 		Result = AutoRTFM::Transact([&]
 			{
@@ -100,10 +89,10 @@ bool FAutoRTFMTraceFilterTests::RunTest(const FString & Parameters)
 				AutoRTFM::AbortTransaction();
 			});
 
-		TEST_CHECK_TRUE(AutoRTFM::ETransactionResult::AbortedByRequest == Result);
-		TEST_CHECK_TRUE(!FTraceFilter::IsObjectTraceable(Object));
-		TEST_CHECK_TRUE(FTraceFilter::IsObjectTraceable(Other));
-		TEST_CHECK_TRUE(FTraceFilter::IsObjectTraceable(Another));
+		TestTrueExpr(AutoRTFM::ETransactionResult::AbortedByRequest == Result);
+		TestFalseExpr(FTraceFilter::IsObjectTraceable(Object));
+		TestTrueExpr(FTraceFilter::IsObjectTraceable(Other));
+		TestTrueExpr(FTraceFilter::IsObjectTraceable(Another));
 
 		Result = AutoRTFM::Transact([&]
 			{
@@ -120,23 +109,23 @@ bool FAutoRTFMTraceFilterTests::RunTest(const FString & Parameters)
 					});
 			});
 
-		TEST_CHECK_TRUE(AutoRTFM::ETransactionResult::Committed == Result);
-		TEST_CHECK_TRUE(FTraceFilter::IsObjectTraceable(Object));
-		TEST_CHECK_TRUE(!FTraceFilter::IsObjectTraceable(Other));
-		TEST_CHECK_TRUE(!FTraceFilter::IsObjectTraceable(Another));
+		TestTrueExpr(AutoRTFM::ETransactionResult::Committed == Result);
+		TestTrueExpr(FTraceFilter::IsObjectTraceable(Object));
+		TestFalseExpr(FTraceFilter::IsObjectTraceable(Other));
+		TestFalseExpr(FTraceFilter::IsObjectTraceable(Another));
 	}
 
 	{
 		UAutoRTFMTestObject* Object = NewObject<UAutoRTFMTestObject>();
 
 		FTraceFilter::SetObjectIsTraceable(Object, false);
-		TEST_CHECK_TRUE(!FTraceFilter::IsObjectTraceable(Object));
+		TestFalseExpr(FTraceFilter::IsObjectTraceable(Object));
 
 		FTraceFilter::MarkObjectTraceable(Object);
-		TEST_CHECK_TRUE(FTraceFilter::IsObjectTraceable(Object));
+		TestTrueExpr(FTraceFilter::IsObjectTraceable(Object));
 
 		FTraceFilter::SetObjectIsTraceable(Object, false);
-		TEST_CHECK_TRUE(!FTraceFilter::IsObjectTraceable(Object));
+		TestFalseExpr(FTraceFilter::IsObjectTraceable(Object));
 
 		AutoRTFM::ETransactionResult Result = AutoRTFM::Transact([&]
 			{
@@ -144,19 +133,19 @@ bool FAutoRTFMTraceFilterTests::RunTest(const FString & Parameters)
 				AutoRTFM::AbortTransaction();
 			});
 
-		TEST_CHECK_TRUE(AutoRTFM::ETransactionResult::AbortedByRequest == Result);
-		TEST_CHECK_TRUE(!FTraceFilter::IsObjectTraceable(Object));
+		TestTrueExpr(AutoRTFM::ETransactionResult::AbortedByRequest == Result);
+		TestFalseExpr(FTraceFilter::IsObjectTraceable(Object));
 
 		Result = AutoRTFM::Transact([&]
 			{
 				FTraceFilter::MarkObjectTraceable(Object);
 			});
 
-		TEST_CHECK_TRUE(AutoRTFM::ETransactionResult::Committed == Result);
-		TEST_CHECK_TRUE(FTraceFilter::IsObjectTraceable(Object));
+		TestTrueExpr(AutoRTFM::ETransactionResult::Committed == Result);
+		TestTrueExpr(FTraceFilter::IsObjectTraceable(Object));
 
 		FTraceFilter::SetObjectIsTraceable(Object, false);
-		TEST_CHECK_TRUE(!FTraceFilter::IsObjectTraceable(Object));
+		TestFalseExpr(FTraceFilter::IsObjectTraceable(Object));
 
 		Result = AutoRTFM::Transact([&]
 			{
@@ -168,10 +157,10 @@ bool FAutoRTFMTraceFilterTests::RunTest(const FString & Parameters)
 				AutoRTFM::AbortTransaction();
 			});
 
-		TEST_CHECK_TRUE(FTraceFilter::IsObjectTraceable(Object));
+		TestTrueExpr(FTraceFilter::IsObjectTraceable(Object));
 
 		FTraceFilter::SetObjectIsTraceable(Object, false);
-		TEST_CHECK_TRUE(!FTraceFilter::IsObjectTraceable(Object));
+		TestFalseExpr(FTraceFilter::IsObjectTraceable(Object));
 
 		Result = AutoRTFM::Transact([&]
 			{
@@ -181,7 +170,7 @@ bool FAutoRTFMTraceFilterTests::RunTest(const FString & Parameters)
 					});
 			});
 
-		TEST_CHECK_TRUE(FTraceFilter::IsObjectTraceable(Object));
+		TestTrueExpr(FTraceFilter::IsObjectTraceable(Object));
 	}
 
 	{
@@ -194,8 +183,8 @@ bool FAutoRTFMTraceFilterTests::RunTest(const FString & Parameters)
 				bTraceable = FTraceFilter::IsObjectTraceable(Object);
 			});
 
-		TEST_CHECK_TRUE(AutoRTFM::ETransactionResult::Committed == Result);
-		TEST_CHECK_TRUE(FTraceFilter::IsObjectTraceable(Object) == bTraceable);
+		TestTrueExpr(AutoRTFM::ETransactionResult::Committed == Result);
+		TestTrueExpr(FTraceFilter::IsObjectTraceable(Object) == bTraceable);
 	}
 
 	return true;
@@ -205,6 +194,4 @@ bool FAutoRTFMTraceFilterTests::RunTest(const FString & Parameters)
 #endif
 }
 
-#undef TEST_CHECK_TRUE
-
 #endif //WITH_DEV_AUTOMATION_TESTS
