Initialize the PrevPrimitiveLocalToWorld on first FPrimitiveInstanceDataManager flush.
If we don't do this and there are no further flushes then PrevPrimitiveLocalToWorld will be zero, and if it is used for motion vector calculation it will give bad results.
It is used for motion vector calculation if we have SetHasPerInstancePrevTransforms(true) on the ISM. Typically that is followed by one or more per frame updates of at least one instance transform, but in the edge case where it isn't we will see bad motion vectors.

diff --git a/Engine/Source/Runtime/Engine/Private/InstancedStaticMesh/ISMInstanceDataManager.cpp b/Engine/Source/Runtime/Engine/Private/InstancedStaticMesh/ISMInstanceDataManager.cpp
--- a/Engine/Source/Runtime/Engine/Private/InstancedStaticMesh/ISMInstanceDataManager.cpp
+++ b/Engine/Source/Runtime/Engine/Private/InstancedStaticMesh/ISMInstanceDataManager.cpp
@@ -672,7 +672,8 @@ bool FPrimitiveInstanceDataManager::FlushChanges(FInstanceUpdateComponentDesc &&
 		ChangeDesc.bPrimitiveTransformChanged = true;
 	}
 
-	FMatrix PrevPrimitiveLocalToWorld = PrimitiveLocalToWorld;
+	// Init previous LocalToWorld to current on inital update
+	FMatrix PrevPrimitiveLocalToWorld = GetState() == ETrackingState::Initial ? ComponentData.PrimitiveLocalToWorld : PrimitiveLocalToWorld;
 	// Update the tracked state
 	PrimitiveLocalToWorld = ComponentData.PrimitiveLocalToWorld;
 	AbsMaxDisplacement = NewAbsMaxDisplacement;
