Skin Weight Profiles: Check whether a CPU-side weight buffer actually exists rather than rely on the CPU data wanted flag.
Tightened up the readback as well to try whether a CPU-side buffer exists or not. Streaming can delay when this becomes available and so we trigger a readback as soon as the vertex data is available to the GPU.
This new code path can be disabled by setting the SkinWeightProfileManager.HandleDelayedLoads cvar to 0.

diff --git a/Engine/Source/Runtime/Engine/Private/Animation/SkinWeightProfile.cpp b/Engine/Source/Runtime/Engine/Private/Animation/SkinWeightProfile.cpp
--- a/Engine/Source/Runtime/Engine/Private/Animation/SkinWeightProfile.cpp
+++ b/Engine/Source/Runtime/Engine/Private/Animation/SkinWeightProfile.cpp
@@ -621,19 +621,41 @@ void FSkinWeightProfilesData::EnqueueGPUReadback()
 {
 	ensure(!ReadbackData.BufferReadback.IsValid());
 
-	if (BaseBuffer && BaseBuffer->GetDataVertexBuffer()->IsWeightDataValid() && BaseBuffer->GetDataVertexBuffer()->GetVertexDataSize())
+	if (FSkinWeightProfileManager::HandleDelayedLoads())
 	{
-		static const FName ReadbackName = TEXT("ReadbackSkinWeightBuffer");
-		FSkinweightReadbackData* DataPtr = &ReadbackData;
-		ENQUEUE_RENDER_COMMAND(FSkinWeightProfilesData_EnqueueGPUReadback)(
-			[DataPtr, this](FRHICommandListImmediate& RHICmdList)
+		if (BaseBuffer && BaseBuffer->GetDataVertexBuffer()->GetVertexDataSize())
 		{
-			DataPtr->BufferReadback.Reset(new FRHIGPUBufferReadback(ReadbackName));
-			if (DataPtr->BufferReadback.IsValid() && BaseBuffer->GetDataVertexBuffer()->VertexBufferRHI->GetSize())
+			static const FName ReadbackName = TEXT("ReadbackSkinWeightBuffer");
+			FSkinweightReadbackData* DataPtr = &ReadbackData;
+			ENQUEUE_RENDER_COMMAND(FSkinWeightProfilesData_EnqueueGPUReadback)(
+				[DataPtr, this](FRHICommandListImmediate& RHICmdList)
 			{
-				DataPtr->BufferReadback->EnqueueCopy(RHICmdList, BaseBuffer->GetDataVertexBuffer()->VertexBufferRHI);
-			}
-		});
+				if (BaseBuffer->GetDataVertexBuffer()->VertexBufferRHI->GetSize())
+				{
+					// Only set up the readback buffer if we have a vertex buffer to read from. It's possible we're still waiting for
+					// the mesh to be streamed in.
+					DataPtr->BufferReadback.Reset(new FRHIGPUBufferReadback(ReadbackName));
+					DataPtr->BufferReadback->EnqueueCopy(RHICmdList, BaseBuffer->GetDataVertexBuffer()->VertexBufferRHI);
+				}
+			});
+		}
+	}
+	else
+	{
+		if (BaseBuffer && BaseBuffer->GetDataVertexBuffer()->IsWeightDataValid() && BaseBuffer->GetDataVertexBuffer()->GetVertexDataSize())
+		{
+			static const FName ReadbackName = TEXT("ReadbackSkinWeightBuffer");
+			FSkinweightReadbackData* DataPtr = &ReadbackData;
+			ENQUEUE_RENDER_COMMAND(FSkinWeightProfilesData_EnqueueGPUReadback)(
+				[DataPtr, this](FRHICommandListImmediate& RHICmdList)
+			{
+				DataPtr->BufferReadback.Reset(new FRHIGPUBufferReadback(ReadbackName));
+				if (DataPtr->BufferReadback.IsValid() && BaseBuffer->GetDataVertexBuffer()->VertexBufferRHI->GetSize())
+				{
+					DataPtr->BufferReadback->EnqueueCopy(RHICmdList, BaseBuffer->GetDataVertexBuffer()->VertexBufferRHI);
+				}
+			});
+		}
 	}
 }
 
@@ -684,23 +706,31 @@ void FSkinWeightProfilesData::ResetGPUReadback()
 	ReadbackData.ReadbackFinishedFrameIndex = INDEX_NONE;
 }
 
+bool FSkinWeightProfilesData::HasProfileStack(const FSkinWeightProfileStack& InProfileStack) const
+{
+	check(InProfileStack == InProfileStack.Normalized());
+	
+	return ProfileStackToBuffer.Contains(InProfileStack);
+}
+
 void FSkinWeightProfilesData::InitialiseProfileBuffer(const FSkinWeightProfileStack& InProfileStack)
 {
 	LLM_SCOPE(ELLMTag::SkeletalMesh);
 
-	check(InProfileStack == InProfileStack.Normalized());
-	
-	if (ProfileStackToBuffer.Contains(InProfileStack))
+	// Have we already constructed this particular profile stack?
+	if (HasProfileStack(InProfileStack))
 	{
-		// Maybe it has already been previously set-up and not unloaded yet
 		return;
 	}
 
 	if (BaseBuffer)
 	{
 		const uint8* BaseBufferData;
-		
-		if (BaseBuffer->GetNeedsCPUAccess())
+
+		// If we have the weight data, then just use that directly. Otherwise, assume that we've been called as a result of a successful
+		// GPU readback.
+		if ((FSkinWeightProfileManager::AllowCPU() && FSkinWeightProfileManager::HandleDelayedLoads() && BaseBuffer->GetDataVertexBuffer()->IsWeightDataValid()) ||
+		    (!FSkinWeightProfileManager::HandleDelayedLoads() && BaseBuffer->GetNeedsCPUAccess()))
 		{
 			BaseBufferData = BaseBuffer->GetDataVertexBuffer()->GetWeightData();
 		}
@@ -724,7 +754,7 @@ void FSkinWeightProfilesData::InitialiseProfileBuffer(const FSkinWeightProfileSt
 			OverrideBuffer->SetOwnerName(OwnerName);
 #endif
 			OverrideBuffer->BeginInitResources();
-		}	
+		}
 	}
 }
 
@@ -848,4 +878,3 @@ void FRuntimeSkinWeightProfileData::ApplyDefaultOverride(FSkinWeightVertexBuffer
 		}
 	}
 }
-
diff --git a/Engine/Source/Runtime/Engine/Private/Animation/SkinWeightProfileManager.cpp b/Engine/Source/Runtime/Engine/Private/Animation/SkinWeightProfileManager.cpp
--- a/Engine/Source/Runtime/Engine/Private/Animation/SkinWeightProfileManager.cpp
+++ b/Engine/Source/Runtime/Engine/Private/Animation/SkinWeightProfileManager.cpp
@@ -19,6 +19,16 @@ static FAutoConsoleVariableRef CVarAllowCPU (
 	ECVF_Cheat
 );
 
+static int32 GHandleDelayedLoads = 1;
+static FAutoConsoleVariableRef CHandleDelayedLoads (
+	TEXT("SkinWeightProfileManager.HandleDelayedLoads"),
+	GHandleDelayedLoads,
+	TEXT("Whether or not to handle skeletal meshes that are severely delayed in loading"),
+	ECVF_Cheat
+);
+
+
+
 FSkinWeightProfileManager::FSkinWeightProfileManager(UWorld* InWorld) : LastGamethreadProfileIndex(INDEX_NONE), WeakWorld(nullptr)
 {
 	if (InWorld)
@@ -82,6 +92,16 @@ FSkinWeightProfileManager* FSkinWeightProfileManager::Get(UWorld* World)
 	return nullptr;
 }
 
+bool FSkinWeightProfileManager::AllowCPU()
+{
+	return GAllowCPU != 0;
+}
+
+bool FSkinWeightProfileManager::HandleDelayedLoads()
+{
+	return GHandleDelayedLoads != 0;
+}
+
 void FSkinWeightProfileManager::OnStartup()
 {
 	FWorldDelegates::OnPreWorldInitialization.AddStatic(&OnWorldInit);
@@ -197,7 +217,7 @@ void FSkinWeightProfileManager::DoTick(float DeltaTime, ENamedThreads::Type Curr
 			FSkeletalMeshRenderData* RenderData = SkeletalMesh ? SkeletalMesh->GetResourceForRendering() : nullptr;
 			// Or if skeletal mesh / render data is invalid
 			bRemove = bRemove ? true : (RenderData == nullptr);
-				
+
 			if (!bRemove)
 			{	
 				bool bAllBuffersReady = true;
@@ -338,8 +358,17 @@ void FSkinWeightProfileManagerAsyncTask::DoTask(ENamedThreads::Type CurrentThrea
 				{
 					FSkeletalMeshLODRenderData& LODRenderData = RenderData->LODRenderData[LODIndex];
 					FSkinWeightProfilesData& SkinweightData = LODRenderData.SkinWeightProfilesData;
+
+					// Don't try to initialize or do a GPU readback for a profile that's already completed.
+					if (FSkinWeightProfileManager::HandleDelayedLoads() && SkinweightData.HasProfileStack(Request.ProfileStack))
+					{
+						continue;
+					}
 					
-					if (!!GAllowCPU && LODRenderData.SkinWeightVertexBuffer.GetNeedsCPUAccess())
+					// If we have CPU-side weights available, use them right away.
+					if (FSkinWeightProfileManager::AllowCPU() &&
+						((FSkinWeightProfileManager::HandleDelayedLoads() && SkinweightData.BaseBuffer->GetDataVertexBuffer()->IsWeightDataValid()) ||
+						(!FSkinWeightProfileManager::HandleDelayedLoads() && LODRenderData.SkinWeightVertexBuffer.GetNeedsCPUAccess())))
 					{
 						SkinweightData.InitialiseProfileBuffer(Request.ProfileStack);
 					}
@@ -366,4 +395,3 @@ void FSkinWeightProfileManagerAsyncTask::DoTask(ENamedThreads::Type CurrentThrea
 		}
 	}
 }
-
diff --git a/Engine/Source/Runtime/Engine/Public/Animation/SkinWeightProfile.h b/Engine/Source/Runtime/Engine/Public/Animation/SkinWeightProfile.h
--- a/Engine/Source/Runtime/Engine/Public/Animation/SkinWeightProfile.h
+++ b/Engine/Source/Runtime/Engine/Public/Animation/SkinWeightProfile.h
@@ -282,7 +282,8 @@ struct FSkinWeightProfilesData
 private:
 	friend class FSkinWeightProfileManager;
 	friend class FSkinWeightProfileManagerAsyncTask;
-	
+
+	bool HasProfileStack(const FSkinWeightProfileStack& InProfileStack) const;
 	void InitialiseProfileBuffer(const FSkinWeightProfileStack& InProfileStack);
 	void ReleaseBuffer(const FSkinWeightProfileStack& InProfileStack, bool bForceRelease = false);
 	
diff --git a/Engine/Source/Runtime/Engine/Public/Animation/SkinWeightProfileManager.h b/Engine/Source/Runtime/Engine/Public/Animation/SkinWeightProfileManager.h
--- a/Engine/Source/Runtime/Engine/Public/Animation/SkinWeightProfileManager.h
+++ b/Engine/Source/Runtime/Engine/Public/Animation/SkinWeightProfileManager.h
@@ -121,6 +121,8 @@ public:
 	static void OnStartup();
 	static void OnShutdown();
 	static ENGINE_API FSkinWeightProfileManager* Get(UWorld* World);
+	static bool AllowCPU();
+	static bool HandleDelayedLoads();
 
 	FSkinWeightProfileManager(UWorld* InWorld);
 
