Substrate: Fix minor typo in help text

#rnx

diff --git a/Engine/Source/Runtime/Engine/Public/Materials/MaterialExpressionSubstrate.h b/Engine/Source/Runtime/Engine/Public/Materials/MaterialExpressionSubstrate.h
--- a/Engine/Source/Runtime/Engine/Public/Materials/MaterialExpressionSubstrate.h
+++ b/Engine/Source/Runtime/Engine/Public/Materials/MaterialExpressionSubstrate.h
@@ -273,7 +273,7 @@ class UMaterialExpressionSubstrateSlabBSDF : public UMaterialExpressionSubstrate
 	FExpressionInput SSSMFP;
 
 	/**
-	 * Scale the mean free path length. Value between 0 and 1. Always used, without or without a subsurface profile. (type = float, unitless, defaults to 1)
+	 * Scale the mean free path length. Value between 0 and 1. Always used, with or without a subsurface profile. (type = float, unitless, defaults to 1)
 	 */
 	UPROPERTY()
 	FExpressionInput SSSMFPScale;
