Fix UDebugDrawService to work under transactions.

diff --git a/Engine/Source/Runtime/Engine/Classes/Debug/DebugDrawService.h b/Engine/Source/Runtime/Engine/Classes/Debug/DebugDrawService.h
--- a/Engine/Source/Runtime/Engine/Classes/Debug/DebugDrawService.h
+++ b/Engine/Source/Runtime/Engine/Classes/Debug/DebugDrawService.h
@@ -3,7 +3,7 @@
 
 #include "ShowFlags.h"
 #include "Kismet/BlueprintFunctionLibrary.h"
-#include "HAL/CriticalSection.h"
+#include "Misc/TransactionallySafeCriticalSection.h"
 #include "DebugDrawService.generated.h"
 
 class APlayerController;
@@ -45,7 +45,7 @@ private:
 	 *		or
 	 *			RecreateRenderState_Concurrent that will unregister and register delegates from multiple threads.
 	 */
-	static FCriticalSection DelegatesLock;
+	static FTransactionallySafeCriticalSection DelegatesLock;
 
 	DECLARE_MULTICAST_DELEGATE_TwoParams(FDebugDrawMulticastDelegate, UCanvas*, APlayerController*);
 
diff --git a/Engine/Source/Runtime/Engine/Private/Debug/DebugDrawService.cpp b/Engine/Source/Runtime/Engine/Private/Debug/DebugDrawService.cpp
--- a/Engine/Source/Runtime/Engine/Private/Debug/DebugDrawService.cpp
+++ b/Engine/Source/Runtime/Engine/Private/Debug/DebugDrawService.cpp
@@ -4,10 +4,11 @@
 #include "SceneView.h"
 #include "UObject/Package.h"
 #include "Engine/Canvas.h"
+#include "Misc/ScopeLock.h"
 
 #include UE_INLINE_GENERATED_CPP_BY_NAME(DebugDrawService)
 
-FCriticalSection UDebugDrawService::DelegatesLock;
+FTransactionallySafeCriticalSection UDebugDrawService::DelegatesLock;
 
 TArray<UDebugDrawService::FDebugDrawMulticastDelegate> UDebugDrawService::Delegates;
 FEngineShowFlags UDebugDrawService::ObservedFlags(ESFIM_Editor);
@@ -25,7 +26,7 @@ FDelegateHandle UDebugDrawService::Register(const TCHAR* Name, const FDebugDrawD
 	FDelegateHandle Result;
 	if (Index != INDEX_NONE)
 	{
-		FScopeLock ScopeLock(&DelegatesLock);
+		UE::TScopeLock ScopeLock(DelegatesLock);
 		if (Index >= Delegates.Num())
 		{
 			Delegates.AddZeroed(Index - Delegates.Num() + 1);
@@ -38,7 +39,7 @@ FDelegateHandle UDebugDrawService::Register(const TCHAR* Name, const FDebugDrawD
 
 void UDebugDrawService::Unregister(const FDelegateHandle HandleToRemove)
 {
-	FScopeLock ScopeLock(&DelegatesLock);
+	UE::TScopeLock ScopeLock(DelegatesLock);
 	for (int32 Flag = 0; Flag < Delegates.Num(); ++Flag)
 	{
 		FDebugDrawMulticastDelegate& MulticastDelegate = Delegates[Flag];
@@ -83,7 +84,7 @@ void UDebugDrawService::Draw(const FEngineShowFlags Flags, UCanvas* Canvas)
 		return;
 	}
 
-	FScopeLock ScopeLock(&DelegatesLock);
+	UE::TScopeLock ScopeLock(DelegatesLock);
 	for (int32 FlagIndex = 0; FlagIndex < Delegates.Num(); ++FlagIndex)
 	{
 		FDebugDrawMulticastDelegate& MulticastDelegate = Delegates[FlagIndex];
diff --git a/Engine/Source/Runtime/Engine/Private/ShowFlags.cpp b/Engine/Source/Runtime/Engine/Private/ShowFlags.cpp
--- a/Engine/Source/Runtime/Engine/Private/ShowFlags.cpp
+++ b/Engine/Source/Runtime/Engine/Private/ShowFlags.cpp
@@ -8,7 +8,7 @@
 #include "RenderUtils.h"
 #include "Engine/EngineBaseTypes.h"
 #include "HAL/IConsoleManager.h"
-#include "Misc/ScopeRWLock.h"
+#include "Misc/TransactionallySafeRWScopeLock.h"
 #include "SystemSettings.h"
 
 static bool IsValidNameChar(TCHAR c)
@@ -927,9 +927,9 @@ namespace CustomShowFlagsInternal
 		return RegisteredCustomShowFlags;
 	}
 
-	FRWLock& GetLock()
+	FTransactionallySafeRWLock& GetLock()
 	{
-		static FRWLock Lock;
+		static FTransactionallySafeRWLock Lock;
 		return Lock;
 	}
 
@@ -974,7 +974,7 @@ FEngineShowFlags::ECustomShowFlag FEngineShowFlags::RegisterCustomShowFlag(const
 			}
 		}
 
-		FRWScopeLock Lock(CustomShowFlagsInternal::GetLock(), SLT_Write);
+		FTransactionallySafeRWScopeLock Lock(CustomShowFlagsInternal::GetLock(), SLT_Write);
 		ReturnIndex = (FEngineShowFlags::ECustomShowFlag)CustomShowFlagsInternal::GetRegisteredCustomShowFlags().AddDefaulted(1);
 		CustomShowFlagsInternal::GetDefaultState().Add(1);
 		CustomShowFlagsInternal::GetNameToIndex().Add(Name, ReturnIndex);
@@ -1021,7 +1021,7 @@ FSimpleMulticastDelegate FEngineShowFlags::OnCustomShowFlagRegistered;
 // Private custom show flags functions
 FEngineShowFlags::ECustomShowFlag FEngineShowFlags::FindCustomShowFlagByName(const FString& Name)
 {
-	FRWScopeLock Lock(CustomShowFlagsInternal::GetLock(), SLT_ReadOnly);
+	FTransactionallySafeRWScopeLock Lock(CustomShowFlagsInternal::GetLock(), SLT_ReadOnly);
 
 	if (FEngineShowFlags::ECustomShowFlag* Existing = CustomShowFlagsInternal::GetNameToIndex().Find(Name))
 	{
@@ -1032,13 +1032,13 @@ FEngineShowFlags::ECustomShowFlag FEngineShowFlags::FindCustomShowFlagByName(con
 
 bool FEngineShowFlags::IsRegisteredCustomShowFlag(ECustomShowFlag Index)
 {
-	FRWScopeLock Lock(CustomShowFlagsInternal::GetLock(), SLT_ReadOnly);
+	FTransactionallySafeRWScopeLock Lock(CustomShowFlagsInternal::GetLock(), SLT_ReadOnly);
 	return CustomShowFlagsInternal::GetRegisteredCustomShowFlags().IsValidIndex((uint32)Index);
 }
 
 FString FEngineShowFlags::GetCustomShowFlagName(ECustomShowFlag Index)
 {
-	FRWScopeLock Lock(CustomShowFlagsInternal::GetLock(), SLT_ReadOnly);
+	FTransactionallySafeRWScopeLock Lock(CustomShowFlagsInternal::GetLock(), SLT_ReadOnly);
 	if (CustomShowFlagsInternal::GetRegisteredCustomShowFlags().IsValidIndex((uint32)Index))
 	{
 		return  CustomShowFlagsInternal::GetRegisteredCustomShowFlags()[(uint32)Index].Name;
@@ -1048,7 +1048,7 @@ FString FEngineShowFlags::GetCustomShowFlagName(ECustomShowFlag Index)
 
 FText FEngineShowFlags::GetCustomShowFlagDisplayName(ECustomShowFlag Index)
 {
-	FRWScopeLock Lock(CustomShowFlagsInternal::GetLock(), SLT_ReadOnly);
+	FTransactionallySafeRWScopeLock Lock(CustomShowFlagsInternal::GetLock(), SLT_ReadOnly);
 	if (CustomShowFlagsInternal::GetRegisteredCustomShowFlags().IsValidIndex((uint32)Index))
 	{
 		return CustomShowFlagsInternal::GetRegisteredCustomShowFlags()[(uint32)Index].DisplayName;
@@ -1058,7 +1058,7 @@ FText FEngineShowFlags::GetCustomShowFlagDisplayName(ECustomShowFlag Index)
 
 EShowFlagGroup FEngineShowFlags::GetCustomShowFlagGroup(ECustomShowFlag Index)
 {
-	FRWScopeLock Lock(CustomShowFlagsInternal::GetLock(), SLT_ReadOnly);
+	FTransactionallySafeRWScopeLock Lock(CustomShowFlagsInternal::GetLock(), SLT_ReadOnly);
 	if (CustomShowFlagsInternal::GetRegisteredCustomShowFlags().IsValidIndex((uint32)Index))
 	{
 		return CustomShowFlagsInternal::GetRegisteredCustomShowFlags()[(uint32)Index].Group;
@@ -1068,7 +1068,7 @@ EShowFlagGroup FEngineShowFlags::GetCustomShowFlagGroup(ECustomShowFlag Index)
 
 void FEngineShowFlags::InitCustomShowFlags(EShowFlagInitMode InitMode)
 {
-	FRWScopeLock Lock(CustomShowFlagsInternal::GetLock(), SLT_ReadOnly);
+	FTransactionallySafeRWScopeLock Lock(CustomShowFlagsInternal::GetLock(), SLT_ReadOnly);
 	if (InitMode == ESFIM_All0)
 	{
 		CustomShowFlags.Init(false, CustomShowFlagsInternal::GetDefaultState().Num());
@@ -1081,7 +1081,7 @@ void FEngineShowFlags::InitCustomShowFlags(EShowFlagInitMode InitMode)
 
 void FEngineShowFlags::UpdateNewCustomShowFlags()
 {
-	FRWScopeLock Lock(CustomShowFlagsInternal::GetLock(), SLT_ReadOnly);
+	FTransactionallySafeRWScopeLock Lock(CustomShowFlagsInternal::GetLock(), SLT_ReadOnly);
 	for (int32 i = CustomShowFlags.Num(); i < CustomShowFlagsInternal::GetDefaultState().Num(); ++i)
 	{
 		CustomShowFlags.Add(CustomShowFlagsInternal::GetDefaultState()[i]);
@@ -1090,7 +1090,7 @@ void FEngineShowFlags::UpdateNewCustomShowFlags()
 
 bool FEngineShowFlags::FindCustomShowFlagDisplayName(const FString& Name, FText& OutText)
 {
-	FRWScopeLock Lock(CustomShowFlagsInternal::GetLock(), SLT_ReadOnly);
+	FTransactionallySafeRWScopeLock Lock(CustomShowFlagsInternal::GetLock(), SLT_ReadOnly);
 	if (FEngineShowFlags::ECustomShowFlag* Existing = CustomShowFlagsInternal::GetNameToIndex().Find(Name))
 	{
 		OutText = CustomShowFlagsInternal::GetRegisteredCustomShowFlags()[(int32)*Existing].DisplayName;
@@ -1101,7 +1101,7 @@ bool FEngineShowFlags::FindCustomShowFlagDisplayName(const FString& Name, FText&
 
 void FEngineShowFlags::EngineOverrideCustomShowFlagsFromCVars()
 {
-	FRWScopeLock Lock(CustomShowFlagsInternal::GetLock(), SLT_ReadOnly);
+	FTransactionallySafeRWScopeLock Lock(CustomShowFlagsInternal::GetLock(), SLT_ReadOnly);
 	for (int32 i = 0; i < CustomShowFlagsInternal::GetRegisteredCustomShowFlags().Num(); ++i)
 	{
 		int32 Val = CustomShowFlagsInternal::GetRegisteredCustomShowFlags()[i].CVar->GetInt();
diff --git a/Engine/Source/Runtime/Engine/Private/Tests/AutoRTFM/AutoRTFMDebugDrawServiceTest.cpp b/Engine/Source/Runtime/Engine/Private/Tests/AutoRTFM/AutoRTFMDebugDrawServiceTest.cpp
new file mode 100644
--- /dev/null
+++ b/Engine/Source/Runtime/Engine/Private/Tests/AutoRTFM/AutoRTFMDebugDrawServiceTest.cpp
@@ -0,0 +1,144 @@
+// Copyright Epic Games, Inc. All Rights Reserved.
+
+#include "Debug/DebugDrawService.h"
+#include "Engine/Canvas.h"
+#include "Misc/AutomationTest.h"
+#include "AutoRTFM/AutoRTFM.h"
+#include "UnrealClient.h"
+#include "HitProxies.h"
+
+#if WITH_DEV_AUTOMATION_TESTS
+
+#define CHECK_EQ(A, B) \
+	UTEST_EQUAL(TEXT(__FILE__ ":" UE_STRINGIZE(__LINE__) ": UTEST_EQUAL_EXPR(" #A ", " #B ")"), A, B)
+
+IMPLEMENT_SIMPLE_AUTOMATION_TEST(FAutoRTFMDebugDrawServiceTest, "AutoRTFM + UDebugDrawService", EAutomationTestFlags::EngineFilter | EAutomationTestFlags::ClientContext | EAutomationTestFlags::ServerContext | EAutomationTestFlags::CommandletContext)
+
+bool FAutoRTFMDebugDrawServiceTest::RunTest(const FString & Parameters)
+{
+	if (!AutoRTFM::ForTheRuntime::IsAutoRTFMRuntimeEnabled())
+	{
+		ExecutionInfo.AddEvent(FAutomationEvent(EAutomationEventType::Info, TEXT("SKIPPED 'FAutoRTFMDebugDrawServiceTest' test. AutoRTFM disabled.")));
+		return true;
+	}
+
+	//ENGINE_API FCanvas(const FGameTime& Time, ERHIFeatureLevel::Type InFeatureLevel, float InDPIScale = 1.0f);
+
+	struct FMyRenderTarget final : public FRenderTarget
+	{
+		FIntPoint GetSizeXY() const override { return FIntPoint(128); }
+	} RenderTarget;
+
+	struct FMyHitProxyConsumer final : public FHitProxyConsumer
+	{
+		void AddHitProxy(HHitProxy*) override {}
+	} HitProxyConsumer;
+
+	FGameTime GameTime;
+
+	FCanvas Canvas(&RenderTarget, &HitProxyConsumer, GameTime, ERHIFeatureLevel::SM5);
+	UCanvas* const CanvasObject = NewObject<UCanvas>();
+	CanvasObject->Canvas = &Canvas;
+
+	auto Section = [&](const TCHAR* Name, auto&& Test)
+	{
+		bool Result = Test();
+		if (!Result)
+		{
+			AddError(FString::Printf(TEXT("In section '%s'."), Name), 1);
+		}
+	};
+
+	Section(TEXT("Transact(Register, Draw, Unregister)"), [&]
+	{
+		bool bHit = false;
+		bool bIsClosed = false;
+
+		AutoRTFM::Transact([&]
+		{
+			FDebugDrawDelegate Delegate;
+			Delegate.BindLambda([&](UCanvas*, APlayerController*)
+			{
+				bHit = true;
+				bIsClosed = AutoRTFM::IsClosed();
+			});
+
+			FDelegateHandle Handle = UDebugDrawService::Register(TEXT("Tonemapper"), Delegate);
+
+			FEngineShowFlags EngineShowFlags(EShowFlagInitMode::ESFIM_Game);
+
+			EngineShowFlags.SetTonemapper(true);
+
+			UDebugDrawService::Draw(EngineShowFlags, CanvasObject);
+
+			UDebugDrawService::Unregister(Handle);
+		});
+
+		CHECK_EQ(bHit, true);
+		CHECK_EQ(bIsClosed, true);
+
+		return true;
+	});
+
+	Section(TEXT("Transact(Register, Unregister), Draw"), [&]
+	{
+		bool bHit = false;
+
+		AutoRTFM::Transact([&]
+		{
+			FDebugDrawDelegate Delegate;
+			Delegate.BindLambda([&](UCanvas*, APlayerController*)
+			{
+				bHit = true;
+			});
+
+			FDelegateHandle Handle = UDebugDrawService::Register(TEXT("Tonemapper"), Delegate);
+
+			UDebugDrawService::Unregister(Handle);
+		});
+
+		FEngineShowFlags EngineShowFlags(EShowFlagInitMode::ESFIM_Game);
+
+		EngineShowFlags.SetTonemapper(true);
+
+		UDebugDrawService::Draw(EngineShowFlags, CanvasObject);
+
+		CHECK_EQ(bHit, false);
+
+		return true;
+	});
+
+	Section(TEXT("Transact(Register, Abort), Draw"), [&]
+	{
+		bool bHit = false;
+
+		AutoRTFM::Transact([&]
+		{
+			FDebugDrawDelegate Delegate;
+			Delegate.BindLambda([&](UCanvas*, APlayerController*)
+			{
+				bHit = true;
+			});
+
+			FDelegateHandle Handle = UDebugDrawService::Register(TEXT("Tonemapper"), Delegate);
+
+			AutoRTFM::AbortTransaction();
+		});
+
+		FEngineShowFlags EngineShowFlags(EShowFlagInitMode::ESFIM_Game);
+
+		EngineShowFlags.SetTonemapper(true);
+
+		UDebugDrawService::Draw(EngineShowFlags, CanvasObject);
+
+		CHECK_EQ(bHit, false);
+
+		return true;
+	});
+
+	return true;
+}
+
+#undef CHECK_EQ
+
+#endif //WITH_DEV_AUTOMATION_TESTS
