Removed unused cvar r.MaterialParameterLegacyChecks.

diff --git a/Engine/Source/Runtime/Engine/Private/Materials/Material.cpp b/Engine/Source/Runtime/Engine/Private/Materials/Material.cpp
--- a/Engine/Source/Runtime/Engine/Private/Materials/Material.cpp
+++ b/Engine/Source/Runtime/Engine/Private/Materials/Material.cpp
@@ -116,13 +116,6 @@
 
 #define LOCTEXT_NAMESPACE "Material"
 
-static TAutoConsoleVariable<int32> CVarMaterialParameterLegacyChecks(
-	TEXT("r.MaterialParameterLegacyChecks"),
-	0,
-	TEXT("When enabled, sanity check new material parameter logic against legacy path.\n")
-	TEXT("Note that this can be slow"),
-	ECVF_RenderThreadSafe);
-
 static TAutoConsoleVariable<bool> CVarMaterialLogErrorOnFailure(
 	TEXT("r.MaterialLogErrorOnFailure"),
 	false,
