[Engine]
* Added state stream path for SkyAtmosphereComponent.

diff --git a/Engine/Source/Runtime/Engine/Classes/Components/SkyAtmosphereComponent.h b/Engine/Source/Runtime/Engine/Classes/Components/SkyAtmosphereComponent.h
--- a/Engine/Source/Runtime/Engine/Classes/Components/SkyAtmosphereComponent.h
+++ b/Engine/Source/Runtime/Engine/Classes/Components/SkyAtmosphereComponent.h
@@ -8,6 +8,7 @@
 #include "GameFramework/Info.h"
 #include "Misc/Guid.h"
 #include "RenderResource.h"
+#include "StateStream/SkyAtmosphereStateStreamHandle.h"
 
 #include "SkyAtmosphereComponent.generated.h"
 
@@ -308,6 +309,10 @@ protected:
 	 * when a map contains the deprecated AtmosphericFogComponent.
 	 */
 	void SetDummyStaticLightingGUIDs();
+
+	#if WITH_STATE_STREAM_ACTOR
+	FSkyAtmosphereHandle Handle;
+	#endif
 };
 
 
diff --git a/Engine/Source/Runtime/Engine/Private/Components/SkyAtmosphereComponent.cpp b/Engine/Source/Runtime/Engine/Private/Components/SkyAtmosphereComponent.cpp
--- a/Engine/Source/Runtime/Engine/Private/Components/SkyAtmosphereComponent.cpp
+++ b/Engine/Source/Runtime/Engine/Private/Components/SkyAtmosphereComponent.cpp
@@ -14,6 +14,7 @@
 #include "MeshElementCollector.h"
 #include "SceneInterface.h"
 #include "SceneProxies/SkyAtmosphereSceneProxy.h"
+#include "StateStream/SkyAtmosphereStateStream.h"
 #include "UObject/UObjectIterator.h"
 #include "PrimitiveDrawingUtils.h"
 #include "UObject/ConstructorHelpers.h"
@@ -30,12 +31,61 @@
 
 #define LOCTEXT_NAMESPACE "SkyAtmosphereComponent"
 
+#if WITH_STATE_STREAM_ACTOR
+constexpr bool UseSkyAtmosphereStateStream = true;
+#endif
 
 
 /*=============================================================================
 	USkyAtmosphereComponent implementation.
 =============================================================================*/
 
+//#define USKYATMOSPHERE_VARIABLE(Type, Owner, Dot, Prefix, Name)
+
+#define USKYATMOSPHERE_VARIABLES \
+		USKYATMOSPHERE_VARIABLE(ESkyAtmosphereTransformMode, , , , TransformMode) \
+		USKYATMOSPHERE_VARIABLE(float, , , , TraceSampleCountScale) \
+		USKYATMOSPHERE_VARIABLE(float, OtherTentDistribution, ., , TipAltitude) \
+		USKYATMOSPHERE_VARIABLE(float, OtherTentDistribution, ., , TipValue) \
+		USKYATMOSPHERE_VARIABLE(float, OtherTentDistribution, ., , Width) \
+		USKYATMOSPHERE_VARIABLE(uint8, , , b, Holdout) \
+		USKYATMOSPHERE_VARIABLE(uint8, , , b, RenderInMainPass) \
+		USKYATMOSPHERE_VARIABLES_WITH_SETFUNCTION
+
+#define USKYATMOSPHERE_VARIABLES_WITH_SETFUNCTION \
+		USKYATMOSPHERE_VARIABLE(float, , , , BottomRadius) \
+		USKYATMOSPHERE_VARIABLE(const FColor&, , , , GroundAlbedo) \
+		USKYATMOSPHERE_VARIABLE(float, , , , AtmosphereHeight) \
+		USKYATMOSPHERE_VARIABLE(float, , , , MultiScatteringFactor) \
+		USKYATMOSPHERE_VARIABLE(float, , , , RayleighScatteringScale) \
+		USKYATMOSPHERE_VARIABLE(FLinearColor, , , , RayleighScattering) \
+		USKYATMOSPHERE_VARIABLE(float, , , , RayleighExponentialDistribution) \
+		USKYATMOSPHERE_VARIABLE(float, , , , MieScatteringScale) \
+		USKYATMOSPHERE_VARIABLE(FLinearColor, , , , MieScattering) \
+		USKYATMOSPHERE_VARIABLE(float, , , , MieAbsorptionScale) \
+		USKYATMOSPHERE_VARIABLE(FLinearColor, , , , MieAbsorption) \
+		USKYATMOSPHERE_VARIABLE(float, , , , MieAnisotropy) \
+		USKYATMOSPHERE_VARIABLE(float, , , , MieExponentialDistribution) \
+		USKYATMOSPHERE_VARIABLE(float, , , , OtherAbsorptionScale) \
+		USKYATMOSPHERE_VARIABLE(FLinearColor, , , , OtherAbsorption) \
+		USKYATMOSPHERE_VARIABLE(FLinearColor, , , , SkyLuminanceFactor) \
+		USKYATMOSPHERE_VARIABLE(FLinearColor, , , , SkyAndAerialPerspectiveLuminanceFactor) \
+		USKYATMOSPHERE_VARIABLE(float, , , , AerialPespectiveViewDistanceScale) \
+		USKYATMOSPHERE_VARIABLE(float, , , , AerialPerspectiveStartDepth) \
+		USKYATMOSPHERE_VARIABLE(float, , , , HeightFogContribution) \
+		USKYATMOSPHERE_VARIABLE(float, , , , TransmittanceMinLightElevationAngle) \
+
+FOverrideAtmosphericLight ConvertAtmosphericLightOverride(bool* Enabled, FVector* Direction)
+{
+	FOverrideAtmosphericLight Override;
+	for (uint32 I=0; I!=NUM_ATMOSPHERE_LIGHTS; ++I)
+	{
+		Override.EnabledMask |= uint8(Enabled[I]) << I;
+		Override.Direction[I] = Direction[I];
+	}
+	return Override;
+}
+
 USkyAtmosphereComponent::USkyAtmosphereComponent(const FObjectInitializer& ObjectInitializer)
 	: Super(ObjectInitializer)
 	, SkyAtmosphereSceneProxy(nullptr)
@@ -163,9 +213,25 @@ void USkyAtmosphereComponent::CreateRenderState_Concurrent(FRegisterComponentCon
 	if (GetVisibleFlag() && !bHidden &&
 		ShouldComponentAddToScene() && ShouldRender() && IsRegistered() && (GetOuter() == NULL || !GetOuter()->HasAnyFlags(RF_ClassDefaultObject)))
 	{
+		bool bBuilt = SkyAtmosphereComponentStaticLightingBuilt(this);
+		#if WITH_STATE_STREAM_ACTOR
+		if (UseSkyAtmosphereStateStream)
+		{
+			FSkyAtmosphereDynamicState Ds;
+			#define USKYATMOSPHERE_VARIABLE(Type, Owner, Dot, Prefix, Name) Ds.Set##Owner##Name(Owner Dot Prefix##Name);
+			USKYATMOSPHERE_VARIABLES
+			#undef USKYATMOSPHERE_VARIABLE
+			Ds.SetBuilt(bBuilt);
+			Ds.SetComponentTransform(GetComponentTransform());
+			Ds.SetOverrideAtmosphericLight(ConvertAtmosphericLightOverride(OverrideAtmosphericLight, OverrideAtmosphericLightDirection));
+			Handle = GetWorld()->GetStateStream<ISkyAtmosphereStateStream>().Game_CreateInstance({}, Ds);
+			return;
+		}
+		#endif
+
 		// Create the scene proxy.
 		SkyAtmosphereSceneProxy = new FSkyAtmosphereSceneProxy(this);
-		GetWorld()->Scene->AddSkyAtmosphere(SkyAtmosphereSceneProxy, SkyAtmosphereComponentStaticLightingBuilt(this));
+		GetWorld()->Scene->AddSkyAtmosphere(SkyAtmosphereSceneProxy, bBuilt);
 	}
 
 }
@@ -180,6 +246,10 @@ void USkyAtmosphereComponent::DestroyRenderState_Concurrent()
 {
 	Super::DestroyRenderState_Concurrent();
 
+	#if WITH_STATE_STREAM_ACTOR
+	Handle = {};
+	#endif
+
 	if (SkyAtmosphereSceneProxy)
 	{
 		GetWorld()->Scene->RemoveSkyAtmosphere(SkyAtmosphereSceneProxy);
@@ -309,9 +379,18 @@ void USkyAtmosphereComponent::OverrideAtmosphereLightDirection(int32 AtmosphereL
 	if (AreDynamicDataChangesAllowed() && SkyAtmosphereSceneProxy &&
 		(!OverrideAtmosphericLight[AtmosphereLightIndex] || OverrideAtmosphericLightDirection[AtmosphereLightIndex]!=LightDirection))
 	{
-		FSceneInterface* Scene = GetWorld()->Scene;
 		OverrideAtmosphericLight[AtmosphereLightIndex] = true;
 		OverrideAtmosphericLightDirection[AtmosphereLightIndex] = LightDirection;
+
+		#if WITH_STATE_STREAM_ACTOR
+		if (UseSkyAtmosphereStateStream)
+		{
+			FSkyAtmosphereDynamicState Ds;
+			Ds.SetOverrideAtmosphericLight(ConvertAtmosphericLightOverride(OverrideAtmosphericLight, OverrideAtmosphericLightDirection));
+			Handle.Update(Ds);
+			return;
+		}
+		#endif
 		MarkRenderStateDirty();
 	}
 }
@@ -358,56 +437,50 @@ void USkyAtmosphereComponent::SetPositionToMatchDeprecatedAtmosphericFog()
 	SetWorldLocation(FVector(0.0f, 0.0f, -100000.0f));
 }
 
-#define SKY_DECLARE_BLUEPRINT_SETFUNCTION(MemberType, MemberName) void USkyAtmosphereComponent::Set##MemberName(MemberType NewValue)\
-{\
-	if (AreDynamicDataChangesAllowed() && MemberName != NewValue)\
-	{\
-		MemberName = NewValue;\
-		MarkRenderStateDirty();\
-	}\
-}\
-
-#define SKY_DECLARE_BLUEPRINT_SETFUNCTION_LINEARCOEFFICIENT(MemberName) void USkyAtmosphereComponent::Set##MemberName(FLinearColor NewValue)\
-{\
-	if (AreDynamicDataChangesAllowed() && MemberName != NewValue)\
-	{\
-		MemberName = NewValue.GetClamped(0.0f, 1e38f); \
-		MarkRenderStateDirty();\
-	}\
-}\
-
-SKY_DECLARE_BLUEPRINT_SETFUNCTION(float, BottomRadius);
-SKY_DECLARE_BLUEPRINT_SETFUNCTION(const FColor&, GroundAlbedo);
-
-SKY_DECLARE_BLUEPRINT_SETFUNCTION(float, AtmosphereHeight);
-SKY_DECLARE_BLUEPRINT_SETFUNCTION(float, MultiScatteringFactor);
-
-SKY_DECLARE_BLUEPRINT_SETFUNCTION(float, RayleighScatteringScale);
-SKY_DECLARE_BLUEPRINT_SETFUNCTION_LINEARCOEFFICIENT(RayleighScattering);
-SKY_DECLARE_BLUEPRINT_SETFUNCTION(float, RayleighExponentialDistribution);
-
-SKY_DECLARE_BLUEPRINT_SETFUNCTION(float, MieScatteringScale);
-SKY_DECLARE_BLUEPRINT_SETFUNCTION_LINEARCOEFFICIENT(MieScattering);
-SKY_DECLARE_BLUEPRINT_SETFUNCTION(float, MieAbsorptionScale);
-SKY_DECLARE_BLUEPRINT_SETFUNCTION_LINEARCOEFFICIENT(MieAbsorption);
-SKY_DECLARE_BLUEPRINT_SETFUNCTION(float, MieAnisotropy);
-SKY_DECLARE_BLUEPRINT_SETFUNCTION(float, MieExponentialDistribution);
-
-SKY_DECLARE_BLUEPRINT_SETFUNCTION(float, OtherAbsorptionScale);
-SKY_DECLARE_BLUEPRINT_SETFUNCTION_LINEARCOEFFICIENT(OtherAbsorption);
-
-SKY_DECLARE_BLUEPRINT_SETFUNCTION_LINEARCOEFFICIENT(SkyLuminanceFactor);
-SKY_DECLARE_BLUEPRINT_SETFUNCTION_LINEARCOEFFICIENT(SkyAndAerialPerspectiveLuminanceFactor);
-SKY_DECLARE_BLUEPRINT_SETFUNCTION(float, AerialPespectiveViewDistanceScale);
-SKY_DECLARE_BLUEPRINT_SETFUNCTION(float, AerialPerspectiveStartDepth);
-SKY_DECLARE_BLUEPRINT_SETFUNCTION(float, HeightFogContribution);
-SKY_DECLARE_BLUEPRINT_SETFUNCTION(float, TransmittanceMinLightElevationAngle);
+FLinearColor SkyAtmosphereGetClamped(FLinearColor Value) { return Value.GetClamped(0.0f, 1e38f); }
+template<typename T> T SkyAtmosphereGetClamped(T Value) { return Value; }
+
+#if WITH_STATE_STREAM_ACTOR
+	#define USKYATMOSPHERE_SETHANDLE(Field) \
+		if (UseSkyAtmosphereStateStream) \
+		{ \
+			FSkyAtmosphereDynamicState Ds; \
+			Ds.Set##Field(Field); \
+			Handle.Update(Ds); \
+			return; \
+		}
+#else
+	#define USKYATMOSPHERE_SETHANDLE(Field)
+#endif
+
+#define USKYATMOSPHERE_VARIABLE(Type, Owner, Dot, Prefix, Name) \
+	void USkyAtmosphereComponent::Set##Name(Type NewValue) \
+	{ \
+		if (AreDynamicDataChangesAllowed() && Name != NewValue) \
+		{ \
+			Name = SkyAtmosphereGetClamped(NewValue); \
+			USKYATMOSPHERE_SETHANDLE(Name) \
+			MarkRenderStateDirty(); \
+		} \
+	}
+
+USKYATMOSPHERE_VARIABLES_WITH_SETFUNCTION
+#undef USKYATMOSPHERE_VARIABLE
 
 void USkyAtmosphereComponent::SetHoldout(bool bNewHoldout)
 {
 	if (bHoldout != bNewHoldout)
 	{
 		bHoldout = bNewHoldout;
+#if WITH_STATE_STREAM_ACTOR
+		if (UseSkyAtmosphereStateStream)
+		{
+			FSkyAtmosphereDynamicState Ds;
+			Ds.SetHoldout(bNewHoldout);
+			Handle.Update(Ds);
+			return;
+		}
+#endif
 		MarkRenderStateDirty();
 	}
 }
@@ -417,6 +490,15 @@ void USkyAtmosphereComponent::SetRenderInMainPass(bool bValue)
 	if (bRenderInMainPass != bValue)
 	{
 		bRenderInMainPass = bValue;
+#if WITH_STATE_STREAM_ACTOR
+		if (UseSkyAtmosphereStateStream)
+		{
+			FSkyAtmosphereDynamicState Ds;
+			Ds.SetRenderInMainPass(bValue);
+			Handle.Update(Ds);
+			return;
+		}
+#endif
 		MarkRenderStateDirty();
 	}
 }
@@ -525,6 +607,27 @@ FSkyAtmosphereSceneProxy::FSkyAtmosphereSceneProxy(const USkyAtmosphereComponent
 	InComponent->GetOverrideLightStatus(OverrideAtmosphericLight, OverrideAtmosphericLightDirection);
 }
 
+FSkyAtmosphereSceneProxy::FSkyAtmosphereSceneProxy(const FSkyAtmosphereDynamicState& Ds)
+	: bStaticLightingBuilt(false)
+	, AtmosphereSetup(Ds)
+	, bHoldout(Ds.bHoldout > 0)
+	, bRenderInMainPass(Ds.bRenderInMainPass > 0)
+{
+	SkyLuminanceFactor = Ds.SkyLuminanceFactor;
+	SkyAndAerialPerspectiveLuminanceFactor = Ds.SkyAndAerialPerspectiveLuminanceFactor;
+	AerialPespectiveViewDistanceScale = Ds.AerialPespectiveViewDistanceScale;
+	HeightFogContribution = Ds.HeightFogContribution;
+	AerialPerspectiveStartDepthKm = Ds.AerialPerspectiveStartDepth;
+	TraceSampleCountScale = Ds.TraceSampleCountScale;
+
+	const FOverrideAtmosphericLight& Overrides = Ds.OverrideAtmosphericLight;
+	for (uint32 I=0; I!=NUM_ATMOSPHERE_LIGHTS; ++I)
+	{
+		OverrideAtmosphericLight[I] = (Overrides.EnabledMask & (1 << I)) != 0;
+		OverrideAtmosphericLightDirection[I] = Overrides.Direction[I];
+	}
+}
+
 FSkyAtmosphereSceneProxy::~FSkyAtmosphereSceneProxy()
 {
 }
diff --git a/Engine/Source/Runtime/Engine/Public/Rendering/SkyAtmosphereCommonData.cpp b/Engine/Source/Runtime/Engine/Public/Rendering/SkyAtmosphereCommonData.cpp
--- a/Engine/Source/Runtime/Engine/Public/Rendering/SkyAtmosphereCommonData.cpp
+++ b/Engine/Source/Runtime/Engine/Public/Rendering/SkyAtmosphereCommonData.cpp
@@ -8,13 +8,24 @@
 
 #include "Components/SkyAtmosphereComponent.h"
 #include "ColorManagement/ColorSpace.h"
+#include "StateStream/SkyAtmosphereStateStream.h"
 
 //PRAGMA_DISABLE_OPTIMIZATION
 
 const float FAtmosphereSetup::CmToSkyUnit = 0.00001f;			// Centimeters to Kilometers
 const float FAtmosphereSetup::SkyUnitToCm = 1.0f / 0.00001f;	// Kilometers to Centimeters
 
-FAtmosphereSetup::FAtmosphereSetup(const USkyAtmosphereComponent& SkyAtmosphereComponent)
+FTentDistribution GetTentDistribution(const USkyAtmosphereComponent& SkyAtmosphereComponent) { return SkyAtmosphereComponent.OtherTentDistribution; }
+FTentDistribution GetTentDistribution(const FSkyAtmosphereDynamicState& Ds)
+{
+	FTentDistribution D;
+	D.TipAltitude = Ds.OtherTentDistributionTipAltitude;
+	D.TipValue = Ds.OtherTentDistributionTipValue;
+	D.Width = Ds.OtherTentDistributionWidth;
+	return D;
+}
+
+template<typename T> void FAtmosphereSetup::InternalInit(const T& SkyAtmosphereComponent)
 {
 	// Convert Tent distribution to linear curve coefficients.
 	auto TentToCoefficients = [](const FTentDistribution& Tent, float& LayerWidth, float& LinTerm0, float&  LinTerm1, float& ConstTerm0, float& ConstTerm1)
@@ -100,7 +111,7 @@ FAtmosphereSetup::FAtmosphereSetup(const USkyAtmosphereComponent& SkyAtmosphereC
 		AbsorptionExtinction = (SkyAtmosphereComponent.OtherAbsorption * SkyAtmosphereComponent.OtherAbsorptionScale).GetClamped(0.0f, 1e38f);
 		AbsorptionExtinction = ConvertCoefficientsFromSRGBToWorkingColorSpace(AbsorptionExtinction);
 
-		TentToCoefficients(SkyAtmosphereComponent.OtherTentDistribution, AbsorptionDensity0LayerWidth, AbsorptionDensity0LinearTerm, AbsorptionDensity1LinearTerm, AbsorptionDensity0ConstantTerm, AbsorptionDensity1ConstantTerm);
+		TentToCoefficients(GetTentDistribution(SkyAtmosphereComponent), AbsorptionDensity0LayerWidth, AbsorptionDensity0LinearTerm, AbsorptionDensity1LinearTerm, AbsorptionDensity0ConstantTerm, AbsorptionDensity1ConstantTerm);
 	}
 
 	TransmittanceMinLightElevationAngle = SkyAtmosphereComponent.TransmittanceMinLightElevationAngle;
@@ -108,6 +119,17 @@ FAtmosphereSetup::FAtmosphereSetup(const USkyAtmosphereComponent& SkyAtmosphereC
 	UpdateTransform(SkyAtmosphereComponent.GetComponentTransform(), uint8(SkyAtmosphereComponent.TransformMode));
 }
 
+
+FAtmosphereSetup::FAtmosphereSetup(const USkyAtmosphereComponent& SkyAtmosphereComponent)
+{
+	InternalInit(SkyAtmosphereComponent);
+}
+
+FAtmosphereSetup::FAtmosphereSetup(const FSkyAtmosphereDynamicState& Ds)
+{
+	InternalInit(Ds);
+}
+
 void FAtmosphereSetup::ApplyWorldOffset(const FVector& InOffset)
 {
 	PlanetCenterKm += InOffset * double(FAtmosphereSetup::CmToSkyUnit);
diff --git a/Engine/Source/Runtime/Engine/Public/Rendering/SkyAtmosphereCommonData.h b/Engine/Source/Runtime/Engine/Public/Rendering/SkyAtmosphereCommonData.h
--- a/Engine/Source/Runtime/Engine/Public/Rendering/SkyAtmosphereCommonData.h
+++ b/Engine/Source/Runtime/Engine/Public/Rendering/SkyAtmosphereCommonData.h
@@ -9,7 +9,7 @@
 #include "CoreMinimal.h"
 
 
-
+struct FSkyAtmosphereDynamicState;
 class USkyAtmosphereComponent;
 
 
@@ -46,6 +46,9 @@ struct FAtmosphereSetup
 	float TransmittanceMinLightElevationAngle;
 
 	ENGINE_API FAtmosphereSetup(const USkyAtmosphereComponent& SkyAtmosphereComponent);
+	ENGINE_API FAtmosphereSetup(const FSkyAtmosphereDynamicState& Ds);
+
+	template<typename T> void InternalInit(const T& SkyAtmosphereComponent);
 
 	ENGINE_API FLinearColor GetTransmittanceAtGroundLevel(const FVector& SunDirection) const;
 
diff --git a/Engine/Source/Runtime/Engine/Public/SceneProxies/SkyAtmosphereSceneProxy.h b/Engine/Source/Runtime/Engine/Public/SceneProxies/SkyAtmosphereSceneProxy.h
--- a/Engine/Source/Runtime/Engine/Public/SceneProxies/SkyAtmosphereSceneProxy.h
+++ b/Engine/Source/Runtime/Engine/Public/SceneProxies/SkyAtmosphereSceneProxy.h
@@ -5,6 +5,7 @@
 #include "EngineDefines.h"
 #include "Rendering/SkyAtmosphereCommonData.h"
 
+struct FSkyAtmosphereDynamicState;
 class FSkyAtmosphereRenderSceneInfo;
 class USkyAtmosphereComponent;
 
@@ -17,6 +18,7 @@ public:
 
 	// Initialization constructor.
 	ENGINE_API FSkyAtmosphereSceneProxy(const USkyAtmosphereComponent* InComponent);
+	ENGINE_API FSkyAtmosphereSceneProxy(const FSkyAtmosphereDynamicState& Ds);
 	ENGINE_API ~FSkyAtmosphereSceneProxy();
 
 	FLinearColor GetSkyLuminanceFactor() const { return SkyLuminanceFactor; }
@@ -53,6 +55,8 @@ private:
 
 	bool OverrideAtmosphericLight[NUM_ATMOSPHERE_LIGHTS];
 	FVector OverrideAtmosphericLightDirection[NUM_ATMOSPHERE_LIGHTS];
+
+	friend class FSkyAtmosphereStateStreamImpl;
 };
 
 ////////////////////////////////////////////////////////////////////////////////////////////////////
diff --git a/Engine/Source/Runtime/Engine/Public/StateStream/SkyAtmosphereStateStream.h b/Engine/Source/Runtime/Engine/Public/StateStream/SkyAtmosphereStateStream.h
new file mode 100644
--- /dev/null
+++ b/Engine/Source/Runtime/Engine/Public/StateStream/SkyAtmosphereStateStream.h
@@ -0,0 +1,135 @@
+// Copyright Epic Games, Inc. All Rights Reserved.
+
+#pragma once
+
+#include "EngineDefines.h"
+#include "SkyAtmosphereStateStreamHandle.h"
+#include "StateStreamDefinitions.h"
+#include "TransformStateStreamHandle.h"
+#include "TransformStateStreamMath.h"
+#include "SkyAtmosphereStateStream.generated.h"
+
+enum class ESkyAtmosphereTransformMode : uint8;
+
+////////////////////////////////////////////////////////////////////////////////////////////////////
+// Static state for mesh instance. Can only be set upon creation
+
+USTRUCT(StateStreamStaticState)
+struct FSkyAtmosphereStaticState
+{
+	GENERATED_USTRUCT_BODY()
+};
+
+
+////////////////////////////////////////////////////////////////////////////////////////////////////
+// Dynamic state for mesh instance. Can be updated inside ticks
+
+USTRUCT()
+struct FOverrideAtmosphericLight
+{
+	GENERATED_USTRUCT_BODY()
+
+	UPROPERTY()
+	uint8 EnabledMask;
+
+	UPROPERTY()
+	FVector Direction[NUM_ATMOSPHERE_LIGHTS];
+
+	bool operator==(const FOverrideAtmosphericLight& O) const { return memcmp(this, &O, sizeof(FOverrideAtmosphericLight)) == 0; }
+};
+
+USTRUCT(StateStreamDynamicState)
+struct FSkyAtmosphereDynamicState
+{
+	GENERATED_USTRUCT_BODY()
+
+// private: // Should be private but complicates some of the code that needs support both this path and UComponent path
+
+	UPROPERTY()
+	ESkyAtmosphereTransformMode TransformMode;
+	UPROPERTY()
+	float BottomRadius;
+	UPROPERTY()
+	FColor GroundAlbedo;
+	UPROPERTY()
+	float AtmosphereHeight;
+	UPROPERTY()
+	float MultiScatteringFactor;
+	UPROPERTY()
+	float TraceSampleCountScale;
+	UPROPERTY()
+	float RayleighScatteringScale;
+	UPROPERTY()
+	FLinearColor RayleighScattering;
+	UPROPERTY()
+	float RayleighExponentialDistribution;
+	UPROPERTY()
+	float MieScatteringScale;
+	UPROPERTY()
+	FLinearColor MieScattering;
+	UPROPERTY()
+	float MieAbsorptionScale;
+	UPROPERTY()
+	FLinearColor MieAbsorption;
+	UPROPERTY()
+	float MieAnisotropy;
+	UPROPERTY()
+	float MieExponentialDistribution;
+	UPROPERTY()
+	float OtherAbsorptionScale;
+	UPROPERTY()
+	FLinearColor OtherAbsorption;
+	UPROPERTY()
+	float OtherTentDistributionTipAltitude = 0.0f;
+	UPROPERTY()
+	float OtherTentDistributionTipValue = 0.0f;
+	UPROPERTY()
+	float OtherTentDistributionWidth = 1.0f;
+	UPROPERTY()
+	FLinearColor SkyLuminanceFactor;
+	UPROPERTY()
+	FLinearColor SkyAndAerialPerspectiveLuminanceFactor;
+	UPROPERTY()
+	float AerialPespectiveViewDistanceScale;
+	UPROPERTY()
+	float HeightFogContribution;
+	UPROPERTY()
+	float TransmittanceMinLightElevationAngle;
+	UPROPERTY()
+	float AerialPerspectiveStartDepth;
+	UPROPERTY()
+	uint8 bHoldout : 1;
+	UPROPERTY()
+	uint8 bRenderInMainPass : 1;
+	UPROPERTY()
+	uint8 bBuilt : 1;
+
+	UPROPERTY()
+	FTransform ComponentTransform;
+
+	UPROPERTY()
+	FOverrideAtmosphericLight OverrideAtmosphericLight;
+
+	//UPROPERTY()
+	//bool OverrideAtmosphericLight[NUM_ATMOSPHERE_LIGHTS];
+
+	//UPROPERTY()
+	//FVector OverrideAtmosphericLightDirection[NUM_ATMOSPHERE_LIGHTS];
+};
+
+
+////////////////////////////////////////////////////////////////////////////////////////////////////
+// Mesh state stream id used for registering dependencies and find statestream
+
+inline constexpr uint32 SkyAtmosphereStateStreamId = 6;
+
+
+////////////////////////////////////////////////////////////////////////////////////////////////////
+// Interface for creating mesh instances
+
+class ISkyAtmosphereStateStream
+{
+public:
+	DECLARE_STATESTREAM(SkyAtmosphere)
+	virtual FSkyAtmosphereHandle Game_CreateInstance(const FSkyAtmosphereStaticState& Ss, const FSkyAtmosphereDynamicState& Ds) = 0;
+};
diff --git a/Engine/Source/Runtime/Engine/Public/StateStream/SkyAtmosphereStateStreamHandle.h b/Engine/Source/Runtime/Engine/Public/StateStream/SkyAtmosphereStateStreamHandle.h
new file mode 100644
--- /dev/null
+++ b/Engine/Source/Runtime/Engine/Public/StateStream/SkyAtmosphereStateStreamHandle.h
@@ -0,0 +1,19 @@
+// Copyright Epic Games, Inc. All Rights Reserved.
+
+#pragma once
+
+#include "StateStreamHandle.h"
+#include "SkyAtmosphereStateStreamHandle.generated.h"
+
+struct FSkyAtmosphereDynamicState;
+
+////////////////////////////////////////////////////////////////////////////////////////////////////
+// Handle for mesh instance
+
+USTRUCT(StateStreamHandle)
+struct FSkyAtmosphereHandle : public FStateStreamHandle
+{
+	GENERATED_USTRUCT_BODY()
+};
+
+////////////////////////////////////////////////////////////////////////////////////////////////////
diff --git a/Engine/Source/Runtime/Renderer/Private/StateStream/SkyAtmosphereStateStreamImpl.cpp b/Engine/Source/Runtime/Renderer/Private/StateStream/SkyAtmosphereStateStreamImpl.cpp
new file mode 100644
--- /dev/null
+++ b/Engine/Source/Runtime/Renderer/Private/StateStream/SkyAtmosphereStateStreamImpl.cpp
@@ -0,0 +1,46 @@
+// Copyright Epic Games, Inc. All Rights Reserved.
+
+#include "SkyAtmosphereStateStreamImpl.h"
+#include "SceneInterface.h"
+#include "SceneProxies/SkyAtmosphereSceneProxy.h"
+#include "StateStreamCreator.h"
+
+////////////////////////////////////////////////////////////////////////////////////////////////////
+
+FSkyAtmosphereStateStreamImpl::FSkyAtmosphereStateStreamImpl(FSceneInterface& InScene)
+:	Scene(InScene)
+{
+}
+
+void FSkyAtmosphereStateStreamImpl::Render_OnCreate(const FSkyAtmosphereStaticState& Ss, const FSkyAtmosphereDynamicState& Ds, FSkyAtmosphereSceneProxy*& UserData, bool IsDestroyedInSameFrame)
+{
+	UserData = new FSkyAtmosphereSceneProxy(Ds);
+	Scene.AddSkyAtmosphere(UserData, Ds.GetBuilt());
+}
+
+void FSkyAtmosphereStateStreamImpl::Render_OnUpdate(const FSkyAtmosphereStaticState& Ss, const FSkyAtmosphereDynamicState& Ds, FSkyAtmosphereSceneProxy*& UserData)
+{
+	// For some reason we need to recreate the entire sceneproxy every time
+#if 0
+	UserData->AtmosphereSetup = FAtmosphereSetup(Ds);
+	Scene.InvalidatePathTracedOutput();
+#else
+	Scene.RemoveSkyAtmosphere(UserData);
+	delete UserData;
+	UserData = new FSkyAtmosphereSceneProxy(Ds);
+	Scene.AddSkyAtmosphere(UserData, Ds.GetBuilt());
+#endif
+}
+
+void FSkyAtmosphereStateStreamImpl::Render_OnDestroy(const FSkyAtmosphereStaticState& Ss, const FSkyAtmosphereDynamicState& Ds, FSkyAtmosphereSceneProxy*& UserData)
+{
+	Scene.RemoveSkyAtmosphere(UserData);
+	delete UserData;
+	UserData = nullptr;
+}
+
+////////////////////////////////////////////////////////////////////////////////////////////////////
+
+STATESTREAM_CREATOR_INSTANCE(FSkyAtmosphereStateStreamImpl)
+
+////////////////////////////////////////////////////////////////////////////////////////////////////
diff --git a/Engine/Source/Runtime/Renderer/Private/StateStream/SkyAtmosphereStateStreamImpl.h b/Engine/Source/Runtime/Renderer/Private/StateStream/SkyAtmosphereStateStreamImpl.h
new file mode 100644
--- /dev/null
+++ b/Engine/Source/Runtime/Renderer/Private/StateStream/SkyAtmosphereStateStreamImpl.h
@@ -0,0 +1,36 @@
+// Copyright Epic Games, Inc. All Rights Reserved.
+
+#pragma once
+
+#include "GenericStateStream.h"
+#include "StateStream/SkyAtmosphereStateStream.h"
+
+#define UE_API RENDERER_API
+
+class FSceneInterface;
+class FSkyAtmosphereSceneProxy;
+
+////////////////////////////////////////////////////////////////////////////////////////////////////
+
+struct FSkyAtmosphereStateStreamSettings : TStateStreamSettings<ISkyAtmosphereStateStream, FSkyAtmosphereSceneProxy>
+{
+	static inline constexpr bool SkipCreatingDeletes = true;
+};
+
+////////////////////////////////////////////////////////////////////////////////////////////////////
+
+class FSkyAtmosphereStateStreamImpl : public TStateStream<FSkyAtmosphereStateStreamSettings>
+{
+public:
+	FSkyAtmosphereStateStreamImpl(FSceneInterface& InScene);
+private:
+	UE_API virtual void Render_OnCreate(const FSkyAtmosphereStaticState& Ss, const FSkyAtmosphereDynamicState& Ds, FSkyAtmosphereSceneProxy*& UserData, bool IsDestroyedInSameFrame) override;
+	UE_API virtual void Render_OnUpdate(const FSkyAtmosphereStaticState& Ss, const FSkyAtmosphereDynamicState& Ds, FSkyAtmosphereSceneProxy*& UserData) override;
+	UE_API virtual void Render_OnDestroy(const FSkyAtmosphereStaticState& Ss, const FSkyAtmosphereDynamicState& Ds, FSkyAtmosphereSceneProxy*& UserData) override;
+
+	FSceneInterface& Scene;
+};
+
+////////////////////////////////////////////////////////////////////////////////////////////////////
+
+#undef UE_API
