[HWRT] Fix crash due to issue with ray tracing geometry streaming.

Context:
- When static mesh streams out a LOD, the corresponding ray tracing geometry is also streamed out (SetRayTracingGeometryGroupCurrentFirstLODIndex(...))
- However if the ray tracing geometry had a pending streaming request, that request was not immediately cancelled and instead relied on logic in ProcessCompletedStreamingRequests(...) discarding the request once it completes.
- This could cause issues in very specific timing situations where we ended up with 2 in-flight requests for the same ray tracing geometry and ProcessCompletedStreamingRequests(...) failed to discard the older request.

Change:
- Explicitly cancel and reset streming-in requests in SetRayTracingGeometryGroupCurrentFirstLODIndex(...) when they're no longer required.
- added extra checks to detect such issues earlier.
- also added cvar to artificially increase IO latency when streaming ray tracing geometry which makes it easier to reproduce this sort of issue.

Repro steps:
- r.Streaming.StressTest.ExtaIOLatency 0
- r.Streaming.StressTest 1
- r.Streaming.StressTest.RayTracing.ExtaIOLatency 10

diff --git a/Engine/Source/Runtime/Engine/Private/Rendering/RayTracingGeometryManager.cpp b/Engine/Source/Runtime/Engine/Private/Rendering/RayTracingGeometryManager.cpp
--- a/Engine/Source/Runtime/Engine/Private/Rendering/RayTracingGeometryManager.cpp
+++ b/Engine/Source/Runtime/Engine/Private/Rendering/RayTracingGeometryManager.cpp
@@ -170,6 +170,7 @@ static float GetInitialBuildPriority(ERTAccelerationStructureBuildPriority InBui
 FRayTracingGeometryManager::BuildRequestIndex FRayTracingGeometryManager::RequestBuildAccelerationStructure(FRayTracingGeometry* InGeometry, ERTAccelerationStructureBuildPriority InPriority, EAccelerationStructureBuildMode InBuildMode)
 {
 	check(InGeometry->RayTracingBuildRequestIndex == INDEX_NONE);
+
 	FBuildRequest Request;
 	Request.BuildPriority = GetInitialBuildPriority(InPriority);
 	Request.Owner = InGeometry;
@@ -303,6 +304,8 @@ void FRayTracingGeometryManager::ReleaseRayTracingGeometryHandle(RayTracingGeome
 		FStreamingRequest& StreamingRequest = StreamingRequests[RegisteredGeometry.StreamingRequestIndex];
 		check(StreamingRequest.GeometryHandle == Handle);
 
+		// TODO: Reset on a worker thread 
+		// because it can cause the render thread to block waiting for IO request to complete
 		StreamingRequest.Reset();
 
 		RegisteredGeometry.StreamingRequestIndex = INDEX_NONE;
@@ -391,7 +394,22 @@ void FRayTracingGeometryManager::SetRayTracingGeometryGroupCurrentFirstLODIndex(
 
 			if (!RegisteredGeometry.Geometry->IsEvicted())
 			{
+				// TODO: call FRayTracingStreamableAsset::ReleaseForStreaming(...) instead
 				RegisteredGeometry.Geometry->ReleaseRHIForStreaming(Batcher);
+
+				if (RegisteredGeometry.StreamingRequestIndex != INDEX_NONE)
+				{
+					// cancel and reset pending streaming request
+
+					FStreamingRequest& StreamingRequest = StreamingRequests[RegisteredGeometry.StreamingRequestIndex];
+					checkf(StreamingRequest.GeometryHandle == GeometryHandle, TEXT("Unexpected ray tracing geometry streaming request owner."));
+
+					// TODO: Reset on a worker thread 
+					// because it can cause the render thread to block waiting for IO request to complete
+					StreamingRequest.Reset();
+
+					RegisteredGeometry.StreamingRequestIndex = INDEX_NONE;
+				}
 			}
 		}
 	}
@@ -574,6 +592,7 @@ void FRayTracingGeometryManager::Tick(FRHICommandList& RHICmdList)
 			{
 				if (RegisteredGeometry.Geometry->GetRHI() != nullptr)
 				{
+					// TODO: call FRayTracingStreamableAsset::ReleaseForStreaming(...)
 					RegisteredGeometry.Geometry->Evict();
 				}
 			}
@@ -750,6 +769,7 @@ void FRayTracingGeometryManager::Tick(FRHICommandList& RHICmdList)
 
 				check(RegisteredGeometry.Geometry->IsValid() && !RegisteredGeometry.Geometry->IsEvicted());
 
+				// TODO: call FRayTracingStreamableAsset::ReleaseForStreaming(...) instead
 				RegisteredGeometry.Geometry->Evict();
 
 				++Index;
@@ -900,6 +920,7 @@ bool FRayTracingGeometryManager::RequestRayTracingGeometryStreamIn(FRHICommandLi
 			return false;
 		}
 
+		checkf(RegisteredGeometry.StreamingRequestIndex == INDEX_NONE, TEXT("Ray Tracing Geometry already has a streaming request in-flight"));
 		RegisteredGeometry.StreamingRequestIndex = NextStreamingRequestIndex;
 
 		FStreamingRequest& StreamingRequest = StreamingRequests[NextStreamingRequestIndex];
diff --git a/Engine/Source/Runtime/Engine/Private/StaticMesh.cpp b/Engine/Source/Runtime/Engine/Private/StaticMesh.cpp
--- a/Engine/Source/Runtime/Engine/Private/StaticMesh.cpp
+++ b/Engine/Source/Runtime/Engine/Private/StaticMesh.cpp
@@ -82,6 +82,7 @@
 #include "Logging/MessageLog.h"
 #include "Misc/UObjectToken.h"
 #include "Streaming/StaticMeshUpdate.h"
+#include "Streaming/TextureStreamingHelpers.h"
 #include "ProfilingDebugging/IoStoreTrace.h"
 
 #include UE_INLINE_GENERATED_CPP_BY_NAME(StaticMesh)
@@ -1891,7 +1892,16 @@ void FStaticMeshRayTracingProxyLOD::IssueRequest(FBulkDataBatchRequest& Request,
 	{
 		Batch.Read(StreamableData, OfflineBVHOffset, OfflineBVHSize);
 	}
-	Batch.Issue(RequestBuffer, AIOP_Low, [](FBulkDataRequest::EStatus) {}, Request);
+	Batch.Issue(RequestBuffer, AIOP_Low, [](FBulkDataRequest::EStatus)
+		{
+#if !UE_BUILD_SHIPPING
+			// On some platforms the IO is too fast to test cancellation requests timing issues.
+			if (FRenderAssetStreamingSettings::ExtraIOLatencyRayTracing > 0)
+			{
+				FPlatformProcess::Sleep(FRenderAssetStreamingSettings::ExtraIOLatencyRayTracing * .001f); // Slow down the streaming.
+			}
+#endif
+		}, Request);
 }
 
 void FStaticMeshRayTracingProxyLOD::InitWithStreamedData(FRHICommandListBase& RHICmdList, FMemoryView StreamedData)
diff --git a/Engine/Source/Runtime/Engine/Private/Streaming/TextureStreamingHelpers.cpp b/Engine/Source/Runtime/Engine/Private/Streaming/TextureStreamingHelpers.cpp
--- a/Engine/Source/Runtime/Engine/Private/Streaming/TextureStreamingHelpers.cpp
+++ b/Engine/Source/Runtime/Engine/Private/Streaming/TextureStreamingHelpers.cpp
@@ -59,6 +59,7 @@ DEFINE_STAT(STAT_RenderAssetStreaming_GameThreadUpdateTime);
 DEFINE_LOG_CATEGORY(LogContentStreaming);
 
 int32 FRenderAssetStreamingSettings::ExtraIOLatency = 0;
+int32 FRenderAssetStreamingSettings::ExtraIOLatencyRayTracing = 0;
 
 ENGINE_API TAutoConsoleVariable<int32> CVarStreamingUseNewMetrics(
 	TEXT("r.Streaming.UseNewMetrics"),
@@ -289,7 +290,13 @@ static TAutoConsoleVariable<int32> CVarStreamingStressTest(
 static TAutoConsoleVariable<int32> CVarStreamingStressTestExtraIOLatency(
 	TEXT("r.Streaming.StressTest.ExtaIOLatency"),
 	10,
-	TEXT("An extra latency in milliseconds for each stream-in requests when doing the stress test."),
+	TEXT("An extra latency in milliseconds for each stream-in request when doing the stress test."),
+	ECVF_Cheat);
+
+static TAutoConsoleVariable<int32> CVarStreamingStressTestRayTracingExtraIOLatency(
+	TEXT("r.Streaming.StressTest.RayTracing.ExtaIOLatency"),
+	0,
+	TEXT("An extra latency in milliseconds for each ray tracing stream-in request when doing the stress test."),
 	ECVF_Cheat);
 
 static TAutoConsoleVariable<int32> CVarStreamingStressTestFramesForFullUpdate(
@@ -354,6 +361,7 @@ void FRenderAssetStreamingSettings::Update()
 		DropMips = 2; 
 		// Increase chances of canceling IO while they are not yet completed.
 		ExtraIOLatency = CVarStreamingStressTestExtraIOLatency.GetValueOnAnyThread();
+		ExtraIOLatencyRayTracing = CVarStreamingStressTestRayTracingExtraIOLatency.GetValueOnAnyThread();
 	}
     else
 #endif
@@ -362,6 +370,7 @@ void FRenderAssetStreamingSettings::Update()
 		FramesForFullUpdate = FMath::Max<int32>(CVarFramesForFullUpdate.GetValueOnAnyThread(), 0);
 		DropMips = CVarStreamingDropMips.GetValueOnAnyThread();
 		ExtraIOLatency = 0;
+		ExtraIOLatencyRayTracing = 0;
     }
 }
 
diff --git a/Engine/Source/Runtime/Engine/Private/Streaming/TextureStreamingHelpers.h b/Engine/Source/Runtime/Engine/Private/Streaming/TextureStreamingHelpers.h
--- a/Engine/Source/Runtime/Engine/Private/Streaming/TextureStreamingHelpers.h
+++ b/Engine/Source/Runtime/Engine/Private/Streaming/TextureStreamingHelpers.h
@@ -120,6 +120,7 @@ struct FRenderAssetStreamingSettings
 
 	bool bStressTest;
 	static int32 ExtraIOLatency;
+	static int32 ExtraIOLatencyRayTracing;
 
 	// Cached values of 
 	bool HighPriorityLoad_Texture[TEXTUREGROUP_MAX];
