[Engine]
* Renamed HeightFogComponent.cpp to ExponentialHeightFogComponent.cpp to match the header.

diff --git a/Engine/Source/Runtime/Engine/Private/Components/ExponentialHeightFogComponent.cpp b/Engine/Source/Runtime/Engine/Private/Components/ExponentialHeightFogComponent.cpp
new file mode 100644
--- /dev/null
+++ b/Engine/Source/Runtime/Engine/Private/Components/ExponentialHeightFogComponent.cpp
@@ -0,0 +1,495 @@
+// Copyright Epic Games, Inc. All Rights Reserved.
+
+/*=============================================================================
+	HeightFogComponent.cpp: Height fog implementation.
+=============================================================================*/
+
+#include "Components/ExponentialHeightFogComponent.h"
+#include "GameFramework/Info.h"
+#include "UObject/ConstructorHelpers.h"
+#include "Engine/World.h"
+#include "SceneInterface.h"
+#include "Engine/Texture2D.h"
+#include "Engine/ExponentialHeightFog.h"
+#include "Net/UnrealNetwork.h"
+#include "Components/BillboardComponent.h"
+#include "UObject/UE5MainStreamObjectVersion.h"
+
+#include UE_INLINE_GENERATED_CPP_BY_NAME(ExponentialHeightFogComponent)
+
+UExponentialHeightFogComponent::UExponentialHeightFogComponent(const FObjectInitializer& ObjectInitializer)
+	: Super(ObjectInitializer)
+{
+	FogInscatteringColor_DEPRECATED = FLinearColor(0.447f, 0.638f, 1.0f);
+	FogInscatteringLuminance = FLinearColor::Black;
+
+	SkyAtmosphereAmbientContributionColorScale = FLinearColor::White;
+
+	DirectionalInscatteringExponent = 4.0f;
+	DirectionalInscatteringStartDistance = 10000.0f;
+	DirectionalInscatteringColor_DEPRECATED = FLinearColor(0.25f, 0.25f, 0.125f);
+	DirectionalInscatteringLuminance = FLinearColor::Black;
+
+	InscatteringTextureTint = FLinearColor::White;
+	FullyDirectionalInscatteringColorDistance = 100000.0f;
+	NonDirectionalInscatteringColorDistance = 1000.0f;
+
+	FogDensity = 0.02f;
+	FogHeightFalloff = 0.2f;
+	// No influence from the second fog as default
+	SecondFogData.FogDensity = 0.0f;
+
+	FogMaxOpacity = 1.0f;
+	StartDistance = 0.0f;
+	EndDistance = 0.0f;
+
+	// disabled by default
+	FogCutoffDistance = 0;
+
+	bHoldout = false;
+	bRenderInMainPass = true;
+
+	VolumetricFogScatteringDistribution = .2f;
+	VolumetricFogAlbedo = FColor::White;
+	VolumetricFogExtinctionScale = 1.0f;
+	VolumetricFogDistance = 6000.0f;
+	VolumetricFogStaticLightingScatteringIntensity = 1;
+}
+
+void UExponentialHeightFogComponent::AddFogIfNeeded()
+{
+	// For safety, clamp the values for SecondFogData here.
+	SecondFogData.ClampToValidRanges();
+	
+	if (ShouldComponentAddToScene() && ShouldRender() && IsRegistered() && ((FogDensity + SecondFogData.FogDensity) * 1000) > UE_DELTA && FogMaxOpacity > UE_DELTA
+		&& (GetOuter() == NULL || !GetOuter()->HasAnyFlags(RF_ClassDefaultObject)))
+	{
+		GetWorld()->Scene->AddExponentialHeightFog(this);
+	}
+}
+
+void UExponentialHeightFogComponent::CreateRenderState_Concurrent(FRegisterComponentContext* Context)
+{
+	Super::CreateRenderState_Concurrent(Context);
+	AddFogIfNeeded();
+}
+
+void UExponentialHeightFogComponent::SendRenderTransform_Concurrent()
+{
+	GetWorld()->Scene->RemoveExponentialHeightFog(this);
+	AddFogIfNeeded();
+	Super::SendRenderTransform_Concurrent();
+}
+
+void UExponentialHeightFogComponent::DestroyRenderState_Concurrent()
+{
+	Super::DestroyRenderState_Concurrent();
+	GetWorld()->Scene->RemoveExponentialHeightFog(this);
+}
+
+#if WITH_EDITOR
+
+bool UExponentialHeightFogComponent::CanEditChange(const FProperty* InProperty) const
+{
+	if (InProperty)
+	{
+		FString PropertyName = InProperty->GetName();
+
+		static const auto CVarFog = IConsoleManager::Get().FindTConsoleVariableDataInt(TEXT("r.SupportExpFogMatchesVolumetricFog"));
+		if (CVarFog && CVarFog->GetValueOnAnyThread() > 0)
+		{
+			if (PropertyName == GET_MEMBER_NAME_STRING_CHECKED(UExponentialHeightFogComponent, DirectionalInscatteringExponent) ||
+				PropertyName == GET_MEMBER_NAME_STRING_CHECKED(UExponentialHeightFogComponent, DirectionalInscatteringStartDistance) ||
+				PropertyName == GET_MEMBER_NAME_STRING_CHECKED(UExponentialHeightFogComponent, FogInscatteringLuminance) ||
+				PropertyName == GET_MEMBER_NAME_STRING_CHECKED(UExponentialHeightFogComponent, DirectionalInscatteringLuminance))
+			{
+				// In this case, all the data will come from the volumetric fog.
+				return false;
+			}
+		}
+
+		if (PropertyName == GET_MEMBER_NAME_STRING_CHECKED(UExponentialHeightFogComponent, DirectionalInscatteringExponent) ||
+			PropertyName == GET_MEMBER_NAME_STRING_CHECKED(UExponentialHeightFogComponent, DirectionalInscatteringStartDistance) ||
+			PropertyName == GET_MEMBER_NAME_STRING_CHECKED(UExponentialHeightFogComponent, DirectionalInscatteringLuminance) ||
+			PropertyName == GET_MEMBER_NAME_STRING_CHECKED(UExponentialHeightFogComponent, FogInscatteringLuminance))
+		{
+			return !InscatteringColorCubemap;
+		}
+
+		if (PropertyName == GET_MEMBER_NAME_STRING_CHECKED(UExponentialHeightFogComponent, FullyDirectionalInscatteringColorDistance) ||
+			PropertyName == GET_MEMBER_NAME_STRING_CHECKED(UExponentialHeightFogComponent, NonDirectionalInscatteringColorDistance) ||
+			PropertyName == GET_MEMBER_NAME_STRING_CHECKED(UExponentialHeightFogComponent, InscatteringTextureTint) ||
+			PropertyName == GET_MEMBER_NAME_STRING_CHECKED(UExponentialHeightFogComponent, InscatteringColorCubemapAngle))
+		{
+			return InscatteringColorCubemap != NULL;
+		}
+
+		if (PropertyName == GET_MEMBER_NAME_STRING_CHECKED(UExponentialHeightFogComponent, FogInscatteringLuminance))
+		{
+			static const auto CVar = IConsoleManager::Get().FindTConsoleVariableDataInt(TEXT("r.SupportSkyAtmosphereAffectsHeightFog"));
+			return CVar && CVar->GetValueOnAnyThread() > 0;
+		}
+	}
+
+	return Super::CanEditChange(InProperty);
+}
+
+void UExponentialHeightFogComponent::PostEditChangeProperty(FPropertyChangedEvent& PropertyChangedEvent)
+{
+	SecondFogData.ClampToValidRanges();
+	FogDensity = FMath::Clamp(FogDensity, 0.0f, 10.0f);
+	FogHeightFalloff = FMath::Clamp(FogHeightFalloff, 0.0f, 2.0f);
+	FogMaxOpacity = FMath::Clamp(FogMaxOpacity, 0.0f, 1.0f);
+	StartDistance = FMath::Clamp(StartDistance, 0.0f, (float)WORLD_MAX);
+	EndDistance = FMath::Clamp(EndDistance, 0.0f, (float)(10 * WORLD_MAX));
+	FogCutoffDistance = FMath::Clamp(FogCutoffDistance, 0.0f, (float)(10 * WORLD_MAX));
+	FullyDirectionalInscatteringColorDistance = FMath::Clamp(FullyDirectionalInscatteringColorDistance, 0.0f, (float)WORLD_MAX);
+	NonDirectionalInscatteringColorDistance = FMath::Clamp(NonDirectionalInscatteringColorDistance, 0.0f, FullyDirectionalInscatteringColorDistance);
+	InscatteringColorCubemapAngle = FMath::Clamp(InscatteringColorCubemapAngle, 0.0f, 360.0f);
+
+	Super::PostEditChangeProperty(PropertyChangedEvent);
+}
+#endif // WITH_EDITOR
+
+void UExponentialHeightFogComponent::SetFogDensity(float Value)
+{
+	if(FogDensity != Value)
+	{
+		FogDensity = Value;
+		MarkRenderStateDirty();
+	}
+}
+
+void UExponentialHeightFogComponent::SetSecondFogDensity(float Value)
+{
+	if(SecondFogData.FogDensity != Value)
+	{
+		SecondFogData.FogDensity = Value;
+		MarkRenderStateDirty();
+	}
+}
+
+void UExponentialHeightFogComponent::SetFogInscatteringColor(FLinearColor Value)
+{
+	if(FogInscatteringLuminance != Value)
+	{
+		FogInscatteringLuminance = Value;
+		MarkRenderStateDirty();
+	}
+}
+
+void UExponentialHeightFogComponent::SetInscatteringColorCubemap(UTextureCube* Value)
+{
+	if(InscatteringColorCubemap != Value)
+	{
+		InscatteringColorCubemap = Value;
+		MarkRenderStateDirty();
+	}
+}
+
+void UExponentialHeightFogComponent::SetInscatteringColorCubemapAngle(float Value)
+{
+	if(InscatteringColorCubemapAngle != Value)
+	{
+		InscatteringColorCubemapAngle = Value;
+		MarkRenderStateDirty();
+	}
+}
+
+void UExponentialHeightFogComponent::SetFullyDirectionalInscatteringColorDistance(float Value)
+{
+	if(FullyDirectionalInscatteringColorDistance != Value)
+	{
+		FullyDirectionalInscatteringColorDistance = Value;
+		MarkRenderStateDirty();
+	}
+}
+
+void UExponentialHeightFogComponent::SetNonDirectionalInscatteringColorDistance(float Value)
+{
+	if(NonDirectionalInscatteringColorDistance != Value)
+	{
+		NonDirectionalInscatteringColorDistance = Value;
+		MarkRenderStateDirty();
+	}
+}
+
+void UExponentialHeightFogComponent::SetInscatteringTextureTint(FLinearColor Value)
+{
+	if(InscatteringTextureTint != Value)
+	{
+		InscatteringTextureTint = Value;
+		MarkRenderStateDirty();
+	}
+}
+
+void UExponentialHeightFogComponent::SetDirectionalInscatteringExponent(float Value)
+{
+	if(DirectionalInscatteringExponent != Value)
+	{
+		DirectionalInscatteringExponent = Value;
+		MarkRenderStateDirty();
+	}
+}
+
+void UExponentialHeightFogComponent::SetDirectionalInscatteringStartDistance(float Value)
+{
+	if(DirectionalInscatteringStartDistance != Value)
+	{
+		DirectionalInscatteringStartDistance = Value;
+		MarkRenderStateDirty();
+	}
+}
+
+void UExponentialHeightFogComponent::SetDirectionalInscatteringColor(FLinearColor Value)
+{
+	if(DirectionalInscatteringLuminance != Value)
+	{
+		DirectionalInscatteringLuminance = Value;
+		MarkRenderStateDirty();
+	}
+}
+
+void UExponentialHeightFogComponent::SetSecondFogHeightOffset(float Value)
+{
+	if(SecondFogData.FogHeightOffset != Value)
+	{
+		SecondFogData.FogHeightOffset = Value;
+		MarkRenderStateDirty();
+	}
+}
+
+void UExponentialHeightFogComponent::SetFogHeightFalloff(float Value)
+{
+	if(FogHeightFalloff != Value)
+	{
+		FogHeightFalloff = Value;
+		MarkRenderStateDirty();
+	}
+}
+
+void UExponentialHeightFogComponent::SetSecondFogHeightFalloff(float Value)
+{
+	if(SecondFogData.FogHeightFalloff != Value)
+	{
+		SecondFogData.FogHeightFalloff = Value;
+		MarkRenderStateDirty();
+	}
+}
+
+void UExponentialHeightFogComponent::SetFogMaxOpacity(float Value)
+{
+	if(FogMaxOpacity != Value)
+	{
+		FogMaxOpacity = Value;
+		MarkRenderStateDirty();
+	}
+}
+
+void UExponentialHeightFogComponent::SetStartDistance(float Value)
+{
+	if(StartDistance != Value)
+	{
+		StartDistance = Value;
+		MarkRenderStateDirty();
+	}
+}
+
+void UExponentialHeightFogComponent::SetEndDistance(float Value)
+{
+	if(EndDistance != Value)
+	{
+		EndDistance = Value;
+		MarkRenderStateDirty();
+	}
+}
+
+void UExponentialHeightFogComponent::SetFogCutoffDistance(float Value)
+{
+	if(FogCutoffDistance != Value)
+	{
+		FogCutoffDistance = Value;
+		MarkRenderStateDirty();
+	}
+}
+
+void UExponentialHeightFogComponent::SetVolumetricFog(bool bNewValue)
+{
+	if(bEnableVolumetricFog != bNewValue)
+	{
+		bEnableVolumetricFog = bNewValue;
+		MarkRenderStateDirty();
+	}
+}
+
+void UExponentialHeightFogComponent::SetVolumetricFogScatteringDistribution(float NewValue)
+{
+	if(VolumetricFogScatteringDistribution != NewValue)
+	{
+		VolumetricFogScatteringDistribution = NewValue;
+		MarkRenderStateDirty();
+	}
+}
+
+void UExponentialHeightFogComponent::SetVolumetricFogExtinctionScale(float NewValue)
+{
+	if (VolumetricFogExtinctionScale != NewValue)
+	{
+		VolumetricFogExtinctionScale = NewValue;
+		MarkRenderStateDirty();
+	}
+}
+
+void UExponentialHeightFogComponent::SetVolumetricFogAlbedo(FColor NewValue)
+{
+	if (VolumetricFogAlbedo != NewValue)
+	{
+		VolumetricFogAlbedo = NewValue;
+		MarkRenderStateDirty();
+	}
+}
+
+void UExponentialHeightFogComponent::SetVolumetricFogEmissive(FLinearColor NewValue)
+{
+	if (VolumetricFogEmissive != NewValue)
+	{
+		VolumetricFogEmissive = NewValue;
+		MarkRenderStateDirty();
+	}
+}
+
+void UExponentialHeightFogComponent::SetVolumetricFogDistance(float NewValue)
+{
+	if(VolumetricFogDistance != NewValue)
+	{
+		VolumetricFogDistance = NewValue;
+		MarkRenderStateDirty();
+	}
+}
+
+void UExponentialHeightFogComponent::SetVolumetricFogStartDistance(float NewValue)
+{
+	if (VolumetricFogStartDistance != NewValue)
+	{
+		VolumetricFogStartDistance = NewValue;
+		MarkRenderStateDirty();
+	}
+}
+
+void UExponentialHeightFogComponent::SetVolumetricFogNearFadeInDistance(float NewValue)
+{
+	if (VolumetricFogNearFadeInDistance != NewValue)
+	{
+		VolumetricFogNearFadeInDistance = NewValue;
+		MarkRenderStateDirty();
+	}
+}
+
+void UExponentialHeightFogComponent::SetSecondFogData(FExponentialHeightFogData NewValue)
+{
+	if(SecondFogData.FogDensity != NewValue.FogDensity ||
+	   SecondFogData.FogHeightOffset != NewValue.FogHeightOffset ||
+	   SecondFogData.FogHeightFalloff != NewValue.FogHeightFalloff)
+	{
+		SecondFogData = NewValue;
+		MarkRenderStateDirty();
+	}
+}
+
+void UExponentialHeightFogComponent::SetHoldout(bool bNewHoldout)
+{
+	if (bHoldout != bNewHoldout)
+	{
+		bHoldout = bNewHoldout;
+		MarkRenderStateDirty();
+	}
+}
+
+void UExponentialHeightFogComponent::SetRenderInMainPass(bool bValue)
+{
+	if (bRenderInMainPass != bValue)
+	{
+		bRenderInMainPass = bValue;
+		MarkRenderStateDirty();
+	}
+}
+
+void UExponentialHeightFogComponent::SetSkyAtmosphereAmbientContributionColorScale(FLinearColor NewValue)
+{
+	if (SkyAtmosphereAmbientContributionColorScale != NewValue)
+	{
+		SkyAtmosphereAmbientContributionColorScale = NewValue;
+		MarkRenderStateDirty();
+	}
+}
+
+void UExponentialHeightFogComponent::Serialize(FArchive& Ar)
+{
+	Super::Serialize(Ar);
+
+	Ar.UsingCustomVersion(FUE5MainStreamObjectVersion::GUID);
+
+	if (Ar.IsLoading() && (Ar.CustomVer(FUE5MainStreamObjectVersion::GUID) < FUE5MainStreamObjectVersion::SkyAtmosphereAffectsHeightFogWithBetterDefault))
+	{
+		FogInscatteringLuminance = FogInscatteringColor_DEPRECATED;
+		DirectionalInscatteringLuminance = DirectionalInscatteringColor_DEPRECATED;
+	}
+}
+
+//////////////////////////////////////////////////////////////////////////
+// AExponentialHeightFog
+
+AExponentialHeightFog::AExponentialHeightFog(const FObjectInitializer& ObjectInitializer)
+	: Super(ObjectInitializer)
+{
+	Component = CreateDefaultSubobject<UExponentialHeightFogComponent>(TEXT("HeightFogComponent0"));
+	RootComponent = Component;
+
+	SetHidden(false);
+
+#if WITH_EDITORONLY_DATA
+	if (!IsRunningCommandlet() && (GetSpriteComponent() != NULL))
+	{
+		// Structure to hold one-time initialization
+		struct FConstructorStatics
+		{
+			ConstructorHelpers::FObjectFinderOptional<UTexture2D> FogTextureObject;
+			FName ID_Fog;
+			FText NAME_Fog;
+			FConstructorStatics()
+				: FogTextureObject(TEXT("/Engine/EditorResources/S_ExpoHeightFog"))
+				, ID_Fog(TEXT("Fog"))
+				, NAME_Fog(NSLOCTEXT("SpriteCategory", "Fog", "Fog"))
+			{
+			}
+		};
+		static FConstructorStatics ConstructorStatics;
+
+		GetSpriteComponent()->Sprite = ConstructorStatics.FogTextureObject.Get();
+		GetSpriteComponent()->SetRelativeScale3D_Direct(FVector(0.5f, 0.5f, 0.5f));
+		GetSpriteComponent()->SpriteInfo.Category = ConstructorStatics.ID_Fog;
+		GetSpriteComponent()->SpriteInfo.DisplayName = ConstructorStatics.NAME_Fog;
+		GetSpriteComponent()->SetupAttachment(Component);
+	}
+#endif // WITH_EDITORONLY_DATA
+}
+
+void AExponentialHeightFog::PostInitializeComponents()
+{
+	Super::PostInitializeComponents();
+
+	bEnabled = Component->GetVisibleFlag();
+}
+
+void AExponentialHeightFog::GetLifetimeReplicatedProps( TArray< FLifetimeProperty > & OutLifetimeProps ) const
+{
+	Super::GetLifetimeReplicatedProps( OutLifetimeProps );
+	
+	DOREPLIFETIME( AExponentialHeightFog, bEnabled );
+}
+
+void AExponentialHeightFog::OnRep_bEnabled()
+{
+	Component->SetVisibility(bEnabled);
+}
+
diff --git a/Engine/Source/Runtime/Engine/Private/Components/HeightFogComponent.cpp b/Engine/Source/Runtime/Engine/Private/Components/HeightFogComponent.cpp
deleted file mode 100644
