Moved the call of UWorld::OnActorSpawned to FinishSpawningActor so that we can be certain the Actor is fully initialized before that event is broadcast.  Previously, spawning an Actor through Blueprints would cause this OnActorSpawned to trigger between the first call of a Deferred Actor Spawn and the FinishSpawning call, meaning that the listener would receive a half-initialized AActor.

diff --git a/Engine/Source/Runtime/Engine/Classes/Engine/World.h b/Engine/Source/Runtime/Engine/Classes/Engine/World.h
--- a/Engine/Source/Runtime/Engine/Classes/Engine/World.h
+++ b/Engine/Source/Runtime/Engine/Classes/Engine/World.h
@@ -2732,6 +2732,9 @@ public:
 	/** Remove a listener for OnActorSpawned events */
 	void RemoveOnActorSpawnedHandler( FDelegateHandle InHandle ) const;
 
+	/** Notify the World that an Actor has finished spawning (all OnActorSpawnedHandlers will run) */
+	void OnActorFinishedSpawning(AActor* Actor);
+
 	/** Add a listener for OnActorPreSpawnInitialization events */
 	FDelegateHandle AddOnActorPreSpawnInitialization(const FOnActorSpawned::FDelegate& InHandler) const;
 
diff --git a/Engine/Source/Runtime/Engine/Private/Actor.cpp b/Engine/Source/Runtime/Engine/Private/Actor.cpp
--- a/Engine/Source/Runtime/Engine/Private/Actor.cpp
+++ b/Engine/Source/Runtime/Engine/Private/Actor.cpp
@@ -4098,6 +4098,11 @@ void AActor::FinishSpawning(const FTransform& UserTransform, bool bIsDefaultTran
 			SCOPE_CYCLE_COUNTER(STAT_PostActorConstruction);
 			PostActorConstruction();
 		}
+
+		if (UWorld* World = GetWorld())
+		{
+			World->OnActorFinishedSpawning(this);
+		}
 	}
 }
 
diff --git a/Engine/Source/Runtime/Engine/Private/LevelActor.cpp b/Engine/Source/Runtime/Engine/Private/LevelActor.cpp
--- a/Engine/Source/Runtime/Engine/Private/LevelActor.cpp
+++ b/Engine/Source/Runtime/Engine/Private/LevelActor.cpp
@@ -746,9 +746,6 @@ AActor* UWorld::SpawnActor( UClass* Class, FTransform const* UserTransformPtr, c
 
 	Actor->CheckDefaultSubobjects();
 
-	// Broadcast notification of spawn
-	OnActorSpawned.Broadcast(Actor);
-
 #if WITH_EDITOR
 	if (GIsEditor)
 	{
@@ -761,9 +758,6 @@ AActor* UWorld::SpawnActor( UClass* Class, FTransform const* UserTransformPtr, c
 	}
 #endif
 
-	// Add this newly spawned actor to the network actor list. Do this after PostSpawnInitialize so that actor has "finished" spawning.
-	AddNetworkActor( Actor );
-
 	return Actor;
 }
 
diff --git a/Engine/Source/Runtime/Engine/Private/World.cpp b/Engine/Source/Runtime/Engine/Private/World.cpp
--- a/Engine/Source/Runtime/Engine/Private/World.cpp
+++ b/Engine/Source/Runtime/Engine/Private/World.cpp
@@ -5998,6 +5998,15 @@ void UWorld::RemoveOnActorSpawnedHandler(FDelegateHandle InHandle) const
 	OnActorSpawned.Remove(InHandle);
 }
 
+void UWorld::OnActorFinishedSpawning(AActor* Actor)
+{
+	// Broadcast notification of spawn
+	OnActorSpawned.Broadcast(Actor);
+
+	// Add this newly spawned actor to the network actor list
+	AddNetworkActor(Actor);
+}
+
 FDelegateHandle UWorld::AddOnActorPreSpawnInitialization(const FOnActorSpawned::FDelegate& InHandler) const
 {
 	return OnActorPreSpawnInitialization.Add(InHandler);
