Fix UWorldPartitionLevelStreamingDynamic::LoadInEditor (only called by UWorldPartitionHLODSourceActorsFromCell::LoadSourceActors) that creates a dummy streaming cell for convenience, but doesn't create it properly, resulting in crashes.
Fixed by creating a dummy runtime cell data object with the HLOD cell name as debug name.

diff --git a/Engine/Source/Runtime/Engine/Private/WorldPartition/HLOD/HLODSourceActorsFromCell.cpp b/Engine/Source/Runtime/Engine/Private/WorldPartition/HLOD/HLODSourceActorsFromCell.cpp
--- a/Engine/Source/Runtime/Engine/Private/WorldPartition/HLOD/HLODSourceActorsFromCell.cpp
+++ b/Engine/Source/Runtime/Engine/Private/WorldPartition/HLOD/HLODSourceActorsFromCell.cpp
@@ -28,6 +28,8 @@ ULevelStreaming* UWorldPartitionHLODSourceActorsFromCell::LoadSourceActors(bool&
 	UWorldPartitionLevelStreamingDynamic* LevelStreaming = UWorldPartitionLevelStreamingDynamic::LoadInEditor(World, LevelStreamingName, Actors);
 	check(LevelStreaming);
 
+	LevelStreaming->GetWorldPartitionRuntimeCell()->RuntimeCellData->DebugName = HLODActor->GetActorNameOrLabel();
+
 	if (!LevelStreaming->GetLoadSucceeded())
 	{
 		bOutDirty = true;
diff --git a/Engine/Source/Runtime/Engine/Private/WorldPartition/WorldPartitionLevelStreamingDynamic.cpp b/Engine/Source/Runtime/Engine/Private/WorldPartition/WorldPartitionLevelStreamingDynamic.cpp
--- a/Engine/Source/Runtime/Engine/Private/WorldPartition/WorldPartitionLevelStreamingDynamic.cpp
+++ b/Engine/Source/Runtime/Engine/Private/WorldPartition/WorldPartitionLevelStreamingDynamic.cpp
@@ -99,6 +99,7 @@ UWorldPartitionLevelStreamingDynamic* UWorldPartitionLevelStreamingDynamic::Load
 
 	// Assign a dummy runtime cell to ensure that code that is testing for IsWorldPartitionRuntimeCell() behaves as expected
 	UWorldPartitionRuntimeLevelStreamingCell* StreamingCell = NewObject<UWorldPartitionRuntimeLevelStreamingCell>(LevelStreaming, NAME_None, RF_Transient);
+	StreamingCell->RuntimeCellData = NewObject<UWorldPartitionRuntimeCellData>(StreamingCell);
 	LevelStreaming->StreamingCell = StreamingCell;
 	
 	LevelStreaming->LevelTransform = FTransform::Identity;
