[Engine]
* Added statestream support for ExponentialHeightFog.

diff --git a/Engine/Source/Runtime/Engine/Classes/Components/ExponentialHeightFogComponent.h b/Engine/Source/Runtime/Engine/Classes/Components/ExponentialHeightFogComponent.h
--- a/Engine/Source/Runtime/Engine/Classes/Components/ExponentialHeightFogComponent.h
+++ b/Engine/Source/Runtime/Engine/Classes/Components/ExponentialHeightFogComponent.h
@@ -4,6 +4,7 @@
 #pragma once
 
 #include "CoreMinimal.h"
+#include "StateStream/ExponentialHeightFogStateStreamHandle.h"
 #include "UObject/ObjectMacros.h"
 #include "Components/SceneComponent.h"
 #include "Rendering/ExponentialHeightFogData.h"
@@ -291,6 +292,10 @@ protected:
 
 	ENGINE_API void AddFogIfNeeded();
 
+	#if WITH_STATE_STREAM_ACTOR
+	FExponentialHeightFogHandle Handle;
+	#endif
+
 public:
 	//~ Begin UObject Interface
 #if WITH_EDITOR
diff --git a/Engine/Source/Runtime/Engine/Private/Components/ExponentialHeightFogComponent.cpp b/Engine/Source/Runtime/Engine/Private/Components/ExponentialHeightFogComponent.cpp
--- a/Engine/Source/Runtime/Engine/Private/Components/ExponentialHeightFogComponent.cpp
+++ b/Engine/Source/Runtime/Engine/Private/Components/ExponentialHeightFogComponent.cpp
@@ -6,6 +6,7 @@
 
 #include "Components/ExponentialHeightFogComponent.h"
 #include "GameFramework/Info.h"
+#include "StateStream/ExponentialHeightFogStateStream.h"
 #include "UObject/ConstructorHelpers.h"
 #include "Engine/World.h"
 #include "SceneInterface.h"
@@ -17,6 +18,47 @@
 
 #include UE_INLINE_GENERATED_CPP_BY_NAME(ExponentialHeightFogComponent)
 
+#if WITH_STATE_STREAM_ACTOR
+constexpr bool UseExponentialHeightFogStateStream = true;
+#endif
+
+//#define UEXPONENTIALHEIGHTFOG_VARIABLE(Type, Prefix, Name)
+
+#define UEXPONENTIALHEIGHTFOG_VARIABLES \
+	UEXPONENTIALHEIGHTFOG_VARIABLE(FExponentialHeightFogData, , SecondFogData) \
+	UEXPONENTIALHEIGHTFOG_VARIABLE(FLinearColor, , FogInscatteringLuminance) \
+	UEXPONENTIALHEIGHTFOG_VARIABLE(TObjectPtr<UTextureCube>, , InscatteringColorCubemap) \
+	UEXPONENTIALHEIGHTFOG_VARIABLE(float, , FullyDirectionalInscatteringColorDistance) \
+	UEXPONENTIALHEIGHTFOG_VARIABLE(FLinearColor, , DirectionalInscatteringLuminance) \
+	UEXPONENTIALHEIGHTFOG_VARIABLE(bool, b, EnableVolumetricFog) \
+	UEXPONENTIALHEIGHTFOG_VARIABLE(float, , VolumetricFogStaticLightingScatteringIntensity) \
+	UEXPONENTIALHEIGHTFOG_VARIABLE(bool, b, OverrideLightColorsWithFogInscatteringColors) \
+	UEXPONENTIALHEIGHTFOG_VARIABLE(bool, b, Holdout) \
+	UEXPONENTIALHEIGHTFOG_VARIABLE(bool, b, RenderInMainPass) \
+	UEXPONENTIALHEIGHTFOG_VARIABLES_WITH_SETFUNCTION
+
+#define UEXPONENTIALHEIGHTFOG_VARIABLES_WITH_SETFUNCTION \
+	UEXPONENTIALHEIGHTFOG_VARIABLE(float, , FogDensity) \
+	UEXPONENTIALHEIGHTFOG_VARIABLE(float, , FogHeightFalloff) \
+	UEXPONENTIALHEIGHTFOG_VARIABLE(FLinearColor, , SkyAtmosphereAmbientContributionColorScale) \
+	UEXPONENTIALHEIGHTFOG_VARIABLE(float, , InscatteringColorCubemapAngle) \
+	UEXPONENTIALHEIGHTFOG_VARIABLE(FLinearColor, , InscatteringTextureTint) \
+	UEXPONENTIALHEIGHTFOG_VARIABLE(float, , NonDirectionalInscatteringColorDistance) \
+	UEXPONENTIALHEIGHTFOG_VARIABLE(float, , DirectionalInscatteringExponent) \
+	UEXPONENTIALHEIGHTFOG_VARIABLE(float, , DirectionalInscatteringStartDistance) \
+	UEXPONENTIALHEIGHTFOG_VARIABLE(float, , FogMaxOpacity) \
+	UEXPONENTIALHEIGHTFOG_VARIABLE(float, , StartDistance) \
+	UEXPONENTIALHEIGHTFOG_VARIABLE(float, , EndDistance) \
+	UEXPONENTIALHEIGHTFOG_VARIABLE(float, , FogCutoffDistance) \
+	UEXPONENTIALHEIGHTFOG_VARIABLE(float, , VolumetricFogScatteringDistribution) \
+	UEXPONENTIALHEIGHTFOG_VARIABLE(FColor, , VolumetricFogAlbedo) \
+	UEXPONENTIALHEIGHTFOG_VARIABLE(FLinearColor, , VolumetricFogEmissive) \
+	UEXPONENTIALHEIGHTFOG_VARIABLE(float, , VolumetricFogExtinctionScale) \
+	UEXPONENTIALHEIGHTFOG_VARIABLE(float, , VolumetricFogDistance) \
+	UEXPONENTIALHEIGHTFOG_VARIABLE(float, , VolumetricFogStartDistance) \
+	UEXPONENTIALHEIGHTFOG_VARIABLE(float, , VolumetricFogNearFadeInDistance) \
+
+
 UExponentialHeightFogComponent::UExponentialHeightFogComponent(const FObjectInitializer& ObjectInitializer)
 	: Super(ObjectInitializer)
 {
@@ -64,7 +106,24 @@ void UExponentialHeightFogComponent::AddFogIfNeeded()
 	if (ShouldComponentAddToScene() && ShouldRender() && IsRegistered() && ((FogDensity + SecondFogData.FogDensity) * 1000) > UE_DELTA && FogMaxOpacity > UE_DELTA
 		&& (GetOuter() == NULL || !GetOuter()->HasAnyFlags(RF_ClassDefaultObject)))
 	{
-		GetWorld()->Scene->AddExponentialHeightFog(this);
+		FExponentialHeightFogDynamicState Ds;
+		#define UEXPONENTIALHEIGHTFOG_VARIABLE(Type, Prefix, Name) Ds.Set##Name(Prefix##Name);
+		UEXPONENTIALHEIGHTFOG_VARIABLES
+		#undef UEXPONENTIALHEIGHTFOG_VARIABLE
+
+		Ds.SetHeight(GetComponentLocation().Z);
+
+		#if WITH_STATE_STREAM_ACTOR
+		if (UseExponentialHeightFogStateStream)
+		{
+			//Ds.SetComponentTransform(GetComponentTransform());
+			//Ds.SetOverrideAtmosphericLight(ConvertAtmosphericLightOverride(OverrideAtmosphericLight, OverrideAtmosphericLightDirection));
+			Handle = GetWorld()->GetStateStream<IExponentialHeightFogStateStream>().Game_CreateInstance({}, Ds);
+			return;
+		}
+		#endif
+
+		GetWorld()->Scene->AddExponentialHeightFog(uint64(this), Ds);
 	}
 }
 
@@ -76,7 +135,7 @@ void UExponentialHeightFogComponent::CreateRenderState_Concurrent(FRegisterCompo
 
 void UExponentialHeightFogComponent::SendRenderTransform_Concurrent()
 {
-	GetWorld()->Scene->RemoveExponentialHeightFog(this);
+	GetWorld()->Scene->RemoveExponentialHeightFog(uint64(this));
 	AddFogIfNeeded();
 	Super::SendRenderTransform_Concurrent();
 }
@@ -84,7 +143,7 @@ void UExponentialHeightFogComponent::SendRenderTransform_Concurrent()
 void UExponentialHeightFogComponent::DestroyRenderState_Concurrent()
 {
 	Super::DestroyRenderState_Concurrent();
-	GetWorld()->Scene->RemoveExponentialHeightFog(this);
+	GetWorld()->Scene->RemoveExponentialHeightFog(uint64(this));
 }
 
 #if WITH_EDITOR
@@ -151,20 +210,41 @@ void UExponentialHeightFogComponent::PostEditChangeProperty(FPropertyChangedEven
 }
 #endif // WITH_EDITOR
 
-void UExponentialHeightFogComponent::SetFogDensity(float Value)
-{
-	if(FogDensity != Value)
-	{
-		FogDensity = Value;
-		MarkRenderStateDirty();
+#if WITH_STATE_STREAM_ACTOR
+	#define UEXPONENTIALHEIGHTFOG_SETHANDLE(Prefix, Field) \
+		if (UseExponentialHeightFogStateStream) \
+		{ \
+			FExponentialHeightFogDynamicState Ds; \
+			Ds.Set##Field(Prefix##Field); \
+			Handle.Update(Ds); \
+			return; \
+		}
+#else
+	#define UEXPONENTIALHEIGHTFOG_SETHANDLE(Prefix, Field)
+#endif
+
+
+#define UEXPONENTIALHEIGHTFOG_VARIABLE(Type, Prefix, Name) \
+	void UExponentialHeightFogComponent::Set##Name(Type NewValue) \
+	{ \
+		if (Name != NewValue) \
+		{ \
+			Name = NewValue; \
+			UEXPONENTIALHEIGHTFOG_SETHANDLE(, Name) \
+			MarkRenderStateDirty(); \
+		} \
 	}
-}
+
+UEXPONENTIALHEIGHTFOG_VARIABLES_WITH_SETFUNCTION
+#undef UEXPONENTIALHEIGHTFOG_VARIABLE
+
 
 void UExponentialHeightFogComponent::SetSecondFogDensity(float Value)
 {
 	if(SecondFogData.FogDensity != Value)
 	{
 		SecondFogData.FogDensity = Value;
+		UEXPONENTIALHEIGHTFOG_SETHANDLE(, SecondFogData)
 		MarkRenderStateDirty();
 	}
 }
@@ -174,6 +254,7 @@ void UExponentialHeightFogComponent::SetFogInscatteringColor(FLinearColor Value)
 	if(FogInscatteringLuminance != Value)
 	{
 		FogInscatteringLuminance = Value;
+		UEXPONENTIALHEIGHTFOG_SETHANDLE(, FogInscatteringLuminance)
 		MarkRenderStateDirty();
 	}
 }
@@ -183,15 +264,7 @@ void UExponentialHeightFogComponent::SetInscatteringColorCubemap(UTextureCube* V
 	if(InscatteringColorCubemap != Value)
 	{
 		InscatteringColorCubemap = Value;
-		MarkRenderStateDirty();
-	}
-}
-
-void UExponentialHeightFogComponent::SetInscatteringColorCubemapAngle(float Value)
-{
-	if(InscatteringColorCubemapAngle != Value)
-	{
-		InscatteringColorCubemapAngle = Value;
+		UEXPONENTIALHEIGHTFOG_SETHANDLE(, InscatteringColorCubemap)
 		MarkRenderStateDirty();
 	}
 }
@@ -201,42 +274,7 @@ void UExponentialHeightFogComponent::SetFullyDirectionalInscatteringColorDistanc
 	if(FullyDirectionalInscatteringColorDistance != Value)
 	{
 		FullyDirectionalInscatteringColorDistance = Value;
-		MarkRenderStateDirty();
-	}
-}
-
-void UExponentialHeightFogComponent::SetNonDirectionalInscatteringColorDistance(float Value)
-{
-	if(NonDirectionalInscatteringColorDistance != Value)
-	{
-		NonDirectionalInscatteringColorDistance = Value;
-		MarkRenderStateDirty();
-	}
-}
-
-void UExponentialHeightFogComponent::SetInscatteringTextureTint(FLinearColor Value)
-{
-	if(InscatteringTextureTint != Value)
-	{
-		InscatteringTextureTint = Value;
-		MarkRenderStateDirty();
-	}
-}
-
-void UExponentialHeightFogComponent::SetDirectionalInscatteringExponent(float Value)
-{
-	if(DirectionalInscatteringExponent != Value)
-	{
-		DirectionalInscatteringExponent = Value;
-		MarkRenderStateDirty();
-	}
-}
-
-void UExponentialHeightFogComponent::SetDirectionalInscatteringStartDistance(float Value)
-{
-	if(DirectionalInscatteringStartDistance != Value)
-	{
-		DirectionalInscatteringStartDistance = Value;
+		UEXPONENTIALHEIGHTFOG_SETHANDLE(, FullyDirectionalInscatteringColorDistance)
 		MarkRenderStateDirty();
 	}
 }
@@ -246,6 +284,7 @@ void UExponentialHeightFogComponent::SetDirectionalInscatteringColor(FLinearColo
 	if(DirectionalInscatteringLuminance != Value)
 	{
 		DirectionalInscatteringLuminance = Value;
+		UEXPONENTIALHEIGHTFOG_SETHANDLE(, DirectionalInscatteringLuminance)
 		MarkRenderStateDirty();
 	}
 }
@@ -255,15 +294,7 @@ void UExponentialHeightFogComponent::SetSecondFogHeightOffset(float Value)
 	if(SecondFogData.FogHeightOffset != Value)
 	{
 		SecondFogData.FogHeightOffset = Value;
-		MarkRenderStateDirty();
-	}
-}
-
-void UExponentialHeightFogComponent::SetFogHeightFalloff(float Value)
-{
-	if(FogHeightFalloff != Value)
-	{
-		FogHeightFalloff = Value;
+		UEXPONENTIALHEIGHTFOG_SETHANDLE(, SecondFogData)
 		MarkRenderStateDirty();
 	}
 }
@@ -273,42 +304,7 @@ void UExponentialHeightFogComponent::SetSecondFogHeightFalloff(float Value)
 	if(SecondFogData.FogHeightFalloff != Value)
 	{
 		SecondFogData.FogHeightFalloff = Value;
-		MarkRenderStateDirty();
-	}
-}
-
-void UExponentialHeightFogComponent::SetFogMaxOpacity(float Value)
-{
-	if(FogMaxOpacity != Value)
-	{
-		FogMaxOpacity = Value;
-		MarkRenderStateDirty();
-	}
-}
-
-void UExponentialHeightFogComponent::SetStartDistance(float Value)
-{
-	if(StartDistance != Value)
-	{
-		StartDistance = Value;
-		MarkRenderStateDirty();
-	}
-}
-
-void UExponentialHeightFogComponent::SetEndDistance(float Value)
-{
-	if(EndDistance != Value)
-	{
-		EndDistance = Value;
-		MarkRenderStateDirty();
-	}
-}
-
-void UExponentialHeightFogComponent::SetFogCutoffDistance(float Value)
-{
-	if(FogCutoffDistance != Value)
-	{
-		FogCutoffDistance = Value;
+		UEXPONENTIALHEIGHTFOG_SETHANDLE(, SecondFogData)
 		MarkRenderStateDirty();
 	}
 }
@@ -318,69 +314,7 @@ void UExponentialHeightFogComponent::SetVolumetricFog(bool bNewValue)
 	if(bEnableVolumetricFog != bNewValue)
 	{
 		bEnableVolumetricFog = bNewValue;
-		MarkRenderStateDirty();
-	}
-}
-
-void UExponentialHeightFogComponent::SetVolumetricFogScatteringDistribution(float NewValue)
-{
-	if(VolumetricFogScatteringDistribution != NewValue)
-	{
-		VolumetricFogScatteringDistribution = NewValue;
-		MarkRenderStateDirty();
-	}
-}
-
-void UExponentialHeightFogComponent::SetVolumetricFogExtinctionScale(float NewValue)
-{
-	if (VolumetricFogExtinctionScale != NewValue)
-	{
-		VolumetricFogExtinctionScale = NewValue;
-		MarkRenderStateDirty();
-	}
-}
-
-void UExponentialHeightFogComponent::SetVolumetricFogAlbedo(FColor NewValue)
-{
-	if (VolumetricFogAlbedo != NewValue)
-	{
-		VolumetricFogAlbedo = NewValue;
-		MarkRenderStateDirty();
-	}
-}
-
-void UExponentialHeightFogComponent::SetVolumetricFogEmissive(FLinearColor NewValue)
-{
-	if (VolumetricFogEmissive != NewValue)
-	{
-		VolumetricFogEmissive = NewValue;
-		MarkRenderStateDirty();
-	}
-}
-
-void UExponentialHeightFogComponent::SetVolumetricFogDistance(float NewValue)
-{
-	if(VolumetricFogDistance != NewValue)
-	{
-		VolumetricFogDistance = NewValue;
-		MarkRenderStateDirty();
-	}
-}
-
-void UExponentialHeightFogComponent::SetVolumetricFogStartDistance(float NewValue)
-{
-	if (VolumetricFogStartDistance != NewValue)
-	{
-		VolumetricFogStartDistance = NewValue;
-		MarkRenderStateDirty();
-	}
-}
-
-void UExponentialHeightFogComponent::SetVolumetricFogNearFadeInDistance(float NewValue)
-{
-	if (VolumetricFogNearFadeInDistance != NewValue)
-	{
-		VolumetricFogNearFadeInDistance = NewValue;
+		UEXPONENTIALHEIGHTFOG_SETHANDLE(b, EnableVolumetricFog)
 		MarkRenderStateDirty();
 	}
 }
@@ -392,6 +326,7 @@ void UExponentialHeightFogComponent::SetSecondFogData(FExponentialHeightFogData
 	   SecondFogData.FogHeightFalloff != NewValue.FogHeightFalloff)
 	{
 		SecondFogData = NewValue;
+		UEXPONENTIALHEIGHTFOG_SETHANDLE(, SecondFogData)
 		MarkRenderStateDirty();
 	}
 }
@@ -401,6 +336,7 @@ void UExponentialHeightFogComponent::SetHoldout(bool bNewHoldout)
 	if (bHoldout != bNewHoldout)
 	{
 		bHoldout = bNewHoldout;
+		UEXPONENTIALHEIGHTFOG_SETHANDLE(b, Holdout)
 		MarkRenderStateDirty();
 	}
 }
@@ -410,15 +346,7 @@ void UExponentialHeightFogComponent::SetRenderInMainPass(bool bValue)
 	if (bRenderInMainPass != bValue)
 	{
 		bRenderInMainPass = bValue;
-		MarkRenderStateDirty();
-	}
-}
-
-void UExponentialHeightFogComponent::SetSkyAtmosphereAmbientContributionColorScale(FLinearColor NewValue)
-{
-	if (SkyAtmosphereAmbientContributionColorScale != NewValue)
-	{
-		SkyAtmosphereAmbientContributionColorScale = NewValue;
+		UEXPONENTIALHEIGHTFOG_SETHANDLE(b, RenderInMainPass)
 		MarkRenderStateDirty();
 	}
 }
diff --git a/Engine/Source/Runtime/Engine/Public/SceneInterface.h b/Engine/Source/Runtime/Engine/Public/SceneInterface.h
--- a/Engine/Source/Runtime/Engine/Public/SceneInterface.h
+++ b/Engine/Source/Runtime/Engine/Public/SceneInterface.h
@@ -65,7 +65,7 @@ using FSHVectorRGB3 = TSHVectorRGB<3>;
 struct FCustomPrimitiveData;
 class FSceneViewFamily;
 
-
+struct FExponentialHeightFogDynamicState;
 struct FPrimitiveSceneDesc;
 struct FInstancedStaticMeshSceneDesc;
 struct FLightSceneDesc;
@@ -343,13 +343,13 @@ public:
 	 * 
 	 * @param FogComponent - fog component to add
 	 */	
-	virtual void AddExponentialHeightFog(class UExponentialHeightFogComponent* FogComponent) = 0;
+	virtual void AddExponentialHeightFog(uint64 Id, const FExponentialHeightFogDynamicState& State) = 0;
 	/** 
 	 * Removes a exponential height fog component from the scene
 	 * 
 	 * @param FogComponent - fog component to remove
 	 */	
-	virtual void RemoveExponentialHeightFog(class UExponentialHeightFogComponent* FogComponent) = 0;
+	virtual void RemoveExponentialHeightFog(uint64 Id) = 0;
 	/**
 	 * @return True if there are any exponential height fog potentially enabled in the scene
 	 */
diff --git a/Engine/Source/Runtime/Engine/Public/StateStream/ExponentialHeightFogStateStream.h b/Engine/Source/Runtime/Engine/Public/StateStream/ExponentialHeightFogStateStream.h
new file mode 100644
--- /dev/null
+++ b/Engine/Source/Runtime/Engine/Public/StateStream/ExponentialHeightFogStateStream.h
@@ -0,0 +1,126 @@
+// Copyright Epic Games, Inc. All Rights Reserved.
+
+#pragma once
+
+#include "EngineDefines.h"
+#include "ExponentialHeightFogStateStreamHandle.h"
+#include "Rendering/ExponentialHeightFogData.h"
+#include "StateStreamDefinitions.h"
+#include "TransformStateStreamMath.h"
+#include "ExponentialHeightFogStateStream.generated.h"
+
+class UTextureCube;
+
+inline bool StateStreamEquals(const FExponentialHeightFogData& A, const FExponentialHeightFogData& B);
+inline void StateStreamInterpolate(FStateStreamInterpolateContext& Context, FExponentialHeightFogData& Out, const FExponentialHeightFogData& From, const FExponentialHeightFogData& To);
+
+////////////////////////////////////////////////////////////////////////////////////////////////////
+
+USTRUCT(StateStreamStaticState)
+struct FExponentialHeightFogStaticState
+{
+	GENERATED_USTRUCT_BODY()
+};
+
+
+////////////////////////////////////////////////////////////////////////////////////////////////////
+
+USTRUCT(StateStreamDynamicState)
+struct FExponentialHeightFogDynamicState
+{
+	GENERATED_USTRUCT_BODY()
+
+// private: // Should be private but complicates some of the code that needs support both this path and UComponent path
+
+	UPROPERTY()
+	float FogDensity;
+	UPROPERTY()
+	float FogHeightFalloff;
+	UPROPERTY()
+	FExponentialHeightFogData SecondFogData;
+	UPROPERTY()
+	FLinearColor FogInscatteringLuminance;
+	UPROPERTY()
+	FLinearColor SkyAtmosphereAmbientContributionColorScale;
+	UPROPERTY()
+	TObjectPtr<UTextureCube> InscatteringColorCubemap;
+	UPROPERTY()
+	float InscatteringColorCubemapAngle;
+	UPROPERTY()
+	FLinearColor InscatteringTextureTint;
+	UPROPERTY()
+	float FullyDirectionalInscatteringColorDistance;
+	UPROPERTY()
+	float NonDirectionalInscatteringColorDistance;
+	UPROPERTY()
+	float DirectionalInscatteringExponent;
+	UPROPERTY()
+	float DirectionalInscatteringStartDistance;
+	UPROPERTY()
+	FLinearColor DirectionalInscatteringLuminance;
+	UPROPERTY()
+	float FogMaxOpacity;
+	UPROPERTY()
+	float StartDistance;
+	UPROPERTY()
+	float EndDistance;
+	UPROPERTY()
+	float FogCutoffDistance;
+	UPROPERTY()
+	bool bEnableVolumetricFog;
+	UPROPERTY()
+	float VolumetricFogScatteringDistribution;
+	UPROPERTY()
+	FColor VolumetricFogAlbedo;
+	UPROPERTY()
+	FLinearColor VolumetricFogEmissive;
+	UPROPERTY()
+	float VolumetricFogExtinctionScale;
+	UPROPERTY()
+	float VolumetricFogDistance;
+	UPROPERTY()
+	float VolumetricFogStartDistance;
+	UPROPERTY()
+	float VolumetricFogNearFadeInDistance;
+	UPROPERTY()
+	float VolumetricFogStaticLightingScatteringIntensity;
+	UPROPERTY()
+	bool bOverrideLightColorsWithFogInscatteringColors;
+	UPROPERTY()
+	uint8 bHoldout : 1;
+	UPROPERTY()
+	uint8 bRenderInMainPass : 1;
+
+	UPROPERTY()
+	float Height;
+};
+
+////////////////////////////////////////////////////////////////////////////////////////////////////
+
+inline constexpr uint32 ExponentialHeightFogStateStreamId = 7;
+
+
+////////////////////////////////////////////////////////////////////////////////////////////////////
+
+class IExponentialHeightFogStateStream
+{
+public:
+	DECLARE_STATESTREAM(ExponentialHeightFog)
+	virtual FExponentialHeightFogHandle Game_CreateInstance(const FExponentialHeightFogStaticState& Ss, const FExponentialHeightFogDynamicState& Ds) = 0;
+};
+
+////////////////////////////////////////////////////////////////////////////////////////////////////
+
+inline bool StateStreamEquals(const FExponentialHeightFogData& A, const FExponentialHeightFogData& B)
+{
+	return A.FogDensity == B.FogDensity && A.FogHeightFalloff == B.FogHeightFalloff && A.FogHeightOffset == B.FogHeightOffset;
+}
+
+inline void StateStreamInterpolate(FStateStreamInterpolateContext& Context, FExponentialHeightFogData& Out, const FExponentialHeightFogData& From, const FExponentialHeightFogData& To)
+{
+	StateStreamInterpolate(Context, Out.FogDensity, From.FogDensity, To.FogDensity);
+	StateStreamInterpolate(Context, Out.FogHeightFalloff, From.FogHeightFalloff, To.FogHeightFalloff);
+	StateStreamInterpolate(Context, Out.FogHeightOffset, From.FogHeightOffset, To.FogHeightOffset);
+}
+
+////////////////////////////////////////////////////////////////////////////////////////////////////
diff --git a/Engine/Source/Runtime/Engine/Public/StateStream/ExponentialHeightFogStateStreamHandle.h b/Engine/Source/Runtime/Engine/Public/StateStream/ExponentialHeightFogStateStreamHandle.h
new file mode 100644
--- /dev/null
+++ b/Engine/Source/Runtime/Engine/Public/StateStream/ExponentialHeightFogStateStreamHandle.h
@@ -0,0 +1,19 @@
+// Copyright Epic Games, Inc. All Rights Reserved.
+
+#pragma once
+
+#include "StateStreamHandle.h"
+#include "ExponentialHeightFogStateStreamHandle.generated.h"
+
+struct FExponentialHeightFogDynamicState;
+
+////////////////////////////////////////////////////////////////////////////////////////////////////
+// Handle for mesh instance
+
+USTRUCT(StateStreamHandle)
+struct FExponentialHeightFogHandle : public FStateStreamHandle
+{
+	GENERATED_USTRUCT_BODY()
+};
+
+////////////////////////////////////////////////////////////////////////////////////////////////////
diff --git a/Engine/Source/Runtime/Renderer/Private/RendererScene.cpp b/Engine/Source/Runtime/Renderer/Private/RendererScene.cpp
--- a/Engine/Source/Runtime/Renderer/Private/RendererScene.cpp
+++ b/Engine/Source/Runtime/Renderer/Private/RendererScene.cpp
@@ -4062,12 +4062,11 @@ void FScene::RemoveLight(ULightComponent* Light)
 	}
 }
 
-void FScene::AddExponentialHeightFog(UExponentialHeightFogComponent* FogComponent)
+void FScene::AddExponentialHeightFog(uint64 Id, const FExponentialHeightFogDynamicState& State)
 {
-	FScene* Scene = this;
-	FExponentialHeightFogSceneInfo HeightFogSceneInfo = FExponentialHeightFogSceneInfo(FogComponent);
+	FExponentialHeightFogSceneInfo HeightFogSceneInfo(Id, State);
 	ENQUEUE_RENDER_COMMAND(FAddFogCommand)(
-		[Scene, HeightFogSceneInfo] (FRHICommandListBase&)
+		[Scene = this, HeightFogSceneInfo] (FRHICommandListBase&)
 		{
 			// Create a FExponentialHeightFogSceneInfo for the component in the scene's fog array.
 			new(Scene->ExponentialFogs) FExponentialHeightFogSceneInfo(HeightFogSceneInfo);
@@ -4075,16 +4074,15 @@ void FScene::AddExponentialHeightFog(UExponentialHeightFogComponent* FogComponen
 		});
 }
 
-void FScene::RemoveExponentialHeightFog(UExponentialHeightFogComponent* FogComponent)
+void FScene::RemoveExponentialHeightFog(uint64 Id)
 {
-	FScene* Scene = this;
 	ENQUEUE_RENDER_COMMAND(FRemoveFogCommand)(
-		[Scene, FogComponent] (FRHICommandListBase&)
+		[Scene = this, Id] (FRHICommandListBase&)
 		{
 			// Remove the given component's FExponentialHeightFogSceneInfo from the scene's fog array.
 			for(int32 FogIndex = 0;FogIndex < Scene->ExponentialFogs.Num();FogIndex++)
 			{
-				if(Scene->ExponentialFogs[FogIndex].Component == FogComponent)
+				if(Scene->ExponentialFogs[FogIndex].Id == Id)
 				{
 					Scene->ExponentialFogs.RemoveAt(FogIndex);
 					Scene->InvalidatePathTracedOutput();
@@ -6997,8 +6995,8 @@ public:
 	virtual void UpdateLightTransform(ULightComponent* Light) override {}
 	virtual void UpdateLightColorAndBrightness(ULightComponent* Light) override {}
 
-	virtual void AddExponentialHeightFog(class UExponentialHeightFogComponent* FogComponent) override {}
-	virtual void RemoveExponentialHeightFog(class UExponentialHeightFogComponent* FogComponent) override {}
+	virtual void AddExponentialHeightFog(uint64 Id, const FExponentialHeightFogDynamicState& State) override {}
+	virtual void RemoveExponentialHeightFog(uint64 Id) override {}
 	virtual bool HasAnyExponentialHeightFog() const override { return false; }
 
 	virtual void AddLocalFogVolume(class FLocalFogVolumeSceneProxy* FogProxy) override {}
diff --git a/Engine/Source/Runtime/Renderer/Private/SceneCore.cpp b/Engine/Source/Runtime/Renderer/Private/SceneCore.cpp
--- a/Engine/Source/Runtime/Renderer/Private/SceneCore.cpp
+++ b/Engine/Source/Runtime/Renderer/Private/SceneCore.cpp
@@ -8,9 +8,9 @@
 #include "SceneInterface.h"
 #include "PrimitiveDrawingUtils.h"
 #include "Misc/ConfigCacheIni.h"
-#include "Components/ExponentialHeightFogComponent.h"
 #include "DepthRendering.h"
 #include "SceneHitProxyRendering.h"
+#include "StateStream/ExponentialHeightFogStateStream.h"
 #include "VelocityRendering.h"
 #include "BasePassRendering.h"
 #include "MobileBasePassRendering.h"
@@ -385,36 +385,36 @@ void FLightPrimitiveInteraction::FlushCachedShadowMapData()
 }
 
 /** Initialization constructor. */
-FExponentialHeightFogSceneInfo::FExponentialHeightFogSceneInfo(const UExponentialHeightFogComponent* InComponent):
-	Component(InComponent),
-	FogMaxOpacity(InComponent->FogMaxOpacity),
-	StartDistance(InComponent->StartDistance),
-	EndDistance(InComponent->EndDistance),
-	FogCutoffDistance(InComponent->FogCutoffDistance),
-	DirectionalInscatteringExponent(InComponent->DirectionalInscatteringExponent),
-	DirectionalInscatteringStartDistance(InComponent->DirectionalInscatteringStartDistance),
-	DirectionalInscatteringColor(InComponent->DirectionalInscatteringLuminance)
+FExponentialHeightFogSceneInfo::FExponentialHeightFogSceneInfo(uint64 InId, const FExponentialHeightFogDynamicState& State) :
+	Id(InId),
+	FogMaxOpacity(State.FogMaxOpacity),
+	StartDistance(State.StartDistance),
+	EndDistance(State.EndDistance),
+	FogCutoffDistance(State.FogCutoffDistance),
+	DirectionalInscatteringExponent(State.DirectionalInscatteringExponent),
+	DirectionalInscatteringStartDistance(State.DirectionalInscatteringStartDistance),
+	DirectionalInscatteringColor(State.DirectionalInscatteringLuminance)
 {
-	FogData[0].Height = InComponent->GetComponentLocation().Z;
-	FogData[1].Height = InComponent->GetComponentLocation().Z + InComponent->SecondFogData.FogHeightOffset;
+	FogData[0].Height = State.Height;
+	FogData[1].Height = State.Height + State.SecondFogData.FogHeightOffset;
 
 	// Scale the densities back down to their real scale
 	// Artists edit the densities scaled up so they aren't entering in minuscule floating point numbers
-	FogData[0].Density = InComponent->FogDensity / 1000.0f;	
-	FogData[0].HeightFalloff = InComponent->FogHeightFalloff / 1000.0f;
-	FogData[1].Density = InComponent->SecondFogData.FogDensity / 1000.0f;
-	FogData[1].HeightFalloff = InComponent->SecondFogData.FogHeightFalloff / 1000.0f;
-
-	FogColor = InComponent->InscatteringColorCubemap ? InComponent->InscatteringTextureTint : InComponent->FogInscatteringLuminance;
-	InscatteringColorCubemap = InComponent->InscatteringColorCubemap;
-	InscatteringColorCubemapAngle = InComponent->InscatteringColorCubemapAngle * (PI / 180.f);
-	FullyDirectionalInscatteringColorDistance = InComponent->FullyDirectionalInscatteringColorDistance;
-	NonDirectionalInscatteringColorDistance = InComponent->NonDirectionalInscatteringColorDistance;
-
-	bEnableVolumetricFog = InComponent->bEnableVolumetricFog;
-	VolumetricFogScatteringDistribution = FMath::Clamp(InComponent->VolumetricFogScatteringDistribution, -.99f, .99f);
-	VolumetricFogAlbedo = FLinearColor(InComponent->VolumetricFogAlbedo);
-	VolumetricFogEmissive = InComponent->VolumetricFogEmissive;
+	FogData[0].Density = State.FogDensity / 1000.0f;	
+	FogData[0].HeightFalloff = State.FogHeightFalloff / 1000.0f;
+	FogData[1].Density = State.SecondFogData.FogDensity / 1000.0f;
+	FogData[1].HeightFalloff = State.SecondFogData.FogHeightFalloff / 1000.0f;
+
+	FogColor = State.InscatteringColorCubemap ? State.InscatteringTextureTint : State.FogInscatteringLuminance;
+	InscatteringColorCubemap = State.InscatteringColorCubemap;
+	InscatteringColorCubemapAngle = State.InscatteringColorCubemapAngle * (PI / 180.f);
+	FullyDirectionalInscatteringColorDistance = State.FullyDirectionalInscatteringColorDistance;
+	NonDirectionalInscatteringColorDistance = State.NonDirectionalInscatteringColorDistance;
+
+	bEnableVolumetricFog = State.bEnableVolumetricFog;
+	VolumetricFogScatteringDistribution = FMath::Clamp(State.VolumetricFogScatteringDistribution, -.99f, .99f);
+	VolumetricFogAlbedo = FLinearColor(State.VolumetricFogAlbedo);
+	VolumetricFogEmissive = State.VolumetricFogEmissive;
 
 	// Apply a scale so artists don't have to work with tiny numbers.
 	// The is only needed because emissive is by default not weighted by the height fog density distribution. 
@@ -423,15 +423,15 @@ FExponentialHeightFogSceneInfo::FExponentialHeightFogSceneInfo(const UExponentia
 	VolumetricFogEmissive.R = FMath::Max(VolumetricFogEmissive.R * EmissiveUnitScale, 0.0f);
 	VolumetricFogEmissive.G = FMath::Max(VolumetricFogEmissive.G * EmissiveUnitScale, 0.0f);
 	VolumetricFogEmissive.B = FMath::Max(VolumetricFogEmissive.B * EmissiveUnitScale, 0.0f);
-	VolumetricFogExtinctionScale = FMath::Max(InComponent->VolumetricFogExtinctionScale, 0.0f);
-	VolumetricFogDistance = FMath::Max(InComponent->VolumetricFogStartDistance + InComponent->VolumetricFogDistance, 0.0f);
-	VolumetricFogStaticLightingScatteringIntensity = FMath::Max(InComponent->VolumetricFogStaticLightingScatteringIntensity, 0.0f);
-	bOverrideLightColorsWithFogInscatteringColors = InComponent->bOverrideLightColorsWithFogInscatteringColors;
-	bHoldout = InComponent->bHoldout;
-	bRenderInMainPass = InComponent->bRenderInMainPass;
+	VolumetricFogExtinctionScale = FMath::Max(State.VolumetricFogExtinctionScale, 0.0f);
+	VolumetricFogDistance = FMath::Max(State.VolumetricFogStartDistance + State.VolumetricFogDistance, 0.0f);
+	VolumetricFogStaticLightingScatteringIntensity = FMath::Max(State.VolumetricFogStaticLightingScatteringIntensity, 0.0f);
+	bOverrideLightColorsWithFogInscatteringColors = State.bOverrideLightColorsWithFogInscatteringColors;
+	bHoldout = State.bHoldout;
+	bRenderInMainPass = State.bRenderInMainPass;
 
-	VolumetricFogStartDistance = InComponent->VolumetricFogStartDistance;
-	VolumetricFogNearFadeInDistance = InComponent->VolumetricFogNearFadeInDistance;
+	VolumetricFogStartDistance = State.VolumetricFogStartDistance;
+	VolumetricFogNearFadeInDistance = State.VolumetricFogNearFadeInDistance;
 
-	SkyAtmosphereAmbientContributionColorScale = InComponent->SkyAtmosphereAmbientContributionColorScale;
+	SkyAtmosphereAmbientContributionColorScale = State.SkyAtmosphereAmbientContributionColorScale;
 }
diff --git a/Engine/Source/Runtime/Renderer/Private/SceneCore.h b/Engine/Source/Runtime/Renderer/Private/SceneCore.h
--- a/Engine/Source/Runtime/Renderer/Private/SceneCore.h
+++ b/Engine/Source/Runtime/Renderer/Private/SceneCore.h
@@ -13,8 +13,8 @@
 class FLightSceneInfo;
 class FPrimitiveSceneInfo;
 class FScene;
-class UExponentialHeightFogComponent;
 class UTextureCube;
+struct FExponentialHeightFogDynamicState;
 
 /**
  * An interaction between a light and a primitive.
@@ -157,8 +157,7 @@ public:
 	/** Number of supported individual fog settings on this ExponentialHeightFog */
 	static constexpr int NumFogs = 2;
 
-	/** The fog component the scene info is for. */
-	const UExponentialHeightFogComponent* Component;
+	uint64 Id;
 	FExponentialHeightFogSceneData FogData[NumFogs];
 	float FogMaxOpacity;
 	float StartDistance;
@@ -189,7 +188,7 @@ public:
 	FLinearColor SkyAtmosphereAmbientContributionColorScale;
 
 	/** Initialization constructor. */
-	FExponentialHeightFogSceneInfo(const UExponentialHeightFogComponent* InComponent);
+	FExponentialHeightFogSceneInfo(uint64 InId, const FExponentialHeightFogDynamicState& State);
 };
 
 /** Returns true if the indirect lighting cache can be used at all. */
diff --git a/Engine/Source/Runtime/Renderer/Private/ScenePrivate.h b/Engine/Source/Runtime/Renderer/Private/ScenePrivate.h
--- a/Engine/Source/Runtime/Renderer/Private/ScenePrivate.h
+++ b/Engine/Source/Runtime/Renderer/Private/ScenePrivate.h
@@ -3440,8 +3440,8 @@ public:
 	virtual void GetPrimitiveUniformShaderParameters_RenderThread(const FPrimitiveSceneInfo* PrimitiveSceneInfo, bool& bHasPrecomputedVolumetricLightmap, FMatrix& PreviousLocalToWorld, int32& SingleCaptureIndex, bool& bOutputVelocity) const override;
 	virtual void UpdateLightTransform(ULightComponent* Light) override;
 	virtual void UpdateLightColorAndBrightness(ULightComponent* Light) override;
-	virtual void AddExponentialHeightFog(UExponentialHeightFogComponent* FogComponent) override;
-	virtual void RemoveExponentialHeightFog(UExponentialHeightFogComponent* FogComponent) override;
+	virtual void AddExponentialHeightFog(uint64 Id, const FExponentialHeightFogDynamicState& State) override;
+	virtual void RemoveExponentialHeightFog(uint64 Id) override;
 	virtual bool HasAnyExponentialHeightFog() const override;
 
 	virtual void AddHairStrands(FHairStrandsInstance* Proxy) override;
diff --git a/Engine/Source/Runtime/Renderer/Private/StateStream/ExponentialHeightFogStateStreamImpl.cpp b/Engine/Source/Runtime/Renderer/Private/StateStream/ExponentialHeightFogStateStreamImpl.cpp
new file mode 100644
--- /dev/null
+++ b/Engine/Source/Runtime/Renderer/Private/StateStream/ExponentialHeightFogStateStreamImpl.cpp
@@ -0,0 +1,32 @@
+// Copyright Epic Games, Inc. All Rights Reserved.
+
+#include "ExponentialHeightFogStateStreamImpl.h"
+#include "SceneInterface.h"
+#include "StateStreamCreator.h"
+
+////////////////////////////////////////////////////////////////////////////////////////////////////
+
+FExponentialHeightFogStateStreamImpl::FExponentialHeightFogStateStreamImpl(FSceneInterface& InScene)
+:	Scene(InScene)
+{
+}
+
+void FExponentialHeightFogStateStreamImpl::Render_OnCreate(const FExponentialHeightFogStaticState& Ss, const FExponentialHeightFogDynamicState& Ds, FExponentialHeightFogSceneProxy*& UserData, bool IsDestroyedInSameFrame)
+{
+	Scene.AddExponentialHeightFog(uint64(&Ss), Ds);
+}
+
+void FExponentialHeightFogStateStreamImpl::Render_OnUpdate(const FExponentialHeightFogStaticState& Ss, const FExponentialHeightFogDynamicState& Ds, FExponentialHeightFogSceneProxy*& UserData)
+{
+}
+
+void FExponentialHeightFogStateStreamImpl::Render_OnDestroy(const FExponentialHeightFogStaticState& Ss, const FExponentialHeightFogDynamicState& Ds, FExponentialHeightFogSceneProxy*& UserData)
+{
+	Scene.RemoveExponentialHeightFog(uint64(&Ss));
+}
+
+////////////////////////////////////////////////////////////////////////////////////////////////////
+
+STATESTREAM_CREATOR_INSTANCE(FExponentialHeightFogStateStreamImpl)
+
+////////////////////////////////////////////////////////////////////////////////////////////////////
diff --git a/Engine/Source/Runtime/Renderer/Private/StateStream/ExponentialHeightFogStateStreamImpl.h b/Engine/Source/Runtime/Renderer/Private/StateStream/ExponentialHeightFogStateStreamImpl.h
new file mode 100644
--- /dev/null
+++ b/Engine/Source/Runtime/Renderer/Private/StateStream/ExponentialHeightFogStateStreamImpl.h
@@ -0,0 +1,36 @@
+// Copyright Epic Games, Inc. All Rights Reserved.
+
+#pragma once
+
+#include "GenericStateStream.h"
+#include "StateStream/ExponentialHeightFogStateStream.h"
+
+#define UE_API RENDERER_API
+
+class FSceneInterface;
+class FExponentialHeightFogSceneProxy;
+
+////////////////////////////////////////////////////////////////////////////////////////////////////
+
+struct FExponentialHeightFogStateStreamSettings : TStateStreamSettings<IExponentialHeightFogStateStream, FExponentialHeightFogSceneProxy>
+{
+	static inline constexpr bool SkipCreatingDeletes = true;
+};
+
+////////////////////////////////////////////////////////////////////////////////////////////////////
+
+class FExponentialHeightFogStateStreamImpl : public TStateStream<FExponentialHeightFogStateStreamSettings>
+{
+public:
+	FExponentialHeightFogStateStreamImpl(FSceneInterface& InScene);
+private:
+	UE_API virtual void Render_OnCreate(const FExponentialHeightFogStaticState& Ss, const FExponentialHeightFogDynamicState& Ds, FExponentialHeightFogSceneProxy*& UserData, bool IsDestroyedInSameFrame) override;
+	UE_API virtual void Render_OnUpdate(const FExponentialHeightFogStaticState& Ss, const FExponentialHeightFogDynamicState& Ds, FExponentialHeightFogSceneProxy*& UserData) override;
+	UE_API virtual void Render_OnDestroy(const FExponentialHeightFogStaticState& Ss, const FExponentialHeightFogDynamicState& Ds, FExponentialHeightFogSceneProxy*& UserData) override;
+
+	FSceneInterface& Scene;
+};
+
+////////////////////////////////////////////////////////////////////////////////////////////////////
+
+#undef UE_API
