Skin Weight Profile Manager: Properly handle platforms where min LOD > 0.

#rnx

diff --git a/Engine/Source/Runtime/Engine/Private/Animation/SkinWeightProfileManager.cpp b/Engine/Source/Runtime/Engine/Private/Animation/SkinWeightProfileManager.cpp
--- a/Engine/Source/Runtime/Engine/Private/Animation/SkinWeightProfileManager.cpp
+++ b/Engine/Source/Runtime/Engine/Private/Animation/SkinWeightProfileManager.cpp
@@ -144,8 +144,10 @@ void FSkinWeightProfileManager::RequestSkinWeightProfileStack(
 
 		if (LODIndex == INDEX_NONE)
 		{
+			// Match what rendering does.
+			constexpr bool bForceLowestLODIdx = true;
 			const int32 NumLODS = Mesh->GetLODNum();
-			for (int32 index = 0; index < NumLODS; ++index)
+			for (int32 index = Mesh->GetMinLodIdx(bForceLowestLODIdx); index < NumLODS; ++index)
 			{
 				ProfileRequest.LODIndices.Add(index);
 			}
@@ -367,7 +369,7 @@ void FSkinWeightProfileManagerAsyncTask::DoTask(ENamedThreads::Type CurrentThrea
 					
 					// If we have CPU-side weights available, use them right away.
 					if (FSkinWeightProfileManager::AllowCPU() &&
-						((FSkinWeightProfileManager::HandleDelayedLoads() && SkinweightData.BaseBuffer->GetDataVertexBuffer()->IsWeightDataValid()) ||
+						((FSkinWeightProfileManager::HandleDelayedLoads() && LODRenderData.SkinWeightVertexBuffer.GetDataVertexBuffer()->IsWeightDataValid()) ||
 						(!FSkinWeightProfileManager::HandleDelayedLoads() && LODRenderData.SkinWeightVertexBuffer.GetNeedsCPUAccess())))
 					{
 						SkinweightData.InitialiseProfileBuffer(Request.ProfileStack);
