[Backout] - CL37898780
Original CL Desc
-----------------------------------------------------------------
Skin Weight Profile Manager: Properly handle platforms where min LOD > 0.

#rnx

diff --git a/Engine/Source/Runtime/Engine/Private/Animation/SkinWeightProfileManager.cpp b/Engine/Source/Runtime/Engine/Private/Animation/SkinWeightProfileManager.cpp
--- a/Engine/Source/Runtime/Engine/Private/Animation/SkinWeightProfileManager.cpp
+++ b/Engine/Source/Runtime/Engine/Private/Animation/SkinWeightProfileManager.cpp
@@ -144,10 +144,8 @@ void FSkinWeightProfileManager::RequestSkinWeightProfileStack(
 
 		if (LODIndex == INDEX_NONE)
 		{
-			// Match what rendering does.
-			constexpr bool bForceLowestLODIdx = true;
 			const int32 NumLODS = Mesh->GetLODNum();
-			for (int32 index = Mesh->GetMinLodIdx(bForceLowestLODIdx); index < NumLODS; ++index)
+			for (int32 index = 0; index < NumLODS; ++index)
 			{
 				ProfileRequest.LODIndices.Add(index);
 			}
@@ -369,7 +367,7 @@ void FSkinWeightProfileManagerAsyncTask::DoTask(ENamedThreads::Type CurrentThrea
 					
 					// If we have CPU-side weights available, use them right away.
 					if (FSkinWeightProfileManager::AllowCPU() &&
-						((FSkinWeightProfileManager::HandleDelayedLoads() && LODRenderData.SkinWeightVertexBuffer.GetDataVertexBuffer()->IsWeightDataValid()) ||
+						((FSkinWeightProfileManager::HandleDelayedLoads() && SkinweightData.BaseBuffer->GetDataVertexBuffer()->IsWeightDataValid()) ||
 						(!FSkinWeightProfileManager::HandleDelayedLoads() && LODRenderData.SkinWeightVertexBuffer.GetNeedsCPUAccess())))
 					{
 						SkinweightData.InitialiseProfileBuffer(Request.ProfileStack);
