dstm:
- wrap FActorComponentTickFunction ticking in a transaction
- add a check() to test whether we're in a transaction when requesting a remote object

diff --git a/Engine/Plugins/Runtime/RigVM/Source/RigVM/RigVM.Build.cs b/Engine/Plugins/Runtime/RigVM/Source/RigVM/RigVM.Build.cs
--- a/Engine/Plugins/Runtime/RigVM/Source/RigVM/RigVM.Build.cs
+++ b/Engine/Plugins/Runtime/RigVM/Source/RigVM/RigVM.Build.cs
@@ -32,7 +32,5 @@ public class RigVM : ModuleRules
 				}
             );
         }
-
-		bDisableAutoRTFMInstrumentation = true;
     }
 }
diff --git a/Engine/Source/Runtime/Engine/Classes/Engine/EngineBaseTypes.h b/Engine/Source/Runtime/Engine/Classes/Engine/EngineBaseTypes.h
--- a/Engine/Source/Runtime/Engine/Classes/Engine/EngineBaseTypes.h
+++ b/Engine/Source/Runtime/Engine/Classes/Engine/EngineBaseTypes.h
@@ -549,8 +549,12 @@ struct FActorComponentTickFunction : public FTickFunction
 {
 	GENERATED_USTRUCT_BODY()
 
+#if UE_WITH_REMOTE_OBJECT_HANDLE
+	TObjectPtr<class UActorComponent> Target;
+#else
 	/**  AActor  component that is the target of this tick **/
 	class UActorComponent*	Target;
+#endif
 
 	/** 
 		* Abstract function actually execute the tick. 
@@ -577,6 +581,11 @@ struct FActorComponentTickFunction : public FTickFunction
 
 	template <typename ExecuteTickLambda>
 	static void ExecuteTickHelper(UActorComponent* Target, bool bTickInEditor, float DeltaTime, ELevelTick TickType, const ExecuteTickLambda& ExecuteTickFunc);	
+
+#if UE_WITH_REMOTE_OBJECT_HANDLE
+	/** Exposes references to GC system */
+	ENGINE_API void AddStructReferencedObjects(FReferenceCollector& Collector);
+#endif
 };
 
 
@@ -585,6 +594,9 @@ struct TStructOpsTypeTraits<FActorComponentTickFunction> : public TStructOpsType
 {
 	enum
 	{
+#if UE_WITH_REMOTE_OBJECT_HANDLE
+		WithAddStructReferencedObjects = true,
+#endif
 		WithCopy = false
 	};
 };
diff --git a/Engine/Source/Runtime/Engine/Private/Actor.cpp b/Engine/Source/Runtime/Engine/Private/Actor.cpp
--- a/Engine/Source/Runtime/Engine/Private/Actor.cpp
+++ b/Engine/Source/Runtime/Engine/Private/Actor.cpp
@@ -329,6 +329,13 @@ void FActorTickFunction::AddStructReferencedObjects(FReferenceCollector& Collect
 }
 #endif
 
+#if UE_WITH_REMOTE_OBJECT_HANDLE
+void FActorComponentTickFunction::AddStructReferencedObjects(FReferenceCollector& Collector)
+{
+	Collector.AddReferencedObject(Target);
+}
+#endif
+
 bool AActor::CheckDefaultSubobjectsInternal() const
 {
 	bool Result = Super::CheckDefaultSubobjectsInternal();
diff --git a/Engine/Source/Runtime/Engine/Private/Components/ActorComponent.cpp b/Engine/Source/Runtime/Engine/Private/Components/ActorComponent.cpp
--- a/Engine/Source/Runtime/Engine/Private/Components/ActorComponent.cpp
+++ b/Engine/Source/Runtime/Engine/Private/Components/ActorComponent.cpp
@@ -40,6 +40,7 @@
 #include "Materials/MaterialInterface.h"
 #include "ObjectCacheContext.h"
 #include "AutoRTFM/AutoRTFM.h"
+#include "UObject/RemoteObjectTransfer.h"
 
 #if WITH_EDITOR
 #include "Kismet2/ComponentEditorUtils.h"
@@ -1311,10 +1312,23 @@ TStructOnScope<FActorComponentInstanceData> UActorComponent::GetComponentInstanc
 void FActorComponentTickFunction::ExecuteTick(float DeltaTime, enum ELevelTick TickType, ENamedThreads::Type CurrentThread, const FGraphEventRef& MyCompletionGraphEvent)
 {
 	TRACE_CPUPROFILER_EVENT_SCOPE(FActorComponentTickFunction::ExecuteTick);
+
+#if UE_WITH_REMOTE_OBJECT_HANDLE
+	auto ExecuteTickWork = [this, DeltaTime, TickType, CurrentThread, &MyCompletionGraphEvent]()
+		{
+			ExecuteTickHelper(Target, Target->bTickInEditor, DeltaTime, TickType, [this, TickType](float DilatedTime)
+			{
+				Target->TickComponent(DilatedTime, TickType, this);
+			});
+		};
+
+	UE::RemoteObject::Transfer::ExecuteTransactional(ExecuteTickWork);
+#else
 	ExecuteTickHelper(Target, Target->bTickInEditor, DeltaTime, TickType, [this, TickType](float DilatedTime)
 	{
 		Target->TickComponent(DilatedTime, TickType, this);
 	});
+#endif
 }
 
 FString FActorComponentTickFunction::DiagnosticMessage()
