[Console] Fixed last line of scrollback buffer not being displayed and previous to last line overlapping with the input field

diff --git a/Engine/Source/Runtime/Engine/Private/UserInterface/Console.cpp b/Engine/Source/Runtime/Engine/Private/UserInterface/Console.cpp
--- a/Engine/Source/Runtime/Engine/Private/UserInterface/Console.cpp
+++ b/Engine/Source/Runtime/Engine/Private/UserInterface/Console.cpp
@@ -642,7 +642,7 @@ void UConsole::ClearOutput()
 void UConsole::OutputTextLine(const FString& Text)
 {
 	// If we are full, delete the first line
-	if (Scrollback.Num() > ConsoleSettings->MaxScrollbackSize)
+	if (Scrollback.Num() >= ConsoleSettings->MaxScrollbackSize)
 	{
 		Scrollback.RemoveAt(0, 1);
 		SBHead = ConsoleSettings->MaxScrollbackSize - 1;
@@ -1429,7 +1429,7 @@ void UConsole::PostRender_Console_Open(UCanvas* Canvas)
 		Canvas->DrawItem(ConsoleTile);
 	}
 
-	PostRender_InputLine(Canvas, FIntPoint(PosLeft, PosTop + Height + 6));
+	PostRender_InputLine(Canvas, FIntPoint(PosLeft, PosTop + Height + 6 + TextH));
 }
 
 void UConsole::DrawLine(
