Refactored how we detach and reattach external object packages during cooking to allow routing the process to entities.

diff --git a/Engine/Source/Runtime/Engine/Classes/Engine/Level.h b/Engine/Source/Runtime/Engine/Classes/Engine/Level.h
--- a/Engine/Source/Runtime/Engine/Classes/Engine/Level.h
+++ b/Engine/Source/Runtime/Engine/Classes/Engine/Level.h
@@ -873,10 +873,6 @@ private:
 
 	UPROPERTY(Transient)
 	TWeakObjectPtr<UObject> EditorPathOwner;
-
-	/** Temporary map of objects to their associated external packages. Used when detaching/attaching external actors packages during cook. */
-	UPROPERTY(Transient)
-	TMap<TObjectPtr<UObject>, TObjectPtr<UPackage>> ObjectsToExternalPackages;
 #endif // #if WITH_EDITORONLY_DATA
 
 	enum class ERouteActorInitializationState : uint8
@@ -939,6 +935,9 @@ public:
 	static void AddReferencedObjects(UObject* InThis, FReferenceCollector& Collector);
 	//~ End UObject Interface.
 
+	void PreSaveFromRoot(FObjectPreSaveRootContext ObjectSaveContext);
+	void PostSaveFromRoot(FObjectPostSaveRootContext ObjectSaveContext);
+	
 #if	WITH_EDITOR
 	//~Begin ITextureStreamingContainer Interface.
 	virtual void InitializeTextureStreamingContainer(uint32 InPackedTextureStreamingQualityLevelFeatureLevel) override;
@@ -1347,12 +1346,6 @@ public:
 	 */
 	static ENGINE_API UPackage* CreateActorPackage(UPackage* InLevelPackage, EActorPackagingScheme InActorPackagingScheme, const FString& InActorPath, const UObject* InMountPointContext = nullptr);
 
-	/**
-	 * Detach or reattach all level actors to from/to their external package
-	 * @param bReattach if false will detach actors from their external package until reattach is called, passing true will reattach actors, no-op for non external actors
-	 */
-	ENGINE_API void DetachAttachAllActorsPackages(bool bReattach);
-
 	/** 
 	*  Called after lighting was built and data gets propagated to this level
 	*  @param	bLightingSuccessful	 Whether lighting build was successful
diff --git a/Engine/Source/Runtime/Engine/Classes/GameFramework/Actor.h b/Engine/Source/Runtime/Engine/Classes/GameFramework/Actor.h
--- a/Engine/Source/Runtime/Engine/Classes/GameFramework/Actor.h
+++ b/Engine/Source/Runtime/Engine/Classes/GameFramework/Actor.h
@@ -1011,6 +1011,10 @@ private:
 	/** The UHLODLayer in which this actor should be included. */
 	UPROPERTY(EditAnywhere, Category = HLOD, meta = (DisplayName = "HLOD Layer"))
 	TObjectPtr<class UHLODLayer> HLODLayer;
+
+	/** Temporary list of objects that were outered to this actor's external package. Used when detaching/attaching external actor packages during cook. */
+	UPROPERTY(Transient)
+	TArray<TObjectPtr<UObject>> ObjectsDetachedFromExternalPackage;
 #endif
 
 public:
@@ -2407,6 +2411,9 @@ public:
 #endif // WITH_EDITOR
 	//~ End UObject Interface
 
+	ENGINE_API virtual void PreSaveFromRoot(FObjectPreSaveRootContext ObjectSaveContext);
+	ENGINE_API virtual void PostSaveFromRoot(FObjectPostSaveRootContext ObjectSaveContext);
+
 #if WITH_EDITOR
 	ENGINE_API virtual bool CanEditChangeComponent(const UActorComponent* Component, const FProperty* InProperty) const;
 #endif
diff --git a/Engine/Source/Runtime/Engine/Private/Actor.cpp b/Engine/Source/Runtime/Engine/Private/Actor.cpp
--- a/Engine/Source/Runtime/Engine/Private/Actor.cpp
+++ b/Engine/Source/Runtime/Engine/Private/Actor.cpp
@@ -1566,6 +1566,58 @@ AActor* AActor::FindActorInPackage(UPackage* InPackage, bool bEvenIfPendingKill)
 	return Actor;
 }
 
+void AActor::PreSaveFromRoot(FObjectPreSaveRootContext ObjectSaveContext)
+{
+#if WITH_EDITOR
+	if (ObjectSaveContext.IsCooking())
+	{
+		if (UPackage* ActorExternalPackage = GetExternalPackage())
+		{
+			DetachExternalPackage();
+
+			// Process objects found in the actor package
+			TArray<UObject*> Objects;
+			const bool bIncludeNestedSubobjects = false;
+			GetObjectsWithPackage(ActorExternalPackage, Objects, bIncludeNestedSubobjects, RF_NoFlags, EInternalObjectFlags::Garbage);
+
+			UPackage* ActorPackage = GetPackage();
+			for (UObject* Object : Objects)
+			{
+				if (!Cast<AActor>(Object) && Object->GetFName() != NAME_PackageMetaData)
+				{
+					// Move objects in the destination actor package
+					Object->Rename(nullptr, ActorPackage);
+
+					// Keep track of which package to reattach later
+					ObjectsDetachedFromExternalPackage.Add(Object);
+				}
+			}
+		}
+	}
+#endif
+}
+
+void AActor::PostSaveFromRoot(FObjectPostSaveRootContext ObjectSaveContext)
+{
+#if WITH_EDITOR
+	if (ObjectSaveContext.IsCooking())
+	{
+		ReattachExternalPackage();
+
+		if (UPackage* ActorExternalPackage = GetExternalPackage())
+		{
+			// Reouter objects previously found in the actors packages to their original packages
+			for (const TObjectPtr<UObject>& ObjectToReattach : ObjectsDetachedFromExternalPackage)
+			{
+				ObjectToReattach->Rename(nullptr, ActorExternalPackage);
+			}
+		}
+
+		ObjectsDetachedFromExternalPackage.Empty();
+	}
+#endif
+}
+
 bool AActor::Rename( const TCHAR* InName, UObject* NewOuter, ERenameFlags Flags )
 {
 	const bool bRenameTest = ((Flags & REN_Test) != 0);
diff --git a/Engine/Source/Runtime/Engine/Private/Level.cpp b/Engine/Source/Runtime/Engine/Private/Level.cpp
--- a/Engine/Source/Runtime/Engine/Private/Level.cpp
+++ b/Engine/Source/Runtime/Engine/Private/Level.cpp
@@ -496,6 +496,40 @@ void ULevel::AddReferencedObjects(UObject* InThis, FReferenceCollector& Collecto
 	Super::AddReferencedObjects( This, Collector );
 }
 
+void ULevel::PreSaveFromRoot(FObjectPreSaveRootContext ObjectSaveContext)
+{
+	for (AActor* Actor : Actors)
+	{
+		if (Actor)
+		{
+			Actor->PreSaveFromRoot(ObjectSaveContext);
+		}
+	}
+
+	// Update components and keep track off whether we need to clean them up afterwards.
+	if(!bAreComponentsCurrentlyRegistered)
+	{
+		UpdateLevelComponents(true);
+		ObjectSaveContext.SetCleanupRequired(true);
+	}
+}
+
+void ULevel::PostSaveFromRoot(FObjectPostSaveRootContext ObjectSaveContext)
+{
+	if( ObjectSaveContext.IsCleanupRequired() )
+	{
+		ClearLevelComponents();
+	}
+
+	for (AActor* Actor : Actors)
+	{
+		if (Actor)
+		{
+			Actor->PostSaveFromRoot(ObjectSaveContext);
+		}
+	}
+}
+
 bool ULevel::TryAddActorToList(AActor* InActor, bool bAddUnique)
 {
 	bool bResult = true;
@@ -4570,65 +4604,6 @@ UPackage* ULevel::CreateActorPackage(UPackage* InLevelPackage, EActorPackagingSc
 	return ActorPackage;
 }
 
-void ULevel::DetachAttachAllActorsPackages(bool bReattach)
-{
-	if (bReattach)
-	{
-		for (AActor* Actor : Actors)
-		{
-			if (Actor)
-			{
-				Actor->ReattachExternalPackage();
-			}
-		}
-
-		// Reouter objects previously found in the actors packages to their original packages
-		for (const TPair<TObjectPtr<UObject>, TObjectPtr<UPackage>>& ObjectToPackage : ObjectsToExternalPackages)
-		{
-			UObject* Object = ObjectToPackage.Key;
-			UPackage* Package = ObjectToPackage.Value;
-
-			Object->Rename(nullptr, Package);
-		}
-
-		ObjectsToExternalPackages.Empty();
-	}
-	else
-	{
-		UPackage* LevelPackage = GetPackage();
-
-		check(ObjectsToExternalPackages.IsEmpty());
-
-		for (AActor* Actor : Actors)
-		{
-			if (Actor)
-			{
-				if (UPackage* ActorExternalPackage = Actor->GetExternalPackage())
-				{
-					Actor->DetachExternalPackage();
-
-					// Process objects found in the source actor package
-					TArray<UObject*> Objects;
-					const bool bIncludeNestedSubobjects = false;
-
-					GetObjectsWithPackage(ActorExternalPackage, Objects, bIncludeNestedSubobjects, RF_NoFlags, EInternalObjectFlags::Garbage);
-					for (UObject* Object : Objects)
-					{
-						if (!Cast<AActor>(Object) && Object->GetFName() != NAME_PackageMetaData)
-						{
-							// Move objects in the destination level package
-							Object->Rename(nullptr, LevelPackage);
-
-							// Keep track of which package this object really belongs to
-							ObjectsToExternalPackages.Emplace(Object, ActorExternalPackage);
-						}
-					}
-				}
-			}
-		}
-	}
-}
-
 void ULevel::OnApplyNewLightingData(bool bLightingSuccessful)
 {
 	// Store level offset that was used during static light data build
diff --git a/Engine/Source/Runtime/Engine/Private/World.cpp b/Engine/Source/Runtime/Engine/Private/World.cpp
--- a/Engine/Source/Runtime/Engine/Private/World.cpp
+++ b/Engine/Source/Runtime/Engine/Private/World.cpp
@@ -1579,24 +1579,15 @@ void UWorld::PreSaveRoot(FObjectPreSaveRootContext ObjectSaveContext)
 	// Execute all pending actor construction scripts
 	FActorDeferredScriptManager::Get().FinishAllCompilation();
 
-	// if we are cooking this world, convert its persistent level to internal actors before doing so
-	if (ObjectSaveContext.IsCooking())
-	{
-		PersistentLevel->DetachAttachAllActorsPackages(/*bReattach*/false);
-	}
-#endif
-
-	// Update components and keep track off whether we need to clean them up afterwards.
 	if(!PersistentLevel->bAreComponentsCurrentlyRegistered)
 	{
-#if WITH_EDITOR
 		// Even if the World has been CleanedUp we find ourselves in a state where we'll add the components back again (and clean it up afterwards), so we reset the CleanedUp state to allow registration of components to occur.
 		check((bIsWorldInitialized && !IsCleanedUp()) || (IsCleanedUp() && !bIsWorldInitialized)); // since we'll be using bIsWorldInitialized to ensure we restore the CleanedUp state properly we check we can rely on it's value
 		ResetCleanedUpState();
-#endif
-		PersistentLevel->UpdateLevelComponents(true);
-		ObjectSaveContext.SetCleanupRequired(true);
 	}
+#endif
+
+	PersistentLevel->PreSaveFromRoot(ObjectSaveContext);
 }
 
 void UWorld::PostSaveRoot( FObjectPostSaveRootContext ObjectSaveContext )
@@ -1607,24 +1598,16 @@ void UWorld::PostSaveRoot( FObjectPostSaveRootContext ObjectSaveContext )
 		return;
 	}
 
+	PersistentLevel->PostSaveFromRoot(ObjectSaveContext);
+
+#if WITH_EDITOR
 	if( ObjectSaveContext.IsCleanupRequired() )
 	{
-		PersistentLevel->ClearLevelComponents();
-#if WITH_EDITOR
 		if (!bIsWorldInitialized)
 		{
 			// If we resetted the CleanUp state to accommodate save we restore it here. 
 			CleanupWorldTag = ++CleanupWorldGlobalTag;
 		}
-#endif
-	}
-
-#if WITH_EDITOR
-	// if we are cooking this world, convert its persistent level back its proper loading strategy if needed
-	// NOTE: can't use bCleanupIsRequired since we don't want to unregister component if they were registered...
-	if (ObjectSaveContext.IsCooking())
-	{
-		PersistentLevel->DetachAttachAllActorsPackages(/*bReattach*/true);
 	}
 
 	if (!ObjectSaveContext.IsProceduralSave() && !(ObjectSaveContext.GetSaveFlags() & SAVE_FromAutosave))
