[HWRT] Don't block render thread when Ray Tracing Geometry stream-in requests are cancelled.

Context:
- Requests were cancelled using FStreamingRequest::Reset() which also blocks waiting for the FBulkDataBatchRequest to complete so that it can also clear RequestBuffer.

Change:
- Modified FStreamingRequest::Reset() so that it expects requests to be completed instead of blocking until they are.
- Moved cancellation logic to CancelStreamingRequest(...) which cancels the FBulkDataBatchRequest without blocking and instead sets a flag identifying the request as cancelled.
- Modified ProcessCompletedStreamingRequests(...) to handle cancelled requests.
- Added some extra checks to request validity.
- Fixed typo in streaming stress test cvars.

Repro steps:
- Run cooked build
- r.Lumen.Visualize 1
- r.Streaming.StressTest.ExtraIOLatency 0
- r.Streaming.StressTest 1
- Increasing r.Streaming.StressTest.RayTracing.ExtraIOLatency should affect rate at which ray tracing BLAS streams-in (visible in Lumen Scene) without affecting frame rate.

diff --git a/Engine/Source/Runtime/Engine/Private/Rendering/RayTracingGeometryManager.cpp b/Engine/Source/Runtime/Engine/Private/Rendering/RayTracingGeometryManager.cpp
--- a/Engine/Source/Runtime/Engine/Private/Rendering/RayTracingGeometryManager.cpp
+++ b/Engine/Source/Runtime/Engine/Private/Rendering/RayTracingGeometryManager.cpp
@@ -299,17 +299,7 @@ void FRayTracingGeometryManager::ReleaseRayTracingGeometryHandle(RayTracingGeome
 	FRegisteredGeometry& RegisteredGeometry = RegisteredGeometries[Handle];
 
 	// Cancel associated streaming request if currently in-flight
-	if (RegisteredGeometry.StreamingRequestIndex != INDEX_NONE)
-	{
-		FStreamingRequest& StreamingRequest = StreamingRequests[RegisteredGeometry.StreamingRequestIndex];
-		check(StreamingRequest.GeometryHandle == Handle);
-
-		// TODO: Reset on a worker thread 
-		// because it can cause the render thread to block waiting for IO request to complete
-		StreamingRequest.Reset();
-
-		RegisteredGeometry.StreamingRequestIndex = INDEX_NONE;
-	}
+	CancelStreamingRequest(RegisteredGeometry);
 
 	if (RegisteredGeometry.Geometry->GroupHandle != INDEX_NONE)
 	{
@@ -394,22 +384,11 @@ void FRayTracingGeometryManager::SetRayTracingGeometryGroupCurrentFirstLODIndex(
 
 			if (!RegisteredGeometry.Geometry->IsEvicted())
 			{
+				// Cancel associated streaming request if currently in-flight
+				CancelStreamingRequest(RegisteredGeometry);
+
 				// TODO: call FRayTracingStreamableAsset::ReleaseForStreaming(...) instead
 				RegisteredGeometry.Geometry->ReleaseRHIForStreaming(Batcher);
-
-				if (RegisteredGeometry.StreamingRequestIndex != INDEX_NONE)
-				{
-					// cancel and reset pending streaming request
-
-					FStreamingRequest& StreamingRequest = StreamingRequests[RegisteredGeometry.StreamingRequestIndex];
-					checkf(StreamingRequest.GeometryHandle == GeometryHandle, TEXT("Unexpected ray tracing geometry streaming request owner."));
-
-					// TODO: Reset on a worker thread 
-					// because it can cause the render thread to block waiting for IO request to complete
-					StreamingRequest.Reset();
-
-					RegisteredGeometry.StreamingRequestIndex = INDEX_NONE;
-				}
 			}
 		}
 	}
@@ -924,6 +903,7 @@ bool FRayTracingGeometryManager::RequestRayTracingGeometryStreamIn(FRHICommandLi
 		RegisteredGeometry.StreamingRequestIndex = NextStreamingRequestIndex;
 
 		FStreamingRequest& StreamingRequest = StreamingRequests[NextStreamingRequestIndex];
+		checkf(!StreamingRequest.IsValid(), TEXT("Unused streaming request are expected to be in invalid state."));
 		NextStreamingRequestIndex = (NextStreamingRequestIndex + 1) % GRayTracingStreamingMaxPendingRequests;
 		++NumStreamingRequests;
 
@@ -972,16 +952,18 @@ void FRayTracingGeometryManager::ProcessCompletedStreamingRequests(FRHICommandLi
 		const int32 PendingRequestIndex = (StartPendingRequestIndex + Index) % GRayTracingStreamingMaxPendingRequests;
 		FStreamingRequest& PendingRequest = StreamingRequests[PendingRequestIndex];
 
-		if (!PendingRequest.IsValid())
-		{
-			++NumCompletedRequests;
-			continue;
-		}
+		checkf(PendingRequest.IsValid(), TEXT("Pending streaming request should be valid."));
 
 		if (PendingRequest.Request.IsCompleted())
 		{
 			++NumCompletedRequests;
 
+			if (PendingRequest.bCancelled)
+			{
+				PendingRequest.Reset();
+				continue;
+			}
+
 			FRegisteredGeometry& RegisteredGeometry = RegisteredGeometries[PendingRequest.GeometryHandle];
 
 			RegisteredGeometry.StreamingRequestIndex = INDEX_NONE;
@@ -990,11 +972,9 @@ void FRayTracingGeometryManager::ProcessCompletedStreamingRequests(FRHICommandLi
 
 			if (RegisteredGeometry.Geometry->IsEvicted() || RegisteredGeometry.Geometry->LODIndex < Group.CurrentFirstLODIdx)
 			{
-				// skip if geometry was evicted while streaming request was being processed
-				continue;
+				// do nothing since geometry was evicted while streaming request was being processed
 			}
-
-			if (!PendingRequest.Request.IsOk())
+			else if (!PendingRequest.Request.IsOk())
 			{
 				// Retry if IO request failed for some reason
 				RegisteredGeometry.StreamableAsset->IssueRequest(PendingRequest.Request, PendingRequest.RequestBuffer);
@@ -1023,6 +1003,20 @@ void FRayTracingGeometryManager::ProcessCompletedStreamingRequests(FRHICommandLi
 	SET_DWORD_STAT(STAT_RayTracingInflightStreamingRequests, NumStreamingRequests);
 }
 
+void FRayTracingGeometryManager::CancelStreamingRequest(FRegisteredGeometry& RegisteredGeometry)
+{
+	if (RegisteredGeometry.StreamingRequestIndex != INDEX_NONE)
+	{
+		FStreamingRequest& StreamingRequest = StreamingRequests[RegisteredGeometry.StreamingRequestIndex];
+		checkf(StreamingRequest.GeometryHandle == RegisteredGeometry.Geometry->RayTracingGeometryHandle,
+			TEXT("Ray tracing geometry streaming request owner mismatch (expected %d, got %d)."), RegisteredGeometry.Geometry->RayTracingGeometryHandle, StreamingRequest.GeometryHandle);
+
+		StreamingRequest.Cancel();
+
+		RegisteredGeometry.StreamingRequestIndex = INDEX_NONE;
+	}
+}
+
 void FRayTracingGeometryManager::BoostPriority(BuildRequestIndex InRequestIndex, float InBoostValue)
 {
 	FScopeLock ScopeLock(&RequestCS);
diff --git a/Engine/Source/Runtime/Engine/Private/Streaming/TextureStreamingHelpers.cpp b/Engine/Source/Runtime/Engine/Private/Streaming/TextureStreamingHelpers.cpp
--- a/Engine/Source/Runtime/Engine/Private/Streaming/TextureStreamingHelpers.cpp
+++ b/Engine/Source/Runtime/Engine/Private/Streaming/TextureStreamingHelpers.cpp
@@ -288,13 +288,13 @@ static TAutoConsoleVariable<int32> CVarStreamingStressTest(
 	ECVF_Cheat);
 
 static TAutoConsoleVariable<int32> CVarStreamingStressTestExtraIOLatency(
-	TEXT("r.Streaming.StressTest.ExtaIOLatency"),
+	TEXT("r.Streaming.StressTest.ExtraIOLatency"),
 	10,
 	TEXT("An extra latency in milliseconds for each stream-in request when doing the stress test."),
 	ECVF_Cheat);
 
 static TAutoConsoleVariable<int32> CVarStreamingStressTestRayTracingExtraIOLatency(
-	TEXT("r.Streaming.StressTest.RayTracing.ExtaIOLatency"),
+	TEXT("r.Streaming.StressTest.RayTracing.ExtraIOLatency"),
 	0,
 	TEXT("An extra latency in milliseconds for each ray tracing stream-in request when doing the stress test."),
 	ECVF_Cheat);
diff --git a/Engine/Source/Runtime/Engine/Public/Rendering/RayTracingGeometryManager.h b/Engine/Source/Runtime/Engine/Public/Rendering/RayTracingGeometryManager.h
--- a/Engine/Source/Runtime/Engine/Public/Rendering/RayTracingGeometryManager.h
+++ b/Engine/Source/Runtime/Engine/Public/Rendering/RayTracingGeometryManager.h
@@ -129,24 +129,36 @@ private:
 		FIoBuffer RequestBuffer;
 		FBulkDataBatchRequest Request;
 
-		RayTracingGeometryHandle GeometryHandle;
+		RayTracingGeometryHandle GeometryHandle = INDEX_NONE;
 
-		bool IsValid() const { return GeometryHandle != INDEX_NONE; }
+		bool bCancelled = false;
 
-		void Reset()
+		bool IsValid() const { return (GeometryHandle != INDEX_NONE) || bCancelled; }
+
+		void Cancel()
 		{
 			GeometryHandle = INDEX_NONE;
+			bCancelled = true;
 
 			if (Request.IsPending())
 			{
 				Request.Cancel();
-				Request.Wait(); // Even after calling Cancel(), we still need to Wait() before we can touch the RequestBuffer.
 			}
+		}
+
+		void Reset()
+		{
+			GeometryHandle = INDEX_NONE;
+			bCancelled = false;
+
+			checkf(!Request.IsPending(), TEXT("Can't cancel ray tracing geometry streaming request that is still in-flight."));
 			Request.Reset();
 			RequestBuffer = {};
 		}
 	};
 
+	void CancelStreamingRequest(FRegisteredGeometry& RegisteredGeometry);
+
 	TSparseArray<FRayTracingGeometryGroup> RegisteredGroups;
 
 	// Used for keeping track of geometries when ray tracing is dynamic
