Fix regression when creating area modifiers (introduced at 36006447).
SetAreaClassToReplace sets ApplyMode and bIsLowAreaModifier, calling SetApplyMode was overriding the change.

diff --git a/Engine/Source/Runtime/Engine/Private/AI/Navigation/NavigationModifier.cpp b/Engine/Source/Runtime/Engine/Private/AI/Navigation/NavigationModifier.cpp
--- a/Engine/Source/Runtime/Engine/Private/AI/Navigation/NavigationModifier.cpp
+++ b/Engine/Source/Runtime/Engine/Private/AI/Navigation/NavigationModifier.cpp
@@ -941,7 +941,6 @@ void FCompositeNavModifier::CreateAreaModifiers(const UPrimitiveComponent* PrimC
 		if (AreaClassToReplace)
 		{
 			AreaMod.SetAreaClassToReplace(AreaClassToReplace);
-			AreaMod.SetApplyMode(ENavigationAreaMode::Replace);
 		}
 		Add(AreaMod);
 	}
@@ -955,7 +954,6 @@ void FCompositeNavModifier::CreateAreaModifiers(const UPrimitiveComponent* PrimC
 		if (AreaClassToReplace)
 		{
 			AreaMod.SetAreaClassToReplace(AreaClassToReplace);
-			AreaMod.SetApplyMode(ENavigationAreaMode::Replace);
 		}
 		Add(AreaMod);
 	}
@@ -969,7 +967,6 @@ void FCompositeNavModifier::CreateAreaModifiers(const UPrimitiveComponent* PrimC
 			if (AreaClassToReplace)
 			{
 				AreaMod.SetAreaClassToReplace(AreaClassToReplace);
-				AreaMod.SetApplyMode(ENavigationAreaMode::Replace);
 			}
 			Add(AreaMod);
 		}
@@ -988,7 +985,6 @@ void FCompositeNavModifier::CreateAreaModifiers(const UPrimitiveComponent* PrimC
 		if (AreaClassToReplace)
 		{
 			AreaMod.SetAreaClassToReplace(AreaClassToReplace);
-			AreaMod.SetApplyMode(ENavigationAreaMode::Replace);
 		}
 		Add(AreaMod);
 	}
diff --git a/Engine/Source/Runtime/NavigationSystem/Private/NavModifierComponent.cpp b/Engine/Source/Runtime/NavigationSystem/Private/NavModifierComponent.cpp
--- a/Engine/Source/Runtime/NavigationSystem/Private/NavModifierComponent.cpp
+++ b/Engine/Source/Runtime/NavigationSystem/Private/NavModifierComponent.cpp
@@ -33,7 +33,6 @@ namespace UE::Navigation::ModComponent::Private
 		if (AreaClassToReplace)
 		{
 			AreaNavModifier.SetAreaClassToReplace(AreaClassToReplace);
-			AreaNavModifier.SetApplyMode(ENavigationAreaMode::Replace);
 		}
 		return AreaNavModifier;
 	}
diff --git a/Engine/Source/Runtime/NavigationSystem/Private/NavModifierVolume.cpp b/Engine/Source/Runtime/NavigationSystem/Private/NavModifierVolume.cpp
--- a/Engine/Source/Runtime/NavigationSystem/Private/NavModifierVolume.cpp
+++ b/Engine/Source/Runtime/NavigationSystem/Private/NavModifierVolume.cpp
@@ -129,7 +129,6 @@ void ANavModifierVolume::GetNavigationData(FNavigationRelevantData& Data) const
 			if (AreaClassToReplace)
 			{
 				AreaMod.SetAreaClassToReplace(AreaClassToReplace);
-				AreaMod.SetApplyMode(ENavigationAreaMode::Replace);
 			}
 			Data.Modifiers.SetMaskFillCollisionUnderneathForNavmesh(true);
 			Data.Modifiers.Add(AreaMod);
