[HWRT] Fix crash when ray tracing geometry is evicted due streaming requests not being cancelled.
- Cancel pending streaming requests when geometry is evicted.
- Moved eviction logic in FRayTracingGeometry manager to a helper function.

diff --git a/Engine/Source/Runtime/Engine/Private/Rendering/RayTracingGeometryManager.cpp b/Engine/Source/Runtime/Engine/Private/Rendering/RayTracingGeometryManager.cpp
--- a/Engine/Source/Runtime/Engine/Private/Rendering/RayTracingGeometryManager.cpp
+++ b/Engine/Source/Runtime/Engine/Private/Rendering/RayTracingGeometryManager.cpp
@@ -557,7 +557,7 @@ void FRayTracingGeometryManager::Tick(FRHICommandList& RHICmdList)
 			}
 			else if (IsRayTracingUsingReferenceBasedResidency() && RegisteredGeometry.Geometry->GetRHI() != nullptr)
 			{
-				RegisteredGeometry.Geometry->Evict();
+				EvictGeometry(RegisteredGeometry);
 			}
 		}
 	}
@@ -571,8 +571,7 @@ void FRayTracingGeometryManager::Tick(FRHICommandList& RHICmdList)
 			{
 				if (RegisteredGeometry.Geometry->GetRHI() != nullptr)
 				{
-					// TODO: call FRayTracingStreamableAsset::ReleaseForStreaming(...)
-					RegisteredGeometry.Geometry->Evict();
+					EvictGeometry(RegisteredGeometry);
 				}
 			}
 
@@ -748,8 +747,7 @@ void FRayTracingGeometryManager::Tick(FRHICommandList& RHICmdList)
 
 				check(RegisteredGeometry.Geometry->IsValid() && !RegisteredGeometry.Geometry->IsEvicted());
 
-				// TODO: call FRayTracingStreamableAsset::ReleaseForStreaming(...) instead
-				RegisteredGeometry.Geometry->Evict();
+				EvictGeometry(RegisteredGeometry);
 
 				++Index;
 			}
@@ -1017,6 +1015,14 @@ void FRayTracingGeometryManager::CancelStreamingRequest(FRegisteredGeometry& Reg
 	}
 }
 
+void FRayTracingGeometryManager::EvictGeometry(FRegisteredGeometry& RegisteredGeometry)
+{
+	// Cancel associated streaming request if currently in-flight
+	CancelStreamingRequest(RegisteredGeometry);
+
+	RegisteredGeometry.Geometry->Evict();
+}
+
 void FRayTracingGeometryManager::BoostPriority(BuildRequestIndex InRequestIndex, float InBoostValue)
 {
 	FScopeLock ScopeLock(&RequestCS);
diff --git a/Engine/Source/Runtime/Engine/Public/Rendering/RayTracingGeometryManager.h b/Engine/Source/Runtime/Engine/Public/Rendering/RayTracingGeometryManager.h
--- a/Engine/Source/Runtime/Engine/Public/Rendering/RayTracingGeometryManager.h
+++ b/Engine/Source/Runtime/Engine/Public/Rendering/RayTracingGeometryManager.h
@@ -159,6 +159,10 @@ private:
 
 	void CancelStreamingRequest(FRegisteredGeometry& RegisteredGeometry);
 
+	// Helper function to evict geometry + necessary streaming management
+	// Must call this instead of FRayTracingGeometry::Evict() directly
+	void EvictGeometry(FRegisteredGeometry& RegisteredGeometry);
+
 	TSparseArray<FRayTracingGeometryGroup> RegisteredGroups;
 
 	// Used for keeping track of geometries when ray tracing is dynamic
