Fix Incremental UnrealEditor MacArm64 build.

diff --git a/Engine/Source/Runtime/Engine/Classes/GameFramework/Actor.h b/Engine/Source/Runtime/Engine/Classes/GameFramework/Actor.h
--- a/Engine/Source/Runtime/Engine/Classes/GameFramework/Actor.h
+++ b/Engine/Source/Runtime/Engine/Classes/GameFramework/Actor.h
@@ -2362,8 +2362,7 @@ public:
 	ENGINE_API virtual bool IsSelectedInEditor() const override;
 
 	/** DuplicationSeedInterface prevents exposing the entire DupSeed map to every actor and provides a method to do so instead. */
-	UE_DEPRECATED(5.6, "Use PreDuplicateFromRoot instead.")
-	struct FDuplicationSeedInterface
+	struct UE_DEPRECATED(5.6, "Use PreDuplicateFromRoot instead.") FDuplicationSeedInterface
 	{
 	public:
 		FDuplicationSeedInterface(TMap<UObject*, UObject*>& InDuplicationSeed) {}
@@ -2373,8 +2372,10 @@ public:
 	};
 
 	/** Populate the duplication seed when duplicating the actor for PIE. This allows objects to be remapped to existing objects rather than duplicating. */
+	PRAGMA_DISABLE_DEPRECATION_WARNINGS
 	UE_DEPRECATED(5.6, "Use PreDuplicateFromRoot instead.")
 	virtual void PopulatePIEDuplicationSeed(FDuplicationSeedInterface& DupSeed) {}
+	PRAGMA_ENABLE_DEPRECATION_WARNINGS
 
 	/** Defines if preview should be shown when you select an actor, but none of its children */
 	virtual bool IsDefaultPreviewEnabled() const
