Skeletal Mesh: Add editor-only functions to set the relative transform of and to reparent a socket.

diff --git a/Engine/Source/Editor/SkeletonEditor/Private/SkeletonTreeSocketItem.cpp b/Engine/Source/Editor/SkeletonEditor/Private/SkeletonTreeSocketItem.cpp
--- a/Engine/Source/Editor/SkeletonEditor/Private/SkeletonTreeSocketItem.cpp
+++ b/Engine/Source/Editor/SkeletonEditor/Private/SkeletonTreeSocketItem.cpp
@@ -18,6 +18,36 @@
 
 #define LOCTEXT_NAMESPACE "FSkeletonTreeSocketItem"
 
+FSkeletonTreeSocketItem::FSkeletonTreeSocketItem(
+	USkeletalMeshSocket* InSocket,
+	ESocketParentType InParentType,
+	bool bInIsCustomized,
+	const TSharedRef<class ISkeletonTree>& InSkeletonTree
+	) :
+		FSkeletonTreeItem(InSkeletonTree),
+		Socket(InSocket),
+		ParentType(InParentType),
+		bIsCustomized(bInIsCustomized),
+		bInlineEditorExpanded(false)
+{
+	OnSocketChangedDelegate = InSocket->OnPropertyChanged().AddLambda(
+		[SkeletonTree=InSkeletonTree.ToWeakPtr()](const USkeletalMeshSocket*, const FProperty* InProperty)
+	{
+		if (InProperty->GetName() == GET_MEMBER_NAME_CHECKED(USkeletalMeshSocket, BoneName))
+		{
+			if (TSharedPtr<ISkeletonTree> Tree = SkeletonTree.Pin())
+			{
+				Tree->Refresh();
+			}
+		}
+	});
+}
+
+FSkeletonTreeSocketItem::~FSkeletonTreeSocketItem()
+{
+	Socket->OnPropertyChanged().Remove(OnSocketChangedDelegate);
+}
+
 void FSkeletonTreeSocketItem::GenerateWidgetForNameColumn( TSharedPtr< SHorizontalBox > Box, const TAttribute<FText>& FilterText, FIsSelected InIsSelected )
 {
 	const FSlateBrush* SocketIcon = ( ParentType == ESocketParentType::Mesh ) ?
diff --git a/Engine/Source/Editor/SkeletonEditor/Private/SkeletonTreeSocketItem.h b/Engine/Source/Editor/SkeletonEditor/Private/SkeletonTreeSocketItem.h
--- a/Engine/Source/Editor/SkeletonEditor/Private/SkeletonTreeSocketItem.h
+++ b/Engine/Source/Editor/SkeletonEditor/Private/SkeletonTreeSocketItem.h
@@ -20,15 +20,9 @@ class FSkeletonTreeSocketItem : public FSkeletonTreeItem
 public:
 	SKELETON_TREE_ITEM_TYPE(FSkeletonTreeSocketItem, FSkeletonTreeItem)
 
-	FSkeletonTreeSocketItem(USkeletalMeshSocket* InSocket, ESocketParentType InParentType, bool bInIsCustomized, const TSharedRef<class ISkeletonTree>& InSkeletonTree)
-		: FSkeletonTreeItem(InSkeletonTree)
-		, Socket(InSocket)
-		, ParentType(InParentType)
-		, bIsCustomized(bInIsCustomized)
-		, bInlineEditorExpanded(false)
-	{}
-	virtual ~FSkeletonTreeSocketItem()
-	{}
+	FSkeletonTreeSocketItem(USkeletalMeshSocket* InSocket, ESocketParentType InParentType, bool bInIsCustomized, const TSharedRef<class ISkeletonTree>& InSkeletonTree);
+
+	virtual ~FSkeletonTreeSocketItem() override;
 
 	/** ISkeletonTreeItem interface */
 	virtual void GenerateWidgetForNameColumn(TSharedPtr< SHorizontalBox > Box, const TAttribute<FText>& FilterText, FIsSelected InIsSelected) override;
@@ -88,4 +82,6 @@ private:
 	/** Delegate for when the context menu requests a rename */
 	DECLARE_DELEGATE(FOnRenameRequested);
 	FOnRenameRequested OnRenameRequested;
+
+	FDelegateHandle OnSocketChangedDelegate;
 };
diff --git a/Engine/Source/Runtime/Engine/Classes/Engine/SkeletalMeshSocket.h b/Engine/Source/Runtime/Engine/Classes/Engine/SkeletalMeshSocket.h
--- a/Engine/Source/Runtime/Engine/Classes/Engine/SkeletalMeshSocket.h
+++ b/Engine/Source/Runtime/Engine/Classes/Engine/SkeletalMeshSocket.h
@@ -12,8 +12,9 @@
 UCLASS(hidecategories=Object, hidecategories=Actor, MinimalAPI)
 class USkeletalMeshSocket : public UObject
 {
-	GENERATED_UCLASS_BODY()
+	GENERATED_BODY()
 
+public:
 	/** 
 	 *	Defines a named attachment location on the USkeletalMesh. 
 	 *	These are set up in editor and used as a shortcut instead of specifying 
@@ -27,19 +28,25 @@ class USkeletalMeshSocket : public UObject
 	FName BoneName;
 
 	UPROPERTY(EditAnywhere, BlueprintReadOnly, Category="Socket Parameters")
-	FVector RelativeLocation;
+	FVector RelativeLocation = FVector::ZeroVector;
 
 	UPROPERTY(EditAnywhere, BlueprintReadOnly, Category="Socket Parameters")
-	FRotator RelativeRotation;
+	FRotator RelativeRotation = FRotator::ZeroRotator;
 
 	UPROPERTY(EditAnywhere, BlueprintReadOnly, Category="Socket Parameters")
-	FVector RelativeScale;
+	FVector RelativeScale = FVector::OneVector;
 
 	/** If true then the hierarchy of bones this socket is attached to will always be 
 	    evaluated, even if it had previously been removed due to the current lod setting */
 	UPROPERTY(EditAnywhere, BlueprintReadOnly, Category="Socket Parameters")
-	bool bForceAlwaysAnimated;
+	bool bForceAlwaysAnimated = true;
 
+#if WITH_EDITOR
+	/** Change the sockets parent to a new bone. The skeleton is used to validate that the bone exists */
+	UFUNCTION(BlueprintCallable, Category="Socket Parameters")
+	ENGINE_API void SetSocketParent(USkeletalMesh* InSkeletalMesh, FName InBoneName);
+#endif
+	
 	UFUNCTION(BlueprintCallable, Category="Components|SkeletalMesh")
 	ENGINE_API FVector GetSocketLocation(const class USkeletalMeshComponent* SkelComp) const;
 
@@ -51,8 +58,15 @@ class USkeletalMeshSocket : public UObject
 	ENGINE_API bool GetSocketMatrix(FMatrix& OutMatrix, const class USkeletalMeshComponent* SkelComp) const;
 
 	/** returns FTransform of Socket local transform */
+	UFUNCTION(BlueprintCallable, Category="Socket Parameters")
 	ENGINE_API FTransform GetSocketLocalTransform() const;
 
+#if WITH_EDITOR
+	/** Sets the relative transform parameters of the socket to the given local FTransform */
+	UFUNCTION(BlueprintCallable, Category="Socket Parameters")
+	ENGINE_API void SetSocketLocalTransform(FTransform InTransform);
+#endif
+
 	/** Utility that returns the current transform for this socket. */
 	ENGINE_API FTransform GetSocketTransform( const class USkeletalMeshComponent* SkelComp ) const;
 
diff --git a/Engine/Source/Runtime/Engine/Private/SkeletalMesh.cpp b/Engine/Source/Runtime/Engine/Private/SkeletalMesh.cpp
--- a/Engine/Source/Runtime/Engine/Private/SkeletalMesh.cpp
+++ b/Engine/Source/Runtime/Engine/Private/SkeletalMesh.cpp
@@ -6346,12 +6346,6 @@ bool USkeletalMesh::ShouldEvaluatePostProcessAnimBP(int32 LODLevel) const
 /*-----------------------------------------------------------------------------
 USkeletalMeshSocket
 -----------------------------------------------------------------------------*/
-USkeletalMeshSocket::USkeletalMeshSocket(const FObjectInitializer& ObjectInitializer)
-	: Super(ObjectInitializer)
-	, bForceAlwaysAnimated(true)
-{
-	RelativeScale = FVector(1.0f, 1.0f, 1.0f);
-}
 
 void USkeletalMeshSocket::InitializeSocketFromLocation(const class USkeletalMeshComponent* SkelComp, FVector WorldLocation, FVector WorldNormal)
 {
@@ -6365,6 +6359,62 @@ void USkeletalMeshSocket::InitializeSocketFromLocation(const class USkeletalMesh
 	}
 }
 
+
+#if WITH_EDITOR
+void USkeletalMeshSocket::SetSocketParent(USkeletalMesh* InSkeletalMesh, FName InBoneName)
+{
+	if (!InSkeletalMesh)
+	{
+		UE_LOG(LogSkeletalMesh, Error, TEXT("SetSocketParent: No skeletal mesh asset given."));
+		return;
+	}
+
+	if (InBoneName == BoneName)
+	{
+		// Nothing more to do.
+		return;
+	}
+
+	// The socket can be owned by either the skeletal mesh or its associated skeleton. We need to ensure that the bone in
+	// question exists on the correct owner so we don't end up with a situation where a skeleton socket refers to a bone that
+	// only exists on the skeletal mesh.
+	if (GetOuter() == InSkeletalMesh)
+	{
+		if (InSkeletalMesh->GetRefSkeleton().FindBoneIndex(InBoneName) == INDEX_NONE)
+		{
+			UE_LOG(LogSkeletalMesh, Error, TEXT("SetSocketParent: The owning skeletal asset (%s) does not contain any bone named '%s'."),
+				*InSkeletalMesh->GetName(), *BoneName.ToString());
+			return;
+		}
+	}
+	else if (GetOuter() == InSkeletalMesh->GetSkeleton())
+	{
+		const USkeleton* Skeleton = InSkeletalMesh->GetSkeleton();
+		if (Skeleton->GetReferenceSkeleton().FindBoneIndex(InBoneName) == INDEX_NONE)
+		{
+			UE_LOG(LogSkeletalMesh, Error, TEXT("SetSocketParent: The owning skeleton (%s) does not contain any bone named '%s'."),
+				*InSkeletalMesh->GetName(), *BoneName.ToString());
+			return;
+		}
+	}
+	else
+	{
+		UE_LOG(LogSkeletalMesh, Error, TEXT("SetSocketParent: Neither the skeletal asset (%s) nor its skeleton are the owners of this socket (%s)."),
+			*InSkeletalMesh->GetName(), *SocketName.ToString());
+	}
+	
+	// Make sure we can undo this change.
+	SetFlags(RF_Transactional);
+	Modify();
+
+	BoneName = InBoneName;
+
+	// Let the world know.
+	ChangedEvent.Broadcast(this, GetClass()->FindPropertyByName(GET_MEMBER_NAME_CHECKED(USkeletalMeshSocket, BoneName)));
+}
+#endif
+
+
 FVector USkeletalMeshSocket::GetSocketLocation(const class USkeletalMeshComponent* SkelComp) const
 {
 	if (ensureAsRuntimeWarning(SkelComp))
@@ -6400,6 +6450,20 @@ FTransform USkeletalMeshSocket::GetSocketLocalTransform() const
 	return FTransform(RelativeRotation, RelativeLocation, RelativeScale);
 }
 
+#if WITH_EDITOR
+void USkeletalMeshSocket::SetSocketLocalTransform(FTransform InTransform)
+{
+	// Make sure we can undo this change.
+	SetFlags(RF_Transactional);
+	Modify();
+	
+	RelativeLocation = InTransform.GetLocation();
+	RelativeRotation = InTransform.GetRotation().Rotator();
+	RelativeScale = InTransform.GetScale3D();
+}
+#endif
+
+
 FTransform USkeletalMeshSocket::GetSocketTransform(const class USkeletalMeshComponent* SkelComp) const
 {
 	FTransform OutTM;
