Removed Hardware LZ compression flag from StaticMesh DDC key as it isn't necessary anymore and prevents sharing between some platforms

#tests DDC fill job: 672f9fa4e56f6317391fa575

diff --git a/Engine/Source/Runtime/Engine/Private/StaticMesh.cpp b/Engine/Source/Runtime/Engine/Private/StaticMesh.cpp
--- a/Engine/Source/Runtime/Engine/Private/StaticMesh.cpp
+++ b/Engine/Source/Runtime/Engine/Private/StaticMesh.cpp
@@ -3528,15 +3528,6 @@ static FString BuildStaticMeshDerivedDataKeySuffix(const ITargetPlatform* Target
 
 	KeySuffix.AppendChar(Mesh->bSupportUniformlyDistributedSampling ? TEXT('1') : TEXT('0'));
 
-	if (TargetPlatform->SupportsFeature(ETargetPlatformFeatures::HardwareLZDecompression))
-	{
-		KeySuffix += TEXT("_HWLZ");
-	}
-	else
-	{
-		KeySuffix += TEXT("_SWLZ");
-	}
-
 	// Value of this CVar affects index buffer <-> painted vertex color correspondence (see UE-51421).
 	static const TConsoleVariableData<int32>* CVar = IConsoleManager::Get().FindTConsoleVariableDataInt(TEXT("r.TriangleOrderOptimization"));
 
