Declare trace service class to allow inheritance

Move implementation to publicly available header to allow inheritance of the service. Additionally add a method to handle custom urls.

diff --git a/Engine/Source/Runtime/Engine/Private/TraceService.cpp b/Engine/Source/Runtime/Engine/Private/TraceService.cpp
--- a/Engine/Source/Runtime/Engine/Private/TraceService.cpp
+++ b/Engine/Source/Runtime/Engine/Private/TraceService.cpp
@@ -9,46 +9,12 @@
 #include "Trace/Trace.h"
 #include "TraceControlMessages.h"
 
-class FTraceServiceImpl
-{
-public:
-	FTraceServiceImpl();
-	FTraceServiceImpl(const TSharedPtr<IMessageBus>&);
-	virtual ~FTraceServiceImpl() {};
-	
-private:
-	void OnStatusPing(const FTraceControlStatusPing& Message, const TSharedRef<IMessageContext>& Context);
-	void OnChannelsPing(const FTraceControlChannelsPing& Message, const TSharedRef<IMessageContext>& Context);
-	void OnSettingsPing(const FTraceControlSettingsPing& Message, const TSharedRef<IMessageContext, ESPMode::ThreadSafe>& Context);
-	void OnDiscoveryPing(const FTraceControlDiscoveryPing& Message, const TSharedRef<IMessageContext>& Context);
-	void OnStop(const FTraceControlStop& Message, const TSharedRef<IMessageContext>& Context);
-	void OnSend(const FTraceControlSend& Message, const TSharedRef<IMessageContext>& Context);
-	void OnChannelSet(const FTraceControlChannelsSet& Message, const TSharedRef<IMessageContext>& Context);
-	void OnFile(const FTraceControlFile& Message, const TSharedRef<IMessageContext>& Context);
-	void OnSnapshotSend(const FTraceControlSnapshotSend& Message, const TSharedRef<IMessageContext>& Context);
-	void OnSnapshotFile(const FTraceControlSnapshotFile& Message, const TSharedRef<IMessageContext>& Context);
-	void OnPause(const FTraceControlPause& Message, const TSharedRef<IMessageContext>& Context);
-	void OnResume(const FTraceControlResume& Message, const TSharedRef<IMessageContext>& Context);
-	void OnBookmark(const FTraceControlBookmark& Message, const TSharedRef<IMessageContext>& Context);
-#if UE_SCREENSHOT_TRACE_ENABLED
-	void OnScreenshot(const FTraceControlScreenshot& Message, const TSharedRef<IMessageContext>& Context);
-#endif
-	void OnSetStatNamedEvents(const FTraceControlSetStatNamedEvents& Message, const TSharedRef<IMessageContext>& Context);
-	
-	static void FillTraceStatusMessage(FTraceControlStatus* Message);
-
-	TSharedPtr<FMessageEndpoint, ESPMode::ThreadSafe> MessageEndpoint;
-	FGuid SessionId;
-	FGuid InstanceId;
-};
-
-
-FTraceServiceImpl::FTraceServiceImpl()
-	: FTraceServiceImpl(IMessagingModule::Get().GetDefaultBus())
+FTraceService::FTraceService()
+	: FTraceService(IMessagingModule::Get().GetDefaultBus())
 {
 }
 
-FTraceServiceImpl::FTraceServiceImpl(const TSharedPtr<IMessageBus>& InBus)
+FTraceService::FTraceService(const TSharedPtr<IMessageBus>& InBus)
 {
 	SessionId = FApp::GetSessionId();
 	InstanceId = FApp::GetInstanceId();
@@ -56,23 +22,23 @@ FTraceServiceImpl::FTraceServiceImpl(const TSharedPtr<IMessageBus>& InBus)
 	if (InBus.IsValid())
 	{
 		MessageEndpoint = FMessageEndpoint::Builder("FTraceService", InBus.ToSharedRef())
-			.Handling<FTraceControlDiscoveryPing>(this, &FTraceServiceImpl::OnDiscoveryPing)
-			.Handling<FTraceControlChannelsSet>(this, &FTraceServiceImpl::OnChannelSet)
-			.Handling<FTraceControlStop>(this, &FTraceServiceImpl::OnStop)
-			.Handling<FTraceControlSend>(this, &FTraceServiceImpl::OnSend)
-			.Handling<FTraceControlFile>(this, &FTraceServiceImpl::OnFile)
-			.Handling<FTraceControlSnapshotSend>(this, &FTraceServiceImpl::OnSnapshotSend)
-			.Handling<FTraceControlSnapshotFile>(this, &FTraceServiceImpl::OnSnapshotFile)
-			.Handling<FTraceControlPause>(this, &FTraceServiceImpl::OnPause)
-			.Handling<FTraceControlResume>(this, &FTraceServiceImpl::OnResume)
-			.Handling<FTraceControlBookmark>(this, &FTraceServiceImpl::OnBookmark)
+			.Handling<FTraceControlDiscoveryPing>(this, &FTraceService::OnDiscoveryPing)
+			.Handling<FTraceControlChannelsSet>(this, &FTraceService::OnChannelSet)
+			.Handling<FTraceControlStop>(this, &FTraceService::OnStop)
+			.Handling<FTraceControlSend>(this, &FTraceService::OnSend)
+			.Handling<FTraceControlFile>(this, &FTraceService::OnFile)
+			.Handling<FTraceControlSnapshotSend>(this, &FTraceService::OnSnapshotSend)
+			.Handling<FTraceControlSnapshotFile>(this, &FTraceService::OnSnapshotFile)
+			.Handling<FTraceControlPause>(this, &FTraceService::OnPause)
+			.Handling<FTraceControlResume>(this, &FTraceService::OnResume)
+			.Handling<FTraceControlBookmark>(this, &FTraceService::OnBookmark)
 #if UE_SCREENSHOT_TRACE_ENABLED
-			.Handling<FTraceControlScreenshot>(this, &FTraceServiceImpl::OnScreenshot)
+			.Handling<FTraceControlScreenshot>(this, &FTraceService::OnScreenshot)
 #endif
-			.Handling<FTraceControlSetStatNamedEvents>(this, &FTraceServiceImpl::OnSetStatNamedEvents)
-			.Handling<FTraceControlStatusPing>(this, &FTraceServiceImpl::OnStatusPing)
-			.Handling<FTraceControlSettingsPing>(this, &FTraceServiceImpl::OnSettingsPing)
-			.Handling<FTraceControlChannelsPing>(this, &FTraceServiceImpl::OnChannelsPing);
+			.Handling<FTraceControlSetStatNamedEvents>(this, &FTraceService::OnSetStatNamedEvents)
+			.Handling<FTraceControlStatusPing>(this, &FTraceService::OnStatusPing)
+			.Handling<FTraceControlSettingsPing>(this, &FTraceService::OnSettingsPing)
+			.Handling<FTraceControlChannelsPing>(this, &FTraceService::OnChannelsPing);
 
 		if (!MessageEndpoint.IsValid())
 		{
@@ -99,7 +65,7 @@ FTraceServiceImpl::FTraceServiceImpl(const TSharedPtr<IMessageBus>& InBus)
 	}
 }
 
-void FTraceServiceImpl::FillTraceStatusMessage(FTraceControlStatus* Message)
+void FTraceService::FillTraceStatusMessage(FTraceControlStatus* Message)
 {
 	// Get the current endpoint and ids
 	Message->Endpoint = FTraceAuxiliary::GetTraceDestinationString();
@@ -120,31 +86,23 @@ void FTraceServiceImpl::FillTraceStatusMessage(FTraceControlStatus* Message)
 	Message->TraceSystemStatus = static_cast<uint8>(FTraceAuxiliary::GetTraceSystemStatus());
 }
 
-void FTraceServiceImpl::OnChannelSet(const FTraceControlChannelsSet& Message, const TSharedRef<IMessageContext>& Context)
+void FTraceService::OnChannelSet(const FTraceControlChannelsSet& Message, const TSharedRef<IMessageContext>& Context)
 {
 	FTraceAuxiliary::EnableChannels(Message.ChannelIdsToEnable);
 	FTraceAuxiliary::DisableChannels(Message.ChannelIdsToDisable);
 }
 
-void FTraceServiceImpl::OnStop(const FTraceControlStop& Message, const TSharedRef<IMessageContext>& Context)
+void FTraceService::OnStop(const FTraceControlStop& Message, const TSharedRef<IMessageContext>& Context)
 {
 	FTraceAuxiliary::Stop();
 }
 
-void FTraceServiceImpl::OnSend(const FTraceControlSend& Message, const TSharedRef<IMessageContext>& Context)
+void FTraceService::OnSend(const FTraceControlSend& Message, const TSharedRef<IMessageContext>& Context)
 {
-	FTraceAuxiliary::FOptions Options;
-	Options.bExcludeTail = Message.bExcludeTail;
-	
-	FTraceAuxiliary::Start(
-		FTraceAuxiliary::EConnectionType::Network,
-		*Message.Host,
-		*Message.Channels,
-		&Options
-	);
+	HandleSendUri(Message);
 }
 
-void FTraceServiceImpl::OnFile(const FTraceControlFile& Message, const TSharedRef<IMessageContext>& Context)
+void FTraceService::OnFile(const FTraceControlFile& Message, const TSharedRef<IMessageContext>& Context)
 {
 	FTraceAuxiliary::FOptions Options;
 	Options.bTruncateFile = Message.bTruncateFile;
@@ -158,39 +116,39 @@ void FTraceServiceImpl::OnFile(const FTraceControlFile& Message, const TSharedRe
 	);
 }
 
-void FTraceServiceImpl::OnSnapshotSend(const FTraceControlSnapshotSend& Message, const TSharedRef<IMessageContext>& Context)
+void FTraceService::OnSnapshotSend(const FTraceControlSnapshotSend& Message, const TSharedRef<IMessageContext>& Context)
 {
 	FTraceAuxiliary::SendSnapshot(*Message.Host);
 }
 
-void FTraceServiceImpl::OnSnapshotFile(const FTraceControlSnapshotFile& Message, const TSharedRef<IMessageContext>& Context)
+void FTraceService::OnSnapshotFile(const FTraceControlSnapshotFile& Message, const TSharedRef<IMessageContext>& Context)
 {
 	FTraceAuxiliary::SendSnapshot(*Message.File);
 }
 
-void FTraceServiceImpl::OnPause(const FTraceControlPause& Message, const TSharedRef<IMessageContext>& Context)
+void FTraceService::OnPause(const FTraceControlPause& Message, const TSharedRef<IMessageContext>& Context)
 {
 	FTraceAuxiliary::Pause();
 }
 
-void FTraceServiceImpl::OnResume(const FTraceControlResume& Message, const TSharedRef<IMessageContext>& Context)
+void FTraceService::OnResume(const FTraceControlResume& Message, const TSharedRef<IMessageContext>& Context)
 {
 	FTraceAuxiliary::Resume();
 }
 
-void FTraceServiceImpl::OnBookmark(const FTraceControlBookmark& Message, const TSharedRef<IMessageContext>& Context)
+void FTraceService::OnBookmark(const FTraceControlBookmark& Message, const TSharedRef<IMessageContext>& Context)
 {
 	TRACE_BOOKMARK(TEXT("%s"), *Message.Label);
 }
 
 #if UE_SCREENSHOT_TRACE_ENABLED
-void FTraceServiceImpl::OnScreenshot(const FTraceControlScreenshot& Message, const TSharedRef<IMessageContext>& Context)
+void FTraceService::OnScreenshot(const FTraceControlScreenshot& Message, const TSharedRef<IMessageContext>& Context)
 {
 	FTraceScreenshot::RequestScreenshot(Message.Name, Message.bShowUI);
 }
 #endif // UE_SCREENSHOT_TRACE_ENABLED
 
-void FTraceServiceImpl::OnSetStatNamedEvents(const FTraceControlSetStatNamedEvents& Message, const TSharedRef<IMessageContext>& Context)
+void FTraceService::OnSetStatNamedEvents(const FTraceControlSetStatNamedEvents& Message, const TSharedRef<IMessageContext>& Context)
 {
 	if (Message.bEnabled && GCycleStatsShouldEmitNamedEvents == 0)
 	{
@@ -202,7 +160,20 @@ void FTraceServiceImpl::OnSetStatNamedEvents(const FTraceControlSetStatNamedEven
 	}
 }
 
-void FTraceServiceImpl::OnStatusPing(const FTraceControlStatusPing& Message, const TSharedRef<IMessageContext>& Context)
+void FTraceService::HandleSendUri(const FTraceControlSend& Message)
+{
+	FTraceAuxiliary::FOptions Options;
+	Options.bExcludeTail = Message.bExcludeTail;
+	
+	FTraceAuxiliary::Start(
+		FTraceAuxiliary::EConnectionType::Network,
+		*Message.Host,
+		*Message.Channels,
+		&Options
+	);
+}
+
+void FTraceService::OnStatusPing(const FTraceControlStatusPing& Message, const TSharedRef<IMessageContext>& Context)
 {
 	FTraceControlStatus* Response = FMessageEndpoint::MakeMessage<FTraceControlStatus>();
 	FillTraceStatusMessage(Response);
@@ -210,7 +181,7 @@ void FTraceServiceImpl::OnStatusPing(const FTraceControlStatusPing& Message, con
 	MessageEndpoint->Send(Response, Context->GetSender());
 }
 
-void FTraceServiceImpl::OnChannelsPing(const FTraceControlChannelsPing& Message, const TSharedRef<IMessageContext>& Context)
+void FTraceService::OnChannelsPing(const FTraceControlChannelsPing& Message, const TSharedRef<IMessageContext>& Context)
 {
 	
 	struct FEnumerateUserData
@@ -258,7 +229,7 @@ void FTraceServiceImpl::OnChannelsPing(const FTraceControlChannelsPing& Message,
 	MessageEndpoint->Send(StatusResponse, Context->GetSender());
 }
 
-void FTraceServiceImpl::OnSettingsPing(const FTraceControlSettingsPing& Message, const TSharedRef<IMessageContext>& Context)
+void FTraceService::OnSettingsPing(const FTraceControlSettingsPing& Message, const TSharedRef<IMessageContext>& Context)
 {
 	FTraceControlSettings* Response = FMessageEndpoint::MakeMessage<FTraceControlSettings>();
 	UE::Trace::FInitializeDesc const* InitDesc = FTraceAuxiliary::GetInitializeDesc();
@@ -287,7 +258,7 @@ void FTraceServiceImpl::OnSettingsPing(const FTraceControlSettingsPing& Message,
 	MessageEndpoint->Send(Response, Context->GetSender());
 }
 
-void FTraceServiceImpl::OnDiscoveryPing(const FTraceControlDiscoveryPing& Message, const TSharedRef<IMessageContext>& Context)
+void FTraceService::OnDiscoveryPing(const FTraceControlDiscoveryPing& Message, const TSharedRef<IMessageContext>& Context)
 {
 	if ((!Message.SessionId.IsValid() && !Message.InstanceId.IsValid()) || (Message.InstanceId == FApp::GetInstanceId() || Message.SessionId == FApp::GetSessionId()))
 	{
@@ -304,13 +275,3 @@ void FTraceServiceImpl::OnDiscoveryPing(const FTraceControlDiscoveryPing& Messag
 	}
 }
 
-
-FTraceService::FTraceService()
-{
-	Impl = MakePimpl<FTraceServiceImpl>();
-}
-
-FTraceService::FTraceService(TSharedPtr<IMessageBus> InBus)
-{
-	Impl = MakePimpl<FTraceServiceImpl>(InBus);
-}
diff --git a/Engine/Source/Runtime/Engine/Public/TraceService.h b/Engine/Source/Runtime/Engine/Public/TraceService.h
--- a/Engine/Source/Runtime/Engine/Public/TraceService.h
+++ b/Engine/Source/Runtime/Engine/Public/TraceService.h
@@ -1,17 +1,50 @@
 // Copyright Epic Games, Inc. All Rights Reserved.
 #pragma once
 
-#include "Templates/PimplPtr.h"
 #include "Templates/SharedPointer.h"
+#include "TraceControlMessages.h"
 
 class IMessageBus;
+class IMessageContext;
+class FMessageEndpoint;
 
 class FTraceService
 {
 public:
 	ENGINE_API FTraceService();
-	ENGINE_API FTraceService(TSharedPtr<IMessageBus> Bus);
+	ENGINE_API FTraceService(const TSharedPtr<IMessageBus>&);
+	ENGINE_API virtual ~FTraceService() {};
 	
 private:
-	TPimplPtr<class FTraceServiceImpl> Impl;
+	void OnStatusPing(const FTraceControlStatusPing& Message, const TSharedRef<IMessageContext>& Context);
+	void OnChannelsPing(const FTraceControlChannelsPing& Message, const TSharedRef<IMessageContext>& Context);
+	void OnSettingsPing(const FTraceControlSettingsPing& Message, const TSharedRef<IMessageContext>& Context);
+	void OnDiscoveryPing(const FTraceControlDiscoveryPing& Message, const TSharedRef<IMessageContext>& Context);
+	void OnStop(const FTraceControlStop& Message, const TSharedRef<IMessageContext>& Context);
+	void OnSend(const FTraceControlSend& Message, const TSharedRef<IMessageContext>& Context);
+	void OnChannelSet(const FTraceControlChannelsSet& Message, const TSharedRef<IMessageContext>& Context);
+	void OnFile(const FTraceControlFile& Message, const TSharedRef<IMessageContext>& Context);
+	void OnSnapshotSend(const FTraceControlSnapshotSend& Message, const TSharedRef<IMessageContext>& Context);
+	void OnSnapshotFile(const FTraceControlSnapshotFile& Message, const TSharedRef<IMessageContext>& Context);
+	void OnPause(const FTraceControlPause& Message, const TSharedRef<IMessageContext>& Context);
+	void OnResume(const FTraceControlResume& Message, const TSharedRef<IMessageContext>& Context);
+	void OnBookmark(const FTraceControlBookmark& Message, const TSharedRef<IMessageContext>& Context);
+	void OnScreenshot(const FTraceControlScreenshot& Message, const TSharedRef<IMessageContext>& Context);
+	void OnSetStatNamedEvents(const FTraceControlSetStatNamedEvents& Message, const TSharedRef<IMessageContext>& Context);
+
+protected:
+	/**
+	 * Allows for overriding handling of how to connect when a send message is received. Default implementation
+	 * assumes uri is a valid host name and tries to establish a regular socket connection.
+	 * @param Message Request that was received
+	 */
+	ENGINE_API virtual void HandleSendUri(const FTraceControlSend& Message);
+
+private:
+	
+	static void FillTraceStatusMessage(FTraceControlStatus* Message);
+
+	TSharedPtr<FMessageEndpoint, ESPMode::ThreadSafe> MessageEndpoint;
+	FGuid SessionId;
+	FGuid InstanceId;
 };
