Changed method in use to restore CleanupWorldTag, we save wether we need to restore it in a flag since it can't be deduced using the initialized flag value

diff --git a/Engine/Source/Runtime/Engine/Classes/Engine/World.h b/Engine/Source/Runtime/Engine/Classes/Engine/World.h
--- a/Engine/Source/Runtime/Engine/Classes/Engine/World.h
+++ b/Engine/Source/Runtime/Engine/Classes/Engine/World.h
@@ -1792,6 +1792,10 @@ public:
 private:
 	uint32 bMarkedObjectsPendingKill:1;
 
+#if WITH_EDITOR
+	uint32 bWorldWasCleanedUp: 1;
+#endif
+
 	uint32 CleanupWorldTag;
 	static uint32 CleanupWorldGlobalTag;
 
diff --git a/Engine/Source/Runtime/Engine/Private/World.cpp b/Engine/Source/Runtime/Engine/Private/World.cpp
--- a/Engine/Source/Runtime/Engine/Private/World.cpp
+++ b/Engine/Source/Runtime/Engine/Private/World.cpp
@@ -668,6 +668,9 @@ UWorld::UWorld( const FObjectInitializer& ObjectInitializer )
 ,	FXSystem(NULL)
 ,	FlushLevelStreamingType(EFlushLevelStreamingType::None)
 ,	NextTravelType(TRAVEL_Relative)
+#if WITH_EDITOR
+,	bWorldWasCleanedUp(false)
+#endif
 ,	CleanupWorldTag(0)
 
 {
@@ -1582,7 +1585,7 @@ void UWorld::PreSaveRoot(FObjectPreSaveRootContext ObjectSaveContext)
 	if(!PersistentLevel->bAreComponentsCurrentlyRegistered)
 	{
 		// Even if the World has been CleanedUp we find ourselves in a state where we'll add the components back again (and clean it up afterwards), so we reset the CleanedUp state to allow registration of components to occur.
-		check((bIsWorldInitialized && !IsCleanedUp()) || (IsCleanedUp() && !bIsWorldInitialized)); // since we'll be using bIsWorldInitialized to ensure we restore the CleanedUp state properly we check we can rely on it's value
+		bWorldWasCleanedUp = IsCleanedUp();
 		ResetCleanedUpState();
 	}
 #endif
@@ -1603,7 +1606,7 @@ void UWorld::PostSaveRoot( FObjectPostSaveRootContext ObjectSaveContext )
 #if WITH_EDITOR
 	if( ObjectSaveContext.IsCleanupRequired() )
 	{
-		if (!bIsWorldInitialized)
+		if (!bWorldWasCleanedUp)
 		{
 			// If we resetted the CleanUp state to accommodate save we restore it here. 
 			CleanupWorldTag = ++CleanupWorldGlobalTag;
