Optimized UWorld::AddToWorld's adding primitives to scene during component incremental registration.
Optimization can be enabled using 'LevelStreaming.AsyncRegisterLevelContext.Enabled'.
Optimization's internal batching can be setup using 'LevelStreaming.AsyncRegisterLevelContext.PrimitiveBatchSize'.
Saves 22-25% of AddToWorld in CitySample's Big_City_LVL on a high end PC.

diff --git a/Engine/Source/Runtime/Engine/Classes/Components/ActorComponent.h b/Engine/Source/Runtime/Engine/Classes/Components/ActorComponent.h
--- a/Engine/Source/Runtime/Engine/Classes/Components/ActorComponent.h
+++ b/Engine/Source/Runtime/Engine/Classes/Components/ActorComponent.h
@@ -29,6 +29,7 @@ class ULevel;
 class UWorld;
 class UPrimitiveComponent;
 struct FSimpleMemberReference;
+struct FAsyncRegisterLevelContext;
 class IPrimitiveComponent;
 
 ENGINE_API extern int32 GEnableDeferredPhysicsCreation;
@@ -36,42 +37,22 @@ ENGINE_API extern int32 GEnableDeferredPhysicsCreation;
 class FRegisterComponentContext
 {
 public:
-	FRegisterComponentContext(UWorld* InWorld)
-		: World(InWorld)
-		, IncrementalUpdateTimeLimitMs(0)
-		, IncrementalUpdateAddPrimitiveBatchSize(0)
-		, IncrementalUpdateNumComponentsToUpdate(0)
-	{}
-
-	~FRegisterComponentContext()
-	{
-		Process();
-	}
+	using FPrimitiveBatch = TArray<UPrimitiveComponent*, FConcurrentLinearArrayAllocator>;
 
-	FRegisterComponentContext& SetIncrementalUpdateTimeLimitMs(double Value) { IncrementalUpdateTimeLimitMs = Value; return *this; }
-	FRegisterComponentContext& SetIncrementalUpdateAddPrimitiveBatchSize(int32 Value) { IncrementalUpdateAddPrimitiveBatchSize = Value; return *this; }
-	FRegisterComponentContext& SetIncrementalUpdateNumComponentsToUpdate(int32 Value) { IncrementalUpdateNumComponentsToUpdate = Value; return *this; }
+	FRegisterComponentContext(UWorld* InWorld);
+	~FRegisterComponentContext();
+
+	FRegisterComponentContext& SetIncrementalUpdateTimeLimitMs(double Value);
+	FRegisterComponentContext& SetIncrementalUpdateAddPrimitiveBatchSize(int32 Value);
+	FRegisterComponentContext& SetIncrementalUpdateNumComponentsToUpdate(int32 Value);
+	FRegisterComponentContext& SetAsyncRegisterLevelContext(FAsyncRegisterLevelContext* Value);
 
 	double GetIncrementalUpdateTimeLimitMs() const { return IncrementalUpdateTimeLimitMs; }
 	int32 GetIncrementalUpdateAddPrimitiveBatchSize() const { return IncrementalUpdateAddPrimitiveBatchSize; }
 	int32 GetIncrementalUpdateNumComponentsToUpdate() const { return IncrementalUpdateNumComponentsToUpdate; }
 
-	void AddPrimitive(UPrimitiveComponent* PrimitiveComponent)
-	{
-		checkSlow(!AddPrimitiveBatches.Contains(PrimitiveComponent));
-		AddPrimitiveBatches.Add(PrimitiveComponent);
-		
-		if ((IncrementalUpdateAddPrimitiveBatchSize > 0) && (AddPrimitiveBatches.Num() > IncrementalUpdateAddPrimitiveBatchSize))
-		{
-			Process();
-		}
-	}
-
-	void AddSendRenderDynamicData(UPrimitiveComponent* PrimitiveComponent)
-	{
-		checkSlow(!SendRenderDynamicDataPrimitives.Contains(PrimitiveComponent));
-		SendRenderDynamicDataPrimitives.Add(PrimitiveComponent);
-	}
+	void AddPrimitive(UPrimitiveComponent* PrimitiveComponent);
+	void AddSendRenderDynamicData(UPrimitiveComponent* PrimitiveComponent);
 
 	ENGINE_API static void SendRenderDynamicData(FRegisterComponentContext* Context, UPrimitiveComponent* PrimitiveComponent);
 
@@ -81,12 +62,18 @@ public:
 	void Process();
 
 private:
+	bool OnIncrementalRegisterComponentsDone();
+
 	UWorld* World;
 	double	IncrementalUpdateTimeLimitMs;
 	int32	IncrementalUpdateAddPrimitiveBatchSize;
 	int32	IncrementalUpdateNumComponentsToUpdate;
-	TArray<UPrimitiveComponent*, FConcurrentLinearArrayAllocator> AddPrimitiveBatches;
-	TArray<UPrimitiveComponent*, FConcurrentLinearArrayAllocator> SendRenderDynamicDataPrimitives;
+	FAsyncRegisterLevelContext* AsyncRegisterLevelContext;
+	
+	FPrimitiveBatch AddPrimitiveBatches;
+	FPrimitiveBatch SendRenderDynamicDataPrimitives;
+
+	friend class ULevel;
 };
 
 #if WITH_EDITOR
diff --git a/Engine/Source/Runtime/Engine/Classes/Engine/Level.h b/Engine/Source/Runtime/Engine/Classes/Engine/Level.h
--- a/Engine/Source/Runtime/Engine/Classes/Engine/Level.h
+++ b/Engine/Source/Runtime/Engine/Classes/Engine/Level.h
@@ -17,6 +17,7 @@
 #include "Engine/TextureStreamingTypes.h"
 #include "Misc/EditorPathObjectInterface.h"
 #include "Components/ActorComponent.h"
+#include "Streaming/AsyncRegisterLevelContext.h"
 #include <atomic>
 
 #include "Level.generated.h"
@@ -836,6 +837,10 @@ private:
 		return PreUnregisteringActorComponents.Num() > 0;
 	}
 
+	ENGINE_API void CreateAsyncRegisterLevelContext();
+	ENGINE_API void DestroyAsyncRegisterLevelContext();
+	
+	TUniquePtr<FAsyncRegisterLevelContext> AsyncRegisterLevelContext;
 	TSet<UActorComponent*> PreUnregisteringActorComponents;
 	TSet<UActorComponent*> PreRegisteringActorComponents;
 	friend struct FLevelRegistrationAccessor;
@@ -1558,6 +1563,22 @@ private:
 		return InLevel->HasPreUnregisteringComponents();
 	}
 
+	ENGINE_API static void CreateAsyncRegisterLevelContext(ULevel* InLevel)
+	{
+		InLevel->CreateAsyncRegisterLevelContext();
+	}
+
+	ENGINE_API static void DestroyAsyncRegisterLevelContext(ULevel* InLevel)
+	{
+		InLevel->DestroyAsyncRegisterLevelContext();
+	}
+
+	ENGINE_API static FAsyncRegisterLevelContext* GetAsyncRegisterLevelContext(ULevel* InLevel)
+	{
+		return InLevel ? InLevel->AsyncRegisterLevelContext.Get() : nullptr;
+	}
+
+	friend class FRegisterComponentContext;
 	friend class UActorComponent;
 	friend class UWorld;
 };
diff --git a/Engine/Source/Runtime/Engine/Internal/Streaming/AsyncRegisterLevelContext.cpp b/Engine/Source/Runtime/Engine/Internal/Streaming/AsyncRegisterLevelContext.cpp
new file mode 100644
--- /dev/null
+++ b/Engine/Source/Runtime/Engine/Internal/Streaming/AsyncRegisterLevelContext.cpp
@@ -0,0 +1,428 @@
+// Copyright Epic Games, Inc. All Rights Reserved.
+
+#include "AsyncRegisterLevelContext.h"
+#include "Async/ParallelFor.h"
+#include "Engine/World.h"
+#include "Engine/Level.h"
+#include "SceneInterface.h"
+#include "Algo/Transform.h"
+#include "Engine/CoreSettings.h"
+#include "Components/PrimitiveComponent.h"
+#if WITH_EDITOR
+#include "StaticMeshCompiler.h"
+#endif
+
+namespace LevelStreaming::AsyncRegisterLevelContext
+{
+	static bool bEnabled = false;
+	static FAutoConsoleVariableRef CVarEnabled(
+		TEXT("LevelStreaming.AsyncRegisterLevelContext.Enabled"),
+		bEnabled,
+		TEXT("Whether to allow level streaming to add primitives to the scene asynchronously while components are being incrementally registered."));
+
+	static int32 PrimitiveBatchSize = 16;
+	static FAutoConsoleVariableRef CVarPrimitiveBatchSize(
+		TEXT("LevelStreaming.AsyncRegisterLevelContext.PrimitiveBatchSize"),
+		PrimitiveBatchSize,
+		TEXT("The number of primitives before starting adding them to the scene during the incremental component registration.\n")
+		TEXT("Used when LevelStreaming.AsyncRegisterLevelContext.Enabled is true."));
+}
+
+// ------------------------------------
+// FAsyncRegisterLevelContext
+
+TUniquePtr<FAsyncRegisterLevelContext> FAsyncRegisterLevelContext::CreateInstance(ULevel* InLevel)
+{
+	if (LevelStreaming::AsyncRegisterLevelContext::bEnabled)
+	{
+		bool bSingleThreaded = !FApp::ShouldUseThreadingForPerformance();
+#if WITH_EDITOR
+		// This is required for async static mesh compilation in case a scene proxy is not async aware.
+		// A stall until the compilation is finished might occur, and this is only supported on the game thread for now.
+		bSingleThreaded |= FStaticMeshCompilingManager::Get().IsAsyncStaticMeshCompilationEnabled();
+#endif
+		if (!bSingleThreaded)
+		{
+			return MakeUnique<FAsyncRegisterLevelContext>(InLevel);
+		}
+	}
+	return nullptr;
+}
+
+FAsyncRegisterLevelContext::FAsyncRegisterLevelContext(ULevel* InLevel)
+	: Level(InLevel)
+	, AsyncAddPrimitiveQueue(*this)
+	, SendRenderDynamicDataPrimitivesQueue(*this)
+	, bCanLaunchNewTasks(true)
+	, bIncrementalRegisterComponentsDone(false)
+{
+}
+
+void FAsyncRegisterLevelContext::AddPrimitive(UPrimitiveComponent* InComponent)
+{
+	check(!bIncrementalRegisterComponentsDone);
+	AsyncAddPrimitiveQueue.AddPrimitive(InComponent);
+	Tick();
+}
+
+void FAsyncRegisterLevelContext::AddSendRenderDynamicData(UPrimitiveComponent* InComponent)
+{
+	check(!bIncrementalRegisterComponentsDone);
+	SendRenderDynamicDataPrimitivesQueue.AddSendRenderDynamicData(InComponent);
+	Tick();
+}
+
+void FAsyncRegisterLevelContext::SetIncrementalRegisterComponentsDone(bool bValue)
+{
+	bIncrementalRegisterComponentsDone = bValue;
+	if (bIncrementalRegisterComponentsDone)
+	{
+		Flush();
+	}
+}
+
+void FAsyncRegisterLevelContext::SetCanLaunchNewTasks(bool bValue)
+{
+	bCanLaunchNewTasks = bValue;
+}
+
+bool FAsyncRegisterLevelContext::HasRemainingWork() const
+{
+	return AsyncAddPrimitiveQueue.HasRemainingWork() || SendRenderDynamicDataPrimitivesQueue.HasRemainingWork();
+}
+
+bool FAsyncRegisterLevelContext::IsRunningAsync() const
+{
+	return AsyncAddPrimitiveQueue.IsRunningAsync() || SendRenderDynamicDataPrimitivesQueue.IsRunningAsync();
+}
+
+void FAsyncRegisterLevelContext::WaitForAsyncTasks()
+{
+	AsyncAddPrimitiveQueue.WaitForAsyncTask();
+	SendRenderDynamicDataPrimitivesQueue.WaitForAsyncTask();
+}
+
+void FAsyncRegisterLevelContext::Flush()
+{
+	TRACE_CPUPROFILER_EVENT_SCOPE(FAsyncRegisterLevelContext::Flush);
+	AsyncAddPrimitiveQueue.Flush();
+	SendRenderDynamicDataPrimitivesQueue.Flush();
+}
+
+bool FAsyncRegisterLevelContext::Tick()
+{
+	check(!bIncrementalRegisterComponentsDone || AsyncAddPrimitiveQueue.NextBatch.IsEmpty());
+	check(!bIncrementalRegisterComponentsDone || SendRenderDynamicDataPrimitivesQueue.NextBatch.IsEmpty());
+
+	TRACE_CPUPROFILER_EVENT_SCOPE(FAsyncRegisterLevelContext::Tick);
+	if (!AsyncAddPrimitiveQueue.Tick())
+	{
+		return false;
+	}
+	return SendRenderDynamicDataPrimitivesQueue.Tick();
+}
+
+// ------------------------------------
+// FAsyncAddPrimitiveQueue
+
+FAsyncAddPrimitiveQueue::FAsyncAddPrimitiveQueue(FAsyncRegisterLevelContext& InContext)
+	: Context(InContext)
+{
+}
+
+FAsyncAddPrimitiveQueue::~FAsyncAddPrimitiveQueue()
+{
+	check(!HasRemainingWork());
+}
+
+bool FAsyncAddPrimitiveQueue::HasRemainingWork() const
+{
+	return !AsyncTask.IsCompleted() || !AddPrimitivesArray.IsEmpty() || !NextBatch.IsEmpty();
+}
+
+bool FAsyncAddPrimitiveQueue::IsRunningAsync() const
+{
+	return AsyncTask.IsValid() && !AsyncTask.IsCompleted();
+}
+
+void FAsyncAddPrimitiveQueue::WaitForAsyncTask()
+{
+	if (AsyncTask.IsValid())
+	{
+		QUICK_SCOPE_CYCLE_COUNTER(STAT_AddPrimitivesTask_Wait);
+		AsyncTask.Wait();
+		AsyncTask.Reset();
+	}
+}
+
+bool FAsyncAddPrimitiveQueue::Tick()
+{
+	if (IsRunningAsync())
+	{
+		if (Context.bIncrementalRegisterComponentsDone)
+		{
+			if (AddPrimitivesArray.Num())
+			{
+				QUICK_SCOPE_CYCLE_COUNTER(STAT_AddPrimitivesTask_Execute_One_GameThread);
+				TWeakObjectPtr<UPrimitiveComponent> Component;
+				while (!AddPrimitivesArray.IsEmpty() && !Component.IsValid())
+				{
+					FPrimitiveBatch& Batch = AddPrimitivesArray[0];
+					while (!Batch.IsEmpty() && !Component.IsValid())
+					{
+						Component = Batch[0];
+						Batch.RemoveAtSwap(0, EAllowShrinking::No);
+					}
+					if (Batch.IsEmpty())
+					{
+						AddPrimitivesArray.RemoveAt(0, EAllowShrinking::No);
+					}
+				};
+
+				if (Component.IsValid())
+				{
+					FSceneInterface* Scene = Context.Level->GetWorld()->Scene;
+					const bool bAppCanEverRender = FApp::CanEverRender();
+					FAddPrimitivesTask::Execute(Component.Get(), Scene, bAppCanEverRender);
+					return false;
+				}
+			}
+			
+			// Nothing else to do, wait for task to finish
+			WaitForAsyncTask();
+		}
+		return false;
+	}
+	AsyncTask.Reset();
+
+	if (AddPrimitivesArray.Num())
+	{
+		// Don't launch any new task when context was marked as waiting for running task
+		if (Context.bCanLaunchNewTasks)
+		{
+			FSceneInterface* Scene = Context.Level->GetWorld()->Scene;
+			const bool bStartAsyncTask = !Context.bIncrementalRegisterComponentsDone;
+			if (bStartAsyncTask)
+			{
+				FPrimitiveBatch Batch = MoveTemp(AddPrimitivesArray[0]);
+				AddPrimitivesArray.RemoveAt(0, EAllowShrinking::No);
+				TArray<FPrimitiveBatch> Batches;
+				Batches.Emplace(MoveTemp(Batch));
+				AsyncTask.Launch(Batches, Scene);
+			}
+			else
+			{
+				QUICK_SCOPE_CYCLE_COUNTER(STAT_AddPrimitivesTask_Execute_GameThread);
+				const int32 MaxCount = GLevelStreamingAddPrimitiveGranularity ? GLevelStreamingAddPrimitiveGranularity : MAX_int32;
+				int32 TotalMoved = 0;
+				TArray<FPrimitiveBatch> Batches;
+				while (!AddPrimitivesArray.IsEmpty() && (TotalMoved < MaxCount))
+				{
+					if (TotalMoved + AddPrimitivesArray[0].Num() <= MaxCount)
+					{
+						TotalMoved += AddPrimitivesArray[0].Num();
+						Batches.Emplace(MoveTemp(AddPrimitivesArray[0]));
+						AddPrimitivesArray.RemoveAt(0, EAllowShrinking::No);
+					}
+					else
+					{
+						int32 Count = MaxCount - TotalMoved;
+						Batches.Emplace_GetRef().Append(AddPrimitivesArray[0].GetData(), Count);
+						AddPrimitivesArray[0].RemoveAt(0, Count, EAllowShrinking::No);
+						TotalMoved += Count;
+					}
+				};
+				FAddPrimitivesTask::Execute(Batches, Scene);
+			}
+		}
+		// Return false so that caller can test time limit
+		return false;
+	}
+
+	return !HasRemainingWork();
+}
+
+void FAsyncAddPrimitiveQueue::AddPrimitive(UPrimitiveComponent* InComponent)
+{
+	checkSlow(!NextBatch.Contains(InComponent));
+	NextBatch.Add(InComponent);
+	if (NextBatch.Num() >= LevelStreaming::AsyncRegisterLevelContext::PrimitiveBatchSize)
+	{
+		Flush();
+	}
+}
+
+void FAsyncAddPrimitiveQueue::Flush()
+{
+	if (!NextBatch.IsEmpty())
+	{
+		Push(NextBatch);
+		check(NextBatch.IsEmpty());
+	}
+}
+
+void FAsyncAddPrimitiveQueue::Push(FPrimitiveBatch& InAddPrimitives)
+{
+	check(!InAddPrimitives.IsEmpty());
+	AddPrimitivesArray.Emplace(MoveTemp(InAddPrimitives));
+}
+
+// ------------------------------------
+// FAsyncAddPrimitiveQueue::FAddPrimitivesTask
+
+void FAsyncAddPrimitiveQueue::FAddPrimitivesTask::Reset()
+{
+	check(IsCompleted());
+	Task = UE::Tasks::TTask<void>();
+	Batches.Reset();
+}
+
+void FAsyncAddPrimitiveQueue::FAddPrimitivesTask::Launch(TArray<FPrimitiveBatch>& InBatches, FSceneInterface* InScene)
+{
+	check(IsCompleted());
+	Batches = MoveTemp(InBatches);
+	Task = UE::Tasks::Launch(UE_SOURCE_LOCATION, [InScene, this]()
+	{
+		QUICK_SCOPE_CYCLE_COUNTER(STAT_AddPrimitivesTask_Execute_Async);
+		FAddPrimitivesTask::Execute(Batches, InScene);
+	});
+}
+
+void FAsyncAddPrimitiveQueue::FAddPrimitivesTask::Execute(UPrimitiveComponent* Component, FSceneInterface* InScene, bool bAppCanEverRender)
+{
+	// AActor::PostRegisterAllComponents (called by AActor::IncrementalRegisterComponents) can trigger code 
+	// that either unregisters or re-registers components. If unregistered, skip this component.
+	// If re-registered, FRegisterComponentContext is not passed, so SceneProxy can be created.
+	if (Component && Component->IsRegistered())
+	{
+		if (Component->IsRenderStateCreated() || !bAppCanEverRender)
+		{
+			// Skip if SceneProxy is already created
+			if (Component->SceneProxy == nullptr)
+			{
+				InScene->AddPrimitive(Component);
+			}
+		}
+		else // Fallback for some edge case where the component renderstate are missing
+		{
+			Component->CreateRenderState_Concurrent(nullptr);
+		}
+	}
+}
+
+void FAsyncAddPrimitiveQueue::FAddPrimitivesTask::Execute(const TArray<FPrimitiveBatch>& InBatches, FSceneInterface* InScene)
+{
+	check(InScene);
+
+	const bool bAppCanEverRender = FApp::CanEverRender();
+
+	int32 Num = 0;
+	for (const FPrimitiveBatch& Batch : InBatches)
+	{
+		Num += Batch.Num();
+	}
+
+	auto GetComponent = [&InBatches](int32 Index) -> UPrimitiveComponent*
+	{
+		int32 MaxIndex = 0;
+		for (const FPrimitiveBatch& Batch : InBatches)
+		{
+			if (Index < Batch.Num())
+			{
+				return Batch[Index].Get();
+			}
+			Index -= Batch.Num();
+		}
+		return nullptr;
+	};
+
+	ParallelFor(Num, [&](int32 Index)
+	{
+		FOptionalTaskTagScope Scope(ETaskTag::EParallelGameThread);
+		UPrimitiveComponent* Component = GetComponent(Index);
+		FAddPrimitivesTask::Execute(Component, InScene, bAppCanEverRender);
+	});
+}
+
+void FAsyncAddPrimitiveQueue::FAddPrimitivesTask::Wait()
+{
+	if (Task.IsValid())
+	{
+		Task.Wait();
+	}
+	check(IsCompleted());
+}
+
+bool FAsyncAddPrimitiveQueue::FAddPrimitivesTask::IsCompleted() const
+{
+	return !Task.IsValid() || Task.IsCompleted();
+}
+
+bool FAsyncAddPrimitiveQueue::FAddPrimitivesTask::IsValid() const
+{
+	return Task.IsValid();
+}
+
+// ------------------------------------
+// FSendRenderDynamicDataPrimitivesQueue
+
+FSendRenderDynamicDataPrimitivesQueue::FSendRenderDynamicDataPrimitivesQueue(FAsyncRegisterLevelContext& InContext)
+	: Context(InContext)
+{}
+
+FSendRenderDynamicDataPrimitivesQueue::~FSendRenderDynamicDataPrimitivesQueue()
+{
+	check(!HasRemainingWork());
+}
+
+void FSendRenderDynamicDataPrimitivesQueue::AddSendRenderDynamicData(UPrimitiveComponent* InComponent)
+{
+	checkSlow(!NextBatch.Contains(InComponent));
+	NextBatch.Add(InComponent);
+	if (NextBatch.Num() >= LevelStreaming::AsyncRegisterLevelContext::PrimitiveBatchSize)
+	{
+		Flush();
+	}
+}
+
+void FSendRenderDynamicDataPrimitivesQueue::Flush()
+{
+	if (!NextBatch.IsEmpty())
+	{
+		Push(NextBatch);
+		check(NextBatch.IsEmpty());
+	}
+}
+
+void FSendRenderDynamicDataPrimitivesQueue::Push(FPrimitiveBatch& InSendRenderDynamicDataPrimitives)
+{
+	check(!InSendRenderDynamicDataPrimitives.IsEmpty());
+	SendRenderDynamicDataPrimitivesArray.Emplace(MoveTemp(InSendRenderDynamicDataPrimitives));
+}
+
+bool FSendRenderDynamicDataPrimitivesQueue::Tick()
+{
+	if (SendRenderDynamicDataPrimitivesArray.Num())
+	{
+		FPrimitiveBatch SendRenderDynamicDataPrimitives = MoveTemp(SendRenderDynamicDataPrimitivesArray[0]);
+		SendRenderDynamicDataPrimitivesArray.RemoveAt(0, EAllowShrinking::No);
+
+		for (TWeakObjectPtr<UPrimitiveComponent>& Primitive : SendRenderDynamicDataPrimitives)
+		{
+			// With incremetal updates the component can be registered, added to send render data queue, then destroyed by an actor chain so we must test it's still valid before sending render data.
+			if (Primitive.IsValid())
+			{
+				Primitive->SendRenderDynamicData_Concurrent();
+			}
+		}
+		// Return false so that caller can test time limit
+		return false;
+	}
+
+	return !HasRemainingWork();
+}
+
+bool FSendRenderDynamicDataPrimitivesQueue::HasRemainingWork() const
+{
+	return !SendRenderDynamicDataPrimitivesArray.IsEmpty() || !NextBatch.IsEmpty();
+}
\ No newline at end of file
diff --git a/Engine/Source/Runtime/Engine/Internal/Streaming/AsyncRegisterLevelContext.h b/Engine/Source/Runtime/Engine/Internal/Streaming/AsyncRegisterLevelContext.h
new file mode 100644
--- /dev/null
+++ b/Engine/Source/Runtime/Engine/Internal/Streaming/AsyncRegisterLevelContext.h
@@ -0,0 +1,98 @@
+// Copyright Epic Games, Inc. All Rights Reserved.
+
+#pragma once
+
+#include "Tasks/Task.h"
+#include "Async/Fundamental/Task.h"
+#include "Components/ActorComponent.h"
+#include "Components/PrimitiveComponent.h"
+
+class ULevel;
+class FSceneInterface;
+
+struct FAsyncRegisterLevelContext;
+
+struct FAsyncAddPrimitiveQueue
+{
+private:
+	using FPrimitiveBatch = TArray<TWeakObjectPtr<UPrimitiveComponent>, FConcurrentLinearArrayAllocator>;
+	FAsyncAddPrimitiveQueue(FAsyncRegisterLevelContext& InContext);
+	~FAsyncAddPrimitiveQueue();
+	void AddPrimitive(UPrimitiveComponent* InComponent);
+	void Push(FPrimitiveBatch& InAddPrimitives);
+	bool Tick();
+	void Flush();
+	bool IsRunningAsync() const;
+	bool HasRemainingWork() const;
+	void WaitForAsyncTask();
+
+	struct FAddPrimitivesTask
+	{
+	public:
+		void Launch(TArray<FPrimitiveBatch>& InBatch, FSceneInterface* InScene);
+		bool IsValid() const;
+		bool IsCompleted() const;
+		void Wait();
+		void Reset();
+		static void Execute(const TArray<FPrimitiveBatch>& InBatch, FSceneInterface* InScene);
+		static void Execute(UPrimitiveComponent* Component, FSceneInterface* InScene, bool bAppCanEverRender);
+
+		TArray<FPrimitiveBatch> Batches;
+		UE::Tasks::TTask<void> Task;
+	};
+
+	FAsyncRegisterLevelContext& Context;
+	FPrimitiveBatch NextBatch;
+	TArray<FPrimitiveBatch> AddPrimitivesArray;
+	FAddPrimitivesTask AsyncTask;
+	friend FAsyncRegisterLevelContext;
+};
+
+struct FSendRenderDynamicDataPrimitivesQueue
+{
+private:
+	FSendRenderDynamicDataPrimitivesQueue(FAsyncRegisterLevelContext& InContext);
+	~FSendRenderDynamicDataPrimitivesQueue();
+	using FPrimitiveBatch = TArray<TWeakObjectPtr<UPrimitiveComponent>, FConcurrentLinearArrayAllocator>;
+	void AddSendRenderDynamicData(UPrimitiveComponent* InComponent);
+	void Push(FPrimitiveBatch& InSendRenderDynamicDataPrimitives);
+	bool Tick();
+	void Flush();
+	bool IsRunningAsync() const { return false; }
+	bool HasRemainingWork() const;
+	void WaitForAsyncTask() {}
+
+	FAsyncRegisterLevelContext& Context;
+	FPrimitiveBatch NextBatch;
+	TArray<FPrimitiveBatch> SendRenderDynamicDataPrimitivesArray;
+	friend FAsyncRegisterLevelContext;
+};
+
+struct FAsyncRegisterLevelContext
+{
+public:
+	using FPrimitiveBatch = FRegisterComponentContext::FPrimitiveBatch;
+	static TUniquePtr<FAsyncRegisterLevelContext> CreateInstance(ULevel* InLevel);
+	FAsyncRegisterLevelContext(ULevel* InLevel);
+	void AddPrimitive(UPrimitiveComponent* InComponent);
+	void AddSendRenderDynamicData(UPrimitiveComponent* InComponent);
+	void SetCanLaunchNewTasks(bool bValue);
+	void SetIncrementalRegisterComponentsDone(bool bValue);
+	bool GetIncrementalRegisterComponentsDone() const { return bIncrementalRegisterComponentsDone; }
+	bool Tick();
+	bool IsRunningAsync() const;
+	bool HasRemainingWork() const;
+	void WaitForAsyncTasks();
+	
+private:
+	void Flush();
+
+	ULevel* Level;
+	FAsyncAddPrimitiveQueue AsyncAddPrimitiveQueue;
+	FSendRenderDynamicDataPrimitivesQueue SendRenderDynamicDataPrimitivesQueue;
+	bool bCanLaunchNewTasks;
+	bool bIncrementalRegisterComponentsDone;
+
+	friend FAsyncAddPrimitiveQueue;
+	friend FSendRenderDynamicDataPrimitivesQueue;
+};
\ No newline at end of file
diff --git a/Engine/Source/Runtime/Engine/Private/Components/ActorComponent.cpp b/Engine/Source/Runtime/Engine/Private/Components/ActorComponent.cpp
--- a/Engine/Source/Runtime/Engine/Private/Components/ActorComponent.cpp
+++ b/Engine/Source/Runtime/Engine/Private/Components/ActorComponent.cpp
@@ -41,6 +41,7 @@
 #include "ObjectCacheContext.h"
 #include "AutoRTFM/AutoRTFM.h"
 #include "UObject/RemoteObjectTransfer.h"
+#include "Streaming/AsyncRegisterLevelContext.h"
 
 #if WITH_EDITOR
 #include "Kismet2/ComponentEditorUtils.h"
@@ -119,22 +120,126 @@ FAutoConsoleVariableRef CVarPrecachePSOsOnComponentRecreateRenderContext(
 	TEXT("If > 0, re-creating a component's render context will also re-precache the component's PSOs.")
 );
 
+FRegisterComponentContext::FRegisterComponentContext(UWorld* InWorld)
+	: World(InWorld)
+	, IncrementalUpdateTimeLimitMs(0)
+	, IncrementalUpdateAddPrimitiveBatchSize(0)
+	, IncrementalUpdateNumComponentsToUpdate(0)
+	, AsyncRegisterLevelContext(nullptr)
+{}
+
+FRegisterComponentContext::~FRegisterComponentContext()
+{
+	Process();
+
+	// Make sure level incremental registration has completed its asynchronous tasks
+	check(!AsyncRegisterLevelContext || !AsyncRegisterLevelContext->IsRunningAsync());
+}
+
+void FRegisterComponentContext::AddPrimitive(UPrimitiveComponent* PrimitiveComponent)
+{
+	if (AsyncRegisterLevelContext)
+	{
+		AsyncRegisterLevelContext->AddPrimitive(PrimitiveComponent);
+	}
+	else
+	{
+		checkSlow(!AddPrimitiveBatches.Contains(PrimitiveComponent));
+		AddPrimitiveBatches.Add(PrimitiveComponent);
+		if ((IncrementalUpdateAddPrimitiveBatchSize > 0) && (AddPrimitiveBatches.Num() >= IncrementalUpdateAddPrimitiveBatchSize))
+		{
+			Process();
+		}
+	}
+}
+
+void FRegisterComponentContext::AddSendRenderDynamicData(UPrimitiveComponent* PrimitiveComponent)
+{
+	if (AsyncRegisterLevelContext)
+	{
+		AsyncRegisterLevelContext->AddSendRenderDynamicData(PrimitiveComponent);
+	}
+	else
+	{
+		checkSlow(!SendRenderDynamicDataPrimitives.Contains(PrimitiveComponent));
+		SendRenderDynamicDataPrimitives.Add(PrimitiveComponent);
+	}
+}
+
+void FRegisterComponentContext::SendRenderDynamicData(FRegisterComponentContext* Context, UPrimitiveComponent* PrimitiveComponent)
+{
+	if (Context)
+	{
+		Context->AddSendRenderDynamicData(PrimitiveComponent);
+	}
+	else
+	{
+		PrimitiveComponent->SendRenderDynamicData_Concurrent();
+	}
+}
+
+FRegisterComponentContext& FRegisterComponentContext::SetIncrementalUpdateTimeLimitMs(double Value)
+{
+	IncrementalUpdateTimeLimitMs = Value;
+	return *this;
+}
+
+FRegisterComponentContext& FRegisterComponentContext::SetIncrementalUpdateAddPrimitiveBatchSize(int32 Value)
+{
+	IncrementalUpdateAddPrimitiveBatchSize = Value;
+	return *this;
+}
+
+FRegisterComponentContext& FRegisterComponentContext::SetIncrementalUpdateNumComponentsToUpdate(int32 Value)
+{
+	IncrementalUpdateNumComponentsToUpdate = Value;
+	return *this;
+}
+
+FRegisterComponentContext& FRegisterComponentContext::SetAsyncRegisterLevelContext(FAsyncRegisterLevelContext* Value)
+{
+	AsyncRegisterLevelContext = Value;
+	return *this;
+}
+
+bool FRegisterComponentContext::OnIncrementalRegisterComponentsDone()
+{
+	// We need to process pending adds prior to rerunning the construction scripts, which may internally
+	// perform removals / adds themselves.
+	if (AsyncRegisterLevelContext)
+	{
+		// Notify the context that we registered all actor components
+		AsyncRegisterLevelContext->SetIncrementalRegisterComponentsDone(true);
+		AsyncRegisterLevelContext->Tick();
+		return !AsyncRegisterLevelContext->HasRemainingWork();
+	}
+	else
+	{
+		Process();
+		return true;
+	}
+}
+
 void FRegisterComponentContext::Process()
 {
 	TRACE_CPUPROFILER_EVENT_SCOPE(FRegisterComponentContext::Process)
-
+	if (AsyncRegisterLevelContext)
+	{
+		AsyncRegisterLevelContext->Tick();
+	}
+	else
+	{
 	if (AddPrimitiveBatches.Num())
 	{
+			FSceneInterface* Scene = World->Scene;
+			const bool bAppCanEverRender = FApp::CanEverRender();
+
 		bool bSingleThreaded = !FApp::ShouldUseThreadingForPerformance();
 #if WITH_EDITOR
 		// This is required for async static mesh compilation in case a scene proxy is not async aware.
 		// A stall until the compilation is finished might occur, and this is only supported on the game thread for now.
 		bSingleThreaded |= FStaticMeshCompilingManager::Get().IsAsyncStaticMeshCompilationEnabled();
 #endif
-
-		FSceneInterface* Scene = World->Scene;
-		const bool bAppCanEverRender = FApp::CanEverRender();
-
 		ParallelFor(AddPrimitiveBatches.Num(),
 			[&](int32 Index)
 			{
@@ -176,16 +281,8 @@ void FRegisterComponentContext::Process()
 	SendRenderDynamicDataPrimitives.Empty();
 }
 
-void FRegisterComponentContext::SendRenderDynamicData(FRegisterComponentContext* Context, UPrimitiveComponent* PrimitiveComponent)
-{
-	if (Context)
-	{
-		Context->AddSendRenderDynamicData(PrimitiveComponent);
-	}
-	else
-	{
-		PrimitiveComponent->SendRenderDynamicData_Concurrent();
-	}
+	check(AddPrimitiveBatches.IsEmpty());
+	check(SendRenderDynamicDataPrimitives.IsEmpty());
 }
 
 void UpdateAllPrimitiveSceneInfosForSingleComponent(UActorComponent* InComponent, TSet<FSceneInterface*>* InScenesToUpdateAllPrimitiveSceneInfosForBatching /* = nullptr*/)
diff --git a/Engine/Source/Runtime/Engine/Private/Level.cpp b/Engine/Source/Runtime/Engine/Private/Level.cpp
--- a/Engine/Source/Runtime/Engine/Private/Level.cpp
+++ b/Engine/Source/Runtime/Engine/Private/Level.cpp
@@ -1779,6 +1779,7 @@ void ULevel::IncrementalUpdateComponents(int32 NumComponentsToUpdate, bool bReru
 	Context.SetIncrementalUpdateTimeLimitMs(InContext ? InContext->GetIncrementalUpdateTimeLimitMs() : 0);
 	Context.SetIncrementalUpdateNumComponentsToUpdate(NumComponentsToUpdate);
 	Context.SetIncrementalUpdateAddPrimitiveBatchSize(InContext ? InContext->GetIncrementalUpdateAddPrimitiveBatchSize() : 0);
+	Context.SetAsyncRegisterLevelContext(AsyncRegisterLevelContext.Get());
 
 	bool bAllowLoop;
 	do
@@ -1838,6 +1839,7 @@ void ULevel::IncrementalUpdateComponents(int32 NumComponentsToUpdate, bool bReru
 #endif
 
 		case EIncrementalComponentState::Finalize:
+			check(!AsyncRegisterLevelContext.IsValid() || !AsyncRegisterLevelContext->HasRemainingWork());
 			IncrementalComponentState = EIncrementalComponentState::Init;
 			CurrentActorIndexForIncrementalUpdate = 0;
 			bAreComponentsCurrentlyRegistered = true;
@@ -1916,8 +1918,20 @@ bool ULevel::IncrementalPreRegisterComponents(FRegisterComponentContext& Context
 	return false;
 }
 
+void ULevel::CreateAsyncRegisterLevelContext()
+{
+	check(!AsyncRegisterLevelContext.IsValid());
+	AsyncRegisterLevelContext = FAsyncRegisterLevelContext::CreateInstance(this);
+}
+
+void ULevel::DestroyAsyncRegisterLevelContext()
+{
+	AsyncRegisterLevelContext.Reset();
+}
+
 bool ULevel::IncrementalRegisterComponents(FRegisterComponentContext& Context)
 {
+	TRACE_CPUPROFILER_EVENT_SCOPE(ULevel::IncrementalRegisterComponents());
 	const int32 NumComponentsToUpdate = Context.GetIncrementalUpdateNumComponentsToUpdate();
 
 	if (OwningWorld)
@@ -1971,11 +1985,11 @@ bool ULevel::IncrementalRegisterComponents(FRegisterComponentContext& Context)
 
 	if (CurrentActorIndexForIncrementalUpdate >= Actors.Num())
 	{
-		// We need to process pending adds prior to rerunning the construction scripts, which may internally
-		// perform removals / adds themselves.
-		Context.Process();
-		CurrentActorIndexForIncrementalUpdate = 0;
-		return true;
+		if (Context.OnIncrementalRegisterComponentsDone())
+		{
+			CurrentActorIndexForIncrementalUpdate = 0;
+			return true;
+		}
 	}
 
 	return false;
diff --git a/Engine/Source/Runtime/Engine/Private/World.cpp b/Engine/Source/Runtime/Engine/Private/World.cpp
--- a/Engine/Source/Runtime/Engine/Private/World.cpp
+++ b/Engine/Source/Runtime/Engine/Private/World.cpp
@@ -130,6 +130,7 @@
 #include "Net/NetPing.h"
 #include "AssetRegistry/AssetRegistryModule.h"
 #include "ProfilingDebugging/CountersTrace.h"
+#include "Streaming/AsyncRegisterLevelContext.h"
 
 #if UE_WITH_IRIS
 #include "Iris/ReplicationSystem/ReplicationSystem.h"
@@ -3260,6 +3261,7 @@ void UWorld::AddToWorld( ULevel* Level, const FTransform& LevelTransform, bool b
 		
 		// Mark level as being the one in process of being made visible.
 		CurrentLevelPendingVisibility = Level;
+		FLevelRegistrationAccessor::CreateAsyncRegisterLevelContext(Level);
 
 		// Add to the UWorld's array of levels, which causes it to be rendered et al.
 		Levels.AddUnique( Level );
@@ -3375,6 +3377,15 @@ void UWorld::AddToWorld( ULevel* Level, const FTransform& LevelTransform, bool b
 		// from disk rather than duplicated
 		const bool bRerunConstructionScript = !FPlatformProperties::RequiresCookedData() && (!bIsGameWorld || !Level->bHasRerunConstructionScripts);
 
+		FAsyncRegisterLevelContext* AsyncRegisterLevelContext = FLevelRegistrationAccessor::GetAsyncRegisterLevelContext(Level);
+		auto SetLevelContextCanLaunchNewTasks = [AsyncRegisterLevelContext](bool bValue)
+		{
+			if (AsyncRegisterLevelContext)
+			{
+				AsyncRegisterLevelContext->SetCanLaunchNewTasks(bValue);
+			}
+		};
+
 		// Incremental update components
 		{
 			// Prepare context used to store batch/parallelize AddPrimitive calls
@@ -3385,6 +3396,8 @@ void UWorld::AddToWorld( ULevel* Level, const FTransform& LevelTransform, bool b
 				.SetIncrementalUpdateAddPrimitiveBatchSize(FMath::Max<int32>(0, GLevelStreamingAddPrimitiveGranularity));
 
 			// Incrementally update components.
+			bool bTimeLimitExceeded = false;
+			bool bHasRunningTasks = false;
 			const int32 NumComponentsToUpdate = (!bConsiderTimeLimit || !bIsGameWorld || IsRunningCommandlet() ? 0 : GLevelStreamingComponentsRegistrationGranularity);
 			do
 			{
@@ -3393,9 +3406,23 @@ void UWorld::AddToWorld( ULevel* Level, const FTransform& LevelTransform, bool b
 				{
 					break;
 				}
-			} while (!IsTimeLimitExceeded(TEXT("updating components"), StartTime, Level, TimeLimit));
+				bTimeLimitExceeded = IsTimeLimitExceeded(TEXT("updating components"), StartTime, Level, TimeLimit);
+				bHasRunningTasks = AsyncRegisterLevelContext ? AsyncRegisterLevelContext->IsRunningAsync() : false;
+				if (AsyncRegisterLevelContext && bTimeLimitExceeded && bHasRunningTasks)
+				{
+					// Block level context to launch new tasks since we reached the time limit
+					SetLevelContextCanLaunchNewTasks(false);
+				}
+			} while (!bTimeLimitExceeded || bHasRunningTasks);
+			
+			// Block level context to launch new tasks since we exit the IncrementalUpdateComponents scope
+			SetLevelContextCanLaunchNewTasks(false);
 		}
 
+		check(!Level->bAreComponentsCurrentlyRegistered || !AsyncRegisterLevelContext || !AsyncRegisterLevelContext->HasRemainingWork());
+		// Allow level context to launch new tasks
+		SetLevelContextCanLaunchNewTasks(true);
+
 		// We are done once all components are attached.
 		Level->bAlreadyUpdatedComponents	= Level->bAreComponentsCurrentlyRegistered;
 		bExecuteNextStep					= Level->bAreComponentsCurrentlyRegistered && (!bConsiderTimeLimit || !IsTimeLimitExceeded(TEXT("updating components"), StartTime, Level, TimeLimit));
@@ -3479,6 +3506,7 @@ void UWorld::AddToWorld( ULevel* Level, const FTransform& LevelTransform, bool b
 
 		// Finished making level visible - allow other levels to be added to the world.
 		CurrentLevelPendingVisibility = nullptr;
+		FLevelRegistrationAccessor::DestroyAsyncRegisterLevelContext(Level);
 
 		if (bIsGameWorld && !Level->bClientOnlyVisible)
 		{
