Remove outdated [Physx Only] property descriptions

This change removes the [Physx Only] property descriptions from the properties in BodyInstance. All of these properties now work in Chaos: Iteration Counts, Locking Constraints

diff --git a/Engine/Source/Editor/DetailCustomizations/Private/BodyInstanceCustomization.cpp b/Engine/Source/Editor/DetailCustomizations/Private/BodyInstanceCustomization.cpp
--- a/Engine/Source/Editor/DetailCustomizations/Private/BodyInstanceCustomization.cpp
+++ b/Engine/Source/Editor/DetailCustomizations/Private/BodyInstanceCustomization.cpp
@@ -1386,18 +1386,6 @@ void FBodyInstanceCustomizationHelper::CustomizeDetails(IDetailLayoutBuilder& De
 			AddMaxAngularVelocity(PhysicsCategory, BodyInstanceHandler);
 		}
 
-		// Hide legacy settings in Chaos
-		BodyInstanceHandler->GetChildHandle(GET_MEMBER_NAME_CHECKED(FBodyInstance, DOFMode))->MarkHiddenByCustomization();
-		BodyInstanceHandler->GetChildHandle(GET_MEMBER_NAME_CHECKED(FBodyInstance, bLockTranslation))->MarkHiddenByCustomization();
-		BodyInstanceHandler->GetChildHandle(GET_MEMBER_NAME_CHECKED(FBodyInstance, bLockRotation))->MarkHiddenByCustomization();
-		BodyInstanceHandler->GetChildHandle(GET_MEMBER_NAME_CHECKED(FBodyInstance, bLockXTranslation))->MarkHiddenByCustomization();
-		BodyInstanceHandler->GetChildHandle(GET_MEMBER_NAME_CHECKED(FBodyInstance, bLockYTranslation))->MarkHiddenByCustomization();
-		BodyInstanceHandler->GetChildHandle(GET_MEMBER_NAME_CHECKED(FBodyInstance, bLockZTranslation))->MarkHiddenByCustomization();
-		BodyInstanceHandler->GetChildHandle(GET_MEMBER_NAME_CHECKED(FBodyInstance, bLockXRotation))->MarkHiddenByCustomization();
-		BodyInstanceHandler->GetChildHandle(GET_MEMBER_NAME_CHECKED(FBodyInstance, bLockYRotation))->MarkHiddenByCustomization();
-		BodyInstanceHandler->GetChildHandle(GET_MEMBER_NAME_CHECKED(FBodyInstance, bLockZRotation))->MarkHiddenByCustomization();
-		BodyInstanceHandler->GetChildHandle(GET_MEMBER_NAME_CHECKED(FBodyInstance, CustomDOFPlaneNormal))->MarkHiddenByCustomization();
-
 		//Add the rest
 		uint32 NumChildren = 0;
 		BodyInstanceHandler->GetNumChildren(NumChildren);
diff --git a/Engine/Source/Runtime/Engine/Classes/PhysicsEngine/BodyInstance.h b/Engine/Source/Runtime/Engine/Classes/PhysicsEngine/BodyInstance.h
--- a/Engine/Source/Runtime/Engine/Classes/PhysicsEngine/BodyInstance.h
+++ b/Engine/Source/Runtime/Engine/Classes/PhysicsEngine/BodyInstance.h
@@ -320,15 +320,15 @@ struct FBodyInstance : public FBodyInstanceCore
 	/** When we are a body within a SkeletalMeshComponent, we cache the index of the bone we represent, to speed up sync'ing physics to anim. */
 	int16 InstanceBoneIndex;
 
-	/** [PhysX Only] This physics body's solver iteration count for position. Increasing this will be more CPU intensive, but better stabilized.  */
+	/** This physics body's solver iteration count for position. Increasing this will be more CPU intensive, but better stabilized.  */
 	UPROPERTY(EditAnywhere, AdvancedDisplay, BlueprintReadOnly, Category=Physics, meta = (ClampMin = "1", UIMin = "1"))
 	uint8 PositionSolverIterationCount;
 
-	/** [PhysX Only] This physics body's solver iteration count for velocity. Increasing this will be more CPU intensive, but better stabilized. */
+	/** This physics body's solver iteration count for velocity. Increasing this will be more CPU intensive, but better stabilized. */
 	UPROPERTY(EditAnywhere, AdvancedDisplay, BlueprintReadOnly, Category = Physics, meta = (ClampMin = "1", UIMin = "1"))
 	uint8 VelocitySolverIterationCount;
 
-	/** [PhysX Only] This physics body's solver iteration count for projection. Increasing this will be more CPU intensive, but better stabilized. */
+	/** This physics body's solver iteration count for projection. Increasing this will be more CPU intensive, but better stabilized. */
 	UPROPERTY(EditAnywhere, AdvancedDisplay, BlueprintReadOnly, Category = Physics, meta = (ClampMin = "1", UIMin = "1"))
 	uint8 ProjectionSolverIterationCount;
 
@@ -370,7 +370,7 @@ public:
 	UPROPERTY(EditAnywhere, AdvancedDisplay, BlueprintReadWrite, Category = Physics)
 	ESleepFamily SleepFamily;
 
-	/** [Physx Only] Locks physical movement along specified axis.*/
+	/** Locks physical movement along specified axis.*/
 	UPROPERTY(EditAnywhere, Category = Physics, meta = (DisplayName = "Mode"))
 	TEnumAsByte<EDOFMode::Type> DOFMode;
 
@@ -409,35 +409,35 @@ public:
 	/////////
 	// SIM SETTINGS
 
-	/** [Physx Only] When a Locked Axis Mode is selected, will lock translation on the specified axis*/
+	/** When a Locked Axis Mode is selected, will lock translation on the specified axis*/
 	UPROPERTY(EditAnywhere, Category = Physics, meta=(DisplayName = "Lock Axis Translation"))
 	uint8 bLockTranslation : 1;
 	
-	/** [Physx Only] When a Locked Axis Mode is selected, will lock rotation to the specified axis*/
+	/** When a Locked Axis Mode is selected, will lock rotation to the specified axis*/
 	UPROPERTY(EditAnywhere, Category = Physics, meta=(DisplayName = "Lock Axis Rotation"))
 	uint8 bLockRotation : 1;
 
-	/** [Physx Only] Lock translation along the X-axis*/
+	/** Lock translation along the X-axis*/
 	UPROPERTY(EditAnywhere, Category = Physics, meta = (DisplayName = "X"))
 	uint8 bLockXTranslation : 1;
 
-	/** [Physx Only] Lock translation along the Y-axis*/
+	/** Lock translation along the Y-axis*/
 	UPROPERTY(EditAnywhere, Category = Physics, meta = (DisplayName = "Y"))
 	uint8 bLockYTranslation : 1;
 
-	/** [Physx Only] Lock translation along the Z-axis*/
+	/** Lock translation along the Z-axis*/
 	UPROPERTY(EditAnywhere, Category = Physics, meta = (DisplayName = "Z"))
 	uint8 bLockZTranslation : 1;
 
-	/** [Physx Only] Lock rotation about the X-axis*/
+	/** Lock rotation about the X-axis*/
 	UPROPERTY(EditAnywhere, Category = Physics, meta = (DisplayName = "X"))
 	uint8 bLockXRotation : 1;
 
-	/** [Physx Only] Lock rotation about the Y-axis*/
+	/** Lock rotation about the Y-axis*/
 	UPROPERTY(EditAnywhere, Category = Physics, meta = (DisplayName = "Y"))
 	uint8 bLockYRotation : 1;
 
-	/** [Physx Only] Lock rotation about the Z-axis*/
+	/** Lock rotation about the Z-axis*/
 	UPROPERTY(EditAnywhere, Category = Physics, meta = (DisplayName = "Z"))
 	uint8 bLockZRotation : 1;
 
@@ -609,16 +609,15 @@ public:
 
 	ENGINE_API void ClearExternalCollisionProfile();
 
-	/** [Physx Only] Locks physical movement along axis. */
+	/** Locks physical movement along axis. */
 	ENGINE_API void SetDOFLock(EDOFMode::Type NewDOFMode);
 
-	/** [Physx Only] */
 	ENGINE_API FVector GetLockedAxis() const;
 	ENGINE_API void CreateDOFLock();
 
 	static ENGINE_API EDOFMode::Type ResolveDOFMode(EDOFMode::Type DOFMode);
 
-	/** [Physx Only] Constraint used to allow for easy DOF setup per bodyinstance */
+	/** Constraint used to allow for easy DOF setup per bodyinstance */
 	FConstraintInstance* DOFConstraint;
 
 	/** The parent body that we are welded to*/
