[Local Exposure] Improved parameterization of soft thresholds.
- Replaced smoothness with strength (where strength = 1 / (fade range in EV)
- strength = 1 results in existing threshold behavior, which means no longer need a checkbox to enable soft threshold, instead it's actually applied when strength < 1.

diff --git a/Engine/Plugins/Cameras/GameplayCameras/Source/GameplayCameras/Private/Core/PostProcessSettingsCollection.cpp b/Engine/Plugins/Cameras/GameplayCameras/Source/GameplayCameras/Private/Core/PostProcessSettingsCollection.cpp
--- a/Engine/Plugins/Cameras/GameplayCameras/Source/GameplayCameras/Private/Core/PostProcessSettingsCollection.cpp
+++ b/Engine/Plugins/Cameras/GameplayCameras/Source/GameplayCameras/Private/Core/PostProcessSettingsCollection.cpp
@@ -128,9 +128,8 @@ void FPostProcessSettingsCollection::OverrideChanged(const FPostProcessSettings&
 		UE_SET_PP(LocalExposureDetailStrength);
 		UE_SET_PP(LocalExposureBlurredLuminanceBlend);
 		UE_SET_PP(LocalExposureBlurredLuminanceKernelSizePercent);
-		UE_SET_PP(LocalExposureUseSoftThresholds);
-		UE_SET_PP(LocalExposureHighlightThresholdSmoothness);
-		UE_SET_PP(LocalExposureShadowThresholdSmoothness);
+		UE_SET_PP(LocalExposureHighlightThresholdStrength);
+		UE_SET_PP(LocalExposureShadowThresholdStrength);
 		UE_SET_PP(LocalExposureMiddleGreyBias);
 		UE_SET_PP(LensFlareIntensity);
 		UE_SET_PP(LensFlareTint);
@@ -434,9 +433,8 @@ void FPostProcessSettingsCollection::InternalLerpChanged(const FPostProcessSetti
 		UE_LERP_PP(LocalExposureDetailStrength);
 		UE_LERP_PP(LocalExposureBlurredLuminanceBlend);
 		UE_LERP_PP(LocalExposureBlurredLuminanceKernelSizePercent);
-		UE_SET_PP(LocalExposureUseSoftThresholds);
-		UE_LERP_PP(LocalExposureHighlightThresholdSmoothness);
-		UE_LERP_PP(LocalExposureShadowThresholdSmoothness);
+		UE_LERP_PP(LocalExposureHighlightThresholdStrength);
+		UE_LERP_PP(LocalExposureShadowThresholdStrength);
 		UE_LERP_PP(LocalExposureMiddleGreyBias);
 		UE_LERP_PP(LensFlareIntensity);
 		UE_LERP_PP(LensFlareTint);
diff --git a/Engine/Shaders/Private/PostProcessHistogramCommon.ush b/Engine/Shaders/Private/PostProcessHistogramCommon.ush
--- a/Engine/Shaders/Private/PostProcessHistogramCommon.ush
+++ b/Engine/Shaders/Private/PostProcessHistogramCommon.ush
@@ -43,10 +43,8 @@ float2 LocalExposure_BilateralGridUVScale;
 
 float LocalExposure_HighlightThreshold;
 float LocalExposure_ShadowThreshold;
-uint LocalExposure_bUseSoftThreshold;
-uint LocalExposure_bSoftThresholdUseSmoothStep;
-float LocalExposure_HighlightThresholdSmoothness;
-float LocalExposure_ShadowThresholdSmoothness;
+float LocalExposure_HighlightThresholdStrength;
+float LocalExposure_ShadowThresholdStrength;
 
 float CalculateEyeAdaptationLuminance(float3 Color)
 {
@@ -257,45 +255,21 @@ float CalculateLogLocalExposure(float LogLum, float BaseLogLum, float LogMiddleG
 
 	// Apply threshold
 	float ThresholdOffset = 0;
-
-	if(LocalExposure_bUseSoftThreshold != 0)
 	{
-		float Blend;
+		float M;
+		float T;
 		if (BaseCentered > 0)
 		{
-			if (LocalExposure_bSoftThresholdUseSmoothStep)
-			{
-				Blend = smoothstep(LocalExposure_HighlightThreshold, LocalExposure_HighlightThreshold + 1.0f / (1.0f - LocalExposure_HighlightThresholdSmoothness), BaseCentered);
-			}
-			else
-			{
-				Blend = saturate((BaseCentered - LocalExposure_HighlightThreshold) * (1.0f - LocalExposure_HighlightThresholdSmoothness));
-			}
+			M = max(0.0f, BaseCentered - LocalExposure_HighlightThreshold);
+			T = smoothstep(LocalExposure_HighlightThreshold, LocalExposure_HighlightThreshold + 1.0f / (1.0f - LocalExposure_HighlightThresholdStrength), abs(BaseCentered));
 		}
 		else
 		{
-			if (LocalExposure_bSoftThresholdUseSmoothStep)
-			{
-				Blend = smoothstep(LocalExposure_ShadowThreshold, LocalExposure_ShadowThreshold + 1.0f / (1.0f - LocalExposure_ShadowThresholdSmoothness), abs(BaseCentered));
-			}
-			else
-			{
-				Blend = saturate((abs(BaseCentered) - LocalExposure_ShadowThreshold) * (1.0f - LocalExposure_ShadowThresholdSmoothness));
-			}
+			M = min(0.0f, BaseCentered + LocalExposure_ShadowThreshold);
+			T = smoothstep(LocalExposure_ShadowThreshold, LocalExposure_ShadowThreshold + 1.0f / (1.0f - LocalExposure_ShadowThresholdStrength), abs(BaseCentered));
 		}
 
-		ContrastScale = lerp(1, ContrastScale, Blend);
-	}
-	else
-	{
-		if (BaseCentered > 0)
-		{
-			ThresholdOffset = BaseCentered - max(0.0f, BaseCentered - LocalExposure_HighlightThreshold);
-		}
-		else
-		{
-			ThresholdOffset = BaseCentered - min(0.0f, BaseCentered + LocalExposure_ShadowThreshold);
-		}
+		ThresholdOffset = (BaseCentered - M) * (1 - T);
 
 		BaseCentered -= ThresholdOffset;
 	}
diff --git a/Engine/Shaders/Private/PostProcessVisualizeLocalExposure.usf b/Engine/Shaders/Private/PostProcessVisualizeLocalExposure.usf
--- a/Engine/Shaders/Private/PostProcessVisualizeLocalExposure.usf
+++ b/Engine/Shaders/Private/PostProcessVisualizeLocalExposure.usf
@@ -94,42 +94,14 @@ float4 MainPS(noperspective float4 UVAndScreenPos : TEXCOORD0, float4 SvPosition
 		{
 			if (BaseCentered - LocalExposure_HighlightThreshold > 0)
 			{
-				float ThresholdFade = 1.0f;
-
-				if (LocalExposure_bUseSoftThreshold)
-				{
-					if (LocalExposure_bSoftThresholdUseSmoothStep)
-					{
-						ThresholdFade = smoothstep(LocalExposure_HighlightThreshold, LocalExposure_HighlightThreshold + 1.0f / (1.0f - LocalExposure_HighlightThresholdSmoothness), BaseCentered);
-					}
-					else
-					{
-						ThresholdFade = saturate((BaseCentered - LocalExposure_HighlightThreshold) * (1.0f - LocalExposure_HighlightThresholdSmoothness));
-					}
-				}
-
-				OutValue = lerp(OutValue, float3(0,1,0), ThresholdFade * DiagonalStripesPattern(SvPosition.xy, 32.0f));
+				OutValue = lerp(OutValue, float3(0,1,0), DiagonalStripesPattern(SvPosition.xy, 32.0f));
 			}
 		}
 		else
 		{
 			if (BaseCentered + LocalExposure_ShadowThreshold < 0)
 			{
-				float ThresholdFade = 1.0f;
-
-				if (LocalExposure_bUseSoftThreshold)
-				{
-					if (LocalExposure_bSoftThresholdUseSmoothStep)
-					{
-						ThresholdFade = smoothstep(LocalExposure_ShadowThreshold, LocalExposure_ShadowThreshold + 1.0f / (1.0f - LocalExposure_ShadowThresholdSmoothness), abs(BaseCentered));
-					}
-					else
-					{
-						ThresholdFade = saturate((abs(BaseCentered) - LocalExposure_ShadowThreshold) * (1.0f - LocalExposure_ShadowThresholdSmoothness));
-					}
-				}
-
-				OutValue = lerp(OutValue, float3(1,0,0), ThresholdFade * DiagonalStripesPattern(SvPosition.xy, 32.0f));
+				OutValue = lerp(OutValue, float3(1,0,0), DiagonalStripesPattern(SvPosition.xy, 32.0f));
 			}
 		}
 	}
diff --git a/Engine/Source/Runtime/Engine/Classes/Engine/Scene.h b/Engine/Source/Runtime/Engine/Classes/Engine/Scene.h
--- a/Engine/Source/Runtime/Engine/Classes/Engine/Scene.h
+++ b/Engine/Source/Runtime/Engine/Classes/Engine/Scene.h
@@ -933,13 +933,10 @@ struct FPostProcessSettings
 	uint8 bOverride_LocalExposureBlurredLuminanceKernelSizePercent:1;
 
 	UPROPERTY(EditAnywhere, BlueprintReadWrite, Category=Overrides, meta=(PinHiddenByDefault, InlineEditConditionToggle))
-	uint8 bOverride_LocalExposureUseSoftThresholds:1;
+	uint8 bOverride_LocalExposureHighlightThresholdStrength :1;
 
 	UPROPERTY(EditAnywhere, BlueprintReadWrite, Category=Overrides, meta=(PinHiddenByDefault, InlineEditConditionToggle))
-	uint8 bOverride_LocalExposureHighlightThresholdSmoothness :1;
-
-	UPROPERTY(EditAnywhere, BlueprintReadWrite, Category=Overrides, meta=(PinHiddenByDefault, InlineEditConditionToggle))
-	uint8 bOverride_LocalExposureShadowThresholdSmoothness :1;
+	uint8 bOverride_LocalExposureShadowThresholdStrength :1;
 
 	UPROPERTY(EditAnywhere, BlueprintReadWrite, Category=Overrides, meta=(PinHiddenByDefault, InlineEditConditionToggle))
 	uint8 bOverride_LocalExposureMiddleGreyBias:1;
@@ -1956,21 +1953,17 @@ struct FPostProcessSettings
 	UPROPERTY(interp, BlueprintReadWrite, Category="Lens|Local Exposure", meta=(UIMin = "0.0", UIMax = "100.0", editcondition = "bOverride_LocalExposureBlurredLuminanceKernelSizePercent", DisplayName = "Blurred Luminance Kernel Size Percent"))
 	float LocalExposureBlurredLuminanceKernelSizePercent;
 
-	/** Enables soft thresholds. */
-	UPROPERTY(EditAnywhere, BlueprintReadWrite, Category = "Lens|Local Exposure", AdvancedDisplay, meta = (editcondition = "bOverride_LocalExposureUseSoftThresholds", DisplayName = "Use Soft Thresholds (Experimental)"))
-	uint32 LocalExposureUseSoftThresholds : 1;
-
 	/** 
-	 * Smoothness of transition between regions considered highlights (only relevant when using Soft Thresholds).
+	 * Strength of the highlight threshold.
 	*/
-	UPROPERTY(interp, BlueprintReadWrite, Category="Lens|Local Exposure", AdvancedDisplay, meta=(UIMin = "0.0", UIMax = "1.0", editcondition = "bOverride_LocalExposureHighlightThresholdSmoothness", DisplayName = "Highlight Threshold Smoothness"))
-	float LocalExposureHighlightThresholdSmoothness;
+	UPROPERTY(interp, BlueprintReadWrite, Category="Lens|Local Exposure", AdvancedDisplay, meta=(UIMin = "0.0", UIMax = "1.0", editcondition = "bOverride_LocalExposureHighlightThresholdStrength", DisplayName = "Highlight Threshold Strength"))
+	float LocalExposureHighlightThresholdStrength;
 
 	/**
-	 * Smoothness of transition between regions considered shadows (only relevant when using Soft Thresholds).
+	 * Strength of the shadow threshold.
 	*/
-	UPROPERTY(interp, BlueprintReadWrite, Category="Lens|Local Exposure", AdvancedDisplay, meta=(UIMin = "0.0", UIMax = "1.0", editcondition = "bOverride_LocalExposureShadowThresholdSmoothness", DisplayName = "Shadow Threshold Smoothness"))
-	float LocalExposureShadowThresholdSmoothness;
+	UPROPERTY(interp, BlueprintReadWrite, Category="Lens|Local Exposure", AdvancedDisplay, meta=(UIMin = "0.0", UIMax = "1.0", editcondition = "bOverride_LocalExposureShadowThresholdStrength", DisplayName = "Shadow Threshold Strength"))
+	float LocalExposureShadowThresholdStrength;
 
 	/**
 	 * Logarithmic adjustment for the local exposure middle grey.
diff --git a/Engine/Source/Runtime/Engine/Private/PostProcessVolume.cpp b/Engine/Source/Runtime/Engine/Private/PostProcessVolume.cpp
--- a/Engine/Source/Runtime/Engine/Private/PostProcessVolume.cpp
+++ b/Engine/Source/Runtime/Engine/Private/PostProcessVolume.cpp
@@ -175,13 +175,6 @@ bool APostProcessVolume::CanEditChange(const FProperty* InProperty) const
 				return Settings.AutoExposureMethod == EAutoExposureMethod::AEM_Histogram;
 			}
 
-			// Parameters that are only used when using soft thhresholds
-			if (PropertyName == GET_MEMBER_NAME_STRING_CHECKED(FPostProcessSettings, LocalExposureHighlightThresholdSmoothness) ||
-				PropertyName == GET_MEMBER_NAME_STRING_CHECKED(FPostProcessSettings, LocalExposureShadowThresholdSmoothness))
-			{
-				return Settings.LocalExposureUseSoftThresholds != 0;
-			}
-
 			// Parameters that are only used for the Sum of Gaussian bloom / not the texture based fft bloom
 			if (PropertyName == GET_MEMBER_NAME_STRING_CHECKED(FPostProcessSettings, BloomThreshold) ||
 				PropertyName == GET_MEMBER_NAME_STRING_CHECKED(FPostProcessSettings, BloomSizeScale) ||
diff --git a/Engine/Source/Runtime/Engine/Private/Scene.cpp b/Engine/Source/Runtime/Engine/Private/Scene.cpp
--- a/Engine/Source/Runtime/Engine/Private/Scene.cpp
+++ b/Engine/Source/Runtime/Engine/Private/Scene.cpp
@@ -517,9 +517,8 @@ FPostProcessSettings::FPostProcessSettings()
 	LocalExposureDetailStrength = 1.0f;
 	LocalExposureBlurredLuminanceBlend = 0.6f;
 	LocalExposureBlurredLuminanceKernelSizePercent = 50.0f;
-	LocalExposureUseSoftThresholds = false;
-	LocalExposureHighlightThresholdSmoothness = 0.6f;
-	LocalExposureShadowThresholdSmoothness = 0.6f;
+	LocalExposureHighlightThresholdStrength = 1.0f;
+	LocalExposureShadowThresholdStrength = 1.0f;
 	LocalExposureMiddleGreyBias = 0.0f;
 
 	// next value might get overwritten by r.DefaultFeature.LensFlare
diff --git a/Engine/Source/Runtime/Engine/Private/SceneView.cpp b/Engine/Source/Runtime/Engine/Private/SceneView.cpp
--- a/Engine/Source/Runtime/Engine/Private/SceneView.cpp
+++ b/Engine/Source/Runtime/Engine/Private/SceneView.cpp
@@ -1661,9 +1661,8 @@ void FSceneView::OverridePostProcessSettings(const FPostProcessSettings& Src, fl
 		LERP_PP(LocalExposureDetailStrength);
 		LERP_PP(LocalExposureBlurredLuminanceBlend);
 		LERP_PP(LocalExposureBlurredLuminanceKernelSizePercent);
-		SET_PP(LocalExposureUseSoftThresholds);
-		LERP_PP(LocalExposureHighlightThresholdSmoothness);
-		LERP_PP(LocalExposureShadowThresholdSmoothness);
+		LERP_PP(LocalExposureHighlightThresholdStrength);
+		LERP_PP(LocalExposureShadowThresholdStrength);
 		LERP_PP(LocalExposureMiddleGreyBias);
 		LERP_PP(LensFlareIntensity);
 		LERP_PP(LensFlareTint);
diff --git a/Engine/Source/Runtime/Renderer/Private/PostProcess/PostProcessLocalExposure.cpp b/Engine/Source/Runtime/Renderer/Private/PostProcess/PostProcessLocalExposure.cpp
--- a/Engine/Source/Runtime/Renderer/Private/PostProcess/PostProcessLocalExposure.cpp
+++ b/Engine/Source/Runtime/Renderer/Private/PostProcess/PostProcessLocalExposure.cpp
@@ -13,12 +13,6 @@
 #include "ShaderCompilerCore.h"
 #include "DataDrivenShaderPlatformInfo.h"
 
-static TAutoConsoleVariable<bool> CVarLocalExposureSoftThresholdSmoothstep(
-	TEXT("r.LocalExposure.SoftThreshold.Smoothstep"),
-	false,
-	TEXT("Whether to use smoothstep when calculating soft thresholds."),
-	ECVF_RenderThreadSafe);
-
 static TAutoConsoleVariable<float> CVarExposureFusionTargetLuminance(
 	TEXT("r.LocalExposure.ExposureFusion.TargetLuminance"),
 	0.5f,
@@ -266,10 +260,8 @@ FLocalExposureParameters GetLocalExposureParameters(const FViewInfo& View, FIntP
 	Parameters.BilateralGridUVScale = LocalExposureBilateralGridUVScale;
 	Parameters.HighlightThreshold = Settings.LocalExposureHighlightThreshold;
 	Parameters.ShadowThreshold = Settings.LocalExposureShadowThreshold;
-	Parameters.bUseSoftThreshold = Settings.LocalExposureUseSoftThresholds;
-	Parameters.bSoftThresholdUseSmoothStep = CVarLocalExposureSoftThresholdSmoothstep.GetValueOnRenderThread();
-	Parameters.HighlightThresholdSmoothness = Settings.LocalExposureHighlightThresholdSmoothness;
-	Parameters.ShadowThresholdSmoothness = Settings.LocalExposureShadowThresholdSmoothness;
+	Parameters.HighlightThresholdStrength = Settings.LocalExposureHighlightThresholdStrength;
+	Parameters.ShadowThresholdStrength = Settings.LocalExposureShadowThresholdStrength;
 	return Parameters;
 }
 
diff --git a/Engine/Source/Runtime/Renderer/Private/PostProcess/PostProcessLocalExposure.h b/Engine/Source/Runtime/Renderer/Private/PostProcess/PostProcessLocalExposure.h
--- a/Engine/Source/Runtime/Renderer/Private/PostProcess/PostProcessLocalExposure.h
+++ b/Engine/Source/Runtime/Renderer/Private/PostProcess/PostProcessLocalExposure.h
@@ -15,10 +15,8 @@ BEGIN_SHADER_PARAMETER_STRUCT(FLocalExposureParameters, )
 	SHADER_PARAMETER(FVector2f, BilateralGridUVScale)
 	SHADER_PARAMETER(float, HighlightThreshold)
 	SHADER_PARAMETER(float, ShadowThreshold)
-	SHADER_PARAMETER(uint32, bUseSoftThreshold)
-	SHADER_PARAMETER(uint32, bSoftThresholdUseSmoothStep)
-	SHADER_PARAMETER(float, HighlightThresholdSmoothness)
-	SHADER_PARAMETER(float, ShadowThresholdSmoothness)
+	SHADER_PARAMETER(float, HighlightThresholdStrength)
+	SHADER_PARAMETER(float, ShadowThresholdStrength)
 END_SHADER_PARAMETER_STRUCT()
 
 FLocalExposureParameters GetLocalExposureParameters(const FViewInfo& View, FIntPoint ViewRectSize, const FEyeAdaptationParameters& EyeAdaptationParameters);
