Fixed CIS - C4061 ('enumeration' is not explicitly handled by a case label) reported when compiling Module.SolarisEditor.cpp (fixed by not includiing PrimitiveComponent.h in AsyncRegisterLevelContext.h, but this could come back if RHI public headers are included again).

diff --git a/Engine/Source/Runtime/Engine/Internal/Streaming/AsyncRegisterLevelContext.h b/Engine/Source/Runtime/Engine/Internal/Streaming/AsyncRegisterLevelContext.h
--- a/Engine/Source/Runtime/Engine/Internal/Streaming/AsyncRegisterLevelContext.h
+++ b/Engine/Source/Runtime/Engine/Internal/Streaming/AsyncRegisterLevelContext.h
@@ -4,11 +4,10 @@
 
 #include "Tasks/Task.h"
 #include "Async/Fundamental/Task.h"
-#include "Components/ActorComponent.h"
-#include "Components/PrimitiveComponent.h"
 
 class ULevel;
 class FSceneInterface;
+class UPrimitiveComponent;
 
 struct FAsyncRegisterLevelContext;
 
@@ -71,7 +70,6 @@ private:
 struct FAsyncRegisterLevelContext
 {
 public:
-	using FPrimitiveBatch = FRegisterComponentContext::FPrimitiveBatch;
 	static TUniquePtr<FAsyncRegisterLevelContext> CreateInstance(ULevel* InLevel);
 	FAsyncRegisterLevelContext(ULevel* InLevel);
 	void AddPrimitive(UPrimitiveComponent* InComponent);
