1-Only start precaching pso after the engine have loaded the default materials.

The underlying issue is that we are trying to precache PSOs when we are loading the default materials, LoadObject can potentially use default materials, and they should not be precache, since UMaterialInterface::InitDefaultMaterials will precache them after the recursive loads, by calling PrecacheDefaultMaterialPSOs. I had the same issue with  EPSOPrecacheMode::PreloadShader, that why IsPSOShaderPreloadingEnabled() is checking if we initialized the default materials before starting to preload shaders. I think the same logic should apply to the PSO caching mode

2-Fix GPuSkin PSO precaching when ManualVF is disabled.

Expose  InitStaticMeshVertexFactoryComponents in StaticMeshResources, and remove the extern declaration.

#tests replayruns (all platforms) + 672b8b72e07ed232abe35e53

diff --git a/Engine/Source/Runtime/Engine/Private/Components/StaticMeshComponent.cpp b/Engine/Source/Runtime/Engine/Private/Components/StaticMeshComponent.cpp
--- a/Engine/Source/Runtime/Engine/Private/Components/StaticMeshComponent.cpp
+++ b/Engine/Source/Runtime/Engine/Private/Components/StaticMeshComponent.cpp
@@ -1654,27 +1654,6 @@ void UStaticMeshComponent::InitResources()
 	}
 }
 
-void InitStaticMeshVertexFactoryComponents(
-	const FStaticMeshVertexBuffers& VertexBuffers,
-	FLocalVertexFactory* VertexFactory,
-	int32 LightMapCoordinateIndex,
-	bool bOverrideColorVertexBuffer,
-	FLocalVertexFactory::FDataType& OutData)
-{
-	VertexBuffers.PositionVertexBuffer.BindPositionVertexBuffer(VertexFactory, OutData);
-	VertexBuffers.StaticMeshVertexBuffer.BindTangentVertexBuffer(VertexFactory, OutData);
-	VertexBuffers.StaticMeshVertexBuffer.BindPackedTexCoordVertexBuffer(VertexFactory, OutData);
-	VertexBuffers.StaticMeshVertexBuffer.BindLightMapVertexBuffer(VertexFactory, OutData, LightMapCoordinateIndex);
-	if (bOverrideColorVertexBuffer)
-	{
-		FColorVertexBuffer::BindDefaultColorVertexBuffer(VertexFactory, OutData, FColorVertexBuffer::NullBindStride::FColorSizeForComponentOverride);
-	}
-	else
-	{
-		VertexBuffers.ColorVertexBuffer.BindColorVertexBuffer(VertexFactory, OutData);
-	}
-}
-
 void UStaticMeshComponent::CollectPSOPrecacheDataImpl(
 	const FVertexFactoryType* VFType, 
 	const FPSOPrecacheParams& BasePrecachePSOParams, 
@@ -1774,7 +1753,7 @@ void UStaticMeshComponent::CollectPSOPrecacheData(const FPSOPrecacheParams& Base
 		int32 LODLightMapCoordinateIndex = LightMapCoordinateIndex < NumTexCoords ? LightMapCoordinateIndex : NumTexCoords - 1;
 		bool bOverrideColorVertexBuffer = LODIndex < LODData.Num() && LODData[LODIndex].OverrideVertexColors != nullptr;
 		FLocalVertexFactory::FDataType Data;
-		InitStaticMeshVertexFactoryComponents(LODRenderData.VertexBuffers, nullptr /*VertexFactory*/, LODLightMapCoordinateIndex, bOverrideColorVertexBuffer, Data);
+		LODRenderData.VertexBuffers.InitCompVF(nullptr /*VertexFactory*/, LODLightMapCoordinateIndex, bOverrideColorVertexBuffer, Data);
 		FLocalVertexFactory::GetVertexElements(GMaxRHIFeatureLevel, EVertexInputStreamType::Default, bSupportsManualVertexFetch, Data, Elements);
 	};
 	
diff --git a/Engine/Source/Runtime/Engine/Private/Materials/Material.cpp b/Engine/Source/Runtime/Engine/Private/Materials/Material.cpp
--- a/Engine/Source/Runtime/Engine/Private/Materials/Material.cpp
+++ b/Engine/Source/Runtime/Engine/Private/Materials/Material.cpp
@@ -2792,7 +2792,7 @@ bool UMaterial::IsCompiling() const
 FGraphEventArray UMaterial::PrecachePSOs(const FPSOPrecacheVertexFactoryDataList& VertexFactoryDataList, const FPSOPrecacheParams& InPreCacheParams, EPSOPrecachePriority Priority, TArray<FMaterialPSOPrecacheRequestID>& OutMaterialPSORequestIDs)
 {
 	FGraphEventArray GraphEvents;
-	if (FApp::CanEverRender() && MaterialResources.Num() > 0 && (PipelineStateCache::IsPSOPrecachingEnabled() || IsPSOShaderPreloadingEnabled()))
+	if (FApp::CanEverRender() && MaterialResources.Num() > 0 && (PipelineStateCache::IsPSOPrecachingEnabled() || IsPSOShaderPreloadingEnabled()) && UMaterialInterface::IsDefaultMaterialInitialized())
 	{
 		EMaterialQualityLevel::Type ActiveQualityLevel = GetCachedScalabilityCVars().MaterialQualityLevel;
 		uint32 FeatureLevelsToCompile = GetFeatureLevelsToCompileForRendering();
diff --git a/Engine/Source/Runtime/Engine/Private/SkeletalRenderGPUSkin.cpp b/Engine/Source/Runtime/Engine/Private/SkeletalRenderGPUSkin.cpp
--- a/Engine/Source/Runtime/Engine/Private/SkeletalRenderGPUSkin.cpp
+++ b/Engine/Source/Runtime/Engine/Private/SkeletalRenderGPUSkin.cpp
@@ -2051,7 +2051,21 @@ void FSkeletalMeshObjectGPUSkin::GetUsedVertexFactoryData(
 
 	if (GPUSkinTechnique != ESkeletalMeshGPUSkinTechnique::Inline || FGPUSkinCache::IsGPUSkinCacheRayTracingSupported())
 	{
-		VertexFactoryDataList.AddUnique(FPSOPrecacheVertexFactoryData(&FGPUSkinPassthroughVertexFactory::StaticType));
+		const FVertexFactoryType* GPUSkinVFType = &FGPUSkinPassthroughVertexFactory::StaticType;
+		bool bSupportsManualVertexFetch = GPUSkinVFType->SupportsManualVertexFetch(GMaxRHIFeatureLevel);
+		if (!bSupportsManualVertexFetch)
+		{
+			FVertexDeclarationElementList VertexElements;
+			bool bOverrideColorVertexBuffer = false;
+			FGPUSkinPassthroughVertexFactory::FDataType Data;
+			LODRenderData.StaticVertexBuffers.InitCompVF(nullptr /*VertexFactory*/, 0, bOverrideColorVertexBuffer, Data);
+			FGPUSkinPassthroughVertexFactory::GetVertexElements(GMaxRHIFeatureLevel, EVertexInputStreamType::Default, bSupportsManualVertexFetch, Data, VertexElements);
+			VertexFactoryDataList.AddUnique(FPSOPrecacheVertexFactoryData(GPUSkinVFType, VertexElements));
+		}
+		else
+		{
+			VertexFactoryDataList.AddUnique(FPSOPrecacheVertexFactoryData(&FGPUSkinPassthroughVertexFactory::StaticType));
+		}
 	}
 
 	if (GPUSkinTechnique != ESkeletalMeshGPUSkinTechnique::MeshDeformer)
diff --git a/Engine/Source/Runtime/Engine/Private/SkinnedAsset.cpp b/Engine/Source/Runtime/Engine/Private/SkinnedAsset.cpp
--- a/Engine/Source/Runtime/Engine/Private/SkinnedAsset.cpp
+++ b/Engine/Source/Runtime/Engine/Private/SkinnedAsset.cpp
@@ -175,9 +175,6 @@ void USkinnedAsset::PostLoad()
 	}
 }
 
-// TODO: move this to a StaticMeshResources.h?
-extern void InitStaticMeshVertexFactoryComponents(const FStaticMeshVertexBuffers& VertexBuffers, FLocalVertexFactory* VertexFactory, int32 LightMapCoordinateIndex, bool bOverrideColorVertexBuffer, FLocalVertexFactory::FDataType& OutData);
-
 void USkinnedAsset::GetVertexFactoryTypesPerMaterialIndex(FPSOPrecacheVertexFactoryDataPerMaterialIndexList& OutList,
 	USkinnedMeshComponent* SkinnedMeshComponent, int32 MinLODIndex, bool bCPUSkin, ERHIFeatureLevel::Type FeatureLevel)
 {
@@ -221,7 +218,7 @@ void USkinnedAsset::GetVertexFactoryTypesPerMaterialIndex(FPSOPrecacheVertexFact
 					FVertexDeclarationElementList VertexElements;
 					bool bOverrideColorVertexBuffer = false;
 					FLocalVertexFactory::FDataType Data;
-					InitStaticMeshVertexFactoryComponents(LODRenderData.StaticVertexBuffers, nullptr /*VertexFactory*/, 0, bOverrideColorVertexBuffer, Data);
+					LODRenderData.StaticVertexBuffers.InitCompVF(nullptr /*VertexFactory*/, 0, bOverrideColorVertexBuffer, Data);
 					FLocalVertexFactory::GetVertexElements(GMaxRHIFeatureLevel, EVertexInputStreamType::Default, bSupportsManualVertexFetch, Data, VertexElements);
 					VFsPerMaterial->VertexFactoryDataList.AddUnique(FPSOPrecacheVertexFactoryData(CPUSkinVFType, VertexElements));
 				}
diff --git a/Engine/Source/Runtime/Engine/Private/StaticMesh.cpp b/Engine/Source/Runtime/Engine/Private/StaticMesh.cpp
--- a/Engine/Source/Runtime/Engine/Private/StaticMesh.cpp
+++ b/Engine/Source/Runtime/Engine/Private/StaticMesh.cpp
@@ -1124,6 +1124,22 @@ static inline void InitOrUpdateResource(FRHICommandListBase& RHICmdList, FRender
 	}
 }
 
+void FStaticMeshVertexBuffers::InitCompVF(FLocalVertexFactory* VertexFactory, int32 LightMapCoordinateIndex, bool bOverrideColorVertexBuffer, FLocalVertexFactory::FDataType& OutData) const
+{
+	PositionVertexBuffer.BindPositionVertexBuffer(VertexFactory, OutData);
+	StaticMeshVertexBuffer.BindTangentVertexBuffer(VertexFactory, OutData);
+	StaticMeshVertexBuffer.BindPackedTexCoordVertexBuffer(VertexFactory, OutData);
+	StaticMeshVertexBuffer.BindLightMapVertexBuffer(VertexFactory, OutData, LightMapCoordinateIndex);
+	if (bOverrideColorVertexBuffer)
+	{
+		FColorVertexBuffer::BindDefaultColorVertexBuffer(VertexFactory, OutData, FColorVertexBuffer::NullBindStride::FColorSizeForComponentOverride);
+	}
+	else
+	{
+		ColorVertexBuffer.BindColorVertexBuffer(VertexFactory, OutData);
+	}
+}
+
 void FStaticMeshVertexBuffers::InitModelBuffers(TArray<FModelVertex>& Vertices)
 {
 	if (Vertices.Num())
diff --git a/Engine/Source/Runtime/Engine/Public/StaticMeshResources.h b/Engine/Source/Runtime/Engine/Public/StaticMeshResources.h
--- a/Engine/Source/Runtime/Engine/Public/StaticMeshResources.h
+++ b/Engine/Source/Runtime/Engine/Public/StaticMeshResources.h
@@ -356,6 +356,7 @@ struct FStaticMeshVertexBuffers
 	}
 
 	void ENGINE_API InitModelBuffers(TArray<FModelVertex>& Vertices);
+	void ENGINE_API InitCompVF(FLocalVertexFactory* VertexFactory, int32 LightMapCoordinateIndex, bool bOverrideColorVertexBuffer, FLocalVertexFactory::FDataType& OutDat) const;
 
 	inline void ENGINE_API InitModelVF(FRHICommandListBase& RHICmdList, FLocalVertexFactory* VertexFactory)
 	{
