Adding a way to do in depth logging based on a filter specified by cvar.
Add sorting of stream cache debugging view.

#rnx

diff --git a/Engine/Source/Runtime/Engine/Private/AudioDecompress.cpp b/Engine/Source/Runtime/Engine/Private/AudioDecompress.cpp
--- a/Engine/Source/Runtime/Engine/Private/AudioDecompress.cpp
+++ b/Engine/Source/Runtime/Engine/Private/AudioDecompress.cpp
@@ -12,6 +12,11 @@
 #include "Stats/StatsTrace.h"
 #include "Sound/StreamedAudioChunkSeekTable.h"
 
+#if !UE_BUILD_SHIPPING
+// We need to hit this from the streaming cache to trace.
+/* extern */ FString ExtraTracingSoundWaveFilterCVar;
+#endif
+
 namespace AudioDecompressPrivate
 {
 #ifndef WITH_AUDIO_DECODER_DIAGNOSTICS
@@ -26,6 +31,14 @@ namespace AudioDecompressPrivate
 		TEXT("Keep a reference to next to play chunk" ),
 		ECVF_Default);
 	
+#if !UE_BUILD_SHIPPING
+	static FAutoConsoleVariableRef CVarExtraTracingSoundWaveFilter(
+		TEXT("au.debug.extra_logging_filter"),
+		ExtraTracingSoundWaveFilterCVar,
+		TEXT("Enable extra logging for any sound wave that has this substring in its name."),
+		ECVF_Default);
+#endif
+
 #if WITH_AUDIO_DECODER_DIAGNOSTICS
 	
 	static FString ForceDecoderErrorOnWaveCVar;
@@ -320,6 +333,9 @@ bool IStreamedCompressedInfo::StreamCompressedData(uint8* Destination, bool bLoo
 	// If we have a pending next chunk from seeking, move to it now.
 	if (StreamSeekBlockIndex != INDEX_NONE)
 	{
+#if !UE_BUILD_SHIPPING
+		UE_CLOG(bWantsExtraLogging, LogAudioStreamCaching, Display, TEXT("[%s.%p] - Seek Requesting Block = %d"), *WriteToString<128>(StreamingSoundWave->GetFName()), this, StreamSeekBlockIndex.load(std::memory_order_relaxed));
+#endif
 		uint32 ChunkSize = 0;
 		const uint8* NewlySeekedChunk = GetLoadedChunk(StreamingSoundWave, StreamSeekBlockIndex, ChunkSize);
 		UE_LOG(LogAudio, Log, TEXT("Seek request for (%s): (%s) Chunk=%d (%s), Offset=%d, OffsetInAudioFrames=%u"),
@@ -377,6 +393,10 @@ bool IStreamedCompressedInfo::StreamCompressedData(uint8* Destination, bool bLoo
 							LastPCMByteSize = 0;
 							LastPCMOffset = 0;
 							bStoringEndOfFile = false;
+
+#if !UE_BUILD_SHIPPING
+							UE_CLOG(bWantsExtraLogging, LogAudioStreamCaching, Display, TEXT("[%s.%p] - Committing Seek Chunk=%d Size=%u"), *WriteToString<128>(StreamingSoundWave->GetFName()), this, CurrentChunkIndex, ChunkSize);
+#endif
 						}
 						else // Seek failed (off the end of the chunk).
 						{
@@ -397,6 +417,9 @@ bool IStreamedCompressedInfo::StreamCompressedData(uint8* Destination, bool bLoo
 	// If next chunk wasn't loaded when last one finished reading, try to get it again now
 	if (SrcBufferData == NULL)
 	{
+#if !UE_BUILD_SHIPPING
+		UE_CLOG(bWantsExtraLogging, LogAudioStreamCaching, Display, TEXT("[%s.%p] - Demand-loading Chunk=%d"), *WriteToString<128>(StreamingSoundWave->GetFName()), this, CurrentChunkIndex);
+#endif
 		uint32 ChunkSize = 0;
 		SrcBufferData = GetLoadedChunk(StreamingSoundWave, CurrentChunkIndex, ChunkSize);
 		if (SrcBufferData)
@@ -490,7 +513,9 @@ bool IStreamedCompressedInfo::StreamCompressedData(uint8* Destination, bool bLoo
 
 		// Decompress the next compression frame of audio (many samples) into the PCM buffer
 		const int32 DecodedFrames = DecompressToPCMBuffer(/*Unused*/ 0);
-
+#if !UE_BUILD_SHIPPING
+		UE_CLOG(bWantsExtraLogging, LogAudioStreamCaching, Display, TEXT("[%s.%p] - Decoded frames = %d, Chunk = %d @ %u"), *WriteToString<128>(StreamingSoundWave->GetFName()), this, DecodedFrames, CurrentChunkIndex, SrcBufferOffset);
+#endif
 		if (DecodedFrames < 0)
 		{
 			FXxHash64 Hash = FXxHash64::HashBuffer(SrcBufferData, SrcBufferDataSize);
@@ -517,9 +542,16 @@ bool IStreamedCompressedInfo::StreamCompressedData(uint8* Destination, bool bLoo
 			// Have we reached the end of buffer?
 			if (SrcBufferOffset >= SrcBufferDataSize - SrcBufferPadding)
 			{
+#if !UE_BUILD_SHIPPING
+				UE_CLOG(bWantsExtraLogging, LogAudioStreamCaching, Display, TEXT("[%s.%p] - End Of Buffer %d = %u/%u (pad = %u)"), *WriteToString<128>(StreamingSoundWave->GetFName()), this, CurrentChunkIndex, SrcBufferOffset, SrcBufferDataSize, SrcBufferPadding);
+#endif
 				// Special case for the last chunk of audio
 				if (CurrentChunkIndex == StreamingSoundWave->GetNumChunks() - 1)
 				{
+#if !UE_BUILD_SHIPPING
+						UE_CLOG(bWantsExtraLogging, LogAudioStreamCaching, Display, TEXT("[%s.%p] - Last Chunk Handling LastPCMByteSize = %u"), *WriteToString<128>(StreamingSoundWave->GetFName()), this, LastPCMByteSize);
+#endif
+
 					// check whether all decoded PCM was written
 					if (LastPCMByteSize == 0)
 					{
@@ -709,6 +741,9 @@ const uint8* IStreamedCompressedInfo::GetLoadedChunk(const FSoundWaveProxyPtr& I
 	}
 	else if (ChunkIndex == 0)
 	{
+#if !UE_BUILD_SHIPPING
+		UE_CLOG(bWantsExtraLogging, LogAudioStreamCaching, Display, TEXT("[%s.%p] - GetLoadedChunk(0)"), *WriteToString<128>(InSoundWave->GetFName()), this);
+#endif
 		TArrayView<const uint8> ZerothChunk = FSoundWaveProxy::GetZerothChunk(InSoundWave, true);
 		OutChunkSize = ZerothChunk.Num();
 		return ZerothChunk.GetData();
@@ -718,6 +753,9 @@ const uint8* IStreamedCompressedInfo::GetLoadedChunk(const FSoundWaveProxyPtr& I
 		CurCompressedChunkHandle = IStreamingManager::Get().GetAudioStreamingManager().GetLoadedChunk(InSoundWave, ChunkIndex, false, true);
 		RetainNextChunkInSequence(InSoundWave,ChunkIndex);
 		OutChunkSize = CurCompressedChunkHandle.Num();
+#if !UE_BUILD_SHIPPING
+		UE_CLOG(bWantsExtraLogging, LogAudioStreamCaching, Display, TEXT("[%s.%p] - GetLoadedChunk(%d, %u)"), *WriteToString<128>(InSoundWave->GetFName()), this, ChunkIndex, OutChunkSize);
+#endif
 		return CurCompressedChunkHandle.GetData();
 	}
 }
@@ -858,6 +896,13 @@ bool ICompressedAudioInfo::StreamCompressedInfo(USoundWave* Wave, FSoundQualityI
 		return false;
 	}
 
+#if !UE_BUILD_SHIPPING
+	if (ExtraTracingSoundWaveFilterCVar.Len() &&
+		Wave->GetName().Contains(ExtraTracingSoundWaveFilterCVar))
+	{
+		bWantsExtraLogging = true;
+	}
+#endif
 	return StreamCompressedInfoInternal(StreamingSoundWave, QualityInfo);
 }
 
@@ -865,6 +910,15 @@ bool ICompressedAudioInfo::StreamCompressedInfo(const FSoundWaveProxyPtr& Wave,
 {
 	// Create our own copy of the proxy object
 	StreamingSoundWave = Wave;
+
+#if !UE_BUILD_SHIPPING
+	if (ExtraTracingSoundWaveFilterCVar.Len() &&
+		Wave->GetFName().ToString().Contains(ExtraTracingSoundWaveFilterCVar))
+	{
+		bWantsExtraLogging = true;
+	}
+#endif
+
 	return StreamCompressedInfoInternal(StreamingSoundWave, QualityInfo);
 }
 
diff --git a/Engine/Source/Runtime/Engine/Private/AudioStreamingCache.cpp b/Engine/Source/Runtime/Engine/Private/AudioStreamingCache.cpp
--- a/Engine/Source/Runtime/Engine/Private/AudioStreamingCache.cpp
+++ b/Engine/Source/Runtime/Engine/Private/AudioStreamingCache.cpp
@@ -18,6 +18,9 @@ AudioStreaming.cpp: Implementation of audio streaming classes.
 
 DEFINE_LOG_CATEGORY(LogAudioStreamCaching);
 
+// Defined in AudioDecompress.cpp off au.debug.extra_logging_filter
+extern FString ExtraTracingSoundWaveFilterCVar;
+
 static int32 DebugMaxElementsDisplayCVar = 128;
 FAutoConsoleVariableRef CVarDebugDisplayCaches(
 	TEXT("au.streamcaching.MaxCachesToDisplay"),
@@ -26,6 +29,27 @@ FAutoConsoleVariableRef CVarDebugDisplayCaches(
 	TEXT("n: Number of elements to display on screen."),
 	ECVF_Default);
 
+enum CacheDisplaySortMethod
+{
+	CDSM_RecentlyUsed,
+	CDSM_LoadTime,
+	CDSM_AccessCount,
+	CDSM_RetainingHandleCount,
+	CDSM_LoadStartTime,
+};
+
+static int32 DebugCacheDisplaySortMethod = 0;
+FAutoConsoleVariableRef CVarDebugCacheDisplaySortMethod(
+	TEXT("au.streamcaching.CacheDisplaySortMethod"),
+	DebugCacheDisplaySortMethod,
+	TEXT("What to sort on when showing the stream cache elements. Requries DEBUG_STREAM_CACHE.\n")
+	TEXT("  0 (def): Recently Used.\n")
+	TEXT("  1      : Load Time.\n")
+	TEXT("  2      : Access Count.\n")
+	TEXT("  3      : Retaining Handle Count.\n")
+	TEXT("  4      : When Loaded.\n"),
+	ECVF_Default);
+
 static int32 KeepCacheMissBufferOnFlushCVar = 1;
 FAutoConsoleVariableRef CVarKeepCacheMissBufferOnFlush(
 	TEXT("au.streamcaching.KeepCacheMissBufferOnFlush"),
@@ -336,6 +360,9 @@ void FAudioChunkCache::FCacheElement::ReleaseRetainedAudioOnSoundWave()
 {
 	if (FSoundWavePtr SoundWaveDataPtr = SoundWaveWeakPtr.Pin())
 	{
+#if !UE_BUILD_SHIPPING
+			UE_CLOG(bWantsExtraLogging, LogAudioStreamCaching, Display, TEXT("[%s] - Release Retained Audio, chunk %d, ObjectKey=%s"), *WriteToString<128>(Key.SoundWaveName), Key.ChunkIndex, *WriteToString<40>(Key.ObjectKey));
+#endif
 		check(SoundWaveDataPtr.IsValid());
 		SoundWaveDataPtr->ReleaseCompressedAudio();
 	}
@@ -734,6 +761,11 @@ uint64 FAudioChunkCache::AddOrTouchChunk(const FChunkKey& InKey, const TSharedPt
 			ExecuteOnLoadCompleteCallback(EAudioChunkLoadResult::AlreadyLoaded, OnLoadCompleted, CallbackThread);
 		}
 
+#if !UE_BUILD_SHIPPING
+		UE_CLOG(FoundElement->bWantsExtraLogging, LogAudioStreamCaching, Display, TEXT("[%s] - AddOrTouchChunk(found), chunk %d, ObjectKey=%s"), *WriteToString<128>(FoundElement->Key.SoundWaveName), FoundElement->Key.ChunkIndex, *WriteToString<40>(FoundElement->Key.ObjectKey));
+#endif
+
+
 #if DEBUG_STREAM_CACHE
 		FoundElement->DebugInfo.NumTimesTouched++;
 
@@ -798,6 +830,9 @@ TArrayView<uint8> FAudioChunkCache::GetChunk(const FChunkKey& InKey, const TShar
 	FCacheElement* FoundElement = FindElementForKey(InKey);
 	if (FoundElement)
 	{
+#if !UE_BUILD_SHIPPING
+		UE_CLOG(FoundElement->bWantsExtraLogging, LogAudioStreamCaching, Display, TEXT("[%s] - GetChunk(found), chunk %d, ObjectKey=%s"), *WriteToString<128>(FoundElement->Key.SoundWaveName), FoundElement->Key.ChunkIndex, *WriteToString<40>(FoundElement->Key.ObjectKey));
+#endif
 		OutCacheOffset = FoundElement->CacheLookupID;
 		TouchElement(FoundElement);
 		if (FoundElement->IsLoadInProgress())
@@ -843,6 +878,10 @@ TArrayView<uint8> FAudioChunkCache::GetChunk(const FChunkKey& InKey, const TShar
 		{
 			if (FStreamedAudioChunk* Chunk = FoundElement->GetChunk(InKey.ChunkIndex))
 			{
+#if !UE_BUILD_SHIPPING
+				UE_CLOG(FoundElement->bWantsExtraLogging, LogAudioStreamCaching, Display, TEXT("[%s] - GetChunk(blocking, hit), chunk %d, ObjectKey=%s"), *WriteToString<128>(FoundElement->Key.SoundWaveName), FoundElement->Key.ChunkIndex, *WriteToString<40>(FoundElement->Key.ObjectKey));
+#endif
+
 				int32 ChunkAudioDataSize = Chunk->AudioDataSize;
 #if DEBUG_STREAM_CACHE
 				const float SampleRate = InSoundWavePtr->GetSampleRate();
@@ -851,6 +890,8 @@ TArrayView<uint8> FAudioChunkCache::GetChunk(const FChunkKey& InKey, const TShar
 				FoundElement->DebugInfo.OffsetInSeconds = OffsetInFrames != INDEX_NONE ? OffsetInFrames / SampleRate : -1.f;
 				FoundElement->DebugInfo.NumTotalChunks = NumChunks - 1; //(exclude chunk 0)
 				FoundElement->DebugInfo.TimeLoadStartedCycles = FPlatformTime::Cycles64();
+				FoundElement->DebugInfo.TimeLoadEndedCycles = FoundElement->DebugInfo.TimeLoadStartedCycles;
+				FoundElement->DebugInfo.TimeToLoadMs = 0;
 #endif
 				MemoryCounterBytes -= FoundElement->ChunkDataSize;
 
@@ -875,9 +916,7 @@ TArrayView<uint8> FAudioChunkCache::GetChunk(const FChunkKey& InKey, const TShar
 				FoundElement->Key = InKey;
 				FoundElement->ChunkDataSize = ChunkAudioDataSize;
 				FoundElement->bIsLoaded = true;
-#if DEBUG_STREAM_CACHE
-				FoundElement->DebugInfo.TimeToLoadMs = FPlatformTime::ToMilliseconds64(FPlatformTime::Cycles64() - FoundElement->DebugInfo.TimeLoadStartedCycles);
-#endif
+
 				// If this value is ever negative, it means that we're decrementing more than we're incrementing:
 				if (ensureMsgf(FoundElement->NumConsumers.GetValue() >= 0, TEXT("NumConsumers was negative for FoundElement. Reseting to 1")))
 				{
@@ -1097,6 +1136,10 @@ uint64 FAudioChunkCache::TrimMemory(uint64 BytesToFree, bool bInAllowRetainedChu
 	{
 		if (CurrentElement->CanEvictChunk() && CurrentElement->ChunkDataSize != 0)
 		{
+#if !UE_BUILD_SHIPPING
+			UE_CLOG(CurrentElement->bWantsExtraLogging, LogAudioStreamCaching, Display, TEXT("[%s] - Evicting for trim, chunk %d, ObjectKey=%s"), *WriteToString<128>(CurrentElement->Key.SoundWaveName), CurrentElement->Key.ChunkIndex, *WriteToString<40>(CurrentElement->Key.ObjectKey));
+#endif
+
 			uint32 ChunkSize = CurrentElement->ChunkDataSize;
 			BytesFreed += ChunkSize;
 			MemoryCounterBytes -= ChunkSize;
@@ -1140,6 +1183,9 @@ uint64 FAudioChunkCache::TrimMemory(uint64 BytesToFree, bool bInAllowRetainedChu
 				CurrentElement->ReleaseRetainedAudioOnSoundWave();
 				if (CurrentElement->CanEvictChunk())
 				{
+#if !UE_BUILD_SHIPPING
+					UE_CLOG(CurrentElement->bWantsExtraLogging, LogAudioStreamCaching, Display, TEXT("[%s] - Evicting for trim, chunk %d, ObjectKey=%s"), *WriteToString<128>(CurrentElement->Key.SoundWaveName), CurrentElement->Key.ChunkIndex, *WriteToString<40>(CurrentElement->Key.ObjectKey));
+#endif
 					uint32 ChunkSize = CurrentElement->ChunkDataSize;
 					BytesFreed += ChunkSize;
 					RetainedBytesFreed += ChunkSize;
@@ -1508,7 +1554,19 @@ FAudioChunkCache::FCacheElement* FAudioChunkCache::InsertChunk(const FChunkKey&
 			}
 		}
 
-		check(CacheElement);
+#if !UE_BUILD_SHIPPING
+		CacheElement->bWantsExtraLogging = false;
+		if (ExtraTracingSoundWaveFilterCVar.Len())
+		{
+			TStringBuilder<128> SoundName;
+			SoundName << InSoundWavePtr->GetFName();
+			if (FCString::Stristr(*SoundName, *ExtraTracingSoundWaveFilterCVar))
+			{
+				CacheElement->bWantsExtraLogging = true;
+				UE_LOG(LogAudioStreamCaching, Display, TEXT("[%s] - Inserting chunk index %d, ObjectKey=%s"), *SoundName, InKey.ChunkIndex, *WriteToString<40>(InKey.ObjectKey));
+			}
+		}
+#endif
 		CacheElement->bIsLoaded = false;
 		CacheElement->Key = InKey;
 		CacheElement->SoundWaveWeakPtr = InSoundWavePtr;
@@ -1567,6 +1625,10 @@ FAudioChunkCache::FCacheElement* FAudioChunkCache::EvictLeastRecentChunk(bool bB
 
 	if (bIsChunkEvictable)
 	{
+#if !UE_BUILD_SHIPPING
+		UE_CLOG(CacheElement->bWantsExtraLogging, LogAudioStreamCaching, Display, TEXT("[%s] - Evicting ununsed LRU chunk index %d, ObjectKey=%s"), *WriteToString<128>(CacheElement->Key.SoundWaveName), CacheElement->Key.ChunkIndex, *WriteToString<40>(CacheElement->Key.ObjectKey));
+#endif
+
 		check(CacheElement->MoreRecentElement != nullptr);
 		check(CacheElement->LessRecentElement == nullptr);
 
@@ -1577,6 +1639,9 @@ FAudioChunkCache::FCacheElement* FAudioChunkCache::EvictLeastRecentChunk(bool bB
 	}
 	else if (bBlockForPendingLoads && bIsChunkLoadingButUnreferenced)
 	{
+#if !UE_BUILD_SHIPPING
+		UE_CLOG(CacheElement->bWantsExtraLogging, LogAudioStreamCaching, Display, TEXT("[%s] - Evicting ununsed loading chunk index %d, ObjectKey=%s"), *WriteToString<128>(CacheElement->Key.SoundWaveName), CacheElement->Key.ChunkIndex, *WriteToString<40>(CacheElement->Key.ObjectKey));
+#endif
 		CacheElement->WaitForAsyncLoadCompletion(true);
 
 		FCacheElement* NewLeastRecentElement = CacheElement->MoreRecentElement;
@@ -1639,6 +1704,10 @@ FAudioChunkCache::FCacheElement* FAudioChunkCache::EvictLeastRecentChunk(bool bB
 			UE_LOG(LogAudioStreamCaching, Warning, TEXT("Cache blown! Please increase the cache size (currently %lu bytes) or load less audio."), ReportCacheSize());
 			return nullptr;
 		}
+
+#if !UE_BUILD_SHIPPING
+		UE_CLOG(CacheElement->bWantsExtraLogging, LogAudioStreamCaching, Display, TEXT("[%s] - Evicting unused chunk index %d, ObjectKey=%s"), *WriteToString<128>(CacheElement->Key.SoundWaveName), CacheElement->Key.ChunkIndex, *WriteToString<40>(CacheElement->Key.ObjectKey));
+#endif
 	}
 
 #if DEBUG_STREAM_CACHE
@@ -1693,8 +1762,6 @@ public:
 
 void FAudioChunkCache::KickOffAsyncLoad(FCacheElement* CacheElement, const FChunkKey& InKey, TFunction<void(EAudioChunkLoadResult)> OnLoadCompleted, ENamedThreads::Type CallbackThread, bool bNeededForPlayback)
 {
-	check(CacheElement);
-
 	FStreamedAudioChunk* Chunk = CacheElement->GetChunk(InKey.ChunkIndex);
 
 	if (nullptr == Chunk)
@@ -1706,6 +1773,9 @@ void FAudioChunkCache::KickOffAsyncLoad(FCacheElement* CacheElement, const FChun
 	int32 ChunkDataSize = Chunk->AudioDataSize;
 
 	EAsyncIOPriorityAndFlags AsyncIOPriority = GetAsyncPriorityForChunk(InKey, bNeededForPlayback);
+#if !UE_BUILD_SHIPPING
+	UE_CLOG(CacheElement->bWantsExtraLogging, LogAudioStreamCaching, Display, TEXT("[%s] - KickOffAsyncLoad pri %d outstanding %d chunk index %d, ObjectKey=%s"), *WriteToString<128>(CacheElement->Key.SoundWaveName), (uint32)AsyncIOPriority, NumberOfLoadsInFlight.GetValue(), CacheElement->Key.ChunkIndex, *WriteToString<40>(CacheElement->Key.ObjectKey));
+#endif
 
 	MemoryCounterBytes -= CacheElement->ChunkDataSize;
 
@@ -1768,8 +1838,8 @@ void FAudioChunkCache::KickOffAsyncLoad(FCacheElement* CacheElement, const FChun
 			}
 			
 #if DEBUG_STREAM_CACHE
-			const float TimeToLoad = FPlatformTime::ToMilliseconds64(FPlatformTime::Cycles64() - CacheElement->DebugInfo.TimeLoadStartedCycles);		
-			CacheElement->DebugInfo.TimeToLoadMs = TimeToLoad;
+			CacheElement->DebugInfo.TimeLoadEndedCycles = FPlatformTime::Cycles64();
+			CacheElement->DebugInfo.TimeToLoadMs = FPlatformTime::ToMilliseconds64(CacheElement->DebugInfo.TimeLoadEndedCycles - CacheElement->DebugInfo.TimeLoadStartedCycles);
 #endif
 			EAudioChunkLoadResult ChunkLoadResult = bRequestFailed ? EAudioChunkLoadResult::Interrupted : EAudioChunkLoadResult::Completed;
 			ExecuteOnLoadCompleteCallback(ChunkLoadResult, OnLoadCompleted, CallbackThread);
@@ -1793,6 +1863,9 @@ void FAudioChunkCache::KickOffAsyncLoad(FCacheElement* CacheElement, const FChun
 	{
 		if (CacheElement->IsLoadInProgress())
 		{
+#if !UE_BUILD_SHIPPING
+			UE_CLOG(CacheElement->bWantsExtraLogging, LogAudioStreamCaching, Display, TEXT("[%s] - Waiting for async load while kicking async load, chunk %d, ObjectKey=%s"), *WriteToString<128>(CacheElement->Key.SoundWaveName), CacheElement->Key.ChunkIndex, *WriteToString<40>(CacheElement->Key.ObjectKey));
+#endif
 			CacheElement->WaitForAsyncLoadCompletion(true);
 		}
 
@@ -1817,6 +1890,18 @@ void FAudioChunkCache::KickOffAsyncLoad(FCacheElement* CacheElement, const FChun
 				TGraphTask<FClearAudioChunkCacheReadRequestTask>::CreateTask().ConstructAndDispatchWhenReady(LocalReadRequest);
 			}
 
+#if !UE_BUILD_SHIPPING
+			if (CacheElement->bWantsExtraLogging)
+			{
+				bool bMatched = CacheElement->Key == InKey;
+				UE_LOG(LogAudioStreamCaching, Display, TEXT("[%s] - AsyncReadComplete index %d, ObjectKey=%s"), *WriteToString<128>(CacheElement->Key.SoundWaveName), CacheElement->Key.ChunkIndex, *WriteToString<40>(CacheElement->Key.ObjectKey));
+				UE_LOG(LogAudioStreamCaching, Display, TEXT("[%s] - Cancelled=%d, Size=%u, InKey %s"), *WriteToString<128>(CacheElement->Key.SoundWaveName), bWasCancelled, ChunkDataSize, bMatched ? TEXT("matched") : TEXT("DID NOT MATCH!"));
+				if (!bMatched)
+				{
+					UE_LOG(LogAudioStreamCaching, Display, TEXT("[%s] - InKey %s / %d / %s"), *WriteToString<128>(CacheElement->Key.SoundWaveName), *WriteToString<128>(InKey.SoundWaveName), InKey.ChunkIndex, *WriteToString<40>(InKey.ObjectKey));
+				}
+			}
+#endif
 			// Populate key and DataSize. The async read request was set up to write directly into CacheElement->ChunkData.
 			CacheElement->Key = InKey;
 			CacheElement->ChunkDataSize = ChunkDataSize;
@@ -1828,8 +1913,8 @@ void FAudioChunkCache::KickOffAsyncLoad(FCacheElement* CacheElement, const FChun
 			}
 			
 #if DEBUG_STREAM_CACHE
-			const float TimeToLoad = FPlatformTime::ToMilliseconds64(FPlatformTime::Cycles64() - CacheElement->DebugInfo.TimeLoadStartedCycles);		
-			CacheElement->DebugInfo.TimeToLoadMs = TimeToLoad;			
+			CacheElement->DebugInfo.TimeLoadEndedCycles = FPlatformTime::Cycles64();
+			CacheElement->DebugInfo.TimeToLoadMs = FPlatformTime::ToMilliseconds64(CacheElement->DebugInfo.TimeLoadEndedCycles - CacheElement->DebugInfo.TimeLoadStartedCycles);
 #endif
 			const EAudioChunkLoadResult LoadResult = bWasCancelled ? EAudioChunkLoadResult::Interrupted : EAudioChunkLoadResult::Completed;
 			ExecuteOnLoadCompleteCallback(LoadResult, OnLoadCompleted, CallbackThread);
@@ -1844,6 +1929,9 @@ void FAudioChunkCache::KickOffAsyncLoad(FCacheElement* CacheElement, const FChun
 		CacheElement->ReadRequest = nullptr;
 		if (Chunk->BulkData.IsBulkDataLoaded())
 		{
+#if !UE_BUILD_SHIPPING
+			UE_CLOG(CacheElement->bWantsExtraLogging, LogAudioStreamCaching, Display, TEXT("[%s] - Chunk %d has bulk data loaded, , ObjectKey=%s"), *WriteToString<128>(CacheElement->Key.SoundWaveName), CacheElement->Key.ChunkIndex, *WriteToString<40>(CacheElement->Key.ObjectKey));
+#endif
 			// If this chunk has been inlined and loaded, move out the data into our newly allocated block.
 			const FBulkDataBuffer<uint8> ChunkMemory = Chunk->MoveOutAsBuffer();
 			
@@ -1873,6 +1961,9 @@ void FAudioChunkCache::KickOffAsyncLoad(FCacheElement* CacheElement, const FChun
 			
 			UE_LOG(LogAudioStreamCaching, Verbose, TEXT("Loading Chunk: %s, %d"), *InKey.SoundWaveName.ToString(), InKey.ChunkIndex);
 
+#if !UE_BUILD_SHIPPING
+			UE_CLOG(CacheElement->bWantsExtraLogging, LogAudioStreamCaching, Display, TEXT("[%s] - CreateStreamingRequest: chunk %d, ObjectKey=%s"), *WriteToString<128>(CacheElement->Key.SoundWaveName), CacheElement->Key.ChunkIndex, *WriteToString<40>(CacheElement->Key.ObjectKey));
+#endif
 			TRACE_IOSTORE_METADATA_SCOPE_TAG(InKey.SoundWaveName);
 			IBulkDataIORequest* LocalReadRequest = Chunk->BulkData.CreateStreamingRequest(0, ChunkDataSize, AsyncIOPriority | AIOP_FLAG_DONTCACHE, &AsyncFileCallBack, CacheElement->ChunkData);
 			if (!LocalReadRequest)
@@ -2105,7 +2196,63 @@ TPair<int, int> FAudioChunkCache::DebugDisplayLegacy(UWorld* World, FViewport* V
 	// A chunk that currently has an async load in flight.
 	const FLinearColor CurrentlyLoadingChunkColor = FLinearColor::Yellow;
 
+	TArray<TPair<FCacheElement*, int32>> SortedCacheElements;
+
+	int32 SortMethod = CDSM_RecentlyUsed;
+#if DEBUG_STREAM_CACHE
+	SortMethod = DebugCacheDisplaySortMethod;
+#endif
+
+	{
+		FCacheElement* CurrentElement = MostRecentElement;
+		for (int32 Index = 0; CurrentElement; Index++, CurrentElement = CurrentElement->LessRecentElement)
+		{
+#if DEBUG_STREAM_CACHE
+			if (SortMethod == CDSM_LoadTime ||
+				SortMethod == CDSM_LoadStartTime)
+			{
+				// If we care about load time, don't display stuff that didn't have a load time.
+				if (CurrentElement->DebugInfo.bWasLoadedFromInlineChunk)
+				{
+					continue;
+				}
+			}
+#endif
+			SortedCacheElements.Add({CurrentElement, Index});
+		}
+	}
+
+
+	switch (SortMethod)
+	{
+	case CDSM_RecentlyUsed: break; // already sorted this way.
+#if DEBUG_STREAM_CACHE
+	case CDSM_LoadTime:
+		{
+			Algo::SortBy(SortedCacheElements, [](TPair<FCacheElement*, int32> ElementPair) { return ElementPair.Key->DebugInfo.TimeToLoadMs; }, TGreater<float>() );
+			break;
+		}
+	case CDSM_AccessCount:
+		{
+			Algo::SortBy(SortedCacheElements, [](TPair<FCacheElement*, int32> ElementPair) { return ElementPair.Key->DebugInfo.NumTimesTouched; }, TGreater<int32>() );
+			break;
+		}
+	case CDSM_RetainingHandleCount:
+		{
+			Algo::SortBy(SortedCacheElements, [](TPair<FCacheElement*, int32> ElementPair) { return ElementPair.Key->NumConsumers.GetValue(); }, TGreater<int32>() );
+			break;
+		}
+	case CDSM_LoadStartTime:
+		{
+			Algo::SortBy(SortedCacheElements, [](TPair<FCacheElement*, int32> ElementPair) { return ElementPair.Key->DebugInfo.TimeLoadStartedCycles; }, TGreater<uint64>() );
+			break;
+		}
+#endif
+	};
 
+	// Trim to our display count
+	SortedCacheElements.SetNum(DebugMaxElementsDisplayCVar);
+	
 	const int32 InitialX = X;
 	const int32 InitialY = Y;
 
@@ -2121,18 +2268,16 @@ TPair<int, int> FAudioChunkCache::DebugDisplayLegacy(UWorld* World, FViewport* V
 
 	// First pass: We run through and get a snap shot of the amount of memory currently in use.
 
-	// Second Pass: We're going to list the actual chunks in the cache.
-	FCacheElement* CurrentElement = MostRecentElement;
-	int32 Index = 0;
-
 	float ColorLerpAmount = 0.0f;
 	const float ColorLerpStep = 0.04f;
-	
+	uint64 CurrentTime = FPlatformTime::Cycles64();
+
 	Audio::FTabularRenderHelper Helper;
 	Helper.AddCol(TEXT("Name"), 64);
 	Helper.AddCol(TEXT("Chunk"), 10);
 	Helper.AddCol(TEXT("Size(kb)"), 10);
 	Helper.AddCol(TEXT("Loadtime(ms)"),20);
+	Helper.AddCol(TEXT("Loaded at(sec)"),20);
 	Helper.AddCol(TEXT("Load-Behavior"), 20);
 	Helper.AddCol(TEXT("Format"), 8);
 	Helper.AddCol(TEXT("Refs"), 6);
@@ -2145,13 +2290,10 @@ TPair<int, int> FAudioChunkCache::DebugDisplayLegacy(UWorld* World, FViewport* V
 	Helper.AddCol(TEXT("Destroyed"),10);
 	
 	// More detailed info about individual chunks here:
-	while (CurrentElement != nullptr)
+	for (TPair<FCacheElement*, int32>& ElementPair : SortedCacheElements)
 	{
-		// We use a CVar to clamp the max amount of chunks we display.
-		if (Index > DebugMaxElementsDisplayCVar)
-		{
-			break;
-		}
+		FCacheElement* CurrentElement = ElementPair.Key;
+		int32 IndexInLRUList = ElementPair.Value;
 
 		int32 NumTotalChunks = -1;
 		int32 NumTimesTouched = -1;
@@ -2165,6 +2307,7 @@ TPair<int, int> FAudioChunkCache::DebugDisplayLegacy(UWorld* World, FViewport* V
 		bool bWaveDestroyed = false;
 		float OffsetInSecs = 0;
 		FName Format;
+		double LoadedAt = 0;
 
 #if DEBUG_STREAM_CACHE
 		NumTotalChunks = CurrentElement->DebugInfo.NumTotalChunks;
@@ -2179,6 +2322,14 @@ TPair<int, int> FAudioChunkCache::DebugDisplayLegacy(UWorld* World, FViewport* V
 		OffsetInSecs = CurrentElement->DebugInfo.OffsetInSeconds;
 		Format = CurrentElement->DebugInfo.Format;
 		TimeToLoad = CurrentElement->DebugInfo.TimeToLoadMs;
+		LoadedAt = -FPlatformTime::ToSeconds64(CurrentTime - CurrentElement->DebugInfo.TimeLoadStartedCycles);
+
+		// we want the next entry to be relative to the first, so only the first entry goes backward in time and
+		// it's easier to spot big load jumps
+		if (SortMethod == CDSM_LoadStartTime)
+		{
+			CurrentTime = CurrentElement->DebugInfo.TimeLoadStartedCycles;
+		}
 #endif
 
 #if WITH_EDITOR
@@ -2219,7 +2370,12 @@ TPair<int, int> FAudioChunkCache::DebugDisplayLegacy(UWorld* World, FViewport* V
 		{
 			TextColor = TrimmedChunkColor;
 		}
-
+#if !UE_BUILD_SHIPPING
+		if (CurrentElement->bWantsExtraLogging)
+		{
+			TextColor = FLinearColor(1.0f, 0, 1.0f, 1.0f);
+		}
+#endif
 		// Scale by load time.
 		const float T = FMath::GetMappedRangeValueClamped( FVector2D(0.f, 100.f), {0.f,1.f}, TimeToLoad);
 		FLinearColor LoadColor = FLinearColor::LerpUsingHSV(FColor::Green, FColor::Red, T);
@@ -2230,6 +2386,7 @@ TPair<int, int> FAudioChunkCache::DebugDisplayLegacy(UWorld* World, FViewport* V
 			{FString::Printf(TEXT("%d/%d"), CurrentElement->Key.ChunkIndex, NumTotalChunks), TextColor}, // chunk/count
 			{FString::Printf(TEXT("%u"),  CurrentElement->ChunkDataSize >> 10),TextColor},			    // size (kb)
 			{FString::Printf(TEXT("%.2f"), TimeToLoad), LoadColor},								        // load time (ms)
+			{FString::Printf(TEXT("%.2f"), LoadedAt), LoadColor},
 			{ShortEnumString(EnumToString(LoadingBehavior)), TextColor },									    // load behavior
 			{Format.ToString(), TextColor },																    // format
 			{FString::Printf(TEXT("%d"), CurrentElement->NumConsumers.GetValue()), TextColor},		    // hard refs
@@ -2241,9 +2398,6 @@ TPair<int, int> FAudioChunkCache::DebugDisplayLegacy(UWorld* World, FViewport* V
 			{LexToString(bIsStaleChunk), TextColor},														    // stale?
 			{LexToString(bWaveDestroyed), TextColor},														    // destroyed?
 		});
-
-		CurrentElement = CurrentElement->LessRecentElement;
-		Index++;
 	}
 
 	Helper.Draw(Canvas,X,Y);
diff --git a/Engine/Source/Runtime/Engine/Public/AudioDecompress.h b/Engine/Source/Runtime/Engine/Public/AudioDecompress.h
--- a/Engine/Source/Runtime/Engine/Public/AudioDecompress.h
+++ b/Engine/Source/Runtime/Engine/Public/AudioDecompress.h
@@ -169,6 +169,11 @@ public:
 
 protected:
 	mutable bool bHasError = false;
+
+#if !UE_BUILD_SHIPPING
+	// If this sound matches a cvar filter for extra logging.
+	bool bWantsExtraLogging = false;
+#endif
 	FSoundWaveProxyPtr StreamingSoundWave;
 };
 
diff --git a/Engine/Source/Runtime/Engine/Public/AudioStreamingCache.h b/Engine/Source/Runtime/Engine/Public/AudioStreamingCache.h
--- a/Engine/Source/Runtime/Engine/Public/AudioStreamingCache.h
+++ b/Engine/Source/Runtime/Engine/Public/AudioStreamingCache.h
@@ -229,6 +229,7 @@ private:
 		int32 NumTimesTouched;
 
 		uint64 TimeLoadStartedCycles = 0;
+		uint64 TimeLoadEndedCycles = 0;
 		// Amount of time spent loading the audio file.
 		float TimeToLoadMs;
 
@@ -262,6 +263,7 @@ private:
 
 		void Reset()
 		{
+			TimeLoadEndedCycles = 0;
 			NumTotalChunks = 0;
 			NumTimesTouched = 0;
 			TimeLoadStartedCycles = 0;
@@ -296,7 +298,13 @@ private:
 		TWeakPtr<FSoundWaveData, ESPMode::ThreadSafe> SoundWaveWeakPtr;
 
 		FThreadSafeBool bIsLoaded;
-		
+
+#if !UE_BUILD_SHIPPING
+		// if the chunk belongs to a sound that wants more logging. This should be under LogAudio so its'
+		// with the other extra logging.
+		bool bWantsExtraLogging = false;
+#endif
+
 		// How many disparate consumers have called GetLoadedChunk.
 		FThreadSafeCounter NumConsumers;
 
