Renamed recently added CVars s.AllowIncrementalPreRegisterComponents and s.AllowIncrementalPreUnregisterComponents to LevelStreaming.AllowIncrementalPreRegisterComponents/LevelStreaming.AllowIncrementalPreUnregisterComponents.

diff --git a/Engine/Source/Runtime/Engine/Private/Level.cpp b/Engine/Source/Runtime/Engine/Private/Level.cpp
--- a/Engine/Source/Runtime/Engine/Private/Level.cpp
+++ b/Engine/Source/Runtime/Engine/Private/Level.cpp
@@ -128,13 +128,13 @@ static FAutoConsoleVariableRef CVarRouteActorInitializationWorkUnitWeighting(
 
 bool GAllowIncrementalPreRegisterComponents = false;
 static FAutoConsoleVariableRef CVarAllowIncrementalPreRegisterComponents(
-	TEXT("s.AllowIncrementalPreRegisterComponents"),
+	TEXT("LevelStreaming.AllowIncrementalPreRegisterComponents"),
 	GAllowIncrementalPreRegisterComponents,
 	TEXT("Whether to allow level streaming incremental pre-registration of actor components."));
 
 bool GAllowIncrementalPreUnregisterComponents = false;
 static FAutoConsoleVariableRef CVarAllowIncrementalPreUnregisterComponents(
-	TEXT("s.AllowIncrementalPreUnregisterComponents"),
+	TEXT("LevelStreaming.AllowIncrementalPreUnregisterComponents"),
 	GAllowIncrementalPreUnregisterComponents,
 	TEXT("Whether to allow level streaming incremental pre-unregistration of actor components."));
 
