Ray Tracing: Fix TextRenderComponent not obeying CastShadow flag

The mesh batch filled in by GetDynamicRayTracingInstance did not initalize this flag. Also renamed Mesh to MeshBatch to match code conventions in other components.

diff --git a/Engine/Source/Runtime/Engine/Private/Components/TextRenderComponent.cpp b/Engine/Source/Runtime/Engine/Private/Components/TextRenderComponent.cpp
--- a/Engine/Source/Runtime/Engine/Private/Components/TextRenderComponent.cpp
+++ b/Engine/Source/Runtime/Engine/Private/Components/TextRenderComponent.cpp
@@ -910,23 +910,24 @@ void FTextRenderSceneProxy::GetDynamicRayTracingInstances(FRayTracingInstanceCol
 		for (int32 BatchIndex = 0; BatchIndex < NumBatches; BatchIndex++)
 		{
 			const FTextBatch& TextBatch = TextBatches[BatchIndex];
-			FMeshBatch& Mesh = CachedRayTracingMaterials.AddDefaulted_GetRef();
-			FMeshBatchElement& BatchElement = Mesh.Elements[0];
+			FMeshBatch& MeshBatch = CachedRayTracingMaterials.AddDefaulted_GetRef();
+			FMeshBatchElement& BatchElement = MeshBatch.Elements[0];
 			BatchElement.IndexBuffer = &IndexBuffer;
-			Mesh.VertexFactory = &VertexFactory;
+			MeshBatch.VertexFactory = &VertexFactory;
 			BatchElement.FirstIndex = TextBatch.IndexBufferOffset;
 			BatchElement.NumPrimitives = TextBatch.IndexBufferCount / 3;
 			BatchElement.MinVertexIndex = TextBatch.VertexBufferOffset;
 			BatchElement.MaxVertexIndex = TextBatch.VertexBufferOffset + TextBatch.VertexBufferCount - 1;
-			Mesh.bDisableBackfaceCulling = false;
-			Mesh.Type = PT_TriangleList;
-			Mesh.DepthPriorityGroup = SDPG_World;
-			Mesh.MaterialRenderProxy = TextBatch.Material->GetRenderProxy();
-			Mesh.bCanApplyViewModeOverrides = !bAlwaysRenderAsText;
-			Mesh.LODIndex = 0;
-			Mesh.bUseWireframeSelectionColoring = IsSelected() ? 1 : 0;
-			Mesh.SegmentIndex = BatchIndex;
-			Mesh.MeshIdInPrimitive = 0;
+			MeshBatch.bDisableBackfaceCulling = false;
+			MeshBatch.Type = PT_TriangleList;
+			MeshBatch.DepthPriorityGroup = SDPG_World;
+			MeshBatch.MaterialRenderProxy = TextBatch.Material->GetRenderProxy();
+			MeshBatch.bCanApplyViewModeOverrides = !bAlwaysRenderAsText;
+			MeshBatch.CastRayTracedShadow = IsShadowCast(Collector.GetReferenceView());
+			MeshBatch.LODIndex = 0;
+			MeshBatch.bUseWireframeSelectionColoring = IsSelected() ? 1 : 0;
+			MeshBatch.SegmentIndex = BatchIndex;
+			MeshBatch.MeshIdInPrimitive = 0;
 		}
 		bNeedsToUpdateRayTracingCache = false;
 	}
