Disables Audio SetOwner notifications to avoid ReplayRun exit crash

diff --git a/Engine/Source/Runtime/Engine/Private/AudioDevice.cpp b/Engine/Source/Runtime/Engine/Private/AudioDevice.cpp
--- a/Engine/Source/Runtime/Engine/Private/AudioDevice.cpp
+++ b/Engine/Source/Runtime/Engine/Private/AudioDevice.cpp
@@ -133,7 +133,7 @@ FAutoConsoleVariableRef CVarDisableLegacyReverb(
 	TEXT("0: Enabled, 1: Disabled"),
 	ECVF_Default);
 
-static int32 DisableOwnerNotify = 0;
+static int32 DisableOwnerNotify = 1;
 FAutoConsoleVariableRef CVarDisableOwnerNotify(
 	TEXT("au.DisableOwnerNotify"),
 	DisableOwnerNotify,
