Ensure that bRunOnAnyThread is correctly set for skeletal mesh component tick functions

Moved from PostInitProperties to OnRegister() to catch all the times it can be modified (construction script etc.)

diff --git a/Engine/Source/Runtime/Engine/Private/Components/SkeletalMeshComponent.cpp b/Engine/Source/Runtime/Engine/Private/Components/SkeletalMeshComponent.cpp
--- a/Engine/Source/Runtime/Engine/Private/Components/SkeletalMeshComponent.cpp
+++ b/Engine/Source/Runtime/Engine/Private/Components/SkeletalMeshComponent.cpp
@@ -505,8 +505,6 @@ void USkeletalMeshComponent::PostInitProperties()
 			AnimationMode = EAnimationMode::AnimationSingleNode;
 		}
 #endif
-
-		PrimaryComponentTick.bRunOnAnyThread = !bEnableAnimation;
 	}
 }
 
@@ -749,6 +747,9 @@ void USkeletalMeshComponent::OnRegister()
 	// to correctly populate that list.
 	ResetLinkedAnimInstances();
 
+	// Ensure the primary tick function is set to tick off the game thread if we are not using internal animation
+	PrimaryComponentTick.bRunOnAnyThread = !bEnableAnimation;
+
 	if(bEnableAnimation)
 	{
 		// We force an initialization here because we're in one of two cases.
