Update FISMComponentDescriptorBase to handle baking an ISM from an actor loaded in the editor.

The methods that copy Instance properties are different if the BodySetup is in a different state.  This detects the state difference and calls the correct version of the method, CopyBodyInstancePropertiesFrom or CopyRuntimeBodyInstancePropertiesFrom

diff --git a/Engine/Source/Runtime/Engine/Private/ISMPartition/ISMComponentDescriptor.cpp b/Engine/Source/Runtime/Engine/Private/ISMPartition/ISMComponentDescriptor.cpp
--- a/Engine/Source/Runtime/Engine/Private/ISMPartition/ISMComponentDescriptor.cpp
+++ b/Engine/Source/Runtime/Engine/Private/ISMPartition/ISMComponentDescriptor.cpp
@@ -165,7 +165,15 @@ void FISMComponentDescriptorBase::InitFrom(const UStaticMeshComponent* Template,
 
 	if (bInitBodyInstance)
 	{
-		BodyInstance.CopyBodyInstancePropertiesFrom(&Template->BodyInstance);
+		// If we are building this batch live in the editor, then the Body Properties will be in a different runtime state.
+		if (Template->BodyInstance.BodySetup.Get() == nullptr)
+		{
+			BodyInstance.CopyBodyInstancePropertiesFrom(&Template->BodyInstance);
+		}
+		else
+		{
+			BodyInstance.CopyRuntimeBodyInstancePropertiesFrom(&Template->BodyInstance);
+		}
 	}
 }
 
