Fix race condition between parent and child class sparse class data patching/copying

In ALT-enabled builds, patching the exposed value handlers of a parent class can race with copying the same handlers to a derived class's sparse class data

Added a critical section to guard against this (which is rarely contended)

diff --git a/Engine/Source/Runtime/Engine/Private/Animation/AnimSubsystem_Base.cpp b/Engine/Source/Runtime/Engine/Private/Animation/AnimSubsystem_Base.cpp
--- a/Engine/Source/Runtime/Engine/Private/Animation/AnimSubsystem_Base.cpp
+++ b/Engine/Source/Runtime/Engine/Private/Animation/AnimSubsystem_Base.cpp
@@ -13,6 +13,8 @@ void FAnimSubsystem_Base::OnPostLoadDefaults(FAnimSubsystemPostLoadDefaultsConte
 
 void FAnimSubsystem_Base::PatchValueHandlers(UClass* InClass)
 {
+	FScopeLock Lock(&ValueHandlersCriticalSection);
+	
 	ExposedValueHandlers.Reset();
 
 	void* SparseClassData = const_cast<void*>(InClass->GetSparseClassData(EGetSparseClassDataMethod::ReturnIfNull));
@@ -30,4 +32,20 @@ void FAnimSubsystem_Base::PatchValueHandlers(UClass* InClass)
 			}
 		}
 	}
-}
\ No newline at end of file
+}
+
+// Unlike the copy assignment operator, the copy constructor doesn't take a lock (it is not called during UClass::CreateSparseClassData)
+FAnimSubsystem_Base::FAnimSubsystem_Base(const FAnimSubsystem_Base& Other)
+	: ExposedValueHandlers(Other.ExposedValueHandlers)
+{
+}
+
+FAnimSubsystem_Base& FAnimSubsystem_Base::operator =(const FAnimSubsystem_Base& Other)
+{
+	// Need to perform a lock as copying can race with PatchValueHandlers in async loading thread-enabled builds
+	FScopeLock Lock(&ValueHandlersCriticalSection);
+
+	ExposedValueHandlers = Other.ExposedValueHandlers;
+
+	return *this;
+}
diff --git a/Engine/Source/Runtime/Engine/Public/Animation/AnimSubsystem_Base.h b/Engine/Source/Runtime/Engine/Public/Animation/AnimSubsystem_Base.h
--- a/Engine/Source/Runtime/Engine/Public/Animation/AnimSubsystem_Base.h
+++ b/Engine/Source/Runtime/Engine/Public/Animation/AnimSubsystem_Base.h
@@ -14,6 +14,12 @@ struct FAnimSubsystem_Base : public FAnimSubsystem
 
 	friend class UAnimBlueprintExtension_Base;
 
+	// Override copy/copy-assignment to avoid racing base class PatchValueHandlers with copying
+	// derived class sparse class data
+	FAnimSubsystem_Base() = default;
+	ENGINE_API FAnimSubsystem_Base(const FAnimSubsystem_Base& Other);
+	ENGINE_API FAnimSubsystem_Base& operator =(const FAnimSubsystem_Base& Other);
+	
 	// FAnimSubsystem interface
 	ENGINE_API virtual void OnPostLoadDefaults(FAnimSubsystemPostLoadDefaultsContext& InContext) override;
 
@@ -22,5 +28,6 @@ struct FAnimSubsystem_Base : public FAnimSubsystem
 	ENGINE_API void PatchValueHandlers(UClass* InClass);
 
 private:
+	FCriticalSection ValueHandlersCriticalSection;
 	TArray<FExposedValueHandler> ExposedValueHandlers;
 };
