Added option to apply the actor editor context data layers to the selected actor(s).
For partition actors (Landscape/Foliage), this will also properly rename the actors, label and change the actor external package.

diff --git a/Engine/Source/Editor/DataLayerEditor/Private/DataLayer/DataLayerEditorSubsystem.cpp b/Engine/Source/Editor/DataLayerEditor/Private/DataLayer/DataLayerEditorSubsystem.cpp
--- a/Engine/Source/Editor/DataLayerEditor/Private/DataLayer/DataLayerEditorSubsystem.cpp
+++ b/Engine/Source/Editor/DataLayerEditor/Private/DataLayer/DataLayerEditorSubsystem.cpp
@@ -68,6 +68,8 @@
 #include "WorldPartition/DataLayer/WorldDataLayers.h"
 #include "WorldPartition/WorldPartition.h"
 #include "WorldPartition/IWorldPartitionEditorModule.h"
+#include "ActorPartition/PartitionActor.h"
+#include "ActorPartition/ActorPartitionSubsystem.h"
 
 class SWidget;
 
@@ -2003,6 +2005,100 @@ bool UDataLayerEditorSubsystem::SetDataLayerShortName(UDataLayerInstance* DataLa
 	return false;
 }
 
+bool UDataLayerEditorSubsystem::MoveActorToDataLayers(AActor* InActor, const TArray<UDataLayerInstance*>& InDataLayerInstances)
+{
+	UWorld* World = GetWorld();
+	UDataLayerManager* DataLayerManager = UDataLayerManager::GetDataLayerManager(World);
+	if (!DataLayerManager)
+	{
+		return false;
+	}
+
+	TSet<UDataLayerInstance*> AllDataLayerInstance(DataLayerManager->GetDataLayerInstances());
+	for (UDataLayerInstance* DataLayerInstance : InDataLayerInstances)
+	{
+		if (!AllDataLayerInstance.Contains(DataLayerInstance))
+		{
+			UE_LOG(LogWorldPartition, Warning, TEXT("Can't move actor %s to Data Layers: Invalid Data Layer %s"), *InActor->GetActorNameOrLabel(), *DataLayerInstance->GetDataLayerShortName());
+			return false;
+		}
+	}
+
+	TSet<const UDataLayerInstance*> ActorDataLayers(InActor->GetDataLayerInstances());
+	TSet<const UDataLayerInstance*> NewDataLayers(InDataLayerInstances);
+	if ((ActorDataLayers.Num() == NewDataLayers.Num()) &&
+		(ActorDataLayers.Intersect(NewDataLayers).Num() == ActorDataLayers.Num()))
+	{
+		return true;
+	}
+
+	// Data Layers of partition actors impact their name. In this case, use a dedicated function that will do the move operation.
+	APartitionActor* PartitionActor = Cast<APartitionActor>(InActor);
+	if (PartitionActor && PartitionActor->IsPartitionActorNameAffectedByDataLayers())
+	{
+		UActorPartitionSubsystem* ActorPartitionSubsystem = UWorld::GetSubsystem<UActorPartitionSubsystem>(World);
+		if (!ActorPartitionSubsystem || !ActorPartitionSubsystem->MoveActorToDataLayers(PartitionActor, InDataLayerInstances))
+		{
+			return false;
+		}
+	}
+	else
+	{
+		// Test moving from/to new External Data Layer
+		const UExternalDataLayerAsset* OldExternalDataLayerAsset = InActor->GetExternalDataLayerAsset();
+		const UDataLayerInstance* const* ExternalDataLayerInstance = Algo::FindByPredicate(InDataLayerInstances, [](const UDataLayerInstance* DataLayerInstance) { return DataLayerInstance->IsA<UExternalDataLayerInstance>(); });
+		const UExternalDataLayerInstance* NewExternalDataLayerInstance = ExternalDataLayerInstance ? Cast<UExternalDataLayerInstance>(*ExternalDataLayerInstance) : nullptr;
+		const UExternalDataLayerAsset* NewExternalDataLayerAsset = NewExternalDataLayerInstance ? NewExternalDataLayerInstance->GetExternalDataLayerAsset() : nullptr;
+		constexpr bool bAllowNonUserManaged = true;
+		FMoveToExternalDataLayerParams Params(NewExternalDataLayerInstance, bAllowNonUserManaged);
+		if ((OldExternalDataLayerAsset != NewExternalDataLayerAsset))
+		{
+			FText FailureReason;
+			if (!FExternalDataLayerHelper::CanMoveActorsToExternalDataLayer({ InActor }, Params, &FailureReason))
+			{
+				UE_LOG(LogWorldPartition, Warning, TEXT("Can't move actor %s to External Data Layer. %s"), *InActor->GetName(), *FailureReason.ToString());
+				return false;
+			}
+		}
+
+		// Move to new Data Layers (except for the External Data Layer)
+		RemoveActorFromAllDataLayers(InActor);
+		for (UDataLayerInstance* DataLayerInstance : InDataLayerInstances)
+		{
+			if (!DataLayerInstance->IsA<UExternalDataLayerInstance>())
+			{
+				AddActorToDataLayer(InActor, DataLayerInstance);
+			}
+		}
+		// Move to new External Data Layer
+		if (OldExternalDataLayerAsset != NewExternalDataLayerAsset)
+		{
+			verify(FExternalDataLayerHelper::MoveActorsToExternalDataLayer({ InActor }, Params));
+		}
+	}
+	return true;
+}
+
+bool UDataLayerEditorSubsystem::ApplyActorEditorContextDataLayersToActors(const TArray<AActor*>& InActors)
+{
+	UDataLayerManager* DataLayerManager = UDataLayerManager::GetDataLayerManager(GetWorld());
+	if (!DataLayerManager)
+	{
+		return false;
+	}
+	TArray<UDataLayerInstance*> DataLayerInstances = DataLayerManager->GetActorEditorContextDataLayers();
+
+	bool bSuccess = true;
+	for (AActor* Actor : InActors)
+	{
+		if (!MoveActorToDataLayers(Actor, DataLayerInstances))
+		{
+			bSuccess = false;
+		}
+	}
+	return bSuccess;
+}
+
 //~ Begin Deprecated
 
 PRAGMA_DISABLE_DEPRECATION_WARNINGS
diff --git a/Engine/Source/Editor/DataLayerEditor/Private/DataLayerEditorModule.cpp b/Engine/Source/Editor/DataLayerEditor/Private/DataLayerEditorModule.cpp
--- a/Engine/Source/Editor/DataLayerEditor/Private/DataLayerEditorModule.cpp
+++ b/Engine/Source/Editor/DataLayerEditor/Private/DataLayerEditorModule.cpp
@@ -209,19 +209,63 @@ void FDataLayerEditorModule::RegisterMenus()
 	FillExternalDataLayerMenuAction.CanExecuteAction = FToolMenuCanExecuteAction::CreateLambda([this](const FToolMenuContext& InContext)
 	{
 		TArray<AActor*> SelectedActors = GetSelectedActors();
+		if (Algo::AnyOf(GetSelectedActors(), [](AActor* Actor) { return Actor && !Actor->IsUserManaged(); }))
+		{
+			return false;
+		}
 		UExternalDataLayerManager* ExternalDataLayerManager = SelectedActors.Num() ? UExternalDataLayerManager::GetExternalDataLayerManager(SelectedActors[0]) : nullptr;
 		return ExternalDataLayerManager && ExternalDataLayerManager->HasInjectedExternalDataLayerAssets();
 	});
 
+	FToolUIAction ApplyActorEditorContextDataLayersToActorAction;
+	ApplyActorEditorContextDataLayersToActorAction.ExecuteAction = FToolMenuExecuteAction::CreateLambda([this](const FToolMenuContext& InContext)
+	{
+		FScopedTransaction Transaction(LOCTEXT("ApplyActorEditorContextDataLayersToActors", "Apply Actor Editor Context Data Layers To Actor(s)"));
+		if (!UDataLayerEditorSubsystem::Get()->ApplyActorEditorContextDataLayersToActors(GetSelectedActors()))
+		{
+			const FText WarningMessage = LOCTEXT("ApplyActorEditorContextDataLayersToActorFailed", "Failed apply actor editor context data layers to actor(s)");
+			const FText DetailedMessage = LOCTEXT("ApplyActorEditorContextDataLayersToActorFailedDetailed", "See log for details.");
+			UE::Private::DataLayerEditorModule::LogWarningAndNotify(WarningMessage, DetailedMessage);
+		}
+	});
+
+	auto FillDataLayerMenu = [this, ApplyActorEditorContextDataLayersToActorAction, FillExternalDataLayerMenu, FillExternalDataLayerMenuAction](UToolMenu* SubMenu)
+	{
+		FToolMenuSection& DataLayerSection = SubMenu->AddSection("Data Layer");
+		DataLayerSection.AddMenuEntry(
+			"ApplyActorEditorContextDataLayersToActorsMenu",
+			LOCTEXT("ApplyActorEditorContextDataLayersToActorMenuEntry", "Apply Actor Editor Context"),
+			LOCTEXT("ApplyActorEditorContextDataLayersToActorMenu_ToolTip", "Applies the Actor Editor Context's Current Data Layers to the actor(s)"),
+			FSlateIcon(FAppStyle::GetAppStyleSetName(), "DataLayer.Editor"),
+			ApplyActorEditorContextDataLayersToActorAction);
+
+		DataLayerSection.AddSubMenu(
+			"ExternalDataLayerSubMenu",
+			LOCTEXT("ExternalDataLayerSubMenu", "External Data Layer"),
+			LOCTEXT("ExternalDataLayerSubMenu_ToolTip", "External Data Layer Utils"),
+			FNewToolMenuDelegate::CreateLambda(FillExternalDataLayerMenu),
+			FillExternalDataLayerMenuAction,
+			EUserInterfaceActionType::Button,
+			false,
+			FSlateIcon(FAppStyle::GetAppStyleSetName(), "DataLayer.External"));
+	};
+
+	FToolUIAction FillDataLayerMenuAction;
+	FillExternalDataLayerMenuAction.IsActionVisibleDelegate = FToolMenuCanExecuteAction::CreateLambda([this](const FToolMenuContext& InContext)
+	{
+		UWorld* EditorWorld = GEditor->GetEditorWorldContext().World();
+		return EditorWorld && EditorWorld->IsPartitionedWorld();
+	});
+
 	ActorSection.AddSubMenu(
 		"ExternalDataLayerSubMenu",
-		LOCTEXT("ExternalDataLayerSubMenu", "External Data Layer"),
-		LOCTEXT("ExternalDataLayerSubMenu_ToolTip", "External Data Layer Utils"),
-		FNewToolMenuDelegate::CreateLambda(FillExternalDataLayerMenu),
-		FillExternalDataLayerMenuAction,
+		LOCTEXT("DataLayerSubMenu", "Data Layer"),
+		LOCTEXT("DataLayerSubMenu_ToolTip", "Data Layer Utils"),
+		FNewToolMenuDelegate::CreateLambda(FillDataLayerMenu),
+		FillDataLayerMenuAction,
 		EUserInterfaceActionType::Button,
 		false,
-		FSlateIcon(FAppStyle::GetAppStyleSetName(), "DataLayer.External"));
+		FSlateIcon(FAppStyle::GetAppStyleSetName(), "DataLayer.Editor"));
 }
 
 TArray<AActor*> FDataLayerEditorModule::GetSelectedActors() const
diff --git a/Engine/Source/Editor/DataLayerEditor/Public/DataLayer/DataLayerEditorSubsystem.h b/Engine/Source/Editor/DataLayerEditor/Public/DataLayer/DataLayerEditorSubsystem.h
--- a/Engine/Source/Editor/DataLayerEditor/Public/DataLayer/DataLayerEditorSubsystem.h
+++ b/Engine/Source/Editor/DataLayerEditor/Public/DataLayer/DataLayerEditorSubsystem.h
@@ -662,6 +662,13 @@ public:
 	 */
 	bool SetDataLayerShortName(UDataLayerInstance* InDataLayerInstance, const FString& InNewShortName);
 
+	/**
+	 * Move actors to the Actor Editor Context's Data Layers.
+	 * 
+	 * @return true if operation succeeded.
+	 */
+	bool ApplyActorEditorContextDataLayersToActors(const TArray<AActor*>& InActors);
+
 	//~ Begin Deprecated
 
 	PRAGMA_DISABLE_DEPRECATION_WARNINGS
@@ -779,6 +786,7 @@ private:
 	void ApplyContext(AActor* InActor, bool bInForceTryApply = false, AActor* InReplacedActor = nullptr);
 	const UExternalDataLayerInstance* GetActorSpawningExternalDataLayerInstance(AActor* InActor) const;
 	static const UExternalDataLayerAsset* GetReferencingWorldSurrogateObjectForObject(UWorld* ReferencingWorld, const FSoftObjectPath& ObjectPath);
+	bool MoveActorToDataLayers(AActor* InActor, const TArray<UDataLayerInstance*>& InDataLayerInstances);
 	void MoveActorToExternalDataLayer(AActor* InActor, const UExternalDataLayerInstance* InExternalDataLayerInstance, bool bInNotifyFailureReason = true);
 	bool ShouldHandleActor(AActor* InActor) const;
 
diff --git a/Engine/Source/Runtime/Engine/Classes/GameFramework/Actor.h b/Engine/Source/Runtime/Engine/Classes/GameFramework/Actor.h
--- a/Engine/Source/Runtime/Engine/Classes/GameFramework/Actor.h
+++ b/Engine/Source/Runtime/Engine/Classes/GameFramework/Actor.h
@@ -4739,6 +4739,7 @@ private:
 	friend class UExternalDataLayerInstance;
 	friend class ULevelInstanceSubsystem;
 	friend class FExternalDataLayerHelper;
+	friend class UActorPartitionSubsystem;
 };
 
 struct FSetActorIsInLevelInstance
diff --git a/Engine/Source/Runtime/Engine/Private/ActorPartition/ActorPartitionSubsystem.cpp b/Engine/Source/Runtime/Engine/Private/ActorPartition/ActorPartitionSubsystem.cpp
--- a/Engine/Source/Runtime/Engine/Private/ActorPartition/ActorPartitionSubsystem.cpp
+++ b/Engine/Source/Runtime/Engine/Private/ActorPartition/ActorPartitionSubsystem.cpp
@@ -12,8 +12,11 @@
 #include "EngineUtils.h"
 
 #if WITH_EDITOR
+#include "Algo/Transform.h"
+#include "WorldPartition/WorldPartitionLog.h"
 #include "WorldPartition/DataLayer/DataLayerManager.h"
 #include "WorldPartition/ContentBundle/ContentBundleEngineSubsystem.h"
+#include "DeletedObjectPlaceholder.h"
 #endif
 
 #include UE_INLINE_GENERATED_CPP_BY_NAME(ActorPartitionSubsystem)
@@ -409,34 +412,39 @@ APartitionActor* UActorPartitionSubsystem::GetActor(const FActorPartitionGetPara
 
 APartitionActor* UActorPartitionSubsystem::GetActor(const TSubclassOf<APartitionActor>& InActorClass, const FCellCoord& InCellCoords, bool bInCreate, const FGuid& InGuid, uint32 InGridSize, bool bInBoundsSearch, TFunctionRef<void(APartitionActor*)> InActorCreated)
 {
-	const uint32 GridSize = InGridSize > 0 ? InGridSize : InActorClass->GetDefaultObject<APartitionActor>()->GetDefaultGridSize(GetWorld());
-	
 	UWorld* World = GetWorld();
 	FGuid ContentBundleGuid = UContentBundleEngineSubsystem::Get()->GetEditingContentBundleGuid();
 	const UDataLayerManager* DataLayerManager = UDataLayerManager::GetDataLayerManager(World);
 	const uint32 DataLayerEditorContextHash = DataLayerManager ? DataLayerManager->GetDataLayerEditorContextHash() : FDataLayerEditorContext::EmptyHash;
 	FActorPartitionIdentifier ActorPartitionId(InActorClass, InGuid, FActorPartitionContextHash::Get(ContentBundleGuid, DataLayerEditorContextHash));
+	return GetActor(ActorPartitionId, InCellCoords, bInCreate, InGridSize, bInBoundsSearch, InActorCreated);
+}
+
+APartitionActor* UActorPartitionSubsystem::GetActor(const FActorPartitionIdentifier& InActorPartitionId, const FCellCoord& InCellCoords, bool bInCreate, uint32 InGridSize, bool bInBoundsSearch, TFunctionRef<void(APartitionActor*)> InActorCreated)
+{
+	UWorld* World = GetWorld();
+	const uint32 GridSize = InGridSize > 0 ? InGridSize : InActorPartitionId.GetClass()->GetDefaultObject<APartitionActor>()->GetDefaultGridSize(World);
 	TMap<FActorPartitionIdentifier, TWeakObjectPtr<APartitionActor>>* ActorsPerId = PartitionedActors.Find(InCellCoords);
 	APartitionActor* FoundActor = nullptr;
 	if (!ActorsPerId)
 	{
-		FoundActor = ActorPartition->GetActor(ActorPartitionId, bInCreate, InCellCoords, GridSize, bInBoundsSearch, InActorCreated);
+		FoundActor = ActorPartition->GetActor(InActorPartitionId, bInCreate, InCellCoords, GridSize, bInBoundsSearch, InActorCreated);
 		if (FoundActor)
 		{
-			PartitionedActors.Add(InCellCoords).Add(ActorPartitionId, FoundActor);
+			PartitionedActors.Add(InCellCoords).Add(InActorPartitionId, FoundActor);
 		}
 	}
 	else
 	{
-		TWeakObjectPtr<APartitionActor>* ActorPtr = ActorsPerId->Find(ActorPartitionId);
+		TWeakObjectPtr<APartitionActor>* ActorPtr = ActorsPerId->Find(InActorPartitionId);
 		if (!ActorPtr || !ActorPtr->IsValid())
 		{
-			FoundActor = ActorPartition->GetActor(ActorPartitionId, bInCreate, InCellCoords, GridSize, bInBoundsSearch, InActorCreated);
+			FoundActor = ActorPartition->GetActor(InActorPartitionId, bInCreate, InCellCoords, GridSize, bInBoundsSearch, InActorCreated);
 			if (FoundActor)
 			{
 				if (!ActorPtr)
 				{
-					ActorsPerId->Add(ActorPartitionId, FoundActor);
+					ActorsPerId->Add(InActorPartitionId, FoundActor);
 				}
 				else
 				{
@@ -453,5 +461,156 @@ APartitionActor* UActorPartitionSubsystem::GetActor(const TSubclassOf<APartition
 	return FoundActor;
 }
 
+bool UActorPartitionSubsystem::MoveActorToDataLayers(APartitionActor* InActor, const TArray<UDataLayerInstance*>& InDataLayerInstances)
+{
+	auto ChangeActorExternalPackage = [](AActor* InActor, const UPackage* InOldActorPackage, UPackage* InNewActorPackage)
+	{
+		const bool bShouldDirty = true;
+		const bool bLevelPackageWasDirty = InActor->GetLevel()->GetPackage()->IsDirty();
+
+		InActor->SetPackageExternal(false, bShouldDirty);
+
+		// Get all other dependant objects in the old actor package
+		TArray<UObject*> DependantObjects;
+		ForEachObjectWithPackage(InOldActorPackage, [&DependantObjects](UObject* Object)
+		{
+			if (!Cast<UMetaData>(Object) && !Cast<UDeletedObjectPlaceholder>(Object) && !Cast<AActor>(Object))
+			{
+				DependantObjects.Add(Object);
+			}
+			return true;
+		}, false, RF_NoFlags, EInternalObjectFlags::Garbage); // Skip garbage objects (like child actors destroyed when de-externalizing the actor)
+
+		InActor->SetPackageExternal(true, bShouldDirty, InNewActorPackage);
+		check(InNewActorPackage == InActor->GetExternalPackage());
+
+		// Move dependant objects into the new actor package
+		for (UObject* DependantObject : DependantObjects) //-V1078
+		{
+			DependantObject->Rename(nullptr, InNewActorPackage, REN_NonTransactional | REN_DontCreateRedirectors | REN_DoNotDirty);
+		}
+
+		// Restore level package dirty flag
+		if (!bLevelPackageWasDirty)
+		{
+			InActor->GetLevel()->GetPackage()->SetDirtyFlag(false);
+		}
+	};
+
+	if (!InActor->IsPartitionActorNameAffectedByDataLayers())
+	{
+		UE_LOG(LogWorldPartition, Warning, TEXT("Failed to move actor %s to Data Layers: Use UDataLayerEditorSubsystem::MoveActorToDataLayers instead."), *InActor->GetName());
+		return false;
+	}
+
+	UWorld* World = InActor->GetWorld();
+	TArray<FName> DataLayerInstanceNames;
+	Algo::Transform(InDataLayerInstances, DataLayerInstanceNames, [](UDataLayerInstance* DataLayerInstance) { return DataLayerInstance->GetDataLayerFName(); });
+	const uint32 DataLayerEditorContextHash = FDataLayerEditorContext(World, DataLayerInstanceNames).GetHash();
+	const FActorPartitionIdentifier ActorPartitionId(InActor->GetClass(), InActor->GetGridGuid(), FActorPartitionContextHash::Get(InActor->GetContentBundleGuid(), DataLayerEditorContextHash));
+	const TUniquePtr<FWorldPartitionActorDesc> ActorDesc = InActor->CreateActorDesc();
+	const FPartitionActorDesc* PartitionActorDesc = (const FPartitionActorDesc*)ActorDesc.Get();
+	const UActorPartitionSubsystem::FCellCoord CellCoord(PartitionActorDesc->GridIndexX, PartitionActorDesc->GridIndexY, PartitionActorDesc->GridIndexZ, InActor->GetLevel());
+	const FString NewActorName = APartitionActor::GetActorName(CellCoord.Level->GetTypedOuter<UWorld>(), ActorPartitionId, InActor->GetGridSize(), CellCoord.X, CellCoord.Y, CellCoord.Z);
+	const FString& OldActorName = InActor->GetName();
+
+	// Test if there's nothing to do
+	if (OldActorName == NewActorName)
+	{
+		return true;
+	}
+
+	// Test moving from/to new External Data Layer
+	const UExternalDataLayerAsset* OldExternalDataLayerAsset = InActor->GetExternalDataLayerAsset();
+	const UDataLayerInstance* const* ExternalDataLayerInstance = Algo::FindByPredicate(InDataLayerInstances, [](const UDataLayerInstance* DataLayerInstance) { return DataLayerInstance->IsA<UExternalDataLayerInstance>(); });
+	const UExternalDataLayerInstance* NewExternalDataLayerInstance = ExternalDataLayerInstance ? Cast<UExternalDataLayerInstance>(*ExternalDataLayerInstance) : nullptr;
+	const UExternalDataLayerAsset* NewExternalDataLayerAsset = NewExternalDataLayerInstance ? NewExternalDataLayerInstance->GetExternalDataLayerAsset() : nullptr;
+	constexpr bool bAllowNonUserManaged = true;
+	FMoveToExternalDataLayerParams Params(NewExternalDataLayerInstance, bAllowNonUserManaged);
+	if ((OldExternalDataLayerAsset != NewExternalDataLayerAsset))
+	{
+		FText FailureReason;
+		if (!FExternalDataLayerHelper::CanMoveActorsToExternalDataLayer({ InActor }, Params, &FailureReason))
+		{
+			UE_LOG(LogWorldPartition, Warning, TEXT("Can't move actor %s to External Data Layer. %s"), *InActor->GetName(), *FailureReason.ToString());
+			return false;
+		}
+	}
+
+	// Find any existing actor with this setup
+	constexpr bool bCreate = false;
+	constexpr bool bBoundsSearch = false;
+	if (APartitionActor* ExistingActor = GetActor(ActorPartitionId, CellCoord, bCreate, InActor->GetGridSize(), bBoundsSearch); ExistingActor && (ExistingActor != InActor))
+	{
+		UE_LOG(LogEngine, Warning, TEXT("Failed to move actor %s to Data Layers: Another partition actor %s already exists."), *InActor->GetName(), *ExistingActor->GetName());
+		return false;
+	}
+
+	// Update PartitionedActors
+	TMap<FActorPartitionIdentifier, TWeakObjectPtr<APartitionActor>>* ActorsPerId = PartitionedActors.Find(CellCoord);
+	if (ActorsPerId)
+	{
+		for (auto It = ActorsPerId->CreateIterator(); It; ++It)
+		{
+			if (It.Value().Get() == InActor)
+			{
+				It.RemoveCurrent();
+				break;
+			}
+		}
+	}
+	else
+	{
+		ActorsPerId = &PartitionedActors.Add(CellCoord);
+	}
+	ActorsPerId->Add(ActorPartitionId, InActor);
+
+	// Handle the case where the actor already exists, but is in the undo stack (was deleted)
+	if (UObject* ExistingObject = StaticFindObject(nullptr, CellCoord.Level, *NewActorName))
+	{
+		AActor* ExistingActor = CastChecked<AActor>(ExistingObject);
+		check(!IsValidChecked(ExistingActor));
+		ExistingActor->Modify();
+		// Don't go through AActor::Rename here because we aren't changing outers (the actor's level). 
+		// We just want to rename that actor out of the way.
+		ExistingActor->UObject::Rename(nullptr, nullptr, REN_DontCreateRedirectors | REN_DoNotDirty | REN_NonTransactional);
+	}
+
+	// Create new actor package
+	TStringBuilderWithBuffer<TCHAR, NAME_SIZE> NewActorPath;
+	NewActorPath += CellCoord.Level->GetPathName();
+	NewActorPath += TEXT(".");
+	NewActorPath += NewActorName;
+	UPackage* NewActorPackage = ULevel::CreateActorPackage(CellCoord.Level->GetPackage(), CellCoord.Level->GetActorPackagingScheme(), *NewActorPath);
+	
+	// Swap content from the old external package to the new one
+	const UPackage* OldActorPackage = InActor->GetExternalPackage();
+	ChangeActorExternalPackage(InActor, OldActorPackage, NewActorPackage);
+
+	// Update name and label
+	InActor->Rename(*NewActorName, nullptr, REN_NonTransactional | REN_DontCreateRedirectors | REN_DoNotDirty);
+	APartitionActor::SetLabelForActor(InActor, ActorPartitionId, InActor->GetGridSize(), CellCoord.X, CellCoord.Y, CellCoord.Z);
+
+	// Move to new Data Layers (except for the External Data Layer)
+	const bool bIncludeExternalDataLayerAsset = false;
+	for (const UDataLayerAsset* DataLayerAsset : InActor->GetDataLayerAssets(bIncludeExternalDataLayerAsset))
+	{
+		FAssignActorDataLayer::RemoveDataLayerAsset(InActor, DataLayerAsset);
+	}
+	for (const UDataLayerInstance* DataLayerInstance : InDataLayerInstances)
+	{
+		if (const UDataLayerAsset* DataLayerAsset = !DataLayerInstance->IsA<UExternalDataLayerInstance>() ? DataLayerInstance->GetAsset() : nullptr)
+		{
+			FAssignActorDataLayer::AddDataLayerAsset(InActor, DataLayerAsset);
+		}
+	}
+	// Move to new External Data Layer
+	if (OldExternalDataLayerAsset != NewExternalDataLayerAsset)
+	{
+		verify(FExternalDataLayerHelper::MoveActorsToExternalDataLayer({ InActor }, Params));
+	}
+	return true;
+}
+
 #endif // WITH_EDITOR
 
diff --git a/Engine/Source/Runtime/Engine/Private/WorldPartition/DataLayer/ExternalDataLayerHelper.cpp b/Engine/Source/Runtime/Engine/Private/WorldPartition/DataLayer/ExternalDataLayerHelper.cpp
--- a/Engine/Source/Runtime/Engine/Private/WorldPartition/DataLayer/ExternalDataLayerHelper.cpp
+++ b/Engine/Source/Runtime/Engine/Private/WorldPartition/DataLayer/ExternalDataLayerHelper.cpp
@@ -131,9 +131,9 @@ namespace UE::Private::ExternalDataLayerHelper
 	}
 }
 
-bool FExternalDataLayerHelper::CanMoveActorsToExternalDataLayer(const TArray<AActor*>& InActors, const UExternalDataLayerInstance* InExternalDataLayerInstance, FText* OutFailureReason)
+bool FExternalDataLayerHelper::CanMoveActorsToExternalDataLayer(const TArray<AActor*>& InActors, const FMoveToExternalDataLayerParams& InParams, FText* OutFailureReason)
 {
-	auto CanMoveActorToExternalDataLayer = [](AActor* InActor, const UExternalDataLayerInstance* InExternalDataLayerInstance, FText& OutFailureReason)
+	auto CanMoveActorToExternalDataLayer = [](AActor* InActor, const FMoveToExternalDataLayerParams& InParams, FText& OutFailureReason)
 	{
 		const TArray<UClass*> IgnoreClasses;
 		const TArray<UPackage*> IgnorePackages;
@@ -147,7 +147,7 @@ bool FExternalDataLayerHelper::CanMoveActorsToExternalDataLayer(const TArray<AAc
 			return false;
 		}
 
-		if (!InActor->IsUserManaged() && !InActor->GetExternalPackage()->HasAnyPackageFlags(PKG_NewlyCreated))
+		if (!InParams.bAllowNonUserManaged && !InActor->IsUserManaged() && !InActor->GetExternalPackage()->HasAnyPackageFlags(PKG_NewlyCreated))
 		{
 			OutFailureReason = FText::Format(LOCTEXT("CantMoveActorToEDL_NotUserManaged", "Actor {0} cannot be manually modified."), FText::FromString(InActor->GetName()));
 			return false;
@@ -164,7 +164,7 @@ bool FExternalDataLayerHelper::CanMoveActorsToExternalDataLayer(const TArray<AAc
 		FExternalDataLayerHelper::IsExternalDataLayerPath(OldPackage->GetPathName(), &OldExternalDataLayerUIDFromPackage);
 		const UExternalDataLayerAsset* OldExternalDataLayerAsset = InActor->GetExternalDataLayerAsset();
 		const FExternalDataLayerUID OldExternalDataLayerUIDFromAsset = OldExternalDataLayerAsset ? OldExternalDataLayerAsset->GetUID() : FExternalDataLayerUID();
-		const UExternalDataLayerAsset* NewExternalDataLayerAsset = InExternalDataLayerInstance ? InExternalDataLayerInstance->GetExternalDataLayerAsset() : nullptr;
+		const UExternalDataLayerAsset* NewExternalDataLayerAsset = InParams.ExternalDataLayerInstance ? InParams.ExternalDataLayerInstance->GetExternalDataLayerAsset() : nullptr;
 
 		// Detect if old actor package and EDL asset mismatch (this can happen during the replace actor process).
 		// In this case, skip the optimization that tries to detect no change that are relying on OldExternalDataLayerAsset (as it's unreliable).
@@ -190,9 +190,9 @@ bool FExternalDataLayerHelper::CanMoveActorsToExternalDataLayer(const TArray<AAc
 			return false;
 		}
 
-		if (InExternalDataLayerInstance && InExternalDataLayerInstance->IsReadOnly())
+		if (InParams.ExternalDataLayerInstance && InParams.ExternalDataLayerInstance->IsReadOnly())
 		{
-			OutFailureReason = FText::Format(LOCTEXT("CantMoveActorToEDL_ReadOnlyExternalDataLayer", "External Data Layer is read-only."), FText::FromString(InExternalDataLayerInstance->GetDataLayerShortName()));
+			OutFailureReason = FText::Format(LOCTEXT("CantMoveActorToEDL_ReadOnlyExternalDataLayer", "External Data Layer is read-only."), FText::FromString(InParams.ExternalDataLayerInstance->GetDataLayerShortName()));
 			return false;
 		}
 
@@ -269,7 +269,7 @@ bool FExternalDataLayerHelper::CanMoveActorsToExternalDataLayer(const TArray<AAc
 	for (AActor* Actor : InActors)
 	{
 		FText FailureReason;
-		if (!CanMoveActorToExternalDataLayer(Actor, InExternalDataLayerInstance, FailureReason))
+		if (!CanMoveActorToExternalDataLayer(Actor, InParams, FailureReason))
 		{
 			if (OutFailureReason)
 			{
@@ -284,11 +284,16 @@ bool FExternalDataLayerHelper::CanMoveActorsToExternalDataLayer(const TArray<AAc
 
 bool FExternalDataLayerHelper::MoveActorsToExternalDataLayer(const TArray<AActor*>& InActors, const UExternalDataLayerInstance* InExternalDataLayerInstance, FText* OutFailureReason)
 {
-	auto MoveActorToExternalDataLayer = [](AActor* InActor, const UExternalDataLayerInstance* InExternalDataLayerInstance)
+	return MoveActorsToExternalDataLayer(InActors, FMoveToExternalDataLayerParams(InExternalDataLayerInstance), OutFailureReason);
+}
+
+bool FExternalDataLayerHelper::MoveActorsToExternalDataLayer(const TArray<AActor*>& InActors, const FMoveToExternalDataLayerParams& InParams, FText* OutFailureReason)
+{
+	auto MoveActorToExternalDataLayer = [](AActor* InActor, const FMoveToExternalDataLayerParams& InParams)
 	{
 		check(InActor->IsMainPackageActor());
 		const UPackage* OldActorPackage = InActor->GetExternalPackage();
-		const UExternalDataLayerAsset* NewExternalDataLayerAsset = InExternalDataLayerInstance ? InExternalDataLayerInstance->GetExternalDataLayerAsset() : nullptr;
+		const UExternalDataLayerAsset* NewExternalDataLayerAsset = InParams.ExternalDataLayerInstance ? InParams.ExternalDataLayerInstance->GetExternalDataLayerAsset() : nullptr;
 		const bool bShouldDirty = true;
 		const bool bLevelPackageWasDirty = InActor->GetLevel()->GetPackage()->IsDirty();
 		InActor->SetPackageExternal(false, bShouldDirty);
@@ -345,7 +350,7 @@ bool FExternalDataLayerHelper::MoveActorsToExternalDataLayer(const TArray<AActor
 	};
 	
 	// First, validate that the whole operation can be done without any validation errors
-	if (!CanMoveActorsToExternalDataLayer(InActors, InExternalDataLayerInstance, OutFailureReason))
+	if (!CanMoveActorsToExternalDataLayer(InActors, InParams, OutFailureReason))
 	{
 		return false;
 	}
@@ -353,7 +358,7 @@ bool FExternalDataLayerHelper::MoveActorsToExternalDataLayer(const TArray<AActor
 	// Change all actors External Data Layer asset
 	for (AActor* Actor : InActors)
 	{
-		if (ensure(MoveActorToExternalDataLayer(Actor, InExternalDataLayerInstance)))
+		if (ensure(MoveActorToExternalDataLayer(Actor, InParams)))
 		{
 			// Basic validation on the actor and its new External Data Layer asset
 			UExternalDataLayerManager* ExternalDataLayerManager = UExternalDataLayerManager::GetExternalDataLayerManager(Actor);
diff --git a/Engine/Source/Runtime/Engine/Public/ActorPartition/ActorPartitionSubsystem.h b/Engine/Source/Runtime/Engine/Public/ActorPartition/ActorPartitionSubsystem.h
--- a/Engine/Source/Runtime/Engine/Public/ActorPartition/ActorPartitionSubsystem.h
+++ b/Engine/Source/Runtime/Engine/Public/ActorPartition/ActorPartitionSubsystem.h
@@ -172,6 +172,12 @@ public:
 	ENGINE_API virtual void Deinitialize() override;
 
 	ENGINE_API void ForEachRelevantActor(const TSubclassOf<APartitionActor>& InActorClass, const FBox& IntersectionBounds, TFunctionRef<bool(APartitionActor*)>InOperation) const;
+
+	/**
+	 * Updates the partition actor data layers to match the provided ones.
+	 * The partition actor name and its external package will change accordingly.
+	 */
+	ENGINE_API bool MoveActorToDataLayers(APartitionActor* InActor, const TArray<UDataLayerInstance*>& InDataLayerInstances);
 #endif
 	ENGINE_API bool IsLevelPartition() const;
 
@@ -181,6 +187,8 @@ protected:
 private:
 
 #if WITH_EDITOR
+	APartitionActor* GetActor(const FActorPartitionIdentifier& InActorPartitionId, const FCellCoord& InCellCoords, bool bInCreate, uint32 InGridSize, bool bInBoundsSearch, TFunctionRef<void(APartitionActor*)> InActorCreated = [](APartitionActor*) {});
+
 	void OnActorPartitionHashInvalidated(const FCellCoord& Hash);
 	void OnWorldPartitionInitialized(UWorldPartition* InWorldPartition);
 
diff --git a/Engine/Source/Runtime/Engine/Public/ActorPartition/PartitionActor.h b/Engine/Source/Runtime/Engine/Public/ActorPartition/PartitionActor.h
--- a/Engine/Source/Runtime/Engine/Public/ActorPartition/PartitionActor.h
+++ b/Engine/Source/Runtime/Engine/Public/ActorPartition/PartitionActor.h
@@ -20,8 +20,8 @@ class APartitionActor: public AActor
 public:
 	ENGINE_API APartitionActor(const FObjectInitializer& ObjectInitializer = FObjectInitializer::Get());
 
-	//~ Begin AActor Interface
 #if WITH_EDITOR
+	//~ Begin AActor Interface
 	virtual bool CanChangeIsSpatiallyLoadedFlag() const override { return false; }
 	ENGINE_API virtual TUniquePtr<class FWorldPartitionActorDesc> CreateClassActorDesc() const override;
 	ENGINE_API virtual uint32 GetDefaultGridSize(UWorld* InWorld) const PURE_VIRTUAL(APartitionActor, return 0;)
@@ -29,9 +29,11 @@ public:
 	virtual bool ShouldIncludeGridSizeInLabel() const { return false; }
 	virtual FGuid GetGridGuid() const { return FGuid(); }
 	ENGINE_API virtual bool IsUserManaged() const override;
-#endif
 	//~ End AActor Interface
 
+	ENGINE_API virtual bool IsPartitionActorNameAffectedByDataLayers() const { return true; }
+#endif
+
 #if WITH_EDITORONLY_DATA
 	ENGINE_API uint32 GetGridSize() const;
 	ENGINE_API void SetGridSize(uint32 InGridSize);
diff --git a/Engine/Source/Runtime/Engine/Public/WorldPartition/DataLayer/DataLayerManager.h b/Engine/Source/Runtime/Engine/Public/WorldPartition/DataLayer/DataLayerManager.h
--- a/Engine/Source/Runtime/Engine/Public/WorldPartition/DataLayer/DataLayerManager.h
+++ b/Engine/Source/Runtime/Engine/Public/WorldPartition/DataLayer/DataLayerManager.h
@@ -241,6 +241,7 @@ private:
 	friend class UActorDescContainerInstance;
 	friend class AWorldSettings;
 	friend class AActor;
+	friend class FDataLayerEditorModule;
 #endif
 
 #if WITH_EDITORONLY_DATA
diff --git a/Engine/Source/Runtime/Engine/Public/WorldPartition/DataLayer/ExternalDataLayerHelper.h b/Engine/Source/Runtime/Engine/Public/WorldPartition/DataLayer/ExternalDataLayerHelper.h
--- a/Engine/Source/Runtime/Engine/Public/WorldPartition/DataLayer/ExternalDataLayerHelper.h
+++ b/Engine/Source/Runtime/Engine/Public/WorldPartition/DataLayer/ExternalDataLayerHelper.h
@@ -12,6 +12,22 @@ class UExternalDataLayerInstance;
 struct FExternalDataLayerUID;
 struct FAssetData;
 
+struct FMoveToExternalDataLayerParams
+{
+	FMoveToExternalDataLayerParams()
+		: ExternalDataLayerInstance (nullptr)
+		, bAllowNonUserManaged(false)
+	{}
+
+	FMoveToExternalDataLayerParams(const UExternalDataLayerInstance* InExternalDataLayerInstance, bool bInAllowNonUserManaged = false)
+		: ExternalDataLayerInstance(InExternalDataLayerInstance)
+		, bAllowNonUserManaged(bInAllowNonUserManaged)
+	{}
+
+	const UExternalDataLayerInstance* ExternalDataLayerInstance;
+	bool bAllowNonUserManaged;
+};
+
 class FExternalDataLayerHelper
 {
 public:
@@ -64,7 +80,7 @@ private:
 	 * Validates that all actors can change their External Data Layer to the new provided value (supports passing null) 
 	 * Returns false if any actor fails and fills OutFailureReason with the reason (if non-null). 
 	 */
-	ENGINE_API static bool CanMoveActorsToExternalDataLayer(const TArray<AActor*>& InActors, const UExternalDataLayerInstance* InExternalDataLayerInstance, FText* OutFailureReason = nullptr);
+	ENGINE_API static bool CanMoveActorsToExternalDataLayer(const TArray<AActor*>& InActors, const FMoveToExternalDataLayerParams& InParams, FText* OutFailureReason = nullptr);
 
 	/** 
 	 * Changes all actors External Data Layer to the new provided value (supports passing null)
@@ -72,7 +88,14 @@ private:
 	 */
 	ENGINE_API static bool MoveActorsToExternalDataLayer(const TArray<AActor*>& InActors, const UExternalDataLayerInstance* InExternalDataLayerInstance, FText* OutFailureReason = nullptr);
 
+	/**
+	 * Changes all actors External Data Layer to the new provided value (supports passing null)
+	 * Returns false if any actor fails and fills OutFailureReason with the reason (if non-null).
+	 */
+	ENGINE_API static bool MoveActorsToExternalDataLayer(const TArray<AActor*>& InActors, const FMoveToExternalDataLayerParams& InParams, FText* OutFailureReason = nullptr);
+
 	friend class FDataLayerEditorModule;
+	friend class UActorPartitionSubsystem;
 	friend class UDataLayerEditorSubsystem;
 	friend class UExternalDataLayerEngineSubsystem;
 	friend class UContentBundleEditingSubmodule;
diff --git a/Engine/Source/Runtime/Landscape/Classes/Landscape.h b/Engine/Source/Runtime/Landscape/Classes/Landscape.h
--- a/Engine/Source/Runtime/Landscape/Classes/Landscape.h
+++ b/Engine/Source/Runtime/Landscape/Classes/Landscape.h
@@ -277,6 +277,10 @@ public:
 
 	LANDSCAPE_API static void SplitHeightmap(ULandscapeComponent* Comp, ALandscapeProxy* TargetProxy = nullptr, class FMaterialUpdateContext* InOutUpdateContext = nullptr, TArray<class FComponentRecreateRenderStateContext>* InOutRecreateRenderStateContext = nullptr, bool InReregisterComponent = true);
 	
+	//~ Begin APartitionActor Interface
+	virtual bool IsPartitionActorNameAffectedByDataLayers() const { return false; }
+	//~ End APartitionActor Interface
+
 	//~ Begin UObject Interface.
 	virtual void PreEditChange(FProperty* PropertyThatWillChange) override;
 	virtual void PostEditChangeProperty(FPropertyChangedEvent& PropertyChangedEvent) override;
