Chaos:: Correct COM and new inertia computation with COM nudge.

#virtualized

diff --git a/Engine/Source/Runtime/Engine/Private/PhysicsEngine/BodyUtils.cpp b/Engine/Source/Runtime/Engine/Private/PhysicsEngine/BodyUtils.cpp
--- a/Engine/Source/Runtime/Engine/Private/PhysicsEngine/BodyUtils.cpp
+++ b/Engine/Source/Runtime/Engine/Private/PhysicsEngine/BodyUtils.cpp
@@ -5,6 +5,7 @@
 #include "PhysicalMaterials/PhysicalMaterial.h"
 
 #include "Chaos/MassProperties.h"
+#include "Chaos/Rotation.h"
 #include "Chaos/Utilities.h"
 #include "Physics/Experimental/ChaosInterfaceUtils.h"
 
@@ -73,7 +74,9 @@ namespace BodyUtils
 		float MassRatio = NewMass / OldMass;
 		MassProps.Mass *= MassRatio;
 		MassProps.InertiaTensor *= MassRatio;
-		MassProps.CenterOfMass += MassModifierTransform.TransformVector(OwningBodyInstance->COMNudge);
+		//COM Nudge is done in actor space. 
+		Chaos::FVec3 COMNudgeInActorSpace = MassModifierTransform.TransformVector(OwningBodyInstance->COMNudge);
+		MassProps.CenterOfMass += COMNudgeInActorSpace;
 
 		// Scale the inertia tensor by the owning body instance's InertiaTensorScale
 		// NOTE: PhysX scales the inertia by the mass increase we would get from the scale change, even though we 
@@ -86,7 +89,10 @@ namespace BodyUtils
 		//		bInertaScaleIncludeMass = false: more sensible behaviour given that InertiaScale does not affect mass
 		if (!(OwningBodyInstance->InertiaTensorScale - FVector::OneVector).IsNearlyZero(1e-3f))
 		{
+			MassProps.InertiaTensor = Chaos::Utilities::ComputeWorldSpaceInertia(MassProps.RotationOfMass, MassProps.InertiaTensor);
+			MassProps.RotationOfMass = Chaos::FRotation3::FromIdentity();
 			MassProps.InertiaTensor = Chaos::Utilities::ScaleInertia(MassProps.InertiaTensor, OwningBodyInstance->InertiaTensorScale, bInertaScaleIncludeMass);
+			TransformToLocalSpace(MassProps);
 		}
 
 		// If we move the center of mass, we need to update the inertia using parallel-axis theorem. If we don't do this
@@ -95,9 +101,14 @@ namespace BodyUtils
 		// which is not always possible, e.g., if we move the CoM outside of a box you will get negative elements in the scaled inertia!
 		if (CVars::bPhysicsComNudgeAdjustInertia)
 		{
-			MassProps.InertiaTensor.M[0][0] += MassProps.Mass * OwningBodyInstance->COMNudge.X * OwningBodyInstance->COMNudge.X;
-			MassProps.InertiaTensor.M[1][1] += MassProps.Mass * OwningBodyInstance->COMNudge.Y * OwningBodyInstance->COMNudge.Y;
-			MassProps.InertiaTensor.M[2][2] += MassProps.Mass * OwningBodyInstance->COMNudge.Z * OwningBodyInstance->COMNudge.Z;
+			MassProps.InertiaTensor = Chaos::Utilities::ComputeWorldSpaceInertia(MassProps.RotationOfMass, MassProps.InertiaTensor);
+			MassProps.RotationOfMass = Chaos::FRotation3::FromIdentity();
+
+			MassProps.InertiaTensor.M[0][0] += MassProps.Mass * COMNudgeInActorSpace.X * COMNudgeInActorSpace.X;
+			MassProps.InertiaTensor.M[1][1] += MassProps.Mass * COMNudgeInActorSpace.Y * COMNudgeInActorSpace.Y;
+			MassProps.InertiaTensor.M[2][2] += MassProps.Mass * COMNudgeInActorSpace.Z * COMNudgeInActorSpace.Z;
+
+			TransformToLocalSpace(MassProps);
 		}
 
 		return MassProps;
