Fix for audio streaming cache debug display if max count is larger than number of elements.
#rnx

diff --git a/Engine/Source/Runtime/Engine/Private/AudioStreamingCache.cpp b/Engine/Source/Runtime/Engine/Private/AudioStreamingCache.cpp
--- a/Engine/Source/Runtime/Engine/Private/AudioStreamingCache.cpp
+++ b/Engine/Source/Runtime/Engine/Private/AudioStreamingCache.cpp
@@ -2251,7 +2251,10 @@ TPair<int, int> FAudioChunkCache::DebugDisplayLegacy(UWorld* World, FViewport* V
 	};
 
 	// Trim to our display count
-	SortedCacheElements.SetNum(DebugMaxElementsDisplayCVar);
+	if (SortedCacheElements.Num() > DebugMaxElementsDisplayCVar)
+	{
+		SortedCacheElements.SetNum(DebugMaxElementsDisplayCVar);
+	}
 	
 	const int32 InitialX = X;
 	const int32 InitialY = Y;
