PR #12346: Add BlueprintType tag to FRuntimeCurveLinearColor

diff --git a/Engine/Source/Runtime/Engine/Classes/Curves/CurveLinearColor.h b/Engine/Source/Runtime/Engine/Classes/Curves/CurveLinearColor.h
--- a/Engine/Source/Runtime/Engine/Classes/Curves/CurveLinearColor.h
+++ b/Engine/Source/Runtime/Engine/Classes/Curves/CurveLinearColor.h
@@ -12,7 +12,7 @@ class UCurveLinearColor;
 // Delegate called right before rendering each slot. This allows the gradient to be modified dynamically per slot.
 DECLARE_MULTICAST_DELEGATE_OneParam(FOnUpdateGradient, UCurveLinearColor* /*GradientAtlas*/);
 
-USTRUCT()
+USTRUCT(BlueprintType)
 struct FRuntimeCurveLinearColor
 {
 	GENERATED_USTRUCT_BODY()
