Schedule FSkinWeightProfileManagerAsyncTask to run on Gamethread

#rnx

diff --git a/Engine/Source/Runtime/Engine/Private/Animation/SkinWeightProfileManager.cpp b/Engine/Source/Runtime/Engine/Private/Animation/SkinWeightProfileManager.cpp
--- a/Engine/Source/Runtime/Engine/Private/Animation/SkinWeightProfileManager.cpp
+++ b/Engine/Source/Runtime/Engine/Private/Animation/SkinWeightProfileManager.cpp
@@ -27,7 +27,20 @@ static FAutoConsoleVariableRef CHandleDelayedLoads (
 	ECVF_Cheat
 );
 
+static int32 GSkipHasPendingInitOrStreamingSkeletalMesh = 1;
+static FAutoConsoleVariableRef CVarSkipHasPendingInitOrStreamingSkeletalMesh (
+	TEXT("SkinWeightProfileManager.SkipHasPendingInitOrStreamingSkeletalMesh"),
+	GSkipHasPendingInitOrStreamingSkeletalMesh,
+	TEXT(""),
+	ECVF_Default | ECVF_RenderThreadSafe
+);
 
+static TAutoConsoleVariable<int32> CVarForceSkinWeightProfileManagerAsyncTaskOnGT(
+	TEXT("SkinWeightProfileManager.ForceSkinWeightProfileManagerAsyncTaskOnGT"),
+	1,
+	TEXT(""),
+	ECVF_Default | ECVF_RenderThreadSafe
+);
 
 FSkinWeightProfileManager::FSkinWeightProfileManager(UWorld* InWorld) : LastGamethreadProfileIndex(INDEX_NONE), WeakWorld(nullptr)
 {
@@ -337,6 +350,28 @@ FName FSkinWeightProfileManagerTickFunction::DiagnosticContext(bool bDetailed)
 	return FName("FSkinWeightProfileManager");
 }
 
+ENamedThreads::Type FSkinWeightProfileManagerAsyncTask::GetDesiredThread()
+{
+#if WITH_EDITOR
+	// Force task to run on GT since it can cause a stall waiting for RT
+	return ENamedThreads::GameThread;
+#else
+
+	if (CVarForceSkinWeightProfileManagerAsyncTaskOnGT.GetValueOnAnyThread())
+	{
+    	return ENamedThreads::GameThread;
+	}
+
+#if PLATFORM_ANDROID
+	    return IsRunningRHIInSeparateThread() ? ENamedThreads::RHIThread : ENamedThreads::ActualRenderingThread;
+#else
+		return ENamedThreads::AnyBackgroundHiPriTask;
+#endif
+
+#endif 		
+}
+
+
 void FSkinWeightProfileManagerAsyncTask::DoTask(ENamedThreads::Type CurrentThread, const FGraphEventRef& MyCompletionGraphEvent)
 {
 #if PROFILE_SKIN_WEIGHT_MANAGER 
@@ -354,6 +389,8 @@ void FSkinWeightProfileManagerAsyncTask::DoTask(ENamedThreads::Type CurrentThrea
 		const FSetProfileRequest& Request = OwnerRequests[RequestIndex];
 		if (const USkeletalMesh* SkeletalMesh = Request.WeakSkeletalMesh.Get())
 		{
+			if (GSkipHasPendingInitOrStreamingSkeletalMesh && SkeletalMesh->HasPendingInitOrStreaming()) continue;
+
 			if (FSkeletalMeshRenderData* RenderData = SkeletalMesh->GetResourceForRendering())
 			{
 				for (const int32 LODIndex : Request.LODIndices)
diff --git a/Engine/Source/Runtime/Engine/Public/Animation/SkinWeightProfileManager.h b/Engine/Source/Runtime/Engine/Public/Animation/SkinWeightProfileManager.h
--- a/Engine/Source/Runtime/Engine/Public/Animation/SkinWeightProfileManager.h
+++ b/Engine/Source/Runtime/Engine/Public/Animation/SkinWeightProfileManager.h
@@ -63,22 +63,7 @@ public:
 		RETURN_QUICK_DECLARE_CYCLE_STAT(FSkinWeightProfileManagerAsyncTask, STATGROUP_TaskGraphTasks);
 	}
 
-	ENamedThreads::Type GetDesiredThread()
-	{
-#if WITH_EDITOR
-		// Force task to run on GT since it can cause a stall waiting for RT
-		return ENamedThreads::GameThread;
-#else
-
-#if PLATFORM_ANDROID
-		return IsRunningRHIInSeparateThread() ? ENamedThreads::RHIThread : ENamedThreads::ActualRenderingThread;
-#else
-		return ENamedThreads::AnyBackgroundHiPriTask;
-#endif
-
-#endif 		
-	}
-
+	ENamedThreads::Type GetDesiredThread(); 
 	static ESubsequentsMode::Type GetSubsequentsMode() { return ESubsequentsMode::TrackSubsequents; }
 	void DoTask(ENamedThreads::Type CurrentThread, const FGraphEventRef& MyCompletionGraphEvent);
 };
