Chaos:: Fix particle jittering issue when involved in a free joint. Disable inertia conditioning on particle when it is involved in a free joint.

diff --git a/Engine/Source/Runtime/Engine/Private/PhysicsEngine/ImmediatePhysics/ImmediatePhysicsChaos/ImmediatePhysicsSimulation_Chaos.cpp b/Engine/Source/Runtime/Engine/Private/PhysicsEngine/ImmediatePhysics/ImmediatePhysicsChaos/ImmediatePhysicsSimulation_Chaos.cpp
--- a/Engine/Source/Runtime/Engine/Private/PhysicsEngine/ImmediatePhysics/ImmediatePhysicsChaos/ImmediatePhysicsSimulation_Chaos.cpp
+++ b/Engine/Source/Runtime/Engine/Private/PhysicsEngine/ImmediatePhysics/ImmediatePhysicsChaos/ImmediatePhysicsSimulation_Chaos.cpp
@@ -20,6 +20,7 @@
 #include "ChaosDebugDraw/ChaosDDScene.h"
 #include "ChaosDebugDraw/ChaosDDTimeline.h"
 #include "ChaosVisualDebugger/ChaosVisualDebuggerTrace.h"
+#include "Chaos/PBDJointConstraintUtilities.h"
 #include "Stats/StatsTrace.h"
 
 //UE_DISABLE_OPTIMIZATION
@@ -1087,17 +1088,20 @@ namespace ImmediatePhysics_Chaos
 						FVec3 ConstraintExtents = FVec3(0);
 						for (const FJointHandle* Joint : Implementation->JointHandles)
 						{
-							const Chaos::TVec2<const FActorHandle*>& JointActors = Joint->GetActorHandles();
-							FVec3 JointArm = FVec3(0);
-							if (JointActors[0] == Actor)
+							if (!FPBDJointUtilities::GetLinearIsFree(Joint->GetConstraint()->GetSettings()))
 							{
-								JointArm = Actor->GetLocalCoMTransform().InverseTransformPositionNoScale(Joint->GetConstraint()->GetSettings().ConnectorTransforms[0].GetTranslation());
+								const Chaos::TVec2<const FActorHandle*>& JointActors = Joint->GetActorHandles();
+								FVec3 JointArm = FVec3(0);
+								if (JointActors[0] == Actor)
+								{
+									JointArm = Actor->GetLocalCoMTransform().InverseTransformPositionNoScale(Joint->GetConstraint()->GetSettings().ConnectorTransforms[0].GetTranslation());
+								}
+								else if (JointActors[1] == Actor)
+								{
+									JointArm = Actor->GetLocalCoMTransform().InverseTransformPositionNoScale(Joint->GetConstraint()->GetSettings().ConnectorTransforms[1].GetTranslation());
+								}
+								ConstraintExtents = FVec3::Max(ConstraintExtents, JointArm.GetAbs());
 							}
-							else if (JointActors[1] == Actor)
-							{
-								JointArm = Actor->GetLocalCoMTransform().InverseTransformPositionNoScale(Joint->GetConstraint()->GetSettings().ConnectorTransforms[1].GetTranslation());
-							}
-							ConstraintExtents = FVec3::Max(ConstraintExtents, JointArm.GetAbs());
 						}
 
 						const FVec3f CoMExtents = FVec3f::Max(FVec3f(CollisionExtents), FVec3f(ConstraintExtents));
diff --git a/Engine/Source/Runtime/Experimental/Chaos/Private/Chaos/MassConditioning.cpp b/Engine/Source/Runtime/Experimental/Chaos/Private/Chaos/MassConditioning.cpp
--- a/Engine/Source/Runtime/Experimental/Chaos/Private/Chaos/MassConditioning.cpp
+++ b/Engine/Source/Runtime/Experimental/Chaos/Private/Chaos/MassConditioning.cpp
@@ -4,6 +4,7 @@
 #include "Chaos/ConstraintHandle.h"
 #include "Chaos/ParticleHandle.h"
 #include "Chaos/PBDJointConstraints.h"
+#include "Chaos/PBDJointConstraintUtilities.h"
 
 namespace Chaos
 {
@@ -94,11 +95,14 @@ namespace Chaos
 			FVec3 ConstraintArm = FVec3(0);
 			if (const FPBDJointConstraintHandle* Joint = Constraint->As<FPBDJointConstraintHandle>())
 			{
-				const TVector<FGeometryParticleHandle*, 2> JointParticles = Joint->GetConstrainedParticles();
-				check((Rigid == JointParticles[0]) || (Rigid == JointParticles[1]));
-				const int32 JointParticleIndex = (Rigid == JointParticles[0]) ? 0 : 1;
-				const FVec3& ActorJointArm = Joint->GetSettings().ConnectorTransforms[JointParticleIndex].GetTranslation();
-				ConstraintArm = CoMTransform.InverseTransformPosition(ActorJointArm);
+				if (!Chaos::FPBDJointUtilities::GetLinearIsFree(Joint->GetSettings()))
+				{
+					const TVector<FGeometryParticleHandle*, 2> JointParticles = Joint->GetConstrainedParticles();
+					check((Rigid == JointParticles[0]) || (Rigid == JointParticles[1]));
+					const int32 JointParticleIndex = (Rigid == JointParticles[0]) ? 0 : 1;
+					const FVec3& ActorJointArm = Joint->GetSettings().ConnectorTransforms[JointParticleIndex].GetTranslation();
+					ConstraintArm = CoMTransform.InverseTransformPosition(ActorJointArm);
+				}
 			}
 
 			ConstraintExtents = FVec3::Max(ConstraintExtents, ConstraintArm.GetAbs());
diff --git a/Engine/Source/Runtime/Experimental/Chaos/Private/Chaos/PBDJointConstraintUtilities.cpp b/Engine/Source/Runtime/Experimental/Chaos/Private/Chaos/PBDJointConstraintUtilities.cpp
--- a/Engine/Source/Runtime/Experimental/Chaos/Private/Chaos/PBDJointConstraintUtilities.cpp
+++ b/Engine/Source/Runtime/Experimental/Chaos/Private/Chaos/PBDJointConstraintUtilities.cpp
@@ -9,6 +9,9 @@
 
 namespace Chaos
 {
+	bool bExcludeFreeJointInertiaConditioning = true;
+	FAutoConsoleVariableRef CVarChaosExcludeFreeJointInertiaConditioning(TEXT("p.Chaos.ExcludeFreeJointForInertiaConditioning"), bExcludeFreeJointInertiaConditioning, TEXT("Exclude the effects from free joints when computing inertia conditioning"));
+
 	FReal GetSwingAngleY(const FRotation3& RSwing)
 	{
 		return 4.0f * FMath::Atan2(RSwing.Y, 1.0f + RSwing.W);
@@ -882,4 +885,10 @@ namespace Chaos
 			return CX;
 		}
 	}
+
+	bool FPBDJointUtilities::GetLinearIsFree(const FPBDJointSettings& JointSettings)
+	{
+		return JointSettings.LinearMotionTypes == TVector<EJointMotionType, 3>(EJointMotionType::Free)
+			&& JointSettings.bLinearPositionDriveEnabled == TVector<bool, 3>(false) && bExcludeFreeJointInertiaConditioning;
+	}
 }
diff --git a/Engine/Source/Runtime/Experimental/Chaos/Public/Chaos/PBDJointConstraintUtilities.h b/Engine/Source/Runtime/Experimental/Chaos/Public/Chaos/PBDJointConstraintUtilities.h
--- a/Engine/Source/Runtime/Experimental/Chaos/Public/Chaos/PBDJointConstraintUtilities.h
+++ b/Engine/Source/Runtime/Experimental/Chaos/Public/Chaos/PBDJointConstraintUtilities.h
@@ -287,6 +287,7 @@ namespace Chaos
 		static FVec3 GetCylinderLimitedPositionError(const FVec3& CX, const FVec3& Axis, const FReal Limit, const EJointMotionType AxisMotion);
 		static FVec3 GetLineLimitedPositionError(const FVec3& CX, const FVec3& Axis, const FReal Limit, const EJointMotionType AxisMotion);
 		static FVec3 GetLimitedPositionError(const FPBDJointSettings& JointSettings, const FRotation3& R0, const FVec3& CX);
+		static CHAOS_API bool GetLinearIsFree(const FPBDJointSettings& JointSettings);
 	};
 }
 
