[HWRT] Fix VB/IB of ray tracing proxies not being streamed out.
- Modified FRayTracingGeometryManager to call StreamableAsset->ReleaseForStreaming(...) instead of Geometry->ReleaseRHIForStreaming(...) when appropriate.
- Modified FRayTracingGeometry::Evict() check that is has been streamed out before evicting instead of patching its own state during eviction (helps detect code paths that are not properly handling streaming).
- Modified FRayTracingGeometryManager::EvictGeometry(...) to stream out the geometry before evicting it.

diff --git a/Engine/Source/Runtime/Engine/Private/Rendering/RayTracingGeometryManager.cpp b/Engine/Source/Runtime/Engine/Private/Rendering/RayTracingGeometryManager.cpp
--- a/Engine/Source/Runtime/Engine/Private/Rendering/RayTracingGeometryManager.cpp
+++ b/Engine/Source/Runtime/Engine/Private/Rendering/RayTracingGeometryManager.cpp
@@ -387,8 +387,7 @@ void FRayTracingGeometryManager::SetRayTracingGeometryGroupCurrentFirstLODIndex(
 				// Cancel associated streaming request if currently in-flight
 				CancelStreamingRequest(RegisteredGeometry);
 
-				// TODO: call FRayTracingStreamableAsset::ReleaseForStreaming(...) instead
-				RegisteredGeometry.Geometry->ReleaseRHIForStreaming(Batcher);
+				StreamOutGeometry(Batcher, RegisteredGeometry);
 			}
 		}
 	}
@@ -557,7 +556,7 @@ void FRayTracingGeometryManager::Tick(FRHICommandList& RHICmdList)
 			}
 			else if (IsRayTracingUsingReferenceBasedResidency() && RegisteredGeometry.Geometry->GetRHI() != nullptr)
 			{
-				EvictGeometry(RegisteredGeometry);
+				EvictGeometry(RHICmdList, RegisteredGeometry);
 			}
 		}
 	}
@@ -571,7 +570,7 @@ void FRayTracingGeometryManager::Tick(FRHICommandList& RHICmdList)
 			{
 				if (RegisteredGeometry.Geometry->GetRHI() != nullptr)
 				{
-					EvictGeometry(RegisteredGeometry);
+					EvictGeometry(RHICmdList, RegisteredGeometry);
 				}
 			}
 
@@ -747,7 +746,7 @@ void FRayTracingGeometryManager::Tick(FRHICommandList& RHICmdList)
 
 				check(RegisteredGeometry.Geometry->IsValid() && !RegisteredGeometry.Geometry->IsEvicted());
 
-				EvictGeometry(RegisteredGeometry);
+				EvictGeometry(RHICmdList, RegisteredGeometry);
 
 				++Index;
 			}
@@ -1015,11 +1014,34 @@ void FRayTracingGeometryManager::CancelStreamingRequest(FRegisteredGeometry& Reg
 	}
 }
 
-void FRayTracingGeometryManager::EvictGeometry(FRegisteredGeometry& RegisteredGeometry)
+void FRayTracingGeometryManager::StreamOutGeometry(FRHIResourceReplaceBatcher& Batcher, FRegisteredGeometry& RegisteredGeometry)
+{
+	if (EnumHasAllFlags(RegisteredGeometry.Geometry->GetGeometryState(), FRayTracingGeometry::EGeometryStateFlags::StreamedIn))
+	{
+		if (RegisteredGeometry.StreamableAsset != nullptr)
+		{
+			RegisteredGeometry.StreamableAsset->ReleaseForStreaming(Batcher);
+		}
+		else
+		{
+			RegisteredGeometry.Geometry->ReleaseRHIForStreaming(Batcher);
+		}
+	}
+}
+
+void FRayTracingGeometryManager::EvictGeometry(FRHICommandListBase& RHICmdList, FRegisteredGeometry& RegisteredGeometry)
 {
 	// Cancel associated streaming request if currently in-flight
 	CancelStreamingRequest(RegisteredGeometry);
 
+	// Both FRayTracingGeometry::ReleaseRHIForStreaming(...) and FRayTracingGeometry::Evict()
+	// call FRayTracingGeometryManager:::RefreshRegisteredGeometry(...) which is unecessary
+	// however there's no straightforward way to avoid that
+	// TODO: investigate possible improvements
+
+	FRHIResourceReplaceBatcher Batcher(RHICmdList, 1);
+	StreamOutGeometry(Batcher, RegisteredGeometry);
+
 	RegisteredGeometry.Geometry->Evict();
 }
 
diff --git a/Engine/Source/Runtime/Engine/Public/Rendering/RayTracingGeometryManager.h b/Engine/Source/Runtime/Engine/Public/Rendering/RayTracingGeometryManager.h
--- a/Engine/Source/Runtime/Engine/Public/Rendering/RayTracingGeometryManager.h
+++ b/Engine/Source/Runtime/Engine/Public/Rendering/RayTracingGeometryManager.h
@@ -159,9 +159,13 @@ private:
 
 	void CancelStreamingRequest(FRegisteredGeometry& RegisteredGeometry);
 
+	// Helper function to stream out geometry
+	// Must call this instead of FRayTracingGeometry::ReleaseRHIForStreaming() directly
+	void StreamOutGeometry(FRHIResourceReplaceBatcher& Batcher, FRegisteredGeometry& RegisteredGeometry);
+
 	// Helper function to evict geometry + necessary streaming management
 	// Must call this instead of FRayTracingGeometry::Evict() directly
-	void EvictGeometry(FRegisteredGeometry& RegisteredGeometry);
+	void EvictGeometry(FRHICommandListBase& RHICmdList, FRegisteredGeometry& RegisteredGeometry);
 
 	TSparseArray<FRayTracingGeometryGroup> RegisteredGroups;
 
diff --git a/Engine/Source/Runtime/RenderCore/Private/RayTracingGeometry.cpp b/Engine/Source/Runtime/RenderCore/Private/RayTracingGeometry.cpp
--- a/Engine/Source/Runtime/RenderCore/Private/RayTracingGeometry.cpp
+++ b/Engine/Source/Runtime/RenderCore/Private/RayTracingGeometry.cpp
@@ -101,18 +101,11 @@ void FRayTracingGeometry::MakeResident(FRHICommandList& RHICmdList)
 void FRayTracingGeometry::Evict()
 {
 	check(!EnumHasAllFlags(GeometryState, EGeometryStateFlags::Evicted) && RayTracingGeometryRHI != nullptr);
+	checkf(!EnumHasAllFlags(GeometryState, EGeometryStateFlags::StreamedIn), TEXT("Ray tracing geometry must be streamed out before it can be evicted."));
 
 	RemoveBuildRequest();
 	RayTracingGeometryRHI.SafeRelease();
 	EnumAddFlags(GeometryState, EGeometryStateFlags::Evicted);
-
-	if (EnumHasAllFlags(GeometryState, EGeometryStateFlags::StreamedIn))
-	{
-		EnumRemoveFlags(GeometryState, EGeometryStateFlags::StreamedIn);
-		EnumRemoveFlags(GeometryState, EGeometryStateFlags::Valid);
-
-		Initializer.Type = ERayTracingGeometryInitializerType::StreamingDestination;
-	}
 	
 	GRayTracingGeometryManager->RefreshRegisteredGeometry(RayTracingGeometryHandle);
 	
