[AutoRTFM] Fix assertions running AutoRTFMEngineTests in debug

UActorComponent::OwnerPrivate can checked after a transaction that constructed it is aborted.
UActorChannel::OpenChannels is expected to hold the ActorChannel on GC.

diff --git a/Engine/Source/Runtime/Engine/Private/Components/ActorComponent.cpp b/Engine/Source/Runtime/Engine/Private/Components/ActorComponent.cpp
--- a/Engine/Source/Runtime/Engine/Private/Components/ActorComponent.cpp
+++ b/Engine/Source/Runtime/Engine/Private/Components/ActorComponent.cpp
@@ -543,7 +543,11 @@ UActorComponent::UActorComponent(const FObjectInitializer& ObjectInitializer /*=
 	, PhysicsStateAsyncCreationState(EPhysicsStateAsyncCreationState::None)
 	, PhysicsStateAsyncCreationScene(nullptr)
 {
-	OwnerPrivate = GetTypedOuter<AActor>();
+	// OwnerPrivate must be valid until destruction, even after aborting.
+	UE_AUTORTFM_OPEN_NO_VALIDATION
+	{
+		OwnerPrivate = GetTypedOuter<AActor>();
+	};
 
 	PrimaryComponentTick.TickGroup = TG_DuringPhysics;
 	PrimaryComponentTick.bStartWithTickEnabled = true;
diff --git a/Engine/Source/Runtime/Engine/Private/Tests/AutoRTFM/AutoRTFMNetDriverTest.cpp b/Engine/Source/Runtime/Engine/Private/Tests/AutoRTFM/AutoRTFMNetDriverTest.cpp
--- a/Engine/Source/Runtime/Engine/Private/Tests/AutoRTFM/AutoRTFMNetDriverTest.cpp
+++ b/Engine/Source/Runtime/Engine/Private/Tests/AutoRTFM/AutoRTFMNetDriverTest.cpp
@@ -27,6 +27,7 @@ bool FAutoRTFMNetDriverTests::RunTest(const FString & Parameters)
 	ActorChannel->OpenedLocally = true;
 	ActorChannel->Connection = Connection;
 	Connection->Channels.Add(ActorChannel);
+	Connection->OpenChannels.Add(ActorChannel);
 
 	FString Description;
 
