[GPUSkinCache] Minor code clean up.

- Move some functions from FGPUSkinCache::FRWBuffersAllocation into .cpp file for faster compilation, especially when selectively disabling optimziation during debugging.
- Rename wrong parameter name NumTriangles to InIntermediateAccumulatedTangentsSize.
- Exclude alternative branch for ispc::UpdateBoneData_CopyBones() when INTEL_ISPC is not enabled, regardless of "bGPUSkin_CopyBones_ISPC_Enabled" configuration.

#rnx

diff --git a/Engine/Source/Runtime/Engine/Private/GPUSkinCache.cpp b/Engine/Source/Runtime/Engine/Private/GPUSkinCache.cpp
--- a/Engine/Source/Runtime/Engine/Private/GPUSkinCache.cpp
+++ b/Engine/Source/Runtime/Engine/Private/GPUSkinCache.cpp
@@ -1242,10 +1242,10 @@ void FGPUSkinCache::DispatchUpdateSkinTangents(FRHICommandList& RHICmdList, FGPU
 	}
 }
 
-FGPUSkinCache::FRWBuffersAllocation* FGPUSkinCache::TryAllocBuffer(uint32 NumVertices, bool WithTangnents, bool UseIntermediateTangents, uint32 NumTriangles, FRHICommandList& RHICmdList, const FName& OwnerName)
+FGPUSkinCache::FRWBuffersAllocation* FGPUSkinCache::TryAllocBuffer(uint32 NumVertices, bool WithTangents, bool UseIntermediateTangents, uint32 InIntermediateAccumulatedTangentsSize, FRHICommandList& RHICmdList, const FName& OwnerName)
 {
 	uint64 MaxSizeInBytes = (uint64)(GSkinCacheSceneMemoryLimitInMB * 1024.0f * 1024.0f);
-	uint64 RequiredMemInBytes = FRWBuffersAllocation::CalculateRequiredMemory(NumVertices, WithTangnents, UseIntermediateTangents, NumTriangles);
+	uint64 RequiredMemInBytes = FRWBuffersAllocation::CalculateRequiredMemory(NumVertices, WithTangents, UseIntermediateTangents, InIntermediateAccumulatedTangentsSize);
 	if (bRequiresMemoryLimit && UsedMemoryInBytes + RequiredMemInBytes >= MaxSizeInBytes)
 	{
 		ExtraRequiredMemory += RequiredMemInBytes;
@@ -1254,7 +1254,7 @@ FGPUSkinCache::FRWBuffersAllocation* FGPUSkinCache::TryAllocBuffer(uint32 NumVer
 		return nullptr;
 	}
 
-	FRWBuffersAllocation* NewAllocation = new FRWBuffersAllocation(NumVertices, WithTangnents, UseIntermediateTangents, NumTriangles, RHICmdList, OwnerName);
+	FRWBuffersAllocation* NewAllocation = new FRWBuffersAllocation(NumVertices, WithTangents, UseIntermediateTangents, InIntermediateAccumulatedTangentsSize, RHICmdList, OwnerName);
 	Allocations.Add(NewAllocation);
 
 	UsedMemoryInBytes += RequiredMemInBytes;
@@ -2490,6 +2490,83 @@ void FGPUSkinCache::DrawVisualizationInfoText(const FName& GPUSkinCacheVisualiza
 	}
 }
 
+FGPUSkinCache::FRWBuffersAllocation::FRWBuffersAllocation(uint32 InNumVertices, bool InWithTangents, bool InUseIntermediateTangents, uint32 InIntermediateAccumulatedTangentsSize, FRHICommandList& RHICmdList, const FName& OwnerName)
+	: NumVertices(InNumVertices), WithTangents(InWithTangents), UseIntermediateTangents(InUseIntermediateTangents), IntermediateAccumulatedTangentsSize(InIntermediateAccumulatedTangentsSize)
+{
+	const static FLazyName PositionsName(TEXT("SkinCachePositions"));
+	const static FLazyName TangentsName(TEXT("SkinCacheTangents"));
+	const static FLazyName IntermediateTangentsName(TEXT("SkinCacheIntermediateTangents"));
+	const static FLazyName IntermediateAccumulatedTangentsName(TEXT("SkinCacheIntermediateAccumulatedTangents"));
+
+	for (int32 Index = 0; Index < NUM_BUFFERS; ++Index)
+	{
+		PositionBuffers[Index].Buffer.ClassName = PositionsName;
+		PositionBuffers[Index].Buffer.OwnerName = OwnerName;
+		PositionBuffers[Index].Buffer.Initialize(RHICmdList, TEXT("SkinCachePositions"), PosBufferBytesPerElement, NumVertices * 3, PF_R32_FLOAT, ERHIAccess::SRVMask, BUF_Static);
+		PositionBuffers[Index].Buffer.Buffer->SetOwnerName(OwnerName);
+		PositionBuffers[Index].AccessState = ERHIAccess::Unknown;
+	}
+	if (WithTangents)
+	{
+		// OpenGL ES does not support writing to RGBA16_SNORM images, instead pack data into SINT in the shader
+		const EPixelFormat TangentsFormat = IsOpenGLPlatform(GMaxRHIShaderPlatform) ? PF_R16G16B16A16_SINT : PF_R16G16B16A16_SNORM;
+
+		Tangents.Buffer.ClassName = TangentsName;
+		Tangents.Buffer.OwnerName = OwnerName;
+		Tangents.Buffer.Initialize(RHICmdList, TEXT("SkinCacheTangents"), TangentBufferBytesPerElement, NumVertices * 2, TangentsFormat, BUF_Static);
+		Tangents.Buffer.Buffer->SetOwnerName(OwnerName);
+		Tangents.AccessState = ERHIAccess::Unknown;
+		if (UseIntermediateTangents)
+		{
+			IntermediateTangents.Buffer.ClassName = IntermediateTangentsName;
+			IntermediateTangents.Buffer.OwnerName = OwnerName;
+			IntermediateTangents.Buffer.Initialize(RHICmdList, TEXT("SkinCacheIntermediateTangents"), TangentBufferBytesPerElement, NumVertices * 2, TangentsFormat, BUF_Static);
+			IntermediateTangents.Buffer.Buffer->SetOwnerName(OwnerName);
+			IntermediateTangents.AccessState = ERHIAccess::Unknown;
+		}
+	}
+	if (IntermediateAccumulatedTangentsSize > 0)
+	{
+		IntermediateAccumulatedTangents.Buffer.ClassName = IntermediateAccumulatedTangentsName;
+		IntermediateAccumulatedTangents.Buffer.OwnerName = OwnerName;
+		IntermediateAccumulatedTangents.Buffer.Initialize(RHICmdList, TEXT("SkinCacheIntermediateAccumulatedTangents"), sizeof(int32), IntermediateAccumulatedTangentsSize * FGPUSkinCache::IntermediateAccumBufferNumInts, PF_R32_SINT, BUF_UnorderedAccess);
+		IntermediateAccumulatedTangents.Buffer.Buffer->SetOwnerName(OwnerName);
+		IntermediateAccumulatedTangents.AccessState = ERHIAccess::Unknown;
+		// The UAV must be zero-filled. We leave it zeroed after each round (see RecomputeTangentsPerVertexPass.usf), so this is only needed on when the buffer is first created.
+		RHICmdList.ClearUAVUint(IntermediateAccumulatedTangents.Buffer.UAV, FUintVector4(0, 0, 0, 0));
+	}
+}
+
+FGPUSkinCache::FRWBuffersAllocation::~FRWBuffersAllocation()
+{
+	for (int32 Index = 0; Index < NUM_BUFFERS; ++Index)
+	{
+		PositionBuffers[Index].Release();
+	}
+	if (WithTangents)
+	{
+		Tangents.Release();
+		IntermediateTangents.Release();
+	}
+	if (IntermediateAccumulatedTangentsSize > 0)
+	{
+		IntermediateAccumulatedTangents.Release();
+	}
+}
+
+uint64 FGPUSkinCache::FRWBuffersAllocation::CalculateRequiredMemory(uint32 InNumVertices, bool InWithTangents, bool InUseIntermediateTangents, uint32 InIntermediateAccumulatedTangentsSize)
+{
+	uint64 PositionBufferSize = PosBufferBytesPerElement * InNumVertices * 3 * NUM_BUFFERS;
+	uint64 TangentBufferSize = InWithTangents ? TangentBufferBytesPerElement * InNumVertices * 2 : 0;
+	uint64 IntermediateTangentBufferSize = 0;
+	if (InUseIntermediateTangents)
+	{
+		IntermediateTangentBufferSize = InWithTangents ? TangentBufferBytesPerElement * InNumVertices * 2 : 0;
+	}
+	uint64 AccumulatedTangentBufferSize = InIntermediateAccumulatedTangentsSize * FGPUSkinCache::IntermediateAccumBufferNumInts * sizeof(int32);
+	return TangentBufferSize + IntermediateTangentBufferSize + PositionBufferSize + AccumulatedTangentBufferSize;
+}
+
 #undef IMPLEMENT_SKIN_CACHE_SHADER_CLOTH
 #undef IMPLEMENT_SKIN_CACHE_SHADER_ALL_SKIN_TYPES
 #undef IMPLEMENT_SKIN_CACHE_SHADER
\ No newline at end of file
diff --git a/Engine/Source/Runtime/Engine/Public/GPUSkinCache.h b/Engine/Source/Runtime/Engine/Public/GPUSkinCache.h
--- a/Engine/Source/Runtime/Engine/Public/GPUSkinCache.h
+++ b/Engine/Source/Runtime/Engine/Public/GPUSkinCache.h
@@ -195,82 +195,10 @@ public:
 	{
 		friend struct FRWBufferTracker;
 
-		FRWBuffersAllocation(uint32 InNumVertices, bool InWithTangents, bool InUseIntermediateTangents, uint32 InIntermediateAccumulatedTangentsSize, FRHICommandList& RHICmdList, const FName& OwnerName)
-			: NumVertices(InNumVertices), WithTangents(InWithTangents), UseIntermediateTangents(InUseIntermediateTangents), IntermediateAccumulatedTangentsSize(InIntermediateAccumulatedTangentsSize)
-		{
-			const static FLazyName PositionsName(TEXT("SkinCachePositions"));
-			const static FLazyName TangentsName(TEXT("SkinCacheTangents"));
-			const static FLazyName IntermediateTangentsName(TEXT("SkinCacheIntermediateTangents"));
-			const static FLazyName IntermediateAccumulatedTangentsName(TEXT("SkinCacheIntermediateAccumulatedTangents"));
-
-			for (int32 Index = 0; Index < NUM_BUFFERS; ++Index)
-			{
-				PositionBuffers[Index].Buffer.ClassName = PositionsName;
-				PositionBuffers[Index].Buffer.OwnerName = OwnerName;
-				PositionBuffers[Index].Buffer.Initialize(RHICmdList, TEXT("SkinCachePositions"), PosBufferBytesPerElement, NumVertices * 3, PF_R32_FLOAT, ERHIAccess::SRVMask, BUF_Static);
-				PositionBuffers[Index].Buffer.Buffer->SetOwnerName(OwnerName);
-				PositionBuffers[Index].AccessState = ERHIAccess::Unknown;
-			}
-			if (WithTangents)
-			{
-				// OpenGL ES does not support writing to RGBA16_SNORM images, instead pack data into SINT in the shader
-				const EPixelFormat TangentsFormat = IsOpenGLPlatform(GMaxRHIShaderPlatform) ? PF_R16G16B16A16_SINT : PF_R16G16B16A16_SNORM;
-				
-				Tangents.Buffer.ClassName = TangentsName;
-				Tangents.Buffer.OwnerName = OwnerName;
-				Tangents.Buffer.Initialize(RHICmdList, TEXT("SkinCacheTangents"), TangentBufferBytesPerElement, NumVertices * 2, TangentsFormat, BUF_Static);
-				Tangents.Buffer.Buffer->SetOwnerName(OwnerName);
-				Tangents.AccessState = ERHIAccess::Unknown;
-				if (UseIntermediateTangents)
-				{
-					IntermediateTangents.Buffer.ClassName = IntermediateTangentsName;
-					IntermediateTangents.Buffer.OwnerName = OwnerName;
-					IntermediateTangents.Buffer.Initialize(RHICmdList, TEXT("SkinCacheIntermediateTangents"), TangentBufferBytesPerElement, NumVertices * 2, TangentsFormat, BUF_Static);
-					IntermediateTangents.Buffer.Buffer->SetOwnerName(OwnerName);
-					IntermediateTangents.AccessState = ERHIAccess::Unknown;
-				}
-			}
-			if (IntermediateAccumulatedTangentsSize > 0)
-			{
-				IntermediateAccumulatedTangents.Buffer.ClassName = IntermediateAccumulatedTangentsName;
-				IntermediateAccumulatedTangents.Buffer.OwnerName = OwnerName;
-				IntermediateAccumulatedTangents.Buffer.Initialize(RHICmdList, TEXT("SkinCacheIntermediateAccumulatedTangents"), sizeof(int32), IntermediateAccumulatedTangentsSize * FGPUSkinCache::IntermediateAccumBufferNumInts, PF_R32_SINT, BUF_UnorderedAccess);
-				IntermediateAccumulatedTangents.Buffer.Buffer->SetOwnerName(OwnerName);
-				IntermediateAccumulatedTangents.AccessState = ERHIAccess::Unknown;
-				// The UAV must be zero-filled. We leave it zeroed after each round (see RecomputeTangentsPerVertexPass.usf), so this is only needed on when the buffer is first created.
-				RHICmdList.ClearUAVUint(IntermediateAccumulatedTangents.Buffer.UAV, FUintVector4(0, 0, 0, 0));
-			}
-		}
-
-		~FRWBuffersAllocation()
-		{
-			for (int32 Index = 0; Index < NUM_BUFFERS; ++Index)
-			{
-				PositionBuffers[Index].Release();
-			}
-			if (WithTangents)
-			{
-				Tangents.Release();
-				IntermediateTangents.Release();
-			}
-			if (IntermediateAccumulatedTangentsSize > 0)
-			{
-				IntermediateAccumulatedTangents.Release();
-			}
-		}
+		FRWBuffersAllocation(uint32 InNumVertices, bool InWithTangents, bool InUseIntermediateTangents, uint32 InIntermediateAccumulatedTangentsSize, FRHICommandList& RHICmdList, const FName& OwnerName);
+		~FRWBuffersAllocation();
 
-		static uint64 CalculateRequiredMemory(uint32 InNumVertices, bool InWithTangents, bool InUseIntermediateTangents, uint32 InIntermediateAccumulatedTangentsSize)
-		{
-			uint64 PositionBufferSize = PosBufferBytesPerElement * InNumVertices * 3 * NUM_BUFFERS;
-			uint64 TangentBufferSize = InWithTangents ? TangentBufferBytesPerElement * InNumVertices * 2 : 0;
-			uint64 IntermediateTangentBufferSize = 0;
-			if (InUseIntermediateTangents)
-			{
-				IntermediateTangentBufferSize = InWithTangents ? TangentBufferBytesPerElement * InNumVertices * 2 : 0;
-			}
-			uint64 AccumulatedTangentBufferSize = InIntermediateAccumulatedTangentsSize * FGPUSkinCache::IntermediateAccumBufferNumInts * sizeof(int32);
-			return TangentBufferSize + IntermediateTangentBufferSize + PositionBufferSize + AccumulatedTangentBufferSize;
-		}
+		static uint64 CalculateRequiredMemory(uint32 InNumVertices, bool InWithTangents, bool InUseIntermediateTangents, uint32 InIntermediateAccumulatedTangentsSize);
 
 		uint64 GetNumBytes() const
 		{
@@ -427,7 +355,7 @@ protected:
 	TSet<FGPUSkinCacheEntry*> PendingProcessRTGeometryEntries;
 	TArray<FDispatchEntry> BatchDispatches;
 
-	FRWBuffersAllocation* TryAllocBuffer(uint32 NumVertices, bool WithTangnents, bool UseIntermediateTangents, uint32 NumTriangles, FRHICommandList& RHICmdList, const FName& OwnerName);
+	FRWBuffersAllocation* TryAllocBuffer(uint32 NumVertices, bool WithTangnents, bool UseIntermediateTangents, uint32 InIntermediateAccumulatedTangentsSize, FRHICommandList& RHICmdList, const FName& OwnerName);
 	void DispatchUpdateSkinTangents(FRHICommandList& RHICmdList, FGPUSkinCacheEntry* Entry, int32 SectionIndex, FSkinCacheRWBuffer*& StagingBuffer, bool bTrianglePass);
 
 	void PrepareUpdateSkinning(
