Add MaterialSlot details export options to UStaticMesh to assist in asset reflection. Includes an asset registry data export step and a slower runtime check that counterbalances the StaticMaterials array async updates.

diff --git a/Engine/Source/Runtime/Engine/Classes/Engine/StaticMesh.h b/Engine/Source/Runtime/Engine/Classes/Engine/StaticMesh.h
--- a/Engine/Source/Runtime/Engine/Classes/Engine/StaticMesh.h
+++ b/Engine/Source/Runtime/Engine/Classes/Engine/StaticMesh.h
@@ -515,6 +515,18 @@ struct FStaticMaterial
 	FMeshUVChannelInfo			UVChannelData;
 };
 
+USTRUCT()
+struct FStaticMaterialMinimalInfo
+{
+	GENERATED_USTRUCT_BODY()
+
+	FStaticMaterialMinimalInfo(FName InMaterialSlotName = NAME_None, FString InMaterialObjectPath = FString())
+		: MaterialSlotName(InMaterialSlotName), MaterialObjectPath(InMaterialObjectPath)
+		{}
+
+	FName MaterialSlotName;
+	FString MaterialObjectPath;
+};
 
 enum EImportStaticMeshVersion
 {
@@ -937,19 +949,32 @@ private:
 	UPROPERTY(BlueprintGetter = GetStaticMaterials, BlueprintSetter = SetStaticMaterials, Category = StaticMesh)
 	TArray<FStaticMaterial> StaticMaterials;
 
+	TArray<FStaticMaterialMinimalInfo> StaticMaterialsInfoCache;
+	mutable FCriticalSection MaterialsInfoCriticalSection;
+	ENGINE_API void UpdateStaticMaterialsInfoCache(const TArray<FStaticMaterial>& InMaterials);
+	ENGINE_API const FString ExportStaticMaterialsInfo() const;
+
 #if WITH_EDITORONLY_DATA
 	UPROPERTY(Transient)
 	TArray<TObjectPtr<UStaticMesh>> AssemblyReferenceCache;
 #endif
 
 public:
-	static FName GetStaticMaterialsName()
+	ENGINE_API static FName GetStaticMaterialsName()
 	{
 		PRAGMA_DISABLE_DEPRECATION_WARNINGS
 		return GET_MEMBER_NAME_CHECKED(UStaticMesh, StaticMaterials);
 		PRAGMA_ENABLE_DEPRECATION_WARNINGS
 	}
 
+	ENGINE_API const TArray<FStaticMaterialMinimalInfo> GetCopyOfStaticMaterialsInfo() const
+	{
+		FScopeLock Lock(&MaterialsInfoCriticalSection);
+		return StaticMaterialsInfoCache;
+	}
+
+	ENGINE_API static const TArray<FStaticMaterialMinimalInfo> ImportStaticMaterialsInfo(const FString& InData);
+
 	TArray<FStaticMaterial>& GetStaticMaterials()
 	{
 		WaitUntilAsyncPropertyReleased(EStaticMeshAsyncProperties::StaticMaterials);
@@ -973,6 +998,7 @@ public:
 		WaitUntilAsyncPropertyReleased(EStaticMeshAsyncProperties::StaticMaterials);
 		PRAGMA_DISABLE_DEPRECATION_WARNINGS
 		StaticMaterials = InStaticMaterials;
+		UpdateStaticMaterialsInfoCache(StaticMaterials);
 		PRAGMA_ENABLE_DEPRECATION_WARNINGS
 	}
 
diff --git a/Engine/Source/Runtime/Engine/Private/StaticMesh.cpp b/Engine/Source/Runtime/Engine/Private/StaticMesh.cpp
--- a/Engine/Source/Runtime/Engine/Private/StaticMesh.cpp
+++ b/Engine/Source/Runtime/Engine/Private/StaticMesh.cpp
@@ -5685,6 +5685,7 @@ void UStaticMesh::GetAssetRegistryTags(FAssetRegistryTagsContext Context) const
 	Context.AddTag(FAssetRegistryTag("Vertices", FString::FromInt(NumVertices), FAssetRegistryTag::TT_Numerical) );
 	Context.AddTag(FAssetRegistryTag("UVChannels", FString::FromInt(NumUVChannels), FAssetRegistryTag::TT_Numerical) );
 	Context.AddTag(FAssetRegistryTag("Materials", FString::FromInt(GetStaticMaterials().Num()), FAssetRegistryTag::TT_Numerical) );
+	Context.AddTag(FAssetRegistryTag(GetStaticMaterialsName(), ExportStaticMaterialsInfo(), FAssetRegistryTag::TT_Numerical));
 	Context.AddTag(FAssetRegistryTag("ApproxSize", ApproxSizeStr, FAssetRegistryTag::TT_Dimensional) );
 	Context.AddTag(FAssetRegistryTag("CollisionPrims", FString::FromInt(NumCollisionPrims), FAssetRegistryTag::TT_Numerical));
 	Context.AddTag(FAssetRegistryTag("LODs", FString::FromInt(NumLODs), FAssetRegistryTag::TT_Numerical));
@@ -6715,13 +6716,16 @@ void UStaticMesh::Serialize(FArchive& Ar)
 	if (Ar.CustomVer(FEditorObjectVersion::GUID) >= FEditorObjectVersion::RefactorMeshEditorMaterials)
 #endif // WITH_EDITORONLY_DATA
 	{
-		Ar << GetStaticMaterials();
+		TArray<FStaticMaterial>& StaticMaterialsArray = GetStaticMaterials();
+		Ar << StaticMaterialsArray;
+		UpdateStaticMaterialsInfoCache(StaticMaterialsArray);
 	}
 #if WITH_EDITORONLY_DATA
 	else if (Ar.IsLoading())
 	{
 		TArray<UMaterialInterface*> Unique_Materials_DEPRECATED;
 		TArray<FName> MaterialSlotNames;
+		TArray<FStaticMaterial>& StaticMaterialsArray = GetStaticMaterials();
 		for (UMaterialInterface *MaterialInterface : Materials_DEPRECATED)
 		{
 			FName MaterialSlotName = MaterialInterface != nullptr ? MaterialInterface->GetFName() : NAME_None;
@@ -6736,13 +6740,14 @@ void UStaticMesh::Serialize(FArchive& Ar)
 				}
 			}
 			MaterialSlotNames.Add(MaterialSlotName);
-			GetStaticMaterials().Add(FStaticMaterial(MaterialInterface, MaterialSlotName));
+			StaticMaterialsArray.Add(FStaticMaterial(MaterialInterface, MaterialSlotName));
 			int32 UniqueIndex = Unique_Materials_DEPRECATED.AddUnique(MaterialInterface);
 #if WITH_EDITOR
 			//We must cleanup the material list since we have a new way to build static mesh
 			bCleanUpRedundantMaterialPostLoad = GetStaticMaterials().Num() > 1;
 #endif
 		}
+		UpdateStaticMaterialsInfoCache(StaticMaterialsArray);
 		Materials_DEPRECATED.Empty();
 
 	}
@@ -9586,6 +9591,54 @@ const TArray<TObjectPtr<UStaticMesh>>& UStaticMesh::GetCachedNaniteAssemblyRefer
 
 #endif
 
+void UStaticMesh::UpdateStaticMaterialsInfoCache(const TArray<FStaticMaterial>& InMaterials)
+{
+	FScopeLock Lock(&MaterialsInfoCriticalSection);
+	StaticMaterialsInfoCache.Reset();
+	for (int32 Index = 0; Index < InMaterials.Num(); Index++)
+	{
+		FName MaterialSlotName = InMaterials[Index].MaterialSlotName;
+		FString MaterialObjectPath = InMaterials[Index].MaterialInterface->GetPathName();
+		StaticMaterialsInfoCache.Add(FStaticMaterialMinimalInfo(MaterialSlotName, MaterialObjectPath));
+	}
+}
+
+#define MATERIAL_SLOT_DELIMITER TEXT(">")
+#define MATERIAL_SLOT_PAIR_DELIMITER TEXT(",")
+const FString UStaticMesh::ExportStaticMaterialsInfo() const
+{
+	FScopeLock Lock(&MaterialsInfoCriticalSection);
+	FString OutText;
+	for (const FStaticMaterialMinimalInfo& Slot : StaticMaterialsInfoCache)
+	{
+		OutText += Slot.MaterialSlotName.ToString();
+		OutText += MATERIAL_SLOT_PAIR_DELIMITER;
+		OutText += Slot.MaterialObjectPath;
+		OutText += MATERIAL_SLOT_DELIMITER;
+	}
+	return OutText;
+}
+
+const TArray<FStaticMaterialMinimalInfo> UStaticMesh::ImportStaticMaterialsInfo(const FString& InData)
+{
+	TArray<FStaticMaterialMinimalInfo> OutArray;
+	if (!InData.IsEmpty())
+	{
+		TArray<FString> Slot;
+		const int32 NumPairs = InData.ParseIntoArray(Slot, MATERIAL_SLOT_DELIMITER, true);
+		for (int32 Index = 0; Index < NumPairs; Index++)
+		{
+			TArray<FString> Pair;
+			int32 NumEntries = Slot[Index].ParseIntoArray(Pair, MATERIAL_SLOT_PAIR_DELIMITER, true);
+			if (NumEntries == 2)
+			{
+				OutArray.Add(FStaticMaterialMinimalInfo(FName(Pair[0]), Pair[1]));
+			}
+		}
+	}
+	return OutArray;
+}
+
 bool UStaticMesh::CanMeshPaintTextureColors() const
 {
 	if (StaticMeshPaintSupport == EStaticMeshPaintSupport::Default)
