Fix engine shutdown crash
- The graph events in FFrameEndSync were not cleared before the task graph was destroyed.
- Modified FlushRenderCommands to use FFrameEndSync rather than directly issuing an FRenderCommandFence, so that the events are cleaned up after a flush. This ensures the arrays of events are empty when the engine shuts down.
- Moved the implementation of FFrameEndSync from Engine to RenderCore to achieve this.

diff --git a/Engine/Source/Programs/ChaosVisualDebugger/Private/ChaosVisualDebuggerMain.cpp b/Engine/Source/Programs/ChaosVisualDebugger/Private/ChaosVisualDebuggerMain.cpp
--- a/Engine/Source/Programs/ChaosVisualDebugger/Private/ChaosVisualDebuggerMain.cpp
+++ b/Engine/Source/Programs/ChaosVisualDebugger/Private/ChaosVisualDebuggerMain.cpp
@@ -75,9 +75,6 @@ int32 RunChaosVisualDebugger(const TCHAR* CommandLine)
 		}
 	}
 
-	// Make sure all rendering thread work is done before we attempt to exit the program
-	FFrameEndSync::Sync(true);
-
 	GEngineLoop.Exit();
 
 	return Result;
diff --git a/Engine/Source/Runtime/Engine/Private/UnrealEngine.cpp b/Engine/Source/Runtime/Engine/Private/UnrealEngine.cpp
--- a/Engine/Source/Runtime/Engine/Private/UnrealEngine.cpp
+++ b/Engine/Source/Runtime/Engine/Private/UnrealEngine.cpp
@@ -454,23 +454,6 @@ FSystemResolution GSystemResolution;
 
 ENGINE_API int32 GUnbuiltHLODCount = 0;
 
-TAutoConsoleVariable<int32> CVarAllowOneFrameThreadLag(
-	TEXT("r.OneFrameThreadLag"),
-	1,
-	TEXT("Whether to allow the rendering thread to lag one frame behind the game thread (0: disabled, otherwise enabled)")
-);
-
-TAutoConsoleVariable<int32> CVarGTSyncType(
-	TEXT("r.GTSyncType"),
-	0,
-	TEXT("Determines how the game thread syncs with the render thread, RHI thread and GPU.\n")
-	TEXT("Syncing to the GPU swap chain flip allows for lower frame latency.\n")
-	TEXT(" <= 0 - Sync the game thread with the N-1 render thread frame. Then sync with the N-m RHI thread frame where m is (2 + (-r.GTSyncType)) (i.e. negative values increase the amount of RHI thread overlap) (default = 0).\n")
-	TEXT("    1 - Sync the game thread with the N-1 RHI thread frame.\n")
-	TEXT("    2 - Sync the game thread with the GPU swap chain flip (only on supported platforms).\n"),
-	ECVF_Default
-);
-
 static FAutoConsoleVariable CVarSystemResolution(
 	TEXT("r.SetRes"),
 	TEXT("1280x720w"),
@@ -752,8 +735,10 @@ void CalculateFPSTimings()
 	double CurrentTime = FPlatformTime::Seconds();
 	float FrameTimeMS = (float)((CurrentTime - LastTime) * 1000.0);
 
-	static auto CVarVsync = IConsoleManager::Get().FindConsoleVariable(TEXT("r.VSync"));
-	if (CVarGTSyncType.GetValueOnAnyThread() == 2 && CVarVsync->GetInt() != 0)
+	static auto CVarVsync      = IConsoleManager::Get().FindConsoleVariable(TEXT("r.VSync"));
+	static auto CVarGTSyncType = IConsoleManager::Get().FindConsoleVariable(TEXT("r.GTSyncType"));
+	
+	if (CVarGTSyncType->GetInt() == 2 && CVarVsync->GetInt() != 0)
 	{
 		float RHIFrameTime = RHIGetFrameTime();
 		if (RHIFrameTime != 0)
@@ -13316,7 +13301,6 @@ DEFINE_STAT(STAT_RedrawViewports);
 DEFINE_STAT(STAT_UpdateLevelStreaming);
 DEFINE_STAT(STAT_RHITickTime);
 DEFINE_STAT(STAT_IntentionalHitch);
-DEFINE_STAT(STAT_FrameSyncTime);
 DEFINE_STAT(STAT_DeferredTickTime);
 
 /** Input stat */
@@ -13352,119 +13336,11 @@ UFont* GetStatsFont()
 
 namespace FFrameEndSync
 {
-	using ESyncDepth = FRenderCommandFence::ESyncDepth;
-
-	struct FRenderThreadFence
-	{
-		// Legacy game code assumes the game thread will never get further than 1 frame ahead of the render thread.
-		// This fence is used to sync the game thread with the N-1 render thread frame.
-		FRenderCommandFence Fence;
-
-		FRenderThreadFence()
-		{
-			Fence.BeginFence(ESyncDepth::RenderThread);
-		}
-
-		~FRenderThreadFence()
-		{
-			Fence.Wait(true);
-		}
-	};
-	TArray<FRenderThreadFence, TInlineAllocator<2>> RenderThreadFences;
-
-	// Additional fences to await. These sync with either the RHI thread or swapchain,
-	// and are used to prevent the game thread running too far ahead of presented frames.
-	TArray<FRenderCommandFence, TInlineAllocator<3>> PipelineFences;
-
 	void Sync(bool bFullSync)
 	{
-		// The "r.OneFrameThreadLag" cvar forces a full sync, meaning the game thread will
-		// not start work until all the rendering work for the previous frame has completed.
-		bFullSync |= CVarAllowOneFrameThreadLag.GetValueOnAnyThread() <= 0;
-
-		SCOPE_CYCLE_COUNTER(STAT_FrameSyncTime);
-
-		check(IsInGameThread());
-
-	#if !UE_BUILD_SHIPPING && PLATFORM_SUPPORTS_FLIP_TRACKING
-		// Set the FrameDebugInfo on platforms that have accurate frame tracking.
-		ENQUEUE_RENDER_COMMAND(FrameDebugInfo)(
-			[CurrentFrameCounter = GFrameCounter, CurrentInputTime = GInputTime](FRHICommandListImmediate& RHICmdList)
-		{
-			RHICmdList.EnqueueLambda(
-				[CurrentFrameCounter, CurrentInputTime](FRHICommandListImmediate&)
-			{
-				// Set the FrameCount and InputTime for input latency stats and flip debugging.
-				RHISetFrameDebugInfo(GRHIPresentCounter - 1, CurrentFrameCounter, CurrentInputTime);
-			});
-		});
-	#endif
-
-		// Always sync with the render thread (either current frame, or N-1 frame)
-		RenderThreadFences.Emplace();
-		while (RenderThreadFences.Num() > (bFullSync ? 0 : 1))
-		{
-			RenderThreadFences.RemoveAt(0);
-		}
-
-		// Insert an additional fence based on how we want to sync with the RHI thread / swapchain
-		ESyncDepth SyncDepth;
-		int32 NumFramesOverlap;
-
-		int32 const GTSyncType = CVarGTSyncType.GetValueOnAnyThread();
-
-		if (bFullSync)
-		{
-			SyncDepth = GTSyncType >= 2
-				? ESyncDepth::Swapchain
-				: ESyncDepth::RHIThread;
-
-			NumFramesOverlap = 0;
-		}
-		else if (GTSyncType >= 2)
-		{
-			SyncDepth = ESyncDepth::Swapchain;
-			NumFramesOverlap = 1;
-		}
-		else if (GTSyncType == 1)
-		{
-			SyncDepth = ESyncDepth::RHIThread;
-			NumFramesOverlap = 1;
-		}
-		else
-		{
-			check(GTSyncType <= 0);
-
-			// Modes <= 0 allows N frames of overlap with the RHI thread.
-			SyncDepth = ESyncDepth::RHIThread;
-			NumFramesOverlap = 2 + (-GTSyncType);
-		}
-
-		if (SyncDepth == ESyncDepth::Swapchain)
-		{
-			// Swapchain sync mode does not work when vsync is disabled. Fallback to RHI thread sync in that case.
-			static auto CVarVsync = IConsoleManager::Get().FindConsoleVariable(TEXT("r.VSync"));
-			check(CVarVsync != nullptr);
-
-			if (CVarVsync->GetInt() == 0)
-			{
-				SyncDepth = ESyncDepth::RHIThread;
-			}
-		}
-
-		PipelineFences.Emplace_GetRef().BeginFence(SyncDepth);
-
-		if (!FTaskGraphInterface::Get().IsThreadProcessingTasks(ENamedThreads::GameThread))
-		{
-			// need to process gamethread tasks at least once a frame no matter what
-			FTaskGraphInterface::Get().ProcessThreadUntilIdle(ENamedThreads::GameThread);
-		}
-
-		while (PipelineFences.Num() > NumFramesOverlap)
-		{
-			PipelineFences[0].Wait(true);
-			PipelineFences.RemoveAt(0);
-		}
+		Sync(bFullSync
+			? EFlushMode::Threads
+			: EFlushMode::EndFrame);
 	}
 }
 
diff --git a/Engine/Source/Runtime/Engine/Public/EngineStats.h b/Engine/Source/Runtime/Engine/Public/EngineStats.h
--- a/Engine/Source/Runtime/Engine/Public/EngineStats.h
+++ b/Engine/Source/Runtime/Engine/Public/EngineStats.h
@@ -27,7 +27,6 @@ DECLARE_CYCLE_STAT_EXTERN(TEXT("RedrawViewports"),STAT_RedrawViewports,STATGROUP
 DECLARE_CYCLE_STAT_EXTERN(TEXT("Update Level Streaming"),STAT_UpdateLevelStreaming,STATGROUP_Engine, ENGINE_API);
 DECLARE_CYCLE_STAT_EXTERN(TEXT("RHI Game Tick"),STAT_RHITickTime,STATGROUP_Engine, ENGINE_API);
 DECLARE_CYCLE_STAT_EXTERN(TEXT("Debug Hitch"),STAT_IntentionalHitch,STATGROUP_Engine, );
-DECLARE_CYCLE_STAT_EXTERN(TEXT("Frame Sync Time"),STAT_FrameSyncTime,STATGROUP_Engine, ENGINE_API);
 DECLARE_CYCLE_STAT_EXTERN(TEXT("Deferred Tick Time"),STAT_DeferredTickTime,STATGROUP_Engine, ENGINE_API);
 
 /** Unit time stats */
diff --git a/Engine/Source/Runtime/Engine/Public/UnrealEngine.h b/Engine/Source/Runtime/Engine/Public/UnrealEngine.h
--- a/Engine/Source/Runtime/Engine/Public/UnrealEngine.h
+++ b/Engine/Source/Runtime/Engine/Public/UnrealEngine.h
@@ -375,10 +375,10 @@ namespace FFrameEndSync
 	 * @param bFullSync - when true, blocks the caller until all rendering work is completed,
 	 * otherwise only blocks until the N - m frame has completed, where m is driven by various config.
 	 */
+	UE_DEPRECATED(5.6, "Use FFrameEndSync::Sync defined in RenderingThread.h, which takes an enum.")
 	ENGINE_API void Sync(bool bFullSync = false);
 };
 
-
 /** Public interface to FEngineLoop so we can call it from editor or editor code */
 class IEngineLoop
 {
diff --git a/Engine/Source/Runtime/Launch/Private/LaunchEngineLoop.cpp b/Engine/Source/Runtime/Launch/Private/LaunchEngineLoop.cpp
--- a/Engine/Source/Runtime/Launch/Private/LaunchEngineLoop.cpp
+++ b/Engine/Source/Runtime/Launch/Private/LaunchEngineLoop.cpp
@@ -6085,7 +6085,7 @@ void FEngineLoop::Tick()
 
 		// This could be perhaps moved down to get greater parallelism
 		// Sync game and render/RHI threads.
-		FFrameEndSync::Sync();
+		FFrameEndSync::Sync(FFrameEndSync::EFlushMode::EndFrame);
 
 		// tick core ticker, threads & deferred commands
 		{
diff --git a/Engine/Source/Runtime/RenderCore/Private/RenderingThread.cpp b/Engine/Source/Runtime/RenderCore/Private/RenderingThread.cpp
--- a/Engine/Source/Runtime/RenderCore/Private/RenderingThread.cpp
+++ b/Engine/Source/Runtime/RenderCore/Private/RenderingThread.cpp
@@ -1344,9 +1344,8 @@ void FlushRenderingCommands()
 	FPendingCleanupObjects* PendingCleanupObjects = GetPendingCleanupObjects();
 
 	// Issue a fence command to the rendering thread and wait for it to complete.
-	FRenderCommandFence Fence;
-	Fence.BeginFence();
-	Fence.Wait();
+	// Use the frame end sync here, so that it cleans up outstanding graph events, which is necessary on engine shutdown.
+	FFrameEndSync::Sync(FFrameEndSync::EFlushMode::Threads);
 
 	// Delete the objects which were enqueued for deferred cleanup before the command queue flush.
 	delete PendingCleanupObjects;
@@ -1983,4 +1982,143 @@ void FRenderCommandPipe::EnqueueAndLaunch(FFunctionVariant&& FunctionVariant, co
 	
 		}, Frame_GameThread->LastTask);
 	}
-}
\ No newline at end of file
+}
+
+TAutoConsoleVariable<int32> CVarAllowOneFrameThreadLag(
+	TEXT("r.OneFrameThreadLag"),
+	1,
+	TEXT("Whether to allow the rendering thread to lag one frame behind the game thread (0: disabled, otherwise enabled)")
+);
+
+TAutoConsoleVariable<int32> CVarGTSyncType(
+	TEXT("r.GTSyncType"),
+	0,
+	TEXT("Determines how the game thread syncs with the render thread, RHI thread and GPU.\n")
+	TEXT("Syncing to the GPU swap chain flip allows for lower frame latency.\n")
+	TEXT(" <= 0 - Sync the game thread with the N-1 render thread frame. Then sync with the N-m RHI thread frame where m is (2 + (-r.GTSyncType)) (i.e. negative values increase the amount of RHI thread overlap) (default = 0).\n")
+	TEXT("    1 - Sync the game thread with the N-1 RHI thread frame.\n")
+	TEXT("    2 - Sync the game thread with the GPU swap chain flip (only on supported platforms).\n"),
+	ECVF_Default
+);
+
+DECLARE_CYCLE_STAT(TEXT("Frame Sync Time"), STAT_FrameSyncTime, STATGROUP_RenderThreadProcessing);
+
+namespace FFrameEndSync
+{
+	using ESyncDepth = FRenderCommandFence::ESyncDepth;
+
+	struct FRenderThreadFence
+	{
+		// Legacy game code assumes the game thread will never get further than 1 frame ahead of the render thread.
+		// This fence is used to sync the game thread with the N-1 render thread frame.
+		FRenderCommandFence Fence;
+
+		FRenderThreadFence()
+		{
+			Fence.BeginFence(ESyncDepth::RenderThread);
+		}
+
+		~FRenderThreadFence()
+		{
+			Fence.Wait(true);
+		}
+	};
+	TArray<FRenderThreadFence, TInlineAllocator<2>> RenderThreadFences;
+
+	// Additional fences to await. These sync with either the RHI thread or swapchain,
+	// and are used to prevent the game thread running too far ahead of presented frames.
+	TArray<FRenderCommandFence, TInlineAllocator<3>> PipelineFences;
+
+	void Sync(EFlushMode FlushMode)
+	{
+		bool bFullSync = FlushMode == EFlushMode::Threads;
+
+		// The "r.OneFrameThreadLag" cvar forces a full sync, meaning the game thread will
+		// not start work until all the rendering work for the previous frame has completed.
+		bFullSync |= CVarAllowOneFrameThreadLag.GetValueOnAnyThread() <= 0;
+
+		SCOPE_CYCLE_COUNTER(STAT_FrameSyncTime);
+
+		check(IsInGameThread());
+
+	#if !UE_BUILD_SHIPPING && PLATFORM_SUPPORTS_FLIP_TRACKING
+		// Set the FrameDebugInfo on platforms that have accurate frame tracking.
+		ENQUEUE_RENDER_COMMAND(FrameDebugInfo)(
+			[CurrentFrameCounter = GFrameCounter, CurrentInputTime = GInputTime](FRHICommandListImmediate& RHICmdList)
+		{
+			RHICmdList.EnqueueLambda(
+				[CurrentFrameCounter, CurrentInputTime](FRHICommandListImmediate&)
+			{
+				// Set the FrameCount and InputTime for input latency stats and flip debugging.
+				RHISetFrameDebugInfo(GRHIPresentCounter - 1, CurrentFrameCounter, CurrentInputTime);
+			});
+		});
+	#endif
+
+		// Always sync with the render thread (either current frame, or N-1 frame)
+		RenderThreadFences.Emplace();
+		while (RenderThreadFences.Num() > (bFullSync ? 0 : 1))
+		{
+			RenderThreadFences.RemoveAt(0);
+		}
+
+		// Insert an additional fence based on how we want to sync with the RHI thread / swapchain
+		ESyncDepth SyncDepth;
+		int32 NumFramesOverlap;
+
+		int32 const GTSyncType = CVarGTSyncType.GetValueOnAnyThread();
+
+		if (bFullSync)
+		{
+			SyncDepth = (GTSyncType >= 2 && FlushMode != EFlushMode::Threads)
+				? ESyncDepth::Swapchain
+				: ESyncDepth::RHIThread;
+
+			NumFramesOverlap = 0;
+		}
+		else if (GTSyncType >= 2)
+		{
+			SyncDepth = ESyncDepth::Swapchain;
+			NumFramesOverlap = 1;
+		}
+		else if (GTSyncType == 1)
+		{
+			SyncDepth = ESyncDepth::RHIThread;
+			NumFramesOverlap = 1;
+		}
+		else
+		{
+			check(GTSyncType <= 0);
+
+			// Modes <= 0 allows N frames of overlap with the RHI thread.
+			SyncDepth = ESyncDepth::RHIThread;
+			NumFramesOverlap = 2 + (-GTSyncType);
+		}
+
+		if (SyncDepth == ESyncDepth::Swapchain)
+		{
+			// Swapchain sync mode does not work when vsync is disabled. Fallback to RHI thread sync in that case.
+			static auto CVarVsync = IConsoleManager::Get().FindConsoleVariable(TEXT("r.VSync"));
+			check(CVarVsync != nullptr);
+
+			if (CVarVsync->GetInt() == 0)
+			{
+				SyncDepth = ESyncDepth::RHIThread;
+			}
+		}
+
+		PipelineFences.Emplace_GetRef().BeginFence(SyncDepth);
+
+		if (!FTaskGraphInterface::Get().IsThreadProcessingTasks(ENamedThreads::GameThread))
+		{
+			// need to process gamethread tasks at least once a frame no matter what
+			FTaskGraphInterface::Get().ProcessThreadUntilIdle(ENamedThreads::GameThread);
+		}
+
+		while (PipelineFences.Num() > NumFramesOverlap)
+		{
+			PipelineFences[0].Wait(true);
+			PipelineFences.RemoveAt(0);
+		}
+	}
+}
diff --git a/Engine/Source/Runtime/RenderCore/Public/RenderingThread.h b/Engine/Source/Runtime/RenderCore/Public/RenderingThread.h
--- a/Engine/Source/Runtime/RenderCore/Public/RenderingThread.h
+++ b/Engine/Source/Runtime/RenderCore/Public/RenderingThread.h
@@ -36,6 +36,21 @@ namespace UE { namespace Trace { class FChannel; } }
 // Render thread API
 ////////////////////////////////////
 
+namespace FFrameEndSync
+{
+	enum class EFlushMode
+	{
+		// Blocks the caller until the N - m frame has completed, where m is driven by various config.
+		EndFrame,
+
+		// Blocks the caller until all rendering work is completed on the CPU. Does not sync with the GPU.
+		Threads
+	};
+
+	// Syncs the game thread based on progress throughout the rendering pipeline.
+	RENDERCORE_API void Sync(EFlushMode FlushMode);
+};
+
 /**
  * Whether the renderer is currently running in a separate thread.
  * If this is false, then all rendering commands will be executed immediately instead of being enqueued in the rendering command buffer.
diff --git a/Engine/Source/Runtime/SlateNullRenderer/Private/SlateNullRenderer.cpp b/Engine/Source/Runtime/SlateNullRenderer/Private/SlateNullRenderer.cpp
--- a/Engine/Source/Runtime/SlateNullRenderer/Private/SlateNullRenderer.cpp
+++ b/Engine/Source/Runtime/SlateNullRenderer/Private/SlateNullRenderer.cpp
@@ -5,7 +5,7 @@
 #include "Rendering/ShaderResourceManager.h"
 #include "Rendering/SlateDrawBuffer.h"
 #if UE_SLATE_NULL_RENDERER_WITH_ENGINE
-#include "UnrealEngine.h"
+#include "RenderingThread.h"
 #endif
 
 
@@ -144,6 +144,6 @@ void FSlateNullRenderer::ClearScenes()
 void FSlateNullRenderer::Sync() const
 {
 #if UE_SLATE_NULL_RENDERER_WITH_ENGINE
-	FFrameEndSync::Sync();
+	FFrameEndSync::Sync(FFrameEndSync::EFlushMode::EndFrame);
 #endif
 }
diff --git a/Engine/Source/Runtime/SlateRHIRenderer/Private/SlateRHIRenderer.cpp b/Engine/Source/Runtime/SlateRHIRenderer/Private/SlateRHIRenderer.cpp
--- a/Engine/Source/Runtime/SlateRHIRenderer/Private/SlateRHIRenderer.cpp
+++ b/Engine/Source/Runtime/SlateRHIRenderer/Private/SlateRHIRenderer.cpp
@@ -1529,7 +1529,7 @@ void FSlateRHIRenderer::FlushCommands() const
 
 void FSlateRHIRenderer::Sync() const
 {
-	FFrameEndSync::Sync();
+	FFrameEndSync::Sync(FFrameEndSync::EFlushMode::EndFrame);
 }
 
 void FSlateRHIRenderer::BeginFrame() const
