First Person: Force disabled Nanite on meshes tagged as FirstPerson or FirstPersonWorldSpaceRepresentation. Nanite does not currently support bOwnerNoSee/bOnlyOwnerSee, so FirstPerson Nanite meshes could show up in the view of other players in local splitscreen and FirstPersonWorldSpaceRepresentation Nanite meshes would likely mess with VSM player shadows being cast onto the scene. This change should be reverted once Nanite supports bOwnerNoSee/bOnlyOwnerSee but until then, it should provide for a more robust user experience.

#tests Local testing with static and skinned meshes. Switching FirstPersonPrimitiveType to FirstPerson on a Nanite mesh correctly makes it render using the non-Nanite path. (r.Nanite.Visualize mask)

diff --git a/Engine/Source/Runtime/Engine/Private/Components/SkinnedMeshComponent.cpp b/Engine/Source/Runtime/Engine/Private/Components/SkinnedMeshComponent.cpp
--- a/Engine/Source/Runtime/Engine/Private/Components/SkinnedMeshComponent.cpp
+++ b/Engine/Source/Runtime/Engine/Private/Components/SkinnedMeshComponent.cpp
@@ -1654,7 +1654,9 @@ bool USkinnedMeshComponent::IsMaterialSlotNameValid(FName MaterialSlotName) cons
 bool USkinnedMeshComponent::ShouldNaniteSkin()
 {
 	const EShaderPlatform ShaderPlatform = GetScene() ? GetScene()->GetShaderPlatform() : GMaxRHIShaderPlatform;
-	return ShouldRenderNaniteSkinnedMeshes() && UseNanite(ShaderPlatform) && HasValidNaniteData();
+	// Force disable Nanite on first person primitives as they cause issues when used together with bOwnerNoSee/bOnlyOwnerSee, which is often the case for FirstPerson.
+	const bool bIsFirstPersonRelevant = FirstPersonPrimitiveType != EFirstPersonPrimitiveType::None;
+	return ShouldRenderNaniteSkinnedMeshes() && UseNanite(ShaderPlatform) && HasValidNaniteData() && !bIsFirstPersonRelevant;
 }
 
 bool USkinnedMeshComponent::ShouldCPUSkin() const
diff --git a/Engine/Source/Runtime/Engine/Private/SkinnedMeshSceneProxyDesc.cpp b/Engine/Source/Runtime/Engine/Private/SkinnedMeshSceneProxyDesc.cpp
--- a/Engine/Source/Runtime/Engine/Private/SkinnedMeshSceneProxyDesc.cpp
+++ b/Engine/Source/Runtime/Engine/Private/SkinnedMeshSceneProxyDesc.cpp
@@ -209,7 +209,9 @@ bool FSkinnedMeshSceneProxyDesc::HasValidNaniteData() const
 bool FSkinnedMeshSceneProxyDesc::ShouldNaniteSkin() const
 {
 	const EShaderPlatform ShaderPlatform = GetScene() ? GetScene()->GetShaderPlatform() : GMaxRHIShaderPlatform;
-	return ShouldRenderNaniteSkinnedMeshes() && UseNanite(ShaderPlatform) && HasValidNaniteData();
+	// Force disable Nanite on first person primitives as they cause issues when used together with bOwnerNoSee/bOnlyOwnerSee, which is often the case for FirstPerson.
+	const bool bIsFirstPersonRelevant = bIsFirstPerson || bIsFirstPersonWorldSpaceRepresentation;
+	return ShouldRenderNaniteSkinnedMeshes() && UseNanite(ShaderPlatform) && HasValidNaniteData() && !bIsFirstPersonRelevant;
 }
 
 FSkeletalMeshObject* FSkinnedMeshSceneProxyDesc::CreateMeshObject(const FSkinnedMeshSceneProxyDesc& Desc)
diff --git a/Engine/Source/Runtime/Engine/Private/StaticMeshSceneProxy.cpp b/Engine/Source/Runtime/Engine/Private/StaticMeshSceneProxy.cpp
--- a/Engine/Source/Runtime/Engine/Private/StaticMeshSceneProxy.cpp
+++ b/Engine/Source/Runtime/Engine/Private/StaticMeshSceneProxy.cpp
@@ -2468,6 +2468,16 @@ FLODMask FStaticMeshSceneProxy::GetLODMask(const FSceneView* View) const
 
 namespace Nanite
 {
+	bool IsStaticMeshFirstPersonRelevant(const UStaticMeshComponent& Component)
+	{
+		return Component.FirstPersonPrimitiveType != EFirstPersonPrimitiveType::None;
+	}
+
+	bool IsStaticMeshFirstPersonRelevant(const FStaticMeshSceneProxyDesc& Component)
+	{
+		return Component.bIsFirstPerson || Component.bIsFirstPersonWorldSpaceRepresentation;
+	}
+
 	template<class T> 
 	bool ShouldCreateNaniteProxy(const T& Component, FMaterialAudit* OutNaniteMaterials)
 	{
@@ -2478,7 +2488,9 @@ namespace Nanite
 		const bool bForceFallback = false;
 	#endif
 
-		if (bForceFallback || Component.bDisallowNanite || Component.bForceDisableNanite)
+		// Force disable Nanite on first person primitives as they cause issues when used together with bOwnerNoSee/bOnlyOwnerSee, which is often the case for FirstPerson.
+		const bool bIsFirstPersonRelevant = IsStaticMeshFirstPersonRelevant(Component);
+		if (bForceFallback || Component.bDisallowNanite || Component.bForceDisableNanite || bIsFirstPersonRelevant)
 		{
 			// Regardless of the static mesh asset supporting Nanite, this component does not want Nanite to be used
 			return false;
