Patch to ensure MaterialInterface Object is not null when serializing StaticMesh slot names to the asset registry.

diff --git a/Engine/Source/Runtime/Engine/Private/StaticMesh.cpp b/Engine/Source/Runtime/Engine/Private/StaticMesh.cpp
--- a/Engine/Source/Runtime/Engine/Private/StaticMesh.cpp
+++ b/Engine/Source/Runtime/Engine/Private/StaticMesh.cpp
@@ -9598,7 +9598,7 @@ void UStaticMesh::UpdateStaticMaterialsInfoCache(const TArray<FStaticMaterial>&
 	for (int32 Index = 0; Index < InMaterials.Num(); Index++)
 	{
 		FName MaterialSlotName = InMaterials[Index].MaterialSlotName;
-		FString MaterialObjectPath = InMaterials[Index].MaterialInterface->GetPathName();
+		FString MaterialObjectPath = InMaterials[Index].MaterialInterface ? InMaterials[Index].MaterialInterface->GetPathName() : FString();
 		StaticMaterialsInfoCache.Add(FStaticMaterialMinimalInfo(MaterialSlotName, MaterialObjectPath));
 	}
 }
