Add basic initial functionality to SubtilesAndClosedCaptionsModule:
Plugin notes:
-) SubtitlesSettings - contains settings that can be set per-Game, uses Game config. Includes font, color, size
-) SubtitlesAssetUserData - analagous with one Subtitle. Contains the subtitle Text, Duration, and Priority. Can be subclassed per-project.
-) SubtitlesSubsystem - Handles queueing up and displaying of SubtitleAssetUserData. Can be subclassed per project to work with extensions to the user data.
-) Added tests for adding, removing, and prioritizing subtitles

Engine change notes:
-) Engine.h and ISubtitlesAndClosedCaptionsModule.h - Engine change introducing a module interface that can be overridden. Engine setting Subtitles.SubtitlesAndClosedCaptionsModuleName has been added so that projecs can opt-in to using new subtitles functionality, or, provide their own implementation or extend upon the basic subtitles functionality.
-) SubtitlesAndClosedCaptionsDelegates.h - Delegates that call .ExecuteIfBound() where subtitle functionality is needed.  The SubtitlesSubsystem binds to these delegates.  Added to minimize engine changes so that the module doesn't have to be retrieved whenever its needed, and also so API functions don't have to be added to the Module which would end up forwarding the calls to the SubtitlesSubsystem.
-) GameViewportClient + SoundWave make use of these delegates.  SoundWave has an array of asset user data, so the subtitles data is added to this array.

diff --git a/Engine/Source/Runtime/Engine/Private/GameViewportClient.cpp b/Engine/Source/Runtime/Engine/Private/GameViewportClient.cpp
--- a/Engine/Source/Runtime/Engine/Private/GameViewportClient.cpp
+++ b/Engine/Source/Runtime/Engine/Private/GameViewportClient.cpp
@@ -32,6 +32,7 @@
 #include "GameFramework/HUD.h"
 #include "FXSystem.h"
 #include "SubtitleManager.h"
+#include "Subtitles/SubtitlesAndClosedCaptionsDelegates.h"
 #include "ImageUtils.h"
 #include "SceneViewExtension.h"
 #include "IHeadMountedDisplay.h"
@@ -1996,6 +1997,7 @@ void UGameViewportClient::Draw(FViewport* InViewport, FCanvas* SceneCanvas)
 							const uint32 SizeY = SceneCanvas->GetRenderTarget()->GetSizeXY().Y;
 							FIntRect SubtitleRegion(FMath::TruncToInt(SizeX * MinPos.X), FMath::TruncToInt(SizeY * MinPos.Y), FMath::TruncToInt(SizeX * MaxPos.X), FMath::TruncToInt(SizeY * MaxPos.Y));
 							FSubtitleManager::GetSubtitleManager()->DisplaySubtitles( SceneCanvas, SubtitleRegion, MyWorld->GetAudioTimeSeconds() );
+							FSubtitlesAndClosedCaptionsDelegates::DisplaySubtitle.ExecuteIfBound(*SceneCanvas);
 							bDisplayedSubtitles = true;
 						}
 					}
diff --git a/Engine/Source/Runtime/Engine/Private/SoundWave.cpp b/Engine/Source/Runtime/Engine/Private/SoundWave.cpp
--- a/Engine/Source/Runtime/Engine/Private/SoundWave.cpp
+++ b/Engine/Source/Runtime/Engine/Private/SoundWave.cpp
@@ -21,6 +21,8 @@
 #include "Misc/App.h"
 #include "Misc/ConfigCacheIni.h"
 #include "SubtitleManager.h"
+#include "Subtitles/ISubtitlesAndClosedCaptionsModule.h"
+#include "Subtitles/SubtitlesAndClosedCaptionsDelegates.h"
 #include "DerivedDataCacheInterface.h"
 #include "EditorFramework/AssetImportData.h"
 #include "ProfilingDebugging/CookStats.h"
@@ -3356,6 +3358,16 @@ void USoundWave::Parse(FAudioDevice* AudioDevice, const UPTRINT NodeWaveInstance
 		}
 
 		WaveInstance = &HandleStart(ActiveSound, NodeWaveInstanceHash);
+
+		IModularFeatures& ModularFeatures = IModularFeatures::Get();
+		if (ModularFeatures.IsModularFeatureAvailable(ISubtitlesAndClosedCaptionsModule::GetModularFeatureName()))
+		{
+			const ISubtitlesAndClosedCaptionsModule& SubtitlesModule = ModularFeatures.GetModularFeature<ISubtitlesAndClosedCaptionsModule>(ISubtitlesAndClosedCaptionsModule::GetModularFeatureName());
+			if (const UAssetUserData * Subtitle{ GetAssetUserDataOfClass(SubtitlesModule.GetAssetUserDataClass()) })
+			{
+				FSubtitlesAndClosedCaptionsDelegates::QueueSubtitle.ExecuteIfBound(*Subtitle);
+			}
+		}
 	}
 	else
 	{
diff --git a/Engine/Source/Runtime/Engine/Private/Subtitles/SubtitlesAndClosedCaptionsDelegates.cpp b/Engine/Source/Runtime/Engine/Private/Subtitles/SubtitlesAndClosedCaptionsDelegates.cpp
new file mode 100644
--- /dev/null
+++ b/Engine/Source/Runtime/Engine/Private/Subtitles/SubtitlesAndClosedCaptionsDelegates.cpp
@@ -0,0 +1,6 @@
+// Copyright Epic Games, Inc. All Rights Reserved.
+
+#include "Subtitles/SubtitlesAndClosedCaptionsDelegates.h"
+
+TDelegate<void (const UAssetUserData&)> FSubtitlesAndClosedCaptionsDelegates::QueueSubtitle;
+TDelegate<void(FCanvas&)> FSubtitlesAndClosedCaptionsDelegates::DisplaySubtitle;
diff --git a/Engine/Source/Runtime/Engine/Public/Subtitles/ISubtitlesAndClosedCaptionsModule.h b/Engine/Source/Runtime/Engine/Public/Subtitles/ISubtitlesAndClosedCaptionsModule.h
new file mode 100644
--- /dev/null
+++ b/Engine/Source/Runtime/Engine/Public/Subtitles/ISubtitlesAndClosedCaptionsModule.h
@@ -0,0 +1,21 @@
+// Copyright Epic Games, Inc. All Rights Reserved.
+
+#pragma once
+
+#include "Features/IModularFeature.h"
+#include "Modules/ModuleInterface.h"
+#include "Templates/SubclassOf.h"
+
+class UAssetUserData;
+
+class ISubtitlesAndClosedCaptionsModule : public IModuleInterface, public IModularFeature
+{
+public:
+	static FName GetModularFeatureName()
+	{
+		static const FName FeatureName = FName(TEXT("SubtitlesAndClosedCaptions"));
+		return FeatureName;
+	}
+
+	virtual TSubclassOf<UAssetUserData> GetAssetUserDataClass() const = 0;
+};
diff --git a/Engine/Source/Runtime/Engine/Public/Subtitles/SubtitlesAndClosedCaptionsDelegates.h b/Engine/Source/Runtime/Engine/Public/Subtitles/SubtitlesAndClosedCaptionsDelegates.h
new file mode 100644
--- /dev/null
+++ b/Engine/Source/Runtime/Engine/Public/Subtitles/SubtitlesAndClosedCaptionsDelegates.h
@@ -0,0 +1,20 @@
+// Copyright Epic Games, Inc. All Rights Reserved.
+
+#pragma once
+#include "Delegates/Delegate.h"
+
+class FCanvas;
+class UAssetUserData;
+
+class ENGINE_API FSubtitlesAndClosedCaptionsDelegates
+{
+public:
+
+	// Have the subtitle subsystem to queue up a subtitle
+	static TDelegate<void(const UAssetUserData&)> QueueSubtitle;
+
+	// Have the subtitle subsystem to display the highest priority subtitle. DisplaySubtitles is expected to be called/drawn per frame.
+	static TDelegate<void(FCanvas&)> DisplaySubtitle;
+};
+
+
