Fix GC crash when setting reverb
#rnx

diff --git a/Engine/Source/Runtime/Engine/Private/AudioDevice.cpp b/Engine/Source/Runtime/Engine/Private/AudioDevice.cpp
--- a/Engine/Source/Runtime/Engine/Private/AudioDevice.cpp
+++ b/Engine/Source/Runtime/Engine/Private/AudioDevice.cpp
@@ -1056,13 +1056,37 @@ void FAudioDevice::CountBytes(FArchive& Ar)
 
 void FAudioDevice::AddReferencedObjects(FReferenceCollector& Collector)
 {
+	/** Threading Note: The FAudioThread may run concurrently with `AddReferencedObjects(...). 
+	 *
+	 * The FAudioThread is suspended during GC to avoid accesing objects during
+	 * or after their destruction. But the FAudioThread does NOT suspend during 
+	 * AddReferencedObjects(...). Any object that needs to be accessed on the FAudioThread
+	 * but needs protection from GC should be adding there references to the reference collector. 
+	 * 
+	 * In FAudioDevice, there are two primary strategies used to achieve this, where systems
+	 * generally choose one strategy depending on their system's needs, though sometimes require
+	 * both strategies in more complex scenarios. 
+	 *
+	 * Strategy 1: Manage the lifetime of the object entirely on the GT, where GT logic can
+	 * ensure that the object will be protected from GC any time the object is accessed
+	 * on the FAudioThread.  
+	 *
+	 * Strategy 2: Manage the lifetime of the object on the AT and utilize the AudioThreadReferenceCritSec
+	 * to protect against modifying the object while it is being read during `AddReferencedObjects(...)`. 
+	 */
+
 	Collector.AddReferencedObject(DefaultBaseSoundMix);
 	Collector.AddReferencedObjects(PrevPassiveSoundMixModifiers);
 	Collector.AddReferencedObjects(SoundMixModifiers);
 
 	for (TPair<FName, FActivatedReverb>& ActivatedReverbPair : ActivatedReverbs)
 	{
-		Collector.AddReferencedObject(ActivatedReverbPair.Value.ReverbSettings.ReverbEffect);
+		Collector.AddPropertyReferencesWithStructARO(FReverbSettings::StaticStruct(), &ActivatedReverbPair.Value.ReverbSettings);
+	}
+
+	if (PendingActivatedReverb)
+	{
+		Collector.AddPropertyReferencesWithStructARO(FReverbSettings::StaticStruct(), &PendingActivatedReverb->ReverbSettings);
 	}
 
 	if (Effects)
@@ -1090,6 +1114,23 @@ void FAudioDevice::AddReferencedObjects(FReferenceCollector& Collector)
 
 	// Make sure we don't try to delete any sound waves which may have in-flight decodes
 	Collector.AddReferencedObjects(ReferencedSoundWaves);
+
+	{
+		// Lock to protect against modification of objects which are owned by the AudioThread, but need
+		// their references tracked to protect against GC. 
+		FScopeLock AudioThreadReferenceLock(&AudioThreadReferenceCritSec);
+
+		Collector.AddPropertyReferencesWithStructARO(FReverbSettings::StaticStruct(), &HighestPriorityActivatedReverb_AudioThread.ReverbSettings);
+
+		for (FAudioVolumeSubmixSendSettings& Settings : CurrentAudioVolumeSettings.SubmixSendSettings)
+		{
+			Collector.AddPropertyReferencesWithStructARO(FAudioVolumeSubmixSendSettings::StaticStruct(), &Settings);
+		}
+		for (FAudioVolumeSubmixOverrideSettings& Settings : CurrentAudioVolumeSettings.SubmixOverrideSettings)
+		{
+			Collector.AddPropertyReferencesWithStructARO(FAudioVolumeSubmixOverrideSettings::StaticStruct(), &Settings);
+		}
+	}
 }
 
 void FAudioDevice::ResetInterpolation()
@@ -2512,19 +2553,26 @@ void FAudioDevice::UpdateHighestPriorityReverb()
 	{
 		ActivatedReverbs.ValueSort([](const FActivatedReverb& A, const FActivatedReverb& B) { return A.Priority > B.Priority; });
 
-		const FActivatedReverb& NewActiveReverbRef = ActivatedReverbs.CreateConstIterator().Value();
-		FAudioThread::RunCommandOnAudioThread([AudioDevice, NewActiveReverbRef]()
+		// Use shared ptr to ensure that pending reverb is protected from GC before it gets applied on the audio thread.
+		PendingActivatedReverb = MakeShared<FActivatedReverb>(ActivatedReverbs.CreateConstIterator().Value());
+		FAudioThread::RunCommandOnAudioThread([AudioDevice, NewActivatedReverbWeakPtr = TWeakPtr<FActivatedReverb>(PendingActivatedReverb)]()
 		{
-			AudioDevice->bHasActivatedReverb = true;
-			AudioDevice->HighestPriorityActivatedReverb = NewActiveReverbRef;
+			if (TSharedPtr<FActivatedReverb> NewActivatedReverb = NewActivatedReverbWeakPtr.Pin())
+			{
+				FScopeLock AudioThreadReferenceLock(&AudioDevice->AudioThreadReferenceCritSec);
+				AudioDevice->bHasActivatedReverb = true;
+				AudioDevice->HighestPriorityActivatedReverb_AudioThread = *NewActivatedReverb;
+			}
 		}, GET_STATID(STAT_AudioUpdateHighestPriorityReverb));
 	}
 	else
 	{
+		PendingActivatedReverb.Reset();
 		FAudioThread::RunCommandOnAudioThread([AudioDevice]()
 		{
+			FScopeLock AudioThreadReferenceLock(&AudioDevice->AudioThreadReferenceCritSec);
 			AudioDevice->bHasActivatedReverb = false;
-			AudioDevice->HighestPriorityActivatedReverb = FActivatedReverb();
+			AudioDevice->HighestPriorityActivatedReverb_AudioThread = FActivatedReverb();
 		}, GET_STATID(STAT_AudioUpdateHighestPriorityReverb));
 	}
 }
@@ -4558,7 +4606,7 @@ void FAudioDevice::StartSources(TArray<FWaveInstance*>& WaveInstances, int32 Fir
 	// Run a command to make sure we add the starting sounds to the referenced sound waves list
 	if (StartingSoundWaves.Num() > 0)
 	{
-		FScopeLock ReferencedSoundWaveLock(&ReferencedSoundWaveCritSec);
+		FScopeLock AudioThreadReferenceLock(&AudioThreadReferenceCritSec);
 
 		for (USoundWave* SoundWave : StartingSoundWaves)
 		{
@@ -4570,7 +4618,7 @@ void FAudioDevice::StartSources(TArray<FWaveInstance*>& WaveInstances, int32 Fir
 void FAudioDevice::UpdateReferencedSoundWaves()
 {
 	{
-		FScopeLock ReferencedSoundWaveLock(&ReferencedSoundWaveCritSec);
+		FScopeLock AudioThreadReferenceLock(&AudioThreadReferenceCritSec);
 
 		for (USoundWave* SoundWave : ReferencedSoundWaves_AudioThread)
 		{
@@ -4851,74 +4899,79 @@ void FAudioDevice::UpdateAudioVolumeEffects()
 	bool bUpdateSubmixEffects = (PlayerAudioVolumeSettings.AudioVolumeID != CurrentAudioVolumeSettings.AudioVolumeID) ||
 		(PlayerAudioVolumeSettings.AudioVolumeID == CurrentAudioVolumeSettings.AudioVolumeID && PlayerAudioVolumeSettings.bChanged);
 
-	CurrentAudioVolumeSettings = PlayerAudioVolumeSettings;
-
-	if (Effects)
 	{
-		// Check if we should be using activated reverb
-		if (bHasActivatedReverb && (HighestPriorityActivatedReverb.Priority > PlayerAudioVolumeSettings.Priority || bUsingDefaultReverb))
+		// Lock for modifications to CurrentAudioVolumeSettings and contention with AddReferencedObjects(...)
+		FScopeLock AudioThreadReferenceLock(&AudioThreadReferenceCritSec);
+
+		CurrentAudioVolumeSettings = PlayerAudioVolumeSettings;
+
+		if (Effects)
 		{
-			CurrentAudioVolumeSettings.ReverbSettings = HighestPriorityActivatedReverb.ReverbSettings;
-		}
+			// Check if we should be using activated reverb
+			if (bHasActivatedReverb && (HighestPriorityActivatedReverb_AudioThread.Priority > PlayerAudioVolumeSettings.Priority || bUsingDefaultReverb))
+			{
+				CurrentAudioVolumeSettings.ReverbSettings = HighestPriorityActivatedReverb_AudioThread.ReverbSettings;
+			}
 
 #if ENABLE_AUDIO_DEBUG
-		// This is for debug visualization only - the audio debugger will show a reverb effect is active, if bApplyReverb is false for the world (default) reverb.
-		// In order to not mislead people using au.debug.reverb, we manually set the reverb effect to none.  This doesn't impact non-debug enabled
-		// builds, as a reverb effect w/ ApplyReverb false is equivalent to a null reverb effect that's been applied.
-		else if (bUsingDefaultReverb && !CurrentAudioVolumeSettings.ReverbSettings.bApplyReverb)
-		{
-			CurrentAudioVolumeSettings.ReverbSettings.bApplyReverb = true;
-			CurrentAudioVolumeSettings.ReverbSettings.ReverbEffect = nullptr;
-		}
+			// This is for debug visualization only - the audio debugger will show a reverb effect is active, if bApplyReverb is false for the world (default) reverb.
+			// In order to not mislead people using au.debug.reverb, we manually set the reverb effect to none.  This doesn't impact non-debug enabled
+			// builds, as a reverb effect w/ ApplyReverb false is equivalent to a null reverb effect that's been applied.
+			else if (bUsingDefaultReverb && !CurrentAudioVolumeSettings.ReverbSettings.bApplyReverb)
+			{
+				CurrentAudioVolumeSettings.ReverbSettings.bApplyReverb = true;
+				CurrentAudioVolumeSettings.ReverbSettings.ReverbEffect = nullptr;
+			}
 #endif // ENABLE_AUDIO_DEBUG
 
-		// Update the master reverb if it has changed
-		if (CurrentAudioVolumeSettings.bChanged || (CurrentAudioVolumeSettings.ReverbSettings != PreviousPlayerAudioVolumeSettings.ReverbSettings))
-		{
-			Effects->SetReverbSettings(CurrentAudioVolumeSettings.ReverbSettings);
-		}
+			// Update the master reverb if it has changed
+			if (CurrentAudioVolumeSettings.bChanged || (CurrentAudioVolumeSettings.ReverbSettings != PreviousPlayerAudioVolumeSettings.ReverbSettings))
+			{
+				Effects->SetReverbSettings(CurrentAudioVolumeSettings.ReverbSettings);
+			}
 
-		// Update the audio effects - reverb, EQ etc
-		Effects->Update();
+			// Update the audio effects - reverb, EQ etc
+			Effects->Update();
 
-		// If we any submix override settings apply those overrides to the indicated submixes
-		if (bUpdateSubmixEffects)
-		{
-			// Clear out any previous submix effect chain overrides if the audio volume changed
-			if (PreviousPlayerAudioVolumeSettings.SubmixOverrideSettings.Num() > 0)
+			// If we any submix override settings apply those overrides to the indicated submixes
+			if (bUpdateSubmixEffects)
 			{
-				for (FAudioVolumeSubmixOverrideSettings& OverrideSettings : PreviousPlayerAudioVolumeSettings.SubmixOverrideSettings)
+				// Clear out any previous submix effect chain overrides if the audio volume changed
+				if (PreviousPlayerAudioVolumeSettings.SubmixOverrideSettings.Num() > 0)
 				{
-					ClearSubmixEffectChainOverride(OverrideSettings.Submix, OverrideSettings.CrossfadeTime);
+					for (FAudioVolumeSubmixOverrideSettings& OverrideSettings : PreviousPlayerAudioVolumeSettings.SubmixOverrideSettings)
+					{
+						ClearSubmixEffectChainOverride(OverrideSettings.Submix, OverrideSettings.CrossfadeTime);
+					}
 				}
-			}
 
-			if (CurrentAudioVolumeSettings.SubmixOverrideSettings.Num() > 0)
-			{
-				for (FAudioVolumeSubmixOverrideSettings& OverrideSettings : CurrentAudioVolumeSettings.SubmixOverrideSettings)
+				if (CurrentAudioVolumeSettings.SubmixOverrideSettings.Num() > 0)
 				{
-					if (OverrideSettings.Submix && OverrideSettings.SubmixEffectChain.Num() > 0)
+					for (FAudioVolumeSubmixOverrideSettings& OverrideSettings : CurrentAudioVolumeSettings.SubmixOverrideSettings)
 					{
+						if (OverrideSettings.Submix && OverrideSettings.SubmixEffectChain.Num() > 0)
+						{
 
-						FSoundEffectSubmixInitData InitData;
-						InitData.SampleRate = GetSampleRate();
+							FSoundEffectSubmixInitData InitData;
+							InitData.SampleRate = GetSampleRate();
 
-						TArray<FSoundEffectSubmixPtr> SubmixEffectPresetChainOverride;
+							TArray<FSoundEffectSubmixPtr> SubmixEffectPresetChainOverride;
 
-						// Build the instances of the new submix preset chain override
-						for (USoundEffectSubmixPreset* SubmixEffectPreset : OverrideSettings.SubmixEffectChain)
-						{
-							if (SubmixEffectPreset)
+							// Build the instances of the new submix preset chain override
+							for (USoundEffectSubmixPreset* SubmixEffectPreset : OverrideSettings.SubmixEffectChain)
 							{
-								InitData.ParentPresetUniqueId = SubmixEffectPreset->GetUniqueID();
-
-								TSoundEffectSubmixPtr SoundEffectSubmix = USoundEffectPreset::CreateInstance<FSoundEffectSubmixInitData, FSoundEffectSubmix>(InitData, *SubmixEffectPreset);
-								SoundEffectSubmix->SetEnabled(true);
-								SubmixEffectPresetChainOverride.Add(SoundEffectSubmix);
+								if (SubmixEffectPreset)
+								{
+									InitData.ParentPresetUniqueId = SubmixEffectPreset->GetUniqueID();
+
+									TSoundEffectSubmixPtr SoundEffectSubmix = USoundEffectPreset::CreateInstance<FSoundEffectSubmixInitData, FSoundEffectSubmix>(InitData, *SubmixEffectPreset);
+									SoundEffectSubmix->SetEnabled(true);
+									SubmixEffectPresetChainOverride.Add(SoundEffectSubmix);
+								}
 							}
-						}
 
-						SetSubmixEffectChainOverride(OverrideSettings.Submix, SubmixEffectPresetChainOverride, OverrideSettings.CrossfadeTime);
+							SetSubmixEffectChainOverride(OverrideSettings.Submix, SubmixEffectPresetChainOverride, OverrideSettings.CrossfadeTime);
+						}
 					}
 				}
 			}
diff --git a/Engine/Source/Runtime/Engine/Public/AudioDevice.h b/Engine/Source/Runtime/Engine/Public/AudioDevice.h
--- a/Engine/Source/Runtime/Engine/Public/AudioDevice.h
+++ b/Engine/Source/Runtime/Engine/Public/AudioDevice.h
@@ -2170,8 +2170,11 @@ private:
 	/** Reverb Effects activated without volumes - Game Thread owned */
 	TMap<FName, FActivatedReverb> ActivatedReverbs;
 
+	/** Most recently activated in-flight reverb going from the GT to the AT */
+	TSharedPtr<FActivatedReverb> PendingActivatedReverb;
+
 	/** The activated reverb that currently has the highest priority - Audio Thread owned */
-	FActivatedReverb HighestPriorityActivatedReverb;
+	FActivatedReverb HighestPriorityActivatedReverb_AudioThread;
 
 	/** The current audio volume settings the listener is in. */
 	FAudioVolumeSettings CurrentAudioVolumeSettings;
@@ -2273,7 +2276,7 @@ private:
 
 	void UpdateReferencedSoundWaves();
 	TArray<USoundWave*> ReferencedSoundWaves_AudioThread;
-	FCriticalSection ReferencedSoundWaveCritSec;
+	FCriticalSection AudioThreadReferenceCritSec;
 
 	TArray<USoundWave*> PrecachingSoundWaves;
 
