Skeletal Mesh: Allow using USkeletalMeshSocket in BP.

#rnx

diff --git a/Engine/Source/Runtime/Engine/Classes/Engine/SkeletalMeshSocket.h b/Engine/Source/Runtime/Engine/Classes/Engine/SkeletalMeshSocket.h
--- a/Engine/Source/Runtime/Engine/Classes/Engine/SkeletalMeshSocket.h
+++ b/Engine/Source/Runtime/Engine/Classes/Engine/SkeletalMeshSocket.h
@@ -9,7 +9,7 @@
 #include "UObject/ScriptMacros.h"
 #include "SkeletalMeshSocket.generated.h"
 
-UCLASS(hidecategories=Object, hidecategories=Actor, MinimalAPI)
+UCLASS(hidecategories=Object, hidecategories=Actor, BlueprintType, MinimalAPI)
 class USkeletalMeshSocket : public UObject
 {
 	GENERATED_BODY()
