ShadowMap.cpp restore use of single channel G8 texture for single channel shadow maps
need to use TC_Grayscale to get that due to CL 30603716

diff --git a/Engine/Source/Runtime/Engine/Private/ShadowMap.cpp b/Engine/Source/Runtime/Engine/Private/ShadowMap.cpp
--- a/Engine/Source/Runtime/Engine/Private/ShadowMap.cpp
+++ b/Engine/Source/Runtime/Engine/Private/ShadowMap.cpp
@@ -327,7 +327,13 @@ void FShadowMapPendingTexture::StartEncoding(const FStaticLightingBuildContext*
 
 		Texture->Source.Init2DWithMipChain(GetSizeX(), GetSizeY(), NumChannelsUsed == 1 ? TSF_G8 : TSF_BGRA8);
 		Texture->MipGenSettings = TMGS_LeaveExistingMips;
+
+		// Uncompressed G8 or BGRA8 platform data :
+		Texture->CompressionSettings = NumChannelsUsed == 1 ? TC_Grayscale : TC_Default;
 		Texture->CompressionNone = true;
+		
+		// outer MipLock
+		FTextureSource::FMipLock Mip0Lock(FTextureSource::ELockState::ReadWrite,&Texture->Source,0);
 
 		// Copy the mip-map data into the UShadowMapTexture2D's mip-map array.
 		for(int32 MipIndex = 0;MipIndex < MipData.Num();MipIndex++)
@@ -348,6 +354,7 @@ void FShadowMapPendingTexture::StartEncoding(const FStaticLightingBuildContext*
 					}
 					else
 					{
+						// 4 distances in RGBA channels
 						((FColor*)DestMipData)[Y * MipSizeX + X] = FColor(SourceSample.Samples[0].Distance, SourceSample.Samples[1].Distance, SourceSample.Samples[2].Distance, SourceSample.Samples[3].Distance);
 					}
 				}
@@ -358,6 +365,8 @@ void FShadowMapPendingTexture::StartEncoding(const FStaticLightingBuildContext*
 	}
 
 	// Update the texture resource.
+	// @todo Oodle : manually calling CachePlatformData is not standard; not sure if this is a bug or was intentional for some reason
+	//		more standard thing should be to just call PostEditChange() here and let UpdateResource do its thing.
 	Texture->CachePlatformData(true, true, false, Compressor);
 
 	bFinishedEncoding = true;
