Add public assets to cook when configured to do so
#rnx

diff --git a/Engine/Source/Runtime/Engine/Classes/Engine/AssetManager.h b/Engine/Source/Runtime/Engine/Classes/Engine/AssetManager.h
--- a/Engine/Source/Runtime/Engine/Classes/Engine/AssetManager.h
+++ b/Engine/Source/Runtime/Engine/Classes/Engine/AssetManager.h
@@ -781,6 +781,9 @@ protected:
 	/** Scans the respective PackagePath for public assets to include in PackagesToCook */
 	ENGINE_API virtual void GatherPublicAssetsForPackage(FName PackagePath, TArray<FName>& PackagesToCook) const;
 
+	/** Scans the respective PackagePath for public assets to include in AssetsToCook */
+	ENGINE_API virtual void GatherPublicAssetsForPackage(FName PackagePath, TSet<FSoftObjectPath>& AssetsToCook) const;
+
 	/** Called right before PIE starts, will refresh asset directory and can be overriden to preload assets */
 	ENGINE_API virtual void PreBeginPIE(bool bStartSimulate);
 
diff --git a/Engine/Source/Runtime/Engine/Private/AssetManager.cpp b/Engine/Source/Runtime/Engine/Private/AssetManager.cpp
--- a/Engine/Source/Runtime/Engine/Private/AssetManager.cpp
+++ b/Engine/Source/Runtime/Engine/Private/AssetManager.cpp
@@ -4352,7 +4352,7 @@ void UAssetManager::ModifyCookReferences(FName PackageName, TArray<FName>& Packa
 	PackagesToCook.SetNum(Algo::Unique(PackagesToCook));
 }
 
-void UAssetManager::GatherPublicAssetsForPackage(FName PackagePath, TArray<FName>& PackagesToCook) const
+void UAssetManager::GatherPublicAssetsForPackage(FName PackagePath, TSet<FSoftObjectPath>& AssetsToCook) const
 {
 	FARFilter Filter;
 	Filter.PackagePaths.Add(PackagePath);
@@ -4360,7 +4360,7 @@ void UAssetManager::GatherPublicAssetsForPackage(FName PackagePath, TArray<FName
 	Filter.bIncludeOnlyOnDiskAssets = true;
 	Filter.WithoutPackageFlags |= PKG_NotExternallyReferenceable;
 
-	GetAssetRegistry().EnumerateAssets(Filter, [&PackagesToCook](const FAssetData& AssetData)
+	GetAssetRegistry().EnumerateAssets(Filter, [&AssetsToCook](const FAssetData& AssetData)
 	{
 		// this package can be externally referenced; include it in the cook
 		if (!AssetData.PackageName.IsNone())
@@ -4369,7 +4369,7 @@ void UAssetManager::GatherPublicAssetsForPackage(FName PackagePath, TArray<FName
 				TEXT("GatherPublicAssetsForPackage: Adding public package [%s] (instigator: [%s]"),
 				*AssetData.PackageName.ToString(),
 				*AssetData.AssetName.ToString());
-			PackagesToCook.AddUnique(AssetData.PackageName);
+			AssetsToCook.Add(AssetData.GetSoftObjectPath());
 		}
 		else
 		{
@@ -4382,6 +4382,31 @@ void UAssetManager::GatherPublicAssetsForPackage(FName PackagePath, TArray<FName
 	});
 }
 
+void UAssetManager::GatherPublicAssetsForPackage(FName PackagePath, TArray<FName>& PackagesToCook) const
+{
+	TSet<FSoftObjectPath> AssetsToCook;
+	GatherPublicAssetsForPackage(PackagePath, AssetsToCook);
+
+	if (AssetsToCook.Num() > 0)
+	{
+		TSet<FName> PackagesToCookSet;
+		PackagesToCookSet.Reserve(PackagesToCook.Num() + AssetsToCook.Num());
+		PackagesToCookSet.Append(PackagesToCook);
+
+		PackagesToCook.Reserve(PackagesToCook.Num() + AssetsToCook.Num());
+		for (const FSoftObjectPath& AssetToCook : AssetsToCook)
+		{
+			const FName PackageToCook = AssetToCook.GetLongPackageFName();
+			bool bAlreadyInSet = false;
+			PackagesToCookSet.Add(PackageToCook, &bAlreadyInSet);
+			if (!bAlreadyInSet)
+			{
+				PackagesToCook.Add(PackageToCook);
+			}
+		}
+	}
+}
+
 bool UAssetManager::ShouldCookForPlatform(const UPackage* Package, const ITargetPlatform* TargetPlatform)
 {
 	return true;
