[Engine]
* Disabled pso precaching if WITH_STATE_STREAM_ACTOR is 1 (will revisit)

diff --git a/Engine/Source/Runtime/Engine/Public/PSOPrecacheFwd.h b/Engine/Source/Runtime/Engine/Public/PSOPrecacheFwd.h
--- a/Engine/Source/Runtime/Engine/Public/PSOPrecacheFwd.h
+++ b/Engine/Source/Runtime/Engine/Public/PSOPrecacheFwd.h
@@ -7,12 +7,13 @@
 #pragma once
 
 #include "Containers/Array.h"
+#include "EngineDefines.h"
 #include "HAL/Platform.h"
 #include "PipelineStateCache.h"
 
 // General switch that decides whether to compile out some PSO precaching code (most importantly, reduce sizeofs of common classes)
 #ifndef UE_WITH_PSO_PRECACHING
-	#define UE_WITH_PSO_PRECACHING	(PLATFORM_SUPPORTS_PSO_PRECACHING || PLATFORM_SUPPORTS_DYNAMIC_SHADER_PRELOADING)
+	#define UE_WITH_PSO_PRECACHING	(PLATFORM_SUPPORTS_PSO_PRECACHING || PLATFORM_SUPPORTS_DYNAMIC_SHADER_PRELOADING) && !WITH_STATE_STREAM_ACTOR
 #endif // UE_WITH_PSO_PRECACHING
 
 // Only enable pso validation when the platform support its. Note: UE_WITH_PSO_PRECACHING can be define when dynamic shader preloading is supported. 
