AnimNext: Ensure we dont try to write to a skeletal mesh component that is also running anim BP

Changes SetEnableAnimation to defer to the end of the frame

diff --git a/Engine/Plugins/Experimental/AnimNext/Source/AnimNext/Private/Graph/RigUnit_AnimNextWriteSkeletalMeshComponentPose.cpp b/Engine/Plugins/Experimental/AnimNext/Source/AnimNext/Private/Graph/RigUnit_AnimNextWriteSkeletalMeshComponentPose.cpp
--- a/Engine/Plugins/Experimental/AnimNext/Source/AnimNext/Private/Graph/RigUnit_AnimNextWriteSkeletalMeshComponentPose.cpp
+++ b/Engine/Plugins/Experimental/AnimNext/Source/AnimNext/Private/Graph/RigUnit_AnimNextWriteSkeletalMeshComponentPose.cpp
@@ -35,6 +35,13 @@ FRigUnit_AnimNextWriteSkeletalMeshComponentPose_Execute()
 		return;
 	}
 
+	// We cannot write to the skeletal mesh component if it is using anim BP
+	if(SkeletalMeshComponent->bEnableAnimation)
+	{
+		UE_LOG(LogAnimation, Warning, TEXT("AnimNext: Could not write to skeletal mesh component - bEnableAnimation is true"))
+		return;
+	}
+
 	FDataHandle RefPoseHandle = FDataRegistry::Get()->GetOrGenerateReferencePose(SkeletalMeshComponent);
 	const UE::AnimNext::FReferencePose& RefPose = RefPoseHandle.GetRef<UE::AnimNext::FReferencePose>();
 
diff --git a/Engine/Source/Runtime/Engine/Classes/Components/SkeletalMeshComponent.h b/Engine/Source/Runtime/Engine/Classes/Components/SkeletalMeshComponent.h
--- a/Engine/Source/Runtime/Engine/Classes/Components/SkeletalMeshComponent.h
+++ b/Engine/Source/Runtime/Engine/Classes/Components/SkeletalMeshComponent.h
@@ -979,6 +979,7 @@ public:
 
 	/**
 	* Whether the built-in animation of this component should run when the component ticks. See USkeletalMeshComponent::bEnableAnimation for more details.
+	* Note that this call defers setting the flag to the end of the frame as the flag cannot change during world ticking (as other tasks may be in flight)
 	*/
 	UFUNCTION(BlueprintCallable, Category = "Components|SkeletalMesh", meta = (Keywords = "Animation"))
 	ENGINE_API void SetEnableAnimation(bool bNewEnableAnimation);
diff --git a/Engine/Source/Runtime/Engine/Private/Components/SkeletalMeshComponent.cpp b/Engine/Source/Runtime/Engine/Private/Components/SkeletalMeshComponent.cpp
--- a/Engine/Source/Runtime/Engine/Private/Components/SkeletalMeshComponent.cpp
+++ b/Engine/Source/Runtime/Engine/Private/Components/SkeletalMeshComponent.cpp
@@ -3269,8 +3269,17 @@ void USkeletalMeshComponent::SetAnimInstanceClass(class UClass* NewClass)
 
 void USkeletalMeshComponent::SetEnableAnimation(bool bNewEnableAnimation)
 {
-	bEnableAnimation = bNewEnableAnimation;
-	PrimaryComponentTick.bRunOnAnyThread = !bEnableAnimation;
+	// Defer setting flag to the end of the frame as other tasks can be in-flight during a tick
+	FTSTicker::GetCoreTicker().AddTicker(ANSI_TO_TCHAR(__FUNCTION__), 0.0f,
+		[bNewEnableAnimation, WeakComponent = TWeakObjectPtr<USkeletalMeshComponent>(this)](float InDeltaTime)
+		{
+			if(USkeletalMeshComponent* Component = WeakComponent.Get())
+			{
+				Component->bEnableAnimation = bNewEnableAnimation;
+				Component->PrimaryComponentTick.bRunOnAnyThread = !bNewEnableAnimation;
+			}
+			return false;
+		});
 }
 
 UAnimInstance* USkeletalMeshComponent::GetAnimInstance() const
