Added some improvements to the physics query handler for remote queries.

diff --git a/Engine/Source/Runtime/CoreUObject/Private/UObject/RemoteObjectTransfer.cpp b/Engine/Source/Runtime/CoreUObject/Private/UObject/RemoteObjectTransfer.cpp
--- a/Engine/Source/Runtime/CoreUObject/Private/UObject/RemoteObjectTransfer.cpp
+++ b/Engine/Source/Runtime/CoreUObject/Private/UObject/RemoteObjectTransfer.cpp
@@ -33,6 +33,9 @@ TMulticastDelegate<void(const TArray<UObject*>& MigratedObjects, FRemoteServerId
 TDelegate<void(FRemoteObjectId ObjectId, FRemoteServerId DatabaseId, FRemoteObjectData& ObjectData)> StoreRemoteObjectDataDelegate;
 TDelegate<void(FRemoteObjectId ObjectId, FRemoteServerId DatabaseId, FRemoteServerId DestinationServerId)> RestoreRemoteObjectDataDelegate;
 
+TDelegate<void()> FlushRemotePhysicsQueryDelegate;
+TDelegate<void()> ClearRemotePhysicsQueryQueueDelegate;
+
 UObject* MigrateObjectFromObjectData(FRemoteObjectData& ObjectData, FRemoteServerId RemoteServerId, UObject* DestinationOuter);
 
 /**
@@ -521,6 +524,8 @@ void ExecuteTransactional(const TFunction<void(void)>& Work)
 			continue;
 		}
 
+		FlushRemotePhysicsQueryDelegate.ExecuteIfBound();
+
 		UE_LOG(LogRemoteObject, VeryVerbose, TEXT("   : READY"));
 
 		// all of the requested objects are locked, try to perform the work
@@ -539,6 +544,8 @@ void ExecuteTransactional(const TFunction<void(void)>& Work)
 
 				bWorkSucceeded = true;
 
+				ClearRemotePhysicsQueryQueueDelegate.ExecuteIfBound();
+
 				if (GTransferQueue.LockedObjects.Num() == GTransferQueue.RequestedObjects.Num())
 				{
 					// the work succeeded and no more object requests were made, we're all done!
diff --git a/Engine/Source/Runtime/CoreUObject/Public/UObject/RemoteObjectTransfer.h b/Engine/Source/Runtime/CoreUObject/Public/UObject/RemoteObjectTransfer.h
--- a/Engine/Source/Runtime/CoreUObject/Public/UObject/RemoteObjectTransfer.h
+++ b/Engine/Source/Runtime/CoreUObject/Public/UObject/RemoteObjectTransfer.h
@@ -139,5 +139,11 @@ namespace UE::RemoteObject::Transfer
 	/** Delegate that restores object data from a database. */
 	extern COREUOBJECT_API TDelegate<void(FRemoteObjectId /*ObjectId*/, FRemoteServerId /*DatabaseId*/, FRemoteServerId /*DestinationServerId*/)> RestoreRemoteObjectDataDelegate;
 
+	/** Delegate that flushes pending physics query requests. */
+	extern COREUOBJECT_API TDelegate<void()> FlushRemotePhysicsQueryDelegate;
+	
+	/** Delegate that marks the end of a successful transaction. Allows the physics query queue to be cleared. */
+	extern COREUOBJECT_API TDelegate<void()> ClearRemotePhysicsQueryQueueDelegate;
+
 } // namespace UE::RemoveObject::Transfer
 
diff --git a/Engine/Source/Runtime/Engine/Public/Physics/DefaultPhysicsQueryHandler.h b/Engine/Source/Runtime/Engine/Public/Physics/DefaultPhysicsQueryHandler.h
--- a/Engine/Source/Runtime/Engine/Public/Physics/DefaultPhysicsQueryHandler.h
+++ b/Engine/Source/Runtime/Engine/Public/Physics/DefaultPhysicsQueryHandler.h
@@ -13,13 +13,13 @@ class UDefaultPhysicsQueryHandler : public UPhysicsQueryHandler
 public:
 	GENERATED_BODY()
 
-	virtual bool Overlap(const EQueryInfo InfoType, const EThreadQueryContext ThreadContext, const UWorld* World, const FOverlapQueryData& OverlapData, const FCommonQueryData& CommonData, TArray<FOverlapResult>& OutOverlaps) override;
+	ENGINE_API virtual bool Overlap(const EQueryInfo InfoType, const EThreadQueryContext ThreadContext, const UWorld* World, const FOverlapQueryData& OverlapData, const FCommonQueryData& CommonData, TArray<FOverlapResult>& OutOverlaps) override;
 
-	virtual bool RaycastTest(const EThreadQueryContext ThreadContext, const UWorld* World, const FRayQueryData& RayData, const FCommonQueryData& CommonData, FHitResult& OutHits) override;
-	virtual bool RaycastSingle(const EThreadQueryContext ThreadContext, const UWorld* World, const FRayQueryData& RayData, const FCommonQueryData& CommonData, FHitResult& OutHits) override;
-	virtual bool RaycastMulti(const EThreadQueryContext ThreadContext, const UWorld* World, const FRayQueryData& RayData, const FCommonQueryData& CommonData, TArray<FHitResult>& OutHits) override;
+	ENGINE_API virtual bool RaycastTest(const EThreadQueryContext ThreadContext, const UWorld* World, const FRayQueryData& RayData, const FCommonQueryData& CommonData, FHitResult& OutHits) override;
+	ENGINE_API virtual bool RaycastSingle(const EThreadQueryContext ThreadContext, const UWorld* World, const FRayQueryData& RayData, const FCommonQueryData& CommonData, FHitResult& OutHits) override;
+	ENGINE_API virtual bool RaycastMulti(const EThreadQueryContext ThreadContext, const UWorld* World, const FRayQueryData& RayData, const FCommonQueryData& CommonData, TArray<FHitResult>& OutHits) override;
 
-	virtual bool SweepTest(const EThreadQueryContext ThreadContext, const UWorld* World, const FSweepQueryData& SweepData, const FCommonQueryData& CommonData, FHitResult& OutHits) override;
-	virtual bool SweepSingle(const EThreadQueryContext ThreadContext, const UWorld* World, const FSweepQueryData& SweepData, const FCommonQueryData& CommonData, FHitResult& OutHits) override;
-	virtual bool SweepMulti(const EThreadQueryContext ThreadContext, const UWorld* World, const FSweepQueryData& SweepData, const FCommonQueryData& CommonData, TArray<FHitResult>& OutHits) override;
+	ENGINE_API virtual bool SweepTest(const EThreadQueryContext ThreadContext, const UWorld* World, const FSweepQueryData& SweepData, const FCommonQueryData& CommonData, FHitResult& OutHits) override;
+	ENGINE_API virtual bool SweepSingle(const EThreadQueryContext ThreadContext, const UWorld* World, const FSweepQueryData& SweepData, const FCommonQueryData& CommonData, FHitResult& OutHits) override;
+	ENGINE_API virtual bool SweepMulti(const EThreadQueryContext ThreadContext, const UWorld* World, const FSweepQueryData& SweepData, const FCommonQueryData& CommonData, TArray<FHitResult>& OutHits) override;
 };
diff --git a/Engine/Source/Runtime/Engine/Public/Physics/GenericPhysicsInterface.h b/Engine/Source/Runtime/Engine/Public/Physics/GenericPhysicsInterface.h
--- a/Engine/Source/Runtime/Engine/Public/Physics/GenericPhysicsInterface.h
+++ b/Engine/Source/Runtime/Engine/Public/Physics/GenericPhysicsInterface.h
@@ -14,7 +14,7 @@ class UWorld;
 namespace Chaos
 {
 	enum class EQueryInfo : uint8;
-	enum class EThreadQueryContext;
+	enum class EThreadQueryContext : uint8;
 
 	struct FCommonQueryData;
 	struct FOverlapQueryData;
diff --git a/Engine/Source/Runtime/Engine/Public/Physics/PhysicsQueryHandler.h b/Engine/Source/Runtime/Engine/Public/Physics/PhysicsQueryHandler.h
--- a/Engine/Source/Runtime/Engine/Public/Physics/PhysicsQueryHandler.h
+++ b/Engine/Source/Runtime/Engine/Public/Physics/PhysicsQueryHandler.h
@@ -27,13 +27,13 @@ public:
 	using FRayQueryData = Chaos::FRayQueryData;
 	using FSweepQueryData = Chaos::FSweepQueryData;
 
-	virtual bool Overlap(const EQueryInfo InfoType, const EThreadQueryContext ThreadContext, const UWorld* World, const FOverlapQueryData& OverlapData, const FCommonQueryData& CommonData, TArray<FOverlapResult>& OutOverlaps);
+	ENGINE_API virtual bool Overlap(const EQueryInfo InfoType, const EThreadQueryContext ThreadContext, const UWorld* World, const FOverlapQueryData& OverlapData, const FCommonQueryData& CommonData, TArray<FOverlapResult>& OutOverlaps);
 
-	virtual bool RaycastTest(const EThreadQueryContext ThreadContext, const UWorld* World, const FRayQueryData& RayData, const FCommonQueryData& CommonData, FHitResult& OutHits);
-	virtual bool RaycastSingle(const EThreadQueryContext ThreadContext, const UWorld* World, const FRayQueryData& RayData, const FCommonQueryData& CommonData, FHitResult& OutHits);
-	virtual bool RaycastMulti(const EThreadQueryContext ThreadContext, const UWorld* World, const FRayQueryData& RayData, const FCommonQueryData& CommonData, TArray<FHitResult>& OutHits);
+	ENGINE_API virtual bool RaycastTest(const EThreadQueryContext ThreadContext, const UWorld* World, const FRayQueryData& RayData, const FCommonQueryData& CommonData, FHitResult& OutHits);
+	ENGINE_API virtual bool RaycastSingle(const EThreadQueryContext ThreadContext, const UWorld* World, const FRayQueryData& RayData, const FCommonQueryData& CommonData, FHitResult& OutHits);
+	ENGINE_API virtual bool RaycastMulti(const EThreadQueryContext ThreadContext, const UWorld* World, const FRayQueryData& RayData, const FCommonQueryData& CommonData, TArray<FHitResult>& OutHits);
 
-	virtual bool SweepTest(const EThreadQueryContext ThreadContext, const UWorld* World, const FSweepQueryData& SweepData, const FCommonQueryData& CommonData, FHitResult& OutHits);
-	virtual bool SweepSingle(const EThreadQueryContext ThreadContext, const UWorld* World, const FSweepQueryData& SweepData, const FCommonQueryData& CommonData, FHitResult& OutHits);
-	virtual bool SweepMulti(const EThreadQueryContext ThreadContext, const UWorld* World, const FSweepQueryData& RaySweepData, const FCommonQueryData& CommonData, TArray<FHitResult>& OutHits);
+	ENGINE_API virtual bool SweepTest(const EThreadQueryContext ThreadContext, const UWorld* World, const FSweepQueryData& SweepData, const FCommonQueryData& CommonData, FHitResult& OutHits);
+	ENGINE_API virtual bool SweepSingle(const EThreadQueryContext ThreadContext, const UWorld* World, const FSweepQueryData& SweepData, const FCommonQueryData& CommonData, FHitResult& OutHits);
+	ENGINE_API virtual bool SweepMulti(const EThreadQueryContext ThreadContext, const UWorld* World, const FSweepQueryData& RaySweepData, const FCommonQueryData& CommonData, TArray<FHitResult>& OutHits);
 };
diff --git a/Engine/Source/Runtime/Engine/Public/Physics/SceneQueryData.h b/Engine/Source/Runtime/Engine/Public/Physics/SceneQueryData.h
--- a/Engine/Source/Runtime/Engine/Public/Physics/SceneQueryData.h
+++ b/Engine/Source/Runtime/Engine/Public/Physics/SceneQueryData.h
@@ -16,7 +16,7 @@ namespace Chaos
 		IsAnything		//is any of the data blocking or touching? only return a bool so don't bother collecting
 	};
 
-	enum class EThreadQueryContext
+	enum class EThreadQueryContext : uint8
 	{
 		GTData,		//use interpolated GT data
 		PTDataWithGTObjects,	//use pt data, but convert back to GT when possible
