Made bEnableMaterialParameterCaching public

diff --git a/Engine/Source/Runtime/Engine/Classes/Components/MeshComponent.h b/Engine/Source/Runtime/Engine/Classes/Components/MeshComponent.h
--- a/Engine/Source/Runtime/Engine/Classes/Components/MeshComponent.h
+++ b/Engine/Source/Runtime/Engine/Classes/Components/MeshComponent.h
@@ -165,7 +165,6 @@ public:
 	ENGINE_API virtual int32 GetMeshPaintTextureCoordinateIndex() const { return 0; }
 
 public:
-	/** Material parameter setting and caching */
 
 	/** Set all occurrences of Scalar Material Parameters with ParameterName in the set of materials of the SkeletalMesh to ParameterValue */
 	UFUNCTION(BlueprintCallable, Category = "Rendering|Material")
@@ -221,9 +220,14 @@ protected:
 		float ScalarParameterDefaultValue = 0.f;
 	};
 
+public:
+
+	/** Whether or not to cache material parameter to speed up setting scalar or vector value on materials */
 	UPROPERTY(EditAnywhere, AdvancedDisplay, BlueprintReadOnly, Category = MaterialParameters)
 	uint8 bEnableMaterialParameterCaching : 1;
 
+protected:
+
 	/** Flag whether or not the cached material parameter indices map is dirty (defaults to true, and is set from SetMaterial/Set(Skeletal)Mesh */
 	uint8 bCachedMaterialParameterIndicesAreDirty : 1;
 
