[Iris] Move PostInit allocations into UObject/NetworkPostInit from NetDriver LLM

#tests local build
#rnx

diff --git a/Engine/Source/Runtime/Engine/Private/DataChannel.cpp b/Engine/Source/Runtime/Engine/Private/DataChannel.cpp
--- a/Engine/Source/Runtime/Engine/Private/DataChannel.cpp
+++ b/Engine/Source/Runtime/Engine/Private/DataChannel.cpp
@@ -21,6 +21,7 @@
 #include "Net/Core/Misc/GuidReferences.h"
 #include "Net/Core/NetCoreModule.h"
 #include "Net/Core/Connection/NetEnums.h"
+#include "UObject/Package.h"
 #include "UObject/UObjectIterator.h"
 #include "EngineStats.h"
 #include "Engine/Engine.h"
@@ -41,6 +42,7 @@
 #include "HAL/LowLevelMemStats.h"
 #include "Net/NetPing.h"
 #include "PacketHandler.h"
+#include "ProfilingDebugging/AssetMetadataTrace.h"
 
 #if UE_NET_REPACTOR_NAME_DEBUG
 #include "Net/NetNameDebug.h"
@@ -3450,6 +3452,10 @@ void UActorChannel::ProcessBunch( FInBunch & Bunch )
 	if (Actor && bSpawnedNewActor)
 	{
 		SCOPE_CYCLE_COUNTER(Stat_PostNetInit);
+		LLM_SCOPE_BYNAME(TEXT("UObject/NetworkPostInit"));
+		LLM_SCOPE_DYNAMIC_STAT_OBJECTPATH(Actor->GetPackage(), ELLMTagSet::Assets);
+		LLM_SCOPE_DYNAMIC_STAT_OBJECTPATH(Actor->GetClass(), ELLMTagSet::AssetClasses);
+		UE_TRACE_METADATA_SCOPE_ASSET(Actor, Actor->GetClass());
 		Actor->PostNetInit();
 	}
 }
diff --git a/Engine/Source/Runtime/Engine/Private/DemoNetDriver.cpp b/Engine/Source/Runtime/Engine/Private/DemoNetDriver.cpp
--- a/Engine/Source/Runtime/Engine/Private/DemoNetDriver.cpp
+++ b/Engine/Source/Runtime/Engine/Private/DemoNetDriver.cpp
@@ -32,6 +32,8 @@
 #include "Engine/ChildConnection.h"
 #include "Net/ReplayPlaylistTracker.h"
 #include "Net/NetworkGranularMemoryLogging.h"
+#include "HAL/LowLevelMemStats.h"
+#include "ProfilingDebugging/AssetMetadataTrace.h"
 
 #include UE_INLINE_GENERATED_CPP_BY_NAME(DemoNetDriver)
 
@@ -3487,7 +3489,13 @@ void UDemoNetDriver::RespawnNecessaryNetStartupActors(TArray<AActor*>& SpawnedAc
 			Actor->MarkComponentsRenderStateDirty();
 
 			// BeginPlay
-			Actor->PostNetInit();
+			{
+				LLM_SCOPE_BYNAME(TEXT("UObject/NetworkPostInit"));
+				LLM_SCOPE_DYNAMIC_STAT_OBJECTPATH(Actor->GetPackage(), ELLMTagSet::Assets);
+				LLM_SCOPE_DYNAMIC_STAT_OBJECTPATH(Actor->GetClass(), ELLMTagSet::AssetClasses);
+				UE_TRACE_METADATA_SCOPE_ASSET(Actor, Actor->GetClass());
+				Actor->PostNetInit();
+			}
 
 			// Call actor rep notifies
 			if (RepLayout.IsValid() && ReceivingRepState)
diff --git a/Engine/Source/Runtime/Engine/Private/Net/Iris/ReplicationSystem/NetActorFactory.cpp b/Engine/Source/Runtime/Engine/Private/Net/Iris/ReplicationSystem/NetActorFactory.cpp
--- a/Engine/Source/Runtime/Engine/Private/Net/Iris/ReplicationSystem/NetActorFactory.cpp
+++ b/Engine/Source/Runtime/Engine/Private/Net/Iris/ReplicationSystem/NetActorFactory.cpp
@@ -475,6 +475,10 @@ void UNetActorFactory::PostInit(const FPostInitContext& Context)
 	if (Context.Handle.IsDynamic())
 	{
 		AActor* Actor = CastChecked<AActor>(Context.Instance);
+		LLM_SCOPE_BYNAME(TEXT("UObject/NetworkPostInit"));
+		LLM_SCOPE_DYNAMIC_STAT_OBJECTPATH(Actor->GetPackage(), ELLMTagSet::Assets);
+		LLM_SCOPE_DYNAMIC_STAT_OBJECTPATH(Actor->GetClass(), ELLMTagSet::AssetClasses);
+		UE_TRACE_METADATA_SCOPE_ASSET(Actor, Actor->GetClass());
 		Actor->PostNetInit();
 	}
 }
