fix deadlock on transacting static mesh changes w/ collision

diff --git a/Engine/Source/Runtime/Engine/Private/PhysicsEngine/BodySetup.cpp b/Engine/Source/Runtime/Engine/Private/PhysicsEngine/BodySetup.cpp
--- a/Engine/Source/Runtime/Engine/Private/PhysicsEngine/BodySetup.cpp
+++ b/Engine/Source/Runtime/Engine/Private/PhysicsEngine/BodySetup.cpp
@@ -1447,6 +1447,16 @@ UBodySetup::FBodySetupTryCookResult UBodySetup::TryGetCookedData(FName Format)
 	{
 		SCOPE_CYCLE_COUNTER(STAT_PhysXCooking);
 
+		// Note: Check ContainsPhysicsTriMeshData before looking at the number of convex elems, to ensure the side effects of ContainsPhysicsTriMeshData happen
+		// (specifically, for static mesh this will ensure the mesh render data is already built)
+		// Note: We also do this *before* the below scope lock, to avoid deadlock 
+		// (because ContainsPhysicsTriMeshData can call GetRenderData() for static mesh, which will wait for async build to run, 
+		//  which will try to take this same lock on a separate thread when the async mesh build calls this same method)
+		if ((CDP == nullptr || CDP->ContainsPhysicsTriMeshData(bMeshCollideAll) == false) && AggGeom.ConvexElems.Num() == 0)
+		{
+			return { nullptr, CookDataResult };
+		}
+
 		// Only one thread should build the physics data. Below calling GetFormat and TryGetDDCBuiltData will
 		// both mutate shared state so we protect this operation with a lock.
 		FScopeLock ScopeLock(&CookedFormatDataMutex);
@@ -1456,13 +1466,6 @@ UBodySetup::FBodySetupTryCookResult UBodySetup::TryGetCookedData(FName Format)
 		// Double-checked locking avoids always taking the lock - we only do it for unbuilt data
 		if(!UseCookedData->Contains(Format))
 		{
-			// Note: Check ContainsPhysicsTriMeshData before looking at the number of convex elems, to ensure the side effects of ContainsPhysicsTriMeshData happen
-			// (specifically, for static mesh this will ensure the mesh render data is already built)
-			if ((CDP == nullptr || CDP->ContainsPhysicsTriMeshData(bMeshCollideAll) == false) && AggGeom.ConvexElems.Num() == 0)
-			{
-				return { nullptr, CookDataResult };
-			}
-
 			// We do not want a FGCObject to be created to prevent garbage collection of our own UBodySetup*
 			// because the scope and lifetime is well defined and FGCObject can't be created on other threads
 			// during garbage collection, which would prevent this function from being run asynchronously.
