When UE_WITH_REMOTE_OBJECT_HANDLE is enabled, NetDrivers can now replicate actors in a transaction.

diff --git a/Engine/Plugins/Runtime/MultiServerReplication/Source/MultiServerReplication/Private/MultiServerBeaconHost.cpp b/Engine/Plugins/Runtime/MultiServerReplication/Source/MultiServerReplication/Private/MultiServerBeaconHost.cpp
--- a/Engine/Plugins/Runtime/MultiServerReplication/Source/MultiServerReplication/Private/MultiServerBeaconHost.cpp
+++ b/Engine/Plugins/Runtime/MultiServerReplication/Source/MultiServerReplication/Private/MultiServerBeaconHost.cpp
@@ -38,6 +38,7 @@ bool AMultiServerBeaconHost::InitHost()
 				NetDriver->Notify = this;
 				NetDriver->InitialConnectTimeout = BeaconConnectionInitialTimeout;
 				NetDriver->ConnectionTimeout = BeaconConnectionTimeout;
+				NetDriver->SetReplicateTransactionally(false);
 				return true;
 			}
 			else
diff --git a/Engine/Source/Runtime/Engine/Classes/Engine/NetDriver.h b/Engine/Source/Runtime/Engine/Classes/Engine/NetDriver.h
--- a/Engine/Source/Runtime/Engine/Classes/Engine/NetDriver.h
+++ b/Engine/Source/Runtime/Engine/Classes/Engine/NetDriver.h
@@ -2225,6 +2225,11 @@ public:
 	/** Called internally by channels to track how many hit net.QueuedBunchTimeFailsafeSeconds */
 	void AddQueuedBunchFailsafeChannel() { ++QueuedBunchFailsafeNumChannels; }
 
+	void SetReplicateTransactionally(bool bInReplicateTransactionally)
+	{
+		bReplicateTransactionally = bInReplicateTransactionally;
+	}
+
 protected:
 	
 	/** Stream of random numbers to be used by this instance of UNetDriver */
@@ -2359,4 +2364,6 @@ private:
 	/** Accumulated number of channels that hit the net.QueuedBunchTimeFailsafeSeconds this frame */
 	uint32 QueuedBunchFailsafeNumChannels = 0;
 
+	/** When using remote objects with UE_WITH_REMOTE_OBJECT_HANDLE, determines whether or not property replication happens within a transaction */
+	bool bReplicateTransactionally = true;
 };
diff --git a/Engine/Source/Runtime/Engine/Private/NetDriver.cpp b/Engine/Source/Runtime/Engine/Private/NetDriver.cpp
--- a/Engine/Source/Runtime/Engine/Private/NetDriver.cpp
+++ b/Engine/Source/Runtime/Engine/Private/NetDriver.cpp
@@ -50,6 +50,7 @@
 #include "Net/NetworkMetricsConfig.h"
 #include "Net/NetSubObjectRegistryGetter.h"
 #include "Net/NetworkGranularMemoryLogging.h"
+#include "UObject/RemoteObjectTransfer.h"
 #include "UObject/Stack.h"
 #if UE_WITH_IRIS
 #include "Iris/IrisConfig.h"
@@ -1016,7 +1017,21 @@ void UNetDriver::TickFlush(float DeltaSeconds)
 		else
 #endif // UE_WITH_IRIS
 		{
+#if UE_WITH_REMOTE_OBJECT_HANDLE
+			if (bReplicateTransactionally)
+			{
+				UE::RemoteObject::Transfer::ExecuteTransactional([this, DeltaSeconds]()
+				{
+					ServerReplicateActors(DeltaSeconds);
+				});
+			}
+			else
+			{
+				ServerReplicateActors(DeltaSeconds);
+			}
+#else
 			ServerReplicateActors(DeltaSeconds);
+#endif // UE_WITH_REMOTE_OBJECT_HANDLE
 		}
 #endif // WITH_SERVER_CODE
 	}
