ShadowMap.h nop comment only about unorm quantizer error

diff --git a/Engine/Source/Runtime/Engine/Public/ShadowMap.h b/Engine/Source/Runtime/Engine/Public/ShadowMap.h
--- a/Engine/Source/Runtime/Engine/Public/ShadowMap.h
+++ b/Engine/Source/Runtime/Engine/Public/ShadowMap.h
@@ -35,15 +35,16 @@ struct FQuantizedSignedDistanceFieldShadowSample
 	{
 		if (Index == 0)
 		{
-			return Distance / 255.0f;
+			return Distance * (1.f / 255.0f);
 		}
 		else
 		{
 			checkSlow(Index == 1);
-			return PenumbraSize / 255.0f;
+			return PenumbraSize * (1.f / 255.0f);
 		}
 	}
-
+	
+	// this is not a correct U8 quantizer (should round, not trunc); see below for correct version
 	void SetFilterableComponent(float InComponent, int32 Index)
 	{
 		if (Index == 0)
@@ -57,6 +58,24 @@ struct FQuantizedSignedDistanceFieldShadowSample
 		}
 	}
 
+	#if 0
+	// correct quantizer; needs to be tested to make sure other biases aren't compensating for the error
+	void SetFilterableComponent(float InComponent, int32 Index)
+	{
+		uint8 Quantized = FColor::QuantizeUNormFloatTo8(InComponent);
+
+		if (Index == 0)
+		{
+			Distance = Quantized;
+		}
+		else
+		{
+			checkSlow(Index == 1);
+			PenumbraSize = Quantized;
+		}
+	}
+	#endif
+
 	/** Equality operator */
 	bool operator==(const FQuantizedSignedDistanceFieldShadowSample& RHS) const
 	{
