ImageUtils::ImageResize and FImage::ResizeTo : get rid of old bad resizers and just call FImageCore::ResizeImage
this produces different output but is not used in significant paths that build content so does not create a large patch delta
note that a bad resizer still exists in the TextureCompressorModule Downscale legacy path

diff --git a/Engine/Source/Developer/MaterialUtilities/Private/MaterialUtilities.cpp b/Engine/Source/Developer/MaterialUtilities/Private/MaterialUtilities.cpp
--- a/Engine/Source/Developer/MaterialUtilities/Private/MaterialUtilities.cpp
+++ b/Engine/Source/Developer/MaterialUtilities/Private/MaterialUtilities.cpp
@@ -1912,6 +1912,7 @@ void FMaterialUtilities::ResizeFlattenMaterial(FFlattenMaterial& InFlattenMateri
 		if (InFlattenMaterial.GetPropertySize(EFlattenMaterialProperties::Diffuse).X != DiffuseSizeX)
 		{
 			TArray<FColor> NewSamples;
+			// @todo : last arg should be true not false
 			FImageUtils::ImageResize(InFlattenMaterial.GetPropertySize(EFlattenMaterialProperties::Diffuse).X, InFlattenMaterial.GetPropertySize(EFlattenMaterialProperties::Diffuse).Y, InFlattenMaterial.GetPropertySamples(EFlattenMaterialProperties::Diffuse), DiffuseSizeX, DiffuseSizeX, NewSamples, false);
 			InFlattenMaterial.GetPropertySamples(EFlattenMaterialProperties::Diffuse).Reset(NewSamples.Num());
 			InFlattenMaterial.GetPropertySamples(EFlattenMaterialProperties::Diffuse).Append(NewSamples);
@@ -1948,6 +1949,7 @@ void FMaterialUtilities::ResizeFlattenMaterial(FFlattenMaterial& InFlattenMateri
 		if (InFlattenMaterial.GetPropertySamples(EFlattenMaterialProperties::Normal).Num() && MaterialSettings.bNormalMap && InFlattenMaterial.GetPropertySize(EFlattenMaterialProperties::Normal).X != NormalSizeX)
 		{
 			TArray<FColor> NewSamples;
+			// @todo : last arg should be false not true
 			FImageUtils::ImageResize(InFlattenMaterial.GetPropertySize(EFlattenMaterialProperties::Normal).X, InFlattenMaterial.GetPropertySize(EFlattenMaterialProperties::Normal).Y, InFlattenMaterial.GetPropertySamples(EFlattenMaterialProperties::Normal), NormalSizeX, NormalSizeX, NewSamples, true);
 			InFlattenMaterial.GetPropertySamples(EFlattenMaterialProperties::Normal).Reset(NewSamples.Num());
 			InFlattenMaterial.GetPropertySamples(EFlattenMaterialProperties::Normal).Append(NewSamples);
diff --git a/Engine/Source/Runtime/Engine/Private/ImageUtils.cpp b/Engine/Source/Runtime/Engine/Private/ImageUtils.cpp
--- a/Engine/Source/Runtime/Engine/Private/ImageUtils.cpp
+++ b/Engine/Source/Runtime/Engine/Private/ImageUtils.cpp
@@ -599,221 +599,67 @@ bool FImageUtils::GetRenderTargetImage(UTextureRenderTarget* TexRT, FImage & Out
 }
 
 
-/**
- * Resizes the given image using a simple average filter and stores it in the destination array.
- *
- * @param SrcWidth		Source image width.
- * @param SrcHeight		Source image height.
- * @param SrcData		Source image data.
- * @param DstWidth		Destination image width.
- * @param DstHeight		Destination image height.
- * @param DstData		Destination image data.
- * @param bLinearSpace	If true, convert colors into linear space before interpolating (slower but more accurate)
- */
+// prefer FImageCore::ResizeImage
 void FImageUtils::ImageResize(int32 SrcWidth, int32 SrcHeight, const TArray<FColor> &SrcData, int32 DstWidth, int32 DstHeight, TArray<FColor> &DstData, bool bLinearSpace, bool bForceOpaqueOutput)
 {
-	// * DEPRECATED do not use this, use FImageCore::ResizeImage instead
-
 	DstData.Empty(DstWidth*DstHeight);
 	DstData.AddZeroed(DstWidth*DstHeight);
 
 	ImageResize(SrcWidth, SrcHeight, TArrayView<const FColor>(SrcData), DstWidth, DstHeight, TArrayView<FColor>(DstData), bLinearSpace, bForceOpaqueOutput);
 }
 
-/**
- * Resizes the given image using a simple average filter and stores it in the destination array.  This version constrains aspect ratio.
- * Accepts TArrayViews but requires that DstData be pre-sized appropriately
- *
- * @param SrcWidth	Source image width.
- * @param SrcHeight	Source image height.
- * @param SrcData	Source image data.
- * @param DstWidth	Destination image width.
- * @param DstHeight Destination image height.
- * @param DstData	Destination image data. (must already be sized to DstWidth*DstHeight)
- */
+// prefer FImageCore::ResizeImage
 void FImageUtils::ImageResize(int32 SrcWidth, int32 SrcHeight, const TArrayView<const FColor> &SrcData, int32 DstWidth, int32 DstHeight, const TArrayView<FColor> &DstData, bool bLinearSpace, bool bForceOpaqueOutput)
 {
-	//@todo OodleImageResize : deprecate ImageResize
-	// * DEPRECATED do not use this, use FImageCore::ResizeImage instead
-
 	check(SrcData.Num() >= SrcWidth * SrcHeight);
 	check(DstData.Num() >= DstWidth * DstHeight);
-
-	float SrcX = 0;
-	float SrcY = 0;
-
-	const float StepSizeX = SrcWidth / (float)DstWidth;
-	const float StepSizeY = SrcHeight / (float)DstHeight;
-
-	for(int32 Y=0; Y<DstHeight;Y++)
-	{
-		int32 PixelPos = Y * DstWidth;
-		SrcX = 0.0f;	
 	
-		for(int32 X=0; X<DstWidth; X++)
-		{
-			int32 PixelCount = 0;
-			float EndX = SrcX + StepSizeX;
-			float EndY = SrcY + StepSizeY;
-			
-			// Generate a rectangular region of pixels and then find the average color of the region.
-			int32 PosY = FMath::TruncToInt(SrcY+0.5f);
-			PosY = FMath::Clamp<int32>(PosY, 0, (SrcHeight - 1));
-
-			int32 PosX = FMath::TruncToInt(SrcX+0.5f);
-			PosX = FMath::Clamp<int32>(PosX, 0, (SrcWidth - 1));
-
-			int32 EndPosY = FMath::TruncToInt(EndY+0.5f);
-			EndPosY = FMath::Clamp<int32>(EndPosY, 0, (SrcHeight - 1));
-
-			int32 EndPosX = FMath::TruncToInt(EndX+0.5f);
-			EndPosX = FMath::Clamp<int32>(EndPosX, 0, (SrcWidth - 1));
-
-			FColor FinalColor;
-			if(bLinearSpace)
-			{
-				FLinearColor LinearStepColor(0.0f,0.0f,0.0f,0.0f);
-				for(int32 PixelX = PosX; PixelX <= EndPosX; PixelX++)
-				{
-					for(int32 PixelY = PosY; PixelY <= EndPosY; PixelY++)
-					{
-						int32 StartPixel =  PixelX + PixelY * SrcWidth;
-
-						// Convert from gamma space to linear space before the addition.
-						LinearStepColor += SrcData[StartPixel];
-						PixelCount++;
-					}
-				}
-				LinearStepColor /= (float)PixelCount;
-
-				// Convert back from linear space to gamma space.
-				FinalColor = LinearStepColor.ToFColor(true);
-			}
-			else
-			{
-				FVector4 StepColor(0,0,0,0);
-				for(int32 PixelX = PosX; PixelX <= EndPosX; PixelX++)
-				{
-					for(int32 PixelY = PosY; PixelY <= EndPosY; PixelY++)
-					{
-						int32 StartPixel =  PixelX + PixelY * SrcWidth;
-						StepColor.X += (float)SrcData[StartPixel].R;
-						StepColor.Y += (float)SrcData[StartPixel].G;
-						StepColor.Z += (float)SrcData[StartPixel].B;
-						StepColor.W += (float)SrcData[StartPixel].A;
-						PixelCount++;
-					}
-				}
-				uint8 FinalR = FMath::Clamp(FMath::TruncToInt32(StepColor.X / (float)PixelCount), 0, 255);
-				uint8 FinalG = FMath::Clamp(FMath::TruncToInt32(StepColor.Y / (float)PixelCount), 0, 255);
-				uint8 FinalB = FMath::Clamp(FMath::TruncToInt32(StepColor.Z / (float)PixelCount), 0, 255);
-				uint8 FinalA = FMath::Clamp(FMath::TruncToInt32(StepColor.W / (float)PixelCount), 0, 255);
-				FinalColor = FColor(FinalR, FinalG, FinalB, FinalA);
-			}
-
-			if ( bForceOpaqueOutput )
-			{
-				FinalColor.A = 255;
-			}
+	if ( SrcData.Num() == 0 || DstData.Num() == 0 )
+	{
+		return;
+	}
 
-			// Store the final averaged pixel color value.
-			DstData[PixelPos] = FinalColor;
+	// bLinearSpace means convert SRGB colors to Linear before resizing
+	//
+	//	(there are a few callers using bLinearSpace=false , most of which seem to be unintentional/bugs
+	//	 MaterialUtilities.cpp attempts to use it right but does NOT)
+	EGammaSpace Gamma = bLinearSpace ? EGammaSpace::sRGB : EGammaSpace::Linear; 
 
-			SrcX = EndX;	
-			PixelPos++;
-		}
+	FImageView SrcImage(&SrcData[0],SrcWidth,SrcHeight,Gamma);
+	FImageView DstImage(&DstData[0],DstWidth,DstHeight,Gamma);
+	
+	FImageCore::ResizeImage(SrcImage,DstImage);
 
-		SrcY += StepSizeY;
+	if ( bForceOpaqueOutput )
+	{
+		FImageCore::SetAlphaOpaque(DstImage);
 	}
 }
 
-/**
- * Resizes the given image using a simple average filter and stores it in the destination array.
- *
- * @param SrcWidth		Source image width.
- * @param SrcHeight		Source image height.
- * @param SrcData		Source image data.
- * @param DstWidth		Destination image width.
- * @param DstHeight		Destination image height.
- * @param DstData		Destination image data.
- */
+// prefer FImageCore::ResizeImage
 void FImageUtils::ImageResize(int32 SrcWidth, int32 SrcHeight, const TArray64<FLinearColor>& SrcData, int32 DstWidth, int32 DstHeight, TArray64<FLinearColor>& DstData)
 {
-	// * DEPRECATED do not use this, use FImageCore::ResizeImage instead
-
 	DstData.Empty(DstWidth * DstHeight);
 	DstData.AddZeroed(DstWidth * DstHeight);
 
 	ImageResize(SrcWidth, SrcHeight, TArrayView64<const FLinearColor>(SrcData), DstWidth, DstHeight, TArrayView64<FLinearColor>(DstData));
 }
 
-/**
- * Resizes the given image using a simple average filter and stores it in the destination array.  This version constrains aspect ratio.
- * Accepts TArrayViews but requires that DstData be pre-sized appropriately
- *
- * @param SrcWidth	Source image width.
- * @param SrcHeight	Source image height.
- * @param SrcData	Source image data.
- * @param DstWidth	Destination image width.
- * @param DstHeight Destination image height.
- * @param DstData	Destination image data. (must already be sized to DstWidth*DstHeight)
- */
+// prefer FImageCore::ResizeImage
 void FImageUtils::ImageResize(int32 SrcWidth, int32 SrcHeight, const TArrayView64<const FLinearColor>& SrcData, int32 DstWidth, int32 DstHeight, const TArrayView64<FLinearColor>& DstData)
 {
-	//@todo OodleImageResize : deprecate ImageResize
-	// * DEPRECATED do not use this, use FImageCore::ResizeImage instead
-
 	check(SrcData.Num() >= SrcWidth * SrcHeight);
 	check(DstData.Num() >= DstWidth * DstHeight);
-
-	float SrcX = 0;
-	float SrcY = 0;
-	const float StepSizeX = SrcWidth / (float)DstWidth;
-	const float StepSizeY = SrcHeight / (float)DstHeight;
-
-	for (int32 Y = 0; Y < DstHeight; Y++)
+	
+	if ( SrcData.Num() == 0 || DstData.Num() == 0 )
 	{
-		int32 PixelPos = Y * DstWidth;
-		SrcX = 0.0f;
-
-		for (int32 X = 0; X < DstWidth; X++)
-		{
-			int32 PixelCount = 0;
-			float EndX = SrcX + StepSizeX;
-			float EndY = SrcY + StepSizeY;
-
-			// Generate a rectangular region of pixels and then find the average color of the region.
-			int32 PosY = FMath::TruncToInt(SrcY + 0.5f);
-			PosY = FMath::Clamp<int32>(PosY, 0, (SrcHeight - 1));
-
-			int32 PosX = FMath::TruncToInt(SrcX + 0.5f);
-			PosX = FMath::Clamp<int32>(PosX, 0, (SrcWidth - 1));
-
-			int32 EndPosY = FMath::TruncToInt(EndY + 0.5f);
-			EndPosY = FMath::Clamp<int32>(EndPosY, 0, (SrcHeight - 1));
-
-			int32 EndPosX = FMath::TruncToInt(EndX + 0.5f);
-			EndPosX = FMath::Clamp<int32>(EndPosX, 0, (SrcWidth - 1));
-
-			FLinearColor FinalColor(0.0f, 0.0f, 0.0f, 0.0f);
-			for (int32 PixelX = PosX; PixelX <= EndPosX; PixelX++)
-			{
-				for (int32 PixelY = PosY; PixelY <= EndPosY; PixelY++)
-				{
-					int32 StartPixel = PixelX + PixelY * SrcWidth;
-					FinalColor += SrcData[StartPixel];
-					PixelCount++;
-				}
-			}
-			FinalColor /= (float)PixelCount;
-
-			// Store the final averaged pixel color value.
-			DstData[PixelPos] = FinalColor;
-			SrcX = EndX;
-			PixelPos++;
-		}
-		SrcY += StepSizeY;
+		return;
 	}
+
+	FImageView SrcImage(&SrcData[0],SrcWidth,SrcHeight);
+	FImageView DstImage(&DstData[0],DstWidth,DstHeight);
+	
+	FImageCore::ResizeImage(SrcImage,DstImage);
 }
 
 /**
@@ -939,7 +785,6 @@ UTexture * FImageUtils::CreateTexture(ETextureClass TextureClass, const FImageVi
 
 void FImageUtils::CropAndScaleImage( int32 SrcWidth, int32 SrcHeight, int32 DesiredWidth, int32 DesiredHeight, const TArray<FColor> &SrcData, TArray<FColor> &DstData  )
 {
-	// DEPRECATED , uses the bad ImageResize
 	//	this is used by the old Thumbnail code, not the new paths
 
 	// Get the aspect ratio, and calculate the dimension of the image to crop
diff --git a/Engine/Source/Runtime/Engine/Public/ImageUtils.h b/Engine/Source/Runtime/Engine/Public/ImageUtils.h
--- a/Engine/Source/Runtime/Engine/Public/ImageUtils.h
+++ b/Engine/Source/Runtime/Engine/Public/ImageUtils.h
@@ -172,26 +172,29 @@ public:
 	ENGINE_API static bool GetTexture2DSourceImage(	UTexture2D* Texture, FImage & OutImage);
 
 	/**
-	 * Resizes the given image using a simple average filter and stores it in the destination array.  This version constrains aspect ratio.
+	 * Resizes the given image using FImageCore::ResizeImage.
+	 * Destination array will be emptied and resized.
 	 *
-	 * DEPRECATED do not use this, use FImageCore::ResizeImage instead
+	 * prefer FImageCore::ResizeImage instead.
+	 * prefer using FImage or FImageView not FColor arrays.
 	 *
 	 * @param SrcWidth				Source image width.
 	 * @param SrcHeight				Source image height.
 	 * @param SrcData				Source image data.
 	 * @param DstWidth				Destination image width.
 	 * @param DstHeight				Destination image height.
-	 * @param DstData				Destination image data.
-	 * @param bLinearSpace			Output in linear space instead of sRGB.
-	 * @param bForceOpaqueOutput	Always output 255 for the alpha channel.
+	 * @param DstData				Destination image data. (will be resized)
+	 * @param bLinearSpace			Convert linear to sRGB for gamma-correct resize.  If this is false source U8s are treated as linear.  Should usually be true.
+	 * @param bForceOpaqueOutput	Always output 255 for the alpha channel.  BEWARE defaults to true so alpha is killed by default.
 	 */
 	ENGINE_API static void ImageResize(int32 SrcWidth, int32 SrcHeight, const TArray<FColor> &SrcData,  int32 DstWidth, int32 DstHeight, TArray<FColor> &DstData, bool bLinearSpace, bool bForceOpaqueOutput=true);
 
 	/**
-	 * Resizes the given image using a simple average filter and stores it in the destination array.  This version constrains aspect ratio.
+	 * Resizes the given image using FImageCore::ResizeImage.
 	 * Accepts TArrayViews but requires that DstData be pre-sized appropriately
 	 *
-	 * DEPRECATED do not use this, use FImageCore::ResizeImage instead
+	 * prefer FImageCore::ResizeImage instead.
+	 * prefer using FImage or FImageView not FColor arrays.
 	 *
 	 * @param SrcWidth				Source image width.
 	 * @param SrcHeight				Source image height.
@@ -199,30 +202,33 @@ public:
 	 * @param DstWidth				Destination image width.
 	 * @param DstHeight				Destination image height.
 	 * @param DstData				Destination image data. (must already be sized to DstWidth*DstHeight)
-	 * @param bLinearSpace			Output in linear space instead of sRGB.
-	 * @param bForceOpaqueOutput	Always output 255 for the alpha channel.
+	 * @param bLinearSpace			Convert linear to sRGB for gamma-correct resize.  If this is false source U8s are treated as linear.  Should usually be true.
+	 * @param bForceOpaqueOutput	Always output 255 for the alpha channel.  BEWARE defaults to true so alpha is killed by default.
 	 */
 	ENGINE_API static void ImageResize(int32 SrcWidth, int32 SrcHeight, const TArrayView<const FColor> &SrcData, int32 DstWidth, int32 DstHeight, const TArrayView<FColor> &DstData, bool bLinearSpace, bool bForceOpaqueOutput=true);
 
 	/**
-	 * Resizes the given image using a simple average filter and stores it in the destination array.  This version constrains aspect ratio.
+	 * Resizes the given image using FImageCore::ResizeImage.
+	 * Destination array will be emptied and resized.
 	 *
-	 * DEPRECATED do not use this, use FImageCore::ResizeImage instead
+	 * prefer FImageCore::ResizeImage instead.
+	 * prefer using FImage or FImageView not FColor arrays.
 	 *
 	 * @param SrcWidth	Source image width.
 	 * @param SrcHeight	Source image height.
 	 * @param SrcData	Source image data.
 	 * @param DstWidth	Destination image width.
 	 * @param DstHeight Destination image height.
-	 * @param DstData	Destination image data.
+	 * @param DstData	Destination image data. (will be resized)
 	 */
 	ENGINE_API static void ImageResize(int32 SrcWidth, int32 SrcHeight, const TArray64<FLinearColor>& SrcData, int32 DstWidth, int32 DstHeight, TArray64<FLinearColor>& DstData);
 
 	/**
-	 * Resizes the given image using a simple average filter and stores it in the destination array.  This version constrains aspect ratio.
+	 * Resizes the given image using FImageCore::ResizeImage.
 	 * Accepts TArrayViews but requires that DstData be pre-sized appropriately
 	 *
-	 * DEPRECATED do not use this, use FImageCore::ResizeImage instead
+	 * prefer FImageCore::ResizeImage instead.
+	 * prefer using FImage or FImageView not FColor arrays.
 	 *
 	 * @param SrcWidth	Source image width.
 	 * @param SrcHeight	Source image height.
@@ -287,7 +293,6 @@ public:
 	 * @param SrcData			Raw image array.
 	 * @param DstData			compressed image array.
 	 *
-	 * DEPRECATED, avoid this, uses the bad image resize
 	 */
 	ENGINE_API static void CropAndScaleImage( int32 SrcWidth, int32 SrcHeight, int32 DesiredWidth, int32 DesiredHeight, const TArray<FColor> &SrcData, TArray<FColor> &DstData  );
 
diff --git a/Engine/Source/Runtime/ImageCore/Private/ImageCore.cpp b/Engine/Source/Runtime/ImageCore/Private/ImageCore.cpp
--- a/Engine/Source/Runtime/ImageCore/Private/ImageCore.cpp
+++ b/Engine/Source/Runtime/ImageCore/Private/ImageCore.cpp
@@ -927,56 +927,6 @@ void FImage::TransformToWorkingColorSpace(const FVector2d& SourceRedChromaticity
 	return FImageCore::TransformToWorkingColorSpace(*this, SourceRedChromaticity, SourceGreenChromaticity, SourceBlueChromaticity, SourceWhiteChromaticity, Method, EqualityTolerance);
 }
 
-// bad legacy resizer, do not use this except when required to maintain legacy behavior ; use "ResizeImage" instead
-static FLinearColor Bad_SampleImage(const FLinearColor* Pixels, int Width, int Height, float X, float Y)
-{
-	const int64 TexelX0 = FMath::FloorToInt(X);
-	const int64 TexelY0 = FMath::FloorToInt(Y);
-	const int64 TexelX1 = FMath::Min<int64>(TexelX0 + 1, Width - 1);
-	const int64 TexelY1 = FMath::Min<int64>(TexelY0 + 1, Height - 1);
-	checkSlow(TexelX0 >= 0 && TexelX0 < Width);
-	checkSlow(TexelY0 >= 0 && TexelY0 < Height);
-
-	const float FracX1 = FMath::Frac(X);
-	const float FracY1 = FMath::Frac(Y);
-	const float FracX0 = 1.0f - FracX1;
-	const float FracY0 = 1.0f - FracY1;
-	const FLinearColor& Color00 = Pixels[TexelY0 * Width + TexelX0];
-	const FLinearColor& Color01 = Pixels[TexelY1 * Width + TexelX0];
-	const FLinearColor& Color10 = Pixels[TexelY0 * Width + TexelX1];
-	const FLinearColor& Color11 = Pixels[TexelY1 * Width + TexelX1];
-	return
-		Color00 * (FracX0 * FracY0) +
-		Color01 * (FracX0 * FracY1) +
-		Color10 * (FracX1 * FracY0) +
-		Color11 * (FracX1 * FracY1);
-}
-
-// bad legacy resizer, do not use this except when required to maintain legacy behavior ; use "ResizeImage" instead
-static void Bad_ResizeImage_Bilinear(const FImageView & SrcImage, const FImageView & DestImage)
-{
-	// Src and Dest should both now be RGBA32F and Linear gamma
-	check( SrcImage.Format == ERawImageFormat::RGBA32F );
-	check( DestImage.Format == ERawImageFormat::RGBA32F );
-	const FLinearColor* SrcPixels = (const FLinearColor*) SrcImage.RawData;
-	FLinearColor* DestPixels = (FLinearColor*) DestImage.RawData;
-	const float DestToSrcScaleX = (float)SrcImage.SizeX / (float)DestImage.SizeX;
-	const float DestToSrcScaleY = (float)SrcImage.SizeY / (float)DestImage.SizeY;
-
-	// not a correct bilinear Resize?  missing 0.5 pixel center shift
-	//		deprecate this and redirect to new resizer
-	for (int64 DestY = 0; DestY < DestImage.SizeY; ++DestY)
-	{
-		const float SrcY = (float)DestY * DestToSrcScaleY;
-		for (int64 DestX = 0; DestX < DestImage.SizeX; ++DestX)
-		{
-			const float SrcX = (float)DestX * DestToSrcScaleX;
-			const FLinearColor Color = Bad_SampleImage(SrcPixels, SrcImage.SizeX, SrcImage.SizeY, SrcX, SrcY);
-			DestPixels[DestY * DestImage.SizeX + DestX] = Color;
-		}
-	}
-}
-
 /* FImage constructors
  *****************************************************************************/
  
@@ -1108,65 +1058,16 @@ void FImageView::CopyTo(FImage& DestImage, ERawImageFormat::Type DestFormat, EGa
 	FImageCore::CopyImage(*this, DestImage);
 }
 
-// bad legacy resizer, do not use this except when required to maintain legacy behavior ; use "ResizeImage" instead
+// legacy API, prefer ResizeImage
 void FImage::ResizeTo(FImage& DestImage, int32 DestSizeX, int32 DestSizeY, ERawImageFormat::Type DestFormat, EGammaSpace DestGammaSpace) const
 {
 	FImageCore::ResizeTo(*this,DestImage,DestSizeX,DestSizeY,DestFormat,DestGammaSpace);
 }
 
-// bad legacy resizer, do not use this except when required to maintain legacy behavior ; use "ResizeImage" instead
+// legacy API, prefer ResizeImage
 void FImageCore::ResizeTo(const FImageView & SourceImage,FImage& DestImage, int32 DestSizeX, int32 DestSizeY, ERawImageFormat::Type DestFormat, EGammaSpace DestGammaSpace)
 {
-	check(SourceImage.NumSlices == 1); // only support 1 slice for now
-
-	FImage TempSrcImage;
-	FImageView SrcImageView = SourceImage;
-	if (SourceImage.Format != ERawImageFormat::RGBA32F)
-	{
-		SourceImage.CopyTo(TempSrcImage, ERawImageFormat::RGBA32F, EGammaSpace::Linear);
-		SrcImageView = TempSrcImage;
-	}
-	else if ( SourceImage.GammaSpace != EGammaSpace::Linear )
-	{
-		UE_LOG(LogImageCore, Warning, TEXT("Resize from source Format RGBA32F was called but source GammaSpace is not Linear"));
-	}
-
-	// now SrcImagePtr should be RGBA32F and Linear
-
-	if (DestFormat == ERawImageFormat::RGBA32F)
-	{
-		if ( DestGammaSpace != EGammaSpace::Linear )
-		{
-			UE_LOG(LogImageCore, Warning, TEXT("Resize to DestFormat RGBA32F was called but DestGammaSpace is not Linear"));
-		}
-
-		DestImage.SizeX = DestSizeX;
-		DestImage.SizeY = DestSizeY;
-		DestImage.NumSlices = 1;
-		DestImage.Format = DestFormat;
-		DestImage.GammaSpace = DestGammaSpace;
-		InitImageStorage(DestImage);
-		Bad_ResizeImage_Bilinear(SrcImageView, DestImage);
-	}
-	else
-	{
-		// first resize from RGBA32F to RGBA32F
-		FImage TempDestImage;
-		TempDestImage.SizeX = DestSizeX;
-		TempDestImage.SizeY = DestSizeY;
-		TempDestImage.NumSlices = 1;
-		TempDestImage.Format = ERawImageFormat::RGBA32F;
-		TempDestImage.GammaSpace = EGammaSpace::Linear;
-		InitImageStorage(TempDestImage);
-		Bad_ResizeImage_Bilinear(SrcImageView, TempDestImage);
-
-		// then convert to dest format/gamma :
-		if ( DestGammaSpace == EGammaSpace::Pow22 )
-		{
-			UE_LOG(LogImageCore, Warning, TEXT("Resize incorrectly used with Pow22 Dest Gamma"));
-		}
-		TempDestImage.CopyTo(DestImage, DestFormat, DestGammaSpace);
-	}
+	FImageCore::ResizeImageAllocDest(SourceImage,DestImage,DestSizeX,DestSizeY,DestFormat,DestGammaSpace);
 }
 
 IMAGECORE_API FImageView FImageView::GetSlice(int32 SliceIndex) const
diff --git a/Engine/Source/Runtime/ImageCore/Public/ImageCore.h b/Engine/Source/Runtime/ImageCore/Public/ImageCore.h
--- a/Engine/Source/Runtime/ImageCore/Public/ImageCore.h
+++ b/Engine/Source/Runtime/ImageCore/Public/ImageCore.h
@@ -267,7 +267,7 @@ struct FImageView : public FImageInfo
 	{
 	}
 
-	// Make an FImageView from an array of FColor pixels as BGRA8
+	// Make an FImageView pointing at an array of FColor pixels as BGRA8 (does not copy the input, points at it)
 	FImageView(const FColor * InColors,int32 InSizeX,int32 InSizeY,EGammaSpace InGammaSpace = EGammaSpace::sRGB)
 	{
 		RawData = (void *) InColors;
@@ -278,7 +278,7 @@ struct FImageView : public FImageInfo
 		GammaSpace = InGammaSpace; 
 	}
 	
-	// Make an FImageView from an array of FLinearColor pixels as RGBA32F
+	// Make an FImageView pointing at an array of FLinearColor pixels as RGBA32F (does not copy the input, points at it)
 	FImageView(const FLinearColor * InColors,int32 InSizeX,int32 InSizeY)
 	{
 		RawData = (void *) InColors;
@@ -289,7 +289,7 @@ struct FImageView : public FImageInfo
 		GammaSpace = EGammaSpace::Linear;
 	}
 	
-	// Make an FImageView from an array of FLinearColor pixels as RGBA16F
+	// Make an FImageView pointing at an array of FLinearColor pixels as RGBA16F (does not copy the input, points at it)
 	FImageView(const FFloat16Color * InColors,int32 InSizeX,int32 InSizeY)
 	{
 		RawData = (void *) InColors;
@@ -505,10 +505,7 @@ public:
 	 * Copies and resizes the image to a destination image with the specified size and format.
 	 * Resize is done using bilinear filtering; gamma correct (light linear) by converting through RGBA32F.
 	 *
-	 * DEPRECATED! Use ResizeImage() instead.
-	 * This function should only be used where legacy behavior must be maintained.
-	 *
-	 * ResizeTo has a bug a slight translation of the image.
+	 * legacy API, prefer ResizeImage
 	 *
 	 * @param DestImage - The destination image.
 	 * @param DestSizeX - Width of the resized image
@@ -924,10 +921,7 @@ IMAGECORE_API void SetAlphaOpaque(const FImageView & InImage);
 	* Copies and resizes the image to a destination image with the specified size and format.
 	* Resize is done using bilinear filtering; gamma correct (light linear) by converting through RGBA32F.
 	*
-	* DEPRECATED! Use ResizeImage() instead.
-	* This function should only be used where legacy behavior must be maintained.
-	*
-	* ResizeTo has a bug a slight translation of the image.
+	* legacy API, prefer ResizeImage
 	*
 	* @param DestImage - The destination image.
 	* @param DestSizeX - Width of the resized image
