[Navigation] update navigation data of a SceneComponent when it gets attached to another component only if it's relevant to navigation.

diff --git a/Engine/Plugins/Experimental/Water/Source/Runtime/Private/WaterBodyActor.cpp b/Engine/Plugins/Experimental/Water/Source/Runtime/Private/WaterBodyActor.cpp
--- a/Engine/Plugins/Experimental/Water/Source/Runtime/Private/WaterBodyActor.cpp
+++ b/Engine/Plugins/Experimental/Water/Source/Runtime/Private/WaterBodyActor.cpp
@@ -709,7 +709,10 @@ void AWaterBody::PostRegisterAllComponents()
 		WaterBodyComponent->OnPostRegisterAllComponents();
 
 		// Now that all components are registered and setup we can notify the navigation system.
-		FNavigationSystem::UpdateActorAndComponentData(*this);
+		if (WaterBodyComponent->IsNavigationRelevant())
+		{
+			FNavigationSystem::UpdateActorAndComponentData(*this);
+		}
 	}
 }
 
diff --git a/Engine/Source/Runtime/Engine/Private/Components/SceneComponent.cpp b/Engine/Source/Runtime/Engine/Private/Components/SceneComponent.cpp
--- a/Engine/Source/Runtime/Engine/Private/Components/SceneComponent.cpp
+++ b/Engine/Source/Runtime/Engine/Private/Components/SceneComponent.cpp
@@ -2381,9 +2381,12 @@ bool USceneComponent::AttachToComponent(USceneComponent* Parent, const FAttachme
 			UpdateOverlaps();
 
 			// Update our owner actor in the navigation system since its associated bounds have changed
-			if (AActor* Owner = GetOwner())
+			if (IsNavigationRelevant())
 			{
-				FNavigationSystem::UpdateActorAndComponentData(*Owner);
+				if (AActor* Owner = GetOwner())
+				{
+					FNavigationSystem::UpdateActorAndComponentData(*Owner);
+				}
 			}
 		}
 
