SVT: Fixed bug in FTileDataTexture::GetMaxTileDataTextureResourceSize where the memory limit defined by r.SparseVolumeTexture.Streaming.ReservedResourcesMemoryLimit was applied incorrectly.

#rnx

diff --git a/Engine/Source/Runtime/Engine/Private/SparseVolumeTexture/SparseVolumeTextureTileDataTexture.cpp b/Engine/Source/Runtime/Engine/Private/SparseVolumeTexture/SparseVolumeTextureTileDataTexture.cpp
--- a/Engine/Source/Runtime/Engine/Private/SparseVolumeTexture/SparseVolumeTextureTileDataTexture.cpp
+++ b/Engine/Source/Runtime/Engine/Private/SparseVolumeTexture/SparseVolumeTextureTileDataTexture.cpp
@@ -46,7 +46,8 @@ int64 FTileDataTexture::GetMaxTileDataTextureResourceSize(int32 InVoxelMemSize)
 		if (GSVTStreamingReservedResourcesMemoryLimit > 0)
 		{
 			constexpr int64 OneMiB = 1024LL * 1024LL;
-			MaxStreamingPoolResourceSize = FMath::Min(MaxStreamingPoolResourceSize, (int64)GSVTStreamingReservedResourcesMemoryLimit) * OneMiB;
+			const int64 ReservedResourcesMemoryLimit = FMath::Min((int64)GSVTStreamingReservedResourcesMemoryLimit, INT64_MAX / OneMiB) * OneMiB;
+			MaxStreamingPoolResourceSize = FMath::Min(MaxStreamingPoolResourceSize, ReservedResourcesMemoryLimit);
 		}
 	}
 	return MaxStreamingPoolResourceSize;
