Chaos::Fix pushing cube's different bahavior under different fps. Scale impulse applied in UCharacterMovementComponent::ApplyImpactPhysicsForces to be consistent with time scale.

diff --git a/Engine/Source/Runtime/Engine/Private/Components/CharacterMovementComponent.cpp b/Engine/Source/Runtime/Engine/Private/Components/CharacterMovementComponent.cpp
--- a/Engine/Source/Runtime/Engine/Private/Components/CharacterMovementComponent.cpp
+++ b/Engine/Source/Runtime/Engine/Private/Components/CharacterMovementComponent.cpp
@@ -369,6 +369,13 @@ namespace CharacterMovementCVars
 		TEXT("Enable a client-authoritative network flow for adding forces to pushed physics objects, requires physics prediction to be enabled in the project settings. NOTE: This is not recommended for networked physics and can produce inconsistency bugs where the client adds multiple forces for a single physics tick on the server, especially if async physics is not enabled."),
 		ECVF_Default);
 
+	static bool bApplyCharacterCollisionForceAsImpulse = false;
+	FAutoConsoleVariableRef CVarApplyCharacterCollisionForceAsImpulse(
+		TEXT("p.ApplyCharacterCollisionForceAsImpulse"),
+		bApplyCharacterCollisionForceAsImpulse,
+		TEXT("Whether to apply character impact physics force as impulse when the object hit is not moving."),
+		ECVF_Default);
+
 #if CSV_PROFILER_STATS
 	bool bClientRecordMovePackedRpcStatsToCsv = false;
 	FAutoConsoleVariableRef CVarClientRecordMovePackedRpcStatsToCsv(
@@ -7656,13 +7663,32 @@ void UCharacterMovementComponent::ApplyImpactPhysicsForces(const FHitResult& Imp
 				if (ComponentVelocity.IsNearlyZero(ZeroVelocityTolerance))
 				{
 					Force *= InitialPushForceFactor;
-					if (!bClientAuthScheduledPushForces)
+
+					if (CharacterMovementCVars::bApplyCharacterCollisionForceAsImpulse)
 					{
-						ImpactComponent->AddImpulseAtLocation(Force, ForcePoint, Impact.BoneName);
+						if (!bClientAuthScheduledPushForces)
+						{
+							ImpactComponent->AddImpulseAtLocation(Force, ForcePoint, Impact.BoneName);
+						}
+						else
+						{
+							ApplyAsyncPhysicsStateAction(ImpactComponent, Impact.BoneName, EPhysicsStateAction::AddImpulseAtPosition, Force, ForcePoint);
+						}
 					}
 					else
 					{
-						ApplyAsyncPhysicsStateAction(ImpactComponent, Impact.BoneName, EPhysicsStateAction::AddImpulseAtPosition, Force, ForcePoint);
+						//Originally the force is applied as impulse, which was inaccurate. The extra scaling is applied so that the force is calibrated to match the old behavior at 30fps. 
+						const float ExtraForceFactor = 30.f;
+						Force *= ExtraForceFactor;
+
+						if (!bClientAuthScheduledPushForces)
+						{
+							ImpactComponent->AddForceAtLocation(Force, ForcePoint, Impact.BoneName);
+						}
+						else
+						{
+							ApplyAsyncPhysicsStateAction(ImpactComponent, Impact.BoneName, EPhysicsStateAction::AddForceAtPosition, Force, ForcePoint);
+						}
 					}
 				}
 				else
