DSTM: Do not attempt to recalculate the transform of a migrated Component (it will have already been set properly and serialized).

diff --git a/Engine/Source/Runtime/Engine/Private/Actor.cpp b/Engine/Source/Runtime/Engine/Private/Actor.cpp
--- a/Engine/Source/Runtime/Engine/Private/Actor.cpp
+++ b/Engine/Source/Runtime/Engine/Private/Actor.cpp
@@ -1207,16 +1207,8 @@ void AActor::PostMigrate(FRemoteServerId DestinationServerId)
 		checkf(World != nullptr, TEXT("Migrated Actor %s (%s) doesn't have a World"), *GetPathName(), *FRemoteObjectId(this).ToString());
 
 		// Set the actor's world transform if it has a native rootcomponent.
-		USceneComponent* const SceneRootComponent = FixupNativeActorComponents(this);
-		FTransform FinalRootComponentTransform = GetTransform();
-		if (SceneRootComponent != nullptr)
-		{
-			// Respect any non-default transform value that the root component may have received from the archetype that's owned
-			// by the native CDO, so the final transform might not always necessarily equate to the passed-in UserSpawnTransform.
-			const FTransform RootTransform(SceneRootComponent->GetRelativeRotation(), SceneRootComponent->GetRelativeLocation(), SceneRootComponent->GetRelativeScale3D());
-			FinalRootComponentTransform = RootTransform * GetTransform();
-			SceneRootComponent->SetWorldTransform(FinalRootComponentTransform, false, nullptr, ETeleportType::ResetPhysics);
-		}
+		const USceneComponent* SceneRootComponent = FixupNativeActorComponents(this);
+		ensureMsgf(SceneRootComponent == RootComponent, TEXT("%hs: We were expecting RootComponent (%s) to be equal to FixupNativeActorComponents (%s). Transform fix-up code would be required to handle this case."), __func__, *GetNameSafe(RootComponent), *GetNameSafe(SceneRootComponent));
 
 		RegisterAllComponents();
 
