Add SetMetadata method to VertexBuffers.
* Allow to set the metadata to non-serialized buffers.

#rnx

diff --git a/Engine/Source/Runtime/Engine/Private/PositionVertexBuffer.cpp b/Engine/Source/Runtime/Engine/Private/PositionVertexBuffer.cpp
--- a/Engine/Source/Runtime/Engine/Private/PositionVertexBuffer.cpp
+++ b/Engine/Source/Runtime/Engine/Private/PositionVertexBuffer.cpp
@@ -186,6 +186,12 @@ void FPositionVertexBuffer::SerializeMetaData(FArchive& Ar)
 	Ar << Stride << NumVertices;
 }
 
+void FPositionVertexBuffer::SetMetaData(uint32 InStride, uint32 InNumVertices)
+{
+	Stride = InStride;
+	NumVertices = InNumVertices;
+}
+
 void FPositionVertexBuffer::ClearMetaData()
 {
 	Stride = NumVertices = 0;
diff --git a/Engine/Source/Runtime/Engine/Private/Rendering/ColorVertexBuffer.cpp b/Engine/Source/Runtime/Engine/Private/Rendering/ColorVertexBuffer.cpp
--- a/Engine/Source/Runtime/Engine/Private/Rendering/ColorVertexBuffer.cpp
+++ b/Engine/Source/Runtime/Engine/Private/Rendering/ColorVertexBuffer.cpp
@@ -242,6 +242,12 @@ void FColorVertexBuffer::SerializeMetaData(FArchive& Ar)
 	Ar << Stride << NumVertices;
 }
 
+void FColorVertexBuffer::SetMetaData(uint32 InStride, uint32 InNumVertices)
+{
+	Stride = InStride;
+	NumVertices = InNumVertices;
+}
+
 void FColorVertexBuffer::ClearMetaData()
 {
 	Stride = NumVertices = 0;
diff --git a/Engine/Source/Runtime/Engine/Private/Rendering/SkinWeightVertexBuffer.cpp b/Engine/Source/Runtime/Engine/Private/Rendering/SkinWeightVertexBuffer.cpp
--- a/Engine/Source/Runtime/Engine/Private/Rendering/SkinWeightVertexBuffer.cpp
+++ b/Engine/Source/Runtime/Engine/Private/Rendering/SkinWeightVertexBuffer.cpp
@@ -113,6 +113,11 @@ void FSkinWeightLookupVertexBuffer::SerializeMetaData(FArchive& Ar)
 	}
 }
 
+void FSkinWeightLookupVertexBuffer::SetMetaData(uint32 InNumVertices)
+{
+	NumVertices = InNumVertices;
+}
+
 void FSkinWeightLookupVertexBuffer::CopyMetaData(const FSkinWeightLookupVertexBuffer& Other)
 {
 	NumVertices = Other.NumVertices;
@@ -374,6 +379,16 @@ void FSkinWeightDataVertexBuffer::SerializeMetaData(FArchive& Ar)
 	}
 }
 
+void FSkinWeightDataVertexBuffer::SetMetaData(uint32 InNumVertices, uint32 InMaxBoneInfluences, bool bInUse16BitBoneIndex, bool bInUse16BitBoneWeight)
+{
+	NumVertices = InNumVertices;
+	bUse16BitBoneIndex = bInUse16BitBoneIndex;
+	bUse16BitBoneWeight = bInUse16BitBoneWeight;
+
+	SetMaxBoneInfluences(InMaxBoneInfluences);
+	NumBoneWeights = NumVertices * MaxBoneInfluences;
+}
+
 void FSkinWeightDataVertexBuffer::CopyMetaData(const FSkinWeightDataVertexBuffer& Other)
 {
 	bVariableBonesPerVertex = Other.bVariableBonesPerVertex;
diff --git a/Engine/Source/Runtime/Engine/Private/Rendering/StaticMeshVertexBuffer.cpp b/Engine/Source/Runtime/Engine/Private/Rendering/StaticMeshVertexBuffer.cpp
--- a/Engine/Source/Runtime/Engine/Private/Rendering/StaticMeshVertexBuffer.cpp
+++ b/Engine/Source/Runtime/Engine/Private/Rendering/StaticMeshVertexBuffer.cpp
@@ -247,6 +247,16 @@ void FStaticMeshVertexBuffer::SerializeMetaData(FArchive& Ar)
 	InitTangentAndTexCoordStrides();
 }
 
+void FStaticMeshVertexBuffer::SetMetaData(uint32 InNumTexCoords, uint32 InNumVertices, bool bInUseFullPrecisionUVs, bool bInUseHighPrecisionTangentBasis)
+{
+	NumTexCoords = InNumTexCoords;
+	NumVertices = InNumVertices;
+	bUseFullPrecisionUVs = bInUseFullPrecisionUVs;
+	bUseHighPrecisionTangentBasis = bInUseHighPrecisionTangentBasis;
+
+	InitTangentAndTexCoordStrides();
+}
+
 void FStaticMeshVertexBuffer::ClearMetaData()
 {
 	NumTexCoords = NumVertices = 0;
diff --git a/Engine/Source/Runtime/Engine/Public/Rendering/ColorVertexBuffer.h b/Engine/Source/Runtime/Engine/Public/Rendering/ColorVertexBuffer.h
--- a/Engine/Source/Runtime/Engine/Public/Rendering/ColorVertexBuffer.h
+++ b/Engine/Source/Runtime/Engine/Public/Rendering/ColorVertexBuffer.h
@@ -64,6 +64,8 @@ public:
 
 	void SerializeMetaData(FArchive& Ar);
 
+	ENGINE_API void SetMetaData(uint32 InStride, uint32 InNumVertices);
+
 	void ClearMetaData();
 
 	/**
diff --git a/Engine/Source/Runtime/Engine/Public/Rendering/PositionVertexBuffer.h b/Engine/Source/Runtime/Engine/Public/Rendering/PositionVertexBuffer.h
--- a/Engine/Source/Runtime/Engine/Public/Rendering/PositionVertexBuffer.h
+++ b/Engine/Source/Runtime/Engine/Public/Rendering/PositionVertexBuffer.h
@@ -72,6 +72,8 @@ public:
 
 	void SerializeMetaData(FArchive& Ar);
 
+	ENGINE_API void SetMetaData(uint32 InStride, uint32 InNumVertices);
+
 	void ClearMetaData();
 
 	/**
diff --git a/Engine/Source/Runtime/Engine/Public/Rendering/SkinWeightVertexBuffer.h b/Engine/Source/Runtime/Engine/Public/Rendering/SkinWeightVertexBuffer.h
--- a/Engine/Source/Runtime/Engine/Public/Rendering/SkinWeightVertexBuffer.h
+++ b/Engine/Source/Runtime/Engine/Public/Rendering/SkinWeightVertexBuffer.h
@@ -105,6 +105,7 @@ public:
 	friend FArchive& operator<<(FArchive& Ar, FSkinWeightLookupVertexBuffer& VertexBuffer);
 
 	void SerializeMetaData(FArchive& Ar);
+	ENGINE_API void SetMetaData(uint32 InNumVertices);
 	void CopyMetaData(const FSkinWeightLookupVertexBuffer& Other);
 
 	// FRenderResource interface.
@@ -208,6 +209,7 @@ public:
 	friend FArchive& operator<<(FArchive& Ar, FSkinWeightDataVertexBuffer& VertexBuffer);
 
 	void SerializeMetaData(FArchive& Ar);
+	ENGINE_API void SetMetaData(uint32 InNumVertices, uint32 InMaxBoneInfluences, bool bInUse16BitBoneIndex, bool bInUse16BitBoneWeight);
 	void CopyMetaData(const FSkinWeightDataVertexBuffer& Other);
 
 	// FRenderResource interface.
diff --git a/Engine/Source/Runtime/Engine/Public/Rendering/StaticMeshVertexBuffer.h b/Engine/Source/Runtime/Engine/Public/Rendering/StaticMeshVertexBuffer.h
--- a/Engine/Source/Runtime/Engine/Public/Rendering/StaticMeshVertexBuffer.h
+++ b/Engine/Source/Runtime/Engine/Public/Rendering/StaticMeshVertexBuffer.h
@@ -219,6 +219,8 @@ public:
 
 	void SerializeMetaData(FArchive& Ar);
 
+	ENGINE_API void SetMetaData(uint32 NumTexCoords, uint32 NumVertices, bool bUseFullPrecisionUVs, bool bUseHighPrecisionTangentBasis);
+
 	void ClearMetaData();
 
 	/**
