Fixed FTextureRenderTarget2DArrayResource::GetSizeY() and FTextureRenderTarget2DArrayResource::GetSizeXY(), which had been broken forever

diff --git a/Engine/Source/Runtime/Engine/Private/TextureRenderTarget2DArray.cpp b/Engine/Source/Runtime/Engine/Private/TextureRenderTarget2DArray.cpp
--- a/Engine/Source/Runtime/Engine/Private/TextureRenderTarget2DArray.cpp
+++ b/Engine/Source/Runtime/Engine/Private/TextureRenderTarget2DArray.cpp
@@ -349,7 +349,7 @@ uint32 FTextureRenderTarget2DArrayResource::GetSizeX() const
  */
 uint32 FTextureRenderTarget2DArrayResource::GetSizeY() const
 {
-	return Owner->SizeX;
+	return Owner->SizeY;
 }
 
 /**
@@ -357,7 +357,7 @@ uint32 FTextureRenderTarget2DArrayResource::GetSizeY() const
  */
 FIntPoint FTextureRenderTarget2DArrayResource::GetSizeXY() const
 {
-	return FIntPoint(Owner->SizeX, Owner->SizeX);
+	return FIntPoint(Owner->SizeX, Owner->SizeY);
 }
 
 float UTextureRenderTarget2DArray::GetDisplayGamma() const
