Texture.cpp disable incorrect check affecting non-streaming textures with mips on consoles when MaxLODSize is set such that LODBias on game platform is non-zero even after cooking
introduced in CL 23365088
Salesforce case 14390448

diff --git a/Engine/Source/Runtime/Engine/Private/Texture.cpp b/Engine/Source/Runtime/Engine/Private/Texture.cpp
--- a/Engine/Source/Runtime/Engine/Private/Texture.cpp
+++ b/Engine/Source/Runtime/Engine/Private/Texture.cpp
@@ -1866,10 +1866,12 @@ FStreamableRenderResourceState UTexture::GetResourcePostInitState(const FTexture
 	//		in a cooked run, it is zero
 
 	// in Editor the mips that will be dropped in cook are still present
-	//	dropping them is simulated by treating them as streamable
+	//	dropping them is simulated by treating them as streamable (even for textures that are NeverStream)
 	if ( ExpectedAssetLODBias > 0 && ! bTextureIsStreamable )
 	{
-		check( ! FPlatformProperties::RequiresCookedData() ); // ExpectedAssetLODBias should have been zero
+		// ExpectedAssetLODBias should be zero except in Editor builds
+		//  but not sure that's actually true on non-Windows client platforms
+		//check( ! FPlatformProperties::RequiresCookedData() ); 
 		NumOfNonStreamingMips = FullMipCount - ExpectedAssetLODBias;
 	}
 
