SVT: Added cvar to disable all SparseVolumeTexture streaming and GPU allocations.

#rnx

diff --git a/Engine/Source/Runtime/Engine/Private/SparseVolumeTexture/SparseVolumeTextureStreamingManager.cpp b/Engine/Source/Runtime/Engine/Private/SparseVolumeTexture/SparseVolumeTextureStreamingManager.cpp
--- a/Engine/Source/Runtime/Engine/Private/SparseVolumeTexture/SparseVolumeTextureStreamingManager.cpp
+++ b/Engine/Source/Runtime/Engine/Private/SparseVolumeTexture/SparseVolumeTextureStreamingManager.cpp
@@ -24,6 +24,14 @@ DEFINE_LOG_CATEGORY(LogSparseVolumeTextureStreamingManager);
 
 DECLARE_GPU_STAT(SVTStreaming);
 
+static int32 GSVTEnabled = 1;
+static FAutoConsoleVariableRef CVarSVTEnabled(
+	TEXT("r.SparseVolumeTexture"),
+	GSVTEnabled,
+	TEXT("Whether SparseVolumeTextures are supported at all. Toggles support for streaming and rendering SVTs."),
+	ECVF_RenderThreadSafe | ECVF_ReadOnly
+);
+
 static int32 GSVTStreamingForceBlockingRequests = 0;
 static FAutoConsoleVariableRef CVarSVTStreamingForceBlockingRequests(
 	TEXT("r.SparseVolumeTexture.Streaming.ForceBlockingRequests"),
@@ -118,12 +126,6 @@ IStreamingManager& GetStreamingManager()
 	return GStreamingManager;
 }
 
-static bool DoesPlatformSupportSparseVolumeTexture(EShaderPlatform Platform)
-{
-	// There are currently no hard platform restrictions for SVT support
-	return true;
-}
-
 struct FStreamingUpdateParameters
 {
 	FStreamingManager* StreamingManager = nullptr;
@@ -155,7 +157,7 @@ void FStreamingManager::InitRHI(FRHICommandListBase& RHICmdList)
 {
 	using namespace UE::DerivedData;
 
-	if (!DoesPlatformSupportSparseVolumeTexture(GMaxRHIShaderPlatform))
+	if (!UseSparseVolumeTextures(GMaxRHIShaderPlatform))
 	{
 		return;
 	}
@@ -172,7 +174,7 @@ void FStreamingManager::InitRHI(FRHICommandListBase& RHICmdList)
 
 void FStreamingManager::ReleaseRHI()
 {
-	if (!DoesPlatformSupportSparseVolumeTexture(GMaxRHIShaderPlatform))
+	if (!UseSparseVolumeTextures(GMaxRHIShaderPlatform))
 	{
 		return;
 	}
@@ -180,7 +182,7 @@ void FStreamingManager::ReleaseRHI()
 
 void FStreamingManager::Add_GameThread(UStreamableSparseVolumeTexture* SparseVolumeTexture)
 {
-	if (!DoesPlatformSupportSparseVolumeTexture(GMaxRHIShaderPlatform) || !SparseVolumeTexture)
+	if (!UseSparseVolumeTextures(GMaxRHIShaderPlatform) || !SparseVolumeTexture)
 	{
 		return;
 	}
@@ -223,7 +225,7 @@ void FStreamingManager::Add_GameThread(UStreamableSparseVolumeTexture* SparseVol
 
 void FStreamingManager::Remove_GameThread(UStreamableSparseVolumeTexture* SparseVolumeTexture)
 {
-	if (!DoesPlatformSupportSparseVolumeTexture(GMaxRHIShaderPlatform) || !SparseVolumeTexture)
+	if (!UseSparseVolumeTextures(GMaxRHIShaderPlatform) || !SparseVolumeTexture)
 	{
 		return;
 	}
@@ -236,7 +238,7 @@ void FStreamingManager::Remove_GameThread(UStreamableSparseVolumeTexture* Sparse
 
 void FStreamingManager::Request_GameThread(UStreamableSparseVolumeTexture* SparseVolumeTexture, uint32 StreamingInstanceKey, float FrameRate, float FrameIndex, float MipLevel, EStreamingRequestFlags Flags)
 {
-	if (!DoesPlatformSupportSparseVolumeTexture(GMaxRHIShaderPlatform) || !SparseVolumeTexture)
+	if (!UseSparseVolumeTextures(GMaxRHIShaderPlatform) || !SparseVolumeTexture)
 	{
 		return;
 	}
@@ -249,7 +251,7 @@ void FStreamingManager::Request_GameThread(UStreamableSparseVolumeTexture* Spars
 
 void FStreamingManager::Update_GameThread()
 {
-	if (!DoesPlatformSupportSparseVolumeTexture(GMaxRHIShaderPlatform))
+	if (!UseSparseVolumeTextures(GMaxRHIShaderPlatform))
 	{
 		return;
 	}
@@ -267,7 +269,7 @@ void FStreamingManager::Update_GameThread()
 void FStreamingManager::Request(UStreamableSparseVolumeTexture* SparseVolumeTexture, uint32 StreamingInstanceKey, float FrameRate, float FrameIndex, float MipLevel, EStreamingRequestFlags Flags)
 {
 	check(IsInRenderingThread());
-	if (!DoesPlatformSupportSparseVolumeTexture(GMaxRHIShaderPlatform) || !SparseVolumeTexture)
+	if (!UseSparseVolumeTextures(GMaxRHIShaderPlatform) || !SparseVolumeTexture)
 	{
 		return;
 	}
@@ -335,7 +337,7 @@ void FStreamingManager::BeginAsyncUpdate(FRDGBuilder& GraphBuilder, bool bUseAsy
 {
 	check(IsInRenderingThread());
 	check(!AsyncState.bUpdateActive);
-	if (!DoesPlatformSupportSparseVolumeTexture(GMaxRHIShaderPlatform) || StreamingInfo.IsEmpty())
+	if (!UseSparseVolumeTextures(GMaxRHIShaderPlatform) || StreamingInfo.IsEmpty())
 	{
 		return;
 	}
@@ -507,7 +509,7 @@ void FStreamingManager::BeginAsyncUpdate(FRDGBuilder& GraphBuilder, bool bUseAsy
 void FStreamingManager::EndAsyncUpdate(FRDGBuilder& GraphBuilder)
 {
 	check(IsInRenderingThread());
-	if (!DoesPlatformSupportSparseVolumeTexture(GMaxRHIShaderPlatform) || StreamingInfo.IsEmpty())
+	if (!UseSparseVolumeTextures(GMaxRHIShaderPlatform) || StreamingInfo.IsEmpty())
 	{
 		return;
 	}
diff --git a/Engine/Source/Runtime/RenderCore/Private/RenderUtils.cpp b/Engine/Source/Runtime/RenderCore/Private/RenderUtils.cpp
--- a/Engine/Source/Runtime/RenderCore/Private/RenderUtils.cpp
+++ b/Engine/Source/Runtime/RenderCore/Private/RenderUtils.cpp
@@ -1424,6 +1424,20 @@ bool DoesPlatformSupportHeterogeneousVolumes(EShaderPlatform Platform)
 		&& !IsForwardShadingEnabled(Platform);
 }
 
+bool DoesPlatformSupportSparseVolumeTextures(EShaderPlatform Platform)
+{
+	// There are currently no hard platform restrictions for SVT support, but mobile is probably out of spec at this point.
+	return IsFeatureLevelSupported(Platform, ERHIFeatureLevel::SM5);
+}
+
+bool UseSparseVolumeTextures(EShaderPlatform Platform)
+{
+	static const auto CVarSparseVolumeTexture = IConsoleManager::Get().FindConsoleVariable(TEXT("r.SparseVolumeTexture"));
+	const bool bIsSVTEnabled = CVarSparseVolumeTexture && CVarSparseVolumeTexture->GetInt() != 0;
+	const bool bIsSVTSupportedOnPlatform = DoesPlatformSupportSparseVolumeTextures(Platform);
+	return bIsSVTEnabled && bIsSVTSupportedOnPlatform;
+}
+
 bool IsSingleLayerWaterDepthPrepassEnabled(const FStaticShaderPlatform Platform, const FStaticFeatureLevel FeatureLevel)
 {
 	static const auto CVarWaterSingleLayerDepthPrepass = IConsoleManager::Get().FindConsoleVariable(TEXT("r.Water.SingleLayer.DepthPrepass"));
diff --git a/Engine/Source/Runtime/RenderCore/Public/RenderUtils.h b/Engine/Source/Runtime/RenderCore/Public/RenderUtils.h
--- a/Engine/Source/Runtime/RenderCore/Public/RenderUtils.h
+++ b/Engine/Source/Runtime/RenderCore/Public/RenderUtils.h
@@ -453,6 +453,16 @@ RENDERCORE_API bool DoesRuntimeSupportHeterogeneousVolumes(EShaderPlatform Platf
  */
 RENDERCORE_API bool DoesPlatformSupportHeterogeneousVolumes(EShaderPlatform Platform);
 
+/**
+ * Returns true if SparseVolumeTextures are supported on the given shader platform.
+ */
+RENDERCORE_API bool DoesPlatformSupportSparseVolumeTextures(EShaderPlatform Platform);
+
+/**
+ * Returns true if SparseVolumeTextures are supported on the given shader platform and are enabled.
+ */
+RENDERCORE_API bool UseSparseVolumeTextures(EShaderPlatform Platform);
+
 /**
 *	(Non-runtime) Checks if the depth prepass for single layer water is enabled. This also depends on virtual shadow maps to be supported on the platform.
 */
