Don't do per instance GPU distance cull HISM.
The HISM logic already does distance culling.
This fixes an issue with per instance fade out.
HISM renders past the cull distance for fading out instances. The GPU culling removes the fading instances.

diff --git a/Engine/Source/Runtime/Engine/Private/HierarchicalInstancedStaticMesh.cpp b/Engine/Source/Runtime/Engine/Private/HierarchicalInstancedStaticMesh.cpp
--- a/Engine/Source/Runtime/Engine/Private/HierarchicalInstancedStaticMesh.cpp
+++ b/Engine/Source/Runtime/Engine/Private/HierarchicalInstancedStaticMesh.cpp
@@ -2005,6 +2005,18 @@ const TArray<FBoxSphereBounds>* FHierarchicalStaticMeshSceneProxy::GetOcclusionQ
 	return &OcclusionBounds;
 }
 
+bool FHierarchicalStaticMeshSceneProxy::GetInstanceDrawDistanceMinMax(FVector2f& OutDistanceMinMax) const
+{
+	// HISM already does distance culling before submitting instances, so there is no need to do distance culling on GPU as well.
+	// Also GPU distance culling breaks dithered fade out for HISM. This is because we apply fade out to instances only when they go 
+	// beyond the culling distance, and at that point they will be GPU culled!
+	// We could apply a tweak to the draw distance returned here to account for this. But there isn't really any point
+	// because we know that the correct distance culling is already applied on the CPU side. So any tweaking of the value would be
+	// to ensure that we don't cull. So we might as well just not do the work on the GPU, which is what happens when we return false here.
+	OutDistanceMinMax = FVector2f(0, 0);
+	return false;
+}
+
 FBoxSphereBounds UHierarchicalInstancedStaticMeshComponent::CalcBounds(const FTransform& BoundTransform) const
 {
 	ensure(BuiltInstanceBounds.IsValid || ClusterTreePtr->Num() == 0);
diff --git a/Engine/Source/Runtime/Engine/Public/HierarchicalStaticMeshSceneProxy.h b/Engine/Source/Runtime/Engine/Public/HierarchicalStaticMeshSceneProxy.h
--- a/Engine/Source/Runtime/Engine/Public/HierarchicalStaticMeshSceneProxy.h
+++ b/Engine/Source/Runtime/Engine/Public/HierarchicalStaticMeshSceneProxy.h
@@ -91,6 +91,8 @@ public:
 		return FirstOcclusionNode > 0;
 	}
 
+	virtual bool GetInstanceDrawDistanceMinMax(FVector2f& OutDistanceMinMax) const override;
+
 	ENGINE_API virtual void DrawStaticElements(FStaticPrimitiveDrawInterface* PDI) override;
 
 	ENGINE_API virtual void ApplyWorldOffset(FRHICommandListBase& RHICmdList, FVector InOffset) override;
