World Partition: don't assign HLOD layers to actor descriptor views for HLOD irrelevant actors.
This avoids reporting an error for HLOD irrelevant actors assigned to streaming grids without an HLOD setup.

diff --git a/Engine/Source/Runtime/Engine/Private/WorldPartition/WorldPartitionStreamingGeneration.cpp b/Engine/Source/Runtime/Engine/Private/WorldPartition/WorldPartitionStreamingGeneration.cpp
--- a/Engine/Source/Runtime/Engine/Private/WorldPartition/WorldPartitionStreamingGeneration.cpp
+++ b/Engine/Source/Runtime/Engine/Private/WorldPartition/WorldPartitionStreamingGeneration.cpp
@@ -720,7 +720,7 @@ class FWorldPartitionStreamingGenerator
 	{
 		// Only assign the default layer to actors that don't have a valid HLOD layer set. HLOD actors will have their 
 		// parent HLOD layer set during HLOD generation.
-		if (!ActorDescView.GetHLODLayer().IsValid() && ActorDescView.GetIsSpatiallyLoaded())
+		if (!ActorDescView.GetHLODLayer().IsValid() && ActorDescView.GetIsSpatiallyLoaded() && ActorDescView.GetActorIsHLODRelevant())
 		{
 			ActorDescView.SetRuntimeHLODLayer(DefaultHLODLayer);
 		}
