Force garbage collection after removing players viewportClient and flagging as garbage

#localization none
#tests local client testing

diff --git a/Engine/Source/Runtime/Engine/Private/GameInstance.cpp b/Engine/Source/Runtime/Engine/Private/GameInstance.cpp
--- a/Engine/Source/Runtime/Engine/Private/GameInstance.cpp
+++ b/Engine/Source/Runtime/Engine/Private/GameInstance.cpp
@@ -1023,6 +1023,7 @@ bool UGameInstance::RemoveLocalPlayer(ULocalPlayer* ExistingPlayer)
 	if (!UObjectBaseUtility::IsGarbageEliminationEnabled())
 	{
 		ExistingPlayer->MarkAsGarbage();
+		GEngine->ForceGarbageCollection();
 	}
 
 	return true;
