Fix for runtime transient textures on tiled platforms.

#rnx

diff --git a/Engine/Source/Runtime/Engine/Classes/Engine/Texture2DArray.cpp b/Engine/Source/Runtime/Engine/Classes/Engine/Texture2DArray.cpp
--- a/Engine/Source/Runtime/Engine/Classes/Engine/Texture2DArray.cpp
+++ b/Engine/Source/Runtime/Engine/Classes/Engine/Texture2DArray.cpp
@@ -67,6 +67,7 @@ UTexture2DArray* UTexture2DArray::CreateTransient(int32 InSizeX, int32 InSizeY,
 		NewTexture->GetPlatformData()->SizeY = InSizeY;
 		NewTexture->GetPlatformData()->SetNumSlices(InArraySize);
 		NewTexture->GetPlatformData()->PixelFormat = InFormat;
+		NewTexture->bNotOfflineProcessed = true;
 
 		// Allocate first mipmap.
 		int32 NumBlocksX = InSizeX / GPixelFormats[InFormat].BlockSizeX;
diff --git a/Engine/Source/Runtime/Engine/Private/Texture2D.cpp b/Engine/Source/Runtime/Engine/Private/Texture2D.cpp
--- a/Engine/Source/Runtime/Engine/Private/Texture2D.cpp
+++ b/Engine/Source/Runtime/Engine/Private/Texture2D.cpp
@@ -1267,6 +1267,7 @@ UTexture2D* UTexture2D::CreateTransient(int32 InSizeX, int32 InSizeY, EPixelForm
 	NewTexture->GetPlatformData()->SizeY = InSizeY;
 	NewTexture->GetPlatformData()->SetNumSlices(1);
 	NewTexture->GetPlatformData()->PixelFormat = InFormat;
+	NewTexture->bNotOfflineProcessed = true;
 
 	// Allocate first mipmap.
 	FTexture2DMipMap* Mip = new FTexture2DMipMap(InSizeX, InSizeY, 1);
diff --git a/Engine/Source/Runtime/Engine/Private/TextureCube.cpp b/Engine/Source/Runtime/Engine/Private/TextureCube.cpp
--- a/Engine/Source/Runtime/Engine/Private/TextureCube.cpp
+++ b/Engine/Source/Runtime/Engine/Private/TextureCube.cpp
@@ -45,6 +45,7 @@ UTextureCube* UTextureCube::CreateTransient(int32 InSizeX, int32 InSizeY, EPixel
 		NewTexture->GetPlatformData()->SetIsCubemap(true);
 		NewTexture->GetPlatformData()->SetNumSlices(6);
 		NewTexture->GetPlatformData()->PixelFormat = InFormat;
+		NewTexture->bNotOfflineProcessed = true;
 
 		// Allocate first mipmap.
 		int32 NumBlocksX = InSizeX / GPixelFormats[InFormat].BlockSizeX;
diff --git a/Engine/Source/Runtime/Engine/Private/TextureCubeArray.cpp b/Engine/Source/Runtime/Engine/Private/TextureCubeArray.cpp
--- a/Engine/Source/Runtime/Engine/Private/TextureCubeArray.cpp
+++ b/Engine/Source/Runtime/Engine/Private/TextureCubeArray.cpp
@@ -330,6 +330,7 @@ UTextureCubeArray* UTextureCubeArray::CreateTransient(int32 InSizeX, int32 InSiz
 		NewTexture->GetPlatformData()->SetIsCubemap(true);
 		NewTexture->GetPlatformData()->SetNumSlices(InArraySize * 6);
 		NewTexture->GetPlatformData()->PixelFormat = InFormat;
+		NewTexture->bNotOfflineProcessed = true;
 
 		// Allocate first mipmap.
 		int32 NumBlocksX = InSizeX / GPixelFormats[InFormat].BlockSizeX;
diff --git a/Engine/Source/Runtime/Engine/Private/VolumeTexture.cpp b/Engine/Source/Runtime/Engine/Private/VolumeTexture.cpp
--- a/Engine/Source/Runtime/Engine/Private/VolumeTexture.cpp
+++ b/Engine/Source/Runtime/Engine/Private/VolumeTexture.cpp
@@ -74,6 +74,7 @@ UVolumeTexture* UVolumeTexture::CreateTransient(int32 InSizeX, int32 InSizeY, in
 		NewTexture->GetPlatformData()->SizeY = InSizeY;
 		NewTexture->GetPlatformData()->SetNumSlices(InSizeZ);
 		NewTexture->GetPlatformData()->PixelFormat = InFormat;
+		NewTexture->bNotOfflineProcessed = true;
 
 		// Allocate first mipmap.
 		int32 NumBlocksX = InSizeX / GPixelFormats[InFormat].BlockSizeX;
