SVT: Added memory stats to track permament SparseVolumeTexture GPU memory allocations as well as the CPU side memory allocations needed for streaming bookkeeping etc.

#rnx

diff --git a/Engine/Source/Runtime/Engine/Private/SparseVolumeTexture/SparseVolumeTextureStreamingManager.cpp b/Engine/Source/Runtime/Engine/Private/SparseVolumeTexture/SparseVolumeTextureStreamingManager.cpp
--- a/Engine/Source/Runtime/Engine/Private/SparseVolumeTexture/SparseVolumeTextureStreamingManager.cpp
+++ b/Engine/Source/Runtime/Engine/Private/SparseVolumeTexture/SparseVolumeTextureStreamingManager.cpp
@@ -24,6 +24,10 @@ DEFINE_LOG_CATEGORY(LogSparseVolumeTextureStreamingManager);
 
 DECLARE_GPU_STAT(SVTStreaming);
 
+DECLARE_MEMORY_STAT(TEXT("Total CPU Memory"), STAT_SVTCPUMemoryTotal, STATGROUP_SparseVolumeTextureMemory);
+DECLARE_MEMORY_STAT_POOL(TEXT("PageTable GPU Memory"), STAT_SVTPageTableMemory, STATGROUP_SparseVolumeTextureMemory, FPlatformMemory::MCR_GPU);
+DEFINE_STAT(STAT_SVTGPUMemoryTotal);
+
 static int32 GSVTEnabled = 1;
 static FAutoConsoleVariableRef CVarSVTEnabled(
 	TEXT("r.SparseVolumeTexture"),
@@ -542,6 +546,8 @@ void FStreamingManager::EndAsyncUpdate(FRDGBuilder& GraphBuilder)
 	++UpdateIndex;
 	AsyncState.bUpdateActive = false;
 	AsyncState.bUpdateIsAsync = false;
+
+	SET_MEMORY_STAT(STAT_SVTCPUMemoryTotal, GetAllocatedSize());
 }
 
 const FStreamingDebugInfo* FStreamingManager::GetStreamingDebugInfo(FRDGBuilder& GraphBuilder) const
@@ -760,6 +766,11 @@ void FStreamingManager::AddInternal(FRDGBuilder& GraphBuilder, FNewSparseVolumeT
 
 				FrameInfo.PageTableTexture = GraphBuilder.ConvertToExternalTexture(PageTableRDG);
 				GraphBuilder.UseExternalAccessMode(PageTableRDG, ERHIAccess::SRVMask, ERHIPipeline::All);
+#if STATS
+				const int64 MemorySize = static_cast<int64>(RHICalcTexturePlatformSize(FrameInfo.PageTableTexture->GetRHI()->GetDesc()).Size);
+				INC_MEMORY_STAT_BY(STAT_SVTPageTableMemory, MemorySize);
+				INC_MEMORY_STAT_BY(STAT_SVTGPUMemoryTotal, MemorySize);
+#endif
 			}
 
 			// Initialize TextureRenderResources
@@ -868,6 +879,11 @@ void FStreamingManager::RemoveInternal(UStreamableSparseVolumeTexture* SparseVol
 		// Release resources
 		for (FFrameInfo& FrameInfo : SVTInfo->PerFrameInfo)
 		{
+#if STATS
+			const int64 MemorySize = static_cast<int64>(RHICalcTexturePlatformSize(FrameInfo.PageTableTexture->GetRHI()->GetDesc()).Size);
+			DEC_MEMORY_STAT_BY(STAT_SVTPageTableMemory, MemorySize);
+			DEC_MEMORY_STAT_BY(STAT_SVTGPUMemoryTotal, MemorySize);
+#endif
 			FrameInfo.PageTableTexture.SafeRelease();
 			InvalidatedSVTFrames.Remove(&FrameInfo);
 		}
@@ -1941,6 +1957,21 @@ void FStreamingManager::RequestDDCData(TConstArrayView<UE::DerivedData::FCacheGe
 
 #endif // WITH_EDITORONLY_DATA
 
+SIZE_T FStreamingManager::GetAllocatedSize() const
+{
+	SIZE_T AllocatedSize = 0;
+	AllocatedSize += SparseVolumeTextureToHandle.GetAllocatedSize();
+	AllocatedSize += StreamingInfo.GetAllocatedSize();
+	for (const auto& Info : StreamingInfo)
+	{
+		AllocatedSize += Info->GetAllocatedSize();
+	}
+	AllocatedSize += RequestsHashTable.GetAllocatedSize();
+	AllocatedSize += PendingRequests.GetAllocatedSize();
+	
+	return AllocatedSize;
+}
+
 FStreamingInstance* FStreamingManager::FStreamingInfo::GetAndUpdateStreamingInstance(uint64 StreamingInstanceKey, const FStreamingInstanceRequest& Request)
 {
 	// Try to find a FStreamingInstance around the requested frame index and matching the same key
@@ -1962,6 +1993,31 @@ FStreamingInstance* FStreamingManager::FStreamingInfo::GetAndUpdateStreamingInst
 	return StreamingInstance;
 }
 
+SIZE_T FStreamingManager::FFrameInfo::GetAllocatedSize() const
+{
+	return TileAllocations.GetAllocatedSize()
+		+ ResidentPages.GetAllocatedSize()
+		+ InvalidatedPages.GetAllocatedSize()
+		+ ResidentTiles.GetAllocatedSize()
+		+ StreamingTiles.GetAllocatedSize()
+		+ TileIndexToPendingRequestIndex.GetAllocatedSize();
+}
+
+SIZE_T FStreamingManager::FStreamingInfo::GetAllocatedSize() const
+{
+	SIZE_T AllocatedSize = 0;
+	for (const FFrameInfo& FrameInfo : PerFrameInfo)
+	{
+		AllocatedSize += FrameInfo.GetAllocatedSize();
+	}
+	AllocatedSize += PerFrameInfo.GetAllocatedSize();
+	AllocatedSize += StreamingInstances.GetAllocatedSize() + StreamingInstances.Num() * sizeof(*StreamingInstances[0]); // StreamingInstances only stores pointers, but each FStreamingInstance has the same size.
+	AllocatedSize += MipLevelStreamingSize.GetAllocatedSize();
+	AllocatedSize += TileAllocator.GetAllocatedSize();
+	AllocatedSize += sizeof(*TileDataTexture);
+	return AllocatedSize;
+}
+
 void FTileAllocator::Init(const FIntVector3& InResolutionInTiles)
 {
 	check(InResolutionInTiles.X > 0 && InResolutionInTiles.Y > 0 && InResolutionInTiles.Z > 0);
@@ -2086,6 +2142,13 @@ void FTileAllocator::Free(uint32 TileCoord)
 	--NumAllocated;
 }
 
+SIZE_T FTileAllocator::GetAllocatedSize() const
+{
+	// FBinaryHeap unfortunately does not support querying the allocated size, so we just compute it ourselves here.
+	const SIZE_T HeapSize = FreeHeap.GetHeapSize() * sizeof(uint32) + FreeHeap.GetIndexSize() * (sizeof(uint64) + sizeof(uint32));
+	return HeapSize + Allocations.GetAllocatedSize();
+}
+
 }
 }
 
diff --git a/Engine/Source/Runtime/Engine/Private/SparseVolumeTexture/SparseVolumeTextureStreamingManager.h b/Engine/Source/Runtime/Engine/Private/SparseVolumeTexture/SparseVolumeTextureStreamingManager.h
--- a/Engine/Source/Runtime/Engine/Private/SparseVolumeTexture/SparseVolumeTextureStreamingManager.h
+++ b/Engine/Source/Runtime/Engine/Private/SparseVolumeTexture/SparseVolumeTextureStreamingManager.h
@@ -16,6 +16,8 @@
 #include "Misc/DateTime.h"
 
 DECLARE_LOG_CATEGORY_EXTERN(LogSparseVolumeTextureStreamingManager, Log, All);
+DECLARE_STATS_GROUP(TEXT("SparseVolumeTexture Memory"), STATGROUP_SparseVolumeTextureMemory, STATCAT_Advanced);
+DECLARE_MEMORY_STAT_POOL_EXTERN(TEXT("Total GPU Memory"), STAT_SVTGPUMemoryTotal, STATGROUP_SparseVolumeTextureMemory, FPlatformMemory::MCR_GPU, );
 
 class UStreamableSparseVolumeTexture;
 class USparseVolumeTextureFrame;
@@ -79,6 +81,9 @@ public:
 	// Frees a given tile. Expects the output of Allocate(), i.e. a packed 8|8|8 coordinate. Freed tiles are placed at the front of the priority queue for tile reuse.
 	void Free(uint32 TileCoord);
 
+	// Returns the size of the memory allocated for bookkeeping.
+	SIZE_T GetAllocatedSize() const;
+
 private:
 	FBinaryHeap<uint64, uint32> FreeHeap;
 	TArray<FAllocation> Allocations; // One entry for every slot in the tile data texture. Maps to the current SVT frame occupying that tile
@@ -127,6 +132,8 @@ private:
 		TBitArray<> StreamingTiles; // Tiles that have logically been streamed in but may not actually be resident in GPU memory yet because the IO request hasn't finished
 		TMap<uint32, uint32> TileIndexToPendingRequestIndex;
 		TRefCountPtr<IPooledRenderTarget> PageTableTexture;
+
+		SIZE_T GetAllocatedSize() const;
 	};
 
 	// Streaming manager internal data for each SVT
@@ -150,6 +157,8 @@ private:
 
 		// Tries to look up an existing FStreamingInstance, creates a new one if it can't find one and finally updates it with the new request.
 		FStreamingInstance* GetAndUpdateStreamingInstance(uint64 StreamingInstanceKey, const FStreamingInstanceRequest& Request);
+
+		SIZE_T GetAllocatedSize() const;
 	};
 
 	// Represents an IO request for tile(s)
@@ -339,6 +348,8 @@ private:
 	UE::DerivedData::FCacheGetChunkRequest BuildDDCRequest(const FResources& Resources, uint32 FirstTileIndex, uint32 NumTiles, uint32 PendingRequestIndex, int32 ChunkIndex);
 	void RequestDDCData(TConstArrayView<UE::DerivedData::FCacheGetChunkRequest> DDCRequests, bool bBlocking);
 #endif
+
+	SIZE_T GetAllocatedSize() const;
 };
 
 }
diff --git a/Engine/Source/Runtime/Engine/Private/SparseVolumeTexture/SparseVolumeTextureTileDataTexture.cpp b/Engine/Source/Runtime/Engine/Private/SparseVolumeTexture/SparseVolumeTextureTileDataTexture.cpp
--- a/Engine/Source/Runtime/Engine/Private/SparseVolumeTexture/SparseVolumeTextureTileDataTexture.cpp
+++ b/Engine/Source/Runtime/Engine/Private/SparseVolumeTexture/SparseVolumeTextureTileDataTexture.cpp
@@ -5,6 +5,8 @@
 #include "SparseVolumeTextureStreamingManager.h" // LogSparseVolumeTextureStreamingManager
 #include "RenderTargetPool.h"
 
+DECLARE_MEMORY_STAT_POOL(TEXT("TileData GPU Memory"), STAT_SVTTileDataMemory, STATGROUP_SparseVolumeTextureMemory, FPlatformMemory::MCR_GPU);
+
 static int32 GSVTStreamingReservedResources = 1;
 static FAutoConsoleVariableRef CVarSVTStreamingReservedResources(
 	TEXT("r.SparseVolumeTexture.Streaming.UseReservedResources"),
@@ -204,15 +206,39 @@ void FTileDataTexture::InitRHI(FRHICommandListBase& RHICmdList)
 	if (FormatA != PF_Unknown)
 	{
 		TileDataTextureA = GRenderTargetPool.FindFreeElement(RHICmdList, FRDGTextureDesc::Create3D(Resolution, FormatA, FClearValueBinding::Black, Flags), TEXT("SparseVolumeTexture.PhysicalTileDataA.RHITexture"));
+#if STATS
+		const int64 MemorySize = static_cast<int64>(RHICalcTexturePlatformSize(TileDataTextureA->GetRHI()->GetDesc()).Size);
+		INC_MEMORY_STAT_BY(STAT_SVTTileDataMemory, MemorySize);
+		INC_MEMORY_STAT_BY(STAT_SVTGPUMemoryTotal, MemorySize);
+#endif
 	}
 	if (FormatB != PF_Unknown)
 	{
 		TileDataTextureB = GRenderTargetPool.FindFreeElement(RHICmdList, FRDGTextureDesc::Create3D(Resolution, FormatB, FClearValueBinding::Black, Flags), TEXT("SparseVolumeTexture.PhysicalTileDataB.RHITexture"));
+#if STATS
+		const int64 MemorySize = static_cast<int64>(RHICalcTexturePlatformSize(TileDataTextureB->GetRHI()->GetDesc()).Size);
+		INC_MEMORY_STAT_BY(STAT_SVTTileDataMemory, MemorySize);
+		INC_MEMORY_STAT_BY(STAT_SVTGPUMemoryTotal, MemorySize);
+#endif
 	}
 }
 
 void FTileDataTexture::ReleaseRHI()
 {
+#if STATS
+	if (FormatA != PF_Unknown)
+	{
+		const int64 MemorySize = static_cast<int64>(RHICalcTexturePlatformSize(TileDataTextureA->GetRHI()->GetDesc()).Size);
+		DEC_MEMORY_STAT_BY(STAT_SVTTileDataMemory, MemorySize);
+		DEC_MEMORY_STAT_BY(STAT_SVTGPUMemoryTotal, MemorySize);
+	}
+	if (FormatB != PF_Unknown)
+	{
+		const int64 MemorySize = static_cast<int64>(RHICalcTexturePlatformSize(TileDataTextureB->GetRHI()->GetDesc()).Size);
+		DEC_MEMORY_STAT_BY(STAT_SVTTileDataMemory, MemorySize);
+		DEC_MEMORY_STAT_BY(STAT_SVTGPUMemoryTotal, MemorySize);
+	}
+#endif
 }
 
 }
