[HWRT] Fix crash when using ray tracing proxy with WPO on FStaticMeshProxy code path.

Context:
- When !RayTracingProxy->bUsingRenderingLODs, the FStaticMeshRayTracingProxyLOD reference geometry different than FStaticMeshLODResources.
- FStaticMeshSceneProxy::GetDynamicRayTracingInstances(...) was mixing data from FStaticMeshRayTracingProxyLOD and FStaticMeshLODResources.
- FStaticMeshSceneProxy::GetDynamicRayTracingInstances(...) was also skipping LOD selection so it could end up using a streamed out ray tracing proxy LOD.

Changes:
- mark RayTracingGeometryGroup as referenced and select a valid LOD when !RayTracingProxy->bUsingRenderingLODs regardless of proxy using WPO or not since we need VB/IB to be streamed-in.
- implemented new method GetRayTracingMeshElement(...) to initialize FMeshBatch used for ray tracing with ray tracing proxy data when !RayTracingProxy->bUsingRenderingLODs.
- refactored SetMeshElementGeometrySource(...) to support both rendering and ray tracing LODs.
- store and use the new FRayTracingLODInfo when !RayTracingProxy->bUsingRenderingLODs instead of using the regular FLODInfo since section data might be different between rendering and ray tracing LODs.

diff --git a/Engine/Source/Runtime/Engine/Private/StaticMesh.cpp b/Engine/Source/Runtime/Engine/Private/StaticMesh.cpp
--- a/Engine/Source/Runtime/Engine/Private/StaticMesh.cpp
+++ b/Engine/Source/Runtime/Engine/Private/StaticMesh.cpp
@@ -1987,7 +1987,7 @@ void FStaticMeshRayTracingProxyLOD::InitWithStreamedData(FRHICommandListBase& RH
 
 	if (!GRHIGlobals.RayTracing.SupportsShaders && !GRHIGlobals.RayTracing.RequiresInlineRayTracingSBT)
 	{
-		// TODO: stream out VB/IB
+		// TODO: stream out VB/IB (unless there are instead using WPO)
 	}
 }
 
diff --git a/Engine/Source/Runtime/Engine/Private/StaticMeshSceneProxy.cpp b/Engine/Source/Runtime/Engine/Private/StaticMeshSceneProxy.cpp
--- a/Engine/Source/Runtime/Engine/Private/StaticMeshSceneProxy.cpp
+++ b/Engine/Source/Runtime/Engine/Private/StaticMeshSceneProxy.cpp
@@ -311,13 +311,33 @@ FStaticMeshSceneProxy::FStaticMeshSceneProxy(const FStaticMeshSceneProxyDesc& In
 	LODs.Empty(RenderData->LODResources.Num());
 	const bool bLODsShareStaticLighting = RenderData->bLODsShareStaticLighting || bForceLODsShareStaticLighting;
 
+	for (int32 LODIndex = 0; LODIndex < RenderData->LODResources.Num(); LODIndex++)
+	{
+		FLODInfo* NewLODInfo = new (LODs) FLODInfo(InProxyDesc, RenderData->LODVertexFactories, LODIndex, ClampedMinLOD, bLODsShareStaticLighting);
+
+		// Under certain error conditions an LOD's material will be set to 
+		// DefaultMaterial. Ensure our material view relevance is set properly.
+		const int32 NumSections = NewLODInfo->Sections.Num();
+		for (int32 SectionIndex = 0; SectionIndex < NumSections; ++SectionIndex)
+		{
+			const FLODInfo::FSectionInfo& SectionInfo = NewLODInfo->Sections[SectionIndex];
+			bAnySectionCastsShadows |= RenderData->LODResources[LODIndex].Sections[SectionIndex].bCastShadow;
+			if (SectionInfo.Material == UMaterial::GetDefaultMaterial(MD_Surface))
+			{
+				MaterialRelevance |= UMaterial::GetDefaultMaterial(MD_Surface)->GetRelevance(FeatureLevel);
+			}
+
+			MaxWPOExtent = FMath::Max(MaxWPOExtent, SectionInfo.Material->GetMaxWorldPositionOffsetDisplacement());
+		}
+	}
+
 #if RHI_RAYTRACING
 	bSupportRayTracing = IsRayTracingAllowed() && InProxyDesc.GetStaticMesh()->bSupportRayTracing;
 	bDynamicRayTracingGeometry = false;
 	bNeedsDynamicRayTracingGeometries = false;
-	
+
 	if (bSupportRayTracing)
-	{		
+	{
 		const bool bWantsRayTracingWPO = MaterialRelevance.bUsesWorldPositionOffset && InProxyDesc.bEvaluateWorldPositionOffsetInRayTracing;
 
 		// r.RayTracing.Geometry.StaticMeshes.WPO is handled in the following way:
@@ -333,28 +353,19 @@ FStaticMeshSceneProxy::FStaticMeshSceneProxy(const FStaticMeshSceneProxyDesc& In
 			// only need dynamic geometries when r.RayTracing.Geometry.StaticMeshes.WPO == 1
 			bNeedsDynamicRayTracingGeometries = CVarRayTracingStaticMeshesWPO.GetValueOnAnyThread() == 1;
 		}
-	}
-#endif
-
-	for(int32 LODIndex = 0;LODIndex < RenderData->LODResources.Num();LODIndex++)
-	{
-		FLODInfo* NewLODInfo = new (LODs) FLODInfo(InProxyDesc, RenderData->LODVertexFactories, LODIndex, ClampedMinLOD, bLODsShareStaticLighting);
 
-		// Under certain error conditions an LOD's material will be set to 
-		// DefaultMaterial. Ensure our material view relevance is set properly.
-		const int32 NumSections = NewLODInfo->Sections.Num();
-		for (int32 SectionIndex = 0; SectionIndex < NumSections; ++SectionIndex)
+		// when ray tracing proxy uses rendering LODs we can use the FLODInfo for ray tracing
+		// otherwise initialize RayTracingLODInfos to be used in GetRayTracingMeshElement(...)
+		if (!RenderData->RayTracingProxy->bUsingRenderingLODs)
 		{
-			const FLODInfo::FSectionInfo& SectionInfo = NewLODInfo->Sections[SectionIndex];
-			bAnySectionCastsShadows |= RenderData->LODResources[LODIndex].Sections[SectionIndex].bCastShadow;
-			if (SectionInfo.Material == UMaterial::GetDefaultMaterial(MD_Surface))
+			RayTracingLODInfos.Empty(RenderData->RayTracingProxy->LODs.Num());
+			for (int32 LODIndex = 0; LODIndex < RenderData->RayTracingProxy->LODs.Num(); LODIndex++)
 			{
-				MaterialRelevance |= UMaterial::GetDefaultMaterial(MD_Surface)->GetRelevance(FeatureLevel);
+				new (RayTracingLODInfos) FRayTracingLODInfo(InProxyDesc, LODIndex);
 			}
-
-			MaxWPOExtent = FMath::Max(MaxWPOExtent, SectionInfo.Material->GetMaxWorldPositionOffsetDisplacement());
 		}
 	}
+#endif
 
 	// WPO is typically used for ambient animations, so don't include in cached shadowmaps
 	// Note mesh animation can also come from PDO or Tessellation but they are typically static uses so we ignore them for cached shadowmaps
@@ -679,7 +690,7 @@ bool FStaticMeshSceneProxy::GetMeshElement(
 	// No support for stateless dithered LOD transitions for movable meshes
 	const bool bDitheredLODTransition = !IsMovable() && Material.IsDitheredLODTransition();
 
-	const uint32 NumPrimitives = SetMeshElementGeometrySource(LODIndex, SectionIndex, bWireframe, bUseReversedIndices, bAllowPreCulledIndices, VertexFactory, OutMeshBatch);
+	const uint32 NumPrimitives = SetMeshElementGeometrySource(Section, LOD.IndexBuffer, LOD.AdditionalIndexBuffers, VertexFactory, bWireframe, bUseReversedIndices, OutMeshBatch);
 
 	if (NumPrimitives > 0)
 	{
@@ -716,6 +727,65 @@ bool FStaticMeshSceneProxy::GetMeshElement(
 	}
 }
 
+bool FStaticMeshSceneProxy::GetRayTracingMeshElement(int32 LODIndex, int32 BatchIndex, int32 SectionIndex, uint8 InDepthPriorityGroup, FMeshBatch& OutMeshBatch) const
+{
+	const ERHIFeatureLevel::Type FeatureLevel = GetScene().GetFeatureLevel();
+	const FStaticMeshRayTracingProxyLOD& LOD = RenderData->RayTracingProxy->LODs[LODIndex];
+	const FStaticMeshVertexFactories& VFs = (*RenderData->RayTracingProxy->LODVertexFactories)[LODIndex];
+	const FStaticMeshSection& Section = (*LOD.Sections)[SectionIndex];
+
+	if (Section.NumTriangles == 0)
+	{
+		return false;
+	}
+
+	const FRayTracingLODInfo& LODInfo = RayTracingLODInfos[LODIndex];
+
+	UMaterialInterface* MaterialInterface = LODInfo.Sections[SectionIndex].Material;
+	const FMaterialRenderProxy* MaterialRenderProxy = MaterialInterface->GetRenderProxy();
+	const FMaterial& Material = MaterialRenderProxy->GetIncompleteMaterialWithFallback(FeatureLevel);
+
+	const FVertexFactory* VertexFactory = &VFs.VertexFactory;
+
+	FMeshBatchElement& OutMeshBatchElement = OutMeshBatch.Elements[0];
+	OutMeshBatchElement.VertexFactoryUserData = VFs.VertexFactory.GetUniformBuffer();
+
+	const bool bWireframe = false;
+	const bool bUseReversedIndices = false;
+
+	// No support for stateless dithered LOD transitions for movable meshes
+	const bool bDitheredLODTransition = !IsMovable() && Material.IsDitheredLODTransition();
+
+	const uint32 NumPrimitives = SetMeshElementGeometrySource(Section, *LOD.IndexBuffer, nullptr, VertexFactory, bWireframe, bUseReversedIndices, OutMeshBatch);
+
+	if (NumPrimitives > 0)
+	{
+		OutMeshBatch.SegmentIndex = SectionIndex;
+		OutMeshBatch.MeshIdInPrimitive = SectionIndex;
+
+		OutMeshBatch.LODIndex = LODIndex;
+		OutMeshBatch.ReverseCulling = IsReversedCullingNeeded(bUseReversedIndices);
+		OutMeshBatch.CastShadow = bCastShadow && Section.bCastShadow;
+#if RHI_RAYTRACING
+		OutMeshBatch.CastRayTracedShadow = OutMeshBatch.CastShadow && bCastDynamicShadow;
+#endif
+		OutMeshBatch.DepthPriorityGroup = (ESceneDepthPriorityGroup)InDepthPriorityGroup;
+		OutMeshBatch.LCI = nullptr;
+		OutMeshBatch.MaterialRenderProxy = MaterialRenderProxy;
+
+		OutMeshBatchElement.MinVertexIndex = Section.MinVertexIndex;
+		OutMeshBatchElement.MaxVertexIndex = Section.MaxVertexIndex;
+
+		SetMeshElementScreenSize(LODIndex, bDitheredLODTransition, OutMeshBatch);
+
+		return true;
+	}
+	else
+	{
+		return false;
+	}
+}
+
 #if RHI_RAYTRACING
 void FStaticMeshSceneProxy::CreateDynamicRayTracingGeometries(FRHICommandListBase& RHICmdList)
 {
@@ -823,7 +893,7 @@ bool FStaticMeshSceneProxy::GetWireframeMeshElement(int32 LODIndex, int32 BatchI
 	OutBatchElement.MinVertexIndex = 0;
 	OutBatchElement.MaxVertexIndex = LODModel.GetNumVertices() - 1;
 
-	const uint32_t NumPrimitives = SetMeshElementGeometrySource(LODIndex, 0, bWireframe, bUseReversedIndices, bAllowPreCulledIndices, VertexFactory, OutMeshBatch);
+	const uint32_t NumPrimitives = SetMeshElementGeometrySource(LODModel.Sections[0], LODModel.IndexBuffer, LODModel.AdditionalIndexBuffers, VertexFactory, bWireframe, bUseReversedIndices, OutMeshBatch);
 	SetMeshElementScreenSize(LODIndex, bDitheredLODTransition, OutMeshBatch);
 
 	return NumPrimitives > 0;
@@ -855,7 +925,7 @@ bool FStaticMeshSceneProxy::GetCollisionMeshElement(
 	const bool bAllowPreCulledIndices = true;
 	const bool bDitheredLODTransition = false;
 
-	SetMeshElementGeometrySource(LODIndex, SectionIndex, bWireframe, bUseReversedIndices, bAllowPreCulledIndices, VertexFactory, OutMeshBatch);
+	SetMeshElementGeometrySource(Section, LOD.IndexBuffer, LOD.AdditionalIndexBuffers, VertexFactory, bWireframe, bUseReversedIndices, OutMeshBatch);
 
 	FMeshBatchElement& OutMeshBatchElement = OutMeshBatch.Elements[0];
 
@@ -993,44 +1063,38 @@ bool FStaticMeshSceneProxy::GetMaterialTextureScales(int32 LODIndex, int32 Secti
 #endif
 
 uint32 FStaticMeshSceneProxy::SetMeshElementGeometrySource(
-	int32 LODIndex,
-	int32 SectionIndex,
+	const FStaticMeshSection& Section,
+	const FRawStaticIndexBuffer& IndexBuffer,
+	const FAdditionalStaticMeshIndexBuffers* AdditionalIndexBuffers,
+	const FVertexFactory* VertexFactory,
 	bool bWireframe,
 	bool bUseReversedIndices,
-	bool bAllowPreCulledIndices, // unused
-	const FVertexFactory* VertexFactory,
 	FMeshBatch& OutMeshBatch) const
 {
-	const FStaticMeshLODResources& LODModel = RenderData->LODResources[LODIndex];
-
-	const FStaticMeshSection& Section = LODModel.Sections[SectionIndex];
 	if (Section.NumTriangles == 0)
 	{
 		return 0;
 	}
 
-	const FLODInfo& LODInfo = LODs[LODIndex];
-	const FLODInfo::FSectionInfo& SectionInfo = LODInfo.Sections[SectionIndex];
-
 	FMeshBatchElement& OutMeshBatchElement = OutMeshBatch.Elements[0];
 	uint32 NumPrimitives = 0;
 
 	if (bWireframe)
 	{
-		if (LODModel.AdditionalIndexBuffers && LODModel.AdditionalIndexBuffers->WireframeIndexBuffer.IsInitialized())
+		if (AdditionalIndexBuffers && AdditionalIndexBuffers->WireframeIndexBuffer.IsInitialized())
 		{
 			OutMeshBatch.Type = PT_LineList;
 			OutMeshBatchElement.FirstIndex = 0;
-			OutMeshBatchElement.IndexBuffer = &LODModel.AdditionalIndexBuffers->WireframeIndexBuffer;
-			NumPrimitives = LODModel.AdditionalIndexBuffers->WireframeIndexBuffer.GetNumIndices() / 2;
+			OutMeshBatchElement.IndexBuffer = &AdditionalIndexBuffers->WireframeIndexBuffer;
+			NumPrimitives = AdditionalIndexBuffers->WireframeIndexBuffer.GetNumIndices() / 2;
 		}
 		else
 		{
 			OutMeshBatch.Type = PT_TriangleList;
 
 			OutMeshBatchElement.FirstIndex = 0;
-			OutMeshBatchElement.IndexBuffer = &LODModel.IndexBuffer;
-			NumPrimitives = LODModel.IndexBuffer.GetNumIndices() / 3;
+			OutMeshBatchElement.IndexBuffer = &IndexBuffer;
+			NumPrimitives = IndexBuffer.GetNumIndices() / 3;
 
 			OutMeshBatch.bWireframe = true;
 			OutMeshBatch.bDisableBackfaceCulling = true;
@@ -1040,7 +1104,7 @@ uint32 FStaticMeshSceneProxy::SetMeshElementGeometrySource(
 	{
 		OutMeshBatch.Type = PT_TriangleList;
 
-		OutMeshBatchElement.IndexBuffer = bUseReversedIndices ? &LODModel.AdditionalIndexBuffers->ReversedIndexBuffer : &LODModel.IndexBuffer;
+		OutMeshBatchElement.IndexBuffer = bUseReversedIndices ? &AdditionalIndexBuffers->ReversedIndexBuffer : &IndexBuffer;
 		OutMeshBatchElement.FirstIndex = Section.FirstIndex;
 		NumPrimitives = Section.NumTriangles;
 	}
@@ -1051,6 +1115,22 @@ uint32 FStaticMeshSceneProxy::SetMeshElementGeometrySource(
 	return NumPrimitives;
 }
 
+uint32 FStaticMeshSceneProxy::SetMeshElementGeometrySource(
+	int32 LODIndex,
+	int32 SectionIndex,
+	bool bWireframe,
+	bool bUseReversedIndices,
+	bool bAllowPreCulledIndices, // unused
+	const FVertexFactory* VertexFactory,
+	FMeshBatch& OutMeshBatch) const
+{
+	const FStaticMeshLODResources& LODModel = RenderData->LODResources[LODIndex];
+
+	const FStaticMeshSection& Section = LODModel.Sections[SectionIndex];
+
+	return SetMeshElementGeometrySource(Section, LODModel.IndexBuffer, LODModel.AdditionalIndexBuffers, VertexFactory, bWireframe, bUseReversedIndices, OutMeshBatch);
+}
+
 void FStaticMeshSceneProxy::SetMeshElementScreenSize(int32 LODIndex, bool bDitheredLODTransition, FMeshBatch& OutMeshBatch) const
 {
 	FMeshBatchElement& OutBatchElement = OutMeshBatch.Elements[0];
@@ -1843,7 +1923,11 @@ void FStaticMeshSceneProxy::GetDynamicRayTracingInstances(FRayTracingInstanceCol
 
 	int32 LODIndex = RayTracingMinLOD;
 
-	if (!bEvaluateWPO)
+	// we need mark the geometry group as referenced (for streaming/residency management) when:
+	// - not evaluating WPO, since the static ray tracing geometry will be used.
+	// - not using rendering LODs, since VB/IB need to be streamed-in
+	//		- TODO: support streaming VB/IB independently from RayTracingGeometry
+	if (!bEvaluateWPO || !RenderData->RayTracingProxy->bUsingRenderingLODs)
 	{
 		Collector.AddReferencedGeometryGroup(RenderData->RayTracingGeometryGroupHandle);
 
@@ -1909,12 +1993,23 @@ void FStaticMeshSceneProxy::GetDynamicRayTracingInstances(FRayTracingInstanceCol
 				{
 					FMeshBatch &MeshBatch = CachedRayTracingMaterials.AddDefaulted_GetRef();
 
-					bool bResult = GetMeshElement(LODIndex, BatchIndex, SectionIndex, PrimitiveDPG, false, false, MeshBatch); // todo: RayTracingLOD vertex factory
+					bool bResult;
+					if (RenderData->RayTracingProxy->bUsingRenderingLODs)
+					{
+						// when using rendering LODs we can reuse the main GetMeshElement(...)
+						bResult = GetMeshElement(LODIndex, BatchIndex, SectionIndex, PrimitiveDPG, false, false, MeshBatch);
+					}
+					else
+					{
+						// otherwise initialize MeshBatch using ray tracing proxy VB/IB/Section data
+						bResult = GetRayTracingMeshElement(LODIndex, BatchIndex, SectionIndex, PrimitiveDPG, MeshBatch);
+					}
+
 					if (!bResult)
 					{
 						// Hidden material
 						MeshBatch.MaterialRenderProxy = UMaterial::GetDefaultMaterial(MD_Surface)->GetRenderProxy();
-						MeshBatch.VertexFactory = &RenderData->LODVertexFactories[LODIndex].VertexFactory;
+						MeshBatch.VertexFactory = &(*RenderData->RayTracingProxy->LODVertexFactories)[LODIndex].VertexFactory;
 					}
 					MeshBatch.ReverseCulling = bReverseCulling; // overwrite what came from GetMeshElement as DXR only needs the user driven flag, not the flipping implied by the transform
 					MeshBatch.SegmentIndex = SectionIndex;
@@ -1939,7 +2034,7 @@ void FStaticMeshSceneProxy::GetDynamicRayTracingInstances(FRayTracingInstanceCol
 		const FMatrix& ThisLocalToWorld = GetLocalToWorld();
 		RayTracingInstance.InstanceTransformsView = MakeArrayView(&ThisLocalToWorld, 1);
 
-		if (bEvaluateWPO && RenderData->LODVertexFactories[LODIndex].VertexFactory.GetType()->SupportsRayTracingDynamicGeometry())
+		if (bEvaluateWPO && (*RenderData->RayTracingProxy->LODVertexFactories)[LODIndex].VertexFactory.GetType()->SupportsRayTracingDynamicGeometry())
 		{
 			const uint32 NumVertices = RayTracingLOD.VertexBuffers->PositionVertexBuffer.GetNumVertices();
 
@@ -2344,6 +2439,41 @@ FLightInteraction FStaticMeshSceneProxy::FLODInfo::GetInteraction(const FLightSc
 	return FLightInteraction::Dynamic();
 }
 
+FStaticMeshSceneProxy::FRayTracingLODInfo::FRayTracingLODInfo(const FStaticMeshSceneProxyDesc& InProxyDesc, int32 LODIndex)
+{
+	const auto FeatureLevel =  InProxyDesc.GetWorld()->GetFeatureLevel();
+
+	FStaticMeshRenderData* MeshRenderData = InProxyDesc.GetStaticMesh()->GetRenderData();
+	FStaticMeshRayTracingProxyLOD& LOD = MeshRenderData->RayTracingProxy->LODs[LODIndex];
+
+	bool bForceDefaultMaterial = InProxyDesc.ShouldRenderProxyFallbackToDefaultMaterial();
+
+	// Gather the materials applied to the LOD.
+	Sections.Empty(LOD.Sections->Num());
+	for(int32 SectionIndex = 0;SectionIndex < LOD.Sections->Num();SectionIndex++)
+	{
+		const FStaticMeshSection& Section = (*LOD.Sections)[SectionIndex];
+		FSectionInfo SectionInfo;
+
+		// Determine the material applied to this element of the LOD.
+		SectionInfo.Material = InProxyDesc.GetMaterial(Section.MaterialIndex);
+
+		if (bForceDefaultMaterial || (GForceDefaultMaterial && SectionInfo.Material && !IsTranslucentBlendMode(*SectionInfo.Material)))
+		{
+			SectionInfo.Material = UMaterial::GetDefaultMaterial(MD_Surface);
+		}
+
+		// If there isn't an applied material, fall back to the default material.
+		if(!SectionInfo.Material)
+		{
+			SectionInfo.Material = UMaterial::GetDefaultMaterial(MD_Surface);
+		}
+
+		// Store the element info.
+		Sections.Add(SectionInfo);
+	}
+}
+
 float FStaticMeshSceneProxy::GetScreenSize( int32 LODIndex ) const
 {
 	return RenderData->ScreenSize[LODIndex].GetValue() * GetLodScreenSizeScale();
diff --git a/Engine/Source/Runtime/Engine/Public/StaticMeshSceneProxy.h b/Engine/Source/Runtime/Engine/Public/StaticMeshSceneProxy.h
--- a/Engine/Source/Runtime/Engine/Public/StaticMeshSceneProxy.h
+++ b/Engine/Source/Runtime/Engine/Public/StaticMeshSceneProxy.h
@@ -14,6 +14,8 @@
 
 class FLocalVertexFactoryUniformShaderParameters;
 class FRawStaticIndexBuffer;
+struct FAdditionalStaticMeshIndexBuffers;
+struct FStaticMeshSection;
 struct FStaticMeshVertexFactories;
 using FStaticMeshVertexFactoriesArray = TArray<FStaticMeshVertexFactories>;
 struct FStaticMeshSceneProxyDesc;
@@ -48,12 +50,15 @@ public:
 	ENGINE_API virtual bool GetMeshElement(
 		int32 LODIndex, 
 		int32 BatchIndex, 
-		int32 ElementIndex, 
+		int32 SectionIndex,
 		uint8 InDepthPriorityGroup, 
 		bool bUseSelectionOutline,
 		bool bAllowPreCulledIndices,
 		FMeshBatch& OutMeshBatch) const;
 
+	/** Sets up a ray tracing FMeshBatch for a specific LOD. */
+	ENGINE_API virtual bool GetRayTracingMeshElement(int32 LODIndex, int32 BatchIndex, int32 SectionIndex, uint8 InDepthPriorityGroup, FMeshBatch& OutMeshBatch) const;
+
 	ENGINE_API virtual void CreateRenderThreadResources(FRHICommandListBase& RHICmdList) override;
 
 	ENGINE_API virtual void DestroyRenderThreadResources() override;
@@ -87,7 +92,7 @@ public:
 	ENGINE_API virtual bool GetInstanceWorldPositionOffsetDisableDistance(float& OutWPODisableDistance) const override;
 
 protected:
-	/** Configures mesh batch vertex / index state. Returns the number of primitives used in the element. */
+	UE_DEPRECATED(5.6, "Provide FStaticMeshSection, IndexBuffer and AdditionalIndexBuffers instead of LODIndex, etc.")
 	ENGINE_API uint32 SetMeshElementGeometrySource(
 		int32 LODIndex,
 		int32 ElementIndex,
@@ -97,6 +102,16 @@ protected:
 		const FVertexFactory* VertexFactory,
 		FMeshBatch& OutMeshElement) const;
 
+	/** Configures mesh batch vertex / index state. Returns the number of primitives used in the element. */
+	ENGINE_API uint32 SetMeshElementGeometrySource(
+		const FStaticMeshSection& Section,
+		const FRawStaticIndexBuffer& IndexBuffer,
+		const FAdditionalStaticMeshIndexBuffers* AdditionalIndexBuffers,
+		const FVertexFactory* VertexFactory,
+		bool bWireframe,
+		bool bUseInversedIndices,
+		FMeshBatch& OutMeshElement) const;
+
 	/** Sets the screen size on a mesh element. */
 	ENGINE_API void SetMeshElementScreenSize(int32 LODIndex, bool bDitheredLODTransition, FMeshBatch& OutMeshBatch) const;
 
@@ -219,9 +234,27 @@ protected:
 		bool bUsesMeshModifyingMaterials;
 	};
 
+	/** Information used by the proxy about a single ray tracing proxy LOD of the mesh. */
+	class FRayTracingLODInfo
+	{
+	public:
+
+		struct FSectionInfo
+		{
+			UMaterialInterface* Material = nullptr;
+		};
+
+		TArray<FSectionInfo, TInlineAllocator<1>> Sections;
+
+		// TODO: Support OverrideColorVertexBuffer (needs to be generated similar to what is done for regular LODs)
+
+		FRayTracingLODInfo(const FStaticMeshSceneProxyDesc& InProxyDesc, int32 InLODIndex);
+	};
+
 	FStaticMeshRenderData* RenderData;
 
 	TArray<FLODInfo> LODs;
+	TArray<FRayTracingLODInfo> RayTracingLODInfos;
 
 	const FDistanceFieldVolumeData* DistanceFieldData;	
 	const FCardRepresentationData* CardRepresentationData;	
diff --git a/Engine/Source/Runtime/Renderer/Private/PrimitiveSceneInfo.cpp b/Engine/Source/Runtime/Renderer/Private/PrimitiveSceneInfo.cpp
--- a/Engine/Source/Runtime/Renderer/Private/PrimitiveSceneInfo.cpp
+++ b/Engine/Source/Runtime/Renderer/Private/PrimitiveSceneInfo.cpp
@@ -1144,7 +1144,7 @@ void CacheRayTracingPrimitive(
 		else
 		{
 			const bool bMustEmitCommand = false;
-			for (const FStaticMeshBatch& Mesh : SceneInfo->StaticMeshes)
+			for (const FStaticMeshBatch& Mesh : SceneInfo->StaticMeshes) // TODO: Handle !RayTracingProxy->bUsingRenderingLODs
 			{
 				CacheRayTracingMeshBatch<bDeferLODCommandIndices>(Mesh, SceneInfo, RayTracingLODData, Commands, CommandContext, RayTracingMeshProcessor, DeferredMeshCommandData, bMustEmitCommand);
 			}
diff --git a/Engine/Source/Runtime/Renderer/Private/RayTracing/RayTracing.cpp b/Engine/Source/Runtime/Renderer/Private/RayTracing/RayTracing.cpp
--- a/Engine/Source/Runtime/Renderer/Private/RayTracing/RayTracing.cpp
+++ b/Engine/Source/Runtime/Renderer/Private/RayTracing/RayTracing.cpp
@@ -1094,6 +1094,8 @@ namespace RayTracing
 					FLODMask LODToRender = ComputeLODForMeshes(SceneInfo->StaticMeshRelevances, View, Bounds.BoxSphereBounds.Origin, Bounds.BoxSphereBounds.SphereRadius, ForcedLODLevel, MeshScreenSizeSquared, CurFirstLODIdx, LODScale, true);
 
 					LODIndex = LODToRender.GetRayTracedLOD();
+
+					// TODO: Handle !RayTracingProxy->bUsingRenderingLODs
 				}
 
 				if (EnumHasAnyFlags(Flags, ERayTracingPrimitiveFlags::CacheInstances))
