Replace broken GeomTools 2D convex hull method with a call to the Core's 2D convex hull method + add unit test coverage

diff --git a/Engine/Source/Runtime/Engine/Private/GeomTools.cpp b/Engine/Source/Runtime/Engine/Private/GeomTools.cpp
--- a/Engine/Source/Runtime/Engine/Private/GeomTools.cpp
+++ b/Engine/Source/Runtime/Engine/Private/GeomTools.cpp
@@ -6,6 +6,7 @@
 #include "GeomTools.h"
 #include "StaticMeshResources.h"
 #include "Engine/StaticMesh.h"
+#include "Math/ConvexHull2d.h"
 
 DEFINE_LOG_CATEGORY_STATIC(LogGeomTools, Log, All);
 
@@ -1619,82 +1620,13 @@ void FGeomTools2D::GenerateConvexHullFromPoints(TArray<FVector2D>& OutConvexHull
 		return;
 	}
 
-	// Find lowest point. If multiple points have the same y, find leftmost
-	int LowestPointIndex = 0;
-	FVector2D& LowestPoint = SourcePoints[0];
-	for (int I = 1; I < SourcePointsCount; ++I)
-	{
-		if ((SourcePoints[I].Y < LowestPoint.Y) || (SourcePoints[I].Y == LowestPoint.Y && SourcePoints[I].X < LowestPoint.X))
-		{
-			LowestPointIndex = I;
-			LowestPoint = SourcePoints[I];
-		}
-	}
-
-	// Get indices sorted by angle they and the lowest point make with the x-axis (excluding the lowest point)
-	FVector2D& P = SourcePoints[LowestPointIndex];
-	TArray<int32> SortedPoints;
-	SortedPoints.Empty(SourcePointsCount);
-	for (int I = 0; I < SourcePointsCount; ++I)
-	{
-		if (I != LowestPointIndex)
-		{
-			SortedPoints.Add(I);
-		}
-	}
-	int SortedPointsCount = SortedPoints.Num();
-	for (int I = 0; I < SortedPointsCount; ++I)
-	{
-		for (int J = I + 1; J < SortedPointsCount; ++J)
-		{
-			FVector2D DI = SourcePoints[SortedPoints[I]] - P;
-			FVector2D DJ = SourcePoints[SortedPoints[J]] - P;
-			if (FMath::Atan2(DI.Y, DI.X) > FMath::Atan2(DJ.Y, DJ.X))
-			{
-				int32 Temp = SortedPoints[I];
-				SortedPoints[I] = SortedPoints[J];
-				SortedPoints[J] = Temp;
-			}
-		}
-	}
-
-	// Initialize hull points and add lowest point to end of sorted points for the algorithm
-	TArray<int> Hull;
-	Hull.Add(LowestPointIndex);
-	SortedPoints.Add(LowestPointIndex);
-	++SortedPointsCount;
-	// Traverse sorted points, removing all prior points that would make a right turn, before adding each new point (Graham scan)
-	for (int I = 0; I < SortedPointsCount; ++I)
-	{
-		int NewPointIndex = SortedPoints[I];
-		FVector2D& C = SourcePoints[NewPointIndex];
-		while (Hull.Num() > 1)
-		{
-			FVector2D& A = SourcePoints[Hull[Hull.Num() - 2]];
-			FVector2D& B = SourcePoints[Hull[Hull.Num() - 1]];
-			if ((B.X - A.X) * (C.Y - A.Y) - (B.Y - A.Y) * (C.X - A.X) < 0.0f)
-			{
-				// Remove last entry
-				Hull.RemoveAt(Hull.Num() - 1);
-			}
-			else
-			{
-				// Left turn, exit loop
-				break;
-			}
-		}
-		Hull.Add(NewPointIndex);
-	}
-
-	// We added the starting point to the end for the algorithm, so remove the duplicate
-	Hull.RemoveAt(Hull.Num() - 1);
+	TArray<int32> HullIndices;
+	ConvexHull2D::ComputeConvexHull(SourcePoints, HullIndices);
+	OutConvexHull.Empty(HullIndices.Num());
 
-	// Result
-	int HullCount = Hull.Num();
-	OutConvexHull.Empty(HullCount);
-	for (int I = 0; I < HullCount; ++I)
+	for (int32 HullIdx : HullIndices)
 	{
-		OutConvexHull.Add(SourcePoints[Hull[I]]);
+		OutConvexHull.Add(SourcePoints[HullIdx]);
 	}
 }
 
