World Partition: rerun construction scripts for actors in levels before executing the cell transformers, as the Blueprint parent classes might have changed and we want to make sure the constructed actors are up to date.

diff --git a/Engine/Source/Runtime/Engine/Private/WorldPartition/WorldPartition.cpp b/Engine/Source/Runtime/Engine/Private/WorldPartition/WorldPartition.cpp
--- a/Engine/Source/Runtime/Engine/Private/WorldPartition/WorldPartition.cpp
+++ b/Engine/Source/Runtime/Engine/Private/WorldPartition/WorldPartition.cpp
@@ -4,6 +4,7 @@
 	WorldPartition.cpp: UWorldPartition implementation
 =============================================================================*/
 #include "WorldPartition/WorldPartition.h"
+#include "Algo/Copy.h"
 #include "Engine/Engine.h"
 #include "Misc/PackageName.h"
 #include "UObject/AssetRegistryTagsContext.h"
@@ -2481,6 +2482,20 @@ void UWorldPartition::ApplyRuntimeCellsTransformerStack(ULevel* InLevel)
 		}
 	};
 
+	// Rerun construction scripts for actors in levels before executing the cell transformers, as the Blueprint parent classes might have changed 
+	// and we want to make sure the constructed actors are up to date.
+	TArray<TObjectPtr<AActor>> LevelActors;
+	Algo::CopyIf(InLevel->Actors, LevelActors, [](AActor* InActor) { return IsValid(InActor); });
+	
+	for (TObjectPtr<AActor>& LevelActor : LevelActors)
+	{
+		// Deal with deleted child actors that might have been deleted from previous calls to RerunConstructionScripts
+		if (IsValid(LevelActor))
+		{
+			LevelActor->RerunConstructionScripts();
+		}
+	}
+
 	ApplyTransformPhase([](const FRuntimeCellTransformerInstance& TransformerInstance, ULevel* Level) { TransformerInstance.PreTransform(Level); });
 	ApplyTransformPhase([](const FRuntimeCellTransformerInstance& TransformerInstance, ULevel* Level) { TransformerInstance.Transform(Level); });
 	ApplyTransformPhase([](const FRuntimeCellTransformerInstance& TransformerInstance, ULevel* Level) { TransformerInstance.PostTransform(Level); });
