Added a repair check to brushactors in the case the owned brushcomponent's brush reference gets removed.

diff --git a/Engine/Source/Runtime/Engine/Private/Brush.cpp b/Engine/Source/Runtime/Engine/Private/Brush.cpp
--- a/Engine/Source/Runtime/Engine/Private/Brush.cpp
+++ b/Engine/Source/Runtime/Engine/Private/Brush.cpp
@@ -341,6 +341,15 @@ void ABrush::PostLoad()
 		}
 
 #if WITH_EDITOR
+		// Fix up corrupted brush references.
+		if (BrushComponent)
+		{
+			if (!ensureMsgf(BrushComponent->Brush, TEXT("Component brush missing, yet the brush actor still has it! This should not be possible.")))
+			{
+				BrushComponent->Brush = Brush;
+			}
+		}
+
 		// Fix up any broken poly normals.
 		// They have not been getting fixed up after vertex editing since at least UE2!
 		for(FPoly& Poly : Brush->Polys->Element)
