Substrate - tagging opaque rough refraction as experimental as it is not well tested, validated in production and not available in the path tracer.

#tests PC editor
#rnx

diff --git a/Engine/Source/Runtime/Engine/Classes/Engine/RendererSettings.h b/Engine/Source/Runtime/Engine/Classes/Engine/RendererSettings.h
--- a/Engine/Source/Runtime/Engine/Classes/Engine/RendererSettings.h
+++ b/Engine/Source/Runtime/Engine/Classes/Engine/RendererSettings.h
@@ -1204,8 +1204,8 @@ class URendererSettings : public UDeveloperSettings
 	"Enable Substrate opaque material rough refractions effect from top layers over layers below."
 	*/
 	UPROPERTY(config, EditAnywhere, Category = Substrate, meta = (
-		ConsoleVariable = "r.Substrate.OpaqueMaterialRoughRefraction", DisplayName = "Substrate opaque material rough refraction",
-		ToolTip = "Enable Substrate opaque material rough refractions effect from top layers over layers below.",
+		ConsoleVariable = "r.Substrate.OpaqueMaterialRoughRefraction", DisplayName = "Substrate opaque material rough refraction (Experimental)",
+		ToolTip = "Enable Substrate opaque material rough refractions effect from top layers over layers below. (Experimental for now since this is not a path validated in production)",
 		ConfigRestartRequired = true))
 		uint32 SubstrateOpaqueMaterialRoughRefraction : 1;
 		
diff --git a/Engine/Source/Runtime/RenderCore/Private/RenderUtils.cpp b/Engine/Source/Runtime/RenderCore/Private/RenderUtils.cpp
--- a/Engine/Source/Runtime/RenderCore/Private/RenderUtils.cpp
+++ b/Engine/Source/Runtime/RenderCore/Private/RenderUtils.cpp
@@ -1826,7 +1826,7 @@ static TAutoConsoleVariable<int32> CVarSubstrateBackCompatibility(
 static TAutoConsoleVariable<int32> CVarSubstrateOpaqueMaterialRoughRefraction(
 	TEXT("r.Substrate.OpaqueMaterialRoughRefraction"),
 	0,
-	TEXT("Enable Substrate opaque material rough refractions effect from top layers over layers below."),
+	TEXT("Enable Substrate opaque material rough refractions effect from top layers over layers below. (Experimental for now since this is not a path validated in production)"),
 	ECVF_ReadOnly | ECVF_RenderThreadSafe);
 
 static TAutoConsoleVariable<int32> CVarSubstrateSheenQuality(
