Use the TextureLODSettings of the preview DP when preview is active

diff --git a/Engine/Source/Runtime/Engine/Classes/Engine/TextureLODSettings.h b/Engine/Source/Runtime/Engine/Classes/Engine/TextureLODSettings.h
--- a/Engine/Source/Runtime/Engine/Classes/Engine/TextureLODSettings.h
+++ b/Engine/Source/Runtime/Engine/Classes/Engine/TextureLODSettings.h
@@ -231,4 +231,20 @@ public:
 	/** Array of LOD settings with entries per group. */
 	UPROPERTY(EditAnywhere, config, Category="Texture LOD Settings")
 	TArray<FTextureLODGroup> TextureLODGroups;
+
+#if ALLOW_OTHER_PLATFORM_CONFIG
+public:
+	ENGINE_API void SetPreviewTextureGroupRedirect(UTextureLODSettings* InPreviewTextureGroupRedirect)
+	{
+		PreviewTextureGroupRedirect = InPreviewTextureGroupRedirect;
+	}
+
+	ENGINE_API UTextureLODSettings* GetPreviewTextureGroupRedirect() const
+	{
+		return PreviewTextureGroupRedirect;
+	}
+
+private:
+	UTextureLODSettings* PreviewTextureGroupRedirect;
+#endif
 };
diff --git a/Engine/Source/Runtime/Engine/Private/DeviceProfiles/DeviceProfile.cpp b/Engine/Source/Runtime/Engine/Private/DeviceProfiles/DeviceProfile.cpp
--- a/Engine/Source/Runtime/Engine/Private/DeviceProfiles/DeviceProfile.cpp
+++ b/Engine/Source/Runtime/Engine/Private/DeviceProfiles/DeviceProfile.cpp
@@ -71,6 +71,14 @@ void UDeviceProfile::GatherParentCVarInformationRecursively(OUT TMap<FString, FS
 
 UTextureLODSettings* UDeviceProfile::GetTextureLODSettings() const
 {
+#if ALLOW_OTHER_PLATFORM_CONFIG
+	UTextureLODSettings* PreviewTextureGroup = GetPreviewTextureGroupRedirect();
+	if (PreviewTextureGroup != nullptr)
+	{
+		return PreviewTextureGroup;
+	}
+#endif
+
 	return (UTextureLODSettings*)this;
 }
 
diff --git a/Engine/Source/Runtime/Engine/Private/DeviceProfiles/DeviceProfileManager.cpp b/Engine/Source/Runtime/Engine/Private/DeviceProfiles/DeviceProfileManager.cpp
--- a/Engine/Source/Runtime/Engine/Private/DeviceProfiles/DeviceProfileManager.cpp
+++ b/Engine/Source/Runtime/Engine/Private/DeviceProfiles/DeviceProfileManager.cpp
@@ -983,9 +983,14 @@ void UDeviceProfileManager::LoadProfiles()
 				ITargetPlatform* Platform = TargetPlatforms[PlatformIndex];
 
 				// Set TextureLODSettings
-				const UTextureLODSettings* TextureLODSettingsObj = FindProfile(Platform->CookingDeviceProfileName(), false);
+				UTextureLODSettings* TextureLODSettingsObj = FindProfile(Platform->CookingDeviceProfileName(), false);
 				checkf(TextureLODSettingsObj, TEXT("No TextureLODSettings found for %s"), *Platform->CookingDeviceProfileName());
 
+#if ALLOW_OTHER_PLATFORM_CONFIG
+				// disable redirect if we are using this for cooking.
+				TextureLODSettingsObj->SetPreviewTextureGroupRedirect(nullptr);
+#endif
+
 				Platform->RegisterTextureLODSettings(TextureLODSettingsObj);
 			}
 
@@ -1064,6 +1069,8 @@ void UDeviceProfileManager::SetPreviewDeviceProfile(UDeviceProfile* DeviceProfil
 
 	PreviewDeviceProfile = DeviceProfile;
 	
+	ActiveDeviceProfile->SetPreviewTextureGroupRedirect(PreviewDeviceProfile);
+
 	// walk over all cvars and apply them
 	IConsoleManager::Get().PreviewPlatformCVars(*DeviceProfile->DeviceType, DeviceProfile->GetName(), PreviewModeTag);
 
@@ -1081,6 +1088,8 @@ void UDeviceProfileManager::RestorePreviewDeviceProfile(FName PreviewModeTag)
 
 	PreviewDeviceProfile = nullptr;
 
+	ActiveDeviceProfile->SetPreviewTextureGroupRedirect(nullptr);
+
 	IConsoleManager::Get().UnsetAllConsoleVariablesWithTag(PreviewModeTag, ECVF_SetByPreview);
 }
 #endif
