CookIndeterminism: AnimSequence: Some AnimSequences have indeterminism in their GetDataModel()->GenerateGuid() and CurrentSkeleton->GetGuid() strings that affects their generated bulk data and triggers intermittent failures of our cook determinism testing. Add AnimIAnimationDataModel::GenerateDebugStateString and store it in UAnimSequence::PreSave's RegisterDeterminismHelper diagnostics to help debug this when the farm reproduces it. This stores a large amount of data stored in the cook attachments for AnimSequences when running with -cookdeterminism, so we may want to configure it off once we have diagnosed the issue.
#rnx

diff --git a/Engine/Plugins/Animation/AnimationData/Source/AnimationData/Private/AnimSequencerDataModel.cpp b/Engine/Plugins/Animation/AnimationData/Source/AnimationData/Private/AnimSequencerDataModel.cpp
--- a/Engine/Plugins/Animation/AnimationData/Source/AnimationData/Private/AnimSequencerDataModel.cpp
+++ b/Engine/Plugins/Animation/AnimationData/Source/AnimationData/Private/AnimSequencerDataModel.cpp
@@ -3,6 +3,7 @@
 #include "AnimSequencerDataModel.h"
 #include "AnimSequencerController.h"
 
+#include "Animation/AnimDataModelHasher.h"
 #include "Animation/AnimSequence.h"
 #include "AnimDataController.h"
 #include "Async/ParallelFor.h"
@@ -586,128 +587,83 @@ UAnimSequence* UAnimationSequencerDataModel::GetAnimationSequence() const
 	return Cast<UAnimSequence>(GetOuter());
 }
 
-FGuid UAnimationSequencerDataModel::GenerateGuid() const
-{	
-	FGuid ReturnGuid;
-	if (CachedRawDataGUID.IsValid())
-	{
-		ReturnGuid = CachedRawDataGUID;
-	}
-	else
-	{
-		FSHA1 Sha;
-		const FString ClassName = GetClass()->GetName();
-		Sha.UpdateWithString(*ClassName, ClassName.Len());
-		
-		auto UpdateSHAWithArray = [&Sha](const auto& Array)
-		{
-			if (Array.Num())
-			{
-				Sha.Update(reinterpret_cast<const uint8*>(Array.GetData()), Array.Num() * Array.GetTypeSize());			
-			}
-		};
+template <typename HasherType>
+void UAnimationSequencerDataModel::GenerateStateHash(HasherType& Hasher) const
+{
+	Hasher.UpdateString(GetClass()->GetName(), TEXT("C"));
 
-		auto UpdateWithChannel = [UpdateSHAWithArray, &Sha](const auto& Channel)
+	auto UpdateWithChannel = [&Hasher](const auto& Channel, const TCHAR* Name)
 		{
-			UpdateSHAWithArray(Channel.GetData().GetTimes());
-			UpdateSHAWithArray(Channel.GetData().GetValues());
+			Hasher.BeginObject(Name);
+			Hasher.UpdateArray(Channel.GetData().GetTimes(), TEXT("T"));
+			Hasher.UpdateArray(Channel.GetData().GetValues(), TEXT("V"));
 			if (Channel.GetDefault().IsSet())
 			{
-				Sha.Update(reinterpret_cast<const uint8*>(&Channel.GetDefault().GetValue()), sizeof(Channel.GetDefault().GetValue()));
+				Hasher.UpdateData(Channel.GetDefault().GetValue(), TEXT("D"));
 			}
+			Hasher.EndObject();
 		};
 
-		if (const UMovieSceneControlRigParameterSection* RigSection = GetFKControlRigSection())
-		{
-			UpdateWithChannel(RigSection->Weight);
-
-			for (const FTransformParameterNameAndCurves& TransformParameter : RigSection->GetTransformParameterNamesAndCurves())
-			{
-				const FString ParameterString = TransformParameter.ParameterName.ToString();
-				Sha.UpdateWithString(*ParameterString, ParameterString.Len());
-				for (int32 Index = 0; Index < 3; ++Index)
-				{							
-					UpdateWithChannel(TransformParameter.Translation[Index]);
-					UpdateWithChannel(TransformParameter.Rotation[Index]);
-					UpdateWithChannel(TransformParameter.Scale[Index]);
-				}
-			}
-
-			for (const FScalarParameterNameAndCurve& ScalarCurve : RigSection->GetScalarParameterNamesAndCurves())
-			{
-				const FString ParameterString = ScalarCurve.ParameterName.ToString();
-				Sha.UpdateWithString(*ParameterString, ParameterString.Len());
-				UpdateWithChannel(ScalarCurve.ParameterCurve);
-			}
-		}
-		
-		auto UpdateWithData = [&Sha](const auto& Data)
-		{
-			Sha.Update(reinterpret_cast<const uint8*>(&Data), sizeof(Data));
-		};
+	if (const UMovieSceneControlRigParameterSection* RigSection = GetFKControlRigSection())
+	{
+		Hasher.BeginObject(TEXT("RIG"));
+		UpdateWithChannel(RigSection->Weight, TEXT("W"));
 
-		for (const FAnimatedBoneAttribute& Attribute : AnimatedBoneAttributes)
+		Hasher.BeginObject(TEXT("T"));
+		for (const FTransformParameterNameAndCurves& TransformParameter : RigSection->GetTransformParameterNamesAndCurves())
 		{
-			UpdateSHAWithArray(Attribute.Identifier.GetName().ToString().GetCharArray());
-			UpdateSHAWithArray(Attribute.Identifier.GetBoneName().ToString().GetCharArray());
-			UpdateWithData(Attribute.Identifier.GetBoneIndex());
-			UpdateSHAWithArray(Attribute.Identifier.GetType()->GetFName().ToString().GetCharArray());
-			const UScriptStruct* TypeStruct = Attribute.Identifier.GetType();
-			const uint32 StructSize = TypeStruct->GetPropertiesSize();
-			const bool bHasTypeHash = TypeStruct->GetCppStructOps()->HasGetTypeHash();
-			for (const FAttributeKey& Key : Attribute.Curve.GetConstRefOfKeys())
+			Hasher.BeginObject();
+			Hasher.UpdateString(TransformParameter.ParameterName.ToString(), TEXT("N"));
+			for (int32 Index = 0; Index < 3; ++Index)
 			{
-				UpdateWithData(Key.Time);
-				if (bHasTypeHash)
-				{
-					const uint32 KeyHash = TypeStruct->GetStructTypeHash(Key.GetValuePtr<uint8>());
-					UpdateWithData(KeyHash);
-				}
-				else
-				{
-					Sha.Update(Key.GetValuePtr<uint8>(), StructSize);
-				}
+				Hasher.BeginObject(TEXT("C"));
+				UpdateWithChannel(TransformParameter.Translation[Index], TEXT("T"));
+				UpdateWithChannel(TransformParameter.Rotation[Index], TEXT("R"));
+				UpdateWithChannel(TransformParameter.Scale[Index], TEXT("S"));
+				Hasher.EndObject();
 			}
+			Hasher.EndObject();
 		}
+		Hasher.EndObject();
 
-		auto UpdateWithFloatCurve = [&UpdateWithData, &UpdateSHAWithArray](const FRichCurve& Curve)
+		Hasher.BeginObject(TEXT("S"));
+		for (const FScalarParameterNameAndCurve& ScalarCurve : RigSection->GetScalarParameterNamesAndCurves())
 		{
-			UpdateWithData(Curve.DefaultValue);
-			UpdateSHAWithArray(Curve.GetConstRefOfKeys());
-			UpdateWithData(Curve.PreInfinityExtrap);
-			UpdateWithData(Curve.PostInfinityExtrap);
-		};
-
-		for (const FTransformCurve& Curve : GetTransformCurves())
-		{
-		   	const FString CurveName = Curve.GetName().ToString();
-    		UpdateSHAWithArray(CurveName.GetCharArray());
-
-			auto UpdateWithComponent = [&UpdateWithFloatCurve](const FVectorCurve& VectorCurve)
-			{
-				for (int32 ChannelIndex = 0; ChannelIndex < 3; ++ChannelIndex)
-				{
-					UpdateWithFloatCurve(VectorCurve.FloatCurves[ChannelIndex]);
-				}
-			};
-	
-			UpdateWithComponent(Curve.TranslationCurve);
-			UpdateWithComponent(Curve.RotationCurve);
-			UpdateWithComponent(Curve.ScaleCurve);
+			Hasher.BeginObject();
+			Hasher.UpdateString(ScalarCurve.ParameterName.ToString(), TEXT("N"));
+			UpdateWithChannel(ScalarCurve.ParameterCurve, TEXT("C"));
+			Hasher.EndObject();
 		}
+		Hasher.EndObject();
+		Hasher.EndObject();
+	}
 
-		Sha.Final();
+	Hasher.UpdateAnimatedBoneAttributes(AnimatedBoneAttributes, TEXT("ABA"));
+	Hasher.UpdateTransformCurves(GetTransformCurves(), TEXT("TC"));
+}
 
-		uint32 Hash[5];
-		Sha.GetHash(reinterpret_cast<uint8*>(Hash));
-		const FGuid Guid(Hash[0] ^ Hash[4], Hash[1], Hash[2], Hash[3]);
-		
-		ReturnGuid = Guid;
+FGuid UAnimationSequencerDataModel::GenerateGuid() const
+{
+	if (CachedRawDataGUID.IsValid())
+	{
+		return CachedRawDataGUID;
 	}
 
-	return ReturnGuid;
+	UE::Anim::DataModel::FHasherSha Hasher;
+	GenerateStateHash(Hasher);
+	return Hasher.FinalGuid();
 }
 
+#if WITH_EDITOR
+FString UAnimationSequencerDataModel::GenerateDebugStateString() const
+{
+	UE::Anim::DataModel::FHasherCopyToText Hasher;
+	GenerateStateHash(Hasher);
+	return Hasher.GetString();
+}
+#endif
+
+
 TScriptInterface<IAnimationDataController> UAnimationSequencerDataModel::GetController()
 {
 	TScriptInterface<IAnimationDataController> Controller = nullptr;
diff --git a/Engine/Plugins/Animation/AnimationData/Source/AnimationData/Private/AnimSequencerDataModel.h b/Engine/Plugins/Animation/AnimationData/Source/AnimationData/Private/AnimSequencerDataModel.h
--- a/Engine/Plugins/Animation/AnimationData/Source/AnimationData/Private/AnimSequencerDataModel.h
+++ b/Engine/Plugins/Animation/AnimationData/Source/AnimationData/Private/AnimSequencerDataModel.h
@@ -111,6 +111,9 @@ public:
 	virtual UAnimSequence* GetAnimationSequence() const override;
 	virtual FAnimDataModelModifiedEvent& GetModifiedEvent() override { return ModifiedEvent; }
 	virtual FGuid GenerateGuid() const override;
+#if WITH_EDITOR
+	virtual FString GenerateDebugStateString() const override;
+#endif
 	virtual TScriptInterface<IAnimationDataController> GetController() override;
 	virtual IAnimationDataModel::FModelNotifier& GetNotifier() override;
 	virtual FAnimDataModelModifiedDynamicEvent& GetModifiedDynamicEvent() override { return ModifiedEventDynamic; }
@@ -170,6 +173,9 @@ protected:
 	
 	void GenerateTransformKeysForControl(const FName& BoneName, const TArray<FFrameNumber>& FrameNumbers, TArray<FTransform>& InOutTransforms) const;
 private:	
+	template <typename HasherType>
+	void GenerateStateHash(HasherType& Hasher) const;
+
 	/** Dynamic delegate event allows scripting to register to any broadcast-ed notify. */
 	UPROPERTY(BlueprintAssignable, Transient, Category = AnimationDataModel, meta = (ScriptName = "ModifiedEvent", AllowPrivateAccess = "true"))
 	FAnimDataModelModifiedDynamicEvent ModifiedEventDynamic;	
diff --git a/Engine/Source/Runtime/Engine/Classes/Animation/AnimData/AnimDataModel.h b/Engine/Source/Runtime/Engine/Classes/Animation/AnimData/AnimDataModel.h
--- a/Engine/Source/Runtime/Engine/Classes/Animation/AnimData/AnimDataModel.h
+++ b/Engine/Source/Runtime/Engine/Classes/Animation/AnimData/AnimDataModel.h
@@ -72,6 +72,7 @@ public:
 	virtual FAnimDataModelModifiedEvent& GetModifiedEvent() override { return ModifiedEvent; }
 	ENGINE_API virtual FGuid GenerateGuid() const override;
 #if WITH_EDITOR
+	ENGINE_API virtual FString GenerateDebugStateString() const override;
 	ENGINE_API virtual void Evaluate(FAnimationPoseData& InOutPoseData, const UE::Anim::DataModel::FEvaluationContext& EvaluationContext) const override;
 #endif
 	ENGINE_API virtual TScriptInterface<IAnimationDataController> GetController() override;
@@ -100,6 +101,8 @@ private:
 	ENGINE_API FRichCurve* GetMutableRichCurve(const FAnimationCurveIdentifier& CurveIdentifier);
 
 	ENGINE_API USkeleton* GetSkeleton() const;
+	template <typename HasherType>
+	void GenerateStateHash(HasherType& Hasher) const;
 
 private:
 	UPROPERTY(Transient)
diff --git a/Engine/Source/Runtime/Engine/Classes/Animation/AnimData/IAnimationDataModel.h b/Engine/Source/Runtime/Engine/Classes/Animation/AnimData/IAnimationDataModel.h
--- a/Engine/Source/Runtime/Engine/Classes/Animation/AnimData/IAnimationDataModel.h
+++ b/Engine/Source/Runtime/Engine/Classes/Animation/AnimData/IAnimationDataModel.h
@@ -352,6 +352,16 @@ public:
 	*/
 	virtual FGuid GenerateGuid() const = 0;
 
+#if WITH_EDITOR
+	/**
+	* @return	String representing the data reported by GenerateGuid, used to debug indeterminism in GenerateGuid.
+	*/
+	virtual FString GenerateDebugStateString() const
+	{
+		return FString::Printf(TEXT("DebugStateString not implemented. Guid == %s"), *GenerateGuid().ToString());
+	}
+#endif
+
 	/**
 	* @return	IAnimationDataController instance used for mutating this model
 	*/
diff --git a/Engine/Source/Runtime/Engine/Private/Animation/AnimDataModel.cpp b/Engine/Source/Runtime/Engine/Private/Animation/AnimDataModel.cpp
--- a/Engine/Source/Runtime/Engine/Private/Animation/AnimDataModel.cpp
+++ b/Engine/Source/Runtime/Engine/Private/Animation/AnimDataModel.cpp
@@ -6,6 +6,7 @@
 #include "Animation/AnimSequence.h"
 
 #include "Algo/Accumulate.h"
+#include "Animation/AnimDataModelHasher.h"
 #include "EngineLogs.h"
 #include "UObject/UE5MainStreamObjectVersion.h"
 #include "Misc/SecureHash.h"
@@ -505,97 +506,49 @@ const FAnimCurveBase* UAnimDataModel::FindCurve(const FAnimationCurveIdentifier&
 	return nullptr;
 }
 
-FGuid UAnimDataModel::GenerateGuid() const
+template <typename HasherType>
+void UAnimDataModel::GenerateStateHash(HasherType& Hasher) const
 {
-	FSHA1 Sha;
-	
-	auto UpdateSHAWithArray = [&](const auto& Array)
-	{
-		Sha.Update((uint8*)Array.GetData(), Array.Num() * Array.GetTypeSize());
-	};
-	   	
-	auto UpdateWithData = [&](const auto& Data)
-	{
-		Sha.Update((uint8*)(&Data), sizeof(Data));
-	};
-
-	auto UpdateWithString = [&](const FString& Data)
-	{
-		Sha.UpdateWithString(*Data, Data.Len());
-	};
-	
+	Hasher.BeginObject(TEXT("BAT"));
 	for (const FBoneAnimationTrack& Track : BoneAnimationTracks)
 	{
-		UpdateSHAWithArray(Track.InternalTrackData.PosKeys);
-		UpdateSHAWithArray(Track.InternalTrackData.RotKeys);
-		UpdateSHAWithArray(Track.InternalTrackData.ScaleKeys);
+		Hasher.BeginObject();
+		Hasher.UpdateArray(Track.InternalTrackData.PosKeys, TEXT("P"));
+		Hasher.UpdateArray(Track.InternalTrackData.RotKeys, TEXT("R"));
+		Hasher.UpdateArray(Track.InternalTrackData.ScaleKeys, TEXT("S"));
+		Hasher.EndObject();
 	}
+	Hasher.EndObject();
 
-    auto UpdateWithFloatCurve = [&UpdateWithData, &UpdateSHAWithArray](const FRichCurve& Curve)
-    {
-    	UpdateWithData(Curve.DefaultValue);
-    	UpdateSHAWithArray(Curve.GetConstRefOfKeys());
-    	UpdateWithData(Curve.PreInfinityExtrap);
-    	UpdateWithData(Curve.PostInfinityExtrap);
-    };
-
+	Hasher.BeginObject(TEXT("FC"));
 	for (const FFloatCurve& Curve : CurveData.FloatCurves)
 	{
-		const FString CurveName = Curve.GetName().ToString();
-		UpdateSHAWithArray(CurveName.GetCharArray()),
-		UpdateWithFloatCurve(Curve.FloatCurve);
+		Hasher.BeginObject();
+		Hasher.UpdateLegacyString(Curve.GetName().ToString(), TEXT("N"));
+		Hasher.UpdateRichCurve(Curve.FloatCurve, TEXT("C"));
+		Hasher.EndObject();
 	}
+	Hasher.EndObject();
 
-    for (const FTransformCurve& Curve : CurveData.TransformCurves)
-    {
-    	const FString CurveName = Curve.GetName().ToString();
-    	UpdateSHAWithArray(CurveName.GetCharArray());
-
-    	auto UpdateWithComponent = [&UpdateWithFloatCurve](const FVectorCurve& VectorCurve)
-    	{
-    		for (int32 ChannelIndex = 0; ChannelIndex < 3; ++ChannelIndex)
-    		{
-    			UpdateWithFloatCurve(VectorCurve.FloatCurves[ChannelIndex]);
-    		}
-    	};
-	
-    	UpdateWithComponent(Curve.TranslationCurve);
-    	UpdateWithComponent(Curve.RotationCurve);
-    	UpdateWithComponent(Curve.ScaleCurve);
-    }
-	
-	for (const FAnimatedBoneAttribute& Attribute : AnimatedBoneAttributes)
-	{
-		UpdateSHAWithArray(Attribute.Identifier.GetName().ToString().GetCharArray());
-		UpdateSHAWithArray(Attribute.Identifier.GetBoneName().ToString().GetCharArray());
-		UpdateWithData(Attribute.Identifier.GetBoneIndex());
-		UpdateSHAWithArray(Attribute.Identifier.GetType()->GetFName().ToString().GetCharArray());
-		const UScriptStruct* TypeStruct = Attribute.Identifier.GetType();
-		const uint32 StructSize = TypeStruct->GetPropertiesSize();
-		const bool bHasTypeHash = TypeStruct->GetCppStructOps()->HasGetTypeHash();
-		for (const FAttributeKey& Key : Attribute.Curve.GetConstRefOfKeys())
-		{
-			UpdateWithData(Key.Time);
-			if (bHasTypeHash)
-			{
-				const uint32 KeyHash = TypeStruct->GetStructTypeHash(Key.GetValuePtr<uint8>());
-				UpdateWithData(KeyHash);
-			}
-			else
-			{
-				Sha.Update(Key.GetValuePtr<uint8>(), StructSize);
-			}
-		}
-	}
+	Hasher.UpdateTransformCurves(CurveData.TransformCurves, TEXT("TC"));
+	Hasher.UpdateAnimatedBoneAttributes(AnimatedBoneAttributes, TEXT("ABA"));
+}
 
-	Sha.Final();
+FGuid UAnimDataModel::GenerateGuid() const
+{
+	UE::Anim::DataModel::FHasherSha Hasher;
+	GenerateStateHash(Hasher);
+	return Hasher.FinalGuid();
+}
 
-	uint32 Hash[5];
-	Sha.GetHash((uint8*)Hash);
-	FGuid Guid(Hash[0] ^ Hash[4], Hash[1], Hash[2], Hash[3]);
-	
-	return Guid;
+#if WITH_EDITOR
+FString UAnimDataModel::GenerateDebugStateString() const
+{
+	UE::Anim::DataModel::FHasherCopyToText Hasher;
+	GenerateStateHash(Hasher);
+	return Hasher.GetString();
 }
+#endif
 
 TScriptInterface<IAnimationDataController> UAnimDataModel::GetController()
 {
diff --git a/Engine/Source/Runtime/Engine/Private/Animation/AnimSequence.cpp b/Engine/Source/Runtime/Engine/Private/Animation/AnimSequence.cpp
--- a/Engine/Source/Runtime/Engine/Private/Animation/AnimSequence.cpp
+++ b/Engine/Source/Runtime/Engine/Private/Animation/AnimSequence.cpp
@@ -939,6 +939,12 @@ public:
 	{
 		FString DDCKey = AnimSequence->CreateDerivedDataKeyString(Context.GetTargetPlatform());
 		FIoHash DDCKeyHash = FIoHash::Zero;
+		const IAnimationDataModel* Model = AnimSequence->DataModelInterface.GetInterface();
+		FString ModelStateString = Model ? Model->GenerateDebugStateString() : TEXT("NoDataModel");
+		PRAGMA_DISABLE_DEPRECATION_WARNINGS;
+		bool bRawDataGuidEqual = Model ? (AnimSequence->RawDataGuid == Model->GenerateGuid()): false;
+		PRAGMA_ENABLE_DEPRECATION_WARNINGS;
+
 		if (!DDCKey.IsEmpty())
 		{
 			FMemoryHasherBlake3 Writer;
@@ -950,6 +956,8 @@ public:
 		Writer.BeginObject();
 		Writer << "DDCKeyHash" << WriteToString<64>(DDCKeyHash);
 		Writer << "DDCKey" << DDCKey;
+		Writer << "Model" << ModelStateString;
+		Writer << "RawDataGuidEqual" << bRawDataGuidEqual;
 		Writer.EndObject();
 
 		Context.AddDiagnostic("UAnimSequence", Writer.Save());
diff --git a/Engine/Source/Runtime/Engine/Public/Animation/AnimDataModelHasher.h b/Engine/Source/Runtime/Engine/Public/Animation/AnimDataModelHasher.h
new file mode 100644
--- /dev/null
+++ b/Engine/Source/Runtime/Engine/Public/Animation/AnimDataModelHasher.h
@@ -0,0 +1,247 @@
+// Copyright Epic Games, Inc. All Rights Reserved.
+
+#pragma once
+
+#include "Animation/AnimCurveTypes.h"
+#include "Animation/AnimData/IAnimationDataModel.h"
+#include "Containers/Array.h"
+#include "Containers/ArrayView.h"
+#include "Curves/RichCurve.h"
+#include "Misc/Guid.h"
+#include "Misc/SecureHash.h"
+
+#if WITH_EDITOR
+#include "Misc/StringBuilder.h"
+#include "String/BytesToHex.h"
+#endif
+
+namespace UE::Anim::DataModel
+{
+
+/**
+ * Base class of Hasher classes used by IAnimationDataModel classes to share code between GenerateGuid and
+ * GenerateDebugStateString.
+ * This class uses no virtuals for performance; the function that takes a Hasher should be templated on the
+ * type of the Hasher.
+ *
+ * This base class shows the required interface and some convenience functions. It uses the Curiously Recurring
+ * Template Pattern to call the Subclass's version of internal function calls.
+ */
+template <typename SubClassType>
+class FHasherBase
+{
+public:
+	// Output functions; the base class does nothing and subclasses must override
+
+	void UpdateBytes(const uint8* Data, int32 Size, const TCHAR* Name)
+	{
+		checkf(false, TEXT("Subclass must implement"));
+	}
+	void BeginObject(const TCHAR* Name = nullptr)
+	{
+		checkf(false, TEXT("Subclass must implement"));
+	}
+	void EndObject()
+	{
+		checkf(false, TEXT("Subclass must implement"));
+	}
+
+	// Convenience functions; the base class implementation does not need to be overridden
+
+	template <typename T>
+	void UpdateData(const T& Data, const TCHAR* Name)
+	{
+		TypedThis()->UpdateBytes(reinterpret_cast<const uint8*>(&Data), sizeof(Data), Name);
+	}
+	template <typename T>
+	void UpdateArray(const TArray<T>& Array, const TCHAR* Name)
+	{
+		TypedThis()->UpdateBytes(reinterpret_cast<const uint8*>(Array.GetData()), Array.Num() * Array.GetTypeSize(), Name);
+	}
+	template <typename T>
+	void UpdateArray(TConstArrayView<T> Array, const TCHAR* Name)
+	{
+		TypedThis()->UpdateBytes(reinterpret_cast<const uint8*>(Array.GetData()), Array.Num() * Array.GetTypeSize(), Name);
+	}
+	void UpdateString(const FString& Data, const TCHAR* Name)
+	{
+		TypedThis()->UpdateBytes(reinterpret_cast<const uint8*>(*Data), Data.Len() * sizeof((*Data)[0]), Name);
+	}
+	/**
+	 * Previously, some code was serializing strings as GetCharArray, which includes the null terminator,
+	 * and we did not want to change behavior because it would change the key. Those sites call UpdateLegacyString
+	 * instead, and the subclasses that want to preserve the key override it.
+	 */
+	void UpdateLegacyString(const FString& Data, const TCHAR* Name)
+	{
+		TypedThis()->UpdateString(Data, Name);
+	}
+
+	void UpdateRichCurve(const FRichCurve& Curve, const TCHAR* Name)
+	{
+		SubClassType& Hasher = *TypedThis();
+
+		Hasher.BeginObject(Name);
+		Hasher.UpdateData(Curve.DefaultValue, TEXT("D"));
+		Hasher.UpdateArray(Curve.GetConstRefOfKeys(), TEXT("K"));
+		Hasher.UpdateData(Curve.PreInfinityExtrap, TEXT("E"));
+		Hasher.UpdateData(Curve.PostInfinityExtrap, TEXT("O"));
+		Hasher.EndObject();
+	};
+
+	void UpdateVectorCurve(const FVectorCurve& VectorCurve, const TCHAR* Name)
+	{
+		SubClassType& Hasher = *TypedThis();
+
+		Hasher.BeginObject(Name);
+		for (int32 ChannelIndex = 0; ChannelIndex < 3; ++ChannelIndex)
+		{
+			Hasher.UpdateRichCurve(VectorCurve.FloatCurves[ChannelIndex], TEXT("C"));
+		}
+		Hasher.EndObject();
+	};
+
+	void UpdateAnimatedBoneAttributes(TConstArrayView<FAnimatedBoneAttribute> AnimatedBoneAttributes, const TCHAR* Name)
+	{
+		SubClassType& Hasher = *TypedThis();
+
+		Hasher.BeginObject(Name);
+		for (const FAnimatedBoneAttribute& Attribute : AnimatedBoneAttributes)
+		{
+			const UScriptStruct* TypeStruct = Attribute.Identifier.GetType();
+			const uint32 StructSize = TypeStruct->GetPropertiesSize();
+			const bool bHasTypeHash = TypeStruct->GetCppStructOps()->HasGetTypeHash();
+
+			Hasher.BeginObject();
+			Hasher.UpdateLegacyString(Attribute.Identifier.GetName().ToString(), TEXT("N"));
+			Hasher.UpdateLegacyString(Attribute.Identifier.GetBoneName().ToString(), TEXT("BN"));
+			Hasher.UpdateData(Attribute.Identifier.GetBoneIndex(), TEXT("BI"));
+			Hasher.UpdateLegacyString(TypeStruct->GetFName().ToString(), TEXT("T"));
+			Hasher.BeginObject(TEXT("K"));
+			for (const FAttributeKey& Key : Attribute.Curve.GetConstRefOfKeys())
+			{
+				Hasher.BeginObject();
+				Hasher.UpdateData(Key.Time, TEXT("T"));
+				if (bHasTypeHash)
+				{
+					const uint32 KeyHash = TypeStruct->GetStructTypeHash(Key.GetValuePtr<uint8>());
+					Hasher.UpdateData(KeyHash, TEXT("H"));
+				}
+				else
+				{
+					Hasher.UpdateBytes(Key.GetValuePtr<uint8>(), StructSize, TEXT("B"));
+				}
+				Hasher.EndObject();
+			}
+			Hasher.EndObject();
+			Hasher.EndObject();
+		}
+		Hasher.EndObject();
+	}
+
+	void UpdateTransformCurves(TConstArrayView<FTransformCurve> TransformCurves, const TCHAR* Name)
+	{
+		SubClassType& Hasher = *TypedThis();
+
+		Hasher.BeginObject(Name);
+		for (const FTransformCurve& Curve : TransformCurves)
+		{
+			Hasher.BeginObject();
+			Hasher.UpdateLegacyString(Curve.GetName().ToString(), TEXT("N"));
+			Hasher.UpdateVectorCurve(Curve.TranslationCurve, TEXT("T"));
+			Hasher.UpdateVectorCurve(Curve.RotationCurve, TEXT("R"));
+			Hasher.UpdateVectorCurve(Curve.ScaleCurve, TEXT("S"));
+			Hasher.EndObject();
+		}
+		Hasher.EndObject();
+	}
+
+private:
+	SubClassType* TypedThis()
+	{
+		return static_cast<SubClassType*>(this);
+	}
+
+};
+
+/** A hasher used for GenerateGuid, writes SHA and converts the SHA to a guid. */
+class FHasherSha : public FHasherBase<FHasherSha>
+{
+public:
+	void UpdateBytes(const uint8* Data, int32 Size, const TCHAR* Name)
+	{
+		Sha.Update(Data, Size);
+	}
+	void UpdateString(const FString& Data, const TCHAR* Name)
+	{
+		Sha.UpdateWithString(*Data, Data.Len());
+	}
+	void UpdateLegacyString(const FString& Data, const TCHAR* Name)
+	{
+		UpdateArray(Data.GetCharArray(), Name);
+	}
+	void BeginObject(const TCHAR* Name = nullptr)
+	{
+	}
+	void EndObject()
+	{
+	}
+
+	FGuid FinalGuid()
+	{
+		checkf(!bFinalized, TEXT("Calling FinalGuid more than once is not implemented."));
+		bFinalized = true;
+		Sha.Final();
+
+		uint32 Hash[5];
+		Sha.GetHash(reinterpret_cast<uint8*>(Hash));
+		return FGuid(Hash[0] ^ Hash[4], Hash[1], Hash[2], Hash[3]);
+	}
+
+private:
+	FSHA1 Sha;
+	bool bFinalized = false;
+};
+
+#if WITH_EDITOR
+/** A hasher used for GenerateDebugStateString. Writes data as string, untyped data as BytesToHex. */
+class FHasherCopyToText : public FHasherBase<FHasherCopyToText>
+{
+public:
+	void UpdateBytes(const uint8* Data, int32 Size, const TCHAR* Name)
+	{
+		Text << Divider << Name << Divider;
+		UE::String::BytesToHex(TConstArrayView<uint8>(Data, Size), Text);
+	}
+	void UpdateString(const FString& Data, const TCHAR* Name)
+	{
+		Text << Divider << Name << Divider << Data;
+	}
+	void BeginObject(const TCHAR* Name = nullptr)
+	{
+		Text << Divider << (Name ? Name : TEXT("")) << Open;
+	}
+	void EndObject()
+	{
+		Text << Close;
+	}
+
+	FStringView GetStringView() const
+	{
+		return Text;
+	}
+	FString GetString() const
+	{
+		return FString(Text);
+	}
+private:
+	TStringBuilder<256> Text;
+	constexpr static TCHAR Divider = '_';
+	constexpr static TCHAR Open = '{';
+	constexpr static TCHAR Close = '}';
+};
+#endif
+
+
+} // namespace UE::Anim::DataModel
+
