Editor: Change a brush ensure that fires on some broken volumes/bsp actors into a warning.

#virtualized

diff --git a/Engine/Source/Runtime/Engine/Private/Brush.cpp b/Engine/Source/Runtime/Engine/Private/Brush.cpp
--- a/Engine/Source/Runtime/Engine/Private/Brush.cpp
+++ b/Engine/Source/Runtime/Engine/Private/Brush.cpp
@@ -344,8 +344,11 @@ void ABrush::PostLoad()
 		// Fix up corrupted brush references.
 		if (BrushComponent)
 		{
-			if (!ensureMsgf(BrushComponent->Brush, TEXT("Component brush missing, yet the brush actor still has it! This should not be possible.")))
+			if (!BrushComponent->Brush)
 			{
+				UE_LOG(LogBrush, Warning, TEXT("Component inside %s had missing brush, but actor had it. "
+					"This should not happen. Resave package '%s' to remove this warning."),
+					*GetNameSafe(this), *GetNameSafe(GetPackage()));
 				BrushComponent->Brush = Brush;
 			}
 		}
