[NetPhysics] Reposition Predictive Interpolation if target looks like a teleport

diff --git a/Engine/Source/Runtime/Engine/Private/PhysicsEngine/PhysicsReplication.cpp b/Engine/Source/Runtime/Engine/Private/PhysicsEngine/PhysicsReplication.cpp
--- a/Engine/Source/Runtime/Engine/Private/PhysicsEngine/PhysicsReplication.cpp
+++ b/Engine/Source/Runtime/Engine/Private/PhysicsEngine/PhysicsReplication.cpp
@@ -254,6 +254,15 @@
 		
 		int32 TargetTickAlignmentClampMultiplier = 2;
 		static FAutoConsoleVariableRef CVarTargetTickAlignmentClampMultiplier(TEXT("np2.PredictiveInterpolation.TargetTickAlignmentClampMultiplier"), TargetTickAlignmentClampMultiplier, TEXT("Multiplier to adjust clamping of target alignment via TickCount. Multiplier is performed on AverageReceiveInterval."));
+
+		int32 TeleportDetectionEnabled = 1;
+		static FAutoConsoleVariableRef CVarTeleportDetectionEnabled(TEXT("np2.PredictiveInterpolation.TeleportDetection.Enabled"), TeleportDetectionEnabled, TEXT("Set to 1 to enable teleportation detection which hard snaps the replicated object if replication thinks a teleportation has happened. Disable by setting to 0 (or any other value currently, leaving room for adding more teleportation detection modes which will use incremental values)."));
+
+		float TeleportDetectionMinDistance = 200.0f;
+		static FAutoConsoleVariableRef CVarTeleportDetectionMinDistance(TEXT("np2.PredictiveInterpolation.TeleportDetection.MinDistance"), TeleportDetectionMinDistance, TEXT("Minimum positional distance between two received target states running teleportation detection."));
+
+		float TeleportDetectionVelocityMultiplier = 1.3f;
+		static FAutoConsoleVariableRef CVarTeleportDetectionVelocityMultiplier(TEXT("np2.PredictiveInterpolation.TeleportDetection.VelocityMultiplier"), TeleportDetectionVelocityMultiplier, TEXT("Multiplier to add leniency when checking if the previous or current velocity could cover the positional offset between previous and current target state. Higher value results in more lenient velocity comparison, i.e. less risk of triggering a hard snap when a teleport has not happened but also less likely to catch a teleport for objects that move while being teleported. Values under 1.0 are not recommended."));
 	}
 
 }
@@ -1081,6 +1090,9 @@
 		// Set if the target is allowed to be altered after this update
 		Target->bAllowTargetAltering = !(Target->TargetState.Flags & ERigidBodyFlags::Sleeping) && !(Input.TargetState.Flags & ERigidBodyFlags::Sleeping);
 
+		// Cache previous linear velocity
+		const FVector PrevLinVel = Target->TargetState.LinVel;
+
 		// Set Target->ReceiveInterval from either SendInterval or the number of physics ticks between receiving input states
 		if (SendInterval > 0)
 		{
@@ -1121,10 +1133,6 @@
 			}
 #endif
 
-			// Cache the position we received this target at, Predictive Interpolation will alter the target state but use this as the source position for reconciliation.
-			Target->PrevPosTarget = Input.TargetState.Position;
-			Target->PrevRotTarget = Input.TargetState.Quaternion;
-		
 			// TickCount is 0 by default at this point and when LOD is used, TickCount will be 0 if no LOD alignment was performed, in this case perform the normal target alignment
 			if (Target->TickCount == 0)
 			{
@@ -1163,6 +1171,29 @@
 					}
 				}
 			}
+
+			// Teleport detection, we don't have specific data that tells us a teleport has happened on the server, so try to detect it by examining the previous and next state
+			if (PhysicsReplicationCVars::PredictiveInterpolationCVars::TeleportDetectionEnabled == 1 && !bFirstTarget && SendInterval > 0 && RigidsSolver->IsUsingFixedDt())
+			{
+				const FVector PosOffset = (Input.TargetState.Position - Target->PrevPosTarget);
+				if (PosOffset.SizeSquared() > (PhysicsReplicationCVars::PredictiveInterpolationCVars::TeleportDetectionMinDistance * PhysicsReplicationCVars::PredictiveInterpolationCVars::TeleportDetectionMinDistance))
+				{
+					const FVector Velocity = Input.TargetState.LinVel.SizeSquared() > PrevLinVel.SizeSquared() ? Input.TargetState.LinVel : PrevLinVel;
+					const float DeltaSeconds = (SendInterval * RigidsSolver->GetAsyncDeltaTime());
+					const float PossibleDistanceSquared = (Velocity * (DeltaSeconds * PhysicsReplicationCVars::PredictiveInterpolationCVars::TeleportDetectionVelocityMultiplier)).SizeSquared();
+
+					if (PossibleDistanceSquared < PosOffset.SizeSquared())
+					{
+						// A teleport has most likely happened, set accumulated error seconds to above limit for hard snapping
+						// TODO: Don't piggyback on AccumulatedErrorSeconds (potentially implement ERigidBodyFlags::Teleported)
+						Target->AccumulatedErrorSeconds = PhysicsReplicationCVars::PredictiveInterpolationCVars::ErrorAccumulationSeconds + 1.0f;
+					}
+				}
+			}
+
+			// Cache the position we received this target at, Predictive Interpolation will alter the target state but use this as the source position for reconciliation.
+			Target->PrevPosTarget = Input.TargetState.Position;
+			Target->PrevRotTarget = Input.TargetState.Quaternion;
 		}
 	}
 
@@ -1760,13 +1791,43 @@
 	const FVector TargetLinVel = FVector(Target.TargetState.LinVel);
 	const FVector TargetAngVel = FVector(FMath::DegreesToRadians(Target.TargetState.AngVel)); // Radians
 
-	/** --- Reconciliation ---
-	* If target velocities are low enough, check the traveled direction and distance from previous frame and compare with replicated linear velocity.
+	/** --- Reconciliation --- */
+	const bool bHardSnap = (!bCanSimulate && PhysicsReplicationCVars::PredictiveInterpolationCVars::bKinematicHardSnap)
+		|| Target.AccumulatedErrorSeconds > PhysicsReplicationCVars::PredictiveInterpolationCVars::ErrorAccumulationSeconds
+		|| PhysicsReplicationCVars::PredictiveInterpolationCVars::bAlwaysHardSnap;
+
+	if (bHardSnap)
+	{
+		Target.AccumulatedErrorSeconds = 0.0f;
+
+		if (Handle->IsKinematic())
+		{
+			// Set a FKinematicTarget to hard snap kinematic object
+			const Chaos::FKinematicTarget KinTarget = Chaos::FKinematicTarget::MakePositionTarget(Target.PrevPosTarget, Target.PrevRotTarget); // Uses EKinematicTargetMode::Position
+			RigidsSolver->GetEvolution()->SetParticleKinematicTarget(Handle, KinTarget);
+		}
+		else 
+		{
+			// Set XRVW to hard snap dynamic object and force recalculation of friction
+			const bool bCorrectConnectedBodies = SettingsCurrent.PredictiveInterpolationSettings.GetCorrectConnectedBodies();
+			RigidsSolver->GetEvolution()->ApplyParticleTransformCorrection(Handle, Target.PrevPosTarget, Target.PrevRotTarget, bCorrectConnectedBodies, /*bInRecalculateFrictionOnConnectedBodies*/ true, ReplicatedParticleIDs);
+			Handle->SetV(TargetLinVel);
+			Handle->SetW(TargetAngVel);
+		}
+
+		// Cache data for next replication
+		Target.PrevLinVel = FVector(Target.TargetState.LinVel);
+
+		// End replication and go to sleep if that's requested
+		return EndReplicationHelper(Target, true);
+	}
+
+	/** If target velocities are low enough, check the traveled direction and distance from previous frame and compare with replicated linear velocity.
 	* If the object isn't moving enough along the replicated velocity it's considered stuck and needs reconciliation.
 	* SoftSnap is performed each tick while there is a registered error, if enough time pass HardSnap forces the object into the correct state. */
 	bool bSoftSnap = !PhysicsReplicationCVars::PredictiveInterpolationCVars::bVelocityBased;
 
-	if ( PhysicsReplicationCVars::PredictiveInterpolationCVars::bDisableErrorVelocityLimits ||
+	if (PhysicsReplicationCVars::PredictiveInterpolationCVars::bDisableErrorVelocityLimits ||
 		(TargetLinVel.Size() < PhysicsReplicationCVars::PredictiveInterpolationCVars::ErrorAccLinVelMaxLimit && TargetAngVel.Size() < PhysicsReplicationCVars::PredictiveInterpolationCVars::ErrorAccAngVelMaxLimit))
 	{
 		const FVector PrevDiff = CurrentState.Position - Target.PrevPos;
@@ -1797,36 +1858,7 @@
 		bSoftSnap = false;
 	}
 
-	const bool bHardSnap = (!bCanSimulate && PhysicsReplicationCVars::PredictiveInterpolationCVars::bKinematicHardSnap)
-		|| Target.AccumulatedErrorSeconds > PhysicsReplicationCVars::PredictiveInterpolationCVars::ErrorAccumulationSeconds
-		|| PhysicsReplicationCVars::PredictiveInterpolationCVars::bAlwaysHardSnap;
-
-	if (bHardSnap)
+	if (Handle->IsKinematic()) // Smooth Kinematic Replication
 	{
-		Target.AccumulatedErrorSeconds = 0.0f;
-
-		if (Handle->IsKinematic())
-		{
-			// Set a FKinematicTarget to hard snap kinematic object
-			const Chaos::FKinematicTarget KinTarget = Chaos::FKinematicTarget::MakePositionTarget(Target.PrevPosTarget, Target.PrevRotTarget); // Uses EKinematicTargetMode::Position
-			RigidsSolver->GetEvolution()->SetParticleKinematicTarget(Handle, KinTarget);
-		}
-		else 
-		{
-			// Set XRVW to hard snap dynamic object and force recalculation of friction
-			const bool bCorrectConnectedBodies = SettingsCurrent.PredictiveInterpolationSettings.GetCorrectConnectedBodies();
-			RigidsSolver->GetEvolution()->ApplyParticleTransformCorrection(Handle, Target.PrevPosTarget, Target.PrevRotTarget, bCorrectConnectedBodies, /*bInRecalculateFrictionOnConnectedBodies*/ true, ReplicatedParticleIDs);
-			Handle->SetV(TargetLinVel);
-			Handle->SetW(TargetAngVel);
-		}
-
-		// Cache data for next replication
-		Target.PrevLinVel = FVector(Target.TargetState.LinVel);
-
-		// End replication and go to sleep if that's requested
-		return EndReplicationHelper(Target, true);
-	}
-	else if (Handle->IsKinematic()) // Smooth Kinematic Replication
-	{
 		const bool bKinematicPrediction = PhysicsReplicationCVars::PredictiveInterpolationCVars::bKinematicPrediction;
 		const float InterpolationTicks = FMath::CeilToInt(Target.AverageReceiveInterval) - (RigidsSolver->GetCurrentFrame() - Target.ReceiveFrame);
