Added blue noise support to materials.
- The BlueNoiseScalar texture is now always loaded on all platforms (2MB)
- Blue noise scalar as common access to all shaders from the view UB as it is convenient to not have to change all material passes (opaque, translucent, distortion, motion, LumenCard, etc.).
- Added BlueNoise samping node (targeted for primary rays using SV ps).
- Added BlueNoise based opaque dithering. Needed to fudge adjust the uniform noise to get a similar look.
- Moved all noise material expressions into a common cpp file.

#tests PC and mobile preview
[RN] Added material graph BlueNoise sampling node. Opaque Masked dithereing is noe using blue noise.

#virtualized

diff --git a/Engine/Config/BaseEngine.ini b/Engine/Config/BaseEngine.ini
--- a/Engine/Config/BaseEngine.ini
+++ b/Engine/Config/BaseEngine.ini
@@ -160,6 +160,7 @@ DebugMeshMaterialName=/Engine/EngineDebugMaterials/DebugMeshMaterial.DebugMeshMa
 EmissiveMeshMaterialName=/Engine/EngineMaterials/EmissiveMeshMaterial.EmissiveMeshMaterial
 PreIntegratedSkinBRDFTextureName=/Engine/EngineMaterials/PreintegratedSkinBRDF.PreintegratedSkinBRDF
 BlueNoiseScalarTextureName=/Engine/EngineMaterials/FastBlueNoise_scalar_128x128x64.FastBlueNoise_scalar_128x128x64
+BlueNoiseScalarMobileTextureName=/Engine/EngineMaterials/FastBlueNoise_scalar_128x128x8.FastBlueNoise_scalar_128x128x8
 BlueNoiseVec2TextureName=/Engine/EngineMaterials/FastBlueNoise_vec2_128x128x64.FastBlueNoise_vec2_128x128x64
 GGXLTCAmpTextureName=/Engine/EngineMaterials/LTC/GGX_LTCAmp.GGX_LTCAmp
 GGXLTCMatTextureName=/Engine/EngineMaterials/LTC/GGX_LTCMat.GGX_LTCMat
diff --git a/Engine/Shaders/Private/MaterialTemplate.ush b/Engine/Shaders/Private/MaterialTemplate.ush
--- a/Engine/Shaders/Private/MaterialTemplate.ush
+++ b/Engine/Shaders/Private/MaterialTemplate.ush
@@ -2004,6 +2004,13 @@ float3 RotateAboutAxis(float4 NormalizedRotationAxisAndAngle, FWSVector3 Positio
 	return R - UAxis;
 }
 
+float ViewScalarBlueNoise(uint2 ScreenCoord, uint FrameIndex)
+{
+	uint3 WrappedCoordinate = uint3(ScreenCoord, FrameIndex) & View.BlueNoiseModuloMasks;
+	uint3 TextureCoordinate = uint3(WrappedCoordinate.x, WrappedCoordinate.z * View.BlueNoiseDimensions.y + WrappedCoordinate.y, 0);
+	return View.BlueNoiseScalarTexture.Load(TextureCoordinate, 0).x;
+}
+
 // Material Expression function
 float MaterialExpressionDepthOfFieldFunction(float SceneDepth, int FunctionValueIndex)
 {
@@ -4284,6 +4291,8 @@ void GetMaterialCoverageAndClipping(FMaterialPixelParameters Parameters, FPixelM
 
 #if MATERIALBLENDING_MASKED && !SINGLE_LAYER_WATER_NO_DISCARD
 	#if MATERIAL_DITHER_OPACITY_MASK
+
+	#if 0
 		/*
 		5 value dither. Every value present in +
 		012
@@ -4301,6 +4310,12 @@ void GetMaterialCoverageAndClipping(FMaterialPixelParameters Parameters, FPixelM
 		float Dither5 = frac( ( Pos.x + Pos.y * 2 - 1.5 + ResolvedView.TemporalAAParams.x ) / 5 );
 		float Noise = frac( dot( float2( 171.0, 231.0 ) / 71, Pos.xy ) );
 		float Dither = ( Dither5 * 5 + Noise ) * (1.0 / 6.0);
+	#else
+		// New dithering based on Blue noise
+		float Noise2  = ViewScalarBlueNoise(Parameters.SvPosition.xy, View.StateFrameIndex);
+		const float LegacyAdjustement = 0.83f; // This is required to get the same look and empty areas as the legacy one.
+		float Dither = LegacyAdjustement * Noise2;
+	#endif
 
 		clip( GetMaterialMask(PixelMaterialInputs) + Dither - 0.5 );
 	#else
diff --git a/Engine/Source/Runtime/Engine/Classes/Engine/Engine.h b/Engine/Source/Runtime/Engine/Classes/Engine/Engine.h
--- a/Engine/Source/Runtime/Engine/Classes/Engine/Engine.h
+++ b/Engine/Source/Runtime/Engine/Classes/Engine/Engine.h
@@ -1370,6 +1370,10 @@ public:
 	UPROPERTY(globalconfig)
 	FSoftObjectPath BlueNoiseScalarTextureName;
 
+	/** Path of the tiled blue-noise texture */
+	UPROPERTY(globalconfig)
+	FSoftObjectPath BlueNoiseScalarMobileTextureName;
+
 	/** Path of the tiled blue-noise texture */
 	UPROPERTY(globalconfig)
 	FSoftObjectPath BlueNoiseVec2TextureName;
@@ -1998,7 +2002,7 @@ public:
 	/** Delay loading this texture until it is needed by the renderer.
 	* The texture is uncompressed and contains no mips so it can't be streamed.
 	*/
-	ENGINE_API void LoadBlueNoiseTexture();
+	ENGINE_API void LoadBlueNoiseTexture(bool LoadVector2BlueNoiseTexture);
 
 	/** Delay loading this texture until it is needed by the renderer. */
 	ENGINE_API void LoadDefaultFilmGrainTexture();
diff --git a/Engine/Source/Runtime/Engine/Private/Materials/HLSLMaterialTranslator.cpp b/Engine/Source/Runtime/Engine/Private/Materials/HLSLMaterialTranslator.cpp
--- a/Engine/Source/Runtime/Engine/Private/Materials/HLSLMaterialTranslator.cpp
+++ b/Engine/Source/Runtime/Engine/Private/Materials/HLSLMaterialTranslator.cpp
@@ -11804,6 +11804,16 @@ int32 FHLSLMaterialTranslator::VectorNoise(int32 Position, EPositionOrigin Posit
 	}
 }
 
+int32 FHLSLMaterialTranslator::ScalarBlueNoise()
+{
+	if (Material->GetMaterialDomain() != EMaterialDomain::MD_Surface)
+	{
+		Errorf(TEXT("ScalarBlueNoise node can only be used within Surface domain materials %s (asset: %s).\r\n"), *Material->GetDebugName(), *Material->GetAssetPath().ToString());
+	}
+	
+	return AddCodeChunk(MCT_Float1, TEXT("ViewScalarBlueNoise(Parameters.SvPosition.xy, View.StateFrameIndex)"));
+}
+
 int32 FHLSLMaterialTranslator::BlackBody( int32 Temp )
 {
 	if( Temp == INDEX_NONE )
diff --git a/Engine/Source/Runtime/Engine/Private/Materials/HLSLMaterialTranslator.h b/Engine/Source/Runtime/Engine/Private/Materials/HLSLMaterialTranslator.h
--- a/Engine/Source/Runtime/Engine/Private/Materials/HLSLMaterialTranslator.h
+++ b/Engine/Source/Runtime/Engine/Private/Materials/HLSLMaterialTranslator.h
@@ -1211,6 +1211,7 @@ protected:
 	virtual int32 TemporalSobol(int32 Index, int32 Seed) override;
 	virtual int32 Noise(int32 Position, EPositionOrigin PositionOrigin, float Scale, int32 Quality, uint8 NoiseFunction, bool bTurbulence, int32 Levels, float OutputMin, float OutputMax, float LevelScale, int32 FilterWidth, bool bTiling, uint32 RepeatSize) override;
 	virtual int32 VectorNoise(int32 Position, EPositionOrigin PositionOrigin, int32 Quality, uint8 NoiseFunction, bool bTiling, uint32 TileSize) override;
+	virtual int32 ScalarBlueNoise() override;
 
 	virtual int32 BlackBody(int32 Temp) override;
 
diff --git a/Engine/Source/Runtime/Engine/Private/Materials/MaterialExpressionNoise.cpp b/Engine/Source/Runtime/Engine/Private/Materials/MaterialExpressionNoise.cpp
new file mode 100644
--- /dev/null
+++ b/Engine/Source/Runtime/Engine/Private/Materials/MaterialExpressionNoise.cpp
@@ -0,0 +1,303 @@
+// Copyright Epic Games, Inc. All Rights Reserved.
+
+/*=============================================================================
+	MaterialExpressionNoise.cpp - Material expressions implementation.
+=============================================================================*/
+
+#include "Materials/MaterialExpressionNoise.h"
+#include "Materials/MaterialExpressionVectorNoise.h"
+#include "Materials/HLSLMaterialTranslator.h"
+#if WITH_EDITOR
+#include "MaterialGraph/MaterialGraphSchema.h"
+#endif
+
+#define LOCTEXT_NAMESPACE "MaterialExpressionNoise"
+
+
+// Extern function from MaterialExpressions.cpp
+FName GetWorldPositionInputName(EPositionOrigin PositionOrigin);
+EWorldPositionIncludedOffsets GetWorldPositionTypeWithOrigin(EPositionOrigin PositionOrigin, bool IncludeOffsets);
+
+
+///////////////////////////////////////////////////////////////////////////////
+// UMaterialExpressionNoise
+///////////////////////////////////////////////////////////////////////////////
+UMaterialExpressionNoise::UMaterialExpressionNoise(const FObjectInitializer& ObjectInitializer)
+	: Super(ObjectInitializer)
+{
+	// Structure to hold one-time initialization
+	struct FConstructorStatics
+	{
+		FText NAME_Utility;
+		FConstructorStatics()
+			: NAME_Utility(LOCTEXT("Utility", "Utility"))
+		{
+		}
+	};
+	static FConstructorStatics ConstructorStatics;
+
+	Scale = 1.0f;
+	Levels = 6;
+	Quality = 1;
+	OutputMin = -1.0f;
+	OutputMax = 1.0f;
+	LevelScale = 2.0f;
+	NoiseFunction = NOISEFUNCTION_SimplexTex;
+	bTurbulence = true;
+	bTiling = false;
+	RepeatSize = 512;
+
+#if WITH_EDITORONLY_DATA
+	MenuCategories.Add(ConstructorStatics.NAME_Utility);
+#endif
+
+}
+
+#if WITH_EDITOR
+bool UMaterialExpressionNoise::CanEditChange(const FProperty* InProperty) const
+{
+	bool bIsEditable = Super::CanEditChange(InProperty);
+	if (bIsEditable && InProperty != nullptr)
+	{
+		FName PropertyFName = InProperty->GetFName();
+
+		bool bTilableNoiseType = NoiseFunction == NOISEFUNCTION_GradientALU || NoiseFunction == NOISEFUNCTION_ValueALU
+			|| NoiseFunction == NOISEFUNCTION_GradientTex || NoiseFunction == NOISEFUNCTION_VoronoiALU;
+
+		bool bSupportsQuality = (NoiseFunction == NOISEFUNCTION_VoronoiALU);
+
+		if (PropertyFName == GET_MEMBER_NAME_CHECKED(UMaterialExpressionNoise, bTiling))
+		{
+			bIsEditable = bTilableNoiseType;
+		}
+		else if (PropertyFName == GET_MEMBER_NAME_CHECKED(UMaterialExpressionNoise, RepeatSize))
+		{
+			bIsEditable = bTilableNoiseType && bTiling;
+		}
+
+		if (PropertyFName == GET_MEMBER_NAME_CHECKED(UMaterialExpressionNoise, Quality))
+		{
+			bIsEditable = bSupportsQuality;
+		}
+	}
+
+	return bIsEditable;
+}
+
+FName UMaterialExpressionNoise::GetInputName(int32 InputIndex) const
+{
+	if (GetInput(InputIndex) == &Position)
+	{
+		return GetWorldPositionInputName(WorldPositionOriginType);
+	}
+
+	return Super::GetInputName(InputIndex);
+}
+
+void UMaterialExpressionNoise::PostEditChangeProperty(FPropertyChangedEvent& PropertyChangedEvent)
+{
+	const FName PropertyName = PropertyChangedEvent.GetPropertyName();
+	if (PropertyName == GET_MEMBER_NAME_STRING_CHECKED(ThisClass, WorldPositionOriginType))
+	{
+		if (GraphNode)
+		{
+			GraphNode->ReconstructNode();
+		}
+	}
+
+	Super::PostEditChangeProperty(PropertyChangedEvent);
+}
+
+int32 UMaterialExpressionNoise::Compile(class FMaterialCompiler* Compiler, int32 OutputIndex)
+{
+	int32 PositionInput;
+
+	if (Position.GetTracedInput().Expression)
+	{
+		PositionInput = Position.Compile(Compiler);
+	}
+	else
+	{
+		PositionInput = Compiler->WorldPosition(GetWorldPositionTypeWithOrigin(WorldPositionOriginType, true));
+	}
+
+	int32 FilterWidthInput;
+
+	if (FilterWidth.GetTracedInput().Expression)
+	{
+		FilterWidthInput = FilterWidth.Compile(Compiler);
+	}
+	else
+	{
+		FilterWidthInput = Compiler->Constant(0);
+	}
+
+	return Compiler->Noise(PositionInput, WorldPositionOriginType, Scale, Quality, NoiseFunction, bTurbulence, Levels, OutputMin, OutputMax, LevelScale, FilterWidthInput, bTiling, RepeatSize);
+}
+
+void UMaterialExpressionNoise::GetCaption(TArray<FString>& OutCaptions) const
+{
+	const UEnum* NFEnum = StaticEnum<ENoiseFunction>();
+	check(NFEnum);
+	OutCaptions.Add(NFEnum->GetDisplayNameTextByValue(NoiseFunction).ToString());
+	OutCaptions.Add(TEXT("Noise"));
+}
+#endif // WITH_EDITOR
+
+
+
+///////////////////////////////////////////////////////////////////////////////
+// UMaterialExpressionVectorNoise
+///////////////////////////////////////////////////////////////////////////////
+UMaterialExpressionVectorNoise::UMaterialExpressionVectorNoise(const FObjectInitializer& ObjectInitializer)
+	: Super(ObjectInitializer)
+{
+	// Structure to hold one-time initialization
+	struct FConstructorStatics
+	{
+		FText NAME_Utility;
+		FConstructorStatics()
+			: NAME_Utility(LOCTEXT("Utility", "Utility"))
+		{
+		}
+	};
+	static FConstructorStatics ConstructorStatics;
+
+	Quality = 1;
+	NoiseFunction = VNF_CellnoiseALU;
+	bTiling = false;
+	TileSize = 300;
+
+#if WITH_EDITORONLY_DATA
+	MenuCategories.Add(ConstructorStatics.NAME_Utility);
+#endif
+}
+
+#if WITH_EDITOR
+bool UMaterialExpressionVectorNoise::CanEditChange(const FProperty* InProperty) const
+{
+	bool bIsEditable = Super::CanEditChange(InProperty);
+	if (bIsEditable && InProperty != nullptr)
+	{
+		FName PropertyFName = InProperty->GetFName();
+
+		bool bSupportsQuality = (NoiseFunction == VNF_VoronoiALU);
+
+		if (PropertyFName == GET_MEMBER_NAME_CHECKED(UMaterialExpressionVectorNoise, TileSize))
+		{
+			bIsEditable = bTiling;
+		}
+
+		else if (PropertyFName == GET_MEMBER_NAME_CHECKED(UMaterialExpressionVectorNoise, Quality))
+		{
+			bIsEditable = bSupportsQuality;
+		}
+	}
+
+	return bIsEditable;
+}
+
+FName UMaterialExpressionVectorNoise::GetInputName(int32 InputIndex) const
+{
+	if (GetInput(InputIndex) == &Position)
+	{
+		return GetWorldPositionInputName(WorldPositionOriginType);
+	}
+
+	return Super::GetInputName(InputIndex);
+}
+
+void UMaterialExpressionVectorNoise::PostEditChangeProperty(FPropertyChangedEvent& PropertyChangedEvent)
+{
+	const FName PropertyName = PropertyChangedEvent.GetPropertyName();
+	if (PropertyName == GET_MEMBER_NAME_STRING_CHECKED(ThisClass, WorldPositionOriginType))
+	{
+		if (GraphNode)
+		{
+			GraphNode->ReconstructNode();
+		}
+	}
+
+	Super::PostEditChangeProperty(PropertyChangedEvent);
+}
+
+int32 UMaterialExpressionVectorNoise::Compile(class FMaterialCompiler* Compiler, int32 OutputIndex)
+{
+	int32 PositionInput;
+
+	if (Position.GetTracedInput().Expression)
+	{
+		PositionInput = Position.Compile(Compiler);
+	}
+	else
+	{
+		PositionInput = Compiler->WorldPosition(GetWorldPositionTypeWithOrigin(WorldPositionOriginType, true));
+	}
+
+	return Compiler->VectorNoise(PositionInput, WorldPositionOriginType, Quality, NoiseFunction, bTiling, TileSize);
+}
+
+void UMaterialExpressionVectorNoise::GetCaption(TArray<FString>& OutCaptions) const
+{
+	const UEnum* VNFEnum = StaticEnum<EVectorNoiseFunction>();
+	check(VNFEnum);
+	OutCaptions.Add(VNFEnum->GetDisplayNameTextByValue(NoiseFunction).ToString());
+	OutCaptions.Add(TEXT("Vector Noise"));
+}
+#endif // WITH_EDITOR
+
+
+
+///////////////////////////////////////////////////////////////////////////////
+// UMaterialExpressionScalarBlueNoise
+///////////////////////////////////////////////////////////////////////////////
+UMaterialExpressionScalarBlueNoise::UMaterialExpressionScalarBlueNoise(const FObjectInitializer& ObjectInitializer)
+	: Super(ObjectInitializer)
+{
+	// Structure to hold one-time initialization
+	struct FConstructorStatics
+	{
+		FText NAME_Utility;
+		FConstructorStatics()
+			: NAME_Utility(LOCTEXT("Noise", "Noise"))
+		{
+		}
+	};
+	static FConstructorStatics ConstructorStatics;
+
+#if WITH_EDITORONLY_DATA
+	MenuCategories.Add(ConstructorStatics.NAME_Utility);
+#endif
+}
+
+#if WITH_EDITOR
+
+uint32 UMaterialExpressionScalarBlueNoise::GetOutputType(int32 OutputIndex) 
+{
+	return MCT_Float1;
+}
+
+int32 UMaterialExpressionScalarBlueNoise::Compile(class FMaterialCompiler* Compiler, int32 OutputIndex)
+{
+	return Compiler->ScalarBlueNoise();
+}
+
+void UMaterialExpressionScalarBlueNoise::GetCaption(TArray<FString>& OutCaptions) const
+{
+	OutCaptions.Add(TEXT("Scalar Blue Noise"));
+}
+
+void UMaterialExpressionScalarBlueNoise::GetExpressionToolTip(TArray<FString>& OutToolTip)
+{
+	OutToolTip.Add(TEXT("Return a blue noise value in [0.1] for each pixel on screen."));
+	OutToolTip.Add(TEXT("Be aware that this node might not play well when seen in ray traced reflections or path tracing,"));
+	OutToolTip.Add(TEXT("when secondary rays cannot use pixel position."));
+}
+
+void UMaterialExpressionScalarBlueNoise::GetConnectorToolTip(int32 InputIndex, int32 OutputIndex, TArray<FString>& OutToolTip)
+{
+	GetExpressionToolTip(OutToolTip); // Single ouptput so reuse the expression tooltip.
+}
+#endif // WITH_EDITOR
+
+#undef LOCTEXT_NAMESPACE
diff --git a/Engine/Source/Runtime/Engine/Private/Materials/MaterialExpressions.cpp b/Engine/Source/Runtime/Engine/Private/Materials/MaterialExpressions.cpp
--- a/Engine/Source/Runtime/Engine/Private/Materials/MaterialExpressions.cpp
+++ b/Engine/Source/Runtime/Engine/Private/Materials/MaterialExpressions.cpp
@@ -149,7 +149,6 @@
 #include "Materials/MaterialExpressionModulo.h"
 #include "Materials/MaterialExpressionMultiply.h"
 #include "Materials/MaterialExpressionNaniteReplace.h"
-#include "Materials/MaterialExpressionNoise.h"
 #include "Materials/MaterialExpressionNormalize.h"
 #include "Materials/MaterialExpressionNeuralPostProcessNode.h"
 #include "Materials/MaterialExpressionObjectBounds.h"
@@ -261,7 +260,6 @@
 #include "Materials/MaterialExpressionTruncateLWC.h"
 #include "Materials/MaterialExpressionTwoSidedSign.h"
 #include "Materials/MaterialExpressionUserSceneTexture.h"
-#include "Materials/MaterialExpressionVectorNoise.h"
 #include "Materials/MaterialExpressionVertexColor.h"
 #include "Materials/MaterialExpressionVertexNormalWS.h"
 #include "Materials/MaterialExpressionVertexTangentWS.h"
@@ -693,7 +691,7 @@ int32 CompileSubstrateBlendFunction(FMaterialCompiler* Compiler, const int32 A,
 	return INDEX_NONE;
 }
 
-EWorldPositionIncludedOffsets GetWorldPositionTypeWithOrigin(EPositionOrigin PositionOrigin, bool IncludeOffsets = true)
+extern EWorldPositionIncludedOffsets GetWorldPositionTypeWithOrigin(EPositionOrigin PositionOrigin, bool IncludeOffsets = true)
 {
 	switch (PositionOrigin)
 	{
@@ -704,7 +702,7 @@ EWorldPositionIncludedOffsets GetWorldPositionTypeWithOrigin(EPositionOrigin Pos
 	return WPT_Default;
 }
 
-FName GetWorldPositionInputName(EPositionOrigin PositionOrigin)
+extern FName GetWorldPositionInputName(EPositionOrigin PositionOrigin)
 {
 	switch (PositionOrigin)
 	{
@@ -22902,231 +22900,6 @@ void UMaterialExpressionNaniteReplace::GetExpressionToolTip(TArray<FString>& Out
 }
 #endif // WITH_EDITOR
 
-///////////////////////////////////////////////////////////////////////////////
-// UMaterialExpressionNoise
-///////////////////////////////////////////////////////////////////////////////
-UMaterialExpressionNoise::UMaterialExpressionNoise(const FObjectInitializer& ObjectInitializer)
-	: Super(ObjectInitializer)
-{
-	// Structure to hold one-time initialization
-	struct FConstructorStatics
-	{
-		FText NAME_Utility;
-		FConstructorStatics()
-			: NAME_Utility(LOCTEXT( "Utility", "Utility" ))
-		{
-		}
-	};
-	static FConstructorStatics ConstructorStatics;
-
-	Scale = 1.0f;
-	Levels = 6;
-	Quality = 1;
-	OutputMin = -1.0f;
-	OutputMax = 1.0f;
-	LevelScale = 2.0f;
-	NoiseFunction = NOISEFUNCTION_SimplexTex;
-	bTurbulence = true;
-	bTiling = false;
-	RepeatSize = 512;
-
-#if WITH_EDITORONLY_DATA
-	MenuCategories.Add(ConstructorStatics.NAME_Utility);
-#endif
-
-}
-
-#if WITH_EDITOR
-bool UMaterialExpressionNoise::CanEditChange(const FProperty* InProperty) const
-{
-	bool bIsEditable = Super::CanEditChange(InProperty);
-	if (bIsEditable && InProperty != nullptr)
-	{
-		FName PropertyFName = InProperty->GetFName();
-
-		bool bTilableNoiseType = NoiseFunction == NOISEFUNCTION_GradientALU || NoiseFunction == NOISEFUNCTION_ValueALU 
-			|| NoiseFunction == NOISEFUNCTION_GradientTex || NoiseFunction == NOISEFUNCTION_VoronoiALU;
-
-		bool bSupportsQuality = (NoiseFunction == NOISEFUNCTION_VoronoiALU);
-
-		if (PropertyFName == GET_MEMBER_NAME_CHECKED(UMaterialExpressionNoise, bTiling))
-		{
-			bIsEditable = bTilableNoiseType;
-		}
-		else if (PropertyFName == GET_MEMBER_NAME_CHECKED(UMaterialExpressionNoise, RepeatSize))
-		{
-			bIsEditable = bTilableNoiseType && bTiling;
-		}
-
-		if (PropertyFName == GET_MEMBER_NAME_CHECKED(UMaterialExpressionNoise, Quality))
-		{
-			bIsEditable = bSupportsQuality;
-		}
-	}
-
-	return bIsEditable;
-}
-
-FName UMaterialExpressionNoise::GetInputName(int32 InputIndex) const
-{
-	if (GetInput(InputIndex) == &Position)
-	{
-		return GetWorldPositionInputName(WorldPositionOriginType);
-	}
-
-	return Super::GetInputName(InputIndex);
-}
-
-void UMaterialExpressionNoise::PostEditChangeProperty(FPropertyChangedEvent& PropertyChangedEvent)
-{
-	const FName PropertyName = PropertyChangedEvent.GetPropertyName();
-	if (PropertyName == GET_MEMBER_NAME_STRING_CHECKED(ThisClass, WorldPositionOriginType))
-	{
-		if (GraphNode)
-		{
-			GraphNode->ReconstructNode();
-		}
-	}
-
-	Super::PostEditChangeProperty(PropertyChangedEvent);
-}
-
-int32 UMaterialExpressionNoise::Compile(class FMaterialCompiler* Compiler, int32 OutputIndex)
-{
-	int32 PositionInput;
-
-	if(Position.GetTracedInput().Expression)
-	{
-		PositionInput = Position.Compile(Compiler);
-	}
-	else
-	{
-		PositionInput = Compiler->WorldPosition(GetWorldPositionTypeWithOrigin(WorldPositionOriginType));
-	}
-
-	int32 FilterWidthInput;
-
-	if(FilterWidth.GetTracedInput().Expression)
-	{
-		FilterWidthInput = FilterWidth.Compile(Compiler);
-	}
-	else
-	{
-		FilterWidthInput = Compiler->Constant(0);
-	}
-
-	return Compiler->Noise(PositionInput, WorldPositionOriginType, Scale, Quality, NoiseFunction, bTurbulence, Levels, OutputMin, OutputMax, LevelScale, FilterWidthInput, bTiling, RepeatSize);
-}
-
-void UMaterialExpressionNoise::GetCaption(TArray<FString>& OutCaptions) const
-{
-	const UEnum* NFEnum = StaticEnum<ENoiseFunction>();
-	check(NFEnum);
-	OutCaptions.Add(NFEnum->GetDisplayNameTextByValue(NoiseFunction).ToString());
-	OutCaptions.Add(TEXT("Noise"));
-}
-#endif // WITH_EDITOR
-
-///////////////////////////////////////////////////////////////////////////////
-// UMaterialExpressionVectorNoise
-///////////////////////////////////////////////////////////////////////////////
-UMaterialExpressionVectorNoise::UMaterialExpressionVectorNoise(const FObjectInitializer& ObjectInitializer)
-	: Super(ObjectInitializer)
-{
-	// Structure to hold one-time initialization
-	struct FConstructorStatics
-	{
-		FText NAME_Utility;
-		FConstructorStatics()
-			: NAME_Utility(LOCTEXT("Utility", "Utility"))
-		{
-		}
-	};
-	static FConstructorStatics ConstructorStatics;
-
-	Quality = 1;
-	NoiseFunction = VNF_CellnoiseALU;
-	bTiling = false;
-	TileSize = 300;
-
-#if WITH_EDITORONLY_DATA
-	MenuCategories.Add(ConstructorStatics.NAME_Utility);
-#endif
-}
-
-#if WITH_EDITOR
-bool UMaterialExpressionVectorNoise::CanEditChange(const FProperty* InProperty) const
-{
-	bool bIsEditable = Super::CanEditChange(InProperty);
-	if (bIsEditable && InProperty != nullptr)
-	{
-		FName PropertyFName = InProperty->GetFName();
-
-		bool bSupportsQuality = (NoiseFunction == VNF_VoronoiALU);
-
-		if (PropertyFName == GET_MEMBER_NAME_CHECKED(UMaterialExpressionVectorNoise, TileSize))
-		{
-			bIsEditable = bTiling;
-		}
-
-		else if (PropertyFName == GET_MEMBER_NAME_CHECKED(UMaterialExpressionVectorNoise, Quality))
-		{
-			bIsEditable = bSupportsQuality;
-		}
-	}
-
-	return bIsEditable;
-}
-
-FName UMaterialExpressionVectorNoise::GetInputName(int32 InputIndex) const
-{
-	if (GetInput(InputIndex) == &Position)
-	{
-		return GetWorldPositionInputName(WorldPositionOriginType);
-	}
-
-	return Super::GetInputName(InputIndex);
-}
-
-void UMaterialExpressionVectorNoise::PostEditChangeProperty(FPropertyChangedEvent& PropertyChangedEvent)
-{
-	const FName PropertyName = PropertyChangedEvent.GetPropertyName();
-	if (PropertyName == GET_MEMBER_NAME_STRING_CHECKED(ThisClass, WorldPositionOriginType))
-	{
-		if (GraphNode)
-		{
-			GraphNode->ReconstructNode();
-		}
-	}
-
-	Super::PostEditChangeProperty(PropertyChangedEvent);
-}
-
-int32 UMaterialExpressionVectorNoise::Compile(class FMaterialCompiler* Compiler, int32 OutputIndex)
-{
-	int32 PositionInput;
-
-	if (Position.GetTracedInput().Expression)
-	{
-		PositionInput = Position.Compile(Compiler);
-	}
-	else
-	{
-		PositionInput = Compiler->WorldPosition(GetWorldPositionTypeWithOrigin(WorldPositionOriginType));
-	}
-
-	return Compiler->VectorNoise(PositionInput, WorldPositionOriginType, Quality, NoiseFunction, bTiling, TileSize);
-}
-
-void UMaterialExpressionVectorNoise::GetCaption(TArray<FString>& OutCaptions) const
-{
-	const UEnum* VNFEnum = StaticEnum<EVectorNoiseFunction>();
-	check(VNFEnum);
-	OutCaptions.Add(VNFEnum->GetDisplayNameTextByValue(NoiseFunction).ToString());
-	OutCaptions.Add(TEXT("Vector Noise"));
-}
-#endif // WITH_EDITOR
-
 ///////////////////////////////////////////////////////////////////////////////
 // UMaterialExpressionBlackBody
 ///////////////////////////////////////////////////////////////////////////////
diff --git a/Engine/Source/Runtime/Engine/Private/UnrealEngine.cpp b/Engine/Source/Runtime/Engine/Private/UnrealEngine.cpp
--- a/Engine/Source/Runtime/Engine/Private/UnrealEngine.cpp
+++ b/Engine/Source/Runtime/Engine/Private/UnrealEngine.cpp
@@ -3436,14 +3436,23 @@ void UEngine::LoadDefaultBloomTexture()
 	}
 }
 
-void UEngine::LoadBlueNoiseTexture()
+void UEngine::LoadBlueNoiseTexture(bool LoadVector2BlueNoiseTexture)
 {
+	// Always load the scalar blue noise texture.
 	if (BlueNoiseScalarTexture == nullptr)
 	{
-		LoadEngineTexture(BlueNoiseScalarTexture, *BlueNoiseScalarTextureName.ToString());
+		if (GetFeatureLevelShadingPath(GMaxRHIFeatureLevel) == EShadingPath::Mobile)
+		{
+			LoadEngineTexture(BlueNoiseScalarTexture, *BlueNoiseScalarMobileTextureName.ToString());
+		}
+		else
+		{
+			LoadEngineTexture(BlueNoiseScalarTexture, *BlueNoiseScalarTextureName.ToString());
+		}
 	}
 
-	if (BlueNoiseVec2Texture == nullptr)
+	// Only load the vector2 blue noise texture if required.
+	if (LoadVector2BlueNoiseTexture && BlueNoiseVec2Texture == nullptr)
 	{
 		LoadEngineTexture(BlueNoiseVec2Texture, *BlueNoiseVec2TextureName.ToString());
 	}
diff --git a/Engine/Source/Runtime/Engine/Public/MaterialCompiler.h b/Engine/Source/Runtime/Engine/Public/MaterialCompiler.h
--- a/Engine/Source/Runtime/Engine/Public/MaterialCompiler.h
+++ b/Engine/Source/Runtime/Engine/Public/MaterialCompiler.h
@@ -540,6 +540,7 @@ public:
 	virtual int32 PerInstanceCustomData(int32 DataIndex, int32 DefaultValueIndex) = 0;
 	virtual int32 PerInstanceCustomData3Vector(int32 DataIndex, int32 DefaultValueIndex) = 0;
 	virtual int32 AntialiasedTextureMask(int32 Tex, int32 UV, float Threshold, uint8 Channel) = 0;
+
 	virtual int32 Sobol(int32 Cell, int32 Index, int32 Seed) = 0;
 	virtual int32 TemporalSobol(int32 Index, int32 Seed) = 0;
 	virtual int32 Noise(int32 Position, EPositionOrigin PositionOrigin, float Scale, int32 Quality, uint8 NoiseFunction, bool bTurbulence, int32 Levels, float OutputMin, float OutputMax, float LevelScale, int32 FilterWidth, bool bTiling, uint32 RepeatSize) = 0;
@@ -554,6 +555,8 @@ public:
 	{
 		return VectorNoise(Position, EPositionOrigin::Absolute, Quality, NoiseFunction, bTiling, RepeatSize);
 	}
+	virtual int32 ScalarBlueNoise() = 0;
+
 	virtual int32 BlackBody( int32 Temp ) = 0;
 	virtual int32 DistanceToNearestSurface(int32 PositionArg, EPositionOrigin PositionOrigin) = 0;
 	UE_DEPRECATED(5.4, "Use DistanceToNearestSurface(int32, EPositionOrigin) instead")
@@ -1106,6 +1109,7 @@ public:
 	{
 		return Compiler->AntialiasedTextureMask(Tex, UV, Threshold, Channel);
 	}
+
 	virtual int32 Sobol(int32 Cell, int32 Index, int32 Seed) override { return Compiler->Sobol(Cell, Index, Seed); }
 	virtual int32 TemporalSobol(int32 Index, int32 Seed) override { return Compiler->TemporalSobol(Index, Seed); }
 	virtual int32 Noise(int32 Position, EPositionOrigin PositionOrigin, float Scale, int32 Quality, uint8 NoiseFunction, bool bTurbulence, int32 Levels, float OutputMin, float OutputMax, float LevelScale, int32 FilterWidth, bool bTiling, uint32 TileSize) override
@@ -1116,6 +1120,8 @@ public:
 	{
 		return Compiler->VectorNoise(Position, PositionOrigin, Quality, NoiseFunction, bTiling, TileSize);
 	}
+	virtual int32 ScalarBlueNoise() override { return Compiler->ScalarBlueNoise(); }
+
 	virtual int32 BlackBody(int32 Temp) override { return Compiler->BlackBody(Temp); }
 	virtual int32 DistanceToNearestSurface(int32 PositionArg, EPositionOrigin PositionOrigin) override { return Compiler->DistanceToNearestSurface(PositionArg, PositionOrigin); }
 	virtual int32 DistanceFieldGradient(int32 PositionArg, EPositionOrigin PositionOrigin) override { return Compiler->DistanceFieldGradient(PositionArg, PositionOrigin); }
diff --git a/Engine/Source/Runtime/Engine/Public/Materials/MaterialExpressionNoise.h b/Engine/Source/Runtime/Engine/Public/Materials/MaterialExpressionNoise.h
--- a/Engine/Source/Runtime/Engine/Public/Materials/MaterialExpressionNoise.h
+++ b/Engine/Source/Runtime/Engine/Public/Materials/MaterialExpressionNoise.h
@@ -122,5 +122,20 @@ class UMaterialExpressionNoise : public UMaterialExpression
 	//~ End UMaterialExpression Interface
 };
 
+UCLASS(MinimalAPI)
+class UMaterialExpressionScalarBlueNoise : public UMaterialExpression
+{
+	GENERATED_UCLASS_BODY()
 
+	// Later, we could add a 2D Position input if needed to overide pixel position
 
+	//~ Begin UMaterialExpression Interface
+#if WITH_EDITOR
+	virtual uint32 GetOutputType(int32 OutputIndex) override;
+	virtual int32 Compile(class FMaterialCompiler* Compiler, int32 OutputIndex) override;
+	virtual void GetCaption(TArray<FString>& OutCaptions) const override;
+	virtual void GetExpressionToolTip(TArray<FString>& OutToolTip) override;
+	virtual void GetConnectorToolTip(int32 InputIndex, int32 OutputIndex, TArray<FString>& OutToolTip) override;
+#endif
+	//~ End UMaterialExpression Interface
+};
diff --git a/Engine/Source/Runtime/Engine/Public/SceneView.h b/Engine/Source/Runtime/Engine/Public/SceneView.h
--- a/Engine/Source/Runtime/Engine/Public/SceneView.h
+++ b/Engine/Source/Runtime/Engine/Public/SceneView.h
@@ -1181,6 +1181,11 @@ BEGIN_GLOBAL_SHADER_PARAMETER_STRUCT_WITH_CONSTRUCTOR(FViewUniformShaderParamete
 
 	SHADER_PARAMETER_SRV(Buffer<float>, PhysicsFieldClipmapBuffer)
 
+	// Blue noise texture
+	SHADER_PARAMETER_TEXTURE(Texture2D<float>, BlueNoiseScalarTexture)
+	SHADER_PARAMETER(FIntVector, BlueNoiseDimensions)
+	SHADER_PARAMETER(FIntVector, BlueNoiseModuloMasks)
+
 	// Ray tracing
 	SHADER_PARAMETER(FVector3f, TLASPreViewTranslationHigh)
 	SHADER_PARAMETER(FVector3f, TLASPreViewTranslationLow)
diff --git a/Engine/Source/Runtime/Renderer/Private/BlueNoise.cpp b/Engine/Source/Runtime/Renderer/Private/BlueNoise.cpp
--- a/Engine/Source/Runtime/Renderer/Private/BlueNoise.cpp
+++ b/Engine/Source/Runtime/Renderer/Private/BlueNoise.cpp
@@ -43,7 +43,11 @@ FBlueNoiseParameters GetBlueNoiseParameters()
 		&&(Out.ModuloMasks.Z + 1) == Out.Dimensions.Z);
 
 	Out.ScalarTexture = GEngine->BlueNoiseScalarTexture->GetResource()->TextureRHI;
-	Out.Vec2Texture = GEngine->BlueNoiseVec2Texture->GetResource()->TextureRHI;
+	Out.Vec2Texture = GSystemTextures.BlackDummy->GetRHI();
+	if (GEngine->BlueNoiseVec2Texture != nullptr)
+	{
+		Out.Vec2Texture = GEngine->BlueNoiseVec2Texture->GetResource()->TextureRHI;
+	}
 	return Out;
 }
 
diff --git a/Engine/Source/Runtime/Renderer/Private/SceneRendering.cpp b/Engine/Source/Runtime/Renderer/Private/SceneRendering.cpp
--- a/Engine/Source/Runtime/Renderer/Private/SceneRendering.cpp
+++ b/Engine/Source/Runtime/Renderer/Private/SceneRendering.cpp
@@ -111,6 +111,7 @@
 #include "VT/VirtualTextureScalability.h"
 #include "VT/VirtualTextureSystem.h"
 #include <type_traits>
+#include "BlueNoise.h"
 
 /*-----------------------------------------------------------------------------
 	Globals
@@ -2203,6 +2204,13 @@ void FViewInfo::SetupUniformBufferParameters(
 	ViewUniformShaderParameters.VTFeedbackBuffer = GVirtualTextureFeedbackBuffer.GetUAV();
 
 	ViewUniformShaderParameters.GPUSceneViewId = GPUSceneViewId;
+
+	{
+		FBlueNoise BlueNoiseParam = GetBlueNoiseGlobalParameters();
+		ViewUniformShaderParameters.BlueNoiseScalarTexture = BlueNoiseParam.BlueNoise.ScalarTexture;
+		ViewUniformShaderParameters.BlueNoiseDimensions = BlueNoiseParam.BlueNoise.Dimensions;
+		ViewUniformShaderParameters.BlueNoiseModuloMasks = BlueNoiseParam.BlueNoise.ModuloMasks;
+	}
 }
 
 void FViewInfo::InitRHIResources(uint32 OverrideNumMSAASamples)
@@ -2800,15 +2808,14 @@ FSceneRenderer::FSceneRenderer(const FSceneViewFamily* InViewFamily, FHitProxyCo
 		}
 #endif
 
-		if (ShouldRenderLumenDiffuseGI(Scene, *ViewInfo) || 
-			ShouldRenderLumenReflections(*ViewInfo) || 
-			ShouldRenderVolumetricCloudWithBlueNoise_GameThread(Scene, *ViewInfo) || 
+		const bool LoadVector2BlueNoiseTexture =
+			ShouldRenderLumenDiffuseGI(Scene, *ViewInfo) ||
+			ShouldRenderLumenReflections(*ViewInfo) ||
+			ShouldRenderVolumetricCloudWithBlueNoise_GameThread(Scene, *ViewInfo) ||
 			UseVirtualShadowMaps(Scene->GetShaderPlatform(), Scene->GetFeatureLevel()) ||
 			Substrate::IsGlintEnabled(ViewInfo->GetShaderPlatform()) ||
-			IsHairStrandsSupported(EHairStrandsShaderType::Strands, ViewInfo->GetShaderPlatform()))
-		{
-			GEngine->LoadBlueNoiseTexture();
-		}
+			IsHairStrandsSupported(EHairStrandsShaderType::Strands, ViewInfo->GetShaderPlatform());
+		GEngine->LoadBlueNoiseTexture(LoadVector2BlueNoiseTexture);
 
 		if (Substrate::IsGlintEnabled(ViewInfo->GetShaderPlatform()))
 		{
