Fixed convex queries when using a physics query handler.

diff --git a/Engine/Source/Runtime/Engine/Private/Collision/SceneQuery.cpp b/Engine/Source/Runtime/Engine/Private/Collision/SceneQuery.cpp
--- a/Engine/Source/Runtime/Engine/Private/Collision/SceneQuery.cpp
+++ b/Engine/Source/Runtime/Engine/Private/Collision/SceneQuery.cpp
@@ -12,6 +12,7 @@
 #include "PhysicsEngine/CollisionQueryFilterCallback.h"
 #include "PhysicsEngine/ExternalSpatialAccelerationPayload.h"
 #include "PhysicsEngine/ScopedSQHitchRepeater.h"
+#include "Chaos/ChaosArchive.h"
 
 #include "Collision/CollisionDebugDrawing.h"
 
@@ -197,6 +198,20 @@ bool BuildQueryShape(const FPhysicsGeometryCollection& GeometryCollection, Chaos
 			OutQueryShape.CollisionShape = FCollisionShape::MakeCapsule(Capsule.GetRadiusf(), Capsule.GetHeightf() * 0.5f);
 			return true;
 		}
+		case ECollisionShapeType::Convex:
+		{
+			FMemoryWriter MemoryWriter(OutQueryShape.ConvexData);
+			Chaos::FChaosArchive Writer(MemoryWriter);
+
+			const Chaos::FImplicitObject& Geometry = GeometryCollection.GetGeometry();
+			Chaos::FImplicitObjectPtr CopiedGeometryPtr = Geometry.CopyGeometry();
+			Chaos::FImplicitObject& CopiedGeometry = *CopiedGeometryPtr;
+
+			Chaos::FImplicitObject::SerializationFactory(Writer, &CopiedGeometry);
+			CopiedGeometry.Serialize(Writer);
+			OutQueryShape.LocalBoundingBox = CopiedGeometry.BoundingBox();
+			return true;
+		}
 		default:
 		{
 			UE_LOG(LogChaos, Warning, TEXT("PhysicsQueryHandler: Invalid shape type. Shape type %d is not yet supported"), (int)ShapeType);
@@ -205,6 +220,17 @@ bool BuildQueryShape(const FPhysicsGeometryCollection& GeometryCollection, Chaos
 	return false;
 }
 
+void ConvertToConvex(const Chaos::FQueryShape& QueryShape, Chaos::FImplicitObjectPtr& OutConvex)
+{
+	ensure(QueryShape.IsConvexShape());
+
+	FMemoryReader MemoryReader(QueryShape.ConvexData);
+	Chaos::FChaosArchive Reader(MemoryReader);
+
+	OutConvex = Chaos::FImplicitObject::SerializationFactory(Reader, nullptr);
+	OutConvex->Serialize(Reader);
+}
+
 template <typename InHitType, ESweepOrRay InGeometryQuery, ESingleMultiOrTest InSingleMultiOrTest>
 struct TSQTraits
 {
@@ -1427,11 +1453,23 @@ namespace Chaos::Private
 	bool OverlapImpl(const UWorld* World, const Chaos::FOverlapQueryData& OverlapData, const FCommonQueryData& CommonData, TArray<FOverlapResult>& OutOverlaps)
 	{
 		const FQuat& Rotation = OverlapData.GeomPose.GetRotation();
-		const FCollisionShape& CollisionShape = OverlapData.QueryShape.CollisionShape;
-		FPhysicsShapeAdapter Adaptor(Rotation, CollisionShape);
-		const FPhysicsGeometry& Geom = Adaptor.GetGeometry();
+		if (OverlapData.QueryShape.IsConvexShape())
+		{
+			FImplicitObjectPtr Convex;
+			ConvertToConvex(OverlapData.QueryShape, Convex);
+			FPhysicsGeometryCollection CollisionShape = FChaosEngineInterface::GetGeometryCollection(*Convex.GetReference());
+			const FPhysicsGeometry& Geometry = CollisionShape.GetGeometry();
 
-		return GeomOverlapMultiImp<OverlapHitType, InfoType>(World, Geom, CollisionShape, OverlapData.GeomPose, OutOverlaps, CommonData.TraceChannel, CommonData.Params, CommonData.ResponseParams, CommonData.ObjectParams, FDefaultAccelContainer{});
+			return GeomOverlapMultiImp<OverlapHitType, InfoType>(World, Geometry, CollisionShape, OverlapData.GeomPose, OutOverlaps, CommonData.TraceChannel, CommonData.Params, CommonData.ResponseParams, CommonData.ObjectParams, FDefaultAccelContainer{});
+		}
+		else
+		{
+			FPhysicsShapeAdapter Adaptor(Rotation, OverlapData.QueryShape.CollisionShape);
+			const FPhysicsGeometry& Geom = Adaptor.GetGeometry();
+			const FCollisionShape& CollisionShape = OverlapData.QueryShape.CollisionShape;
+
+			return GeomOverlapMultiImp<OverlapHitType, InfoType>(World, Geom, CollisionShape, OverlapData.GeomPose, OutOverlaps, CommonData.TraceChannel, CommonData.Params, CommonData.ResponseParams, CommonData.ObjectParams, FDefaultAccelContainer{});
+		}
 	}
 
 	template <EQueryInfo InfoType, typename ... ArgTypes>
@@ -1490,8 +1528,20 @@ namespace Chaos::Private
 	template <ESingleMultiOrTest SingleMultiOrTest, typename HitResultsType>
 	bool SweepCommon(const EThreadQueryContext ThreadContext, const UWorld* World, const Chaos::FSweepQueryData& SweepData, const Chaos::FCommonQueryData& CommonData, HitResultsType& OutHits)
 	{
-		FGeomSQAdditionalInputs GeomInputs(SweepData.QueryShape.CollisionShape, SweepData.GeomRot);
-		return TraceCommonImp<ESweepOrRay::Sweep, SingleMultiOrTest>(ThreadContext, World, OutHits, GeomInputs, SweepData.Start, SweepData.End, CommonData.TraceChannel, CommonData.Params, CommonData.ResponseParams, CommonData.ObjectParams, FDefaultAccelContainer{});
+		if (SweepData.QueryShape.IsConvexShape())
+		{
+			Chaos::FImplicitObjectPtr Convex;
+			ConvertToConvex(SweepData.QueryShape, Convex);
+			FPhysicsGeometryCollection GeomCollection = FChaosEngineInterface::GetGeometryCollection(*Convex.GetReference());
+
+			FGeomCollectionSQAdditionalInputs GeomInputs(GeomCollection, SweepData.GeomRot);
+			return TraceCommonImp<ESweepOrRay::Sweep, SingleMultiOrTest>(ThreadContext, World, OutHits, GeomInputs, SweepData.Start, SweepData.End, CommonData.TraceChannel, CommonData.Params, CommonData.ResponseParams, CommonData.ObjectParams, FDefaultAccelContainer{});
+		}
+		else
+		{
+			FGeomSQAdditionalInputs GeomInputs(SweepData.QueryShape.CollisionShape, SweepData.GeomRot);
+			return TraceCommonImp<ESweepOrRay::Sweep, SingleMultiOrTest>(ThreadContext, World, OutHits, GeomInputs, SweepData.Start, SweepData.End, CommonData.TraceChannel, CommonData.Params, CommonData.ResponseParams, CommonData.ObjectParams, FDefaultAccelContainer{});
+		}
 	}
 
 	bool FQueryInterface_Internal::SweepTest(const EThreadQueryContext ThreadContext, const UWorld* World, const FSweepQueryData& SweepData, const FCommonQueryData& CommonData, FHitResult& OutHits)
diff --git a/Engine/Source/Runtime/Engine/Public/Physics/SceneQueryData.h b/Engine/Source/Runtime/Engine/Public/Physics/SceneQueryData.h
--- a/Engine/Source/Runtime/Engine/Public/Physics/SceneQueryData.h
+++ b/Engine/Source/Runtime/Engine/Public/Physics/SceneQueryData.h
@@ -6,6 +6,8 @@
 #include "Math/MathFwd.h"
 #include "CollisionShape.h"
 #include "CollisionQueryParams.h"
+#include "Chaos/Aabb.h"
+#include "Chaos/Core.h"
 
 namespace Chaos
 {
@@ -33,7 +35,14 @@ namespace Chaos
 
 	struct FQueryShape
 	{
+		bool IsConvexShape() const
+		{
+			return !ConvexData.IsEmpty();
+		}
+
 		FCollisionShape CollisionShape;
+		TArray<uint8> ConvexData;
+		FAABB3 LocalBoundingBox;
 	};
 
 	struct FOverlapQueryData
