Fix ensure cause by invalidated CachedBounds in UInstancedStaticMeshComponent (AddInstancesInternal is called in a callstack where the bound cache get invalidated).
Issue was introduced at CL 36539228.

diff --git a/Engine/Source/Runtime/Engine/Private/InstancedStaticMesh.cpp b/Engine/Source/Runtime/Engine/Private/InstancedStaticMesh.cpp
--- a/Engine/Source/Runtime/Engine/Private/InstancedStaticMesh.cpp
+++ b/Engine/Source/Runtime/Engine/Private/InstancedStaticMesh.cpp
@@ -3405,6 +3405,7 @@ TArray<int32> UInstancedStaticMeshComponent::AddInstancesInternal(TConstArrayVie
 			// since it was skipped because IsNavigationRelevant() requires at least one instance.
 			if (GetInstanceCount() == 1)
 			{
+				CalcAndCacheNavigationBounds();
 				FNavigationSystem::RegisterComponent(*this);
 			}
 
