Fix bad cast in FPerQualityLevelProperty.

#rnx

diff --git a/Engine/Source/Runtime/Engine/Private/PerQualityLevelProperties.cpp b/Engine/Source/Runtime/Engine/Private/PerQualityLevelProperties.cpp
--- a/Engine/Source/Runtime/Engine/Private/PerQualityLevelProperties.cpp
+++ b/Engine/Source/Runtime/Engine/Private/PerQualityLevelProperties.cpp
@@ -340,7 +340,7 @@ void FPerQualityLevelProperty<StructType, ValueType, _BasePropertyName>::StripQu
 
 		if (This->PerQuality.Num() > 0)
 		{
-			int32 PreviousQualityLevel = This->Default;
+			ValueType PreviousQualityLevel = This->Default;
 
 			for (TSet<int32>::TIterator It(CookQualityLevelInfo); It; ++It)
 			{
@@ -511,7 +511,7 @@ template ENGINE_API float FPerQualityLevelProperty<FPerQualityLevelFloat, float,
 
 FString FPerQualityLevelFloat::ToString() const
 {
-	FString Result = FString::FromInt(Default);
+	FString Result = FString::Printf(TEXT("%f"), Default);
 
 #if WITH_EDITORONLY_DATA
 	TArray<int32> QualityLevels;
