[World Partition Runtime HashSet] Handle invalid partition layer setups.

diff --git a/Engine/Source/Runtime/Engine/Private/WorldPartition/RuntimeHashSet/WorldPartitionRuntimeHashSetStreamingGeneration.cpp b/Engine/Source/Runtime/Engine/Private/WorldPartition/RuntimeHashSet/WorldPartitionRuntimeHashSetStreamingGeneration.cpp
--- a/Engine/Source/Runtime/Engine/Private/WorldPartition/RuntimeHashSet/WorldPartitionRuntimeHashSetStreamingGeneration.cpp
+++ b/Engine/Source/Runtime/Engine/Private/WorldPartition/RuntimeHashSet/WorldPartitionRuntimeHashSetStreamingGeneration.cpp
@@ -43,7 +43,10 @@ bool UWorldPartitionRuntimeHashSet::GenerateRuntimePartitionsStreamingDescs(cons
 			NameToRuntimePartitionMap.Emplace(ActorSetInstance.RuntimeGrid, RuntimePartition);
 		}
 
-		RuntimePartitionsToActorSetMap.FindOrAdd(RuntimePartition).Add(&ActorSetInstance);
+		if (RuntimePartition)
+		{
+			RuntimePartitionsToActorSetMap.FindOrAdd(RuntimePartition).Add(&ActorSetInstance);
+		}
 	});
 
 	//
