Another potential fix for crash with new blue nosie texture.

#tests PC editor dx12 with PIE
#rnx

diff --git a/Engine/Source/Runtime/Engine/Private/SceneManagement.cpp b/Engine/Source/Runtime/Engine/Private/SceneManagement.cpp
--- a/Engine/Source/Runtime/Engine/Private/SceneManagement.cpp
+++ b/Engine/Source/Runtime/Engine/Private/SceneManagement.cpp
@@ -1335,6 +1335,8 @@ FViewUniformShaderParameters::FViewUniformShaderParameters()
 	SpecularProfileSampler = TStaticSamplerState<SF_Bilinear, AM_Clamp, AM_Clamp, AM_Clamp>::GetRHI();
 
 	VTFeedbackBuffer = GEmptyStructuredBufferWithUAV->UnorderedAccessViewRHI;
+
+	BlueNoiseScalarTexture = GBlackTextureWithSRV->TextureRHI;
 }
 
 FInstancedViewUniformShaderParameters::FInstancedViewUniformShaderParameters()
