Anim Runtime: Make sure RequiredBones.IsValid() when checking FBoneReference::IsValidToEvaluate()

diff --git a/Engine/Source/Runtime/Engine/Private/Animation/BoneReference.cpp b/Engine/Source/Runtime/Engine/Private/Animation/BoneReference.cpp
--- a/Engine/Source/Runtime/Engine/Private/Animation/BoneReference.cpp
+++ b/Engine/Source/Runtime/Engine/Private/Animation/BoneReference.cpp
@@ -57,7 +57,7 @@ bool FBoneReference::Initialize(const USkeleton* Skeleton)
 
 bool FBoneReference::IsValidToEvaluate(const FBoneContainer& RequiredBones) const
 {
-	return (BoneIndex != INDEX_NONE && RequiredBones.Contains(BoneIndex));
+	return (BoneIndex != INDEX_NONE && RequiredBones.IsValid() && RequiredBones.Contains(BoneIndex));
 }
 
 FSkeletonPoseBoneIndex FBoneReference::GetSkeletonPoseIndex(const FBoneContainer& RequiredBones) const
