Material Layer Tool - Bugs with deleting Slab and then recrea

diff --git a/Engine/Source/Editor/MaterialEditor/Private/SMaterialSubstrateTree.cpp b/Engine/Source/Editor/MaterialEditor/Private/SMaterialSubstrateTree.cpp
--- a/Engine/Source/Editor/MaterialEditor/Private/SMaterialSubstrateTree.cpp
+++ b/Engine/Source/Editor/MaterialEditor/Private/SMaterialSubstrateTree.cpp
@@ -76,6 +76,12 @@ bool IsAssetDataAMaterialLayerBlendFunction(const FAssetData& AssetData)
 	return AssetData.IsInstanceOf(UMaterialFunctionMaterialLayerBlend::StaticClass()) || AssetData.IsInstanceOf(UMaterialFunctionMaterialLayerBlendInstance::StaticClass());
 }
 
+// Check if the assetdata is avalid MaterialLayerBlend Function or FunctionInstance
+bool FilterAssetDataAMaterialLayerBlendFunction(const FAssetData& AssetData)
+{
+	return !AssetData.IsInstanceOf(UMaterialFunctionMaterialLayerBlend::StaticClass()) && !AssetData.IsInstanceOf(UMaterialFunctionMaterialLayerBlendInstance::StaticClass());
+}
+
 /// ===========================================================================================================
 /// SMaterialSubstrateItem Methods
 /// ===========================================================================================================
@@ -634,19 +640,14 @@ void SMaterialSubstrateTreeItem::Construct(const FArguments& InArgs, const TShar
                 		})
                 		.OnObjectChanged_Lambda([=, this](const FAssetData& InAssetData)
                 		{
-							if (IsAssetDataAMaterialLayerBlendFunction(InAssetData))
-							{
-								FSoftObjectPath ObjPath = InAssetData.GetSoftObjectPath();
-								AssetChild->ParameterHandle->SetValue(ObjPath.TryLoad());
+							FSoftObjectPath ObjPath = InAssetData.GetSoftObjectPath();
+							AssetChild->ParameterHandle->SetValue(ObjPath.TryLoad());
 
-								Tree->FunctionInstanceHandle->NotifyPostChange(EPropertyChangeType::ValueSet);
-								Tree->CreateGroupsWidget();
-								Tree->RequestTreeRefresh();
-							}
+							Tree->FunctionInstanceHandle->NotifyPostChange(EPropertyChangeType::ValueSet);
+							Tree->CreateGroupsWidget();
+							Tree->RequestTreeRefresh();
                 		})
-                		// TODO: We need to authorize several UCLass in this asset picker UMaterialFunctionMaterialLayerBlend AND UMaterialFunctionMaterialLayerBlendInstance
-						// the widget only support one, we need to find a solution
-						//.AllowedClass(UMaterialFunctionMaterialLayerBlend::StaticClass())
+						.OnShouldFilterAsset_Static(&FilterAssetDataAMaterialLayerBlendFunction)
                 		.AllowClear(true)
                 		.DisplayUseSelected(false)
                 		.DisplayBrowse(false);
diff --git a/Engine/Source/Runtime/Engine/Private/Materials/MaterialExpressions.cpp b/Engine/Source/Runtime/Engine/Private/Materials/MaterialExpressions.cpp
--- a/Engine/Source/Runtime/Engine/Private/Materials/MaterialExpressions.cpp
+++ b/Engine/Source/Runtime/Engine/Private/Materials/MaterialExpressions.cpp
@@ -18644,15 +18644,56 @@ FLayerNodeId FMaterialLayersFunctionsTree::AddNode(const FPayload& InPayload, FN
 	return NodeId;
 }	
 
+int FMaterialLayersFunctionsTree::RemoveNodeBranch(FNode* BranchParentNode, TArray<FNodeId>& RemovedIds, TArray<FPayload>& RemovedPayloads)
+{
+	int NumRemoved = 0;
+	if (BranchParentNode)
+	{
+		FNodeId RemovedNodeId = BranchParentNode->ChildrenHead;
+		while (IsValidId(RemovedNodeId))
+		{
+			FNode* RemovedNode = GetNode(RemovedNodeId);
+
+			// Record the removed node and payload
+			RemovedIds.Add(RemovedNodeId);
+			RemovedPayloads.Add(Payloads[RemovedNodeId]);
+			NumRemoved++;
+
+			// Remove the potential branch of nodes UNDER this removed node
+			NumRemoved += RemoveNodeBranch(RemovedNode, RemovedIds, RemovedPayloads);
+			
+			// Clear the deleted node
+			FNodeId NextSyblingId = RemovedNode->NextSibling;
+			Nodes[RemovedNodeId] = FNode();
+			Payloads[RemovedNodeId] = FPayload();
 
+			// Move on to the next sybling
+			RemovedNodeId = NextSyblingId;
+		}
+
+		BranchParentNode->ChildrenHead = InvalidId; // Disconnect the branch removed from the parent
+	}
+
+	return NumRemoved;
+}
 
-FLayerNodeIdArray FMaterialLayersFunctionsTree::RemoveNode(FNodeId RemovedNodeId)
+std::pair<TArray<FMaterialLayersFunctionsTree::FNodeId>, TArray<FMaterialLayersFunctionsTree::FPayload>> FMaterialLayersFunctionsTree::RemoveNode(FNodeId RemovedNodeId)
 {
-	FLayerNodeIdArray RemovedIds;
+	TArray<FNodeId> RemovedIds;
+	TArray<FPayload> RemovedPayloads;
 
 	FNode* RemovedNode = GetNode(RemovedNodeId);
 	if (!RemovedNode)
-		return RemovedIds;
+		return { RemovedIds, RemovedPayloads };
+
+	// Record the removed node and payload
+	RemovedIds.Add(RemovedNodeId);
+	RemovedPayloads.Add(Payloads[RemovedNodeId]);
+
+	// remove the branch of nodes starting at the removed node
+	RemoveNodeBranch(RemovedNode, RemovedIds, RemovedPayloads);
+
+	// Then remove the node itself
 
 	FNodeId NextId = RemovedNode->NextSibling;
 
@@ -18698,11 +18739,11 @@ FLayerNodeIdArray FMaterialLayersFunctionsTree::RemoveNode(FNodeId RemovedNodeId
 		PrevNode->NextSibling = NextId;
 	}
 
+	// Clear the node and payload that just have been removed
 	Nodes[RemovedNodeId] = FNode();
 	Payloads[RemovedNodeId] = FPayload();
-	RemovedIds.Add(RemovedNodeId);
 
-	return RemovedIds;
+	return { RemovedIds, RemovedPayloads };
 }
 
 void FMaterialLayersFunctionsTree::MoveNode(FNodeId InNodeId, FNodeId DstParentId, int32 InSiblingIndex)
@@ -18810,22 +18851,6 @@ FString FMaterialLayersFunctionsTree::Log(FString InTab) const
 
 namespace MLFT {
 
-	void RemovePayloadIndex(FMaterialLayersFunctionsTree& Tree, int32 InRemovedIdx)
-	{
-		for (auto& Payload : Tree.Payloads)
-		{
-			if (Payload.Layer > InRemovedIdx)
-			{
-				Payload.Layer--;
-			}
-			if (Payload.Blend > InRemovedIdx - 1)
-			{
-				Payload.Blend--;
-			}
-		}
-	}
-
-
 	FString Print(const FMaterialLayersFunctionsRuntimeData* LFRD, FString InTab)
 	{
 		FString NewLine = TEXT("\r\n");
@@ -19813,15 +19838,45 @@ void FMaterialLayersFunctions::RemoveLayerNodeAt(FLayerNodeId InNodeId)
 
 	if (Tree.IsValidId(InNodeId))
 	{
-		FLayerNodePayload* Payload = Tree.GetPayload(InNodeId);
-		if (Payload && (Payload->Layer >= 0 && Payload->Layer < Layers.Num()))
-		{
-			RemoveBlendedLayerAt(Payload->Layer);
+		auto [RemovedNodes, RemovedPayloads] = Tree.RemoveNode(InNodeId);
 
-			MLFT::RemovePayloadIndex(Tree, Payload->Layer);
+		// The node and its branch have been removed from the tree
+		// Now remove the actual asset from Layers and Blends functions at the corresponding indices (stored in the payload)
+		
+		// Extract the actual asset indices from the returned removed payload 
+		TArray<int32> RemovedAssetIndices;
+		for (auto& Payload : RemovedPayloads)
+		{
+			if (Payload.Layer >= 0 && Payload.Layer < Layers.Num())
+			{
+				RemovedAssetIndices.Add(Payload.Layer);
+			}
 		}
 
-		Tree.RemoveNode(InNodeId);
+		// THen sort the indices in growing order
+		RemovedAssetIndices.Sort();
+
+		// And now from back to front
+		// REMOVE the actual asset in the Layers/Blends arrays
+		// AND adjust the stored indices in the tree node payload
+		for (int i = 0; i < RemovedAssetIndices.Num(); ++i)
+		{
+			int AssetIndex = RemovedAssetIndices[RemovedAssetIndices.Num() - 1 - i];
+
+			RemoveBlendedLayerAt(AssetIndex);
+
+			for (auto& Payload : Tree.Payloads)
+			{
+				if (Payload.Layer > AssetIndex)
+				{
+					Payload.Layer--;
+				}
+				if (Payload.Blend > AssetIndex)
+				{
+					Payload.Blend--;
+				}
+			}
+		}	
 	}
 
 	MLFT::Log(this, __FUNCTION__);
diff --git a/Engine/Source/Runtime/Engine/Public/Materials/MaterialLayersFunctions.h b/Engine/Source/Runtime/Engine/Public/Materials/MaterialLayersFunctions.h
--- a/Engine/Source/Runtime/Engine/Public/Materials/MaterialLayersFunctions.h
+++ b/Engine/Source/Runtime/Engine/Public/Materials/MaterialLayersFunctions.h
@@ -278,8 +278,8 @@ public:
 	// detach the node from the tree hierarchy as well as all the children recursively
 	// the nodes are all cleared in the array of nodes
 	// the payloads are NOT cleared and still hold on their value
-	// return the array of Id of the removed nodes
-	TArray<FNodeId> RemoveNode(FNodeId RemovedNodeId);
+	// return the array of Id of the removed nodes AND the array of Payloads removed
+	std::pair<TArray<FNodeId>, TArray<FPayload>> RemoveNode(FNodeId RemovedNodeId);
 	
 	// Move node from one position to another
 	void MoveNode(FNodeId InNodeId, FNodeId DstParentId, int32 InSiblingIndex);
@@ -337,6 +337,8 @@ public:
 private:
 	FNodeId AllocNode(const FPayload& InPayload);
 
+	// Remove iteratively a full branch of nodes
+	int RemoveNodeBranch(FNode* BranchParentNode, TArray<FNodeId>& RemovedIds, TArray<FPayload>& RemovedPayloads);
 };
 
 template<>
