Texture: clarify tooltip and comments about TC_Displacementmap (reads source A)

diff --git a/Engine/Source/Runtime/Engine/Classes/Engine/TextureDefines.h b/Engine/Source/Runtime/Engine/Classes/Engine/TextureDefines.h
--- a/Engine/Source/Runtime/Engine/Classes/Engine/TextureDefines.h
+++ b/Engine/Source/Runtime/Engine/Classes/Engine/TextureDefines.h
@@ -377,13 +377,13 @@ enum TextureCompressionSettings : int
 	TC_Default					UMETA(DisplayName = "Default (BC1 or BC3 with A)"),
 	TC_Normalmap				UMETA(DisplayName = "Normalmap (BC5)"),
 	TC_Masks					UMETA(DisplayName = "Masks (no sRGB)"),
-	TC_Grayscale				UMETA(DisplayName = "Grayscale (G8/16, RGB8 sRGB)"),
-	TC_Displacementmap			UMETA(DisplayName = "Displacementmap (G8/16)"),
+	TC_Grayscale				UMETA(DisplayName = "Grayscale (G8/16, RGB8 sRGB) from source R"),
+	TC_Displacementmap			UMETA(DisplayName = "Displacementmap (G8/16) from source A", ToolTip = "For Nanite displacement use TC_Alpha"),
 	TC_VectorDisplacementmap	UMETA(DisplayName = "VectorDisplacementmap (RGBA8)"),
 	TC_HDR						UMETA(DisplayName = "HDR (RGBA16F, no sRGB)"),
 	TC_EditorIcon				UMETA(DisplayName = "UserInterface2D (RGBA8)"),
-	TC_Alpha					UMETA(DisplayName = "Alpha (no sRGB, BC4)"),
-	TC_DistanceFieldFont		UMETA(DisplayName = "DistanceFieldFont (G8)"),
+	TC_Alpha					UMETA(DisplayName = "Alpha (no sRGB, BC4) from source R", ToolTip = "Single channel compressed; can also use for displacement and gray"),
+	TC_DistanceFieldFont		UMETA(DisplayName = "DistanceFieldFont (G8) from source A"),
 	TC_HDR_Compressed			UMETA(DisplayName = "HDR Compressed (RGB, BC6H)"),
 	TC_BC7						UMETA(DisplayName = "BC7 (RGBA)"),
 	TC_HalfFloat				UMETA(DisplayName = "Half Float (R16F)"),
diff --git a/Engine/Source/Runtime/Engine/Private/Texture.cpp b/Engine/Source/Runtime/Engine/Private/Texture.cpp
--- a/Engine/Source/Runtime/Engine/Private/Texture.cpp
+++ b/Engine/Source/Runtime/Engine/Private/Texture.cpp
@@ -4127,6 +4127,10 @@ FName GetDefaultTextureFormatName( const ITargetPlatformSettings* TargetPlatform
 	else if (FormatSettings.CompressionSettings == TC_Grayscale || 
 		FormatSettings.CompressionSettings == TC_Displacementmap)
 	{
+		// TC_Displacementmap and TC_Grayscale are almost identical
+		//	EXCEPT Grayscale sets bReplicateRed (R -> RGB) and DisplacementMap sets bReplicateAlpha (A -> RGBA)
+		//	Grayscale reads source R, Displacementmap reads source A
+		
 		// Grayscale is G8 output, unless source is specifically G16
 		//	(eg. RGBA16 source still uses G8 output, not G16)
 		if (SourceFormat == TSF_G16)
