World Partition: actors will now inherit HLOD layers set on their parent level instance actor if they are set to default.

diff --git a/Engine/Source/Runtime/Engine/Private/WorldPartition/WorldPartitionStreamingGeneration.cpp b/Engine/Source/Runtime/Engine/Private/WorldPartition/WorldPartitionStreamingGeneration.cpp
--- a/Engine/Source/Runtime/Engine/Private/WorldPartition/WorldPartitionStreamingGeneration.cpp
+++ b/Engine/Source/Runtime/Engine/Private/WorldPartition/WorldPartitionStreamingGeneration.cpp
@@ -584,6 +584,7 @@ class FWorldPartitionStreamingGenerator
 		FActorContainerID ID;
 		FActorContainerID ParentID;
 		FGuid ContentBundleID;
+		FSoftObjectPath HLODLayer;
 
 		FContainerCollectionInstanceDescriptor(const FContainerCollectionInstanceDescriptor& Other) = delete;
 		FContainerCollectionInstanceDescriptor(FContainerCollectionInstanceDescriptor&& Other) = default;
@@ -716,13 +717,13 @@ class FWorldPartitionStreamingGenerator
 		}
 	}
 
-	void ResolveHLODLayer(FStreamingGenerationActorDescView& ActorDescView, const FSoftObjectPath& DefaultHLODLayer)
+	void ResolveHLODLayer(FStreamingGenerationActorDescView& ActorDescView, const FSoftObjectPath& ParentHLODLayer)
 	{
-		// Only assign the default layer to actors that don't have a valid HLOD layer set. HLOD actors will have their 
-		// parent HLOD layer set during HLOD generation.
+		// Only assign the parent layer to actors that don't have a valid HLOD layer set and that are HLOD relevant. 
+		// HLOD actors will have their parent HLOD layer set during HLOD generation.
 		if (!ActorDescView.GetHLODLayer().IsValid() && ActorDescView.GetIsSpatiallyLoaded() && ActorDescView.GetActorIsHLODRelevant())
 		{
-			ActorDescView.SetRuntimeHLODLayer(DefaultHLODLayer);
+			ActorDescView.SetRuntimeHLODLayer(ParentHLODLayer);
 		}
 	}
 
@@ -1004,6 +1005,7 @@ class FWorldPartitionStreamingGenerator
 				// placed inside a Content Bundle will propagate their Content Bundle GUID to child instances.
 				SubContainerInstanceDescriptor.ContentBundleID = ContainerCollectionInstanceDescriptor.ContentBundleID;
 				SubContainerInstanceDescriptor.InstanceData = InheritParentContainerPerInstanceData(ContainerCollectionInstanceDescriptor, ContainerCollectionInstanceView);
+				SubContainerInstanceDescriptor.HLODLayer = ContainerCollectionInstanceView.GetHLODLayer().IsValid() ? ContainerCollectionInstanceView.GetHLODLayer() : ContainerCollectionInstanceDescriptor.HLODLayer;
 
 				if (WorldPartitionSubsystem && ContainerID.IsMainContainer() && ContainerCollectionInstanceView.GetChildContainerFilterType() == EWorldPartitionActorFilterType::Loading)
 				{
@@ -1048,6 +1050,7 @@ class FWorldPartitionStreamingGenerator
 		MainContainerCollectionInstance.OwnerName = TEXT("MainContainer");
 		MainContainerCollectionInstance.ContentBundleID = InContainerInstanceCollection.GetContentBundleGuid();
 		MainContainerCollectionInstance.InstanceData.bIsSpatiallyLoaded = true; // Since we apply AND logic on spatially loaded flag recursively, startup value must be true
+		MainContainerCollectionInstance.HLODLayer = WorldPartitionContext ? FSoftObjectPath(WorldPartitionContext->GetDefaultHLODLayer()) : FSoftObjectPath();
 
 		// Create Child Containers
 		CreateActorDescriptorViewsRecursive(MoveTemp(MainContainerCollectionInstance));
@@ -1087,7 +1090,7 @@ class FWorldPartitionStreamingGenerator
 			ResolveRuntimeSpatiallyLoaded(ActorDescView);
 			ResolveRuntimeGrid(ActorDescView);
 			ResolveRuntimeDataLayers(ActorDescView, ContainerCollectionInstanceDescriptor.DataLayerResolvers);
-			ResolveHLODLayer(ActorDescView, WorldPartitionContext ? FSoftObjectPath(WorldPartitionContext->GetDefaultHLODLayer()) : FSoftObjectPath());
+			ResolveHLODLayer(ActorDescView, ContainerCollectionInstanceDescriptor.HLODLayer);
 			ResolveParentView(ActorDescView, ContainerCollectionInstanceDescriptor.ActorDescViewMap.Get(), ContainerCollectionInstanceDescriptor.EditorOnlyActorDescSet, ContainerCollectionInstanceDescriptor.EditorOnlyParentActorTransforms);
 		};
 
