Fix issue where sequence length of child montages could get out of sync with the parent asset. Also force update of common frame rate when child montage is created to prevent display of very large number of frames

diff --git a/Engine/Source/Runtime/Engine/Private/Animation/AnimMontage.cpp b/Engine/Source/Runtime/Engine/Private/Animation/AnimMontage.cpp
--- a/Engine/Source/Runtime/Engine/Private/Animation/AnimMontage.cpp
+++ b/Engine/Source/Runtime/Engine/Private/Animation/AnimMontage.cpp
@@ -1213,6 +1213,8 @@ void UAnimMontage::RefreshParentAssetData()
 		}
 	}
 
+	UpdateCommonTargetFrameRate();
+
 	OnMontageChanged.Broadcast();
 }
 
diff --git a/Engine/Source/Runtime/Engine/Private/Animation/AnimSequenceBase.cpp b/Engine/Source/Runtime/Engine/Private/Animation/AnimSequenceBase.cpp
--- a/Engine/Source/Runtime/Engine/Private/Animation/AnimSequenceBase.cpp
+++ b/Engine/Source/Runtime/Engine/Private/Animation/AnimSequenceBase.cpp
@@ -967,12 +967,30 @@ void UAnimSequenceBase::RefreshParentAssetData()
 
 	ValidateModel();
 	Controller->OpenBracket(LOCTEXT("RefreshParentAssetData_Bracket", "Refreshing Parent Asset Data"));
-	{
+	{	
 		const IAnimationDataModel* ParentDataModel = ParentSeqBase->GetDataModel();
-
-		if (DataModelInterface->GetNumberOfFrames() != ParentDataModel->GetNumberOfFrames())
+		
+		if (!FMath::IsNearlyEqual(GetPlayLength(), ParentSeqBase->GetPlayLength()))
 		{
-			Controller->SetNumberOfFrames(ParentDataModel->GetNumberOfFrames());
+			const FFrameRate& ParentFrameRate = ParentDataModel->GetFrameRate();
+			const FFrameRate& ChildFrameRate = GetDataModelInterface()->GetFrameRate();
+
+			if (ChildFrameRate.IsMultipleOf(ParentFrameRate) || ChildFrameRate.IsFactorOf(ParentFrameRate))
+			{
+				// Ensure that play lengths are the same by converting the parent number of frames into the child's frame rate
+				FFrameTime ChildNumFrames = FFrameRate::TransformTime(ParentFrameRate.AsFrameTime(ParentSeqBase->GetPlayLength()), ParentFrameRate, ChildFrameRate);
+				Controller->SetNumberOfFrames(ChildNumFrames.RoundToFrame());
+			}
+			else
+			{
+				UE_LOG(LogTemp, Display,
+						TEXT("Child %s %s has incompatible frame rate with Parent %s. Child: %s, Parent: %s. Reimport the child animation with a compatible frame rate."),
+					*ParentSeqBase->GetClass()->GetName(), 
+					*GetName(), 
+					*ParentSeqBase->GetName(), 
+					*ChildFrameRate.ToPrettyText().ToString(), 
+					*ParentFrameRate.ToPrettyText().ToString());
+			}
 		}
 		
 		Controller->RemoveAllCurvesOfType(ERawCurveTrackTypes::RCT_Float);
