engine - update for "experimental PoseSearch" comments to "PoseSearch"

[RN] skip

diff --git a/Engine/Source/Runtime/Engine/Classes/Animation/AnimNode_SequencePlayer.h b/Engine/Source/Runtime/Engine/Classes/Animation/AnimNode_SequencePlayer.h
--- a/Engine/Source/Runtime/Engine/Classes/Animation/AnimNode_SequencePlayer.h
+++ b/Engine/Source/Runtime/Engine/Classes/Animation/AnimNode_SequencePlayer.h
@@ -68,7 +68,7 @@ protected:
 	// The start position [range: 0 - sequence length] to use when initializing. When looping, play will still jump back to the beginning when reaching the end.
 	virtual float GetStartPosition() const { return 0.0f; }
 
-	// Use pose matching to choose the start position. Requires experimental PoseSearch plugin.
+	// Use pose matching to choose the start position. Requires PoseSearch plugin.
 	virtual bool GetStartFromMatchingPose() const { return false; }
 
 	// Set the start time of this node.
@@ -141,7 +141,7 @@ protected:
 	UPROPERTY(EditAnywhere, Category = Settings, meta = (PinHiddenByDefault, FoldProperty))
 	bool bLoopAnimation = true;
 
-	// Use pose matching to choose the start position. Requires experimental PoseSearch plugin.
+	// Use pose matching to choose the start position. Requires PoseSearch plugin.
 	UPROPERTY(EditAnywhere, Category = PoseMatching, meta = (PinHiddenByDefault, FoldProperty))
 	bool bStartFromMatchingPose = false;
 #endif
@@ -226,7 +226,7 @@ protected:
 	UPROPERTY(EditAnywhere, Category = Settings, meta = (PinHiddenByDefault))
 	bool bLoopAnimation = true;
 
-	// Use pose matching to choose the start position. Requires experimental PoseSearch plugin.
+	// Use pose matching to choose the start position. Requires PoseSearch plugin.
 	UPROPERTY(EditAnywhere, Category = PoseMatching, meta = (PinHiddenByDefault))
 	bool bStartFromMatchingPose = false;
 
