Bring back CL38157186 with fix from Ben Ziegler: change actor ExecuteConstruction to correctly register blueprint-added components on actors that have been fully registered, even if world has not been initialized.

diff --git a/Engine/Source/Runtime/Engine/Private/ActorConstruction.cpp b/Engine/Source/Runtime/Engine/Private/ActorConstruction.cpp
--- a/Engine/Source/Runtime/Engine/Private/ActorConstruction.cpp
+++ b/Engine/Source/Runtime/Engine/Private/ActorConstruction.cpp
@@ -883,7 +883,7 @@ bool AActor::ExecuteConstruction(const FTransform& Transform, const FRotationCon
 			for (UActorComponent* ActorComponent : PostSCSComponents)
 			{
 				// Limit registration to components that are known to have been created during SCS execution
-				if (!ActorComponent->IsRegistered() && ActorComponent->bAutoRegister && IsValidChecked(ActorComponent) && GetWorld()->bIsWorldInitialized
+				if (!ActorComponent->IsRegistered() && ActorComponent->bAutoRegister && IsValidChecked(ActorComponent) && (GetWorld()->bIsWorldInitialized || bHasRegisteredAllComponents)
 					&& (ActorComponent->CreationMethod == EComponentCreationMethod::SimpleConstructionScript || !PreSCSComponents.Contains(ActorComponent)))
 				{
 					USimpleConstructionScript::RegisterInstancedComponent(ActorComponent);
diff --git a/Engine/Source/Runtime/Engine/Private/WorldPartition/WorldPartition.cpp b/Engine/Source/Runtime/Engine/Private/WorldPartition/WorldPartition.cpp
--- a/Engine/Source/Runtime/Engine/Private/WorldPartition/WorldPartition.cpp
+++ b/Engine/Source/Runtime/Engine/Private/WorldPartition/WorldPartition.cpp
@@ -2467,23 +2467,43 @@ FBox UWorldPartition::GetRuntimeWorldBounds() const
 
 void UWorldPartition::ApplyRuntimeCellsTransformerStack(ULevel* InLevel)
 {
-	auto ApplyTransformPhase = [this, InLevel](TFunction<void(const FRuntimeCellTransformerInstance& TransformerInstance, ULevel* Level)> Func)
+	if (!RuntimeCellsTransformerStack.IsEmpty())
 	{
-		for (const FRuntimeCellTransformerInstance& TransformerInstance : RuntimeCellsTransformerStack)
+		auto ApplyTransformPhase = [this, InLevel](TFunction<void(const FRuntimeCellTransformerInstance& TransformerInstance, ULevel* Level)> Func)
 		{
-			TPair<double, int32>& TotalStats = RuntimeCellsTransformerStackTimes.FindOrAdd(TransformerInstance.Class);
-			TotalStats.Key -= FPlatformTime::Seconds();
+			for (const FRuntimeCellTransformerInstance& TransformerInstance : RuntimeCellsTransformerStack)
+			{
+				TPair<double, int32>& TotalStats = RuntimeCellsTransformerStackTimes.FindOrAdd(TransformerInstance.Class);
+				TotalStats.Key -= FPlatformTime::Seconds();
 		
-			Func(TransformerInstance, InLevel);
+				Func(TransformerInstance, InLevel);
 		
-			TotalStats.Key += FPlatformTime::Seconds();
-			TotalStats.Value++;
+				TotalStats.Key += FPlatformTime::Seconds();
+				TotalStats.Value++;
+			}
+		};
+
+		// Rerun construction scripts for actors in levels before executing the cell transformers, as the Blueprint parent classes might have changed 
+		// and we want to make sure the constructed actors are up to date.
+		TArray<TObjectPtr<AActor>> LevelActors;
+		Algo::CopyIf(InLevel->Actors, LevelActors, [](AActor* InActor)
+		{
+			return IsValid(InActor) && !InActor->IsChildActor();
+		});
+	
+		for (TObjectPtr<AActor>& LevelActor : LevelActors)
+		{
+			// Deal with deleted child actors that might have been deleted from previous calls to RerunConstructionScripts
+			if (IsValid(LevelActor))
+			{
+				LevelActor->RerunConstructionScripts();
+			}
 		}
-	};
 
-	ApplyTransformPhase([](const FRuntimeCellTransformerInstance& TransformerInstance, ULevel* Level) { TransformerInstance.PreTransform(Level); });
-	ApplyTransformPhase([](const FRuntimeCellTransformerInstance& TransformerInstance, ULevel* Level) { TransformerInstance.Transform(Level); });
-	ApplyTransformPhase([](const FRuntimeCellTransformerInstance& TransformerInstance, ULevel* Level) { TransformerInstance.PostTransform(Level); });
+		ApplyTransformPhase([](const FRuntimeCellTransformerInstance& TransformerInstance, ULevel* Level) { TransformerInstance.PreTransform(Level); });
+		ApplyTransformPhase([](const FRuntimeCellTransformerInstance& TransformerInstance, ULevel* Level) { TransformerInstance.Transform(Level); });
+		ApplyTransformPhase([](const FRuntimeCellTransformerInstance& TransformerInstance, ULevel* Level) { TransformerInstance.PostTransform(Level); });
+	}
 }
 #endif
 
