Add PostAssetCompilation callbacks to the respective components for FAnimBankCompilingManager and FSkinnedAssetCompilingManager to enable recreating renderstate when async asset compilation is finished.

diff --git a/Engine/Source/Runtime/Engine/Classes/Components/SkinnedMeshComponent.h b/Engine/Source/Runtime/Engine/Classes/Components/SkinnedMeshComponent.h
--- a/Engine/Source/Runtime/Engine/Classes/Components/SkinnedMeshComponent.h
+++ b/Engine/Source/Runtime/Engine/Classes/Components/SkinnedMeshComponent.h
@@ -552,6 +552,14 @@ public:
 	 */
 	ENGINE_API virtual bool HasValidNaniteData() const;
 
+
+#if WITH_EDITOR
+	/**
+	 * Called by the asset compilers (FAnimBankCompilingManager / FSkinnedAssetCompilingManager) when done to allow render state to be created.
+	 */
+	ENGINE_API virtual void PostAssetCompilation() {};
+#endif
+
 #if UE_ENABLE_DEBUG_DRAWING
 	/** Get whether to draw this mesh's debug skeleton */
 	bool ShouldDrawDebugSkeleton() const
diff --git a/Engine/Source/Runtime/Engine/Private/Animation/AnimBankCompiler.cpp b/Engine/Source/Runtime/Engine/Private/Animation/AnimBankCompiler.cpp
--- a/Engine/Source/Runtime/Engine/Private/Animation/AnimBankCompiler.cpp
+++ b/Engine/Source/Runtime/Engine/Private/Animation/AnimBankCompiler.cpp
@@ -15,6 +15,7 @@
 #include "Animation/AnimBank.h"
 #include "ProfilingDebugging/CountersTrace.h"
 #include "AnimationSequenceCompiler.h"
+#include "Components/AnimBankMeshComponent.h"
 
 #define LOCTEXT_NAMESPACE "AnimBankCompiler"
 
@@ -217,6 +218,12 @@ void FAnimBankCompilingManager::PostCompilation(UAnimBank* AnimBank)
 
 		AnimBank->InitResources();
 
+		FObjectCacheContextScope ObjectCacheScope;
+		for (UAnimBankMeshComponent* Component : ObjectCacheScope.GetContext().GetAnimBankMeshComponents(AnimBank))
+		{
+			Component->PostAssetCompilation();
+		}
+
 		// Calling this delegate during app exit might be quite dangerous and lead to crash
 		// if the content browser wants to refresh a thumbnail it might try to load a package
 		// which will then fail due to various reasons related to the editor shutting down.
diff --git a/Engine/Source/Runtime/Engine/Private/ObjectCacheContext.cpp b/Engine/Source/Runtime/Engine/Private/ObjectCacheContext.cpp
--- a/Engine/Source/Runtime/Engine/Private/ObjectCacheContext.cpp
+++ b/Engine/Source/Runtime/Engine/Private/ObjectCacheContext.cpp
@@ -16,6 +16,7 @@
 #include "HAL/LowLevelMemStats.h"
 #include "Misc/ScopeRWLock.h"
 #include "Components/ComponentInterfaceIterator.h"
+#include "Components/AnimBankMeshComponent.h"
 
 #define LOCTEXT_NAMESPACE "ObjectCache"
 
@@ -685,25 +686,6 @@ TObjectCacheIterator<IPrimitiveComponent> FObjectCacheContext::GetPrimitiveCompo
 	return TObjectCacheIterator<IPrimitiveComponent>(MakeArrayView(PrimitiveComponents.GetValue()));
 }
 
-TObjectCacheIterator<IStaticMeshComponent> FObjectCacheContext::GetStaticMeshComponents()
-{
-	using namespace ObjectCacheContextImpl;
-
-	if (!StaticMeshComponents.IsSet())
-	{
-		TRACE_CPUPROFILER_EVENT_SCOPE(ComputeStaticMeshComponents);
-
-		TArray<IStaticMeshComponent*> Array;
-		Array.Reserve(4096);
-		for (TComponentInterfaceIterator<IStaticMeshComponent> It(RF_ClassDefaultObject, true /*bIncludeDerivedClasses*/, GetObjectCacheInternalFlagsExclusion()); It; ++It)
-		{
-			Array.Add(*It);
-		}
-		StaticMeshComponents = MoveTemp(Array);
-	}
-	return TObjectCacheIterator<IStaticMeshComponent>(MakeArrayView(StaticMeshComponents.GetValue()));
-}
-
 TObjectCacheIterator<UTexture> FObjectCacheContext::GetUsedTextures(UMaterialInterface* MaterialInterface)
 {
 	using namespace ObjectCacheContextImpl;
@@ -778,25 +760,90 @@ TObjectCacheIterator<UTextureCollection> FObjectCacheContext::GetUsedTextureColl
 	return TObjectCacheIterator<UTextureCollection>(TextureCollections->Array());
 }
 
-TObjectCacheIterator<USkinnedMeshComponent> FObjectCacheContext::GetSkinnedMeshComponents()
+
+template<typename UComponentType, typename UAssetType>
+template <typename ObjectIteratorType>
+TObjectCacheIterator<UComponentType> FObjectCacheContext::TComponentTypeCache<UComponentType,UAssetType>::Get()
 {
 	using namespace ObjectCacheContextImpl;
 
-	if (!SkinnedMeshComponents.IsSet())
+	if (!Components.IsSet())
 	{
 		TRACE_CPUPROFILER_EVENT_SCOPE(ComputeSkinnedMeshComponents);
 
-		TArray<USkinnedMeshComponent*> Array;
+		TArray<UComponentType*> Array;
 		Array.Reserve(4096);
-		for (TObjectIterator<USkinnedMeshComponent> It(RF_ClassDefaultObject, true /*bIncludeDerivedClasses*/, GetObjectCacheInternalFlagsExclusion()); It; ++It)
+		for (ObjectIteratorType It(RF_ClassDefaultObject, true /*bIncludeDerivedClasses*/, GetObjectCacheInternalFlagsExclusion()); It; ++It)
 		{
 			Array.Add(*It);
 		}
-		SkinnedMeshComponents = MoveTemp(Array);
+		Components = MoveTemp(Array);
+	}
+	return TObjectCacheIterator<UComponentType>(MakeArrayView(Components.GetValue()));
+}
+
+template<typename UComponentType, typename UAssetType>
+template <typename ObjectIteratorType, typename GetNumFuncType, typename GetFuncType>
+TObjectCacheIterator<UComponentType> FObjectCacheContext::TComponentTypeCache<UComponentType,UAssetType>::Get(UAssetType *AssetType, GetNumFuncType GetNumFunc, GetFuncType GetFunc)
+{
+	const TSet<UComponentType*>& Result = GetSet<ObjectIteratorType>(AssetType, GetNumFunc, GetFunc);
+	return TObjectCacheIterator<UComponentType>(Result.Array());
+}
+
+
+template<typename UComponentType, typename UAssetType>
+template <typename ObjectIteratorType, typename GetNumFuncType, typename GetFuncType>
+const TSet<UComponentType* >& FObjectCacheContext::TComponentTypeCache<UComponentType,UAssetType>::GetSet(UAssetType *AssetType, GetNumFuncType GetNumFunc, GetFuncType GetFunc)
+{
+	if (!AssetToComponents.IsSet())
+	{
+		TRACE_CPUPROFILER_EVENT_SCOPE(ComputeSkinnedMeshToComponents);
+
+		TMap<TObjectKey<UAssetType>, TSet<UComponentType*>> TempMap;
+		TempMap.Reserve(8192);
+		for (UComponentType* Component : Get<ObjectIteratorType>())
+		{
+			int32 Num = GetNumFunc(Component);
+			for (int32 Index = 0; Index < Num; ++Index)
+			{
+				TempMap.FindOrAdd(GetFunc(Component, Index)).FindOrAdd(Component);
+			}
+		}
+		AssetToComponents = MoveTemp(TempMap);
 	}
-	return TObjectCacheIterator<USkinnedMeshComponent>(MakeArrayView(SkinnedMeshComponents.GetValue()));
+
+	static TSet<UComponentType*> EmptySet;
+	TSet<UComponentType*>* Set = AssetToComponents.GetValue().Find(AssetType);
+	const TSet<UComponentType*>& Result = Set ? *Set : EmptySet;
+
+	return Result;
+}
+
+
+TObjectCacheIterator<USkinnedMeshComponent> FObjectCacheContext::GetSkinnedMeshComponents()
+{
+	return SkinnedMeshCache.Get<TObjectIterator<USkinnedMeshComponent>>();
+}
+
+TObjectCacheIterator<USkinnedMeshComponent> FObjectCacheContext::GetSkinnedMeshComponents(USkinnedAsset* InSkinnedAsset)
+{
+	return SkinnedMeshCache.Get<TObjectIterator<USkinnedMeshComponent>>(InSkinnedAsset, 
+		[](USkinnedMeshComponent* Component){ return 1; },
+		[](USkinnedMeshComponent* Component, int32 ) { return Component->GetSkinnedAsset(); });
+}
+
+TObjectCacheIterator<UAnimBankMeshComponent> FObjectCacheContext::GetAnimBankMeshComponents()
+{
+	return AnimBankMeshCache.Get<TObjectIterator<UAnimBankMeshComponent>>();
 }
 
+TObjectCacheIterator<UAnimBankMeshComponent> FObjectCacheContext::GetAnimBankMeshComponents(UAnimBank* InAnimBank)
+{
+	return AnimBankMeshCache.Get<TObjectIterator<UAnimBankMeshComponent>>(InAnimBank, 
+		[](UAnimBankMeshComponent* Component){ return Component->AnimBankItems.Num();  },
+		[](UAnimBankMeshComponent* Component, int32 Index)->const UAnimBank* { return Component->AnimBankItems[Index].BankAsset; }
+	);
+}
 
 TObjectCacheIterator<UMaterialInterface> FObjectCacheContext::GetUsedMaterials(IPrimitiveComponent* Component)
 {
@@ -847,6 +894,11 @@ TObjectCacheIterator<UMaterialInterface> FObjectCacheContext::GetUsedMaterials(I
 	return TObjectCacheIterator<UMaterialInterface>(Materials->Array());
 }
 
+TObjectCacheIterator<IStaticMeshComponent> FObjectCacheContext::GetStaticMeshComponents()
+{
+	return StaticMeshCache.Get<TComponentInterfaceIterator<IStaticMeshComponent>>();
+}
+
 TObjectCacheIterator<IStaticMeshComponent> FObjectCacheContext::GetStaticMeshComponents(UStaticMesh* InStaticMesh)
 {
 	using namespace ObjectCacheContextImpl;
@@ -857,22 +909,9 @@ TObjectCacheIterator<IStaticMeshComponent> FObjectCacheContext::GetStaticMeshCom
 	}
 #endif
 
-	if (!StaticMeshToComponents.IsSet())
-	{
-		TRACE_CPUPROFILER_EVENT_SCOPE(ComputeStaticMeshToComponents);
-
-		TMap<TObjectKey<UStaticMesh>, TSet<IStaticMeshComponent*>> TempMap;
-		TempMap.Reserve(8192);
-		for (IStaticMeshComponent* Component : GetStaticMeshComponents())
-		{
-			TempMap.FindOrAdd(Component->GetStaticMesh()).FindOrAdd(Component);
-		}
-		StaticMeshToComponents = MoveTemp(TempMap);
-	}
-
-	static TSet<IStaticMeshComponent*> EmptySet;
-	TSet<IStaticMeshComponent*>* Set = StaticMeshToComponents.GetValue().Find(InStaticMesh);
-	const TSet<IStaticMeshComponent*>& ComputeResult = Set ? *Set : EmptySet;
+	const TSet<IStaticMeshComponent*>& ComputeResult = StaticMeshCache.GetSet<TComponentInterfaceIterator<IStaticMeshComponent>>(InStaticMesh, 
+		[](IStaticMeshComponent*) { return 1; },
+		[](IStaticMeshComponent* StaticMeshComponent, int32 Index) { return StaticMeshComponent->GetStaticMesh(); });
 
 #if WITH_EDITOR
 	if (IsReverseLookupComparisonActive(EObjectReverseLookupMask::StaticMeshToComponentLookupCache))
@@ -1034,28 +1073,6 @@ TObjectCacheIterator<UTexture> FObjectCacheContext::GetTexturesAffectedByTexture
 
 #endif
 
-TObjectCacheIterator<USkinnedMeshComponent> FObjectCacheContext::GetSkinnedMeshComponents(USkinnedAsset* InSkinnedAsset)
-{
-	if (!SkinnedAssetToComponents.IsSet())
-	{
-		TRACE_CPUPROFILER_EVENT_SCOPE(ComputeSkinnedMeshToComponents);
-
-		TMap<TObjectKey<USkinnedAsset>, TSet<USkinnedMeshComponent*>> TempMap;
-		TempMap.Reserve(8192);
-		for (USkinnedMeshComponent* Component : GetSkinnedMeshComponents())
-		{
-			TempMap.FindOrAdd(Component->GetSkinnedAsset()).FindOrAdd(Component);
-		}
-		SkinnedAssetToComponents = MoveTemp(TempMap);
-	}
-
-	static TSet<USkinnedMeshComponent*> EmptySet;
-	TSet<USkinnedMeshComponent*>* Set = SkinnedAssetToComponents.GetValue().Find(InSkinnedAsset);
-	const TSet<USkinnedMeshComponent*>& Result = Set ? *Set : EmptySet;
-
-	return TObjectCacheIterator<USkinnedMeshComponent>(Result.Array());
-}
-
 TObjectCacheIterator<UMaterialInterface> FObjectCacheContext::GetMaterialsAffectedByMaterials(TArrayView<UMaterialInterface*> InMaterials)
 {
 	using namespace ObjectCacheContextImpl;
diff --git a/Engine/Source/Runtime/Engine/Private/SkinnedAssetCompiler.cpp b/Engine/Source/Runtime/Engine/Private/SkinnedAssetCompiler.cpp
--- a/Engine/Source/Runtime/Engine/Private/SkinnedAssetCompiler.cpp
+++ b/Engine/Source/Runtime/Engine/Private/SkinnedAssetCompiler.cpp
@@ -239,6 +239,11 @@ void FSkinnedAssetCompilingManager::PostCompilation(USkinnedAsset* SkinnedAsset)
 			}
 		}
 
+		for (USkinnedMeshComponent* Component : ObjectCacheScope.GetContext().GetSkinnedMeshComponents(SkinnedAsset))
+		{
+			Component->PostAssetCompilation();
+		}
+
 		// Calling this delegate during app exit might be quite dangerous and lead to crash
 		// if the content browser wants to refresh a thumbnail it might try to load a package
 		// which will then fail due to various reasons related to the editor shutting down.
diff --git a/Engine/Source/Runtime/Engine/Public/ObjectCacheContext.h b/Engine/Source/Runtime/Engine/Public/ObjectCacheContext.h
--- a/Engine/Source/Runtime/Engine/Public/ObjectCacheContext.h
+++ b/Engine/Source/Runtime/Engine/Public/ObjectCacheContext.h
@@ -18,9 +18,10 @@ class UTexture;
 class UTextureCollection;
 class USkeletalMesh;
 class USkinnedAsset;
-
 class IPrimitiveComponent;
 class IStaticMeshComponent;
+class UAnimBankMeshComponent;
+class UAnimBank;
 
 /**
  *  Class to abstract implementation details of the containers used inside the
@@ -74,6 +75,8 @@ public:
 	ENGINE_API TObjectCacheIterator<IStaticMeshComponent>  GetStaticMeshComponents();
 	ENGINE_API TObjectCacheIterator<USkinnedMeshComponent> GetSkinnedMeshComponents();
 	ENGINE_API TObjectCacheIterator<USkinnedMeshComponent> GetSkinnedMeshComponents(USkinnedAsset* InSkinnedAsset);
+	ENGINE_API TObjectCacheIterator<UAnimBankMeshComponent> GetAnimBankMeshComponents();
+	ENGINE_API TObjectCacheIterator<UAnimBankMeshComponent> GetAnimBankMeshComponents(UAnimBank* InAnimBank);
 	ENGINE_API TObjectCacheIterator<IStaticMeshComponent>  GetStaticMeshComponents(UStaticMesh* InStaticMesh);
 	ENGINE_API TObjectCacheIterator<UMaterialInterface>    GetMaterialsAffectedByTexture(UTexture* InTexture);
 	ENGINE_API TObjectCacheIterator<UMaterialInterface>    GetMaterialsAffectedByTextureCollection(UTextureCollection* InTextureCollection);
@@ -99,11 +102,27 @@ private:
 	TOptional<TMap<TObjectKey<UTexture>, TSet<UTexture*>>> TextureToTextures;
 #endif
 	TOptional<TMap<TObjectKey<UMaterialInterface>, TSet<IPrimitiveComponent*>>> MaterialToPrimitives;
-	TOptional<TMap<TObjectKey<UStaticMesh>, TSet<IStaticMeshComponent*>>> StaticMeshToComponents;
-	TOptional<TMap<TObjectKey<USkinnedAsset>, TSet<USkinnedMeshComponent*>>> SkinnedAssetToComponents;
-	TOptional<TArray<USkinnedMeshComponent*>> SkinnedMeshComponents;
-	TOptional<TArray<IStaticMeshComponent*>> StaticMeshComponents;
+
+	template <typename UComponentType, typename UAssetType>
+	struct TComponentTypeCache
+	{
+		TOptional<TArray<UComponentType*>> Components;
+		TOptional<TMap<TObjectKey<UAssetType>, TSet<UComponentType*>>> AssetToComponents;
+
+		template <typename ObjectIteratorType>
+		TObjectCacheIterator<UComponentType> Get();
+
+		template <typename ObjectIteratorType, typename GetNumFuncType, typename GetFuncType>
+		TObjectCacheIterator<UComponentType> Get(UAssetType* AssetType, GetNumFuncType GetNumFunc, GetFuncType GetFunc);
+
+		template <typename ObjectIteratorType, typename GetNumFuncType, typename GetFuncType>
+		const TSet<UComponentType* >& GetSet(UAssetType* AssetType, GetNumFuncType GetNumFunc, GetFuncType GetFunc);
+	};
+
+	TComponentTypeCache<USkinnedMeshComponent, USkinnedAsset> SkinnedMeshCache;
+	TComponentTypeCache<UAnimBankMeshComponent, UAnimBank> AnimBankMeshCache;
 	TOptional<TArray<IPrimitiveComponent*>> PrimitiveComponents;
+	TComponentTypeCache<IStaticMeshComponent, UStaticMesh> StaticMeshCache;
 };
 
 /**
