Fixed some MarkRenderStateDirty related issues.
* Changed MarkRenderStateDirty to handle the case where ShouldCreateRenderState have previously suppressed creation (used almost exclusively in SMC) and removed some calls to RecreateRenderState_Concurrent.
* Moved call to ReregisterComponent out of UActorComponent::PostApplyToComponent and instead call it after ApplyToComponent is finished avoiding callstacks where the proxy is recreated half-way through applying data.

diff --git a/Engine/Source/Runtime/Engine/Private/ComponentInstanceDataCache.cpp b/Engine/Source/Runtime/Engine/Private/ComponentInstanceDataCache.cpp
--- a/Engine/Source/Runtime/Engine/Private/ComponentInstanceDataCache.cpp
+++ b/Engine/Source/Runtime/Engine/Private/ComponentInstanceDataCache.cpp
@@ -522,6 +522,19 @@ bool FActorComponentInstanceData::MatchesComponent(const UActorComponent* Compon
 	return FActorComponentInstanceSourceInfo(SourceComponentTemplate, SourceComponentCreationMethod, SourceComponentTypeSerializedIndex).MatchesComponent(Component, ComponentTemplate);
 }
 
+void FActorComponentInstanceData::DoApplyToComponent(UActorComponent* Component, const ECacheApplyPhase CacheApplyPhase)
+{
+	// First call the virtual implementation
+	ApplyToComponent(Component, CacheApplyPhase);
+
+	// Moved out of PostApplyToComponent as that is called from the middle of the ApplyToComponent. The logical thing would be to move the call to PostApplyToComponent here
+	// but that has dependencies now where some components rely on that call happening in the middle, after properties may have been applied (E.g., UStatiMeshComponent).
+	if (Component->IsRegistered())
+	{
+		Component->ReregisterComponent();
+	}
+}
+
 void FActorComponentInstanceData::ApplyToComponent(UActorComponent* Component, const ECacheApplyPhase CacheApplyPhase)
 {
 	// After the user construction script has run we will re-apply all the cached changes that do not conflict
@@ -776,7 +789,7 @@ void FComponentInstanceDataCache::ApplyToActor(AActor* Actor, const ECacheApplyP
 					const TStructOnScope<FActorComponentInstanceData>& ComponentInstanceData = ComponentsInstanceData[ComponentInstanceDataIt.GetIndex()];
 					if (ComponentInstanceData->MatchesComponent(ComponentInstance, ComponentTemplate))
 					{
-						ComponentInstanceData->ApplyToComponent(ComponentInstance, CacheApplyPhase);
+						ComponentInstanceData->DoApplyToComponent(ComponentInstance, CacheApplyPhase);
 						ComponentInstanceDataToConsider[ComponentInstanceDataIt.GetIndex()] = false;
 						break;
 					}
diff --git a/Engine/Source/Runtime/Engine/Private/Components/ActorComponent.cpp b/Engine/Source/Runtime/Engine/Private/Components/ActorComponent.cpp
--- a/Engine/Source/Runtime/Engine/Private/Components/ActorComponent.cpp
+++ b/Engine/Source/Runtime/Engine/Private/Components/ActorComponent.cpp
@@ -2371,8 +2371,9 @@ void UActorComponent::DoDeferredRenderUpdates_Concurrent()
 
 void UActorComponent::MarkRenderStateDirty()
 {
+	bool bCreateOrUpdateRenderState = bRenderStateCreated || ShouldCreateRenderState();
 	// If registered and has a render state to mark as dirty
-	if(IsRegistered() && bRenderStateCreated && (!bRenderStateDirty || !GetWorld()))
+	if(IsRegistered() && bCreateOrUpdateRenderState && (!bRenderStateDirty || !GetWorld()))
 	{
 		// Flag as dirty
 		bRenderStateDirty = true;
@@ -3036,10 +3037,6 @@ void UActorComponent::Serialize(FArchive& Ar)
 
 void UActorComponent::PostApplyToComponent()
 {
-	if (IsRegistered())
-	{
-		ReregisterComponent();
-	}
 }
 
 AActor* UActorComponent::GetActorOwnerNoninline() const
diff --git a/Engine/Source/Runtime/Engine/Private/Components/StaticMeshComponent.cpp b/Engine/Source/Runtime/Engine/Private/Components/StaticMeshComponent.cpp
--- a/Engine/Source/Runtime/Engine/Private/Components/StaticMeshComponent.cpp
+++ b/Engine/Source/Runtime/Engine/Private/Components/StaticMeshComponent.cpp
@@ -2363,17 +2363,7 @@ bool UStaticMeshComponent::SetStaticMesh(UStaticMesh* NewMesh)
 #endif // UE_WITH_PSO_PRECACHING
 
 	// Need to send this to render thread at some point
-	if (IsRenderStateCreated())
-	{
-		MarkRenderStateDirty();
-	}
-	// If we didn't have a valid StaticMesh assigned before
-	// our render state might not have been created so
-	// do it now.
-	else if (ShouldCreateRenderState())
-	{
-		RecreateRenderState_Concurrent();
-	}
+	MarkRenderStateDirty();
 
 	// Update physics representation right away
 	RecreatePhysicsState();
diff --git a/Engine/Source/Runtime/Engine/Public/ComponentInstanceDataCache.h b/Engine/Source/Runtime/Engine/Public/ComponentInstanceDataCache.h
--- a/Engine/Source/Runtime/Engine/Public/ComponentInstanceDataCache.h
+++ b/Engine/Source/Runtime/Engine/Public/ComponentInstanceDataCache.h
@@ -158,9 +158,6 @@ public:
 	/** Determines if any instance data was actually saved. */
 	virtual bool ContainsData() const { return SavedProperties.Num() > 0; }
 
-	/** Applies this component instance data to the supplied component */
-	ENGINE_API virtual void ApplyToComponent(UActorComponent* Component, const ECacheApplyPhase CacheApplyPhase);
-
 	/** Replaces any references to old instances during Actor reinstancing */
 	virtual void FindAndReplaceInstances(const TMap<UObject*, UObject*>& OldToNewInstanceMap) { };
 
@@ -170,8 +167,14 @@ public:
 
 	const UObject* GetComponentTemplate() const { return SourceComponentTemplate; }
 
+	/** Applies this component instance data to the supplied component */
+	ENGINE_API virtual void DoApplyToComponent(UActorComponent* Component, const ECacheApplyPhase CacheApplyPhase);
+
 protected:
 
+	/** Applies this component instance data to the supplied component */
+	ENGINE_API virtual void ApplyToComponent(UActorComponent* Component, const ECacheApplyPhase CacheApplyPhase);
+
 	/** The template used to create the source component */
 	UPROPERTY()
 	TObjectPtr<const UObject> SourceComponentTemplate;
