Move the UpdateResource of Textures when the device profiles changes in preview from EditorEngine to TextureLODSettings to fix some edge cases.

diff --git a/Engine/Source/Editor/UnrealEd/Private/EditorEngine.cpp b/Engine/Source/Editor/UnrealEd/Private/EditorEngine.cpp
--- a/Engine/Source/Editor/UnrealEd/Private/EditorEngine.cpp
+++ b/Engine/Source/Editor/UnrealEd/Private/EditorEngine.cpp
@@ -6363,13 +6363,6 @@ void UEditorEngine::SetPreviewPlatform(const FPreviewPlatformInfo& NewPreviewPla
 		{
 			ObjIterator->LoadConfig();
 		}
-
-		//Call UpdateResource on Textures when changing the ShaderPlatformPreview to update the CachePlatformData
-		UTexture* Texture = Cast<UTexture>(*ObjIterator);
-		if (Texture != nullptr)
-		{
-			Texture->UpdateResource();
-		}
 	}
 
 	constexpr bool bUpdateProgressDialog = true;
@@ -6487,13 +6480,6 @@ void UEditorEngine::ToggleFeatureLevelPreview()
 		{
 			ObjIterator->LoadConfig();
 		}
-
-		//Call UpdateResource on Textures when changing the ShaderPlatformPreview to update the CachePlatformData
-		UTexture* Texture = Cast<UTexture>(*ObjIterator);
-		if (Texture != nullptr)
-		{
-			Texture->UpdateResource();
-		}
 	}
 
 	Scalability::ApplyCachedQualityLevelForShaderPlatform(GetActiveShaderPlatform());
diff --git a/Engine/Source/Runtime/Engine/Classes/Engine/TextureLODSettings.h b/Engine/Source/Runtime/Engine/Classes/Engine/TextureLODSettings.h
--- a/Engine/Source/Runtime/Engine/Classes/Engine/TextureLODSettings.h
+++ b/Engine/Source/Runtime/Engine/Classes/Engine/TextureLODSettings.h
@@ -234,10 +234,7 @@ public:
 
 #if ALLOW_OTHER_PLATFORM_CONFIG
 public:
-	ENGINE_API void SetPreviewTextureGroupRedirect(UTextureLODSettings* InPreviewTextureGroupRedirect)
-	{
-		PreviewTextureGroupRedirect = InPreviewTextureGroupRedirect;
-	}
+	ENGINE_API void SetPreviewTextureGroupRedirect(UTextureLODSettings* InPreviewTextureGroupRedirect);
 
 	ENGINE_API UTextureLODSettings* GetPreviewTextureGroupRedirect() const
 	{
diff --git a/Engine/Source/Runtime/Engine/Private/TextureLODSettings.cpp b/Engine/Source/Runtime/Engine/Private/TextureLODSettings.cpp
--- a/Engine/Source/Runtime/Engine/Private/TextureLODSettings.cpp
+++ b/Engine/Source/Runtime/Engine/Private/TextureLODSettings.cpp
@@ -6,6 +6,10 @@
 #include "ImageCoreUtils.h"
 #include "Interfaces/ITargetPlatform.h"
 
+#if ALLOW_OTHER_PLATFORM_CONFIG
+#include "UObject/UObjectIterator.h"
+#endif
+
 #include UE_INLINE_GENERATED_CPP_BY_NAME(TextureLODSettings)
 
 int32 GUITextureLODBias = 0;
@@ -436,3 +440,22 @@ ETextureMipLoadOptions UTextureLODSettings::GetMipLoadOptions(const UTexture* Te
 	}
 }
 
+#if ALLOW_OTHER_PLATFORM_CONFIG
+void UTextureLODSettings::SetPreviewTextureGroupRedirect(UTextureLODSettings* InPreviewTextureGroupRedirect)
+{
+	if (PreviewTextureGroupRedirect != InPreviewTextureGroupRedirect)
+	{
+		PreviewTextureGroupRedirect = InPreviewTextureGroupRedirect;
+
+		for (FThreadSafeObjectIterator ObjIterator(UObject::StaticClass(), RF_NoFlags); ObjIterator; ++ObjIterator)
+		{
+			//Call UpdateResource on Textures when changing the TextureLodGroup Settings to update the CachePlatformData
+			UTexture* Texture = Cast<UTexture>(*ObjIterator);
+			if (Texture != nullptr)
+			{
+				Texture->UpdateResource();
+			}
+		}
+	}
+}
+#endif
\ No newline at end of file
