[World Partition] When dumping streaming generation output logs, also output actor descriptor views data that can differ from the actor descriptor they reflects.

diff --git a/Engine/Source/Runtime/Engine/Private/WorldPartition/WorldPartitionActorDesc.cpp b/Engine/Source/Runtime/Engine/Private/WorldPartition/WorldPartitionActorDesc.cpp
--- a/Engine/Source/Runtime/Engine/Private/WorldPartition/WorldPartitionActorDesc.cpp
+++ b/Engine/Source/Runtime/Engine/Private/WorldPartition/WorldPartitionActorDesc.cpp
@@ -511,7 +511,7 @@ FString FWorldPartitionActorDesc::ToString(EToStringMode Mode) const
 
 	auto AppendFromBool = [&Append, Mode](const TCHAR* Name, bool bValue)
 	{
-		Append(Name, (Mode >= EToStringMode::Compact) ? (bValue ? TEXT("true") : TEXT("false")) : (bValue ? TEXT("1") : TEXT("0")));
+		Append(Name, *LexToString(bValue));
 	};
 
 	if (Mode >= EToStringMode::Compact)
diff --git a/Engine/Source/Runtime/Engine/Private/WorldPartition/WorldPartitionStreamingGeneration.cpp b/Engine/Source/Runtime/Engine/Private/WorldPartition/WorldPartitionStreamingGeneration.cpp
--- a/Engine/Source/Runtime/Engine/Private/WorldPartition/WorldPartitionStreamingGeneration.cpp
+++ b/Engine/Source/Runtime/Engine/Private/WorldPartition/WorldPartitionStreamingGeneration.cpp
@@ -254,6 +254,21 @@ bool FStreamingGenerationActorDescView::IsEditorOnlyReference(const FGuid& Refer
 	return Super::IsEditorOnlyReference(ReferenceGuid) || ForcedInvalidReference.Contains(ReferenceGuid);
 }
 
+FString FStreamingGenerationActorDescView::ToString(FWorldPartitionActorDesc::EToStringMode Mode) const
+{
+	FString Result = Super::ToString(Mode);
+	if (Mode >= FWorldPartitionActorDesc::EToStringMode::Compact)
+	{
+		Result += FString::Printf(TEXT(" View=(RuntimeHLODLayer=%s bIsForcedNonSpatiallyLoaded=%s bIsForcedNoRuntimeGrid=%s bIsForcedNoDataLayers=%s bIsForceNoHLODLayer=%s)"), 
+			RuntimeHLODLayer.IsSet() ? *RuntimeHLODLayer->ToString() : TEXT("none"),
+			*LexToString(bIsForcedNonSpatiallyLoaded),
+			*LexToString(bIsForcedNoRuntimeGrid),
+			*LexToString(bIsForcedNoDataLayers),
+			*LexToString(bIsForceNoHLODLayer));
+	}
+	return Result;
+}
+
 const TArray<FGuid>& FStreamingGenerationActorDescView::GetEditorReferences() const
 {
 	return EditorReferences;
diff --git a/Engine/Source/Runtime/Engine/Public/WorldPartition/WorldPartitionStreamingGeneration.h b/Engine/Source/Runtime/Engine/Public/WorldPartition/WorldPartitionStreamingGeneration.h
--- a/Engine/Source/Runtime/Engine/Public/WorldPartition/WorldPartitionStreamingGeneration.h
+++ b/Engine/Source/Runtime/Engine/Public/WorldPartition/WorldPartitionStreamingGeneration.h
@@ -57,8 +57,9 @@ public:
 	ENGINE_API virtual bool GetIsSpatiallyLoaded() const override;
 	ENGINE_API virtual FSoftObjectPath GetHLODLayer() const override;
 	ENGINE_API virtual const TArray<FGuid>& GetReferences() const override;
-	ENGINE_API const FDataLayerInstanceNames& GetDataLayerInstanceNames() const override;
+	ENGINE_API virtual const FDataLayerInstanceNames& GetDataLayerInstanceNames() const override;
 	ENGINE_API virtual bool IsEditorOnlyReference(const FGuid& ReferenceGuid) const override;
+	ENGINE_API virtual FString ToString(FWorldPartitionActorDesc::EToStringMode Mode) const override;
 	//~ End FWorldPartitionActorDescInstanceView interface
 
 	bool IsUnsaved() const { return bIsUnsaved; }
