[VisualLogger] threading fixes and improvements
- added MT access detector for the list of maps of object entries
- made EditorBaseTimeStamp thread safe
- SVisualLogger now listens to FWorldDelegates::OnPostWorldInitialization instead FEditorDelegates::PostPIEStarted so EditorBaseTimeStamp can be reset properly before other threads start creating new entries for the new world. This fixes an ensure related to mismatching timestamps when merging entries.

diff --git a/Engine/Source/Developer/LogVisualizer/Private/SVisualLogger.cpp b/Engine/Source/Developer/LogVisualizer/Private/SVisualLogger.cpp
--- a/Engine/Source/Developer/LogVisualizer/Private/SVisualLogger.cpp
+++ b/Engine/Source/Developer/LogVisualizer/Private/SVisualLogger.cpp
@@ -111,7 +111,7 @@ SVisualLogger::SVisualLogger()
 SVisualLogger::~SVisualLogger()
 {
 #if WITH_EDITOR
-	FEditorDelegates::PostPIEStarted.Remove(PostPIEStartedHandle);
+	FWorldDelegates::OnPostWorldInitialization.Remove(PostWorldInitializationHandle);
 	FGameDelegates::Get().GetEndPlayMapDelegate().RemoveAll(this);
 #endif
 		
@@ -180,7 +180,7 @@ void SVisualLogger::Construct(const FArguments& InArgs, const TSharedRef<SDockTa
 	ProcessFilterVolumes();
 
 #if WITH_EDITOR
-	PostPIEStartedHandle = FEditorDelegates::PostPIEStarted.AddLambda([this](const bool bIsSimulating)
+	PostWorldInitializationHandle = FWorldDelegates::OnPostWorldInitialization.AddLambda([this](UWorld*, const UWorld::InitializationValues)
 	{
 		UWorld* World = FVisualLoggerEditorInterface::Get()->GetWorld();
 		if (World != nullptr && World != LastUsedWorld)
diff --git a/Engine/Source/Developer/LogVisualizer/Private/SVisualLogger.h b/Engine/Source/Developer/LogVisualizer/Private/SVisualLogger.h
--- a/Engine/Source/Developer/LogVisualizer/Private/SVisualLogger.h
+++ b/Engine/Source/Developer/LogVisualizer/Private/SVisualLogger.h
@@ -139,7 +139,7 @@ protected:
 	bool bGotHistogramData;
 
 #if WITH_EDITOR
-	FDelegateHandle PostPIEStartedHandle;
+	FDelegateHandle PostWorldInitializationHandle;
 #endif
 
 	FDelegateHandle DrawOnCanvasDelegateHandle;
diff --git a/Engine/Source/Runtime/Engine/Private/VisualLogger/VisualLogger.cpp b/Engine/Source/Runtime/Engine/Private/VisualLogger/VisualLogger.cpp
--- a/Engine/Source/Runtime/Engine/Private/VisualLogger/VisualLogger.cpp
+++ b/Engine/Source/Runtime/Engine/Private/VisualLogger/VisualLogger.cpp
@@ -36,7 +36,7 @@ namespace UE::VisLog::Private
 	namespace EditorOnly
 	{
 		/** When in Editor, use this TimeStamp as a Base time so our values don't grow too large (and are canonically around 0.0). */
-		static double EditorBaseTimeStamp = 0.0;
+		static std::atomic<double> EditorBaseTimeStamp = 0.0;
 	}
 #endif
 }
@@ -282,10 +282,9 @@ double FVisualLogger::GetTimeStampForObject(const UObject* Object) const
 		if (WorldForTimeStamp)
 		{
 			using namespace UE::VisLog::Private;
-			if (EditorOnly::EditorBaseTimeStamp <= 0.0)
-			{
-				EditorOnly::EditorBaseTimeStamp = WorldForTimeStamp->TimeSeconds;
-			}
+
+			double Expected = 0;
+			EditorOnly::EditorBaseTimeStamp.compare_exchange_strong(Expected, WorldForTimeStamp->TimeSeconds);
 
 			return WorldForTimeStamp->TimeSeconds - EditorOnly::EditorBaseTimeStamp;
 		}
@@ -340,7 +339,7 @@ void FVisualLogger::AddObjectToAllowList(const UObject& InObject)
 		}
 
 	    // We are not locking the thread entry map, we do not expect other threads adding new entries while editing allow list
-	    // @todo we should add a thread access detector
+	    UE_MT_SCOPED_READ_ACCESS(CurrentEntryMapsAccessDetector);
 		for (FVisualLoggerObjectEntryMap* ThreadCurrentEntryMap : ThreadCurrentEntryMaps)
 		{
 			checkf(ThreadCurrentEntryMap, TEXT("Unexpected null map entry pointer"));
@@ -365,7 +364,7 @@ void FVisualLogger::ClearObjectAllowList()
 		}
 
 	    // We are not locking the thread entry map, we do not expect other threads adding new entries while editing allow list
-	    // @todo we should add a thread access detector
+	    UE_MT_SCOPED_READ_ACCESS(CurrentEntryMapsAccessDetector);
 		for (FVisualLoggerObjectEntryMap* ThreadCurrentEntryMap : ThreadCurrentEntryMaps)
 		{
 			checkf(ThreadCurrentEntryMap, TEXT("Unexpected null map entry pointer"));
@@ -615,7 +614,7 @@ void FVisualLogger::Tick(float DeltaTime)
 void FVisualLogger::FlushThreadsEntries()
 {
 	// We are not locking the thread entry map, we do not expect other threads adding new entries while flushing
-	// @todo we should add a thread access detector
+	UE_MT_SCOPED_READ_ACCESS(CurrentEntryMapsAccessDetector);
 	for (FVisualLoggerObjectEntryMap* ThreadCurrentEntryMap : ThreadCurrentEntryMaps)
 	{
 		checkf(ThreadCurrentEntryMap, TEXT("Unexpected null map entry pointer"));
@@ -667,6 +666,7 @@ FVisualLogger::FVisualLoggerObjectEntryMap& FVisualLogger::GetThreadCurrentEntry
 	{
 		GThreadCurrentEntryPerObject = new FVisualLoggerObjectEntryMap();
 		FWriteScopeLock Lock(EntryRWLock);
+		UE_MT_SCOPED_WRITE_ACCESS(CurrentEntryMapsAccessDetector);
 		ThreadCurrentEntryMaps.Add(GThreadCurrentEntryPerObject);
 	}
 	checkf(GThreadCurrentEntryPerObject, TEXT("Unexpected null map entry pointer"));
@@ -874,6 +874,8 @@ void FVisualLogger::Cleanup(UWorld* OldWorld, const bool bReleaseMemory)
                 ObjectToWorldMap.Reset();
                 ChildToOwnerMap.Reset();
                 CurrentEntryPerObject.Reset();
+
+				UE_MT_SCOPED_WRITE_ACCESS(CurrentEntryMapsAccessDetector);
 				for (FVisualLoggerObjectEntryMap* ThreadCurrentEntryMap : ThreadCurrentEntryMaps)
 				{
 					checkf(ThreadCurrentEntryMap, TEXT("Unexpected null map entry pointer"));
@@ -890,6 +892,8 @@ void FVisualLogger::Cleanup(UWorld* OldWorld, const bool bReleaseMemory)
 						FObjectKey Obj = It.Key();
                         ObjectToWorldMap.Remove(Obj);
                         CurrentEntryPerObject.Remove(Obj);
+
+						UE_MT_SCOPED_WRITE_ACCESS(CurrentEntryMapsAccessDetector);
 						for (FVisualLoggerObjectEntryMap* ThreadCurrentEntryMap : ThreadCurrentEntryMaps)
 						{
 							checkf(ThreadCurrentEntryMap, TEXT("Unexpected null map entry pointer"));
@@ -916,6 +920,8 @@ void FVisualLogger::Cleanup(UWorld* OldWorld, const bool bReleaseMemory)
 		ObjectToWorldMap.Reset();
 		ChildToOwnerMap.Reset();
 		CurrentEntryPerObject.Reset();
+
+		UE_MT_SCOPED_WRITE_ACCESS(CurrentEntryMapsAccessDetector);
 		for (FVisualLoggerObjectEntryMap* ThreadCurrentEntryMap : ThreadCurrentEntryMaps)
 		{
 			checkf(ThreadCurrentEntryMap, TEXT("Unexpected null map entry pointer"));
diff --git a/Engine/Source/Runtime/Engine/Public/VisualLogger/VisualLogger.h b/Engine/Source/Runtime/Engine/Public/VisualLogger/VisualLogger.h
--- a/Engine/Source/Runtime/Engine/Public/VisualLogger/VisualLogger.h
+++ b/Engine/Source/Runtime/Engine/Public/VisualLogger/VisualLogger.h
@@ -901,12 +901,19 @@ protected:
 
 	// Visual Logger extensions map
 	TMap<FName, FVisualLogExtensionInterface*> AllExtensions;
+
 	// last generated unique id for given times tamp
 	TMap<double, int32> LastUniqueIds;
+
 	// Current entry with all data
 	FVisualLoggerObjectEntryMap CurrentEntryPerObject;
+
 	// Threads current entry maps
 	TArray<FVisualLoggerObjectEntryMap*> ThreadCurrentEntryMaps;
+
+	// Multithread access detector for ThreadCurrentEntryMaps operations (adding/removing map to/from the list or in the map itself)
+	UE_MT_DECLARE_RW_ACCESS_DETECTOR(CurrentEntryMapsAccessDetector);
+
 	// Read Write lock protecting object entries
 	mutable FRWLock EntryRWLock;
 	struct FObjectNames
