[HWRT] Add CSV stats to FinishGatherInstances(...) and FRayTracingGeometryManager::Tick(...)

diff --git a/Engine/Source/Runtime/Engine/Private/Rendering/RayTracingGeometryManager.cpp b/Engine/Source/Runtime/Engine/Private/Rendering/RayTracingGeometryManager.cpp
--- a/Engine/Source/Runtime/Engine/Private/Rendering/RayTracingGeometryManager.cpp
+++ b/Engine/Source/Runtime/Engine/Private/Rendering/RayTracingGeometryManager.cpp
@@ -502,6 +502,7 @@ void FRayTracingGeometryManager::Tick(FRHICommandList& RHICmdList)
 
 	TRACE_CPUPROFILER_EVENT_SCOPE(FRayTracingGeometryManager::Tick);
 	QUICK_SCOPE_CYCLE_COUNTER(STAT_FRayTracingGeometryManager_Tick);
+	CSV_SCOPED_TIMING_STAT_EXCLUSIVE(RayTracingGeometryManager_Tick);
 
 	// TODO: investigate fine grained locking to minimize blocking progress on render command pipes
 	// - Don't touch registered geometry/group arrays from render command pipes
diff --git a/Engine/Source/Runtime/Renderer/Private/RayTracing/RayTracing.cpp b/Engine/Source/Runtime/Renderer/Private/RayTracing/RayTracing.cpp
--- a/Engine/Source/Runtime/Renderer/Private/RayTracing/RayTracing.cpp
+++ b/Engine/Source/Runtime/Renderer/Private/RayTracing/RayTracing.cpp
@@ -1698,6 +1698,7 @@ namespace RayTracing
 	{
 		TRACE_CPUPROFILER_EVENT_SCOPE(RayTracing_FinishGatherInstances);
 		SCOPE_CYCLE_COUNTER(STAT_RayTracing_FinishGatherInstances);
+		RDG_CSV_STAT_EXCLUSIVE_SCOPE(GraphBuilder, RayTracing_FinishGatherInstances);
 
 		TaskData.GatherRelevantPrimitivesTask.Wait();
 
