FirstPerson: Added IsFirstPerson material graph node.

diff --git a/Engine/Shaders/Private/MaterialTemplate.ush b/Engine/Shaders/Private/MaterialTemplate.ush
--- a/Engine/Shaders/Private/MaterialTemplate.ush
+++ b/Engine/Shaders/Private/MaterialTemplate.ush
@@ -3770,6 +3770,21 @@ half GetMaterialMaskedOpacity(FPixelMaterialInputs PixelMaterialInputs)
 }
 #endif
 
+bool IsFirstPerson_FromFlags(uint Flags)
+{
+	return (Flags & PRIMITIVE_SCENE_DATA_FLAG_IS_FIRST_PERSON_PRIMITIVE) != 0;
+}
+
+bool IsFirstPerson(FMaterialVertexParameters Parameters)
+{
+	return IsFirstPerson_FromFlags(GetPrimitiveData(Parameters).Flags);
+}
+
+bool IsFirstPerson(FMaterialPixelParameters Parameters)
+{
+	return IsFirstPerson_FromFlags(GetPrimitiveData(Parameters).Flags);
+}
+
 #if SUPPORT_FIRST_PERSON_RENDERING
 
 float GetMaterialFirstPersonInterpolationAlpha(FMaterialVertexParameters Parameters)
@@ -3837,7 +3852,7 @@ float3 GetMaterialWorldPositionOffset(FMaterialVertexParameters Parameters)
 	}
 #if SUPPORT_FIRST_PERSON_RENDERING
 	BRANCH
-	if ((GetPrimitiveData(Parameters).Flags & PRIMITIVE_SCENE_DATA_FLAG_IS_FIRST_PERSON_PRIMITIVE) != 0)
+	if (IsFirstPerson(Parameters))
 	{
 		const float3 FirstPersonPosition = mul(Parameters.WorldPosition + WPO, (float3x3)ResolvedView.FirstPersonTransform);
 		const float3 FirstPersonWPO = FirstPersonPosition - Parameters.WorldPosition;
@@ -3866,7 +3881,7 @@ float3 GetMaterialPreviousWorldPositionOffset(FMaterialVertexParameters Paramete
 	}
 #if SUPPORT_FIRST_PERSON_RENDERING
 	BRANCH
-	if ((GetPrimitiveData(Parameters).Flags & PRIMITIVE_SCENE_DATA_FLAG_IS_FIRST_PERSON_PRIMITIVE) != 0)
+	if (IsFirstPerson(Parameters))
 	{
 		const float3 FirstPersonPosition = mul(Parameters.WorldPosition + WPO, (float3x3)ResolvedView.PrevFirstPersonTransform);
 		const float3 FirstPersonWPO = FirstPersonPosition - Parameters.WorldPosition;
diff --git a/Engine/Source/Runtime/Engine/Private/Materials/HLSLMaterialTranslator.cpp b/Engine/Source/Runtime/Engine/Private/Materials/HLSLMaterialTranslator.cpp
--- a/Engine/Source/Runtime/Engine/Private/Materials/HLSLMaterialTranslator.cpp
+++ b/Engine/Source/Runtime/Engine/Private/Materials/HLSLMaterialTranslator.cpp
@@ -11551,6 +11551,16 @@ int32 FHLSLMaterialTranslator::PixelNormalWS()
 	return AddInlinedCodeChunk(MCT_Float3,TEXT("Parameters.WorldNormal"));	
 }
 
+int32 FHLSLMaterialTranslator::IsFirstPerson()
+{
+	if (!CheckPrimitivePropertyCompatibity(ANSI_TO_TCHAR(__FUNCTION__)))
+	{
+		return INDEX_NONE;
+	}
+
+	return AddInlinedCodeChunkZeroDeriv(MCT_Float, TEXT("(IsFirstPerson(Parameters) ? 1.0f : 0.0f)"));
+}
+
 int32 FHLSLMaterialTranslator::Derivative(int32 A, EDervativeComponent Component)
 {
 	if (A == INDEX_NONE)
diff --git a/Engine/Source/Runtime/Engine/Private/Materials/HLSLMaterialTranslator.h b/Engine/Source/Runtime/Engine/Private/Materials/HLSLMaterialTranslator.h
--- a/Engine/Source/Runtime/Engine/Private/Materials/HLSLMaterialTranslator.h
+++ b/Engine/Source/Runtime/Engine/Private/Materials/HLSLMaterialTranslator.h
@@ -1193,6 +1193,7 @@ protected:
 	virtual int32 VertexNormal() override;
 	virtual int32 VertexTangent() override;
 	virtual int32 PixelNormalWS() override;
+	virtual int32 IsFirstPerson() override;
 
 	virtual int32 DDX(int32 A) override;
 	virtual int32 DDY(int32 A) override;
diff --git a/Engine/Source/Runtime/Engine/Private/Materials/MaterialExpressionHLSL.cpp b/Engine/Source/Runtime/Engine/Private/Materials/MaterialExpressionHLSL.cpp
--- a/Engine/Source/Runtime/Engine/Private/Materials/MaterialExpressionHLSL.cpp
+++ b/Engine/Source/Runtime/Engine/Private/Materials/MaterialExpressionHLSL.cpp
@@ -243,6 +243,7 @@
 #include "Materials/MaterialExpressionWhileLoop.h"
 #include "Materials/MaterialExpressionWorldPosition.h"
 #include "Materials/MaterialExpressionFirstPersonOutput.h"
+#include "Materials/MaterialExpressionIsFirstPerson.h"
 #include "Materials/MaterialFunctionInterface.h"
 #include "Materials/MaterialParameterCollection.h"
 #include "MaterialShared.h"
@@ -5329,4 +5330,11 @@ UE::Shader::EValueType UMaterialExpressionFirstPersonOutput::GetCustomOutputType
 	return EValueType::Void;
 }
 
+bool UMaterialExpressionIsFirstPerson::GenerateHLSLExpression(FMaterialHLSLGenerator& Generator, UE::HLSLTree::FScope& Scope, int32 OutputIndex, UE::HLSLTree::FExpression const*& OutExpression) const
+{
+	using namespace UE::HLSLTree::Material;
+	OutExpression = Generator.GetTree().NewExpression<FExpressionExternalInput>(EExternalInput::IsFirstPerson);
+	return true;
+}
+
 #endif // WITH_EDITOR
diff --git a/Engine/Source/Runtime/Engine/Private/Materials/MaterialExpressions.cpp b/Engine/Source/Runtime/Engine/Private/Materials/MaterialExpressions.cpp
--- a/Engine/Source/Runtime/Engine/Private/Materials/MaterialExpressions.cpp
+++ b/Engine/Source/Runtime/Engine/Private/Materials/MaterialExpressions.cpp
@@ -299,6 +299,7 @@
 #include "Materials/MaterialExpressionSparseVolumeTextureBase.h"
 #include "Materials/MaterialExpressionSparseVolumeTextureObject.h"
 #include "Materials/MaterialExpressionSparseVolumeTextureSample.h"
+#include "Materials/MaterialExpressionIsFirstPerson.h"
 #include "EditorSupportDelegates.h"
 #if WITH_EDITOR
 #include "MaterialGraph/MaterialGraphNode_Comment.h"
@@ -28129,4 +28130,40 @@ FString UMaterialExpressionFirstPersonOutput::GetDisplayName() const
 	return TEXT("First Person Output");
 }
 
+UMaterialExpressionIsFirstPerson::UMaterialExpressionIsFirstPerson(const FObjectInitializer& ObjectInitializer)
+	: Super(ObjectInitializer)
+{
+#if WITH_EDITORONLY_DATA
+	// Structure to hold one-time initialization
+	struct FConstructorStatics
+	{
+		FText NAME_Constants;
+		FText NAME_FirstPerson;
+		FConstructorStatics()
+			: NAME_Constants(LOCTEXT("Constants", "Constants")),
+			NAME_FirstPerson(LOCTEXT("FirstPerson", "First Person"))
+		{
+		}
+	};
+	static FConstructorStatics ConstructorStatics;
+
+	MenuCategories.Add(ConstructorStatics.NAME_Constants);
+	MenuCategories.Add(ConstructorStatics.NAME_FirstPerson);
+
+	bShaderInputData = true;
+#endif
+}
+
+#if WITH_EDITOR
+int32 UMaterialExpressionIsFirstPerson::Compile(class FMaterialCompiler* Compiler, int32 OutputIndex)
+{
+	return Compiler->IsFirstPerson();
+}
+
+void UMaterialExpressionIsFirstPerson::GetCaption(TArray<FString>& OutCaptions) const
+{
+	OutCaptions.Add(TEXT("IsFirstPerson"));
+}
+#endif // WITH_EDITOR
+
 #undef LOCTEXT_NAMESPACE
diff --git a/Engine/Source/Runtime/Engine/Private/Materials/MaterialHLSLTree.cpp b/Engine/Source/Runtime/Engine/Private/Materials/MaterialHLSLTree.cpp
--- a/Engine/Source/Runtime/Engine/Private/Materials/MaterialHLSLTree.cpp
+++ b/Engine/Source/Runtime/Engine/Private/Materials/MaterialHLSLTree.cpp
@@ -183,6 +183,8 @@ FExternalInputDescription GetExternalInputDescription(EExternalInput Input)
 	case EExternalInput::DistanceCullFade: return FExternalInputDescription(TEXT("DistanceCullFade"), Shader::EValueType::Float1);
 
 	case EExternalInput::IsOrthographic: return FExternalInputDescription(TEXT("IsOrthographic"), Shader::EValueType::Float1);
+	
+	case EExternalInput::IsFirstPerson: return FExternalInputDescription(TEXT("IsFirstPerson"), Shader::EValueType::Float1);
 
 	case EExternalInput::AOMask: return FExternalInputDescription(TEXT("PrecomputedAOMask"), Shader::EValueType::Float1);
 
@@ -506,6 +508,8 @@ void FExpressionExternalInput::EmitValueShader(FEmitContext& Context, FEmitScope
 		case EExternalInput::DistanceCullFade: Code = TEXT("GetDistanceCullFade()"); break;
 
 		case EExternalInput::IsOrthographic: Code = TEXT("((View.ViewToClip[3][3] < 1.0f) ? 0.0f : 1.0f)"); break;
+		
+		case EExternalInput::IsFirstPerson: Code = TEXT("(IsFirstPerson(Parameters) ? 1.0f : 0.0f)"); break;
 
 		case EExternalInput::AOMask: Code = TEXT("Parameters.AOMaterialMask"); break;
 
diff --git a/Engine/Source/Runtime/Engine/Public/MaterialCompiler.h b/Engine/Source/Runtime/Engine/Public/MaterialCompiler.h
--- a/Engine/Source/Runtime/Engine/Public/MaterialCompiler.h
+++ b/Engine/Source/Runtime/Engine/Public/MaterialCompiler.h
@@ -527,6 +527,7 @@ public:
 	virtual int32 VertexNormal() = 0;
 	virtual int32 VertexTangent() = 0;
 	virtual int32 PixelNormalWS() = 0;
+	virtual int32 IsFirstPerson() = 0;
 
 	virtual int32 CustomExpression(class UMaterialExpressionCustom* Custom, int32 OutputIndex, TArray<int32>& CompiledInputs) = 0;
 	virtual int32 CustomOutput(class UMaterialExpressionCustomOutput* Custom, int32 OutputIndex, int32 OutputCode) = 0;
@@ -1097,6 +1098,7 @@ public:
 	virtual int32 VertexNormal() override { return Compiler->VertexNormal(); }
 	virtual int32 VertexTangent() override { return Compiler->VertexTangent(); }
 	virtual int32 PixelNormalWS() override { return Compiler->PixelNormalWS(); }
+	virtual int32 IsFirstPerson() override { return Compiler->IsFirstPerson(); }
 
 	virtual int32 CustomExpression(class UMaterialExpressionCustom* Custom, int32 OutputIndex, TArray<int32>& CompiledInputs) override { return Compiler->CustomExpression(Custom, OutputIndex, CompiledInputs); }
 	virtual int32 CustomOutput(class UMaterialExpressionCustomOutput* Custom, int32 OutputIndex, int32 OutputCode) override { return Compiler->CustomOutput(Custom, OutputIndex, OutputCode); }
diff --git a/Engine/Source/Runtime/Engine/Public/MaterialHLSLTree.h b/Engine/Source/Runtime/Engine/Public/MaterialHLSLTree.h
--- a/Engine/Source/Runtime/Engine/Public/MaterialHLSLTree.h
+++ b/Engine/Source/Runtime/Engine/Public/MaterialHLSLTree.h
@@ -217,6 +217,8 @@ enum class EExternalInput : uint8
 	DistanceCullFade,
 
 	IsOrthographic,
+	
+	IsFirstPerson,
 
 	AOMask,
 
diff --git a/Engine/Source/Runtime/Engine/Public/Materials/MaterialExpressionIsFirstPerson.h b/Engine/Source/Runtime/Engine/Public/Materials/MaterialExpressionIsFirstPerson.h
new file mode 100644
--- /dev/null
+++ b/Engine/Source/Runtime/Engine/Public/Materials/MaterialExpressionIsFirstPerson.h
@@ -0,0 +1,23 @@
+// Copyright Epic Games, Inc. All Rights Reserved.
+
+
+#pragma once
+
+#include "CoreMinimal.h"
+#include "UObject/ObjectMacros.h"
+#include "Materials/MaterialExpression.h"
+#include "MaterialExpressionIsFirstPerson.generated.h"
+
+UCLASS(collapsecategories, hidecategories = Object, MinimalAPI)
+class UMaterialExpressionIsFirstPerson : public UMaterialExpression
+{
+	GENERATED_UCLASS_BODY()
+
+	//~ Begin UMaterialExpression Interface
+#if WITH_EDITOR
+	virtual int32 Compile(class FMaterialCompiler* Compiler, int32 OutputIndex) override;
+	virtual void GetCaption(TArray<FString>& OutCaptions) const override;
+	virtual bool GenerateHLSLExpression(FMaterialHLSLGenerator& Generator, UE::HLSLTree::FScope& Scope, int32 OutputIndex, UE::HLSLTree::FExpression const*& OutExpression) const override;
+#endif
+	//~ End UMaterialExpression Interface
+};
