Fixed FCollisionResponse::operator== so it doesn't return true if both containers have the same number of entries but different contents

This fixes a bug where saving an asset or placing an instance in a level could cause its components to have incorrect collision responses

#tests Saved and reloaded an asset that was affected by this bug and its components maintained the correct collision responses. Placed an instance of the asset while in game and collision was correct there as well
#rnx

diff --git a/Engine/Source/Runtime/Engine/Private/PhysicsEngine/BodyInstance.cpp b/Engine/Source/Runtime/Engine/Private/PhysicsEngine/BodyInstance.cpp
--- a/Engine/Source/Runtime/Engine/Private/PhysicsEngine/BodyInstance.cpp
+++ b/Engine/Source/Runtime/Engine/Private/PhysicsEngine/BodyInstance.cpp
@@ -310,15 +310,23 @@ bool FCollisionResponse::operator==(const FCollisionResponse& Other) const
 	{
 		for(int32 ResponseIdx = 0; ResponseIdx < ResponseArray.Num(); ++ResponseIdx)
 		{
+			bool bFound = false;
 			for(int32 InternalIdx = 0; InternalIdx < ResponseArray.Num(); ++InternalIdx)
 			{
 				if(ResponseArray[ResponseIdx].Channel == Other.ResponseArray[InternalIdx].Channel)
 				{
-					bCollisionResponseEqual &= ResponseArray[ResponseIdx] == Other.ResponseArray[InternalIdx];
+					if (ResponseArray[ResponseIdx].Response != Other.ResponseArray[InternalIdx].Response)
+					{
+						return false;
+					}
+					bFound = true;
 					break;
 				}
 			}
-			
+			if(!bFound)
+			{
+				return false;
+			}
 		}
 	}
 
