Make BodyInstance set a decent debug name when there is no Component (for Chaos Visual Debugger etc and Scene Graph)

diff --git a/Engine/Source/Runtime/Engine/Private/PhysicsEngine/BodyInstance.cpp b/Engine/Source/Runtime/Engine/Private/PhysicsEngine/BodyInstance.cpp
--- a/Engine/Source/Runtime/Engine/Private/PhysicsEngine/BodyInstance.cpp
+++ b/Engine/Source/Runtime/Engine/Private/PhysicsEngine/BodyInstance.cpp
@@ -1120,7 +1120,7 @@ void FBodyInstance::UpdatePhysicsFilterData()
 TAutoConsoleVariable<int32> CDisableQueryOnlyActors(TEXT("p.DisableQueryOnlyActors"), 0, TEXT("If QueryOnly is used, actors are marked as simulation disabled. This is NOT compatible with origin shifting at the moment."));
 
 #if USE_BODYINSTANCE_DEBUG_NAMES
-TSharedPtr<TArray<ANSICHAR>> GetDebugDebugName(const UPrimitiveComponent* PrimitiveComp, const UBodySetup* BodySetup, FString& DebugName)
+TSharedPtr<TArray<ANSICHAR>> GetDebugDebugName(const UPrimitiveComponent* PrimitiveComp, const UObject* SourceObject, const UBodySetup* BodySetup, FString& DebugName)
 {
 	// Setup names
 	// Make the debug name for this geometry...
@@ -1131,6 +1131,14 @@ TSharedPtr<TArray<ANSICHAR>> GetDebugDebugName(const UPrimitiveComponent* Primit
 	{
 		DebugName += FString::Printf(TEXT("%s %s "), *AActor::GetDebugName(PrimitiveComp->GetOwner()), *PrimitiveComp->GetPathName());
 	}
+	else if (SourceObject && SourceObject->GetOuter())
+	{
+		DebugName += FString::Printf(TEXT("%s %s "), *SourceObject->GetOuter()->GetName(), *SourceObject->GetName());
+	}
+	else if (SourceObject)
+	{
+		DebugName += FString::Printf(TEXT("%s "), *SourceObject->GetName());
+	}
 
 	if (BodySetup->BoneName != NAME_None)
 	{
@@ -1236,7 +1244,7 @@ TInitBodiesHelperBase<InAllocatorType>::TInitBodiesHelperBase(TArray<FBodyInstan
 	, DisableQueryOnlyActors(!!CDisableQueryOnlyActors.GetValueOnAnyThread())
 {
 #if USE_BODYINSTANCE_DEBUG_NAMES
-	PhysXName = GetDebugDebugName(PrimitiveComp, BodySetup, *DebugName);
+	PhysXName = GetDebugDebugName(PrimitiveComp, SourceObject, BodySetup, *DebugName);
 #endif
 }
 
