[WorldPartitionActorDesc] Add support for actors using transforms relative to their parent
* Change serialized actor transform to be relative for actors with parent actors, actor-to-world transform otherwise
* Add new transient member - ActorTransform, storing current actor-to-world transform, considering parent transforms
* Update transient actor-to-world ActorTransform whenever the actor is moved or whenever there are any changes to its parent hierarchy (new parent, parent moved, etc.)
* Change serialized runtime and editor bounds to be relative to actor-to-world ActorTransform
* Store relative bounds as OBB instead of AABB
* Keep runtime and editor bounds as transient members and update them whenever actor-to-world ActorTransform changes
* This approach helps avoid having to resave all child ActorDescs whenever parent actor's transform changes

diff --git a/Engine/Source/Runtime/Core/Public/UObject/FortniteMainBranchObjectVersion.h b/Engine/Source/Runtime/Core/Public/UObject/FortniteMainBranchObjectVersion.h
--- a/Engine/Source/Runtime/Core/Public/UObject/FortniteMainBranchObjectVersion.h
+++ b/Engine/Source/Runtime/Core/Public/UObject/FortniteMainBranchObjectVersion.h
@@ -543,6 +543,9 @@ struct FFortniteMainBranchObjectVersion
 
 		// Landscape texture patches in Texture Asset source mode now use proper resolution when calculating transform
 		LandscapeTexturePatchUsesTextureAssetResolution,
+		
+		// Added support for relative transform in WorldPartitionActorDesc
+		WorldPartitionActorDescSerializeRelativeTransform,
 
 		// -----<new versions can be added above this line>-------------------------------------------------
 		VersionPlusOne,
diff --git a/Engine/Source/Runtime/Engine/Private/WorldPartition/ActorDescContainer.cpp b/Engine/Source/Runtime/Engine/Private/WorldPartition/ActorDescContainer.cpp
--- a/Engine/Source/Runtime/Engine/Private/WorldPartition/ActorDescContainer.cpp
+++ b/Engine/Source/Runtime/Engine/Private/WorldPartition/ActorDescContainer.cpp
@@ -280,10 +280,24 @@ void UActorDescContainer::Initialize(const FInitializeParams& InitParams)
 		TRACE_CPUPROFILER_EVENT_SCOPE(RegisterDescriptors);
 		for (auto& [ActorGuid, ActorDesc] : ValidActorDescs)
 		{
+			AddChildActorToParentMap(ActorDesc.Get());
 			RegisterActorDescriptor(ActorDesc.Release());
 		}
 	}
 
+	{
+		TRACE_CPUPROFILER_EVENT_SCOPE(UpdateActorToWorld);
+		for (FActorDescList::TIterator<> ActorDescIterator(this); ActorDescIterator; ++ActorDescIterator)
+		{
+			// Update ActorToWorld for all actors, starting from "root" actors
+			if (ActorDescIterator->GetParentActor().IsValid() == false)
+			{
+				ActorDescIterator->UpdateActorToWorld();
+				PropagateActorToWorldUpdate(*ActorDescIterator);
+			}
+		}
+	}
+
 	{
 		TRACE_CPUPROFILER_EVENT_SCOPE(OnActorDescContainerInitialized)
 		OnActorDescContainerInitialized.Broadcast(this);
@@ -487,9 +501,13 @@ void UActorDescContainer::OnObjectPreSave(UObject* Object, FObjectPreSaveContext
 				if (TUniquePtr<FWorldPartitionActorDesc>* ExistingActorDesc = GetActorDescriptor(Actor->GetActorGuid()))
 				{
 					// Existing actor
+					RemoveChildActorFromParentMap(ExistingActorDesc->Get());
 					OnActorDescUpdating(ExistingActorDesc->Get());
 					FWorldPartitionActorDescUtils::UpdateActorDescriptorFromActor(Actor, *ExistingActorDesc);
 					OnActorDescUpdated(ExistingActorDesc->Get());
+					AddChildActorToParentMap(ExistingActorDesc->Get());
+
+					PropagateActorToWorldUpdate(ExistingActorDesc->Get());
 				}
 				else
 				{
@@ -497,6 +515,9 @@ void UActorDescContainer::OnObjectPreSave(UObject* Object, FObjectPreSaveContext
 					FWorldPartitionActorDesc* AddedActorDesc = Actor->CreateActorDesc().Release();
 					RegisterActorDescriptor(AddedActorDesc);
 					OnActorDescAdded(AddedActorDesc);
+					AddChildActorToParentMap(AddedActorDesc);
+
+					PropagateActorToWorldUpdate(AddedActorDesc);
 				}
 			}
 		}
@@ -581,6 +602,7 @@ bool UActorDescContainer::RemoveActor(const FGuid& ActorGuid)
 {
 	if (TUniquePtr<FWorldPartitionActorDesc>* ExistingActorDesc = GetActorDescriptor(ActorGuid))
 	{
+		RemoveChildActorFromParentMap(ExistingActorDesc->Get());
 		OnActorDescRemoved(ExistingActorDesc->Get());
 		UnregisterActorDescriptor(ExistingActorDesc->Get());
 		ExistingActorDesc->Reset();
@@ -636,5 +658,48 @@ void UActorDescContainer::OnActorDescUpdated(FWorldPartitionActorDesc* ActorDesc
 {
 	OnActorDescUpdatedEvent.Broadcast(ActorDesc);
 }
+
+void UActorDescContainer::AddChildActorToParentMap(FWorldPartitionActorDesc* ActorDesc)
+{
+	FGuid ParentActorGuid = ActorDesc->GetParentActor();
+	if (ParentActorGuid.IsValid())
+	{
+		ParentActorToChildrenMap.FindOrAdd(ParentActorGuid).Add(ActorDesc->GetGuid());
+	}
+}
+
+void UActorDescContainer::RemoveChildActorFromParentMap(FWorldPartitionActorDesc* ActorDesc)
+{
+	FGuid ParentActorGuid = ActorDesc->GetParentActor();
+	if (ParentActorGuid.IsValid())
+	{
+		if (TSet<FGuid>* ChildActors = ParentActorToChildrenMap.Find(ParentActorGuid))
+		{
+			ChildActors->Remove(ActorDesc->GetGuid());
+			if (ChildActors->IsEmpty())
+			{
+				ParentActorToChildrenMap.Remove(ParentActorGuid);
+			}
+		}
+	}
+}
+
+void UActorDescContainer::PropagateActorToWorldUpdate(FWorldPartitionActorDesc* ActorDesc)
+{
+	if (const TSet<FGuid>* Children = ParentActorToChildrenMap.Find(ActorDesc->GetGuid()))
+	{
+		for (const FGuid& ChildGuid : *Children)
+		{
+			if (FWorldPartitionActorDesc* ChildActor = GetActorDesc(ChildGuid))
+			{
+				OnActorDescUpdating(ChildActor);
+				ChildActor->UpdateActorToWorld();
+				OnActorDescUpdated(ChildActor);
+
+				PropagateActorToWorldUpdate(ChildActor);
+			}
+		}
+	}
+}
 #endif
 
diff --git a/Engine/Source/Runtime/Engine/Private/WorldPartition/WorldPartitionActorDesc.cpp b/Engine/Source/Runtime/Engine/Private/WorldPartition/WorldPartitionActorDesc.cpp
--- a/Engine/Source/Runtime/Engine/Private/WorldPartition/WorldPartitionActorDesc.cpp
+++ b/Engine/Source/Runtime/Engine/Private/WorldPartition/WorldPartitionActorDesc.cpp
@@ -44,7 +44,9 @@ static FGuid GetDefaultActorDescGuid(const FWorldPartitionActorDesc* ActorDesc)
 }
 
 FWorldPartitionActorDesc::FWorldPartitionActorDesc()
-	: RuntimeBounds(ForceInit)
+	: ActorTransformRelative(FTransform::Identity)
+	, RuntimeBoundsRelative(ForceInit)
+	, EditorBoundsRelative(ForceInit)
 	, bIsSpatiallyLoaded(false)
 	, bActorIsEditorOnly(false)
 	, bActorIsRuntimeOnly(false)
@@ -54,7 +56,11 @@ FWorldPartitionActorDesc::FWorldPartitionActorDesc()
 	, bIsUsingDataLayerAsset(false)
 	, ActorNativeClass(nullptr)
 	, Container(nullptr)
+	, ActorTransform(FTransform::Identity)
+	, RuntimeBounds(ForceInit)
+	, EditorBounds(ForceInit)
 	, bIsDefaultActorDesc(false)
+	, bHasValidRelativeBounds(false)
 {}
 
 void FWorldPartitionActorDesc::Init(const AActor* InActor)
@@ -88,15 +94,21 @@ void FWorldPartitionActorDesc::Init(const AActor* InActor)
 
 	check(Guid.IsValid());
 
-	ActorTransform = InActor->GetActorTransform();
-
 	RuntimeBounds.Init();
 	EditorBounds.Init();
 	
 	if (!bIsDefaultActorDesc)
 	{
+		ActorTransform = InActor->GetActorTransform();
+		ActorTransformRelative = InActor->GetRootComponent() ? InActor->GetRootComponent()->GetRelativeTransform() : ActorTransform;
+
 		InActor->GetStreamingBounds(RuntimeBounds, EditorBounds);
 		FixupStreamingBounds();
+
+		RuntimeBoundsRelative = FWorldPartitionRelativeBounds(RuntimeBounds).InverseTransformBy(ActorTransform);
+		EditorBoundsRelative = FWorldPartitionRelativeBounds(EditorBounds).InverseTransformBy(ActorTransform);
+
+		bHasValidRelativeBounds = true;
 	}
 
 	RuntimeGrid = InActor->GetRuntimeGrid();
@@ -339,9 +351,9 @@ bool FWorldPartitionActorDesc::Equals(const FWorldPartitionActorDesc* Other) con
 		ActorPackage == Other->ActorPackage &&
 		ActorPath == Other->ActorPath &&
 		ActorLabel == Other->ActorLabel &&
-		ActorTransform.Equals(Other->ActorTransform, 0.1f) &&
-		RuntimeBounds.Equals(Other->RuntimeBounds, 0.1f) &&
-		EditorBounds.Equals(Other->EditorBounds, 0.1f) &&
+		ActorTransformRelative.Equals(Other->ActorTransformRelative, 0.1f) &&
+		RuntimeBoundsRelative.Equals(Other->RuntimeBoundsRelative, 0.1f) &&
+		EditorBoundsRelative.Equals(Other->EditorBoundsRelative, 0.1f) &&
 		RuntimeGrid == Other->RuntimeGrid &&
 		bIsSpatiallyLoaded == Other->bIsSpatiallyLoaded &&
 		bActorIsEditorOnly == Other->bActorIsEditorOnly &&
@@ -374,8 +386,8 @@ bool FWorldPartitionActorDesc::ShouldResave(const FWorldPartitionActorDesc* Othe
 		bActorIsEditorOnly != Other->bActorIsEditorOnly ||
 		bActorIsRuntimeOnly != Other->bActorIsRuntimeOnly ||
 		bActorIsMainWorldOnly != Other->bActorIsMainWorldOnly ||
-		RuntimeBounds.IsValid != Other->RuntimeBounds.IsValid||
-		EditorBounds.IsValid != Other->EditorBounds.IsValid||
+		RuntimeBoundsRelative.IsValid() != Other->RuntimeBoundsRelative.IsValid() ||
+		EditorBoundsRelative.IsValid() != Other->EditorBoundsRelative.IsValid() ||
 		HLODLayer != Other->HLODLayer ||
 		ParentActor != Other->ParentActor ||
 		ContentBundleGuid != Other->ContentBundleGuid ||
@@ -389,10 +401,10 @@ bool FWorldPartitionActorDesc::ShouldResave(const FWorldPartitionActorDesc* Othe
 	}
 
 	// Tolerate up to 5% for bounds change
-	if (RuntimeBounds.IsValid)
+	if (RuntimeBoundsRelative.IsValid())
 	{
-		const FBox ThisBounds = GetRuntimeBounds();
-		const FBox OtherBounds = Other->GetRuntimeBounds();
+		const FBox ThisBounds = RuntimeBoundsRelative.ToAABB();
+		const FBox OtherBounds = Other->RuntimeBoundsRelative.ToAABB();
 		const FVector BoundsChangeTolerance = ThisBounds.GetSize() * 0.05f;
 		const FVector MinDiff = FVector(OtherBounds.Min - ThisBounds.Min).GetAbs();
 		const FVector MaxDiff = FVector(OtherBounds.Max - ThisBounds.Max).GetAbs();
@@ -506,6 +518,15 @@ FString FWorldPartitionActorDesc::ToString(EToStringMode Mode) const
 			}
 		}
 
+		if constexpr (std::is_same_v<Type, FWorldPartitionRelativeBounds>)
+		{
+			if (!Value.IsValid())
+			{
+				Append(Name, TEXT("IsValid=false"));
+				return;
+			}
+		}
+
 		AppendFromString(Name, Value.ToString());
 	};
 
@@ -533,7 +554,9 @@ FString FWorldPartitionActorDesc::ToString(EToStringMode Mode) const
 		AppendToString(TEXT("Label"), GetActorLabel());
 		AppendFromBool(TEXT("SpatiallyLoaded"), bIsSpatiallyLoaded);
 		AppendToString(TEXT("EditorBounds"), EditorBounds);
+		AppendToString(TEXT("EditorBoundsRelative"), EditorBoundsRelative);
 		AppendToString(TEXT("RuntimeBounds"), RuntimeBounds);
+		AppendToString(TEXT("RuntimeBoundsRelative"), RuntimeBoundsRelative);
 		AppendToString(TEXT("RuntimeGrid"), RuntimeGrid);
 		AppendFromBool(TEXT("EditorOnly"), bActorIsEditorOnly);
 		AppendFromBool(TEXT("RuntimeOnly"), bActorIsRuntimeOnly);
@@ -626,7 +649,14 @@ void FWorldPartitionActorDesc::Serialize(FArchive& Ar)
 
 		if (Ar.CustomVer(FFortniteSeasonBranchObjectVersion::GUID) >= FFortniteSeasonBranchObjectVersion::WorldPartitionActorDescActorTransformSerialization)
 		{
-			Ar << ActorTransform;
+			if (Ar.CustomVer(FFortniteMainBranchObjectVersion::GUID) < FFortniteMainBranchObjectVersion::WorldPartitionActorDescSerializeRelativeTransform)
+			{
+				Ar << ActorTransform;
+			}
+			else
+			{
+				Ar << ActorTransformRelative;
+			}
 		}
 	}
 
@@ -663,10 +693,15 @@ void FWorldPartitionActorDesc::Serialize(FArchive& Ar)
 				EditorBounds.Init();
 			}
 		}
-		else
+		else if (Ar.CustomVer(FFortniteMainBranchObjectVersion::GUID) < FFortniteMainBranchObjectVersion::WorldPartitionActorDescSerializeRelativeTransform)
 		{
 			Ar << RuntimeBounds << EditorBounds;
 		}
+		else
+		{
+			Ar << RuntimeBoundsRelative << EditorBoundsRelative;
+			bHasValidRelativeBounds = true;
+		}
 	}
 	else
 	{
@@ -943,6 +978,33 @@ bool FWorldPartitionActorDesc::IsRuntimeRelevant(const FWorldPartitionActorDescI
 	return !IsMainWorldOnly() || InActorDescInstance->GetContainerInstance()->GetContainerID().IsMainContainer();
 }
 
+void FWorldPartitionActorDesc::UpdateActorToWorld()
+{
+	FTransform ParentTransform = FTransform::Identity;
+	if (ParentActor.IsValid())
+	{
+		FWorldPartitionActorDesc& ParentActorDesc = Container->GetActorDescChecked(ParentActor);
+		ParentTransform = ParentActorDesc.GetActorTransform();
+	}
+
+	if (bHasValidRelativeBounds == false)
+	{
+		ActorTransformRelative = ActorTransform * ParentTransform.Inverse();
+
+		RuntimeBoundsRelative = FWorldPartitionRelativeBounds(RuntimeBounds).InverseTransformBy(ActorTransform);
+		EditorBoundsRelative = FWorldPartitionRelativeBounds(EditorBounds).InverseTransformBy(ActorTransform);
+
+		bHasValidRelativeBounds = true;
+	}
+	else
+	{
+		ActorTransform = ActorTransformRelative * ParentTransform;
+
+		RuntimeBounds = RuntimeBoundsRelative.TransformBy(ActorTransform).ToAABB();
+		EditorBounds = EditorBoundsRelative.TransformBy(ActorTransform).ToAABB();
+	}
+}
+
 void FWorldPartitionActorDesc::FixupStreamingBounds()
 {
 	const bool bRuntimeBoundsContainsNaN = RuntimeBounds.IsValid && RuntimeBounds.ContainsNaN();
diff --git a/Engine/Source/Runtime/Engine/Public/WorldPartition/ActorDescContainer.h b/Engine/Source/Runtime/Engine/Public/WorldPartition/ActorDescContainer.h
--- a/Engine/Source/Runtime/Engine/Public/WorldPartition/ActorDescContainer.h
+++ b/Engine/Source/Runtime/Engine/Public/WorldPartition/ActorDescContainer.h
@@ -205,6 +205,13 @@ private:
 	ENGINE_API void RegisterEditorDelegates();
 	ENGINE_API void UnregisterEditorDelegates();
 
+	TMap<FGuid, TSet<FGuid>> ParentActorToChildrenMap;
+
+	void AddChildActorToParentMap(FWorldPartitionActorDesc* ActorDesc);
+	void RemoveChildActorFromParentMap(FWorldPartitionActorDesc* ActorDesc);
+
+	void PropagateActorToWorldUpdate(FWorldPartitionActorDesc* ActorDesc);
+
 	static FUObjectAnnotationSparse<FDeletedObjectPlaceholderAnnotation, true> DeletedObjectPlaceholdersAnnotation;
 
 protected:
diff --git a/Engine/Source/Runtime/Engine/Public/WorldPartition/WorldPartitionActorDesc.h b/Engine/Source/Runtime/Engine/Public/WorldPartition/WorldPartitionActorDesc.h
--- a/Engine/Source/Runtime/Engine/Public/WorldPartition/WorldPartitionActorDesc.h
+++ b/Engine/Source/Runtime/Engine/Public/WorldPartition/WorldPartitionActorDesc.h
@@ -66,6 +66,120 @@ private:
 	TVariant<TArray<uint8>, FActorDescArchive*> DataSource;
 };
 
+struct FWorldPartitionRelativeBounds
+{
+	FWorldPartitionRelativeBounds()
+		: bIsValid(false)
+	{ }
+
+	FWorldPartitionRelativeBounds(EForceInit ForceInit)
+		: Center(ForceInit)
+		, Rotation(ForceInit)
+		, Extents(ForceInit)
+		, bIsValid(false)
+	{ }
+
+	FWorldPartitionRelativeBounds(const FVector& InCenter, const FQuat& InRotation, const FVector& InExtents)
+		: Center(InCenter)
+		, Rotation(InRotation)
+		, Extents(InExtents)
+		, bIsValid(true)
+	{ }
+
+	FWorldPartitionRelativeBounds(const FBox& InBox)
+		: FWorldPartitionRelativeBounds(ForceInit)
+	{
+		if (InBox.IsValid)
+		{
+			InBox.GetCenterAndExtents(Center, Extents);
+			Rotation = FQuat::Identity;
+			bIsValid = true;
+		}
+	}
+
+	bool Equals(const FWorldPartitionRelativeBounds& Other, double Tolerance = UE_KINDA_SMALL_NUMBER) const
+	{
+		return (!bIsValid && !Other.bIsValid) 
+			|| ((bIsValid && Other.bIsValid) && Center.Equals(Other.Center, Tolerance) && Rotation.Equals(Other.Rotation, Tolerance) && Extents.Equals(Other.Extents, Tolerance));
+	}
+
+	bool IsValid() const
+	{
+		return bIsValid;
+	}
+
+	FBox ToAABB() const
+	{
+		FBox Result(ForceInit);
+
+		if (bIsValid)
+		{
+			double X = Extents.X;
+			double Y = Extents.Y;
+			double Z = Extents.Z;
+
+			Result += Rotation * FVector(-X, -Y, -Z) + Center;
+			Result += Rotation * FVector( X, -Y, -Z) + Center;
+			Result += Rotation * FVector(-X,  Y, -Z) + Center;
+			Result += Rotation * FVector( X,  Y, -Z) + Center;
+			Result += Rotation * FVector(-X, -Y,  Z) + Center;
+			Result += Rotation * FVector( X, -Y,  Z) + Center;
+			Result += Rotation * FVector(-X,  Y,  Z) + Center;
+			Result += Rotation * FVector( X,  Y,  Z) + Center;
+		}
+
+		return Result;
+	}
+
+	FString ToString() const
+	{
+		return FString::Printf(TEXT("IsValid=%s, Center=(%s), Rotation=(%s), Extents=(%s)"), bIsValid ? TEXT("true") : TEXT("false"), *Center.ToString(), *Rotation.Euler().ToString(), *Extents.ToString());
+	}
+
+	friend FArchive& operator<<(FArchive& Ar, FWorldPartitionRelativeBounds& Bounds)
+	{
+		return Ar << Bounds.Center << Bounds.Rotation << Bounds.Extents << Bounds.bIsValid;
+	}
+
+	[[nodiscard]] FWorldPartitionRelativeBounds TransformBy(const FTransform& Transform) const
+	{
+		if (!bIsValid)
+		{
+			return FWorldPartitionRelativeBounds(ForceInit);
+		}
+
+		FWorldPartitionRelativeBounds NewBounds;
+		NewBounds.Center = Transform.TransformPosition(Center);
+		NewBounds.Rotation = Transform.GetRotation() * Rotation;
+		NewBounds.Extents = Extents * Transform.GetScale3D();
+		NewBounds.bIsValid = true;
+
+		return NewBounds;
+	}
+
+	[[nodiscard]] FWorldPartitionRelativeBounds InverseTransformBy(const FTransform& Transform) const
+	{
+		if (!bIsValid)
+		{
+			return FWorldPartitionRelativeBounds(ForceInit);
+		}
+
+		FWorldPartitionRelativeBounds NewBounds;
+		NewBounds.Center = Transform.InverseTransformPosition(Center);
+		NewBounds.Rotation = Transform.GetRotation().Inverse() * Rotation;
+		NewBounds.Extents = Extents * FTransform::GetSafeScaleReciprocal(Transform.GetScale3D());
+		NewBounds.bIsValid = true;
+
+		return NewBounds;
+	}
+
+private:
+	FVector Center;
+	FQuat Rotation;
+	FVector Extents;
+	bool bIsValid;
+};
+
 struct FWorldPartitionAssetDataPatcher
 {
 	virtual ~FWorldPartitionAssetDataPatcher() {}
@@ -353,6 +467,8 @@ public:
 		Container = InContainer;
 	}
 
+	ENGINE_API virtual void UpdateActorToWorld();
+
 	enum class EToStringMode : uint8
 	{
 		Guid,
@@ -402,6 +518,12 @@ protected:
 		ActorTransform = From->ActorTransform;
 		RuntimeBounds = From->RuntimeBounds;
 		EditorBounds = From->EditorBounds;
+
+
+		ActorTransformRelative = From->ActorTransformRelative;
+		RuntimeBoundsRelative = From->RuntimeBoundsRelative;
+		EditorBoundsRelative = From->EditorBoundsRelative;
+		bHasValidRelativeBounds = From->bHasValidRelativeBounds;
 	}
 
 	virtual uint32 GetSizeOf() const { return sizeof(FWorldPartitionActorDesc); }
@@ -415,9 +537,9 @@ protected:
 	FName							ActorPackage;	// Not serialized, comes from initialization data
 	FSoftObjectPath					ActorPath;		// Not serialized, comes from initialization data
 	FName							ActorLabel;
-	FTransform						ActorTransform;
-	FBox							RuntimeBounds;
-	FBox							EditorBounds;
+	FTransform						ActorTransformRelative;			// Relative transform for actors with parent actor, actor-to-world transform otherwise
+	FWorldPartitionRelativeBounds	RuntimeBoundsRelative;	// Runtime bounds relative to actor-to-world transform
+	FWorldPartitionRelativeBounds	EditorBoundsRelative;	// Editor bounds relative to actor-to-world transform
 	FName							RuntimeGrid;
 	bool							bIsSpatiallyLoaded;
 	bool							bActorIsEditorOnly;
@@ -445,7 +567,11 @@ protected:
 	FString							ActorLabelString;
 	FString							ActorDisplayClassNameString;
 	UActorDescContainer*			Container;
+	FTransform						ActorTransform;	// Current actor-to-world transform, considering parent transform
+	FBox							RuntimeBounds;	// Current runtime bounds, considering parent transform
+	FBox							EditorBounds;	// Current editor bounds, considering parent transform
 	bool							bIsDefaultActorDesc;
+	bool							bHasValidRelativeBounds;
 
 	static ENGINE_API TMap<TSubclassOf<AActor>, FActorDescDeprecator> Deprecators;
 
