PerQualityLevelProperties: Add feeback to the warning about unconfigured platforms, to report which object is requesting the platform. Objects can specify arbitrary platform names that might not exist and cause the warning.
#rnx

diff --git a/Engine/Source/Runtime/Engine/Private/PerQualityLevelProperties.cpp b/Engine/Source/Runtime/Engine/Private/PerQualityLevelProperties.cpp
--- a/Engine/Source/Runtime/Engine/Private/PerQualityLevelProperties.cpp
+++ b/Engine/Source/Runtime/Engine/Private/PerQualityLevelProperties.cpp
@@ -97,7 +97,7 @@ namespace QualityLevelProperty
 		return EnginePlatforms;
 	}
 
-	FSupportedQualityLevelArray PerPlatformOverrideMapping(FString& InPlatformName)
+	FSupportedQualityLevelArray PerPlatformOverrideMapping(FString& InPlatformName, UObject* RequestingAsset)
 	{
 		FSupportedQualityLevelArray* CachedMappingQualitLevelInfo = nullptr;
 		{
@@ -151,14 +151,19 @@ namespace QualityLevelProperty
 				int32 Value = FNameToQualityLevel(FName(*MappingStr));
 				if (Value == INDEX_NONE)
 				{
-					UE_LOG(LogCore, Warning, TEXT("Bad QualityLevelMapping input value. Need to be either [low,medium,high,epic,cinematic]"));
+					UE_LOG(LogCore, Warning, TEXT("Bad QualityLevelMapping input value in %sEngine.ini. ")
+						TEXT("Needs to be one of [low,medium,high,epic,cinematic]."),
+						*EnginePlatformName.ToString());
 					continue;
-				}		
+				}
 				QualityLevels.Add(Value);
 			}
 			else
 			{
-				UE_LOG(LogCore, Warning, TEXT("Didnt found QualityLevelMapping in the %sEngine.ini. Need to define QualityLevelMapping under the [SystemSettings] section. All perplatform MinLOD will not be converted to PerQuality"), *EnginePlatformName.ToString());
+				UE_LOG(LogCore, Warning, TEXT("Didn't find QualityLevelMapping in %sEngine.ini, for platform %s that was requested by %s. ")
+					TEXT("Need to define QualityLevelMapping under the [SystemSettings] section. All perplatform MinLOD will not be converted to PerQuality."),
+					*EnginePlatformName.ToString(), *EnginePlatformName.ToString(),
+					(RequestingAsset ? *RequestingAsset->GetPathName() : TEXT("<Unknown>")));
 			}
 		}
 
diff --git a/Engine/Source/Runtime/Engine/Private/SkeletalMesh.cpp b/Engine/Source/Runtime/Engine/Private/SkeletalMesh.cpp
--- a/Engine/Source/Runtime/Engine/Private/SkeletalMesh.cpp
+++ b/Engine/Source/Runtime/Engine/Private/SkeletalMesh.cpp
@@ -3981,7 +3981,7 @@ void USkeletalMesh::FinishPostLoadInternal(FSkinnedAssetPostLoadContext& Context
 					FSupportedQualityLevelArray QualityLevels;
 					FString PlatformEntry = Pair.Key.ToString();
 
-					QualityLevels = QualityLevelProperty::PerPlatformOverrideMapping(PlatformEntry);
+					QualityLevels = QualityLevelProperty::PerPlatformOverrideMapping(PlatformEntry, this);
 
 					// we now have a range of quality levels supported on that platform or from that group
 					// note: 
diff --git a/Engine/Source/Runtime/Engine/Private/StaticMesh.cpp b/Engine/Source/Runtime/Engine/Private/StaticMesh.cpp
--- a/Engine/Source/Runtime/Engine/Private/StaticMesh.cpp
+++ b/Engine/Source/Runtime/Engine/Private/StaticMesh.cpp
@@ -7278,7 +7278,7 @@ void UStaticMesh::ExecutePostLoadInternal(FStaticMeshPostLoadContext& Context)
 						FSupportedQualityLevelArray QualityLevels;
 						FString PlatformEntry = Pair.Key.ToString();
 
-						QualityLevels = QualityLevelProperty::PerPlatformOverrideMapping(PlatformEntry);
+						QualityLevels = QualityLevelProperty::PerPlatformOverrideMapping(PlatformEntry, this);
 
 						// we now have a range of quality levels supported on that platform or from that group
 						// note: 
diff --git a/Engine/Source/Runtime/Engine/Public/PerQualityLevelProperties.h b/Engine/Source/Runtime/Engine/Public/PerQualityLevelProperties.h
--- a/Engine/Source/Runtime/Engine/Public/PerQualityLevelProperties.h
+++ b/Engine/Source/Runtime/Engine/Public/PerQualityLevelProperties.h
@@ -16,7 +16,8 @@ FPerQualityLevelProperties.h: Property types that can be overridden on a quality
 #include "PerQualityLevelProperties.generated.h"
 
 #if WITH_EDITOR
-	typedef TSet<int32> FSupportedQualityLevelArray;
+class UObject;
+typedef TSet<int32> FSupportedQualityLevelArray;
 #endif
 
 
@@ -65,7 +66,7 @@ namespace QualityLevelProperty
 
 #if WITH_EDITOR
 	ENGINE_API TArray<FName> GetEnginePlatformsForPlatformOrGroupName(const FString& InPlatformName);
-	ENGINE_API FSupportedQualityLevelArray PerPlatformOverrideMapping(FString& InPlatformName);
+	ENGINE_API FSupportedQualityLevelArray PerPlatformOverrideMapping(FString& InPlatformName, UObject* RequestingAsset = nullptr);
 #endif
 };
 
