[Bugfix] Correct bit value for EMoveComponentFlags::MOVECOMP_CheckBlockingRootActorInIgnoreList (from 10110 to 10000)

#12424

diff --git a/Engine/Source/Runtime/Engine/Classes/Components/SceneComponent.h b/Engine/Source/Runtime/Engine/Classes/Components/SceneComponent.h
--- a/Engine/Source/Runtime/Engine/Classes/Components/SceneComponent.h
+++ b/Engine/Source/Runtime/Engine/Classes/Components/SceneComponent.h
@@ -62,7 +62,7 @@ enum EMoveComponentFlags
 	/** avoid dispatching blocking hit events when the hit started in penetration (and is not ignored, see MOVECOMP_NeverIgnoreBlockingOverlaps). */
 	MOVECOMP_DisableBlockingOverlapDispatch		= 0x0008,	
 	/** Compare the root actor of a blocking hit with the ignore UPrimitiveComponent::MoveIgnoreActors array */
-	MOVECOMP_CheckBlockingRootActorInIgnoreList	= 0x0016,	
+	MOVECOMP_CheckBlockingRootActorInIgnoreList	= 0x0010,	
 };
 // Declare bitwise operators to allow EMoveComponentFlags to be combined but still retain type safety
 ENUM_CLASS_FLAGS(EMoveComponentFlags);
