World Partition: instead of crashing when unable to unregister an actor descriptor from a container, trigger a slow path that will search by the guid and report the new actor name.

diff --git a/Engine/Source/Runtime/Engine/Private/WorldPartition/ActorDescContainer.cpp b/Engine/Source/Runtime/Engine/Private/WorldPartition/ActorDescContainer.cpp
--- a/Engine/Source/Runtime/Engine/Private/WorldPartition/ActorDescContainer.cpp
+++ b/Engine/Source/Runtime/Engine/Private/WorldPartition/ActorDescContainer.cpp
@@ -390,22 +390,38 @@ bool UActorDescContainer::IsActorDescHandled(const AActor* InActor) const
 	return IsActorDescHandled(InActor, false);
 }
 
-void UActorDescContainer::RegisterActorDescriptor(FWorldPartitionActorDesc* ActorDesc)
+void UActorDescContainer::RegisterActorDescriptor(FWorldPartitionActorDesc* InActorDesc)
 {
-	FActorDescList::AddActorDescriptor(ActorDesc);
+	FActorDescList::AddActorDescriptor(InActorDesc);
 
-	ActorDesc->SetContainer(this);
+	InActorDesc->SetContainer(this);
 		
-	ActorsByName.Add(ActorDesc->GetActorName(), ActorsByGuid.FindChecked(ActorDesc->GetGuid()));
+	ActorsByName.Add(InActorDesc->GetActorName(), ActorsByGuid.FindChecked(InActorDesc->GetGuid()));
 
-	UE_LOG(LogWorldPartition, Verbose, TEXT("\tRegistered actor descriptor '%s'"), *ActorDesc->ToString(FWorldPartitionActorDesc::EToStringMode::Verbose));
+	UE_LOG(LogWorldPartition, Verbose, TEXT("\tRegistered actor descriptor '%s'"), *InActorDesc->ToString(FWorldPartitionActorDesc::EToStringMode::Verbose));
 }
 
-void UActorDescContainer::UnregisterActorDescriptor(FWorldPartitionActorDesc* ActorDesc)
+void UActorDescContainer::UnregisterActorDescriptor(FWorldPartitionActorDesc* InActorDesc)
 {
-	FActorDescList::RemoveActorDescriptor(ActorDesc);
-	ActorDesc->SetContainer(nullptr);
-	verifyf(ActorsByName.Remove(ActorDesc->GetActorName()), TEXT("Missing actor '%s' from container '%s'"), *ActorDesc->GetActorNameString(), *ContainerPackageName.ToString());
+	FActorDescList::RemoveActorDescriptor(InActorDesc);
+	InActorDesc->SetContainer(nullptr);
+
+	if (!ActorsByName.Remove(InActorDesc->GetActorName()))
+	{
+		bool bRemoved = false;
+		for (auto& [ActorName, ActorDesc] : ActorsByName)
+		{
+			if ((*ActorDesc)->GetGuid() == InActorDesc->GetGuid())
+			{
+				UE_LOG(LogWorldPartition, Log, TEXT("Removed actor '%s' from container '%s' with unexpected name `%s`"), *InActorDesc->GetActorNameString(), *ContainerPackageName.ToString(), *ActorName.ToString());
+				ActorsByName.Remove(ActorName);				
+				bRemoved = true;
+				break;
+			}
+		}
+
+		UE_CLOG(!bRemoved, LogWorldPartition, Log, TEXT("Missing actor '%s' from container '%s'"), *InActorDesc->GetActorNameString(), *ContainerPackageName.ToString());
+	}
 }
 
 bool UActorDescContainer::ShouldHandleActorEvent(const AActor* Actor, bool bInUseLoadedPath) const
diff --git a/Engine/Source/Runtime/Engine/Public/WorldPartition/ActorDescContainer.h b/Engine/Source/Runtime/Engine/Public/WorldPartition/ActorDescContainer.h
--- a/Engine/Source/Runtime/Engine/Public/WorldPartition/ActorDescContainer.h
+++ b/Engine/Source/Runtime/Engine/Public/WorldPartition/ActorDescContainer.h
@@ -133,8 +133,8 @@ public:
 	const TArray<FAssetData>& GetInvalidActors() const { return InvalidActors; }
 	void ClearInvalidActors() { InvalidActors.Empty(); }
 
-	ENGINE_API void RegisterActorDescriptor(FWorldPartitionActorDesc* ActorDesc);
-	ENGINE_API void UnregisterActorDescriptor(FWorldPartitionActorDesc* ActorDesc);
+	ENGINE_API void RegisterActorDescriptor(FWorldPartitionActorDesc* InActorDesc);
+	ENGINE_API void UnregisterActorDescriptor(FWorldPartitionActorDesc* InActorDesc);
 
 	ENGINE_API void OnActorDescAdded(FWorldPartitionActorDesc* NewActorDesc);
 	ENGINE_API void OnActorDescRemoved(FWorldPartitionActorDesc* ActorDesc);
