Forbid the swapping of non-copyable types.

diff --git a/Engine/Plugins/Experimental/AnimNextAnimGraph/Source/AnimNextAnimGraphTestSuite/Private/AnimNextRuntimeTest.cpp b/Engine/Plugins/Experimental/AnimNextAnimGraph/Source/AnimNextAnimGraphTestSuite/Private/AnimNextRuntimeTest.cpp
--- a/Engine/Plugins/Experimental/AnimNextAnimGraph/Source/AnimNextAnimGraphTestSuite/Private/AnimNextRuntimeTest.cpp
+++ b/Engine/Plugins/Experimental/AnimNextAnimGraph/Source/AnimNextAnimGraphTestSuite/Private/AnimNextRuntimeTest.cpp
@@ -16,13 +16,15 @@ namespace UE::AnimNext
 	FScopedClearNodeTemplateRegistry::FScopedClearNodeTemplateRegistry()
 	{
 		FNodeTemplateRegistry& Registry = FNodeTemplateRegistry::Get();
-		Swap(Registry, TmpRegistry);
+		// Can't use Swap because the move constructor is private and Swap thinks it's non-swappable
+		FMemory::Memswap(&Registry, &TmpRegistry, sizeof(FNodeTemplateRegistry));
 	}
 
 	FScopedClearNodeTemplateRegistry::~FScopedClearNodeTemplateRegistry()
 	{
 		FNodeTemplateRegistry& Registry = FNodeTemplateRegistry::Get();
-		Swap(Registry, TmpRegistry);
+		// Can't use Swap because the move constructor is private and Swap thinks it's non-swappable
+		FMemory::Memswap(&Registry, &TmpRegistry, sizeof(FNodeTemplateRegistry));
 	}
 
 	bool FTestUtils::LoadFromArchiveBuffer(UAnimNextAnimationGraph& AnimationGraph, TArray<FNodeHandle>& NodeHandles, const TArray<uint8>& SharedDataArchiveBuffer)
diff --git a/Engine/Plugins/Media/ElectraPlayer/Source/ElectraPlayerRuntime/Private/Runtime/Player/AdaptiveStreamingPlayer.cpp b/Engine/Plugins/Media/ElectraPlayer/Source/ElectraPlayerRuntime/Private/Runtime/Player/AdaptiveStreamingPlayer.cpp
--- a/Engine/Plugins/Media/ElectraPlayer/Source/ElectraPlayerRuntime/Private/Runtime/Player/AdaptiveStreamingPlayer.cpp
+++ b/Engine/Plugins/Media/ElectraPlayer/Source/ElectraPlayerRuntime/Private/Runtime/Player/AdaptiveStreamingPlayer.cpp
@@ -1694,7 +1694,7 @@ bool FAdaptiveStreamingPlayer::InternalHandleThreadMessages()
 					// Add to the list of completed requests for which we need to find the next or retry segment to fetch.
 					FPendingSegmentRequest NextReq;
 					NextReq.Request = pRequest;
-					NextPendingSegmentRequests.Enqueue(MoveTemp(NextReq));
+					NextPendingSegmentRequests->Enqueue(MoveTemp(NextReq));
 				}
 				break;
 			}
@@ -3030,7 +3030,7 @@ void FAdaptiveStreamingPlayer::InternalHandlePendingFirstSegmentRequest(const FT
 			PendingFirstSegmentRequest->GetRequestedStreams(RequestedStreams);
 			for(auto &RequestedStream : RequestedStreams)
 			{
-				ReadyWaitingSegmentRequests.Enqueue(RequestedStream);
+				ReadyWaitingSegmentRequests->Enqueue(RequestedStream);
 			}
 		}
 		PendingFirstSegmentRequest.Reset();
@@ -3121,7 +3121,7 @@ void FAdaptiveStreamingPlayer::RequestNewPeriodStreams(EStreamType InType, FPend
 				SetPlaystartOptions(NextReq.StartoverPosition.Options);
 				NextReq.bStartOver = true;
 				NextReq.bDidRequestNewPeriodStreams = true;
-				NextPendingSegmentRequests.Enqueue(MoveTemp(NextReq));
+				NextPendingSegmentRequests->Enqueue(MoveTemp(NextReq));
 			}
 		};
 
@@ -3164,12 +3164,12 @@ void FAdaptiveStreamingPlayer::HandlePendingMediaSegmentRequests()
 	FTimeValue CurrentPlaybackPos = GetPlayPosition();
 
 	// Are there completed downloads for which we need to find the next segment?
-	TQueue<FPendingSegmentRequest> pendingRequests;
+	TUniqueObj<TQueue<FPendingSegmentRequest>> pendingRequests;
 	Swap(pendingRequests, NextPendingSegmentRequests);
-	while(!pendingRequests.IsEmpty())
+	while(!pendingRequests->IsEmpty())
 	{
 		FPendingSegmentRequest FinishedReq;
-		pendingRequests.Dequeue(FinishedReq);
+		pendingRequests->Dequeue(FinishedReq);
 
 		TSharedPtrTS<IManifest::IPlayPeriod> SegmentPlayPeriod;
 
@@ -3200,7 +3200,7 @@ void FAdaptiveStreamingPlayer::HandlePendingMediaSegmentRequests()
 		// If not yet ready to check again put the request back into the list.
 		if (FinishedReq.AtTime.IsValid() && FinishedReq.AtTime > CurrentTime)
 		{
-			NextPendingSegmentRequests.Enqueue(MoveTemp(FinishedReq));
+			NextPendingSegmentRequests->Enqueue(MoveTemp(FinishedReq));
 		}
 		else if (SegmentPlayPeriod.IsValid())
 		{
@@ -3241,12 +3241,12 @@ void FAdaptiveStreamingPlayer::HandlePendingMediaSegmentRequests()
 					FinishedReq.Period->SetStreamPreferences(EStreamType::Audio, StreamSelectionAttributesAud);
 					FinishedReq.Period->SetStreamPreferences(EStreamType::Subtitle, StreamSelectionAttributesTxt);
 					FinishedReq.Period->Load();
-					NextPendingSegmentRequests.Enqueue(MoveTemp(FinishedReq));
+					NextPendingSegmentRequests->Enqueue(MoveTemp(FinishedReq));
 					continue;
 				}
 				else if (PeriodState == IManifest::IPlayPeriod::EReadyState::Loading)
 				{
-					NextPendingSegmentRequests.Enqueue(MoveTemp(FinishedReq));
+					NextPendingSegmentRequests->Enqueue(MoveTemp(FinishedReq));
 					continue;
 				}
 				else if (PeriodState == IManifest::IPlayPeriod::EReadyState::Loaded)
@@ -3254,12 +3254,12 @@ void FAdaptiveStreamingPlayer::HandlePendingMediaSegmentRequests()
 					// Set the current average video bitrate in the player options for the period to retrieve if necessary.
 					PlayerOptions.Set(OptionKeyCurrentAvgStartingVideoBitrate, FVariantValue(StreamSelector->GetAverageBandwidth()));
 					FinishedReq.Period->PrepareForPlay();
-					NextPendingSegmentRequests.Enqueue(MoveTemp(FinishedReq));
+					NextPendingSegmentRequests->Enqueue(MoveTemp(FinishedReq));
 					continue;
 				}
 				else if (PeriodState == IManifest::IPlayPeriod::EReadyState::Preparing)
 				{
-					NextPendingSegmentRequests.Enqueue(MoveTemp(FinishedReq));
+					NextPendingSegmentRequests->Enqueue(MoveTemp(FinishedReq));
 					continue;
 				}
 				else if (PeriodState == IManifest::IPlayPeriod::EReadyState::IsReady)
@@ -3333,7 +3333,7 @@ void FAdaptiveStreamingPlayer::HandlePendingMediaSegmentRequests()
 				case IManifest::FResult::EType::TryAgainLater:
 				{
 					FinishedReq.AtTime = Result.GetRetryAgainAtTime();
-					NextPendingSegmentRequests.Enqueue(MoveTemp(FinishedReq));
+					NextPendingSegmentRequests->Enqueue(MoveTemp(FinishedReq));
 					break;
 				}
 				case IManifest::FResult::EType::Found:
@@ -3389,7 +3389,7 @@ void FAdaptiveStreamingPlayer::HandlePendingMediaSegmentRequests()
 					}
 
 					NextSegment->SetExecutionDelay(UTCNow, ActionDelay);
-					ReadyWaitingSegmentRequests.Enqueue(NextSegment);
+					ReadyWaitingSegmentRequests->Enqueue(NextSegment);
 					break;
 				}
 				case IManifest::FResult::EType::PastEOS:
@@ -3402,14 +3402,14 @@ void FAdaptiveStreamingPlayer::HandlePendingMediaSegmentRequests()
 						{
 							// Next period is either not there yet or is being resolved right now.
 							FinishedReq.AtTime = Result.GetRetryAgainAtTime();
-							NextPendingSegmentRequests.Enqueue(MoveTemp(FinishedReq));
+							NextPendingSegmentRequests->Enqueue(MoveTemp(FinishedReq));
 							break;
 						}
 						case IManifest::FResult::EType::Found:
 						{
 							// Get the next period ready for playback. Put it back into the queue.
 							FinishedReq.Period = NextPeriod;
-							NextPendingSegmentRequests.Enqueue(MoveTemp(FinishedReq));
+							NextPendingSegmentRequests->Enqueue(MoveTemp(FinishedReq));
 							break;
 						}
 						case IManifest::FResult::EType::PastEOS:
@@ -3470,12 +3470,12 @@ void FAdaptiveStreamingPlayer::HandlePendingMediaSegmentRequests()
 
 
 	// Go through the queue of ready-to-fetch segments to see if they should be fetched now.
-	TQueue<TSharedPtrTS<IStreamSegment>> readyRequests;
+	TUniqueObj<TQueue<TSharedPtrTS<IStreamSegment>>> readyRequests;
 	Swap(readyRequests, ReadyWaitingSegmentRequests);
-	while(!readyRequests.IsEmpty())
+	while(!readyRequests->IsEmpty())
 	{
 		TSharedPtrTS<IStreamSegment> ReadyReq;
-		readyRequests.Dequeue(ReadyReq);
+		readyRequests->Dequeue(ReadyReq);
 
 		bool bExecuteNow = true;
 		bool bSkipOver = false;
@@ -3522,12 +3522,12 @@ void FAdaptiveStreamingPlayer::HandlePendingMediaSegmentRequests()
 				// Skip over this segment. Add to the list of completed requests for which we need to find the next segment to fetch.
 				FPendingSegmentRequest NextReq;
 				NextReq.Request = ReadyReq;
-				NextPendingSegmentRequests.Enqueue(MoveTemp(NextReq));
+				NextPendingSegmentRequests->Enqueue(MoveTemp(NextReq));
 			}
 		}
 		else
 		{
-			ReadyWaitingSegmentRequests.Enqueue(ReadyReq);
+			ReadyWaitingSegmentRequests->Enqueue(ReadyReq);
 		}
 	}
 
@@ -3541,27 +3541,27 @@ void FAdaptiveStreamingPlayer::CancelPendingMediaSegmentRequests(EStreamType Str
 		StreamReaderHandler->CancelRequest(StreamType, true);
 	}
 	// Go through the queue of ready-to-fetch segments and remove those that are to be canceled.
-	TQueue<TSharedPtrTS<IStreamSegment>> readyRequests;
+	TUniqueObj<TQueue<TSharedPtrTS<IStreamSegment>>> readyRequests;
 	Swap(readyRequests, ReadyWaitingSegmentRequests);
-	while(!readyRequests.IsEmpty())
+	while(!readyRequests->IsEmpty())
 	{
 		TSharedPtrTS<IStreamSegment> ReadyReq;
-		readyRequests.Dequeue(ReadyReq);
+		readyRequests->Dequeue(ReadyReq);
 		if (ReadyReq->GetType() != StreamType)
 		{
-			ReadyWaitingSegmentRequests.Enqueue(ReadyReq);
+			ReadyWaitingSegmentRequests->Enqueue(ReadyReq);
 		}
 	}
 	// Likewise for pending segments that hadn't been made ready yet.
-	TQueue<FPendingSegmentRequest> pendingRequests;
+	TUniqueObj<TQueue<FPendingSegmentRequest>> pendingRequests;
 	Swap(pendingRequests, NextPendingSegmentRequests);
-	while(!pendingRequests.IsEmpty())
+	while(!pendingRequests->IsEmpty())
 	{
 		FPendingSegmentRequest PendingReq;
-		pendingRequests.Dequeue(PendingReq);
+		pendingRequests->Dequeue(PendingReq);
 		if (PendingReq.StreamType != StreamType)
 		{
-			NextPendingSegmentRequests.Enqueue(PendingReq);
+			NextPendingSegmentRequests->Enqueue(PendingReq);
 		}
 	}
 }
@@ -3697,7 +3697,7 @@ void FAdaptiveStreamingPlayer::InternalHandleSegmentTrackChanges(const FTimeValu
 						NextReq.StartoverPosition.Time = PlaybackState.GetPlayPosition();
 						SetPlaystartOptions(NextReq.StartoverPosition.Options);
 						NextReq.StreamType = InType;
-						NextPendingSegmentRequests.Enqueue(MoveTemp(NextReq));
+						NextPendingSegmentRequests->Enqueue(MoveTemp(NextReq));
 
 						bChangeMade = true;
 					}
@@ -4277,11 +4277,11 @@ void FAdaptiveStreamingPlayer::InternalStartAt(const FSeekParam& NewPosition, co
 	check(!PendingStartRequest.IsValid());
 	check(!InitialPlayPeriod.IsValid());
 	check(!PendingFirstSegmentRequest.IsValid());
-	check(NextPendingSegmentRequests.IsEmpty());
-	check(ReadyWaitingSegmentRequests.IsEmpty());
+	check(NextPendingSegmentRequests->IsEmpty());
+	check(ReadyWaitingSegmentRequests->IsEmpty());
 	check(CompletedSegmentRequests.Num() == 0);
-	NextPendingSegmentRequests.Empty();
-	ReadyWaitingSegmentRequests.Empty();
+	NextPendingSegmentRequests->Empty();
+	ReadyWaitingSegmentRequests->Empty();
 	PendingFirstSegmentRequest.Reset();
 	CompletedSegmentRequests.Empty();
 	InitialPlayPeriod.Reset();
@@ -4459,8 +4459,8 @@ void FAdaptiveStreamingPlayer::InternalStop(bool bHoldCurrentFrame)
 	// Release any pending segment requests.
 	PendingStartRequest.Reset();
 	PendingFirstSegmentRequest.Reset();
-	NextPendingSegmentRequests.Empty();
-	ReadyWaitingSegmentRequests.Empty();
+	NextPendingSegmentRequests->Empty();
+	ReadyWaitingSegmentRequests->Empty();
 	CompletedSegmentRequests.Empty();
 	InitialPlayPeriod.Reset();
 	CurrentPlayPeriodVideo.Reset();
diff --git a/Engine/Plugins/Media/ElectraPlayer/Source/ElectraPlayerRuntime/Private/Runtime/Player/AdaptiveStreamingPlayerInternal.h b/Engine/Plugins/Media/ElectraPlayer/Source/ElectraPlayerRuntime/Private/Runtime/Player/AdaptiveStreamingPlayerInternal.h
--- a/Engine/Plugins/Media/ElectraPlayer/Source/ElectraPlayerRuntime/Private/Runtime/Player/AdaptiveStreamingPlayerInternal.h
+++ b/Engine/Plugins/Media/ElectraPlayer/Source/ElectraPlayerRuntime/Private/Runtime/Player/AdaptiveStreamingPlayerInternal.h
@@ -25,6 +25,8 @@
 #include "Player/AdaptiveStreamingPlayerABR.h"
 #include "Player/AdaptivePlayerOptionKeynames.h"
 
+#include "Templates/UniqueObj.h"
+
 #include "Utilities/UtilsMP4.h"
 
 #include "ElectraCDM.h"
@@ -2311,8 +2313,8 @@ private:
 	TSharedPtrTS<IManifest::IPlayPeriod>								CurrentPlayPeriodText;
 	TSharedPtrTS<FPendingStartRequest>									PendingStartRequest;
 	TSharedPtrTS<IStreamSegment>										PendingFirstSegmentRequest;
-	TQueue<FPendingSegmentRequest>										NextPendingSegmentRequests;
-	TQueue<TSharedPtrTS<IStreamSegment>>								ReadyWaitingSegmentRequests;
+	TUniqueObj<TQueue<FPendingSegmentRequest>>							NextPendingSegmentRequests;
+	TUniqueObj<TQueue<TSharedPtrTS<IStreamSegment>>>					ReadyWaitingSegmentRequests;
 	TMultiMap<EStreamType, TSharedPtrTS<IStreamSegment>>				CompletedSegmentRequests;
 	bool																bFirstSegmentRequestIsForLooping;
 
diff --git a/Engine/Source/Developer/DerivedDataCache/Private/DerivedDataCacheStore.cpp b/Engine/Source/Developer/DerivedDataCache/Private/DerivedDataCacheStore.cpp
--- a/Engine/Source/Developer/DerivedDataCache/Private/DerivedDataCacheStore.cpp
+++ b/Engine/Source/Developer/DerivedDataCache/Private/DerivedDataCacheStore.cpp
@@ -25,7 +25,7 @@ void FCacheStoreRequestTimer::Stop()
 	}
 	const FMonotonicTimePoint EndTime = FMonotonicTimePoint::Now();
 	const bool bIsInGameThread = IsInGameThread();
-	TUniqueLock Lock(Stats.Mutex);
+	TUniqueLock Lock(Stats.Mutex.Get());
 	(bIsInGameThread ? Stats.MainThreadTime : Stats.OtherThreadTime) += EndTime - StartTime;
 	Stats.StartTime = FMath::Min(Stats.StartTime, StartTime);
 	Stats.EndTime = FMath::Max(Stats.EndTime, EndTime);
diff --git a/Engine/Source/Developer/DerivedDataCache/Private/DerivedDataCacheStore.h b/Engine/Source/Developer/DerivedDataCache/Private/DerivedDataCacheStore.h
--- a/Engine/Source/Developer/DerivedDataCache/Private/DerivedDataCacheStore.h
+++ b/Engine/Source/Developer/DerivedDataCache/Private/DerivedDataCacheStore.h
@@ -6,6 +6,7 @@
 #include "Containers/ContainersFwd.h"
 #include "DerivedDataCache.h"
 #include "Misc/MonotonicTime.h"
+#include "Templates/DontCopy.h"
 
 class FText;
 struct FDerivedDataCacheResourceStat;
@@ -252,7 +253,7 @@ struct FCacheStoreRequestStats final
 	/** The status with which this request completed. */
 	EStatus Status = EStatus::Ok;
 	/** Mutex that can be used by cache stores that process a request on multiple threads. */
-	mutable FMutex Mutex;
+	mutable TDontCopy<FMutex> Mutex;
 
 	void AddLatency(FMonotonicTimeSpan Latency);
 
diff --git a/Engine/Source/Developer/DerivedDataCache/Private/HttpCacheStore.cpp b/Engine/Source/Developer/DerivedDataCache/Private/HttpCacheStore.cpp
--- a/Engine/Source/Developer/DerivedDataCache/Private/HttpCacheStore.cpp
+++ b/Engine/Source/Developer/DerivedDataCache/Private/HttpCacheStore.cpp
@@ -707,7 +707,7 @@ TSharedPtr<FJsonObject> FHttpCacheStore::FHttpOperation::GetBodyAsJson() const
 void FHttpCacheStore::FHttpOperation::GetStats(FRequestStats& OutStats) const
 {
 	const FHttpResponseStats& Stats = Response->GetStats();
-	TUniqueLock Lock(OutStats.Mutex);
+	TUniqueLock Lock(OutStats.Mutex.Get());
 	OutStats.PhysicalReadSize += Stats.RecvSize;
 	OutStats.PhysicalWriteSize += Stats.SendSize;
 	if (const EHttpMethod Method = Response->GetMethod(); Method == EHttpMethod::Get || Method == EHttpMethod::Head)
@@ -1044,7 +1044,7 @@ void FHttpCacheStore::FPutPackageOp::EndPutBlob(FHttpOperation* Operation, uint6
 		if (StatusCode >= 200 && StatusCode <= 204)
 		{
 			SuccessfulBlobUploads.fetch_add(1, std::memory_order_relaxed);
-			TUniqueLock Lock(RequestStats.Mutex);
+			TUniqueLock Lock(RequestStats.Mutex.Get());
 			RequestStats.LogicalWriteSize += LogicalSize;
 		}
 	}
@@ -1503,7 +1503,7 @@ void FHttpCacheStore::FGetRecordOp::EndGetValue(FHttpOperation& Operation, const
 			break;
 		}
 
-		TUniqueLock Lock(RequestStats.Mutex);
+		TUniqueLock Lock(RequestStats.Mutex.Get());
 		RequestStats.LogicalReadSize += CompressedBuffer.GetRawSize();
 	}
 
diff --git a/Engine/Source/Developer/DerivedDataCache/Private/ZenCacheStore.cpp b/Engine/Source/Developer/DerivedDataCache/Private/ZenCacheStore.cpp
--- a/Engine/Source/Developer/DerivedDataCache/Private/ZenCacheStore.cpp
+++ b/Engine/Source/Developer/DerivedDataCache/Private/ZenCacheStore.cpp
@@ -254,7 +254,7 @@ struct FZenCacheStore::TRequestWithStats
 	void EndRequest(FZenCacheStore& Outer, const EStatus Status) const
 	{
 		{
-			TUniqueLock Lock(Stats.Mutex);
+			TUniqueLock Lock(Stats.Mutex.Get());
 			Stats.EndTime = FMonotonicTimePoint::Now();
 			Stats.Status = Status;
 		}
diff --git a/Engine/Source/Runtime/Core/Public/Templates/UnrealTemplate.h b/Engine/Source/Runtime/Core/Public/Templates/UnrealTemplate.h
--- a/Engine/Source/Runtime/Core/Public/Templates/UnrealTemplate.h
+++ b/Engine/Source/Runtime/Core/Public/Templates/UnrealTemplate.h
@@ -627,6 +627,11 @@ UE_INTRINSIC_CAST FORCEINLINE constexpr T&& Forward(std::remove_reference_t<T>&&
 template <typename T>
 inline void Swap(T& A, T& B)
 {
+	// std::is_swappable isn't correct here, because we allow bitwise swapping of types containing e.g. const and reference members,
+	// but we don't want to allow swapping of types which are UE_NONCOPYABLE or equivalent.  We also allow bitwise swapping of arrays, so
+	// extents should be removed first.
+	static_assert(std::is_move_constructible_v<std::remove_all_extents_t<T>>, "Cannot swap non-movable types");
+
 	if constexpr (TUseBitwiseSwap<T>::Value)
 	{
 		TTypeCompatibleBytes<T> Temp;
diff --git a/Engine/Source/Runtime/Engine/Private/AudioStreamingCache.cpp b/Engine/Source/Runtime/Engine/Private/AudioStreamingCache.cpp
--- a/Engine/Source/Runtime/Engine/Private/AudioStreamingCache.cpp
+++ b/Engine/Source/Runtime/Engine/Private/AudioStreamingCache.cpp
@@ -750,7 +750,7 @@ uint64 FAudioChunkCache::AddOrTouchChunk(const FChunkKey& InKey, const TSharedPt
 		return InvalidAudioStreamCacheLookupID;
 	}
 
-	FScopeLock ScopeLock(&CacheMutationCriticalSection);
+	FScopeLock ScopeLock(&CacheMutationCriticalSection.Get());
 
 	FCacheElement* FoundElement = FindElementForKey(InKey);
 	if (FoundElement)
@@ -817,7 +817,7 @@ uint64 FAudioChunkCache::AddOrTouchChunk(const FChunkKey& InKey, const TSharedPt
 			const uint32 TotalNumChunksInWave = CacheElement->GetNumChunks();
 
 			FCacheMissInfo CacheMissInfo = { InKey.SoundWaveName, InKey.ChunkIndex, TotalNumChunksInWave, false };
-			CacheMissQueue.Enqueue(MoveTemp(CacheMissInfo));
+			CacheMissQueue->Enqueue(MoveTemp(CacheMissInfo));
 		}
 
 		return CacheElement->CacheLookupID;
@@ -826,7 +826,7 @@ uint64 FAudioChunkCache::AddOrTouchChunk(const FChunkKey& InKey, const TSharedPt
 
 TArrayView<uint8> FAudioChunkCache::GetChunk(const FChunkKey& InKey, const TSharedPtr<FSoundWaveData>& InSoundWavePtr, bool bBlockForLoadCompletion, bool bNeededForPlayback, uint64& OutCacheOffset)
 {
-	FScopeLock ScopeLock(&CacheMutationCriticalSection);
+	FScopeLock ScopeLock(&CacheMutationCriticalSection.Get());
 	FCacheElement* FoundElement = FindElementForKey(InKey);
 	if (FoundElement)
 	{
@@ -893,7 +893,7 @@ TArrayView<uint8> FAudioChunkCache::GetChunk(const FChunkKey& InKey, const TShar
 				FoundElement->DebugInfo.TimeLoadEndedCycles = FoundElement->DebugInfo.TimeLoadStartedCycles;
 				FoundElement->DebugInfo.TimeToLoadMs = 0;
 #endif
-				MemoryCounterBytes -= FoundElement->ChunkDataSize;
+				*MemoryCounterBytes -= FoundElement->ChunkDataSize;
 
 				{
 					LLM_SCOPE(ELLMTag::AudioStreamCacheCompressedData);
@@ -910,7 +910,7 @@ TArrayView<uint8> FAudioChunkCache::GetChunk(const FChunkKey& InKey, const TShar
 					}
 				}
 
-				MemoryCounterBytes += ChunkAudioDataSize;
+				*MemoryCounterBytes += ChunkAudioDataSize;
 
 				// Populate key and DataSize. The async read request was set up to write directly into CacheElement->ChunkData.
 				FoundElement->Key = InKey;
@@ -944,7 +944,7 @@ TArrayView<uint8> FAudioChunkCache::GetChunk(const FChunkKey& InKey, const TShar
 			// Chunks missing. Log this as a miss.
 			const uint32 TotalNumChunksInWave = FoundElement->GetNumChunks();
 			FCacheMissInfo CacheMissInfo = { InKey.SoundWaveName, InKey.ChunkIndex, TotalNumChunksInWave, false };
-			CacheMissQueue.Enqueue(MoveTemp(CacheMissInfo));
+			CacheMissQueue->Enqueue(MoveTemp(CacheMissInfo));
 		}
 
 		// We missed, return an empty array view.
@@ -954,7 +954,7 @@ TArrayView<uint8> FAudioChunkCache::GetChunk(const FChunkKey& InKey, const TShar
 
 void FAudioChunkCache::AddNewReferenceToChunk(const FChunkKey& InKey)
 {
-	FScopeLock ScopeLock(&CacheMutationCriticalSection);
+	FScopeLock ScopeLock(&CacheMutationCriticalSection.Get());
 	FCacheElement* FoundElement = FindElementForKey(InKey);
 	if (ensure(FoundElement))
 	{
@@ -966,7 +966,7 @@ void FAudioChunkCache::AddNewReferenceToChunk(const FChunkKey& InKey)
 
 void FAudioChunkCache::RemoveReferenceToChunk(const FChunkKey& InKey)
 {
-	FScopeLock ScopeLock(&CacheMutationCriticalSection);
+	FScopeLock ScopeLock(&CacheMutationCriticalSection.Get());
 	FCacheElement* FoundElement = FindElementForKey(InKey);
 	if (ensure(FoundElement))
 	{
@@ -978,7 +978,7 @@ void FAudioChunkCache::RemoveReferenceToChunk(const FChunkKey& InKey)
 
 void FAudioChunkCache::ClearCache()
 {
-	FScopeLock ScopeLock(&CacheMutationCriticalSection);
+	FScopeLock ScopeLock(&CacheMutationCriticalSection.Get());
 	const uint32 NumChunks = CachePool.Num();
 
 	UE_LOG(LogAudioStreamCaching, Verbose, TEXT("Clearing Cache"));
@@ -1012,7 +1012,7 @@ void FAudioChunkCache::AddForceInlineSoundWave(const FSoundWaveProxyPtr& SoundWa
 	int64 MemoryCount = Data ? Data->GetBulkDataSize() : 0;
 	int32 RefCount = 0;
 	{
-		FScopeLock Lock(&SoundWaveMemoryTrackerCritSec);
+		FScopeLock Lock(&SoundWaveMemoryTrackerCritSec.Get());
 		FSoundWaveMemoryTracker& Tracker = SoundWaveTracker.FindOrAdd(SoundWave);
 		checkf(Tracker.RefCount >= 0, TEXT("AudioStreamCache::AddForceInlineSoundWave: ref count for Added sound wave is negative!: %s"), *SoundWave->GetFName().ToString());
 		RefCount = ++Tracker.RefCount;
@@ -1031,7 +1031,7 @@ void FAudioChunkCache::AddForceInlineSoundWave(const FSoundWaveProxyPtr& SoundWa
 	// we only increment memory count for the first time the sound wave is added
 	if (RefCount == 1)
 	{
-		ForceInlineMemoryCounterBytes += MemoryCount;
+		*ForceInlineMemoryCounterBytes += MemoryCount;
 		const uint64 MemoryUsageBytes = GetCurrentMemoryUsageBytes();
 		if (TrimCacheWhenOverBudgetCVar != 0 && MemoryUsageBytes > MemoryLimitBytes)
 		{
@@ -1066,7 +1066,7 @@ void FAudioChunkCache::RemoveForceInlineSoundWave(const FSoundWaveProxyPtr& Soun
 
 	// scope lock
 	{
-		FScopeLock Lock(&SoundWaveMemoryTrackerCritSec);
+		FScopeLock Lock(&SoundWaveMemoryTrackerCritSec.Get());
 
 		FSoundWaveMemoryTracker* Tracker = SoundWaveTracker.Find(SoundWave);
 		checkf(Tracker != nullptr, TEXT("AudioStreamCache::RemoveForceInlineSoundWave: Attempted to remove SoundWave that was never added, or has already been removed: %s"), *SoundWave->GetFName().ToString());
@@ -1082,15 +1082,15 @@ void FAudioChunkCache::RemoveForceInlineSoundWave(const FSoundWaveProxyPtr& Soun
 	
 	if (RefCount == 0)
 	{
-		ForceInlineMemoryCounterBytes -= MemoryCount;
+		*ForceInlineMemoryCounterBytes -= MemoryCount;
 	}
 }
 
 void FAudioChunkCache::AddMemoryCountedFeature(const FAudioStreamCacheMemoryHandle& Feature)
 {
 	UE_LOG(LogAudioStreamCaching, Log, TEXT("Adding Memory Counted Feature (%s) Memory Usage: %d bytes"), *Feature.GetFeatureName().ToString(), (int32)Feature.GetMemoryUseInBytes());
-	const uint64 OldMemoryCount = FeatureMemoryCounterBytes.AddExchange(Feature.GetMemoryUseInBytes());
-	UE_LOG(LogAudioStreamCaching, Log, TEXT("Total Memory Usage for all features: %d -> %d bytes"), (int32)OldMemoryCount, (int32)FeatureMemoryCounterBytes.Load());
+	const uint64 OldMemoryCount = FeatureMemoryCounterBytes->AddExchange(Feature.GetMemoryUseInBytes());
+	UE_LOG(LogAudioStreamCaching, Log, TEXT("Total Memory Usage for all features: %d -> %d bytes"), (int32)OldMemoryCount, (int32)FeatureMemoryCounterBytes->Load());
 	const uint64 MemoryUsageBytes = GetCurrentMemoryUsageBytes();
 	if (TrimCacheWhenOverBudgetCVar != 0 && MemoryUsageBytes > MemoryLimitBytes)
 	{
@@ -1111,14 +1111,14 @@ void FAudioChunkCache::AddMemoryCountedFeature(const FAudioStreamCacheMemoryHand
 void FAudioChunkCache::RemoveMemoryCountedFeature(const FAudioStreamCacheMemoryHandle& Feature)
 {
 	UE_LOG(LogAudioStreamCaching, Log, TEXT("Removing Memory Counted Feature (%s) Memory Usage: %d"), *Feature.GetFeatureName().ToString(), (int32)Feature.GetMemoryUseInBytes());
-	checkf(FeatureMemoryCounterBytes.Load() >= Feature.GetMemoryUseInBytes(), TEXT("Count (%lu) < Remove (%lu)"), FeatureMemoryCounterBytes.Load(), Feature.GetMemoryUseInBytes());
-	const uint32 OldMemoryCount = FeatureMemoryCounterBytes.SubExchange(Feature.GetMemoryUseInBytes());
-	UE_LOG(LogAudioStreamCaching, Log, TEXT("Total Memory Usage for all features: %d -> %d"), (int32)OldMemoryCount, (int32)FeatureMemoryCounterBytes.Load());
+	checkf(FeatureMemoryCounterBytes->Load() >= Feature.GetMemoryUseInBytes(), TEXT("Count (%lu) < Remove (%lu)"), FeatureMemoryCounterBytes->Load(), Feature.GetMemoryUseInBytes());
+	const uint32 OldMemoryCount = FeatureMemoryCounterBytes->SubExchange(Feature.GetMemoryUseInBytes());
+	UE_LOG(LogAudioStreamCaching, Log, TEXT("Total Memory Usage for all features: %d -> %d"), (int32)OldMemoryCount, (int32)FeatureMemoryCounterBytes->Load());
 }
 
 uint64 FAudioChunkCache::TrimMemory(uint64 BytesToFree, bool bInAllowRetainedChunkTrimming)
 {
-	FScopeLock ScopeLock(&CacheMutationCriticalSection);
+	FScopeLock ScopeLock(&CacheMutationCriticalSection.Get());
 
 	if (!MostRecentElement || MostRecentElement->LessRecentElement == nullptr)
 	{
@@ -1142,7 +1142,7 @@ uint64 FAudioChunkCache::TrimMemory(uint64 BytesToFree, bool bInAllowRetainedChu
 
 			uint32 ChunkSize = CurrentElement->ChunkDataSize;
 			BytesFreed += ChunkSize;
-			MemoryCounterBytes -= ChunkSize;
+			*MemoryCounterBytes -= ChunkSize;
 
 			// Empty the chunk data and invalidate the key.
 			check(CurrentElement->ChunkData);
@@ -1189,7 +1189,7 @@ uint64 FAudioChunkCache::TrimMemory(uint64 BytesToFree, bool bInAllowRetainedChu
 					uint32 ChunkSize = CurrentElement->ChunkDataSize;
 					BytesFreed += ChunkSize;
 					RetainedBytesFreed += ChunkSize;
-					MemoryCounterBytes -= ChunkSize;
+					*MemoryCounterBytes -= ChunkSize;
 
 					// Empty the chunk data and invalidate the key.
 					check(CurrentElement->ChunkData);
@@ -1239,7 +1239,7 @@ void FAudioChunkCache::BlockForAllPendingLoads() const
 		}
 
 		{
-			FScopeLock ScopeLock(const_cast<FCriticalSection*>(&CacheMutationCriticalSection));
+			FScopeLock ScopeLock(&CacheMutationCriticalSection.Get());
 
 			// Iterate through every element until we find one with a load in progress.
 			const FCacheElement* CurrentElement = MostRecentElement;
@@ -1254,7 +1254,7 @@ void FAudioChunkCache::BlockForAllPendingLoads() const
 
 void FAudioChunkCache::CancelAllPendingLoads()
 {
-	FScopeLock ScopeLock(&CacheMutationCriticalSection);
+	FScopeLock ScopeLock(&CacheMutationCriticalSection.Get());
 	FCacheElement* CurrentElement = MostRecentElement;
 	while (CurrentElement != nullptr)
 	{
@@ -1306,7 +1306,7 @@ FString FAudioChunkCache::FlushCacheMissLog()
 	TQueue<FCacheMissInfo> BackupQueue;
 
 	FCacheMissInfo CacheMissInfo;
-	while (CacheMissQueue.Dequeue(CacheMissInfo))
+	while (CacheMissQueue->Dequeue(CacheMissInfo))
 	{
 		ConcatenatedCacheMisses.Append(CacheMissInfo.SoundWaveName.ToString());
 		ConcatenatedCacheMisses.Append(TEXT("\t, "));
@@ -1358,7 +1358,7 @@ FString FAudioChunkCache::FlushCacheMissLog()
 	{
 		while (BackupQueue.Dequeue(CacheMissInfo))
 		{
-			CacheMissQueue.Enqueue(CacheMissInfo);
+			CacheMissQueue->Enqueue(CacheMissInfo);
 		}
 	}
 
@@ -1367,7 +1367,7 @@ FString FAudioChunkCache::FlushCacheMissLog()
 
 FAudioChunkCache::FCacheElement* FAudioChunkCache::FindElementForKey(const FChunkKey& InKey)
 {
-	FScopeLock ScopeLock(&CacheMutationCriticalSection);
+	FScopeLock ScopeLock(&CacheMutationCriticalSection.Get());
 
 	const uint64 CacheOffset = GetCacheLookupIDForChunk(InKey);
 
@@ -1460,7 +1460,7 @@ void FAudioChunkCache::TouchElement(FCacheElement* InElement)
 	check(MostRecentElement == nullptr || MostRecentElement != LeastRecentElement);
 	check(InElement->LessRecentElement != InElement);
 
-	FScopeLock ScopeLock(&CacheMutationCriticalSection);
+	FScopeLock ScopeLock(&CacheMutationCriticalSection.Get());
 
 	// If this is already the most recent element, we don't need to do anything.
 	if (InElement == MostRecentElement)
@@ -1519,7 +1519,7 @@ FAudioChunkCache::FCacheElement* FAudioChunkCache::InsertChunk(const FChunkKey&
 	FCacheElement* CacheElement = nullptr;
 
 	{
-		FScopeLock ScopeLock(&CacheMutationCriticalSection);
+		FScopeLock ScopeLock(&CacheMutationCriticalSection.Get());
 
 		if (ShouldAddNewChunk())
 		{
@@ -1598,7 +1598,7 @@ FAudioChunkCache::FCacheElement* FAudioChunkCache::InsertChunk(const FChunkKey&
 
 void FAudioChunkCache::SetUpLeastRecentChunk()
 {
-	FScopeLock ScopeLock(&CacheMutationCriticalSection);
+	FScopeLock ScopeLock(&CacheMutationCriticalSection.Get());
 
 	FCacheElement* CacheElement = MostRecentElement;
 	while (CacheElement->LessRecentElement != nullptr)
@@ -1777,7 +1777,7 @@ void FAudioChunkCache::KickOffAsyncLoad(FCacheElement* CacheElement, const FChun
 	UE_CLOG(CacheElement->bWantsExtraLogging, LogAudioStreamCaching, Display, TEXT("[%s] - KickOffAsyncLoad pri %d outstanding %d chunk index %d, ObjectKey=%s"), *WriteToString<128>(CacheElement->Key.SoundWaveName), (uint32)AsyncIOPriority, NumberOfLoadsInFlight.GetValue(), CacheElement->Key.ChunkIndex, *WriteToString<40>(CacheElement->Key.ObjectKey));
 #endif
 
-	MemoryCounterBytes -= CacheElement->ChunkDataSize;
+	*MemoryCounterBytes -= CacheElement->ChunkDataSize;
 
 	{
 		LLM_SCOPE(ELLMTag::AudioStreamCacheCompressedData);
@@ -1787,7 +1787,7 @@ void FAudioChunkCache::KickOffAsyncLoad(FCacheElement* CacheElement, const FChun
 		CacheElement->ChunkDataSize = Chunk->AudioDataSize;
 	}
 
-	MemoryCounterBytes += CacheElement->ChunkDataSize;
+	*MemoryCounterBytes += CacheElement->ChunkDataSize;
 
 #if DEBUG_STREAM_CACHE
 	CacheElement->DebugInfo.NumTotalChunks = CacheElement->GetNumChunks() - 1;
@@ -2066,7 +2066,7 @@ bool FAudioChunkCache::DoesKeyContainValidChunkIndex(const FChunkKey& InKey, con
 
 uint64 FAudioChunkCache::GetCacheLookupIDForChunk(const FChunkKey& InChunkKey) const
 {
-	FScopeLock Lock(&CacheMutationCriticalSection);
+	FScopeLock Lock(&CacheMutationCriticalSection.Get());
 	const uint64* ID = CacheLookupIdMap.Find(InChunkKey);
 
 	if (ID)
@@ -2081,7 +2081,7 @@ uint64 FAudioChunkCache::GetCacheLookupIDForChunk(const FChunkKey& InChunkKey) c
 
 void FAudioChunkCache::SetCacheLookupIDForChunk(const FChunkKey& InChunkKey, uint64 InCacheLookupID)
 {
-	FScopeLock Lock(&CacheMutationCriticalSection);
+	FScopeLock Lock(&CacheMutationCriticalSection.Get());
 	CacheLookupIdMap.FindOrAdd(InChunkKey, InCacheLookupID);
 }
 
@@ -2172,7 +2172,7 @@ uint64 FCachedAudioStreamingManager::TrimMemory(uint64 NumBytesToFree)
 
 TPair<int, int> FAudioChunkCache::DebugDisplayLegacy(UWorld* World, FViewport* Viewport, FCanvas* Canvas, int32 X, int32 Y, const FVector* ViewLocation, const FRotator* ViewRotation) const
 {
-	FScopeLock ScopeLock(const_cast<FCriticalSection*>(&CacheMutationCriticalSection));
+	FScopeLock ScopeLock(&CacheMutationCriticalSection.Get());
 
 	// Color scheme:
 	static constexpr float ColorMax = 256.0f;
@@ -2410,7 +2410,7 @@ TPair<int, int> FAudioChunkCache::DebugDisplayLegacy(UWorld* World, FViewport* V
 
 FString FAudioChunkCache::DebugPrint()
 {
-	FScopeLock ScopeLock(const_cast<FCriticalSection*>(&CacheMutationCriticalSection));
+	FScopeLock ScopeLock(const_cast<FCriticalSection*>(&CacheMutationCriticalSection.Get()));
 
 	FString OutputString;
 
@@ -2440,15 +2440,15 @@ FString FAudioChunkCache::DebugPrint()
 	}
 
 	// Num bytes in use should include Force Inline data!
-	NumBytesCounter += ForceInlineMemoryCounterBytes;
+	NumBytesCounter += *ForceInlineMemoryCounterBytes;
 
 	// Num bytes should include feature data!
-	NumBytesCounter += FeatureMemoryCounterBytes;
+	NumBytesCounter += *FeatureMemoryCounterBytes;
 
 	// Convert to megabytes and print the total size:
 	const double NumMegabytesInUse = (double)NumBytesCounter / (1024 * 1024);
-	const double NumMegabytesForceInline = (double)ForceInlineMemoryCounterBytes / (1024 * 1024);
-	const double NumMegabytesExternalFeatures = (double)FeatureMemoryCounterBytes / (1024 * 1024);
+	const double NumMegabytesForceInline = (double)*ForceInlineMemoryCounterBytes / (1024 * 1024);
+	const double NumMegabytesExternalFeatures = (double)*FeatureMemoryCounterBytes / (1024 * 1024);
 	const double NumMegabytesRetained = (double)NumBytesRetained / (1024 * 1024);
 
 	const double MaxCacheSizeMB = ((double)MemoryLimitBytes) / (1024 * 1024);
@@ -2458,7 +2458,7 @@ FString FAudioChunkCache::DebugPrint()
 
 	FString CacheMemoryHeader = *FString::Printf(TEXT("External Features:\t, Force Inline:\t, Retaining:\t, Loaded:\t, Max Potential Usage:\t, \n"));
 	FString CacheMemoryUsage = *FString::Printf(TEXT("%.4f Megabytes (%.3f%% of total capacity)\t %.4f Megabytes (%.3f%% of total capacity)\t %.4f Megabytes (%.3f%% of total capacity)\t,  %.4f Megabytes (%" UINT64_FMT " bytes)\t, %.4f Megabytes\t, \n"), 
-		NumMegabytesExternalFeatures, PercentageOfCacheExternalFeatures, NumMegabytesForceInline, PercentageOfCacheForceInlined, NumMegabytesRetained, PercentageOfCacheRetained, NumMegabytesInUse, MemoryCounterBytes.Load(), MaxCacheSizeMB);
+		NumMegabytesExternalFeatures, PercentageOfCacheExternalFeatures, NumMegabytesForceInline, PercentageOfCacheForceInlined, NumMegabytesRetained, PercentageOfCacheRetained, NumMegabytesInUse, MemoryCounterBytes->Load(), MaxCacheSizeMB);
 	OutputString += CacheMemoryHeader + CacheMemoryUsage + TEXT("\n");
 
 	// Second Pass: We're going to list the actual chunks in the cache.
@@ -2563,7 +2563,7 @@ static const FLinearColor ColorExternalFeatures(255 / ColorMax, 100 / ColorMax,
 
 TPair<int, int> FAudioChunkCache::DebugDisplay(UWorld* World, FViewport* Viewport, FCanvas* Canvas, int32 X, int32 Y, const FVector* ViewLocation, const FRotator* ViewRotation) const
 {
-	FScopeLock ScopeLock(const_cast<FCriticalSection*>(&CacheMutationCriticalSection));
+	FScopeLock ScopeLock(&CacheMutationCriticalSection.Get());
 
 	// Draw our header
 	const int32 InitialX = X;
@@ -2684,8 +2684,8 @@ TPair<int, int> FAudioChunkCache::DebugDisplay(UWorld* World, FViewport* Viewpor
 		}
 	}
 
-	uint64 ForceInlineBytes = ForceInlineMemoryCounterBytes.Load();
-	uint64 ExternalFeaturesBytes = FeatureMemoryCounterBytes.Load();
+	uint64 ForceInlineBytes = ForceInlineMemoryCounterBytes->Load();
+	uint64 ExternalFeaturesBytes = FeatureMemoryCounterBytes->Load();
 	NumBytesCounter += ForceInlineBytes;
 	NumBytesCounter += ExternalFeaturesBytes;
 	// Convert to megabytes and print the total size:
@@ -2939,7 +2939,7 @@ TPair<int, int> FAudioChunkCache::DebugDisplay(UWorld* World, FViewport* Viewpor
 
 TPair<int, int> FAudioChunkCache::DebugVisualDisplay(UWorld* World, FViewport* Viewport, FCanvas* Canvas, int32 X, int32 Y, const FVector* ViewLocation, const FRotator* ViewRotation) const
 {
-	FScopeLock ScopeLock(const_cast<FCriticalSection*>(&CacheMutationCriticalSection));
+	FScopeLock ScopeLock(&CacheMutationCriticalSection.Get());
 
 	// Second Pass: We're going to list the actual chunks in the cache.
 	FCacheElement* CurrentElement = MostRecentElement;
diff --git a/Engine/Source/Runtime/Engine/Public/AudioStreamingCache.h b/Engine/Source/Runtime/Engine/Public/AudioStreamingCache.h
--- a/Engine/Source/Runtime/Engine/Public/AudioStreamingCache.h
+++ b/Engine/Source/Runtime/Engine/Public/AudioStreamingCache.h
@@ -18,6 +18,8 @@ AudioStreaming.h: Definitions of classes used for audio streaming.
 #include "AudioStreaming.h"
 #include "Sound/SoundWave.h"
 #include "Sound/SoundWaveLoadingBehavior.h"
+#include "Templates/DontCopy.h"
+#include "Templates/UniqueObj.h"
 #include "UObject/ObjectKey.h"
 
 
@@ -108,7 +110,6 @@ public:
 	};
 
 	FAudioChunkCache(uint32 InMaxChunkSize, uint32 NumChunks, uint64 InMemoryLimitInBytes);
-	
 	~FAudioChunkCache();
 
 	// Places chunk in cache, or puts this chunk back at the top of the cache if it's already loaded. Returns the static lookup ID of the chunk in the cache on success,
@@ -418,19 +419,19 @@ private:
 	int32 ChunksInUse;
 
 	// This counter is used to start evicting chunks before we hit CachePool.Num().
-	TAtomic<uint64> MemoryCounterBytes;
+	TUniqueObj<TAtomic<uint64>> MemoryCounterBytes;
 	uint64 MemoryLimitBytes;
 
-	TAtomic<uint64> ForceInlineMemoryCounterBytes;
+	TUniqueObj<TAtomic<uint64>> ForceInlineMemoryCounterBytes;
 
-	TAtomic<uint64> FeatureMemoryCounterBytes;
+	TUniqueObj<TAtomic<uint64>> FeatureMemoryCounterBytes;
 
 	// Number of async load operations we have currently in flight.
 	FThreadSafeCounter NumberOfLoadsInFlight;
 
 	// Critical section: only used when we are modifying element positions in the cache. This only happens in TouchElement, EvictLeastRecentChunk, and TrimMemory.
 	// Individual cache elements should be thread safe to access.
-	mutable FCriticalSection CacheMutationCriticalSection;
+	mutable TDontCopy<FCriticalSection> CacheMutationCriticalSection;
 
 	// Map that USoundWaves, FSoundWaveProxys, FChunkKeys, and FAudioChunkHandles can use to
 	// quickly lookup where their chunks are currently stored in the cache
@@ -444,8 +445,8 @@ private:
 	// A map to look up the number of times a sound wave has been added for memory tracking
 	// as well as the memory added for the "latest" addition
 	TMap<FSoundWaveProxyPtr, FSoundWaveMemoryTracker> SoundWaveTracker;
-	FCriticalSection SoundWaveMemoryTrackerCritSec;
-	 
+	TDontCopy<FCriticalSection> SoundWaveMemoryTrackerCritSec;
+
 	// This struct is used for logging cache misses.
 	struct FCacheMissInfo
 	{
@@ -456,12 +457,12 @@ private:
 	};
 
 	// This queue is pushed to anytime GetChunk fails to get the chunk. and bLogCacheMisses is true.
-	TQueue<FCacheMissInfo> CacheMissQueue;
+	TUniqueObj<TQueue<FCacheMissInfo>> CacheMissQueue;
 
 	// This is set to true when BeginLoggingCacheMisses is called. 
 	bool bLogCacheMisses;
 
-	uint64 GetCurrentMemoryUsageBytes() const { return MemoryCounterBytes + ForceInlineMemoryCounterBytes + FeatureMemoryCounterBytes; }
+	uint64 GetCurrentMemoryUsageBytes() const { return *MemoryCounterBytes + *ForceInlineMemoryCounterBytes + *FeatureMemoryCounterBytes; }
 
 	// Returns cached element if it exists in our cache, nullptr otherwise.
 	// If the index of the element is already known, it can be used here to avoid searching the cache.
