Run game NetDriver RPCs in transactions when UE_WITH_REMOTE_OBJECT_HANDLE is true.

diff --git a/Engine/Plugins/Runtime/MultiServerReplication/Source/MultiServerReplication/Private/MultiServerBeaconClient.cpp b/Engine/Plugins/Runtime/MultiServerReplication/Source/MultiServerReplication/Private/MultiServerBeaconClient.cpp
--- a/Engine/Plugins/Runtime/MultiServerReplication/Source/MultiServerReplication/Private/MultiServerBeaconClient.cpp
+++ b/Engine/Plugins/Runtime/MultiServerReplication/Source/MultiServerReplication/Private/MultiServerBeaconClient.cpp
@@ -88,6 +88,7 @@ void AMultiServerBeaconClient::ConnectToServer(const FString& ConnectInfo)
 	FURL ConnectURL(NULL, *ConnectInfo, TRAVEL_Absolute);
 	if (InitClient(ConnectURL))
 	{
+		NetDriver->SetReplicateTransactionally(false);
 		bSuccess = true;
 	}
 	else
diff --git a/Engine/Source/Runtime/Engine/Private/DataReplication.cpp b/Engine/Source/Runtime/Engine/Private/DataReplication.cpp
--- a/Engine/Source/Runtime/Engine/Private/DataReplication.cpp
+++ b/Engine/Source/Runtime/Engine/Private/DataReplication.cpp
@@ -27,6 +27,7 @@
 #include "Net/Core/PushModel/Types/PushModelPerNetDriverState.h"
 #include "Net/RPCDoSDetection.h"
 #include "UObject/CoreNetContext.h"
+#include "UObject/RemoteObjectTransfer.h"
 
 DECLARE_LLM_MEMORY_STAT(TEXT("NetObjReplicator"), STAT_NetObjReplicatorLLM, STATGROUP_LLMFULL);
 LLM_DEFINE_TAG(NetObjReplicator, NAME_None, TEXT("Networking"), GET_STATFNAME(STAT_NetObjReplicatorLLM), GET_STATFNAME(STAT_NetworkingSummaryLLM));
@@ -1387,10 +1388,19 @@ bool FObjectReplicator::ReceivedRPC(FNetBitReader& Reader, const FReplicationFla
 				// Call the function.
 				if (Connection->Driver->IsExecuteRPCFunctionsEnabled())
 				{
-					UE::Net::FScopedNetContextRPC CallingRPC;
-					UE::Net::Private::FScopedRemoteRPCMode ReceivingRemoteRPC(Function, UE::Net::Private::ERemoteFunctionMode::Receiving);
-
-					Object->ProcessEvent(Function, Parms);
+#if UE_WITH_REMOTE_OBJECT_HANDLE
+					if (Connection->Driver->bReplicateTransactionally)
+					{
+						UE::RemoteObject::Transfer::ExecuteTransactional([Object, Function, Parms]()
+							{
+								CallProcessEventForReceivedRPC(Object, Function, Parms);
+							});
+					}
+					else
+#endif
+					{
+						CallProcessEventForReceivedRPC(Object, Function, Parms);
+					}
 				}
 			}
 		}
@@ -1418,6 +1428,14 @@ bool FObjectReplicator::ReceivedRPC(FNetBitReader& Reader, const FReplicationFla
 	return true;
 }
 
+void FObjectReplicator::CallProcessEventForReceivedRPC(UObject* Object, UFunction* Function, uint8* Params)
+{
+	UE::Net::FScopedNetContextRPC CallingRPC;
+	UE::Net::Private::FScopedRemoteRPCMode ReceivingRemoteRPC(Function, UE::Net::Private::ERemoteFunctionMode::Receiving);
+
+	Object->ProcessEvent(Function, Params);
+}
+
 void FObjectReplicator::UpdateGuidToReplicatorMap()
 {
 	LLM_SCOPE_BYTAG(NetObjReplicator);
diff --git a/Engine/Source/Runtime/Engine/Public/Net/DataReplication.h b/Engine/Source/Runtime/Engine/Public/Net/DataReplication.h
--- a/Engine/Source/Runtime/Engine/Public/Net/DataReplication.h
+++ b/Engine/Source/Runtime/Engine/Public/Net/DataReplication.h
@@ -320,6 +320,8 @@ public:
 
 private:
 	TWeakObjectPtr<UObject> WeakObjectPtr;
+
+	static void CallProcessEventForReceivedRPC(UObject* Object, UFunction* Function, uint8* Params);
 };
 
 class FScopedActorRoleSwap
