PoseAsset: Add GetLocalSpaceTransform util for retrieving bone transform in local space
#rnx

diff --git a/Engine/Source/Runtime/Engine/Classes/Animation/PoseAsset.h b/Engine/Source/Runtime/Engine/Classes/Animation/PoseAsset.h
--- a/Engine/Source/Runtime/Engine/Classes/Animation/PoseAsset.h
+++ b/Engine/Source/Runtime/Engine/Classes/Animation/PoseAsset.h
@@ -360,6 +360,9 @@ public:
 	// util to return transform of a bone from the pose asset in component space, by walking up tracks in pose asset */
 	ENGINE_API FTransform GetComponentSpaceTransform(FName BoneName, const TArray<FTransform>& LocalTransforms) const;
 
+	// util to return transform of a bone from the pose asset in local space */
+	ENGINE_API const FTransform& GetLocalSpaceTransform(FName BoneName, int32 PoseIndex = 0) const;
+
 	ENGINE_API int32 DeletePoses(TArray<FName> PoseNamesToDelete);
 	ENGINE_API int32 DeleteCurves(TArray<FName> CurveNamesToDelete);
 	ENGINE_API bool ConvertSpace(bool bNewAdditivePose, int32 NewBasePoseInde);
diff --git a/Engine/Source/Runtime/Engine/Private/Animation/PoseAsset.cpp b/Engine/Source/Runtime/Engine/Private/Animation/PoseAsset.cpp
--- a/Engine/Source/Runtime/Engine/Private/Animation/PoseAsset.cpp
+++ b/Engine/Source/Runtime/Engine/Private/Animation/PoseAsset.cpp
@@ -1718,6 +1718,17 @@ FTransform UPoseAsset::GetComponentSpaceTransform(FName BoneName, const TArray<F
 	return ComponentSpaceTransform;
 }
 
+const FTransform& UPoseAsset::GetLocalSpaceTransform(FName BoneName, int32 PoseIndex) const
+{
+	const int32 TrackIdx = GetTrackIndexByName(BoneName);
+	if (!PoseContainer.Poses.IsValidIndex(PoseIndex) || TrackIdx == INDEX_NONE)
+	{
+		UE_LOG(LogAnimation, Error, TEXT("Can't find bone with name %s when trying to get local space transform"), *BoneName.ToString());
+		return FTransform::Identity;
+	}
+	return PoseContainer.Poses[PoseIndex].SourceLocalSpacePose[TrackIdx];
+}
+
 bool UPoseAsset::ConvertSpace(bool bNewAdditivePose, int32 NewBasePoseIndex)
 {
 	// first convert to full pose first
