SVT: Fixed SparseVolumeTexture streaming manager not freeing allocated FIoBuffers at the earliest opportunity, leading to runaway memory allocations.

diff --git a/Engine/Source/Runtime/Engine/Private/SparseVolumeTexture/SparseVolumeTextureStreamingManager.cpp b/Engine/Source/Runtime/Engine/Private/SparseVolumeTexture/SparseVolumeTextureStreamingManager.cpp
--- a/Engine/Source/Runtime/Engine/Private/SparseVolumeTexture/SparseVolumeTextureStreamingManager.cpp
+++ b/Engine/Source/Runtime/Engine/Private/SparseVolumeTexture/SparseVolumeTextureStreamingManager.cpp
@@ -1529,7 +1529,6 @@ void FStreamingManager::InstallReadyRequests()
 
 	UploadTasks.Reset();
 	UploadTasks.Reserve(AsyncState.NumReadyRequests * 2 /*slack for splitting large uploads*/);
-	RequestsToCleanUp.Reset();
 
 #if WITH_EDITORONLY_DATA
 	TMap<const FResources*, const uint8*> ResourceToBulkPointer;
@@ -1632,11 +1631,6 @@ void FStreamingManager::InstallReadyRequests()
 			FrameInfo.ResidentTiles[TileIndex] = true;
 		}
 
-		// Cleanup
-		{
-			RequestsToCleanUp.Add(PendingRequestIndex);
-		}
-
 		InvalidatedSVTFrames.Add(&FrameInfo);
 	}
 
@@ -1665,21 +1659,6 @@ void FStreamingManager::InstallReadyRequests()
 	});
 
 	// Clean up requests
-	for (int32 PendingRequestIndex : RequestsToCleanUp)
-	{
-		FPendingRequest& PendingRequest = PendingRequests[PendingRequestIndex];
-		LOCK_PENDING_REQUEST(PendingRequest);
-#if WITH_EDITORONLY_DATA
-		PendingRequest.SharedBuffer.Reset();
-#endif
-		if (!PendingRequest.Request.IsNone())
-		{
-			check(PendingRequest.Request.IsCompleted());
-			PendingRequest.Request.Reset();
-		}
-	}
-
-#if DO_CHECK // Clear processed pending requests for better debugging
 	for (int32 i = 0; i < AsyncState.NumReadyRequests; ++i)
 	{
 		const int32 PendingRequestIndex = (StartPendingRequestIndex + i) % MaxPendingRequests;
@@ -1687,7 +1666,6 @@ void FStreamingManager::InstallReadyRequests()
 		LOCK_PENDING_REQUEST(PendingRequest);
 		PendingRequest.Reset();
 	}
-#endif
 
 #if WITH_EDITORONLY_DATA
 	// Unlock BulkData
diff --git a/Engine/Source/Runtime/Engine/Private/SparseVolumeTexture/SparseVolumeTextureStreamingManager.h b/Engine/Source/Runtime/Engine/Private/SparseVolumeTexture/SparseVolumeTextureStreamingManager.h
--- a/Engine/Source/Runtime/Engine/Private/SparseVolumeTexture/SparseVolumeTextureStreamingManager.h
+++ b/Engine/Source/Runtime/Engine/Private/SparseVolumeTexture/SparseVolumeTextureStreamingManager.h
@@ -326,7 +326,6 @@ private:
 	TSet<FFrameInfo*> InvalidatedSVTFrames; // Set of SVT frames where tiles have been streamed in or out. Used in PatchPageTable().
 	TArray<FTileRange> TileRangesToStream; // Output of FilterRequests(), consumed in IssueRequests().
 	TArray<FTileDataTask> UploadTasks; // Accessed on the async thread
-	TArray<int32> RequestsToCleanUp; // Accessed on the async thread. Stores indices into PendingRequests.
 	TBitArray<> ResidentPagesNew; // Used as temporary memory in PatchPageTable()
 	TBitArray<> ResidentPagesDiff; // Used as temporary memory in PatchPageTable()
 
