UnrealEngine.cpp and ImageWrapper nop comments
be careful about use of ERGBFormat::RGBA
it is reversed RB order from what FColor uses

diff --git a/Engine/Source/Runtime/Engine/Private/UnrealEngine.cpp b/Engine/Source/Runtime/Engine/Private/UnrealEngine.cpp
--- a/Engine/Source/Runtime/Engine/Private/UnrealEngine.cpp
+++ b/Engine/Source/Runtime/Engine/Private/UnrealEngine.cpp
@@ -1314,6 +1314,10 @@ public:
 	*/
 	virtual bool CompressImage( const TArray< uint8 >& InUncompressedData, const int32 InWidth, const int32 InHeight, TArray< uint8 >& OutCompressedData )
 	{
+		// This incorrectly treats the colors as being RGBA order, not BGRA , eg. NOT FColor
+		//	that is benign as long as CompressImage/DecompressImage do the same thing
+		//	but if you used different functions (eg. ones that decode directly do FImage or assume FColor) then be aware.
+
 		bool bSucceeded = false;
 		OutCompressedData.Reset();
 		if( InUncompressedData.Num() > 0 )
@@ -1343,6 +1347,10 @@ public:
 	*/
 	virtual bool DecompressImage( const TArray< uint8 >& InCompressedData, const int32 InWidth, const int32 InHeight, TArray< uint8 >& OutUncompressedData )
 	{
+		// This incorrectly treats the colors as being RGBA order, not BGRA , eg. NOT FColor
+		//	that is benign as long as CompressImage/DecompressImage do the same thing
+		//	but if you used different functions (eg. ones that decode directly do FImage or assume FColor) then be aware.
+
 		bool bSucceeded = false;
 		OutUncompressedData.Reset();
 		if( InCompressedData.Num() > 0 )
@@ -1402,6 +1410,12 @@ public:
 	*/
 	virtual bool CompressImage( const TArray< uint8 >& InUncompressedData, const int32 InWidth, const int32 InHeight, TArray< uint8 >& OutCompressedData )
 	{
+		// This incorrectly treats the colors as being RGBA order, not BGRA , eg. NOT FColor
+		//	that is benign as long as CompressImage/DecompressImage do the same thing
+		//	but if you used different functions (eg. ones that decode directly do FImage or assume FColor) then be aware.
+
+		// BTW this does make the JPEG slightly less ideal because RB are swapped from how they should be converted to CbCr
+
 		bool bSucceeded = false;
 		OutCompressedData.Reset();
 		if( InUncompressedData.Num() > 0 )
@@ -1432,6 +1446,10 @@ public:
 	*/
 	virtual bool DecompressImage( const TArray< uint8 >& InCompressedData, const int32 InWidth, const int32 InHeight, TArray< uint8 >& OutUncompressedData )
 	{
+		// This incorrectly treats the colors as being RGBA order, not BGRA , eg. NOT FColor
+		//	that is benign as long as CompressImage/DecompressImage do the same thing
+		//	but if you used different functions (eg. ones that decode directly do FImage or assume FColor) then be aware.
+
 		bool bSucceeded = false;
 		OutUncompressedData.Reset();
 		if( InCompressedData.Num() > 0 )
diff --git a/Engine/Source/Runtime/ImageWrapper/Public/IImageWrapper.h b/Engine/Source/Runtime/ImageWrapper/Public/IImageWrapper.h
--- a/Engine/Source/Runtime/ImageWrapper/Public/IImageWrapper.h
+++ b/Engine/Source/Runtime/ImageWrapper/Public/IImageWrapper.h
@@ -76,10 +76,10 @@ enum class ERGBFormat : int8
 {
 	Invalid = -1,
 
-	// Red, Green, Blue and Alpha
+	// Red, Green, Blue and Alpha ; requires RB swap from FColor
 	RGBA =  0,
 
-	// Blue, Green, Red and Alpha
+	// Blue, Green, Red and Alpha ; is ERawImageFormat::BGRA8 and Fcolor
 	BGRA =  1,
 
 	// Gray scale
