Wrapping changes to OnActorSpawned into a CVar:  s.DelayOnActorSpawnedUntilFinishedSpawning.  Default is TRUE to use new (fixed) behavior.

diff --git a/Engine/Source/Runtime/Engine/Private/LevelActor.cpp b/Engine/Source/Runtime/Engine/Private/LevelActor.cpp
--- a/Engine/Source/Runtime/Engine/Private/LevelActor.cpp
+++ b/Engine/Source/Runtime/Engine/Private/LevelActor.cpp
@@ -46,6 +46,17 @@ static TAutoConsoleVariable<int32> CVarAllowDestroyNonNetworkActors(
 	TEXT("When enabled, allows Clients in Networked Games to destroy non-networked actors (AActor::Role == ROLE_None). Does not change behavior on Servers or Standalone games.")
 );
 
+namespace UE::Gameplay::CVars
+{
+	bool bDelayOnActorSpawnedUntilFinishedSpawning = true;
+	static FAutoConsoleVariableRef CVarDelayOnActorSpawnedUntilFinishedSpawning(
+		TEXT("s.DelayOnActorSpawnedUntilFinishedSpawning"),
+		bDelayOnActorSpawnedUntilFinishedSpawning,
+		TEXT("Default: True (new/fixed). Introduced in UE5.6 that UWorld::OnActorSpawned will run only after the Actor has FinishedSpawning (is fully initialized).  False to keep old behavior.")
+	);
+}
+
+
 #define LINE_CHECK_TRACING 0
 
 #if LINE_CHECK_TRACING
@@ -751,6 +762,14 @@ AActor* UWorld::SpawnActor( UClass* Class, FTransform const* UserTransformPtr, c
 		return NULL;
 	}
 
+	// This is too early, we were not guaranteed to call FinishSpawning in PostSpawnInitialize:
+	// CVar in case some ugly case comes up
+	if (!UE::Gameplay::CVars::bDelayOnActorSpawnedUntilFinishedSpawning)
+	{
+		// Broadcast notification of spawn
+		OnActorSpawned.Broadcast(Actor);
+	}
+
 	Actor->CheckDefaultSubobjects();
 
 #if WITH_EDITOR
@@ -765,6 +784,14 @@ AActor* UWorld::SpawnActor( UClass* Class, FTransform const* UserTransformPtr, c
 	}
 #endif
 
+	// This is too early, we were not guaranteed to call FinishSpawning in PostSpawnInitialize:
+	// CVar in case some ugly case comes up
+	if (!UE::Gameplay::CVars::bDelayOnActorSpawnedUntilFinishedSpawning)
+	{
+		// Add this newly spawned actor to the network actor list. Do this after PostSpawnInitialize so that actor has "finished" spawning.
+		AddNetworkActor( Actor );
+	}
+
 	return Actor;
 }
 
diff --git a/Engine/Source/Runtime/Engine/Private/World.cpp b/Engine/Source/Runtime/Engine/Private/World.cpp
--- a/Engine/Source/Runtime/Engine/Private/World.cpp
+++ b/Engine/Source/Runtime/Engine/Private/World.cpp
@@ -212,6 +212,11 @@ namespace UE::Private::World
 	};
 }
 
+namespace UE::Gameplay::CVars
+{
+	extern bool bDelayOnActorSpawnedUntilFinishedSpawning;
+}
+
 /*-----------------------------------------------------------------------------
 	FAdaptiveAddToWorld implementation.
 -----------------------------------------------------------------------------*/
@@ -6014,11 +6019,14 @@ void UWorld::RemoveOnActorSpawnedHandler(FDelegateHandle InHandle) const
 
 void UWorld::OnActorFinishedSpawning(AActor* Actor)
 {
-	// Broadcast notification of spawn
-	OnActorSpawned.Broadcast(Actor);
+	if (UE::Gameplay::CVars::bDelayOnActorSpawnedUntilFinishedSpawning)
+	{
+		// Broadcast notification of spawn
+		OnActorSpawned.Broadcast(Actor);
 
-	// Add this newly spawned actor to the network actor list
-	AddNetworkActor(Actor);
+		// Add this newly spawned actor to the network actor list
+		AddNetworkActor(Actor);
+	}
 }
 
 FDelegateHandle UWorld::AddOnActorPreSpawnInitialization(const FOnActorSpawned::FDelegate& InHandler) const
