Anim bank instance data changes

* Changed instance transform type and packing using FTransform3f instead of FMatrix (double precision), 144/144 byte ->  64/44 byte (memory/disk) .
* Made most properties and data protected & not BP accessible, for BP access we will need to expose accessor fucntions as needed.
* Added more general (not ISMC specific) instance data manager that also supports serializing the ID mapping such that it properly supports undo/redo etc.
* Added Optimize functionality that pre-sorts the instances in the compontent, which means it can be directly sent to the renderer without remapping.
* Implemented support for using these in the AnimBank component  such that it can precompute the instance hierarchy and use async instance data transfer to RT and similar.
* Added tests to avoid creating proxies when async building assets and handling when it is finished (fixing the missing trees issue).

diff --git a/Engine/Plugins/PCG/Source/PCG/Private/Elements/PCGAnimBankMeshSpawner.cpp b/Engine/Plugins/PCG/Source/PCG/Private/Elements/PCGAnimBankMeshSpawner.cpp
--- a/Engine/Plugins/PCG/Source/PCG/Private/Elements/PCGAnimBankMeshSpawner.cpp
+++ b/Engine/Plugins/PCG/Source/PCG/Private/Elements/PCGAnimBankMeshSpawner.cpp
@@ -647,22 +647,18 @@ void FPCGAnimBankMeshSpawnerElement::SpawnAnimBankMeshInstances(FPCGAnimBankMesh
 	const int32 NewInstanceCount = InstanceList.Instances.Num();
 	const int32 NumCustomDataFloats = PackedCustomData.NumCustomDataFloats;
 
-	check((ABMC->NumCustomDataFloats == 0 && PreExistingInstanceCount == 0) || ABMC->NumCustomDataFloats == NumCustomDataFloats);
+	check((ABMC->GetNumCustomDataFloats() == 0 && PreExistingInstanceCount == 0) || ABMC->GetNumCustomDataFloats() == NumCustomDataFloats);
 	ABMC->SetNumCustomDataFloats(NumCustomDataFloats);
 
-	// The index in ABMC PerInstanceSMCustomData where we should pick up to begin inserting new floats
-	const int32 PreviousCustomDataOffset = PreExistingInstanceCount * NumCustomDataFloats;
-
 	// Populate the instances
-	ABMC->AddInstances(InstanceList.Instances, /*bShouldReturnIndices=*/false, /*bWorldSpace=*/true);
+	TArray<FPrimitiveInstanceId> NewIds = ABMC->AddInstances(InstanceList.Instances, /*bShouldReturnIndices=*/true, /*bWorldSpace=*/true);
 
 	// Copy new CustomData into the ABMC PerInstanceSMCustomData
 	if (NumCustomDataFloats > 0)
 	{
-		check(PreviousCustomDataOffset + PackedCustomData.CustomData.Num() == ABMC->InstanceCustomData.Num());
 		for (int32 NewIndex = 0; NewIndex < NewInstanceCount; ++NewIndex)
 		{
-			ABMC->SetCustomData(PreExistingInstanceCount + NewIndex, MakeArrayView(&PackedCustomData.CustomData[NewIndex * NumCustomDataFloats], NumCustomDataFloats));
+			ABMC->SetCustomData(NewIds[NewIndex], MakeArrayView(&PackedCustomData.CustomData[NewIndex * NumCustomDataFloats], NumCustomDataFloats));
 		}
 	}
 
diff --git a/Engine/Plugins/PCG/Source/PCG/Private/Helpers/PCGActorHelpers.cpp b/Engine/Plugins/PCG/Source/PCG/Private/Helpers/PCGActorHelpers.cpp
--- a/Engine/Plugins/PCG/Source/PCG/Private/Helpers/PCGActorHelpers.cpp
+++ b/Engine/Plugins/PCG/Source/PCG/Private/Helpers/PCGActorHelpers.cpp
@@ -291,7 +291,7 @@ UPCGManagedAnimBankMeshComponent* UPCGActorHelpers::GetOrCreateManagedABMC(AActo
 					{
 						if (IsValid(ABMC) &&
 							ABMC->GetOwner() == InTargetActor &&
-							ABMC->NumCustomDataFloats == InParams.NumCustomDataFloats &&
+							ABMC->GetNumCustomDataFloats() == InParams.NumCustomDataFloats &&
 							Resource->GetDescriptor() == Descriptor &&
 							ABMC->HasAnyFlags(RF_Transient) == InParams.bTransient)
 						{
diff --git a/Engine/Source/Runtime/Core/Public/UObject/UE5MainStreamObjectVersion.h b/Engine/Source/Runtime/Core/Public/UObject/UE5MainStreamObjectVersion.h
--- a/Engine/Source/Runtime/Core/Public/UObject/UE5MainStreamObjectVersion.h
+++ b/Engine/Source/Runtime/Core/Public/UObject/UE5MainStreamObjectVersion.h
@@ -372,7 +372,10 @@ struct FUE5MainStreamObjectVersion
 
 		// Visual Logger format now includes a WorldTimeStamp in addition to TimeStamp for easier debugging between multiple instances.
 		VisualLoggerAddedSeparateWorldTime,
-		
+
+		// Added support for InstanceDataManagerSerialization and changed format for the instances to FTransform3f
+		AnimBankInstanceDataSerializationV2,
+	
 		// -----<new versions can be added above this line>-------------------------------------------------
 		VersionPlusOne,
 		LatestVersion = VersionPlusOne - 1
diff --git a/Engine/Source/Runtime/Engine/Classes/Animation/AnimBank.h b/Engine/Source/Runtime/Engine/Classes/Animation/AnimBank.h
--- a/Engine/Source/Runtime/Engine/Classes/Animation/AnimBank.h
+++ b/Engine/Source/Runtime/Engine/Classes/Animation/AnimBank.h
@@ -202,6 +202,9 @@ public:
 
 	inline const FAnimBankData& GetData() const
 	{
+#if WITH_EDITOR
+		check(!IsCompiling());
+#endif
 		return Data;
 	}
 
diff --git a/Engine/Source/Runtime/Engine/Classes/Components/AnimBankMeshComponent.h b/Engine/Source/Runtime/Engine/Classes/Components/AnimBankMeshComponent.h
--- a/Engine/Source/Runtime/Engine/Classes/Components/AnimBankMeshComponent.h
+++ b/Engine/Source/Runtime/Engine/Classes/Components/AnimBankMeshComponent.h
@@ -4,10 +4,10 @@
 
 #include "Components/StaticMeshComponent.h"
 #include "Components/SkeletalMeshComponent.h"
-#include "Components/InstancedStaticMeshComponent.h"
 #include "Rendering/NaniteInterface.h"
 #include "Animation/AnimBank.h"
 #include "InstanceDataSceneProxy.h"
+#include "InstanceData/InstanceDataManager.h"
 #include "AnimBankMeshComponent.generated.h"
 
 class FPrimitiveSceneProxy;
@@ -23,57 +23,44 @@ struct FAnimBankMeshInstanceData
 	GENERATED_USTRUCT_BODY()
 
 	UPROPERTY(EditAnywhere, Category=Instances)
-	FMatrix Transform;
+	FTransform3f Transform;
 
 	UPROPERTY(EditAnywhere, Category = Animation)
 	uint32 BankIndex;
 
-	uint32 Padding[3]; // Need to respect 16 byte alignment for serialization
-
 	FAnimBankMeshInstanceData()
-	: Transform(FMatrix::Identity)
+	: Transform(FTransform3f::Identity)
 	, BankIndex(0)
 	{
-		Padding[0] = 0;
-		Padding[1] = 0;
-		Padding[2] = 0;
 	}
 
-	FAnimBankMeshInstanceData(const FMatrix& InTransform, uint32 InBankIndex)
+	FAnimBankMeshInstanceData(const FTransform3f& InTransform, uint32 InBankIndex)
 	: Transform(InTransform)
 	, BankIndex(InBankIndex)
 	{
-		Padding[0] = 0;
-		Padding[1] = 0;
-		Padding[2] = 0;
 	}
 
 	friend FArchive& operator<<(FArchive& Ar, FAnimBankMeshInstanceData& InstanceData)
 	{
-		// @warning BulkSerialize: FAnimBankMeshInstanceData is serialized as memory dump
-		// See TArray::BulkSerialize for detailed description of implied limitations.
 		Ar << InstanceData.Transform;
 		Ar << InstanceData.BankIndex;
-		Ar << InstanceData.Padding[0];
-		Ar << InstanceData.Padding[1];
-		Ar << InstanceData.Padding[2];
 		return Ar;
 	}
 };
 
-UCLASS(ClassGroup=Rendering, hidecategories=(Object,Activation,Collision,"Components|Activation",Physics), editinlinenew, meta=(BlueprintSpawnableComponent))
-class ENGINE_API UAnimBankMeshComponent : public USkinnedMeshComponent
+UCLASS(ClassGroup=Rendering, hidecategories=(Object,Activation,Collision,"Components|Activation",Physics), editinlinenew, meta=(BlueprintSpawnableComponent), Blueprintable, MinimalAPI)
+class UAnimBankMeshComponent : public USkinnedMeshComponent
 {
-	GENERATED_BODY()
-
-public:
+	GENERATED_UCLASS_BODY()
+	
+protected:
 	/** Array of instances, bulk serialized. */
 	UPROPERTY(EditAnywhere, SkipSerialization, DisplayName="Instances", Category=Instances, meta=(MakeEditWidget=true, EditFixedOrder))
 	TArray<FAnimBankMeshInstanceData> InstanceData;
 
 	/** Defines the number of floats that will be available per instance for custom data */
 	UPROPERTY(EditAnywhere, Category=Instances, AdvancedDisplay)
-	int32 NumCustomDataFloats;
+	int32 NumCustomDataFloats = 0;
 
 	/** Array of custom data for instances. This will contains NumCustomDataFloats*InstanceCount entries. The entries are represented sequantially, in instance order. Can be read in a material and manipulated through Blueprints.
 	*	Example: If NumCustomDataFloats is 1, then each entry will belong to an instance. Custom data 0 will belong to Instance 0. Custom data 1 will belong to Instance 1 etc.
@@ -82,143 +69,147 @@ public:
 	UPROPERTY(EditAnywhere, EditFixedSize, SkipSerialization, DisplayName="Custom Data", Category=Instances, AdvancedDisplay, meta=(EditFixedOrder))
 	TArray<float> InstanceCustomData;
 
-	UPROPERTY(EditAnywhere, BlueprintReadWrite, Category = Animation, meta = (ShowOnlyInnerProperties))
-	TArray<struct FAnimBankItem> AnimBankItems;
-
 	/** Screen space footprint (in primary view) cutoff which dictates the far distance the instance will play back animation. 
 	 *  Using 0.0 (the default) falls back to a global threshold, a negative value disables the cutoff.
 	 */
-	UPROPERTY(EditAnywhere, BlueprintReadWrite, Category = Animation)
+	UPROPERTY(EditAnywhere, BlueprintReadOnly, Category = Animation)
 	float AnimationMinScreenSize = 0.0f;
 
 	/** If true, this component will avoid serializing its per instance data / those properties will also not be editable */
 	UPROPERTY()
 	uint8 bInheritPerInstanceData : 1;
 
+	UPROPERTY(EditAnywhere, BlueprintReadOnly, Category = Animation, meta = (ShowOnlyInnerProperties))
+	TArray<struct FAnimBankItem> AnimBankItems;
+
 public:
-	UAnimBankMeshComponent(FVTableHelper& Helper);
+
+	ENGINE_API UAnimBankMeshComponent(FVTableHelper& Helper);
 	virtual ~UAnimBankMeshComponent();
 
+	//~ Begin UObject Interface
 	virtual void Serialize(FArchive& Ar) override;
-
 	virtual void PostLoad() override;
-
 	virtual void BeginDestroy() override;
-
-	virtual void TickComponent(float DeltaTime, enum ELevelTick TickType, FActorComponentTickFunction* ThisTickFunction) override;
-	virtual void OnRegister() override;
-
-	virtual const Nanite::FResources* GetNaniteResources() const;
-
-	virtual void CreateRenderState_Concurrent(FRegisterComponentContext* Context) override;
-	virtual void DestroyRenderState_Concurrent() override;
-
-	virtual FPrimitiveSceneProxy* CreateSceneProxy() override;
-
 #if WITH_EDITOR
 	virtual void PostEditChangeChainProperty(FPropertyChangedChainEvent& PropertyChangedEvent) override;
 	virtual void PostEditUndo() override;
+	virtual void BeginCacheForCookedPlatformData(const ITargetPlatform* TargetPlatform) override;
+	virtual bool IsCachedCookedPlatformDataLoaded( const ITargetPlatform* TargetPlatform ) override;
 #endif
+	//~ End UObject Interface
 
 	//~ Begin UActorComponent Interface
-	virtual TStructOnScope<FActorComponentInstanceData> GetComponentInstanceData() const override;
-	//~ End UActorComponent Interface
+	virtual void OnRegister() override;
+	virtual void SendRenderInstanceData_Concurrent() override;
+	virtual void TickComponent(float DeltaTime, enum ELevelTick TickType, FActorComponentTickFunction* ThisTickFunction) override;
 
-	/** Handle changes that must happen before the proxy is recreated. */
-	void PreApplyComponentInstanceData(struct FAnimBankMeshComponentInstanceData* ComponentInstanceData);
+	virtual void CreateRenderState_Concurrent(FRegisterComponentContext* Context) override;
+	virtual void DestroyRenderState_Concurrent() override;
 
-	/** Applies the cached component instance data to a newly blueprint constructed component. */
-	virtual void ApplyComponentInstanceData(struct FAnimBankMeshComponentInstanceData* ComponentInstanceData);
+	virtual TStructOnScope<FActorComponentInstanceData> GetComponentInstanceData() const override;
+	//~ End UActorComponent Interface
 
+	//~ Begin USceneComponent Interface
 	virtual FBoxSphereBounds CalcBounds(const FTransform& LocalToWorld) const override;
+	virtual FPrimitiveSceneProxy* CreateSceneProxy() override;
+	virtual void OnUpdateTransform(EUpdateTransformFlags UpdateTransformFlags, ETeleportType Teleport) override;
+	//~ End USceneComponent Interface
 
+	//~ Begin USkinnedMeshComponent Interface
+	virtual const Nanite::FResources* GetNaniteResources() const;
 	virtual void RefreshBoneTransforms(FActorComponentTickFunction* TickFunction = nullptr) override;
+#if WITH_EDITOR
+	/**
+	 * Called by the asset compilers (FAnimBankCompilingManager / FSkinnedAssetCompilingManager) when done to allow render state to be created.
+	 */
+	virtual void PostAssetCompilation();
+
+#endif
+	//~ End USkinnedMeshComponent Interface
+
+public:
 
 	/** Add an instance to this component. Transform is given in local space of this component unless bWorldSpace is set. */
 	UFUNCTION(BlueprintCallable, Category = "Components|AnimBankMesh")
-	virtual int32 AddInstance(const FTransform& InstanceTransform, int32 BankIndex, bool bWorldSpace = false);
+	ENGINE_API FPrimitiveInstanceId AddInstance(const FTransform& InstanceTransform, int32 BankIndex, bool bWorldSpace = false);
 
 	/** Add multiple instances to this component. Transform is given in local space of this component unless bWorldSpace is set. */
 	UFUNCTION(BlueprintCallable, Category = "Components|AnimBankMesh")
-	virtual TArray<int32> AddInstances(const TArray<FTransform>& InstanceTransforms, bool bShouldReturnIndices, bool bWorldSpace = false, bool bUpdateNavigation = true);
+	ENGINE_API TArray<FPrimitiveInstanceId> AddInstances(const TArray<FTransform>& InstanceTransforms, bool bShouldReturnIds, bool bWorldSpace = false);
 
 	/** Update custom data for specific instance */
 	UFUNCTION(BlueprintCallable, Category = "Components|AnimBankMesh")
-	virtual bool SetCustomDataValue(int32 InstanceIndex, int32 CustomDataIndex, float CustomDataValue, bool bMarkRenderStateDirty = false);
+	ENGINE_API bool SetCustomDataValue(FPrimitiveInstanceId InstanceId, int32 CustomDataIndex, float CustomDataValue);
+
+	/** Update all custom data values for specific instance, the size of the array view must match the NumCustomDataFloats */
+	ENGINE_API bool SetCustomData(FPrimitiveInstanceId InstanceId, TArrayView<const float> CustomDataFloats);
 
 	/** Update number of custom data entries per instance. This applies to all instances and will reallocate the full custom data buffer and reset all values to 0 */
 	UFUNCTION(BlueprintCallable, Category = "Components|AnimBankMesh")
-	virtual void SetNumCustomDataFloats(int32 InNumCustomDataFloats);
-
-	/** Per Instance Custom Data */
-	virtual bool SetCustomData(int32 InstanceIndex, TArrayView<const float> CustomDataFloats, bool bMarkRenderStateDirty = false);
-
-	/** Preallocated memory to include the new added instances count, to prevent reallloc during the add operation. */
-	virtual void PreAllocateInstancesMemory(int32 AddedInstanceCount);
+	ENGINE_API void SetNumCustomDataFloats(int32 InNumCustomDataFloats);
 
 	/** Get the transform for the instance specified. Instance is returned in local space of this component unless bWorldSpace is set.  Returns True on success. */
 	UFUNCTION(BlueprintCallable, Category = "Components|AnimBankMesh")
-	bool GetInstanceTransform(int32 InstanceIndex, FTransform& OutInstanceTransform, bool bWorldSpace = false) const;
+	ENGINE_API bool GetInstanceTransform(FPrimitiveInstanceId InstanceId, FTransform& OutInstanceTransform, bool bWorldSpace = false) const;
 
 	/** Remove the instance specified. Returns True on success. */
 	UFUNCTION(BlueprintCallable, Category = "Components|AnimBankMesh")
-	virtual bool RemoveInstance(int32 InstanceIndex);
+	ENGINE_API bool RemoveInstance(FPrimitiveInstanceId InstanceId);
 
 	/** Remove the instances specified. Returns True on success. */
 	UFUNCTION(BlueprintCallable, Category = "Components|AnimBankMesh")
-	virtual bool RemoveInstances(const TArray<int32>& InstancesToRemove);
-
-	/** Remove the instances specified. Returns True on success.
-	* @param InstancesToRemove array of the instance indices to remove ( can be sorted @see bInstanceArrayAlreadySortedInReverseOrder )
-	* @param bInstanceArrayAlreadySortedInReverseOrder true is the array is already sorted in reverse order
-	*/
-	virtual bool RemoveInstances(const TArray<int32>& InstancesToRemove, bool bInstanceArrayAlreadySortedInReverseOrder);
+	ENGINE_API void RemoveInstances(const TArray<FPrimitiveInstanceId>& InstancesToRemove);
 
 	/** Clear all instances being rendered by this component. */
 	UFUNCTION(BlueprintCallable, Category = "Components|AnimBankMesh")
-	virtual void ClearInstances();
+	ENGINE_API void ClearInstances();
 
 	UFUNCTION(BlueprintCallable, Category = "Components|AnimBankMesh")
-	virtual int32 AddAnimBankItem(const FAnimBankItem& BankItem);
+	ENGINE_API int32 AddAnimBankItem(const FAnimBankItem& BankItem);
+
+	/**
+	 * Optimize the instance data by performing sorting according to spatial hash on the _source_ data.
+	 * Note that this reorders the instances and thus any indexing will change. By default resets the ID mapping to identity.
+	 *  @param bShouldRetainIdMap	If true, the id mapping is updated instead of reset to identity, this retains the validity of the IDs but adds some memory and storage cost (for the ID mapping).
+	 */
+	UFUNCTION(BlueprintCallable, Category = "Components|AnimBankMesh")
+	ENGINE_API void OptimizeInstanceData(bool bShouldRetainIdMap = false);
+
+	ENGINE_API bool IsEnabled() const;
+	ENGINE_API int32 GetInstanceCount() const;
 
-	bool IsEnabled() const;
-	int32 GetInstanceCount() const;
+	ENGINE_API TArrayView<const struct FAnimBankItem> GetAnimBankItems() const;
+
+	ENGINE_API void SetAnimBankItems(TArrayView<const struct FAnimBankItem> InAnimBankItems);
+
+	int32 GetNumCustomDataFloats() const { return NumCustomDataFloats; }
+protected:
+
+	/** Handle changes that must happen before the proxy is recreated. */
+	void PreApplyComponentInstanceData(struct FAnimBankMeshComponentInstanceData* ComponentInstanceData);
+
+	/** Applies the cached component instance data to a newly blueprint constructed component. */
+	void ApplyComponentInstanceData(struct FAnimBankMeshComponentInstanceData* ComponentInstanceData);
+
+	FInstanceDataManagerSourceDataDesc GetComponentDesc(ERHIFeatureLevel::Type FeatureLevel);
 
 #if WITH_EDITOR
 	/** One bit per instance if the instance is selected. */
 	TBitArray<> SelectedInstances;
-
-	/** Indicates that the user has purposely chosen to show the instance list in the details panel, despite the performance warning. */
-	bool bForceShowAllInstancesDetails = false;
-
-	/** The reason why a deletion operation is currently happening. */
-	enum class EInstanceDeletionReason : uint8
-	{
-		NotDeleting, /** There is currently no deletion in progress. */
-		EntryAlreadyRemoved, /** The instance has been deleted externally. Data synchronization in progress. */
-		EntryRemoval, /** The instance is being removed. */
-		Clearing /** All instances are being removed. */
-	};
-
-	/** This will be set to the appropriate state when one or more instances are in the process of being
-	 *  deleted. This is primarily used for functions that round trip to this class, such as callbacks
-	 *  for deselecting instances. */
-	EInstanceDeletionReason DeletionState = EInstanceDeletionReason::NotDeleting;
 #endif
 
-protected:
 	bool bIsInstanceDataApplyCompleted = true;
 
 	void CreateHitProxyData(TArray<TRefCountPtr<HHitProxy>>& HitProxies);
 
-	/** Internal version of AddInstance */
-	int32 AddInstanceInternal(int32 InstanceIndex, FAnimBankMeshInstanceData* InNewInstanceData, const FTransform& InstanceTransform, int32 BankIndex, bool bWorldSpace);
-
-	/** Internal implementation of AddInstances */
-	TArray<int32> AddInstancesInternal(TConstArrayView<FTransform> InstanceTransforms, bool bShouldReturnIndices, bool bWorldSpace, bool bUpdateNavigation = true);
+	/** 
+	 */
+	FPrimitiveInstanceId AddInstanceInternal(int32 InstanceIndex, const FTransform& InstanceTransform, int32 BankIndex, bool bWorldSpace);
 
-	/** Internal version of RemoveInstance */
-	bool RemoveInstanceInternal(int32 InstanceIndex, bool InstanceAlreadyRemoved, bool bForceRemoveAtSwap = false, bool bUpdateNavigation = true);
+	/** 
+	 */
+	bool RemoveInstanceInternal(int32 InstanceIndex, bool InstanceAlreadyRemoved);
 
 private:
 	void ApplyInheritedPerInstanceData(const UAnimBankMeshComponent* InArchetype);
@@ -226,8 +217,13 @@ private:
 	bool ShouldInheritPerInstanceData() const;
 
 	/** Sets up new instance data to sensible defaults, creates physics counterparts if possible. */
-	void SetupNewInstanceData(FAnimBankMeshInstanceData& InOutNewInstanceData, int32 InInstanceIndex, const FTransform& InInstanceTransform, int32 BankIndex);
+	void SetupNewInstanceData(FAnimBankMeshInstanceData& InOutNewInstanceData, int32 InInstanceIndex, const FTransform3f& InInstanceTransform, int32 BankIndex);
+
+	FInstanceDataManager InstanceDataManager;
 
+	friend class FAnimBankMeshSceneProxy;
+	friend class FAnimBankMeshObject;
+	friend struct FAnimBankMeshComponentInstanceData;
 private:
 	void SetSkinnedAssetCallback();
 };
diff --git a/Engine/Source/Runtime/Engine/Private/Animation/AnimBank.cpp b/Engine/Source/Runtime/Engine/Private/Animation/AnimBank.cpp
--- a/Engine/Source/Runtime/Engine/Private/Animation/AnimBank.cpp
+++ b/Engine/Source/Runtime/Engine/Private/Animation/AnimBank.cpp
@@ -1081,6 +1081,12 @@ FIoHash UAnimBank::BeginCacheDerivedData(const ITargetPlatform* TargetPlatform)
 		return KeyHash;
 	}
 
+	// If nothing has changed and we already started the compilation we should be all good.
+	if (CacheTasksByKeyHash.Contains(KeyHash))
+	{
+		return KeyHash;
+	}
+
 	// Make sure we finish the previous build before starting another one
 	FAnimBankCompilingManager::Get().FinishCompilation({ this });
 
@@ -1425,7 +1431,7 @@ void FAnimBankMeshComponentDescriptor::InitFrom(const UAnimBankMeshComponent* Te
 	ComponentClass = Template->GetClass();
 	Mobility = Template->Mobility;
 	SkinnedAsset = Template->GetSkinnedAsset();
-	BankItems = Template->AnimBankItems;
+	BankItems = Template->GetAnimBankItems();
 }
 
 uint32 FAnimBankMeshComponentDescriptor::ComputeHash() const
@@ -1444,7 +1450,7 @@ void FAnimBankMeshComponentDescriptor::InitComponent(UAnimBankMeshComponent* Com
 {
 	Component->SetSkinnedAsset(SkinnedAsset);
 	Component->Mobility = Mobility;
-	Component->AnimBankItems = BankItems;
+	Component->SetAnimBankItems(BankItems);
 }
 
 void FAnimBankMeshComponentDescriptor::PostLoadFixup(UObject* Loader)
@@ -1501,11 +1507,7 @@ void FSoftAnimBankMeshComponentDescriptor::InitFrom(const UAnimBankMeshComponent
 	ComponentClass = Template->GetClass();
 	Mobility = Template->Mobility;
 	SkinnedAsset = Template->GetSkinnedAsset();
-	BankItems.Reset(Template->AnimBankItems.Num());
-	for (const FAnimBankItem& BankItem : Template->AnimBankItems)
-	{
-		BankItems.Emplace(BankItem);
-	}
+	BankItems = Template->GetAnimBankItems();
 }
 
 uint32 FSoftAnimBankMeshComponentDescriptor::ComputeHash() const
@@ -1523,11 +1525,12 @@ void FSoftAnimBankMeshComponentDescriptor::InitComponent(UAnimBankMeshComponent*
 {
 	Component->SetSkinnedAsset(SkinnedAsset.LoadSynchronous());
 	Component->Mobility = Mobility;
-	Component->AnimBankItems.Reset(BankItems.Num());
+	TArray<FAnimBankItem> AnimBankItems;
 	for (const FSoftAnimBankItem& BankItem : BankItems)
 	{
-		Component->AnimBankItems.Emplace(BankItem);
+		AnimBankItems.Emplace(BankItem);
 	}
+	Component->SetAnimBankItems(AnimBankItems);
 }
 
 void FSoftAnimBankMeshComponentDescriptor::PostLoadFixup(UObject* Loader)
diff --git a/Engine/Source/Runtime/Engine/Private/Components/AnimBankMeshComponent.cpp b/Engine/Source/Runtime/Engine/Private/Components/AnimBankMeshComponent.cpp
--- a/Engine/Source/Runtime/Engine/Private/Components/AnimBankMeshComponent.cpp
+++ b/Engine/Source/Runtime/Engine/Private/Components/AnimBankMeshComponent.cpp
@@ -18,6 +18,8 @@
 #include "SkeletalRenderPublic.h"
 #include "SkinningDefinitions.h"
 #include "SkinnedMeshSceneProxyDesc.h"
+#include "InstanceData/InstanceUpdateChangeSet.h"
+#include "InstanceData/InstanceDataUpdateUtils.h"
 
 #include UE_INLINE_GENERATED_CPP_BY_NAME(AnimBankMeshComponent)
 
@@ -62,76 +64,11 @@ public:
 			}
 		}
 	#endif
-
-		MeshBounds = InComponent->GetSkinnedAsset()->GetBounds();
-
-		FInstanceSceneDataBuffers::FAccessTag AccessTag(PointerHash(this));
-		FInstanceSceneDataBuffers::FWriteView ProxyData = InstanceSceneDataBuffersImpl.BeginWriteAccess(AccessTag);
-
 		const bool bForceRefPose = CVarAnimBankForceRefPose.GetValueOnAnyThread() != 0;
 		const bool bUseAnimBank = !bForceRefPose && InComponent->AnimBankItems.Num() > 0;
 
-		ProxyData.Flags.bHasPerInstanceHierarchyOffset = false;
-		ProxyData.Flags.bHasPerInstanceLocalBounds = bUseAnimBank && InComponent->AnimBankItems.Num() > 1;
-		ProxyData.Flags.bHasPerInstanceDynamicData = false;
-		ProxyData.Flags.bHasPerInstanceSkinningData = true;
-	#if WITH_EDITOR
-		ProxyData.Flags.bHasPerInstanceEditorData = false;// GIsEditor != 0 && HasPerInstanceHitProxies();
-	#endif
-
-		const int32 InstanceCount = InComponent->GetInstanceCount();
-
-		InstanceTransforms.Reset(InstanceCount);
-		if (ProxyData.Flags.bHasPerInstanceLocalBounds)
-		{
-			InstanceBounds.Reset(InstanceCount);
-		}
-		else if (bUseAnimBank && InComponent->AnimBankItems.Num() == 1)
-		{
-			// No use wasting memory/perf with per-instance local bounds if they all point to the same anim bank sequence
-			const FAnimBankItem& BankItem = InComponent->AnimBankItems[0];
-			if (BankItem.BankAsset != nullptr && CVarAnimBankSampledBounds.GetValueOnAnyThread() != 0)
-			{
-				const FAnimBankData& BankData = BankItem.BankAsset->GetData();
-				if (BankItem.SequenceIndex < BankData.Entries.Num())
-				{
-					MeshBounds = BankData.Entries[BankItem.SequenceIndex].SampledBounds;
-				}
-			}
-		}
-
-		ProxyData.InstanceSkinningData.SetNumUninitialized(InstanceCount);
-
-		for (int32 InstanceIndex = 0; InstanceIndex < InstanceCount; ++InstanceIndex)
-		{
-			const FAnimBankMeshInstanceData& Instance = InComponent->InstanceData[InstanceIndex];
-
-			if (ProxyData.Flags.bHasPerInstanceLocalBounds)
-			{
-				FRenderBounds& SampledBounds = InstanceBounds.Emplace_GetRef(MeshBounds);
-
-				if (Instance.BankIndex < uint32(InComponent->AnimBankItems.Num()))
-				{
-					const FAnimBankItem& BankItem = InComponent->AnimBankItems[Instance.BankIndex];
-					if (BankItem.BankAsset != nullptr)
-					{
-						const FAnimBankData& BankData = BankItem.BankAsset->GetData();
-						if (BankItem.SequenceIndex < BankData.Entries.Num())
-						{
-							SampledBounds = BankData.Entries[BankItem.SequenceIndex].SampledBounds;
-						}
-					}
-				}
-			}
-
-			InstanceTransforms.Emplace(FRenderTransform(Instance.Transform));
-			ProxyData.InstanceSkinningData[InstanceIndex] = Instance.BankIndex * MaxBoneTransformCount * 2u /* Current and Previous */;
-		}
-
-		InstanceSceneDataBuffersImpl.EndWriteAccess(AccessTag);
-
-		// Nanite supports the GPUScene instance data buffer.
-		SetupInstanceSceneDataBuffers(&InstanceSceneDataBuffersImpl);
+		InstanceDataSceneProxy = InComponent->InstanceDataManager.GetOrCreateProxy();
+		SetupInstanceSceneDataBuffers(InstanceDataSceneProxy->GeInstanceSceneDataBuffers());
 
 		// Note: ideally this would be picked up from the Flags.bHasPerInstanceDynamicData above, but that path is not great at the moment.
 		bAlwaysHasVelocity = true;
@@ -164,62 +101,6 @@ public:
 		}
 	}
 
-	virtual ~FAnimBankMeshSceneProxy()
-	{
-	}
-
-	virtual void OnTransformChanged(FRHICommandListBase& RHICmdList) override
-	{
-		Super::OnTransformChanged(RHICmdList);
-
-		FInstanceSceneDataBuffers::FAccessTag AccessTag(PointerHash(this));
-		FInstanceSceneDataBuffers::FWriteView ProxyData = InstanceSceneDataBuffersImpl.BeginWriteAccess(AccessTag);
-
-		InstanceSceneDataBuffersImpl.SetPrimitiveLocalToWorld(GetLocalToWorld(), AccessTag);
-
-		const FRenderTransform& PrimitiveToRelativeWorld = InstanceSceneDataBuffersImpl.GetPrimitiveToRelativeWorld(AccessTag);
-
-		const int32 InstanceCount = InstanceTransforms.Num();
-		ProxyData.InstanceToPrimitiveRelative.Reset(InstanceCount);
-		if (ProxyData.Flags.bHasPerInstanceLocalBounds)
-		{
-			ProxyData.InstanceLocalBounds.Reset(InstanceCount);
-		}
-
-		for (int32 InstanceIndex = 0; InstanceIndex < InstanceCount; ++InstanceIndex)
-		{
-			FRenderTransform InstanceToPrimitiveRelativeWorld = InstanceTransforms[InstanceIndex] * PrimitiveToRelativeWorld;
-			
-			// Remove shear
-			InstanceToPrimitiveRelativeWorld.Orthogonalize();
-
-			//FRenderTransform& PrevInstanceToPrimitiveRelative = ProxyData.PrevInstanceToPrimitiveRelative.Emplace_GetRef();
-			//PrevInstanceToPrimitiveRelative = InstanceToPrimitiveRelative;
-
-			ProxyData.InstanceToPrimitiveRelative.Emplace(InstanceToPrimitiveRelativeWorld);
-
-			if (ProxyData.Flags.bHasPerInstanceLocalBounds)
-			{
-				ProxyData.InstanceLocalBounds.Emplace(InstanceBounds[InstanceIndex]);
-			}
-		}
-
-		if (!ProxyData.Flags.bHasPerInstanceLocalBounds)
-		{
-			if (InstanceCount > 0)
-			{
-				ProxyData.InstanceLocalBounds.SetNum(1);
-				ProxyData.InstanceLocalBounds[0] = MeshBounds;
-			}
-			else
-			{
-				ProxyData.InstanceLocalBounds.SetNum(0);
-			}
-		}
-
-		InstanceSceneDataBuffersImpl.EndWriteAccess(AccessTag);
-	}
-
 	virtual void CreateRenderThreadResources(FRHICommandListBase& RHICmdList) override
 	{
 		Super::CreateRenderThreadResources(RHICmdList);
@@ -272,13 +153,12 @@ public:
 		Super::DestroyRenderThreadResources();
 	}
 
+	// FSkinnedSceneProxy interface
 	virtual const TConstArrayView<uint64> GetAnimationProviderData(bool& bOutValid) const override
 	{
 		bOutValid = (AnimBankIds.Num() == UniqueAnimationCount);
 		return AnimBankIds;
 	}
-
-	// FSkinnedSceneProxy interface
 	virtual float GetAnimationMinScreenSize() const  override { return AnimationMinScreenSize; }
 
 public:
@@ -290,12 +170,7 @@ public:
 	}
 
 private:
-	FInstanceSceneDataBuffers InstanceSceneDataBuffersImpl;
-
-	TArray<FRenderTransform> InstanceTransforms;
-	TArray<FRenderBounds> InstanceBounds;
-	FRenderBounds MeshBounds;
-
+	TSharedPtr<FInstanceDataSceneProxy, ESPMode::ThreadSafe> InstanceDataSceneProxy; 
 	TArray<uint64> AnimBankIds;
 	TArray<FAnimBankRecordHandle> AnimBankHandles;
 	TArray<FAnimBankItem> AnimBankItems;
@@ -430,7 +305,7 @@ public:
 	{
 		for (int32 LODIndex = 0; LODIndex < InRenderData->LODRenderData.Num(); ++LODIndex)
 		{
-			new(LODs) FAnimBankMeshObjectLOD(InFeatureLevel, InRenderData, LODIndex);
+			LODs.Emplace(InFeatureLevel, InRenderData, LODIndex);
 		}
 
 		InitResources(InMeshDesc);
@@ -443,7 +318,6 @@ public:
 
 	virtual ~FAnimBankMeshObject()
 	{
-		delete DynamicData;
 	}
 
 	virtual void InitResources(const FSkinnedMeshSceneProxyDesc& InMeshDesc) override
@@ -487,7 +361,7 @@ public:
 	{
 		// Create the new dynamic data for use by the rendering thread
 		// this data is only deleted when another update is sent
-		FAnimBankMeshData* NewDynamicData = new FAnimBankMeshData(InSkinnedAsset, AnimBankItems, LODIndex, PreviousBoneTransformUpdateMode);
+		TUniquePtr<FAnimBankMeshData> NewDynamicData = MakeUnique<FAnimBankMeshData>(InSkinnedAsset, AnimBankItems, LODIndex, PreviousBoneTransformUpdateMode);
 
 		uint64 FrameNumberToPrepare = GFrameCounter;
 		uint32 RevisionNumber = 0;
@@ -501,10 +375,10 @@ public:
 		{
 			FAnimBankMeshObject* MeshObject = this;
 			ENQUEUE_RENDER_COMMAND(SkelMeshObjectUpdateDataCommand)(UE::RenderCommandPipe::SkeletalMesh,
-				[MeshObject, FrameNumberToPrepare, RevisionNumber, NewDynamicData](FRHICommandList& RHICmdList)
+				[MeshObject, FrameNumberToPrepare, RevisionNumber, NewDynamicData=MoveTemp(NewDynamicData)](FRHICommandList& RHICmdList) mutable
 				{
 					FScopeCycleCounter Context(MeshObject->GetStatId());
-					MeshObject->UpdateDynamicData_RenderThread(RHICmdList, NewDynamicData, FrameNumberToPrepare, RevisionNumber);
+					MeshObject->UpdateDynamicData_RenderThread(RHICmdList, MoveTemp(NewDynamicData), FrameNumberToPrepare, RevisionNumber);
 				}
 			);
 		}
@@ -512,16 +386,13 @@ public:
 
 	void UpdateDynamicData_RenderThread(
 		FRHICommandList& RHICmdList,
-		FAnimBankMeshData* InDynamicData,
+		TUniquePtr<FAnimBankMeshData>&& InDynamicData,
 		uint64 FrameNumberToPrepare,
 		uint32 RevisionNumber
 	)
 	{
-		// We should be done with the old data at this point
-		delete DynamicData;
-
 		// Update with new data
-		DynamicData = InDynamicData;
+		DynamicData = MoveTemp(InDynamicData);
 		check(DynamicData);
 
 		check(IsInParallelRenderingThread());
@@ -635,7 +506,7 @@ public:
 	virtual bool IsNaniteMesh() const override { return true; }
 
 private:
-	FAnimBankMeshData* DynamicData = nullptr;
+	TUniquePtr<FAnimBankMeshData> DynamicData;
 
 	struct FAnimBankMeshObjectLOD
 	{
@@ -704,21 +575,70 @@ static FSkeletalMeshObject* CreateAnimBankMeshObjectFn(void* UserData, USkinnedM
 
 UAnimBankMeshComponent::UAnimBankMeshComponent(FVTableHelper& Helper)
 : Super(Helper)
-, NumCustomDataFloats(0)
 , bInheritPerInstanceData(false)
+, InstanceDataManager(this)
 {
 }
 
+UAnimBankMeshComponent::UAnimBankMeshComponent(const FObjectInitializer& ObjectInitializer) 
+: Super(ObjectInitializer) 
+, bInheritPerInstanceData(false)
+, InstanceDataManager(this)
+{
+}
+
+
 UAnimBankMeshComponent::~UAnimBankMeshComponent()
 {
 
 }
 
+
+
+struct FAnimBankMeshInstanceData_Deprecated
+{
+	FMatrix Transform;
+	uint32 BankIndex;
+	uint32 Padding[3]; // Need to respect 16 byte alignment for bulk-serialization
+
+	FAnimBankMeshInstanceData_Deprecated()
+	: Transform(FMatrix::Identity)
+	, BankIndex(0)
+	{
+		Padding[0] = 0;
+		Padding[1] = 0;
+		Padding[2] = 0;
+	}
+
+	FAnimBankMeshInstanceData_Deprecated(const FMatrix& InTransform, uint32 InBankIndex)
+	: Transform(InTransform)
+	, BankIndex(InBankIndex)
+	{
+		Padding[0] = 0;
+		Padding[1] = 0;
+		Padding[2] = 0;
+	}
+
+	friend FArchive& operator<<(FArchive& Ar, FAnimBankMeshInstanceData_Deprecated& InstanceData)
+	{
+		// @warning BulkSerialize: FAnimBankMeshInstanceData is serialized as memory dump
+		// See TArray::BulkSerialize for detailed description of implied limitations.
+		Ar << InstanceData.Transform;
+		Ar << InstanceData.BankIndex;
+		Ar << InstanceData.Padding[0];
+		Ar << InstanceData.Padding[1];
+		Ar << InstanceData.Padding[2];
+		return Ar;
+	}
+};
+
 void UAnimBankMeshComponent::Serialize(FArchive& Ar)
 {
 	LLM_SCOPE(ELLMTag::SkeletalMesh);
 	Super::Serialize(Ar);
 
+	Ar.UsingCustomVersion(FUE5MainStreamObjectVersion::GUID);
+
 	bool bCooked = Ar.IsCooking();
 	Ar << bCooked;
 
@@ -735,9 +655,24 @@ void UAnimBankMeshComponent::Serialize(FArchive& Ar)
 		// Read existing data if it was serialized
 		TArray<FAnimBankMeshInstanceData> TempInstanceData;
 		TArray<float> TempInstanceCustomData;
+
 		if (bHasSkipSerializationPropertiesData)
 		{
-			TempInstanceData.BulkSerialize(Ar, false /* force per element serialization */);
+			if (Ar.CustomVer(FUE5MainStreamObjectVersion::GUID) < FUE5MainStreamObjectVersion::AnimBankInstanceDataSerializationV2)
+			{
+				TArray<FAnimBankMeshInstanceData_Deprecated> TempInstanceData_Deprecated;
+				TempInstanceData_Deprecated.BulkSerialize(Ar, false /* force per element serialization */);
+
+				TempInstanceData.Reserve(TempInstanceData_Deprecated.Num());
+				for (const auto& Item : TempInstanceData_Deprecated)
+				{
+					TempInstanceData.Emplace(FTransform3f(FMatrix44f(Item.Transform)), Item.BankIndex);
+				}
+			}
+			else
+			{
+				Ar << TempInstanceData;
+			}
 			TempInstanceCustomData.BulkSerialize(Ar);
 		}
 
@@ -755,7 +690,7 @@ void UAnimBankMeshComponent::Serialize(FArchive& Ar)
 	}
 	else if (bHasSkipSerializationPropertiesData)
 	{
-		InstanceData.BulkSerialize(Ar, false /* force per element serialization */);
+		Ar << InstanceData;
 		InstanceCustomData.BulkSerialize(Ar);
 	}
 
@@ -765,6 +700,30 @@ void UAnimBankMeshComponent::Serialize(FArchive& Ar)
 		Ar << SelectedInstances;
 	}
 #endif
+
+	if (Ar.CustomVer(FUE5MainStreamObjectVersion::GUID) >= FUE5MainStreamObjectVersion::AnimBankInstanceDataSerializationV2)
+	{
+		InstanceDataManager.Serialize(Ar, bCooked);
+	}
+	else if (Ar.IsLoading())
+	{
+		// Prior to this the id mapping was not saved so we need to reset it.
+		InstanceDataManager.Reset(InstanceData.Num());
+	}
+
+	if (bCooked)
+	{
+		if (Ar.IsLoading())
+		{
+			InstanceDataManager.ReadCookedRenderData(Ar);
+		}
+#if WITH_EDITOR
+		else if (Ar.IsSaving())
+		{
+			InstanceDataManager.WriteCookedRenderData(Ar, GetComponentDesc(GMaxRHIFeatureLevel));
+		}
+#endif
+	}
 }
 
 void UAnimBankMeshComponent::PostLoad()
@@ -810,6 +769,52 @@ int32 UAnimBankMeshComponent::GetInstanceCount() const
 	return InstanceData.Num();
 }
 
+TArrayView<const struct FAnimBankItem> UAnimBankMeshComponent::GetAnimBankItems() const
+{
+	return MakeArrayView(AnimBankItems);
+}
+
+void UAnimBankMeshComponent::SetAnimBankItems(TArrayView<const struct FAnimBankItem> InAnimBankItems)
+{
+	AnimBankItems = InAnimBankItems;
+	// We use the transform dirty state to drive the update of the animation data (to defer the need to add more bits), so we mark those as dirty here.
+	InstanceDataManager.TransformsChangedAll();
+	MarkRenderStateDirty();
+}
+
+template <typename ArrayType>
+inline void ReorderArray(ArrayType& InOutDataArray, const TArray<int32>& OldIndexArray, int32 ElementStride = 1)
+{
+	// TODO: this fails for FAnimBankMeshInstanceData, it is not really a problem here because the ctor just performs initializaiton and we'll overwrite it anyway 
+	// static_assert(std::is_trivially_constructible_v<decltype(InOutDataArray[0])>, "Must have trivially constructible type to safely call SetNumUninitialized below");
+	ArrayType TmpDataArray = MoveTemp(InOutDataArray);
+	InOutDataArray.Empty(TmpDataArray.Num());
+	for (int32 NewIndex = 0; NewIndex < TmpDataArray.Num(); ++NewIndex)
+	{
+		int32 OldIndex = OldIndexArray[NewIndex];
+		for (int32 SubIndex = 0; SubIndex < ElementStride; ++SubIndex)
+		{
+			InOutDataArray.Add(TmpDataArray[OldIndex * ElementStride + SubIndex]);
+		}
+	}
+}
+
+void UAnimBankMeshComponent::OptimizeInstanceData(bool bShouldRetainIdMap)
+{
+	// compute the optimal order 
+	TArray<int32> IndexRemap = InstanceDataManager.Optimize(GetComponentDesc(GMaxRHIFeatureLevel), bShouldRetainIdMap);
+	
+	if (!IndexRemap.IsEmpty())
+	{
+		// Reorder instances according to the remap
+		ReorderArray(InstanceData, IndexRemap);
+		ReorderArray(InstanceCustomData, IndexRemap, NumCustomDataFloats);
+#if WITH_EDITOR
+		ReorderArray(SelectedInstances, IndexRemap);
+#endif
+	}
+}
+
 void UAnimBankMeshComponent::ApplyInheritedPerInstanceData(const UAnimBankMeshComponent* InArchetype)
 {
 	check(InArchetype);
@@ -828,9 +833,9 @@ bool UAnimBankMeshComponent::ShouldInheritPerInstanceData(const UAnimBankMeshCom
 	return (bInheritPerInstanceData || !bEditableWhenInherited) && InArchetype && InArchetype->IsInBlueprint() && !IsTemplate();
 }
 
-void UAnimBankMeshComponent::SetupNewInstanceData(FAnimBankMeshInstanceData& InOutNewInstanceData, int32 InInstanceIndex, const FTransform& InInstanceTransform, int32 InBankIndex)
+void UAnimBankMeshComponent::SetupNewInstanceData(FAnimBankMeshInstanceData& InOutNewInstanceData, int32 InInstanceIndex, const FTransform3f& InInstanceTransform, int32 InBankIndex)
 {
-	InOutNewInstanceData.Transform = InInstanceTransform.ToMatrixWithScale();
+	InOutNewInstanceData.Transform = InInstanceTransform;
 	InOutNewInstanceData.BankIndex = InBankIndex;
 
 	if (bPhysicsStateCreated)
@@ -844,6 +849,129 @@ const Nanite::FResources* UAnimBankMeshComponent::GetNaniteResources() const
 	return Super::GetNaniteResources();
 }
 
+#if WITH_EDITOR
+
+void UAnimBankMeshComponent::PostAssetCompilation()
+{
+	InstanceDataManager.ClearChangeTracking();
+	MarkRenderStateDirty();
+}
+
+#endif 
+
+FInstanceDataManagerSourceDataDesc UAnimBankMeshComponent::GetComponentDesc(ERHIFeatureLevel::Type FeatureLevel)
+{
+	FInstanceDataManagerSourceDataDesc ComponentDesc;
+
+
+	ComponentDesc.PrimitiveMaterialDesc = GetUsedMaterialPropertyDesc(FeatureLevel);
+
+	FInstanceDataFlags Flags;
+	Flags.bHasPerInstanceRandom = ComponentDesc.PrimitiveMaterialDesc.bAnyMaterialHasPerInstanceRandom;
+	Flags.bHasPerInstanceCustomData = ComponentDesc.PrimitiveMaterialDesc.bAnyMaterialHasPerInstanceCustomData && NumCustomDataFloats != 0;
+#if WITH_EDITOR
+	Flags.bHasPerInstanceEditorData = GIsEditor != 0 && bHasPerInstanceHitProxies;
+#endif
+	
+	const bool bForceRefPose = CVarAnimBankForceRefPose.GetValueOnAnyThread() != 0;
+	const bool bUseAnimBank = !bForceRefPose && AnimBankItems.Num() > 0;
+
+	Flags.bHasPerInstanceHierarchyOffset = false;
+	Flags.bHasPerInstanceLocalBounds = bUseAnimBank && AnimBankItems.Num() > 1;
+	Flags.bHasPerInstanceDynamicData = false;
+	Flags.bHasPerInstanceSkinningData = true;
+
+
+	Flags.bHasPerInstanceLMSMUVBias = false;//IsStaticLightingAllowed();
+
+	ComponentDesc.Flags = Flags;
+	
+	// TODO: rename
+	ComponentDesc.MeshBounds = GetSkinnedAsset()->GetBounds();
+	ComponentDesc.NumCustomDataFloats = NumCustomDataFloats;
+	ComponentDesc.NumInstances = InstanceData.Num();
+
+	ComponentDesc.PrimitiveLocalToWorld = GetRenderMatrix();
+	ComponentDesc.ComponentMobility = Mobility;
+
+	const FReferenceSkeleton& RefSkeleton = GetSkinnedAsset()->GetRefSkeleton();
+	uint32 MaxBoneTransformCount = RefSkeleton.GetRawBoneNum();
+
+	ComponentDesc.BuildChangeSet = [&, MaxBoneTransformCount, MeshBounds = ComponentDesc.MeshBounds](FInstanceUpdateChangeSet& ChangeSet)
+	{
+		// publish data
+		ChangeSet.GetTransformWriter().Gather([&](int32 InstanceIndex) -> FRenderTransform { return FRenderTransform(InstanceData[InstanceIndex].Transform.ToMatrixWithScale()); });
+		ChangeSet.GetCustomDataWriter().Gather(MakeArrayView(InstanceCustomData), NumCustomDataFloats);
+
+		ChangeSet.GetSkinningDataWriter().Gather(
+			[&](int32 InstanceIndex) -> uint32 
+			{ 
+				return InstanceData[InstanceIndex].BankIndex * MaxBoneTransformCount * 2u; 
+			});
+
+		ChangeSet.GetLocalBoundsWriter().Gather(
+			[&](int32 InstanceIndex) -> FRenderBounds 
+			{				
+				uint32 BankIndex = InstanceData[InstanceIndex].BankIndex;
+				if (BankIndex < uint32(AnimBankItems.Num()))
+				{
+					const FAnimBankItem& BankItem = AnimBankItems[BankIndex];
+					if (BankItem.BankAsset != nullptr)
+					{
+						const FAnimBankData& BankData = BankItem.BankAsset->GetData();
+						if (BankItem.SequenceIndex < BankData.Entries.Num())
+						{
+							return BankData.Entries[BankItem.SequenceIndex].SampledBounds;
+						}
+					}
+				}
+				return MeshBounds; 
+			});
+
+#if WITH_EDITOR
+	if (ChangeSet.Flags.bHasPerInstanceEditorData)
+	{
+		// TODO: the way hit proxies are managed seems daft, why don't we just add them when needed and store them in an array alonside the instances?
+		//       this will always force us to update all the hit proxy data for every instances.
+		TArray<TRefCountPtr<HHitProxy>> HitProxies;
+		CreateHitProxyData(HitProxies);
+		ChangeSet.SetEditorData(HitProxies, SelectedInstances);
+	}
+#endif
+
+
+	};
+
+	return ComponentDesc;
+}
+
+void UAnimBankMeshComponent::SendRenderInstanceData_Concurrent()
+{
+	Super::SendRenderInstanceData_Concurrent();
+
+	// If the primitive isn't hidden update its instances.
+	const bool bDetailModeAllowsRendering = true;//DetailMode <= GetCachedScalabilityCVars().DetailMode;
+	// The proxy may not be created, this can happen when a SM is async loading for example.
+	if (bDetailModeAllowsRendering && (ShouldRender() || bCastHiddenShadow || bAffectIndirectLightingWhileHidden || bRayTracingFarField))
+	{
+		if (SceneProxy != nullptr)
+		{
+			// Make sure the instance data proxy is up to date:
+			if (InstanceDataManager.FlushChanges(GetComponentDesc(SceneProxy->GetScene().GetFeatureLevel())))
+			{
+				UpdateBounds();
+				GetWorld()->Scene->UpdatePrimitiveInstances(this);
+			}
+		}
+		else
+		{
+			UpdateBounds();
+			GetWorld()->Scene->AddPrimitive(this);
+		}
+	}
+}
+
+
 void UAnimBankMeshComponent::CreateRenderState_Concurrent(FRegisterComponentContext* Context)
 {
 	MeshObjectFactory = &CreateAnimBankMeshObjectFn;
@@ -876,6 +1004,14 @@ FPrimitiveSceneProxy* UAnimBankMeshComponent::CreateSceneProxy()
 		return nullptr;
 	}
 
+	for (FAnimBankItem& BankItem : AnimBankItems)
+	{
+		if (BankItem.BankAsset->IsCompiling())
+		{
+			return nullptr;
+		}
+	}
+
 	if (CheckPSOPrecachingAndBoostPriority() && GetPSOPrecacheProxyCreationStrategy() == EPSOPrecacheProxyCreationStrategy::DelayUntilPSOPrecached)
 	{
 		UE_LOG(LogAnimBank, Verbose, TEXT("Skipping CreateSceneProxy for UAnimBankMeshComponent %s (UAnimBankMeshComponent PSOs are still compiling)"), *GetFullName());
@@ -910,6 +1046,12 @@ FPrimitiveSceneProxy* UAnimBankMeshComponent::CreateSceneProxy()
 			{
 				Result = Super::CreateSceneProxy();
 			}
+
+			// Unclear exactly how this is supposed to work with a non-instanced proxy - will be interesting...
+			if (Result)
+			{
+				InstanceDataManager.FlushChanges(GetComponentDesc(Result->GetScene().GetFeatureLevel()));
+			}
 		}
 	}
 
@@ -920,10 +1062,18 @@ FPrimitiveSceneProxy* UAnimBankMeshComponent::CreateSceneProxy()
 	return Result;
 }
 
+void UAnimBankMeshComponent::OnUpdateTransform(EUpdateTransformFlags UpdateTransformFlags, ETeleportType Teleport)
+{
+	Super::OnUpdateTransform(UpdateTransformFlags, Teleport);
+	InstanceDataManager.PrimitiveTransformChanged();
+}
+
 #if WITH_EDITOR
 
 void UAnimBankMeshComponent::PostEditChangeChainProperty(FPropertyChangedChainEvent& PropertyChangedEvent)
 {
+	// Always clear the change tracking because in the editor, attributes may have been set without any sort of notification
+	InstanceDataManager.ClearChangeTracking();
 	if (PropertyChangedEvent.Property != nullptr)
 	{
 		if (PropertyChangedEvent.Property->GetFName() == GET_MEMBER_NAME_CHECKED(UAnimBankMeshComponent, InstanceData))
@@ -936,7 +1086,6 @@ void UAnimBankMeshComponent::PostEditChangeChainProperty(FPropertyChangedChainEv
 
 				AddInstanceInternal(
 					AddedAtIndex,
-					&InstanceData[AddedAtIndex],
 					PropertyChangedEvent.ChangeType == EPropertyChangeType::ArrayAdd ? FTransform::Identity : FTransform(InstanceData[AddedAtIndex].Transform),
 					PropertyChangedEvent.ChangeType == EPropertyChangeType::ArrayAdd ? 0 : InstanceData[AddedAtIndex].BankIndex,
 					/*bWorldSpace*/false
@@ -959,21 +1108,16 @@ void UAnimBankMeshComponent::PostEditChangeChainProperty(FPropertyChangedChainEv
 			else if (PropertyChangedEvent.ChangeType == EPropertyChangeType::ValueSet)
 			{
 			}
-			
 			MarkRenderStateDirty();
 		}
-		else if (PropertyChangedEvent.Property->GetFName() == GET_MEMBER_NAME_CHECKED(FAnimBankMeshInstanceData, Transform))
+		else if (PropertyChangedEvent.Property->GetFName() == GET_MEMBER_NAME_CHECKED(FAnimBankMeshInstanceData, Transform) 
+			 || PropertyChangedEvent.Property->GetFName() == GET_MEMBER_NAME_CHECKED(FAnimBankMeshInstanceData, BankIndex))
 		{
-
+			MarkRenderStateDirty();
 		}
 		else if (PropertyChangedEvent.Property->GetFName() == "NumCustomDataFloats")
 		{
-			// Can't just call SetNumCustomDataFloats because it doesn't do anything if the value is the same, and the editor has already modified the value...
-			NumCustomDataFloats = FMath::Max(NumCustomDataFloats, 0);
-
-			// Clear out and reinit to 0
-			InstanceCustomData.Empty(InstanceData.Num() * NumCustomDataFloats);
-			InstanceCustomData.SetNumZeroed(InstanceData.Num() * NumCustomDataFloats);
+			SetNumCustomDataFloats(NumCustomDataFloats);
 		}
 		else if (PropertyChangedEvent.PropertyChain.GetActiveMemberNode()->GetValue()->GetFName() == "InstanceCustomData")
 		{
@@ -982,6 +1126,7 @@ void UAnimBankMeshComponent::PostEditChangeChainProperty(FPropertyChangedChainEv
 			{
 				int32 InstanceIndex = ChangedCustomValueIndex / NumCustomDataFloats;
 			}
+			MarkRenderStateDirty();
 		}
 		else if (
 			(PropertyChangedEvent.Property->GetFName() == GET_MEMBER_NAME_CHECKED(UAnimBankMeshComponent, AnimBankItems)) ||
@@ -1026,6 +1171,31 @@ void UAnimBankMeshComponent::PostEditUndo()
 	MarkRenderStateDirty();
 }
 
+void UAnimBankMeshComponent::BeginCacheForCookedPlatformData(const ITargetPlatform* TargetPlatform)
+{
+	for (FAnimBankItem& BankItem : AnimBankItems)
+	{
+		if (BankItem.BankAsset)
+		{
+			BankItem.BankAsset->BeginCacheForCookedPlatformData(TargetPlatform);
+		}
+	}
+
+}
+
+bool UAnimBankMeshComponent::IsCachedCookedPlatformDataLoaded( const ITargetPlatform* TargetPlatform ) 
+{
+	for (FAnimBankItem& BankItem : AnimBankItems)
+	{
+		if (BankItem.BankAsset->IsCompiling())
+		{
+			return false;
+		}
+	}
+
+	return Super::IsCachedCookedPlatformDataLoaded(TargetPlatform);
+}
+
 #endif
 
 TStructOnScope<FActorComponentInstanceData> UAnimBankMeshComponent::GetComponentInstanceData() const
@@ -1080,7 +1250,6 @@ void UAnimBankMeshComponent::ApplyComponentInstanceData(struct FAnimBankMeshComp
 
 	SelectedInstances = InstancedMeshData->SelectedInstances;
 	bHasPerInstanceHitProxies = InstancedMeshData->bHasPerInstanceHitProxies;
-	bIsInstanceDataApplyCompleted = true;
 #endif
 }
 
@@ -1121,7 +1290,7 @@ FBoxSphereBounds UAnimBankMeshComponent::CalcBounds(const FTransform& LocalToWor
 					for (int32 InstanceIndex = 0; InstanceIndex < InstanceData.Num(); InstanceIndex++)
 					{
 						const FAnimBankMeshInstanceData& Instance = InstanceData[InstanceIndex];
-						BoundsBuilder += MergedBounds.TransformBy(Instance.Transform * BoundTransformMatrix);
+						BoundsBuilder += MergedBounds.TransformBy(FTransform(Instance.Transform) * LocalToWorld);
 					}
 				}
 			}
@@ -1133,13 +1302,17 @@ FBoxSphereBounds UAnimBankMeshComponent::CalcBounds(const FTransform& LocalToWor
 					if (Instance.BankIndex < uint32(AnimBankItems.Num()))
 					{
 						const FAnimBankItem& BankItem = AnimBankItems[Instance.BankIndex];
-						if (BankItem.BankAsset != nullptr)
+						if (BankItem.BankAsset != nullptr
+#if WITH_EDITOR
+							&& !BankItem.BankAsset->IsCompiling()
+#endif
+							)
 						{
 							const FAnimBankData& BankData = BankItem.BankAsset->GetData();
 							if (BankItem.SequenceIndex < BankData.Entries.Num())
 							{
 								const FBox BankBounds = BankData.Entries[BankItem.SequenceIndex].SampledBounds.GetBox();
-								BoundsBuilder += BankBounds.TransformBy(Instance.Transform * BoundTransformMatrix);
+								BoundsBuilder += BankBounds.TransformBy(FTransform(Instance.Transform) * LocalToWorld);
 							}
 						}
 					}
@@ -1155,7 +1328,7 @@ FBoxSphereBounds UAnimBankMeshComponent::CalcBounds(const FTransform& LocalToWor
 			{
 				for (int32 InstanceIndex = 0; InstanceIndex < InstanceData.Num(); InstanceIndex++)
 				{
-					BoundsBuilder += InstanceBounds.TransformBy(InstanceData[InstanceIndex].Transform * BoundTransformMatrix);
+					BoundsBuilder += InstanceBounds.TransformBy(FTransform(InstanceData[InstanceIndex].Transform) * LocalToWorld);
 				}
 
 				return BoundsBuilder;
@@ -1199,78 +1372,105 @@ void UAnimBankMeshComponent::RefreshBoneTransforms(FActorComponentTickFunction*
 	//bNeedsRefreshTransform = false;
 }
 
-int32 UAnimBankMeshComponent::AddInstance(const FTransform& InstanceTransform, int32 BankIndex, bool bWorldSpace)
-{
-	return AddInstanceInternal(InstanceData.Num(), nullptr, InstanceTransform, BankIndex, bWorldSpace);
-}
-
-TArray<int32> UAnimBankMeshComponent::AddInstances(const TArray<FTransform>& InstanceTransforms, bool bShouldReturnIndices, bool bWorldSpace, bool bUpdateNavigation)
+FPrimitiveInstanceId UAnimBankMeshComponent::AddInstance(const FTransform& InstanceTransform, int32 BankIndex, bool bWorldSpace)
 {
-	return AddInstancesInternal(InstanceTransforms, bShouldReturnIndices, bWorldSpace, bUpdateNavigation);
+	return AddInstanceInternal(InstanceData.Num(), InstanceTransform, BankIndex, bWorldSpace);
 }
 
-bool UAnimBankMeshComponent::SetCustomDataValue(int32 InstanceIndex, int32 CustomDataIndex, float CustomDataValue, bool bMarkRenderStateDirty)
+TArray<FPrimitiveInstanceId> UAnimBankMeshComponent::AddInstances(const TArray<FTransform>& InstanceTransforms, bool bShouldReturnIds, bool bWorldSpace)
 {
-	if (!InstanceData.IsValidIndex(InstanceIndex) || CustomDataIndex < 0 || CustomDataIndex >= NumCustomDataFloats)
+	TArray<FPrimitiveInstanceId> NewInstanceIds;
+	if (InstanceTransforms.IsEmpty())
 	{
-		return false;
+		return NewInstanceIds;
 	}
 
 	Modify();
 
-	InstanceCustomData[InstanceIndex * NumCustomDataFloats + CustomDataIndex] = CustomDataValue;
+	const int32 NumToAdd = InstanceTransforms.Num();
 
-	if (bMarkRenderStateDirty)
+	if (bShouldReturnIds)
 	{
-		MarkRenderStateDirty();
+		NewInstanceIds.SetNumUninitialized(NumToAdd);
 	}
 
-	return true;
+
+	// Reserve memory space
+	const int32 NewNumInstances = InstanceData.Num() + NumToAdd;
+	InstanceData.Reserve(NewNumInstances);
+	InstanceCustomData.Reserve(NumCustomDataFloats * NewNumInstances);
+#if WITH_EDITOR
+	SelectedInstances.Reserve(NewNumInstances);
+#endif
+
+	for (int32 AddIndex = 0; AddIndex < NumToAdd; ++AddIndex)
+	{
+		const FTransform& InstanceTransform = InstanceTransforms[AddIndex];
+		int32 BankIndex = 0; // TODO
+		FPrimitiveInstanceId InstanceId = AddInstanceInternal(InstanceData.Num(), InstanceTransform, BankIndex, bWorldSpace);
+		if (bShouldReturnIds)
+		{
+			NewInstanceIds[AddIndex] = InstanceId;
+		}
+	}
+
+	return NewInstanceIds;
 }
 
-void UAnimBankMeshComponent::SetNumCustomDataFloats(int32 InNumCustomDataFloats)
+bool UAnimBankMeshComponent::SetCustomDataValue(FPrimitiveInstanceId InstanceId, int32 CustomDataIndex, float CustomDataValue)
 {
-	if (FMath::Max(InNumCustomDataFloats, 0) != NumCustomDataFloats)
-	{
-		NumCustomDataFloats = FMath::Max(InNumCustomDataFloats, 0);
+	int32 InstanceIndex = InstanceDataManager.IdToIndex(InstanceId);
 
-		// Clear out and reinit to 0
-		InstanceCustomData.Empty(InstanceData.Num() * NumCustomDataFloats);
-		InstanceCustomData.SetNumZeroed(InstanceData.Num() * NumCustomDataFloats);
+	if (!InstanceData.IsValidIndex(InstanceIndex) || CustomDataIndex < 0 || CustomDataIndex >= NumCustomDataFloats)
+	{
+		return false;
 	}
+
+	Modify();
+
+	InstanceDataManager.CustomDataChanged(InstanceIndex);
+	InstanceCustomData[InstanceIndex * NumCustomDataFloats + CustomDataIndex] = CustomDataValue;
+
+	return true;
 }
 
-bool UAnimBankMeshComponent::SetCustomData(int32 InstanceIndex, TArrayView<const float> CustomDataFloats, bool bMarkRenderStateDirty)
+bool UAnimBankMeshComponent::SetCustomData(FPrimitiveInstanceId InstanceId, TArrayView<const float> CustomDataFloats)
 {
+	int32 InstanceIndex = InstanceDataManager.IdToIndex(InstanceId);
+
 	if (!InstanceData.IsValidIndex(InstanceIndex) || CustomDataFloats.Num() == 0)
 	{
 		return false;
 	}
 
-	if (bMarkRenderStateDirty)
-	{
-		Modify();
-	}
+	Modify();
 
 	const int32 NumToCopy = FMath::Min(CustomDataFloats.Num(), NumCustomDataFloats);
+	InstanceDataManager.CustomDataChanged(InstanceIndex);
 	FMemory::Memcpy(&InstanceCustomData[InstanceIndex * NumCustomDataFloats], CustomDataFloats.GetData(), NumToCopy * CustomDataFloats.GetTypeSize());
+	return true;
+}
 
-	if (bMarkRenderStateDirty)
+void UAnimBankMeshComponent::SetNumCustomDataFloats(int32 InNumCustomDataFloats)
+{
+	if (FMath::Max(InNumCustomDataFloats, 0) != NumCustomDataFloats)
 	{
-		MarkRenderStateDirty();
+		NumCustomDataFloats = FMath::Max(InNumCustomDataFloats, 0);
 	}
 
-	return true;
-}
+	if (InstanceData.Num() * NumCustomDataFloats != InstanceCustomData.Num())
+	{
+		InstanceDataManager.NumCustomDataChanged();
 
-void UAnimBankMeshComponent::PreAllocateInstancesMemory(int32 AddedInstanceCount)
-{
-	InstanceData.Reserve(InstanceData.Num() + AddedInstanceCount);
-	InstanceCustomData.Reserve(InstanceCustomData.Num() + AddedInstanceCount * NumCustomDataFloats);
+		// Clear out and reinit to 0
+		InstanceCustomData.Empty(InstanceData.Num() * NumCustomDataFloats);
+		InstanceCustomData.SetNumZeroed(InstanceData.Num() * NumCustomDataFloats);
+	}
 }
 
-bool UAnimBankMeshComponent::GetInstanceTransform(int32 InstanceIndex, FTransform& OutInstanceTransform, bool bWorldSpace) const
+bool UAnimBankMeshComponent::GetInstanceTransform(FPrimitiveInstanceId InstanceId, FTransform& OutInstanceTransform, bool bWorldSpace) const
 {
+	int32 InstanceIndex = InstanceDataManager.IdToIndex(InstanceId);
 	if (!InstanceData.IsValidIndex(InstanceIndex))
 	{
 		return false;
@@ -1287,70 +1487,38 @@ bool UAnimBankMeshComponent::GetInstanceTransform(int32 InstanceIndex, FTransfor
 	return true;
 }
 
-bool UAnimBankMeshComponent::RemoveInstance(int32 InstanceIndex)
-{
-	return RemoveInstanceInternal(InstanceIndex, false);
-}
-
-bool UAnimBankMeshComponent::RemoveInstances(const TArray<int32>& InstancesToRemove)
-{
-	return RemoveInstances(InstancesToRemove, false /*bInstanceArrayAlreadySortedInReverseOrder*/);
-}
-
-bool UAnimBankMeshComponent::RemoveInstances(const TArray<int32>& InstancesToRemove, bool bInstanceArrayAlreadySortedInReverseOrder)
+bool UAnimBankMeshComponent::RemoveInstance(FPrimitiveInstanceId InstanceId)
 {
-	if (InstancesToRemove.IsEmpty())
+	int32 InstanceIndex = InstanceDataManager.IdToIndex(InstanceId);
+	if (InstanceData.IsValidIndex(InstanceIndex))
 	{
-		return false;
+		Modify();
+		return RemoveInstanceInternal(InstanceIndex, false);
 	}
-	
-	auto RemoveInstancedWithSortedArray = [this](const TArray<int32>& SortedInstancesToRemove) -> bool
-	{
-		if (!InstanceData.IsValidIndex(SortedInstancesToRemove[0]) || !InstanceData.IsValidIndex(SortedInstancesToRemove.Last()))
-		{
-			return false;
-		}
-
-		for (const int32 InstanceIndex : SortedInstancesToRemove)
-		{
-			RemoveInstanceInternal(InstanceIndex, false);
-		}
+	return false;
+}
 
-		return true;
-	};
+void UAnimBankMeshComponent::RemoveInstances(const TArray<FPrimitiveInstanceId>& InstancesToRemove)
+{
+	Modify();
 
-	bool bSuccess = false;
-	if (bInstanceArrayAlreadySortedInReverseOrder)
-	{
-		bSuccess = RemoveInstancedWithSortedArray(InstancesToRemove);
-	}
-	else
+	for (FPrimitiveInstanceId InstanceId : InstancesToRemove)
 	{
-		// Sort so Remove doesn't alter the indices of items still to remove
-		TArray<int32> SortedInstancesToRemove = InstancesToRemove;
-		SortedInstancesToRemove.Sort(TGreater<int32>());
-		bSuccess = RemoveInstancedWithSortedArray(SortedInstancesToRemove);
+		int32 InstanceIndex = InstanceDataManager.IdToIndex(InstanceId);
+		RemoveInstanceInternal(InstanceIndex, false);
 	}
-
-	return bSuccess;
 }
 
 void UAnimBankMeshComponent::ClearInstances()
 {
-#if WITH_EDITOR
-	DeletionState = EInstanceDeletionReason::Clearing;
-#endif
-
-	const int32 PrevNumInstances = GetInstanceCount();
+	Modify();
 
 	// Clear all the per-instance data
 	InstanceData.Empty();
 	InstanceCustomData.Empty();
 
-	//ProxySize = 0;
-
 #if WITH_EDITOR
-	DeletionState = EInstanceDeletionReason::NotDeleting;
+	SelectedInstances.Empty();
 #endif
 }
 
@@ -1424,115 +1592,55 @@ void UAnimBankMeshComponent::CreateHitProxyData(TArray<TRefCountPtr<HHitProxy>>&
 	}
 }
 
-int32 UAnimBankMeshComponent::AddInstanceInternal(int32 InstanceIndex, FAnimBankMeshInstanceData* InNewInstanceData, const FTransform& InstanceTransform, int32 BankIndex, bool bWorldSpace)
+FPrimitiveInstanceId UAnimBankMeshComponent::AddInstanceInternal(int32 InstanceIndex, const FTransform& InstanceTransform, int32 BankIndex, bool bWorldSpace)
 {
-	FAnimBankMeshInstanceData* NewInstanceData = InNewInstanceData;
-
-	if (NewInstanceData == nullptr)
+	// This happens because the editor modifies the InstanceData array _before_ callbacks. If we could change the UI to not do that we could remove this ugly hack.
+	if (!InstanceData.IsValidIndex(InstanceIndex))
 	{
-		NewInstanceData = &InstanceData.AddDefaulted_GetRef();
+		check(InstanceIndex == InstanceData.Num());
+		InstanceData.AddDefaulted();
 	}
 
-	const FTransform LocalTransform = bWorldSpace ? InstanceTransform.GetRelativeTransform(GetComponentTransform()) : InstanceTransform;
-	SetupNewInstanceData(*NewInstanceData, InstanceIndex, LocalTransform, BankIndex);
+	FPrimitiveInstanceId InstanceId = InstanceDataManager.Add(InstanceIndex);
+
+	const FTransform3f LocalTransform = FTransform3f(bWorldSpace ? InstanceTransform.GetRelativeTransform(GetComponentTransform()) : InstanceTransform);
+	SetupNewInstanceData(InstanceData[InstanceIndex], InstanceIndex, LocalTransform, BankIndex);
 
 	// Add custom data to instance
 	InstanceCustomData.AddZeroed(NumCustomDataFloats);
 
 #if WITH_EDITOR
-	if (SelectedInstances.Num())
-	{
-		SelectedInstances.Add(false);
-	}
+	SelectedInstances.Add(false);
 #endif
 
-	return InstanceIndex;
+	return InstanceId;
 }
 
-TArray<int32> UAnimBankMeshComponent::AddInstancesInternal(TConstArrayView<FTransform> InstanceTransforms, bool bShouldReturnIndices, bool bWorldSpace, bool bUpdateNavigation)
+bool UAnimBankMeshComponent::RemoveInstanceInternal(int32 InstanceIndex, bool bInstanceAlreadyRemoved)
 {
-	const int32 Count = InstanceTransforms.Num();
-
-	TArray<int32> NewInstanceIndices;
-
-	if (bShouldReturnIndices)
+	if (!ensure(bInstanceAlreadyRemoved || InstanceData.IsValidIndex(InstanceIndex)))
 	{
-		NewInstanceIndices.Reserve(Count);
+		return false;
 	}
-
-	int32 InstanceIndex = InstanceData.Num();
-
-	InstanceCustomData.AddZeroed(NumCustomDataFloats * Count);
-
-#if WITH_EDITOR
-	SelectedInstances.Add(false, Count);
-#endif
-
-	const int32 NumInstances = InstanceIndex + Count;
-	InstanceData.Reserve(NumInstances);
-
-	for (const FTransform& InstanceTransform : InstanceTransforms)
+	InstanceDataManager.RemoveAtSwap(InstanceIndex);
+	
+	// remove instance
+	if (!bInstanceAlreadyRemoved)
 	{
-		FAnimBankMeshInstanceData& NewInstanceData = InstanceData.AddDefaulted_GetRef();
-
-		int32 BankIndex = 0; // TODO
-
-		const FTransform LocalTransform = bWorldSpace ? InstanceTransform.GetRelativeTransform(GetComponentTransform()) : InstanceTransform;
-		SetupNewInstanceData(NewInstanceData, InstanceIndex, LocalTransform, BankIndex);
-
-		if (bShouldReturnIndices)
-		{
-			NewInstanceIndices.Add(InstanceIndex);
-		}
-
-		++InstanceIndex;
+		InstanceData.RemoveAtSwap(InstanceIndex, EAllowShrinking::No);
 	}
-
-	return NewInstanceIndices;
-}
-
-bool UAnimBankMeshComponent::RemoveInstanceInternal(int32 InstanceIndex, bool InstanceAlreadyRemoved, bool bForceRemoveAtSwap, bool bUpdateNavigation)
-{
-#if WITH_EDITOR
-	DeletionState = InstanceAlreadyRemoved ? EInstanceDeletionReason::EntryAlreadyRemoved : EInstanceDeletionReason::EntryRemoval;
-#endif
-
-	const bool bUseRemoveAtSwap = true;
-
-	// remove instance
-	if (!InstanceAlreadyRemoved && InstanceData.IsValidIndex(InstanceIndex))
+	
+	if (InstanceCustomData.IsValidIndex(InstanceIndex * NumCustomDataFloats))
 	{
-		if (bUseRemoveAtSwap)
-		{
-			InstanceData.RemoveAtSwap(InstanceIndex, EAllowShrinking::No);
-			InstanceCustomData.RemoveAtSwap(InstanceIndex * NumCustomDataFloats, NumCustomDataFloats, EAllowShrinking::No);
-		}
-		else
-		{
-			InstanceData.RemoveAt(InstanceIndex);
-			InstanceCustomData.RemoveAt(InstanceIndex * NumCustomDataFloats, NumCustomDataFloats);
-		}
+		InstanceCustomData.RemoveAtSwap(InstanceIndex * NumCustomDataFloats, NumCustomDataFloats, EAllowShrinking::No);
 	}
 
 #if WITH_EDITOR
 	// remove selection flag if array is filled in
 	if (SelectedInstances.IsValidIndex(InstanceIndex))
 	{
-		if (bUseRemoveAtSwap)
-		{
-			SelectedInstances.RemoveAtSwap(InstanceIndex);
-		}
-		else
-		{
-			SelectedInstances.RemoveAt(InstanceIndex);
-		}
+		SelectedInstances.RemoveAtSwap(InstanceIndex);
 	}
 #endif
-
-	const int32 LastInstanceIndex = InstanceData.Num();
-
-#if WITH_EDITOR
-	DeletionState = EInstanceDeletionReason::NotDeleting;
-#endif
 	return true;
 }
diff --git a/Engine/Source/Runtime/Engine/Private/InstanceData/InstanceDataManager.cpp b/Engine/Source/Runtime/Engine/Private/InstanceData/InstanceDataManager.cpp
new file mode 100644
--- /dev/null
+++ b/Engine/Source/Runtime/Engine/Private/InstanceData/InstanceDataManager.cpp
@@ -0,0 +1,1242 @@
+// Copyright Epic Games, Inc. All Rights Reserved.
+
+#include "InstanceData/InstanceDataManager.h"
+#include "InstanceData/InstanceUpdateChangeSet.h"
+#include "InstanceData/InstanceDataUpdateUtils.h"
+#include "Components/PrimitiveComponent.h"
+#include "DataDrivenShaderPlatformInfo.h"
+#include "SceneInterface.h"
+#include "Math/DoubleFloat.h"
+#include "RenderUtils.h"
+#include "Rendering/MotionVectorSimulation.h"
+
+#if 0
+
+#define LOG_INST_DATA(_Format_, ...) \
+{ \
+	FString Tmp = FString::Printf(_Format_, ##__VA_ARGS__);\
+	UE_LOG(LogInstanceProxy, Log, TEXT("%p, %s"), PrimitiveComponent.IsValid() ? PrimitiveComponent.Get() : nullptr, *Tmp);\
+}
+
+#else
+	#define LOG_INST_DATA(_Format_, ...) 
+#endif
+
+static TAutoConsoleVariable<bool> CVarInstanceDataManagerBuildOptimizedForTesting(
+	TEXT("r.InstanceData.BuildOptimizedForTesting"),
+	false,
+	TEXT("."));
+
+
+template <typename IndexRemapType>
+void ApplyDataChanges(FInstanceUpdateChangeSet& ChangeSet, const IndexRemapType& IndexRemap, int32 PostUpdateNumInstances, FInstanceSceneDataBuffers::FWriteView& ProxyData)
+{
+	ProxyData.PrimitiveToRelativeWorld = ChangeSet.PrimitiveToRelativeWorld;
+	ProxyData.PrimitiveWorldSpaceOffset = ChangeSet.PrimitiveWorldSpaceOffset;
+		
+	// TODO: DISP - Fix me (this comment came along from FPrimitiveSceneProxy::SetInstanceLocalBounds and is probably still true...)
+	const FVector3f PadExtent = GetLocalBoundsPadExtent(ProxyData.PrimitiveToRelativeWorld, ChangeSet.AbsMaxDisplacement);
+	if (!ChangeSet.Flags.bHasPerInstanceLocalBounds)
+	{
+		ProxyData.InstanceLocalBounds = MoveTemp(ChangeSet.InstanceLocalBounds);
+
+		for (FRenderBounds& Bounds : ProxyData.InstanceLocalBounds)
+		{
+			Bounds.Min -= PadExtent;
+			Bounds.Max += PadExtent;
+		}
+	}
+	else
+	{
+		ChangeSet.GetLocalBoundsReader().Scatter(ProxyData.InstanceLocalBounds, [PadExtent](FRenderBounds& Bounds) -> void 
+			{	
+				Bounds.Min -= PadExtent;
+				Bounds.Max += PadExtent;
+			}, IndexRemap);
+	}
+
+	// unpack transform deltas
+	ApplyTransformUpdates(ChangeSet.GetTransformDelta(), IndexRemap, ChangeSet.PrimitiveToRelativeWorld, ChangeSet.Transforms, PostUpdateNumInstances, ProxyData.InstanceToPrimitiveRelative);
+	if (ChangeSet.Flags.bHasPerInstanceDynamicData)
+	{
+		FRenderTransform PrevPrimitiveToRelativeWorld = ChangeSet.PreviousPrimitiveToRelativeWorld.Get(ChangeSet.PrimitiveToRelativeWorld);
+		ApplyTransformUpdates(ChangeSet.GetTransformDelta(), IndexRemap, PrevPrimitiveToRelativeWorld, ChangeSet.PrevTransforms, PostUpdateNumInstances, ProxyData.PrevInstanceToPrimitiveRelative);
+	}
+	else
+	{
+		ProxyData.PrevInstanceToPrimitiveRelative.Reset();
+	}
+
+	ApplyAttributeChanges(ChangeSet, IndexRemap, ProxyData);
+}
+
+class FUpdatableInstanceDataSceneProxy : public FInstanceDataSceneProxy
+{
+public:
+	using FInstanceDataSceneProxy::FInstanceDataSceneProxy;
+
+	virtual FInstanceDataUpdateTaskInfo *GetUpdateTaskInfo() override
+	{ 
+		return &InstanceDataUpdateTaskInfo; 
+	}
+
+	virtual void Update(FInstanceUpdateChangeSet&& ChangeSet)
+	{
+		SCOPED_NAMED_EVENT(FISMCInstanceDataSceneProxy_Update, FColor::Emerald);
+		check(!ChangeSet.IsFullUpdate());
+
+	#if WITH_EDITOR
+		// replace the HP container.
+		if (ChangeSet.HitProxyContainer)
+		{
+			HitProxyContainer = MoveTemp(ChangeSet.HitProxyContainer);
+		}
+	#endif
+
+		DecStatCounters();
+
+		FInstanceSceneDataBuffers::FAccessTag AccessTag(PointerHash(this));
+		FInstanceSceneDataBuffers::FWriteView ProxyData = InstanceSceneDataBuffers.BeginWriteAccess(AccessTag);
+
+		ProxyData.Flags = ChangeSet.Flags;
+
+		int32 PostUpdateNumInstances = ChangeSet.NumSourceInstances;
+
+		// Handle data movement, needs old & new ID maps
+		// These can only be caused by removes, which means an item can only ever move towards lower index in the array.
+		// Thus, we can always safely overwrite the data in the new slot, since we do them in increasing order.
+		// NOTE: If we start allowing some other kind of permutation of the ISM data, this assumption will break.
+		// TODO: Add validation code somewhere in the pipeline.
+		const auto IndexDelta = ChangeSet.GetIndexChangedDelta();
+		for (auto It = IndexDelta.GetIterator(); It; ++It)
+		{
+			// Index in the source (e.g., component)
+			const int32 ToIndex = It.GetIndex();
+			if (!ChangeSet.InstanceAttributeTracker.TestFlag<FInstanceAttributeTracker::EFlag::Added>(ToIndex))
+			{
+				int32 ItemIndex = It.GetItemIndex();
+				FPrimitiveInstanceId InstanceId = ChangeSet.bIdentityIdMap ? FPrimitiveInstanceId{ToIndex} : ChangeSet.IndexToIdMapDeltaData[ItemIndex];
+				if (InstanceIdIndexMap.IsValidId(InstanceId))
+				{
+					const int32 FromIndex =  InstanceIdIndexMap.IdToIndex(InstanceId);
+
+					ProxyData.InstanceToPrimitiveRelative[ToIndex] = ProxyData.InstanceToPrimitiveRelative[FromIndex];
+					CondMove(ChangeSet.Flags.bHasPerInstanceCustomData, ProxyData.InstanceCustomData, FromIndex, ToIndex, ChangeSet.NumCustomDataFloats);
+					CondMove(ChangeSet.Flags.bHasPerInstanceRandom, ProxyData.InstanceRandomIDs, FromIndex, ToIndex);
+					CondMove(ChangeSet.Flags.bHasPerInstanceLMSMUVBias, ProxyData.InstanceLightShadowUVBias, FromIndex, ToIndex);
+	#if WITH_EDITOR
+					CondMove(ChangeSet.Flags.bHasPerInstanceEditorData, ProxyData.InstanceEditorData, FromIndex, ToIndex);
+	#endif
+				}
+			}
+		}
+
+		UpdateIdMapping(ChangeSet, FIdentityIndexRemap(), InstanceIdIndexMap);
+		check(ChangeSet.NumSourceInstances == InstanceIdIndexMap.GetMaxInstanceIndex());
+
+		FIdentityIndexRemap IndexRemap;
+		ApplyDataChanges(ChangeSet, IndexRemap, PostUpdateNumInstances, ProxyData);
+
+		InstanceSceneDataBuffers.EndWriteAccess(AccessTag);
+
+		InstanceSceneDataBuffers.ValidateData();
+
+		IncStatCounters();
+	}
+
+	virtual void Build(FInstanceUpdateChangeSet&& ChangeSet)
+	{
+		DecStatCounters();
+		check(ChangeSet.IsFullUpdate());
+		// TODO: add similar checks: checkSlow(!ChangeSet.GetTransformDelta().IsDelta());
+		// TODO: add similar checks: checkSlow(!ChangeSet.GetCustomDataDelta().IsDelta() || (!ChangeSet.Flags.bHasPerInstanceCustomData && ChangeSet.GetCustomDataDelta().IsEmpty()));
+		// TODO: add similar checks: checkSlow(!ChangeSet.GetInstanceLightShadowUVBiasDelta().IsDelta() || ChangeSet.GetInstanceLightShadowUVBiasDelta().IsEmpty());
+	#if WITH_EDITOR
+		// TODO: add similar checks: checkSlow(!ChangeSet.GetInstanceEditorDataDelta().IsDelta() || ChangeSet.GetInstanceEditorDataDelta().IsEmpty());
+		// replace the HP container.
+		if (ChangeSet.HitProxyContainer)
+		{
+			HitProxyContainer = MoveTemp(ChangeSet.HitProxyContainer);
+		}
+	#endif
+
+		FInstanceSceneDataBuffers::FAccessTag AccessTag(PointerHash(this));
+		FInstanceSceneDataBuffers::FWriteView WriteView = InstanceSceneDataBuffers.BeginWriteAccess(AccessTag);
+
+		WriteView.Flags = ChangeSet.Flags;
+
+		if (ChangeSet.PrecomputedOptimizationData)
+		{
+			// We don't store an ID mapping for this case, since we assume a full rebuild is needed to handle any changes at all.
+			// TODO: it may be possible that we'd want to keep the ID mapping at some point if we had a case where a specific instance needed 
+			//       to be addressed on the RT. However, since it is not free to keep the remap around we don't do that until needed.
+			InstanceIdIndexMap.Reset(ChangeSet.NumSourceInstances);
+
+			// If the optimization data contains an index remap, we must rearrange the instance data to match (or spatial hash ranges won't match).
+			// Otherwise we can take the most optimal unsorted path (which allows moving attribute arrays).
+			if (ChangeSet.PrecomputedOptimizationData->ProxyIndexToComponentIndexRemap.IsEmpty())
+			{
+				ApplyDataChanges(ChangeSet, FIdentityIndexRemap(), InstanceIdIndexMap.GetMaxInstanceIndex(), WriteView);
+			}
+			else
+			{
+				FSrcIndexRemap SortedInstancesRemap(ChangeSet.PrecomputedOptimizationData->ProxyIndexToComponentIndexRemap);
+				ApplyDataChanges(ChangeSet, SortedInstancesRemap, InstanceIdIndexMap.GetMaxInstanceIndex(), WriteView);
+			}
+	
+			InstanceSceneDataBuffers.SetImmutable(FInstanceSceneDataImmutable(ChangeSet.PrecomputedOptimizationData->Hashes), WriteView.AccessTag);
+		}
+		else
+		{
+			UpdateIdMapping(ChangeSet, FIdentityIndexRemap(), InstanceIdIndexMap);
+			check(ChangeSet.NumSourceInstances == InstanceIdIndexMap.GetMaxInstanceIndex());
+
+			FIdentityIndexRemap IndexRemap;
+			ApplyDataChanges(ChangeSet, IndexRemap, InstanceIdIndexMap.GetMaxInstanceIndex(), WriteView);
+		}
+		InstanceSceneDataBuffers.EndWriteAccess(AccessTag);
+
+		InstanceSceneDataBuffers.ValidateData();
+
+		IncStatCounters();
+	}
+
+protected:
+	FInstanceDataUpdateTaskInfo InstanceDataUpdateTaskInfo;
+	FInstanceIdIndexMap InstanceIdIndexMap;
+#if WITH_EDITOR
+	/**Container for hitproxies that are used by the instances, uses the FDeferredCleanupInterface machinery to delete itself back on the game thread when replaced. */
+	TPimplPtr<FOpaqueHitProxyContainer> HitProxyContainer;
+#endif
+};
+
+
+#if 1
+namespace RenderingSpatialHash
+{
+	template <typename ScalarType>
+	inline FArchive& operator<<(FArchive& Ar, TLocation<ScalarType>& Item)
+	{
+		Ar << Item.Coord;
+		Ar << Item.Level;
+		return Ar;
+	}
+}
+
+inline FArchive& operator<<(FArchive& Ar, FInstanceSceneDataBuffers::FCompressedSpatialHashItem& Item)
+{
+	Ar << Item.Location;
+	Ar << Item.NumInstances;
+
+	return Ar;
+}
+#endif
+
+FInstanceDataManager::FInstanceDataManager(UPrimitiveComponent* InPrimitiveComponent) 
+	: bNumCustomDataChanged(false)
+	, bPrimitiveTransformChanged(false)
+	, bAnyInstanceChange(false)
+	, PrimitiveComponent(InPrimitiveComponent) 
+{
+	// Don't do anything if this is not a "real" ISM being tracked (this logic shopuld move out).
+	if (PrimitiveComponent.IsValid() && PrimitiveComponent->HasAnyFlags(RF_ClassDefaultObject | RF_ArchetypeObject))
+	{
+		TrackingState = ETrackingState::Disabled;
+	}
+	LOG_INST_DATA(TEXT("FInstanceDataManager %s, TrackingState=%s"), *PrimitiveComponent->GetFullName(), TrackingState == ETrackingState::Disabled ? TEXT("Disabled") : TEXT("Initial"));
+}
+
+FPrimitiveInstanceId FInstanceDataManager::Add(int32 InInstanceAddAtIndex)
+{
+	if (GetState() == ETrackingState::Disabled)
+	{
+		return FPrimitiveInstanceId{};
+	}
+
+	ValidateMapping();
+
+	// 1. determine if we need to enable explicit tracking, this happens when an instance is inserted (i.e., not added at the end).
+	bool bIsInsert = InInstanceAddAtIndex != GetMaxInstanceIndex();
+
+	// Create explicit mapping if we need it now
+	if (bIsInsert && HasIdentityMapping())
+	{
+		CreateExplicitIdentityMapping();
+	}
+
+	MarkComponentRenderInstancesDirty();
+
+	if (HasIdentityMapping())
+	{
+		int32 InstanceId = NumInstances++;
+		MarkChangeHelper<EChangeFlag::Added>(InstanceId);
+
+		LOG_INST_DATA(TEXT("Add(IDX: %d, bInsert: %d) -> Id: %d"), InInstanceAddAtIndex, bInsert ? 1 : 0, InstanceId);
+		return FPrimitiveInstanceId { InstanceId };
+	}
+
+	FPrimitiveInstanceId InstanceId{ValidInstanceIdMask.FindAndSetFirstZeroBit(IdSearchStartIndex)};
+	if (!InstanceId.IsValid())
+	{
+		InstanceId = FPrimitiveInstanceId{ValidInstanceIdMask.Add(true)};
+	}
+	// Optimize search for next time
+	IdSearchStartIndex = InstanceId.Id;
+	IdToIndexMap.SetNumUninitialized(ValidInstanceIdMask.Num());
+
+	// if these do not line up, then we are inserting an instance, this is a thing in the editor
+	if (InInstanceAddAtIndex != IndexToIdMap.Num())
+	{
+		check(bIsInsert);
+		IdToIndexMap[InstanceId.Id] = InInstanceAddAtIndex;
+		LOG_INST_DATA(TEXT("IdToIndexMap[%d] = %d"), InstanceId.Id, InInstanceAddAtIndex);
+		// first move all the existing data down one step by inserting the new one
+		IndexToIdMap.Insert(InstanceId, InInstanceAddAtIndex);
+		// then update all the relevant id->index mappings
+		for (int32 Index = InInstanceAddAtIndex + 1; Index < IndexToIdMap.Num(); ++Index)
+		{
+			FPrimitiveInstanceId MovedId = IndexToIdMap[Index];
+			IdToIndexMap[MovedId.Id] = Index;
+			LOG_INST_DATA(TEXT("IdToIndexMap[%d] = %d"), MovedId.GetAsIndex(), Index);
+			InstanceUpdateTracker.MarkIndex<EChangeFlag::IndexChanged>(Index, GetMaxInstanceIndex());
+		}
+
+		// TODO: DO this instead of all the InstanceUpdateTracker.MarkIndex<EChangeFlag::IndexChanged>(Index, GetMaxInstanceIndex()); 
+		// The proxy update can't handle the reordering anyway and not sure we should bother.
+		// InstanceUpdateTracker.MarkAllChanged();
+	}
+	else
+	{
+		int32 InstanceIndex = IndexToIdMap.Num();
+		IdToIndexMap[InstanceId.Id] = InstanceIndex;
+		LOG_INST_DATA(TEXT("IdToIndexMap[%d] = %d"), InstanceId.Id, InstanceIndex);
+
+		IndexToIdMap.Add(InstanceId);
+	}
+	NumInstances = IndexToIdMap.Num();
+	check(ValidInstanceIdMask.Num() >= NumInstances);
+	MarkChangeHelper<EChangeFlag::Added>(InstanceId);
+	LOG_INST_DATA(TEXT("Add(IDX: %d, bIsInsert: %d) -> Id: %d"), InInstanceAddAtIndex, bIsInsert, InstanceId.Id);
+
+	ValidateMapping();
+	return InstanceId;
+}
+
+void FInstanceDataManager::RemoveAtSwap(int32 InstanceIndex)
+{
+	if (GetState() == ETrackingState::Disabled)
+	{
+		return;
+	}
+
+	ValidateMapping();
+
+	FPrimitiveInstanceId InstanceId = IndexToId(InstanceIndex);
+	// resize to the max at once so we don't have to grow piecemeal
+	InstanceUpdateTracker.RemoveAtSwap(InstanceId, InstanceIndex, GetMaxInstanceIndex());
+
+	// If the remove would cause reordering, we create the explicit mapping
+	const bool bCausesReordering = InstanceIndex != NumInstances - 1;
+	if (bCausesReordering && HasIdentityMapping())
+	{
+		CreateExplicitIdentityMapping();
+	}
+
+	MarkComponentRenderInstancesDirty();
+	FreeInstanceId(InstanceId);
+
+	// If we still have the identity mapping, we must be removing the last item
+	if (HasIdentityMapping())
+	{
+		check(!bCausesReordering);
+		--NumInstances;
+		LOG_INST_DATA(TEXT("RemoveAtSwap(IDX: %d) -> Id: %d"), InstanceIndex, InstanceId.Id);
+		return;
+	}
+
+	FPrimitiveInstanceId LastInstanceId = IndexToIdMap.Pop();
+	NumInstances = IndexToIdMap.Num();
+	check(ValidInstanceIdMask.Num() >= NumInstances);
+
+	if (InstanceId != LastInstanceId)
+	{
+		IdToIndexMap[LastInstanceId.Id] = InstanceIndex;
+		LOG_INST_DATA(TEXT("IdToIndexMap[%d] = %d"), LastInstanceId.Id, InstanceIndex);
+		IndexToIdMap[InstanceIndex] = LastInstanceId;
+	}
+	ValidateMapping();
+	LOG_INST_DATA(TEXT("RemoveAtSwap(IDX: %d) -> Id: %d"), InstanceIndex, InstanceId.Id);
+}
+	
+void FInstanceDataManager::RemoveAt(int32 InstanceIndex)
+{
+	if (GetState() == ETrackingState::Disabled)
+	{
+		return;
+	}
+
+	ValidateMapping();
+
+	FPrimitiveInstanceId InstanceId = IndexToId(InstanceIndex);
+
+	InstanceUpdateTracker.RemoveAt(InstanceId, InstanceIndex, GetMaxInstanceIndex());
+		
+	const bool bCausesReordering = InstanceIndex != NumInstances - 1;
+	if (bCausesReordering && HasIdentityMapping())
+	{
+		CreateExplicitIdentityMapping();
+	}
+
+	MarkComponentRenderInstancesDirty();
+	FreeInstanceId(InstanceId);
+
+	// If we still have the identity mapping, do the simplified tracking update
+	if (HasIdentityMapping())
+	{
+		check(!bCausesReordering);
+		--NumInstances;
+		LOG_INST_DATA(TEXT("RemoveAt(IDX: %d) -> Id: %d"), InstanceIndex, InstanceId.Id);
+		return;
+	}
+
+	if (InstanceIndex == IndexToIdMap.Num() - 1)
+	{
+		IndexToIdMap.SetNum(InstanceIndex);
+	}
+	else
+	{
+		IndexToIdMap.RemoveAt(InstanceIndex);
+		for (int32 Index = InstanceIndex; Index < IndexToIdMap.Num(); ++Index)
+		{
+			FPrimitiveInstanceId MovedId = IndexToIdMap[Index];
+			IdToIndexMap[MovedId.Id] = Index;
+			LOG_INST_DATA(TEXT("IdToIndexMap[%d] = %d"), MovedId.GetAsIndex(), Index);
+		}
+	}
+	NumInstances = IndexToIdMap.Num();
+	check(ValidInstanceIdMask.Num() >= NumInstances);
+	LOG_INST_DATA(TEXT("RemoveAt(IDX: %d) -> Id: %d"), InstanceIndex, InstanceId.Id);
+
+	ValidateMapping();
+}
+
+void FInstanceDataManager::TransformChanged(int32 InstanceIndex)
+{
+	LOG_INST_DATA(TEXT("TransformChanged(IDX: %d)"), InstanceIndex);
+	MarkChangeHelper<EChangeFlag::TransformChanged>( InstanceIndex);
+}
+
+void FInstanceDataManager::TransformChanged(FPrimitiveInstanceId InstanceId)
+{
+	LOG_INST_DATA(TEXT("TransformChanged(ID: %d)"), InstanceId.Id);
+	MarkChangeHelper<EChangeFlag::TransformChanged>(InstanceId);
+}
+
+void FInstanceDataManager::TransformsChangedAll()
+{
+	if (GetState() == ETrackingState::Disabled)
+	{
+		return;
+	}
+
+	LOG_INST_DATA(TEXT("TransformsChangedAll(%s)"), TEXT(""));
+	bTransformChangedAllInstances = true;
+	MarkComponentRenderInstancesDirty();
+}
+
+void FInstanceDataManager::CustomDataChanged(int32 InstanceIndex)
+{
+	LOG_INST_DATA(TEXT("CustomDataChanged(IDX: %d)"), InstanceIndex);
+	MarkChangeHelper<EChangeFlag::CustomDataChanged>(InstanceIndex);
+}
+
+void FInstanceDataManager::BakedLightingDataChanged(int32 InstanceIndex)
+{
+	LOG_INST_DATA(TEXT("BakedLightingDataChanged(IDX: %d)"), InstanceIndex);
+	AllChangedFlags.bHasPerInstanceLMSMUVBias = true;
+	MarkComponentRenderInstancesDirty();
+}
+
+void FInstanceDataManager::BakedLightingDataChangedAll()
+{
+	if (GetState() == ETrackingState::Disabled)
+	{
+		return;
+	}
+
+	LOG_INST_DATA(TEXT("BakedLightingDataChangedAll(%s)"), TEXT(""));
+	AllChangedFlags.bHasPerInstanceLMSMUVBias = true;
+	MarkComponentRenderInstancesDirty();
+}
+
+void FInstanceDataManager::NumCustomDataChanged()
+{
+	if (GetState() == ETrackingState::Disabled)
+	{
+		return;
+	}
+
+	LOG_INST_DATA(TEXT("NumCustomDataChanged(%s)"), TEXT(""));
+	bNumCustomDataChanged = true;
+	MarkComponentRenderInstancesDirty();
+}
+
+#if WITH_EDITOR
+
+void FInstanceDataManager::EditorDataChangedAll()
+{
+	if (GetState() == ETrackingState::Disabled)
+	{
+		return;
+	}
+
+	LOG_INST_DATA(TEXT("EditorDataChangedAll(%s)"), TEXT(""));
+	AllChangedFlags.bHasPerInstanceEditorData = true;
+	MarkComponentRenderInstancesDirty();
+}
+
+#endif
+
+void FInstanceDataManager::PrimitiveTransformChanged()
+{
+	if (GetState() == ETrackingState::Disabled)
+	{
+		return;
+	}
+
+	LOG_INST_DATA(TEXT("PrimitiveTransformChanged(%s)"), TEXT(""));
+	bPrimitiveTransformChanged = true;
+	MarkComponentRenderInstancesDirty();
+}
+
+bool FInstanceDataManager::HasAnyInstanceChanges() const
+{
+	return bAnyInstanceChange || bNumCustomDataChanged || AllChangedFlags.Packed != 0u || bTransformChangedAllInstances
+		|| InstanceUpdateTracker.HasAnyChanges();
+}
+
+/**
+ * Describes what has changed, that can be derived from the primitive desc, or internal tracking state.
+ */
+struct FChangeDesc2
+{
+	FInstanceDataFlags ChangedFlags;
+	union
+	{
+		struct 
+		{
+			uint8 bUntrackedState : 1;
+			uint8 bInstancesChanged : 1;
+			uint8 bPrimitiveTransformChanged : 1;
+			uint8 bMaterialUsageFlagsChanged : 1;
+			uint8 bMaxDisplacementChanged : 1;
+			uint8 bStaticMeshBoundsChanged : 1;
+		};
+		uint32 Packed;
+	};
+	
+	FChangeDesc2(bool bFullChange = false)
+	{
+		Packed = 0u;
+		bUntrackedState = bFullChange;
+	}
+
+	bool HasAnyChange() const
+	{
+		return Packed != 0u || ChangedFlags.Packed != 0u;
+	}
+};
+
+void FInstanceDataManager::GatherDefaultData(const FInstanceDataManagerSourceDataDesc& ComponentData, FInstanceUpdateChangeSet& ChangeSet) const
+{
+	// Collect the delta data to be able to update the index mapping.
+	ChangeSet.MaxInstanceId = GetMaxAllocatedInstanceId();
+	ChangeSet.bIdentityIdMap = IsIdentity();
+
+	const FVector3f PrimitiveWorldSpacePositionHigh = FDFVector3{ ComponentData.PrimitiveLocalToWorld.GetOrigin() }.High;
+	ChangeSet.PrimitiveWorldSpaceOffset = FVector{ PrimitiveWorldSpacePositionHigh };
+	ChangeSet.PrimitiveToRelativeWorld = FDFMatrix::MakeToRelativeWorldMatrix(PrimitiveWorldSpacePositionHigh, ComponentData.PrimitiveLocalToWorld).M;
+	ChangeSet.AbsMaxDisplacement = AbsMaxDisplacement;
+	if (ChangeSet.Flags.bHasPerInstanceCustomData)
+	{
+		ChangeSet.NumCustomDataFloats = NumCustomDataFloats;
+	}
+
+	// Only gather data after the rest is initialized
+	if (!ChangeSet.bIdentityIdMap)
+	{
+		Gather(ChangeSet.GetIndexChangedDelta(), ChangeSet.IndexToIdMapDeltaData, IndexToIdMap);
+	}
+		
+	// Patch up the special local bounds data.
+	if (!ChangeSet.Flags.bHasPerInstanceLocalBounds)
+	{
+		check(ChangeSet.InstanceLocalBounds.IsEmpty());
+		// This is the odd one out.
+		ChangeSet.SetSharedLocalBounds(ComponentData.MeshBounds);
+	}
+}
+
+void FInstanceDataManager::InitChangeSet(const FInstanceDataManagerSourceDataDesc& ComponentData, FInstanceUpdateChangeSet& ChangeSet) const
+{
+	// Force all data since we don't use the delta state
+	ChangeSet.ForceFullFlags = ComponentData.Flags; 
+	ChangeSet.Flags = ComponentData.Flags;
+
+	GatherDefaultData(ComponentData, ChangeSet);
+}
+
+void FInstanceDataManager::InitChangeSet(const FChangeDesc2 &ChangeDesc, const FInstanceDataManagerSourceDataDesc& ComponentData, FInstanceUpdateChangeSet& ChangeSet)
+{
+	ChangeSet.ForceFullFlags = ChangeDesc.ChangedFlags; 
+	// also trigger full copy if num changed.
+	ChangeSet.ForceFullFlags.bHasPerInstanceCustomData = ChangeSet.ForceFullFlags.bHasPerInstanceCustomData || bNumCustomDataChanged;
+	ChangeSet.Flags = ComponentData.Flags;
+
+	GatherDefaultData(ComponentData, ChangeSet);
+}
+
+bool FInstanceDataManager::FlushChanges(FInstanceDataManagerSourceDataDesc&& ComponentData)
+{
+	if (GetState() == ETrackingState::Disabled)
+	{
+		return false;
+	}
+
+	// if bHasPerInstanceCustomData is set, then NumCustomDataFloats must be non-zero
+	check(!ComponentData.Flags.bHasPerInstanceCustomData || ComponentData.NumCustomDataFloats != 0);
+
+	// for a proxy that is not visible to the render thread we can dispatch the update task immediately, saving queueing a RT command and potentially increasing overlap by some ammount.
+	const bool bIsUnattached = !InstanceDataProxy.IsValid() || InstanceDataProxy.GetSharedReferenceCount() == 1;
+	if (!InstanceDataProxy.IsValid())
+	{
+		InstanceDataProxy = GetOrCreateProxyInternal();
+	}
+
+	float NewAbsMaxDisplacement = FMath::Max(-ComponentData.PrimitiveMaterialDesc.MinMaxMaterialDisplacement.X, ComponentData.PrimitiveMaterialDesc.MinMaxMaterialDisplacement.Y)
+		+ ComponentData.PrimitiveMaterialDesc.MaxWorldPositionOffsetDisplacement;
+
+	FChangeDesc2 ChangeDesc;
+	ChangeDesc.bUntrackedState = GetState() != ETrackingState::Tracked && GetState() != ETrackingState::Optimized;
+
+	// Figure out the deltas.
+	if (!ChangeDesc.bUntrackedState)
+	{
+		ChangeDesc.bInstancesChanged =  HasAnyInstanceChanges();
+
+		// TODO: make all this depend on tracking rather than diffing (we may do diffing as a means of debugging perhaps)
+		ChangeDesc.bPrimitiveTransformChanged = bTransformChangedAllInstances || bPrimitiveTransformChanged; 
+		ChangeDesc.bMaterialUsageFlagsChanged = Flags != ComponentData.Flags;
+		ChangeDesc.bMaxDisplacementChanged = AbsMaxDisplacement != NewAbsMaxDisplacement;
+		ChangeDesc.bStaticMeshBoundsChanged = !MeshBounds.Equals(ComponentData.MeshBounds);
+		ChangeDesc.ChangedFlags.Packed = Flags.Packed ^ ComponentData.Flags.Packed;
+	}
+	
+	//TBD: leave in the proxy?
+	// FMatrix PrevPrimitiveLocalToWorld = PrimitiveLocalToWorld;
+	
+	// Update the tracked state
+	// PrimitiveLocalToWorld = ComponentData.PrimitiveLocalToWorld;
+	AbsMaxDisplacement = NewAbsMaxDisplacement;
+	MeshBounds = ComponentData.MeshBounds;
+	Flags = ComponentData.Flags;
+
+	// TODO: detect change and toggle full update
+	NumCustomDataFloats = ComponentData.NumCustomDataFloats;
+
+	if (!ChangeDesc.HasAnyChange())
+	{
+		return false;
+	}
+	
+	// If we got here & the state is "optimized" then we know the precomputed data is now invalid and we ditch it.
+	// This should not happen in a cooked client, ideally.
+	if (GetState() == ETrackingState::Optimized)
+	{
+		UE_LOG(LogTemp, Log, TEXT("Discarded PrecomputedOptimizationData"));
+		PrecomputedOptimizationData.Reset();
+	}
+
+	// After an update has been sent, we need to track all deltas.
+	ETrackingState SuccessorTrackingState = ETrackingState::Tracked;
+	{
+#if WITH_EDITOR
+		// We should rebuild the optimization data if the cvar is set & this is not the first update and we have cooked data.
+		// TODO: this may not be 100% robust against things like invalidations triggered by unexpected stuff, like BP machinations.
+		//       Most such events also need to clear PrecomputedOptimizationData, but must be carefully done to not clear when actually wanted (in cooked). 
+		bool bRebuildOptimized = CVarInstanceDataManagerBuildOptimizedForTesting.GetValueOnGameThread()
+			&& !(ChangeDesc.bUntrackedState && PrecomputedOptimizationData.IsValid());
+#else
+		constexpr bool bRebuildOptimized = false;
+#endif
+		bool bNeedFullUpdate = ChangeDesc.bUntrackedState 
+			// Note: this is a bit inefficient and we could resend just the relevant attributes, but happens only when the material changes which is hopefully not something we need to optimize for.
+			|| ChangeDesc.bMaterialUsageFlagsChanged 
+			// it was optimized so we need to build everything 
+			|| GetState() == ETrackingState::Optimized
+			|| bRebuildOptimized;
+
+		// NOTE: Moving the update tracker to the change set implicitly resets it.
+		FInstanceUpdateChangeSet ChangeSet(
+			bNeedFullUpdate, 
+			MoveTemp(InstanceUpdateTracker),
+			ComponentData.NumInstances);
+
+		ChangeSet.bUpdateAllInstanceTransforms = ChangeDesc.bPrimitiveTransformChanged || bTransformChangedAllInstances;
+		// Initialize the change set before collecting instance change data.
+		InitChangeSet(ChangeDesc, ComponentData, ChangeSet);
+
+		// Callback to the owner to fill in change data.
+		ComponentData.BuildChangeSet(ChangeSet);
+
+		// make sure the custom data change is correctly tracked
+		check(!ChangeSet.Flags.bHasPerInstanceCustomData || NumCustomDataFloats == ChangeSet.NumCustomDataFloats);
+		
+		// TODO: add more cross validation here somehow
+		//check(!ChangeSet.Flags.bHasPerInstanceCustomData || ChangeSet.PerInstanceCustomData.);
+		check(ChangeSet.Flags.bHasPerInstanceLocalBounds || ChangeSet.InstanceLocalBounds.Num() == 1);
+		//check(!ChangeSet.Flags.bHasPerInstanceLocalBounds || ChangeSet.InstanceLocalBounds.Num() == ChangeSet.GetTransformDelta().GetNumItems());
+
+
+		// If we have per-instance previous local to world, they are expected to be in the local space of the _previous_ local to world. If they are in fact not (e.g., if someone sets them explicitly from world space) then, well, this won't be correct
+		if (ChangeSet.Flags.bHasPerInstanceDynamicData)
+		{
+			// TODO: move supplying this out to the ComponentData desc.
+			TOptional<FTransform> PreviousTransform = FMotionVectorSimulation::Get().GetPreviousTransform(PrimitiveComponent.Get());
+			if (PreviousTransform.IsSet())
+			{
+				ChangeSet.PreviousPrimitiveToRelativeWorld = FDFMatrix::MakeToRelativeWorldMatrix(FVector3f{ ChangeSet.PrimitiveWorldSpaceOffset }, PreviousTransform.GetValue().ToMatrixWithScale()).M;
+			}
+			// TODO: Move this to the proxy
+			//else
+			//{
+			//	ChangeSet.PreviousPrimitiveToRelativeWorld = FDFMatrix::MakeToRelativeWorldMatrix(FVector3f{ ChangeSet.PrimitiveWorldSpaceOffset }, PrevPrimitiveLocalToWorld).M;				
+			//}
+		}
+
+		// Note: this does affect TrackingState but we reset this below anyways.
+		ClearChangeTracking();
+
+		// Assemble header info to enable nonblocking primitive update.
+		FInstanceDataBufferHeader InstanceDataBufferHeader;
+
+		InstanceDataBufferHeader.NumInstances = ChangeSet.NumSourceInstances;
+		InstanceDataBufferHeader.PayloadDataStride = FInstanceSceneDataBuffers::CalcPayloadDataStride(ChangeSet.Flags, ChangeSet.NumCustomDataFloats, 0);
+		InstanceDataBufferHeader.Flags = ChangeSet.Flags;
+
+		CSV_CUSTOM_STAT_GLOBAL(NumInstanceTransformUpdates, ChangeSet.Transforms.Num(), ECsvCustomStatOp::Accumulate);
+
+#if WITH_EDITOR
+		if (bRebuildOptimized)
+		{
+			check(ChangeSet.IsFullUpdate());
+			LOG_INST_DATA(TEXT("Optimized Rebuild"));
+			// check(ComponentData.ComponentMobility == EComponentMobility::Static || ComponentData.ComponentMobility == EComponentMobility::Stationary);
+			DispatchInstanceDataUpdateTask(bIsUnattached, InstanceDataProxy, InstanceDataBufferHeader, [ChangeSet = MoveTemp(ChangeSet), InstanceDataProxy = InstanceDataProxy, NumHeaderInstances = InstanceDataBufferHeader.NumInstances] () mutable 
+			{
+				// Build the optimization data on the fly. This path is only for testing purposes as we otherwise want the data to be pre-cooked & passed along.
+				ChangeSet.PrecomputedOptimizationData = MakeShared<FPrecomputedInstanceSpatialHashData>(PrecomputeOptimizationData(ChangeSet));
+				InstanceDataProxy->Build(MoveTemp(ChangeSet));
+				check(NumHeaderInstances == InstanceDataProxy->GeInstanceSceneDataBuffers()->GetNumInstances());
+			});
+			SuccessorTrackingState = ETrackingState::Optimized;
+		}
+		else 
+#endif // WITH_EDITOR
+		if (bNeedFullUpdate)
+		{
+			LOG_INST_DATA(TEXT("Full Build %s"), TEXT(""));
+			if (PrecomputedOptimizationData.IsValid())
+			{
+				LOG_INST_DATA(TEXT("  Optimized Build (%s)"), PrecomputedOptimizationData.IsValid() ? TEXT("Precomputed") : TEXT(""));
+				ChangeSet.PrecomputedOptimizationData = PrecomputedOptimizationData;
+				SuccessorTrackingState = ETrackingState::Optimized;
+			}
+			DispatchInstanceDataUpdateTask(bIsUnattached, InstanceDataProxy, InstanceDataBufferHeader, [ChangeSet = MoveTemp(ChangeSet), InstanceDataProxy = InstanceDataProxy, NumHeaderInstances = InstanceDataBufferHeader.NumInstances] () mutable 
+			{
+				InstanceDataProxy->Build(MoveTemp(ChangeSet));
+				check(NumHeaderInstances == InstanceDataProxy->GeInstanceSceneDataBuffers()->GetNumInstances());
+			});
+		}
+		else
+		{
+			LOG_INST_DATA(TEXT("Delta Update %s"), TEXT(""));
+			check(!ChangeDesc.bUntrackedState);
+			DispatchInstanceDataUpdateTask(bIsUnattached, InstanceDataProxy, InstanceDataBufferHeader, [ChangeSet = MoveTemp(ChangeSet), InstanceDataProxy = InstanceDataProxy, NumHeaderInstances = InstanceDataBufferHeader.NumInstances] () mutable 
+			{
+				InstanceDataProxy->Update(MoveTemp(ChangeSet));
+				check(NumHeaderInstances == InstanceDataProxy->GeInstanceSceneDataBuffers()->GetNumInstances());
+			});
+		}
+	}
+
+	TrackingState = SuccessorTrackingState;
+	return true;
+}
+
+void FInstanceDataManager::ClearChangeTracking()
+{
+	if (GetState() == ETrackingState::Disabled)
+	{
+		return;
+	}
+
+	// When tracking data is cleared, we loose connection to previously tracked state until the next update is sent.
+	TrackingState = ETrackingState::Initial;
+
+	InstanceUpdateTracker.Reset();
+	bNumCustomDataChanged = false;
+	AllChangedFlags.Packed = 0u;
+	bTransformChangedAllInstances = false;
+	bPrimitiveTransformChanged = false;
+	bAnyInstanceChange = false;
+}
+
+int32 FInstanceDataManager::GetMaxAllocatedInstanceId() const
+{
+	return HasIdentityMapping() ? NumInstances : ValidInstanceIdMask.Num();
+}
+
+void FInstanceDataManager::CreateExplicitIdentityMapping()
+{
+	check(HasIdentityMapping());
+	IndexToIdMap.SetNumUninitialized(NumInstances);
+	IdToIndexMap.SetNumUninitialized(NumInstances);
+	for (int32 Index = 0; Index < NumInstances; ++Index)
+	{
+		IndexToIdMap[Index] = FPrimitiveInstanceId{Index};
+		IdToIndexMap[Index] = Index;
+	}
+	ValidInstanceIdMask.Reset();
+	ValidInstanceIdMask.SetNum(NumInstances, true);
+	IdSearchStartIndex = NumInstances;
+}
+
+template<FInstanceDataManager::EChangeFlag Flag>
+void FInstanceDataManager::MarkChangeHelper(int32 InstanceIndex)
+{
+	if (GetState() == ETrackingState::Disabled)
+	{
+		return;
+	}
+
+	if (GetState() != ETrackingState::Tracked)
+	{
+		bAnyInstanceChange = true;
+		MarkComponentRenderInstancesDirty();
+		return;
+	}
+	InstanceUpdateTracker.MarkIndex<Flag>(InstanceIndex, GetMaxInstanceIndex());
+	MarkComponentRenderInstancesDirty();
+}
+
+template<FInstanceDataManager::EChangeFlag Flag>
+void FInstanceDataManager::MarkChangeHelper(FPrimitiveInstanceId InstanceId)
+{
+	if (GetState() != ETrackingState::Tracked)
+	{
+		bAnyInstanceChange = true;
+		MarkComponentRenderInstancesDirty();
+		return;
+	}
+	MarkChangeHelper<Flag>(IdToIndex(InstanceId));
+}
+
+void FInstanceDataManager::MarkComponentRenderInstancesDirty()
+{
+	if (UPrimitiveComponent *PrimitiveComponentPtr = PrimitiveComponent.Get())
+	{
+		PrimitiveComponentPtr->MarkRenderInstancesDirty();
+	}
+}
+
+bool FInstanceDataManager::HasIdentityMapping() const
+{
+	return IndexToIdMap.IsEmpty();
+}
+
+void FInstanceDataManager::FreeInstanceId(FPrimitiveInstanceId InstanceId)
+{
+	LOG_INST_DATA(TEXT("FreeInstanceId(Id: %d)"), InstanceId.Id);
+
+	if (!HasIdentityMapping())
+	{
+		IdToIndexMap[InstanceId.Id] = INDEX_NONE;
+		ValidInstanceIdMask[InstanceId.Id] = false;
+		// Must start from the lowest free index since we'd otherwise get holes when adding things.
+		IdSearchStartIndex = FMath::Min(IdSearchStartIndex, InstanceId.Id);
+	}
+
+	LOG_INST_DATA(TEXT("IdToIndexMap[%d] = %d"), InstanceId.Id, INDEX_NONE);
+}
+
+TSharedPtr<FUpdatableInstanceDataSceneProxy, ESPMode::ThreadSafe> FInstanceDataManager::GetOrCreateProxyInternal()
+{
+	LOG_INST_DATA(TEXT("GetOrCreateProxy"));
+	if (!InstanceDataProxy)
+	{
+		InstanceDataProxy = MakeShared<FUpdatableInstanceDataSceneProxy>();
+	}
+
+	return InstanceDataProxy;
+}
+
+TSharedPtr<FInstanceDataSceneProxy, ESPMode::ThreadSafe> FInstanceDataManager::GetOrCreateProxy()
+{
+	return GetOrCreateProxyInternal();
+}
+
+TSharedPtr<FInstanceDataSceneProxy, ESPMode::ThreadSafe> FInstanceDataManager::GetProxy()
+{
+	return InstanceDataProxy;
+}
+
+#if DO_GUARD_SLOW
+void FInstanceDataManager::ValidateMapping() const
+{
+	check(HasIdentityMapping() || IndexToIdMap.Num() == NumInstances);
+	for (int32 Index = 0; Index < IndexToIdMap.Num(); ++Index)
+	{
+		FPrimitiveInstanceId Id = IndexToIdMap[Index];
+		check(ValidInstanceIdMask[Id.GetAsIndex()]);
+		check(Index == IdToIndexMap[Id.GetAsIndex()]);
+	}
+	for (int32 Id = 0; Id < IdToIndexMap.Num(); ++Id)
+	{
+		int32 Index = IdToIndexMap[Id];
+		if (Index != INDEX_NONE)
+		{
+			check(ValidInstanceIdMask[Id]);
+			check(IndexToIdMap[Index].GetAsIndex() == Id);
+		}
+		else
+		{
+			check(!ValidInstanceIdMask[Id]);
+		}
+	}
+	int32 FirstFalse = ValidInstanceIdMask.Find(false);
+	check(FirstFalse < 0 || FirstFalse >= IdSearchStartIndex);
+}
+#endif
+
+SIZE_T FInstanceDataManager::GetAllocatedSize() const
+{
+	return ValidInstanceIdMask.GetAllocatedSize() +
+		InstanceUpdateTracker.GetAllocatedSize();
+}
+
+static bool ShouldUsePrecomputed()
+{
+	static const auto CVarPrecomputed = IConsoleManager::Get().FindTConsoleVariableDataInt(TEXT("r.SceneCulling.Precomputed"));
+	static const auto CVarSceneCull = IConsoleManager::Get().FindTConsoleVariableDataInt(TEXT("r.SceneCulling"));
+
+	return CVarSceneCull && CVarSceneCull->GetValueOnAnyThread() != 0
+		&& CVarPrecomputed && CVarPrecomputed->GetValueOnAnyThread() != 0;
+}
+
+
+#if WITH_EDITOR
+
+bool FInstanceDataManager::ShouldWriteCookedData(const ITargetPlatform* TargetPlatform, int32 NumInstancesToBuildFor)
+{
+	EComponentMobility::Type Mobility = PrimitiveComponent.IsValid() ? PrimitiveComponent->Mobility.GetValue() : EComponentMobility::Type::Movable;
+
+	// Only cook for static & stationary
+	bool bValidTypeAndMobility = true;//(Mobility == EComponentMobility::Static || Mobility == EComponentMobility::Stationary);
+
+	static const auto MinInstanceCountToOptimizeCVar = IConsoleManager::Get().FindConsoleVariable(TEXT("r.InstanceData.MinInstanceCountToOptimize"));
+	// Default to 2 if the cvar doesn't exist for some reason.
+	const int32 MinInstanceCountToOptimizeFor = MinInstanceCountToOptimizeCVar != nullptr ? MinInstanceCountToOptimizeCVar->GetInt() : 2;
+
+	return bValidTypeAndMobility
+		&& NumInstancesToBuildFor >= MinInstanceCountToOptimizeFor
+		&& ShouldUsePrecomputed()
+		&& DoesTargetPlatformSupportNanite(TargetPlatform);
+
+}
+
+void FInstanceDataManager::BeginCacheForCookedPlatformData(const ITargetPlatform* TargetPlatform, FInstanceDataManagerSourceDataDesc&& ComponentData)
+{
+	// Already precomputed, we don't need to do it twice (could add checks to see that it is not incorrect for some obscure reason)
+	if (PrecomputedOptimizationData.IsValid())
+	{
+		return;
+	}
+	
+	bool bShouldBuild = ShouldWriteCookedData(TargetPlatform, ComponentData.NumInstances);
+
+	// TODO: we could kick an async thread here if that is preferrable for the cooker?
+	if (bShouldBuild && ComponentData.BuildChangeSet)
+	{
+		uint32 StartTime = FPlatformTime::Cycles();
+		PrecomputedOptimizationData = MakeShared<FPrecomputedInstanceSpatialHashData>(PrecomputeOptimizationData(MoveTemp(ComponentData)));
+		uint32 EndtTime = FPlatformTime::Cycles();
+
+		UE_LOG(LogTemp, Log, TEXT("Build Instance Spatial Hashes (%.2fms), Instances: %d, Hashes: %d, Remap Size: %d"), FPlatformTime::ToMilliseconds( EndtTime - StartTime ), ComponentData.NumInstances, PrecomputedOptimizationData->Hashes.Num(), PrecomputedOptimizationData->ProxyIndexToComponentIndexRemap.Num());
+	}	
+}
+
+class FSpatialHashSortBuilder
+{
+public:
+	//int32 FirstLevel;
+	struct FSortedInstanceItem
+	{
+		RenderingSpatialHash::FLocation64 InstanceLoc;
+		int32 InstanceIndex;
+	};
+
+	template <typename GetWorldSpaceInstanceSphereFuncType>
+	void BuildOptimizedSpatialHashOrder(int32 NumInstances, GetWorldSpaceInstanceSphereFuncType &&GetWorldSpaceInstanceSphere)
+	{
+		// TODO: Fix the cvar interaction here, need to unify code path with SceneCulling, or something maybe even come up with some heuristic to allow using fewer bits?
+		int32 FirstLevel = 0;
+		{
+			static const auto CVarInstanceHierarchyMinCellSize = IConsoleManager::Get().FindTConsoleVariableDataFloat(TEXT("r.SceneCulling.MinCellSize"));
+			if (CVarInstanceHierarchyMinCellSize)
+			{
+				// TODO: only one code path to compute this value!!!
+				FirstLevel = RenderingSpatialHash::CalcLevel(CVarInstanceHierarchyMinCellSize->GetValueOnAnyThread() - 1.0);
+			}
+		}
+
+		SortedInstances.Reserve(NumInstances);
+		for (int32 InstanceIndex = 0; InstanceIndex < NumInstances; ++InstanceIndex)
+		{
+			FSphere InstanceWorldSpaceSphere = GetWorldSpaceInstanceSphere(InstanceIndex);
+
+			RenderingSpatialHash::FLocation64 InstanceLoc = RenderingSpatialHash::CalcLevelAndLocationClamped(InstanceWorldSpaceSphere.Center, InstanceWorldSpaceSphere.W, FirstLevel);
+
+			FSortedInstanceItem& Item = SortedInstances.AddDefaulted_GetRef();
+			Item.InstanceLoc = InstanceLoc;
+			Item.InstanceIndex = InstanceIndex;
+		}
+	
+		// Sort the instances according to hash location (first level, then coordinate) and last on instance Index.
+		SortedInstances.Sort(
+			[](const FSortedInstanceItem& A, const FSortedInstanceItem& B) -> bool
+			{
+				if (A.InstanceLoc.Level != B.InstanceLoc.Level)
+				{
+					return A.InstanceLoc.Level < B.InstanceLoc.Level;
+				}
+				if (A.InstanceLoc.Coord.X != B.InstanceLoc.Coord.X)
+				{
+					return A.InstanceLoc.Coord.X < B.InstanceLoc.Coord.X;
+				}
+				if (A.InstanceLoc.Coord.Y != B.InstanceLoc.Coord.Y)
+				{
+					return A.InstanceLoc.Coord.Y < B.InstanceLoc.Coord.Y;
+				}
+				if (A.InstanceLoc.Coord.Z != B.InstanceLoc.Coord.Z)
+				{
+					return A.InstanceLoc.Coord.Z < B.InstanceLoc.Coord.Z;
+				}
+				return A.InstanceIndex < B.InstanceIndex;
+			}
+		);
+	}
+	TArray<FSortedInstanceItem> SortedInstances;
+};
+
+FPrecomputedInstanceSpatialHashData FInstanceDataManager::PrecomputeOptimizationData(FInstanceDataManagerSourceDataDesc&& ComponentData)
+{
+	FInstanceUpdateChangeSet ChangeSet(ComponentData.NumInstances, ComponentData.Flags);
+	InitChangeSet(ComponentData, ChangeSet);
+
+	// Callback to the owner to fill in change-set data.
+	// Note: this makes a copy of the data, which is somewhat wasteful but gets the format converted from whatever the owner might have, otherwise we'd need some other abstraction for the data here.
+	ComponentData.BuildChangeSet(ChangeSet);
+
+	return PrecomputeOptimizationData(ChangeSet);
+}
+
+FPrecomputedInstanceSpatialHashData FInstanceDataManager::PrecomputeOptimizationData(FInstanceUpdateChangeSet& ChangeSet)
+{
+	FSpatialHashSortBuilder SortBuilder;
+
+	SortBuilder.BuildOptimizedSpatialHashOrder(ChangeSet.NumSourceInstances,
+		[&](int32 InstanceIndex) -> FSphere
+		{
+			int32 BoundsIndex = FMath::Min(ChangeSet.InstanceLocalBounds.Num() - 1, InstanceIndex);
+			// this can totally be optimized
+			FSphere3f LocalInstanceSphere = FSphere3f(ChangeSet.InstanceLocalBounds[BoundsIndex].GetCenter(), ChangeSet.InstanceLocalBounds[BoundsIndex].GetExtent().Size());
+			FRenderTransform InstanceLocalToWorld = ChangeSet.Transforms[InstanceIndex] * ChangeSet.PrimitiveToRelativeWorld;
+			FSphere Result = FSphere(LocalInstanceSphere.TransformBy(InstanceLocalToWorld.ToMatrix44f()));
+			Result.Center += ChangeSet.PrimitiveWorldSpaceOffset;
+			return Result; 
+		}				
+	);
+
+	FPrecomputedInstanceSpatialHashData Result;
+
+	// Pack down the spatial hashes & index remap
+	Result.ProxyIndexToComponentIndexRemap.SetNumUninitialized(ChangeSet.NumSourceInstances);
+
+	FInstanceSceneDataBuffers::FCompressedSpatialHashItem CurrentItem;
+	CurrentItem.NumInstances = 0;
+
+	bool bIsIdentityIndexMap = true;
+
+	for (int32 InstanceIndex = 0; InstanceIndex < SortBuilder.SortedInstances.Num(); ++InstanceIndex)
+	{
+		int32 ComponentInstanceIndex = SortBuilder.SortedInstances[InstanceIndex].InstanceIndex;
+		bIsIdentityIndexMap = bIsIdentityIndexMap && InstanceIndex == ComponentInstanceIndex;
+		Result.ProxyIndexToComponentIndexRemap[InstanceIndex] = ComponentInstanceIndex;
+
+		bool bSameLoc = CurrentItem.NumInstances > 0 && CurrentItem.Location == SortBuilder.SortedInstances[InstanceIndex].InstanceLoc;
+		if (bSameLoc)
+		{
+			CurrentItem.NumInstances += 1;
+		}
+		else
+		{
+			if (CurrentItem.NumInstances > 0)
+			{
+				Result.Hashes.Add(CurrentItem);
+			}
+			CurrentItem.Location = SortBuilder.SortedInstances[InstanceIndex].InstanceLoc;
+			CurrentItem.NumInstances = 1;
+		}
+	}
+	if (CurrentItem.NumInstances > 0)
+	{
+		Result.Hashes.Add(CurrentItem);
+	}
+
+	// Don't store a 1:1 mapping
+	if (bIsIdentityIndexMap)
+	{
+		Result.ProxyIndexToComponentIndexRemap.Reset();
+	}
+
+	return Result;
+}
+
+void FInstanceDataManager::WriteCookedRenderData(FArchive& Ar, FInstanceDataManagerSourceDataDesc&& ComponentData)
+{
+	bool bHasCookedData = false;
+
+	bool bShouldBuild = ShouldWriteCookedData(Ar.CookingTarget(), ComponentData.NumInstances);
+
+	if (bShouldBuild)
+	{
+		if (!PrecomputedOptimizationData.IsValid())
+		{
+			if (ComponentData.BuildChangeSet)
+			{
+				PrecomputedOptimizationData = MakeShared<const FPrecomputedInstanceSpatialHashData>(PrecomputeOptimizationData(MoveTemp(ComponentData)));
+			}
+		}
+		
+		if (PrecomputedOptimizationData.IsValid())
+		{
+			// We have to copy the whole thing to be able to serialize?
+			FPrecomputedInstanceSpatialHashData OptData = *PrecomputedOptimizationData;
+
+			// Serialize the stuff we need.
+			bHasCookedData = true;
+			Ar << bHasCookedData;
+
+			OptData.Hashes.BulkSerialize(Ar);
+			OptData.ProxyIndexToComponentIndexRemap.BulkSerialize(Ar);
+		}
+	}
+
+	if (!bHasCookedData)
+	{
+		// write the bool if we didn't write any data previously
+		Ar << bHasCookedData;
+	}
+}
+#endif // WITH_EDITOR
+
+void FInstanceDataManager::ReadCookedRenderData(FArchive& Ar)
+{
+	bool bHasCookedData = false;
+	Ar << bHasCookedData;
+
+	if (bHasCookedData)
+	{
+		FPrecomputedInstanceSpatialHashData Tmp;
+
+		// TODO: Pack the data representation to far fewer bits
+		Tmp.Hashes.BulkSerialize(Ar);
+		// TODO: RLE-compress
+		Tmp.ProxyIndexToComponentIndexRemap.BulkSerialize(Ar);
+
+		// Ditch the precomputed data if it has been disabled (in the runtime), even if the cook was done with the data enabled.
+		if (ShouldUsePrecomputed())
+		{
+			PrecomputedOptimizationData = MakeShared<const FPrecomputedInstanceSpatialHashData>(Tmp);
+		}
+		else
+		{
+			PrecomputedOptimizationData.Reset();
+		}
+	}
+}
+
+void FInstanceDataManager::Serialize(FArchive& Ar, bool bCookedOrCooking)
+{
+	FInstanceIdIndexMap::Serialize(Ar);
+
+	// if we're loading a non-identity map then restore the ID allocation
+	if (Ar.IsLoading())
+	{
+		ValidInstanceIdMask.Empty(IdToIndexMap.Num());
+		if (!IdToIndexMap.IsEmpty())
+		{
+			ValidInstanceIdMask.SetNum(IdToIndexMap.Num(), false);
+
+			for (int32 InstanceIndex = 0; InstanceIndex < IndexToIdMap.Num(); ++InstanceIndex)
+			{
+				FPrimitiveInstanceId Id = IndexToIdMap[InstanceIndex];
+				ValidInstanceIdMask[Id.GetAsIndex()] = true;
+			}
+
+		}
+		IdSearchStartIndex = 0;
+		ClearChangeTracking();
+	}
+}
+
+TArray<int32> FInstanceDataManager::Optimize(FInstanceDataManagerSourceDataDesc&& ComponentData, bool bShouldRetainIdMap)
+{
+#if WITH_EDITOR
+	FPrecomputedInstanceSpatialHashData OptData = PrecomputeOptimizationData(MoveTemp(ComponentData));
+	
+	// Note: Currently this just ditches the spatial hash data again. This is the simple and robust solution because of something mutates the data
+	//       nothing breaks as we recompute the optimization data anyway during cook. In the case where nothing has changed we'll detect the identity reodering and ditch the reorder table.
+
+	//
+	if (bShouldRetainIdMap)
+	{
+		if (!OptData.ProxyIndexToComponentIndexRemap.IsEmpty())
+		{
+			if (HasIdentityMapping())
+			{
+				CreateExplicitIdentityMapping();
+			}
+		
+			TArray<FPrimitiveInstanceId> OldIndexToIdMap = MoveTemp(IndexToIdMap);
+			IndexToIdMap.SetNumUninitialized(OldIndexToIdMap.Num());
+
+			for (int32 NewIndex = 0; NewIndex < OptData.ProxyIndexToComponentIndexRemap.Num(); ++NewIndex)
+			{
+				int32 OldIndex = OptData.ProxyIndexToComponentIndexRemap[NewIndex];
+				FPrimitiveInstanceId InstanceId = OldIndexToIdMap[OldIndex];
+				Update(InstanceId, NewIndex);
+			}
+		}
+	}
+	else
+	{
+		// As we've moved the instances around we'll drop the ID mapping noew
+		Reset(GetMaxInstanceIndex());
+	}
+	return OptData.ProxyIndexToComponentIndexRemap;
+#else
+	return TArray<int32>();
+#endif // WITH_EDITOR
+}
diff --git a/Engine/Source/Runtime/Engine/Private/InstanceData/InstanceDataUpdateUtils.h b/Engine/Source/Runtime/Engine/Private/InstanceData/InstanceDataUpdateUtils.h
new file mode 100644
--- /dev/null
+++ b/Engine/Source/Runtime/Engine/Private/InstanceData/InstanceDataUpdateUtils.h
@@ -0,0 +1,436 @@
+// Copyright Epic Games, Inc. All Rights Reserved.
+
+#pragma once
+
+#include "CoreTypes.h"
+#include "HAL/UnrealMemory.h"
+#include "Math/UnrealMathUtility.h"
+#include "Containers/Array.h"
+#include "RenderingThread.h"
+#include "RenderTransform.h"
+#include "InstanceDataSceneProxy.h"
+#include "InstanceData/InstanceDataManager.h"
+#include "InstanceData/InstanceUpdateChangeSet.h"
+
+#define INSTANCE_DATA_UPDATE_ENABLE_ASYNC_TASK 1
+
+/**
+ * Write InSource that was previously gathered using the above methods to the final destination array OutDest
+ * using the same delta information. 
+ * If there is no delta and the index remap is identity, it performs a move of the source data to the final array, saving a malloc & copy.
+ */
+template <typename DeltaType, typename ValueType, typename IndexRemapType>
+void Scatter(const DeltaType& Delta, TArray<ValueType>& OutDest, int32 DestNumElements, TArray<ValueType>&& InSource, const IndexRemapType& IndexRemap, int32 ElementStride = 1)
+{
+	check(InSource.Num() == Delta.GetNumItems() * ElementStride);
+	if (Delta.IsDelta() || !IndexRemap.IsIdentity())
+	{
+		OutDest.SetNumUninitialized(DestNumElements * ElementStride);
+		for (auto It = Delta.GetIterator(); It; ++It)
+		{
+			int32 SrcIndex = It.GetItemIndex();
+			int32 DestIndex = It.GetIndex();
+
+			if (IndexRemap.Remap(SrcIndex, DestIndex))
+			{
+				FMemory::Memcpy(&OutDest[DestIndex * ElementStride], &InSource[SrcIndex * ElementStride], ElementStride * sizeof(ValueType));
+			}
+		}
+	}
+	else
+	{
+		check(InSource.Num() == DestNumElements * ElementStride);
+		OutDest = MoveTemp(InSource);
+	}
+}
+
+/**
+ * Also takes a flag to optionally reset the array in case the attribute is disabled
+ */
+template <typename DeltaType, typename ValueType, typename IndexRemapType>
+inline void Scatter(bool bHasData, const DeltaType& Delta, TArray<ValueType>& DestData, int32 NumOutElements, TArray<ValueType>&& InData, const IndexRemapType& IndexRemap, int32 ElementStride = 1)
+{
+	if (bHasData)
+	{
+		::Scatter(Delta, DestData, NumOutElements, MoveTemp(InData), IndexRemap, ElementStride);
+	}
+	else
+	{
+		DestData.Reset();
+	}
+}
+
+/**
+ * Gather the needed values from InSource to OutDest, according to the delta.
+ * If there is no delta, it will perform a bulk copy.
+ */
+template <typename DeltaType, typename ValueType, typename InValueArrayType>
+void Gather(const DeltaType& Delta, TArray<ValueType>& OutDest, const InValueArrayType& InSource, int32 ElementStride = 1)
+{
+	// strides & element count matches - just copy the data
+	if (InSource.Num() == Delta.GetNumItems() * ElementStride)
+	{
+		OutDest = InSource;
+	}
+	else if (Delta.IsEmpty())
+	{
+		OutDest.Reset();
+	}
+	else if (Delta.IsDelta() || ElementStride != 1)
+	{
+		OutDest.Reset(Delta.GetNumItems() * ElementStride);
+		for (auto It = Delta.GetIterator(); It; ++It)
+		{
+			check(OutDest.Num() < Delta.GetNumItems() * ElementStride);
+			OutDest.Append(&InSource[It.GetIndex() * ElementStride], ElementStride);
+		}
+	}
+}
+
+/**
+ */
+struct FIdentityIndexRemap
+{
+	FORCEINLINE constexpr bool IsIdentity() const { return true; }
+	FORCEINLINE int32 operator[](int32 InIndex) const { return InIndex; }
+	FORCEINLINE constexpr bool RemapDestIndex(int32 Index) const { return true; }
+	FORCEINLINE bool Remap(int32& SrcIndex, int32& DstIndex) const  { return true; }
+};
+
+/**
+ * Uses an array to remap the source index, leaving the destination index unchanged.
+ */
+struct FSrcIndexRemap
+{
+	FORCEINLINE constexpr bool IsIdentity() const { return false; }
+	FORCEINLINE FSrcIndexRemap(const TArray<int32>& InIndexRemap) : IndexRemap(InIndexRemap) {}
+	FORCEINLINE bool RemapDestIndex(int32& Index) const { return true;}
+	FORCEINLINE bool Remap(int32& SrcIndex, int32& DstIndex) const 
+	{ 
+		SrcIndex = IndexRemap[SrcIndex];
+		return true; 
+	}
+	const TArray<int32>& IndexRemap;
+};
+
+/**
+ * Helper function to conditionally move a single element
+ */
+template<typename ValueType>
+void CondMove(bool bCondition, TArray<ValueType>& Data, int32 FromIndex, int32 ToIndex, int32 NumElements = 1)
+{
+	if (bCondition)
+	{
+		FMemory::Memcpy(&Data[ToIndex * NumElements], &Data[FromIndex * NumElements], NumElements * sizeof(ValueType));
+	}
+}
+
+/**
+ * Vector register version of FRenderTransform, used to preload the primitive to world transform into registers
+ */
+struct FRenderTransformVectorRegister
+{
+	VectorRegister4f R0;
+	VectorRegister4f R1;
+	VectorRegister4f R2;
+	VectorRegister4f Origin;
+
+	FORCEINLINE FRenderTransformVectorRegister(const FRenderTransform& RenderTransform)
+	{
+		//  we can use unaligmed vectorized load since we know there is data beyond the three rows (the origin), so it is ok to load whatever into the 4th component.
+		R0 = VectorLoad(&RenderTransform.TransformRows[0].X);
+		R1 = VectorLoad(&RenderTransform.TransformRows[1].X);
+		R2 = VectorLoad(&RenderTransform.TransformRows[2].X);
+		// But not for the origin
+		Origin = VectorLoadFloat3(&RenderTransform.Origin);
+	}
+};
+
+FORCEINLINE_DEBUGGABLE FRenderTransform VectorMatrixMultiply(const FRenderTransform& LocalToPrimitive, const FRenderTransformVectorRegister& PrimitiveToWorld)
+{
+	FRenderTransform Result;
+
+	// First row of result (Matrix1[0] * Matrix2).
+	{
+		//  we can use unaligmed vectorized load since we know there is data beyond the three rows (the origin), so it is ok to load whatever into the 4th component.
+		const VectorRegister4Float ARow = VectorLoad(&LocalToPrimitive.TransformRows[0].X);
+		VectorRegister4Float R0 = VectorMultiply(VectorReplicate(ARow, 0), PrimitiveToWorld.R0);
+		R0 = VectorMultiplyAdd(VectorReplicate(ARow, 1), PrimitiveToWorld.R1, R0);
+		R0 = VectorMultiplyAdd(VectorReplicate(ARow, 2), PrimitiveToWorld.R2, R0);
+	
+		// We can use unaligmed vectorized store since we know there is data beyond the three floats that is written later
+		// Note: stomps the X of the TransformRows[1]
+		VectorStore(R0, &Result.TransformRows[0].X);		
+	}
+
+	// Second row of result (Matrix1[1] * Matrix2).
+	{
+		//  we can use unaligmed vectorized load since we know there is data beyond the three rows (the origin), so it is ok to load whatever into the 4th component.
+		const VectorRegister4Float ARow = VectorLoad(&LocalToPrimitive.TransformRows[1].X);
+		VectorRegister4Float R1 = VectorMultiply(VectorReplicate(ARow, 0), PrimitiveToWorld.R0);
+		R1 = VectorMultiplyAdd(VectorReplicate(ARow, 1), PrimitiveToWorld.R1, R1);
+		R1 = VectorMultiplyAdd(VectorReplicate(ARow, 2), PrimitiveToWorld.R2, R1);
+
+		// We can use unaligmed vectorized store since we know there is data beyond the three floats that is written later
+		// Note: stomps the X of the TransformRows[2]
+		VectorStore(R1, &Result.TransformRows[1].X);
+	}
+
+	// Third row of result (Matrix1[2] * Matrix2).
+	{
+		//  we can use unaligmed vectorized load since we know there is data beyond the three rows (the origin), so it is ok to load whatever into the 4th component.
+		const VectorRegister4Float ARow = VectorLoad(&LocalToPrimitive.TransformRows[2].X);
+		VectorRegister4Float R2 = VectorMultiply(VectorReplicate(ARow, 0), PrimitiveToWorld.R0);
+		R2 = VectorMultiplyAdd(VectorReplicate(ARow, 1), PrimitiveToWorld.R1, R2);
+		R2 = VectorMultiplyAdd(VectorReplicate(ARow, 2), PrimitiveToWorld.R2, R2);
+
+		// We can use unaligmed vectorized store since we know there is data beyond the three floats that is written later
+		// Note: stomps the X of the Origin
+		VectorStore(R2, &Result.TransformRows[2].X);
+	}
+
+	// Fourth row of result (Matrix1[3] * Matrix2).
+	{
+		//  can _NOT_ use VectorLoad, or we'll run off the end of the FRenderTransform struct.
+		const VectorRegister4Float ARow = VectorLoadFloat3(&LocalToPrimitive.Origin);
+		
+		// Add B3 at once (instead of mult by 1.0 which would have been the fourth value in the 4x4 version of the matrix)
+		VectorRegister4Float R3 = VectorMultiplyAdd(VectorReplicate(ARow, 0), PrimitiveToWorld.R0, PrimitiveToWorld.Origin);
+		R3 = VectorMultiplyAdd(VectorReplicate(ARow, 1), PrimitiveToWorld.R1, R3);
+		R3 = VectorMultiplyAdd(VectorReplicate(ARow, 2), PrimitiveToWorld.R2, R3);
+
+		VectorStoreFloat3(R3, &Result.Origin);
+	}
+	return Result;
+}
+
+
+/**
+ * Helper function to apply transform update that selectively performs Orthogonalize only if the primitive transform has any non-uniform scale.
+ */
+
+template <typename DeltaType, typename IndexRemapType>
+FORCEINLINE void ApplyTransformUpdates(const DeltaType& DeltaRange, const IndexRemapType& IndexRemap, const FRenderTransform& PrimitiveToRelativeWorld, const TArray<FRenderTransform>& InstanceTransforms, int32 PostUpdateNumTransforms, TArray<FRenderTransform>& OutInstanceToPrimitiveRelative)
+{
+	OutInstanceToPrimitiveRelative.SetNumUninitialized(PostUpdateNumTransforms);
+
+	if (DeltaRange.IsEmpty())
+	{
+		return;
+	}
+
+	if (PrimitiveToRelativeWorld.IsScaleNonUniform())
+	{
+		FRenderTransformVectorRegister PrimitiveToRelativeWorldVR(PrimitiveToRelativeWorld);
+		for (auto It = DeltaRange.GetIterator(); It; ++It)
+		{
+			int32 ItemIndex = It.GetItemIndex();
+			int32 InstanceIndex = It.GetIndex();
+
+			if (IndexRemap.Remap(ItemIndex, InstanceIndex))
+			{
+				FRenderTransform LocalToPrimitiveRelativeWorld = VectorMatrixMultiply(InstanceTransforms[ItemIndex], PrimitiveToRelativeWorldVR);
+				// Remove shear
+				LocalToPrimitiveRelativeWorld.Orthogonalize();
+				OutInstanceToPrimitiveRelative[InstanceIndex] = LocalToPrimitiveRelativeWorld;
+			}
+		}
+	}
+	else
+	{
+		FRenderTransformVectorRegister PrimitiveToRelativeWorldVR(PrimitiveToRelativeWorld);
+		for (auto It = DeltaRange.GetIterator(); It; ++It)
+		{
+			int32 ItemIndex = It.GetItemIndex();
+			int32 InstanceIndex = It.GetIndex();
+
+			if (IndexRemap.Remap(ItemIndex, InstanceIndex))
+			{
+				OutInstanceToPrimitiveRelative[InstanceIndex] = VectorMatrixMultiply(InstanceTransforms[ItemIndex], PrimitiveToRelativeWorldVR);
+			}
+		}
+	}
+
+};
+
+inline FVector3f GetLocalBoundsPadExtent(const FRenderTransform& LocalToWorld, float PadAmount)
+{
+	if (FMath::Abs(PadAmount) < UE_SMALL_NUMBER)
+	{
+		return FVector3f::ZeroVector;
+	}
+
+	FVector3f Scale = LocalToWorld.GetScale();
+	return FVector3f(
+		Scale.X > 0.0f ? PadAmount / Scale.X : 0.0f,
+		Scale.Y > 0.0f ? PadAmount / Scale.Y : 0.0f,
+		Scale.Z > 0.0f ? PadAmount / Scale.Z : 0.0f);
+}
+
+
+template <typename IndexRemapType>
+inline void UpdateIdMapping(FInstanceUpdateChangeSet& ChangeSet, const IndexRemapType& IndexRemap, FInstanceIdIndexMap& OutInstanceIdIndexMap)
+{
+	// update mapping, create explicit mapping if needed
+	if (ChangeSet.bIdentityIdMap && IndexRemap.IsIdentity())
+	{
+		// Reset to identity mapping with the new number of instances
+		OutInstanceIdIndexMap.Reset(ChangeSet.NumSourceInstances);
+	}
+	else 
+	{
+		auto IndexDelta = ChangeSet.GetIndexChangedDelta();
+		bool bIsFull = !IndexDelta.IsDelta() || ChangeSet.NumSourceInstances == ChangeSet.IndexToIdMapDeltaData.Num();
+
+		// Efficient full-data update path if there is no index remap
+		if (bIsFull && IndexRemap.IsIdentity())
+		{
+			OutInstanceIdIndexMap.RebuildFromIndexToIdMap(MoveTemp(ChangeSet.IndexToIdMapDeltaData), ChangeSet.MaxInstanceId);
+		}
+		else if (!IndexDelta.IsEmpty())
+		{
+			// General path that handles incremental removes and other updates.
+			OutInstanceIdIndexMap.ResizeExplicit(ChangeSet.NumSourceInstances, ChangeSet.MaxInstanceId);
+
+			// If any were removed, we need to clear the associated IDs _before_ updating (since they may have been added again)
+			for (TConstSetBitIterator<> It(ChangeSet.InstanceAttributeTracker.GetRemovedIterator()); It; ++It)
+			{
+				// There may be more bits set as things that are marked as removed may no longer be in the map
+				if (It.GetIndex() >= OutInstanceIdIndexMap.GetMaxInstanceId())
+				{
+					break;
+				}
+				OutInstanceIdIndexMap.SetInvalid(FPrimitiveInstanceId{It.GetIndex()});
+			}
+
+			// Update index mappings
+			for (auto It = IndexDelta.GetIterator(); It; ++It)
+			{
+				int32 NewInstanceIndex = It.GetIndex();
+				int32 ItemIndex = It.GetItemIndex();
+			
+				IndexRemap.Remap(ItemIndex, NewInstanceIndex);
+
+				FPrimitiveInstanceId InstanceId = ChangeSet.bIdentityIdMap ? FPrimitiveInstanceId{ItemIndex} : ChangeSet.IndexToIdMapDeltaData[ItemIndex];
+				OutInstanceIdIndexMap.Update(InstanceId, NewInstanceIndex);
+			}
+		}
+	}
+}
+
+template <typename IndexRemapType>
+inline void ApplyAttributeChanges(FInstanceUpdateChangeSet& ChangeSet, const IndexRemapType& IndexRemap, FInstanceSceneDataBuffers::FWriteView& ProxyData)
+{
+	ChangeSet.GetCustomDataReader().Scatter(ProxyData.InstanceCustomData, IndexRemap);
+	ProxyData.NumCustomDataFloats = ChangeSet.NumCustomDataFloats;
+	check(ProxyData.Flags.bHasPerInstanceCustomData || ProxyData.NumCustomDataFloats == 0);
+	ChangeSet.GetSkinningDataReader().Scatter(ProxyData.InstanceSkinningData, IndexRemap);
+	ChangeSet.GetLightShadowUVBiasReader().Scatter(ProxyData.InstanceLightShadowUVBias, IndexRemap);
+
+#if WITH_EDITOR
+	ChangeSet.GetEditorDataReader().Scatter(ProxyData.InstanceEditorData, IndexRemap);
+#endif
+
+	// Delayed per instance random generation, moves it off the GT and RT, but still sucks
+	if (ChangeSet.Flags.bHasPerInstanceRandom)
+	{
+		// TODO: only need to process added instances? No help for ISM since the move path would be taken.
+		// TODO: OTOH for HISM there is no meaningful data, so just skipping and letting the SetNumZeroed fill in the blanks is fine.
+
+		// TODO: Move this to the caller, i.e., the update lambda?
+		ProxyData.InstanceRandomIDs.SetNumZeroed(ChangeSet.NumSourceInstances);
+		if (ChangeSet.GeneratePerInstanceRandomIds)
+		{
+			// NOTE: this is not super efficient(!)
+			TArray<float> TmpInstanceRandomIDs;
+			TmpInstanceRandomIDs.SetNumZeroed(ChangeSet.NumSourceInstances);
+			ChangeSet.GeneratePerInstanceRandomIds(TmpInstanceRandomIDs);
+			FIdentityDeltaRange PerInstanceRandomDelta(TmpInstanceRandomIDs.Num());
+			Scatter(PerInstanceRandomDelta, ProxyData.InstanceRandomIDs, ChangeSet.NumSourceInstances, MoveTemp(TmpInstanceRandomIDs), IndexRemap);
+		}
+		//else 
+		//{
+		//	IndexRemap.Scatter(true, PerInstanceRandomDelta, ProxyData.InstanceRandomIDs, ChangeSet.NumSourceInstances, MoveTemp(ChangeSet.InstanceRandomIDs));
+		//}
+	}
+	else
+	{
+		ProxyData.InstanceRandomIDs.Reset();
+	}
+}
+
+template <typename TaskLambdaType>
+void BeginInstanceDataUpdateTask(FInstanceDataUpdateTaskInfo& InstanceDataUpdateTaskInfo, TaskLambdaType&& TaskLambda, const FInstanceDataBufferHeader& InInstanceDataBufferHeader)
+{
+	// Make sure any previous tasks are done.
+	InstanceDataUpdateTaskInfo.WaitForUpdateCompletion();
+	InstanceDataUpdateTaskInfo.InstanceDataBufferHeader = InInstanceDataBufferHeader;
+
+#if INSTANCE_DATA_UPDATE_ENABLE_ASYNC_TASK
+	InstanceDataUpdateTaskInfo.UpdateTaskHandle = UE::Tasks::Launch(TEXT("FInstanceDataUpdateTaskInfo::BeginUpdateTask"), MoveTemp(TaskLambda));
+#else
+	InstanceDataUpdateTaskInfo.UpdateTaskHandle = UE::Tasks::FTask();
+	TaskLambda();
+#endif
+}
+
+template <typename ProxyType, typename TaskLambdaType>
+void DispatchInstanceDataUpdateTask(bool bIsUnattached, const TSharedPtr<ProxyType, ESPMode::ThreadSafe>& InstanceDataProxy, const FInstanceDataBufferHeader& InstanceDataBufferHeader, TaskLambdaType&& TaskLambda)
+{
+	FInstanceDataUpdateTaskInfo *InstanceDataUpdateTaskInfo = InstanceDataProxy->GetUpdateTaskInfo();
+#if DO_CHECK
+	auto OuterTaskLambda = [InnerTaskLambda = MoveTemp(TaskLambda), InstanceDataBufferHeader, InstanceDataProxy = InstanceDataProxy]() mutable
+	{
+		FInstanceDataUpdateTaskInfo *InstanceDataUpdateTaskInfo = InstanceDataProxy->GetUpdateTaskInfo();
+		check(!InstanceDataUpdateTaskInfo || InstanceDataUpdateTaskInfo->GetHeader() == InstanceDataBufferHeader);
+		InnerTaskLambda();
+		check(!InstanceDataUpdateTaskInfo || InstanceDataUpdateTaskInfo->GetHeader() == InstanceDataBufferHeader);
+		check(!InstanceDataUpdateTaskInfo || InstanceDataBufferHeader.NumInstances == InstanceDataProxy->GeInstanceSceneDataBuffers()->GetNumInstances());
+		// check(InstanceDataBufferHeader.Flags == InstanceDataProxy->GetData().GetFlags());
+		const FInstanceDataFlags HeaderFlags = InstanceDataBufferHeader.Flags;
+		const bool bHasAnyPayloadData = HeaderFlags.bHasPerInstanceHierarchyOffset 
+			|| HeaderFlags.bHasPerInstanceLocalBounds 
+			|| HeaderFlags.bHasPerInstanceDynamicData 
+			|| HeaderFlags.bHasPerInstanceLMSMUVBias 
+			|| HeaderFlags.bHasPerInstanceCustomData 
+			|| HeaderFlags.bHasPerInstancePayloadExtension 
+			|| HeaderFlags.bHasPerInstanceSkinningData 
+			|| HeaderFlags.bHasPerInstanceEditorData;
+		check(bHasAnyPayloadData || InstanceDataBufferHeader.PayloadDataStride == 0);
+		check(!bHasAnyPayloadData || InstanceDataBufferHeader.PayloadDataStride != 0);
+	};
+#else
+	TaskLambdaType OuterTaskLambda = MoveTemp(TaskLambda);
+#endif
+	// Dispatch from any thread.
+	if (bIsUnattached)
+	{
+		if (InstanceDataUpdateTaskInfo)
+		{
+			BeginInstanceDataUpdateTask(*InstanceDataUpdateTaskInfo, MoveTemp(OuterTaskLambda), InstanceDataBufferHeader);
+		}
+		else
+		{
+			OuterTaskLambda();
+		}
+	}
+	else
+	{
+		// Mutating an existing data, must dispatch from RT (such that it does not happen mid-frame).
+		// (One could imagine other scheduling mechanisms)
+		ENQUEUE_RENDER_COMMAND(UpdateInstanceProxyData)(
+			[InstanceDataUpdateTaskInfo, 
+			InstanceDataBufferHeader, 
+			OuterTaskLambda = MoveTemp(OuterTaskLambda)](FRHICommandList& RHICmdList) mutable
+		{
+			if (InstanceDataUpdateTaskInfo)
+			{
+				BeginInstanceDataUpdateTask(*InstanceDataUpdateTaskInfo, MoveTemp(OuterTaskLambda), InstanceDataBufferHeader);
+			}
+			else
+			{
+				OuterTaskLambda();
+			}
+		});
+	}
+}
diff --git a/Engine/Source/Runtime/Engine/Private/InstanceData/InstanceUpdateChangeSet.cpp b/Engine/Source/Runtime/Engine/Private/InstanceData/InstanceUpdateChangeSet.cpp
new file mode 100644
--- /dev/null
+++ b/Engine/Source/Runtime/Engine/Private/InstanceData/InstanceUpdateChangeSet.cpp
@@ -0,0 +1,76 @@
+// Copyright Epic Games, Inc. All Rights Reserved.
+
+#include "InstanceData/InstanceUpdateChangeSet.h"
+#include "InstanceData/InstanceDataUpdateUtils.h"
+#include "Engine/InstancedStaticMesh.h"
+#include "DataDrivenShaderPlatformInfo.h"
+#include "Rendering/RenderingSpatialHash.h"
+#include "Rendering/MotionVectorSimulation.h"
+#include "SceneInterface.h"
+
+
+
+#if WITH_EDITOR
+
+/**
+ * Ugly little wrapper to make sure hitproxies that are kept alive by the proxy on the RT are deleted on the GT
+ */
+class FOpaqueHitProxyContainer
+{
+public:
+	FOpaqueHitProxyContainer(const TArray<TRefCountPtr<HHitProxy>>& InHitProxies) : HitProxies(InHitProxies) {}
+	~FOpaqueHitProxyContainer();
+
+private:
+	FOpaqueHitProxyContainer(const FOpaqueHitProxyContainer &) = delete;
+	void operator=(const FOpaqueHitProxyContainer &) = delete;
+
+	TArray<TRefCountPtr<HHitProxy>> HitProxies;
+};
+
+FOpaqueHitProxyContainer::~FOpaqueHitProxyContainer()
+{
+	struct DeferDeleteHitProxies : FDeferredCleanupInterface
+	{
+		DeferDeleteHitProxies(TArray<TRefCountPtr<HHitProxy>>&& InHitProxies) : HitProxies(MoveTemp(InHitProxies)) {}
+		TArray<TRefCountPtr<HHitProxy>> HitProxies;
+	};
+
+	BeginCleanup(new DeferDeleteHitProxies(MoveTemp(HitProxies)));
+}
+
+
+void FInstanceUpdateChangeSet::SetEditorData(const TArray<TRefCountPtr<HHitProxy>>& HitProxies, const TBitArray<> &InSelectedInstances)//, bool bWasHitProxiesReallocated)
+{
+	HitProxyContainer = MakeOpaqueHitProxyContainer(HitProxies);
+	for (int32 Index = 0; Index < HitProxies.Num(); ++Index)
+	{
+		// Record if the instance is selected
+		FColor HitProxyColor(ForceInit);
+		bool bSelected = InSelectedInstances.IsValidIndex(Index) && InSelectedInstances[Index];
+		if (HitProxies.IsValidIndex(Index))
+		{
+			HitProxyColor = HitProxies[Index]->Id.GetColor();
+		}
+		InstanceEditorData.Add(FInstanceEditorData::Pack(HitProxyColor, bSelected));
+	}
+	SelectedInstances = InSelectedInstances;
+}
+
+#endif
+
+void FInstanceUpdateChangeSet::SetSharedLocalBounds(const FRenderBounds &Bounds)
+{
+	check(!Flags.bHasPerInstanceLocalBounds);
+	InstanceLocalBounds.SetNum(1);
+	InstanceLocalBounds[0] = Bounds;
+}
+
+#if WITH_EDITOR
+
+TPimplPtr<FOpaqueHitProxyContainer> MakeOpaqueHitProxyContainer(const TArray<TRefCountPtr<HHitProxy>>& InHitProxies)
+{
+	return MakePimpl<FOpaqueHitProxyContainer>(InHitProxies);
+}
+
+#endif
diff --git a/Engine/Source/Runtime/Engine/Private/InstanceDataSceneProxy.cpp b/Engine/Source/Runtime/Engine/Private/InstanceDataSceneProxy.cpp
--- a/Engine/Source/Runtime/Engine/Private/InstanceDataSceneProxy.cpp
+++ b/Engine/Source/Runtime/Engine/Private/InstanceDataSceneProxy.cpp
@@ -60,6 +60,50 @@ void FInstanceIdIndexMap::CreateExplicitIdentityMapping()
 	}
 }
 
+void FInstanceIdIndexMap::Serialize(FArchive& Ar)
+{
+	// Serialize the smaller of the two maps (1:1 with instances, whereas the other might have holes).
+	Ar << IndexToIdMap;
+	Ar << NumInstances;
+
+	int32 MaxInstanceId = GetMaxInstanceId();
+	Ar << MaxInstanceId;
+	
+	if (Ar.IsLoading()) 
+	{
+		// if we're loading a non-identity map then restore the ID to index map as well
+		if (!IndexToIdMap.IsEmpty())
+		{
+			RebuildFromIndexToIdMap(MaxInstanceId);
+		}
+		else
+		{
+			// identity mapping, make sure the array is empty
+			IdToIndexMap.Reset();
+		}
+	}
+}
+
+void FInstanceIdIndexMap::RebuildFromIndexToIdMap(int32 MaxInstanceId)
+{
+	IdToIndexMap.SetNum(MaxInstanceId);
+	for (int32 Index = 0; Index < MaxInstanceId; ++Index)
+	{
+		IdToIndexMap[Index] = INDEX_NONE;
+	}
+	for (int32 InstanceIndex = 0; InstanceIndex < IndexToIdMap.Num(); ++InstanceIndex)
+	{
+		FPrimitiveInstanceId Id = IndexToIdMap[InstanceIndex];
+		IdToIndexMap[Id.Id] = InstanceIndex;
+	}
+}
+
+void FInstanceIdIndexMap::RebuildFromIndexToIdMap(TArray<FPrimitiveInstanceId>&& InIndexToIdMap, int32 MaxInstanceId)
+{
+	IndexToIdMap = MoveTemp(InIndexToIdMap);
+	RebuildFromIndexToIdMap(MaxInstanceId);
+}
+
 FInstanceSceneDataBuffers::FInstanceSceneDataBuffers()
 {
 }
diff --git a/Engine/Source/Runtime/Engine/Private/InstancedStaticMesh/ISMInstanceDataManager.cpp b/Engine/Source/Runtime/Engine/Private/InstancedStaticMesh/ISMInstanceDataManager.cpp
--- a/Engine/Source/Runtime/Engine/Private/InstancedStaticMesh/ISMInstanceDataManager.cpp
+++ b/Engine/Source/Runtime/Engine/Private/InstancedStaticMesh/ISMInstanceDataManager.cpp
@@ -2,7 +2,7 @@
 
 #include "InstancedStaticMesh/ISMInstanceDataManager.h"
 #include "InstancedStaticMesh/ISMInstanceUpdateChangeSet.h"
-#include "InstancedStaticMesh/ISMScatterGatherUtil.h"
+#include "InstanceData/InstanceDataUpdateUtils.h"
 
 #include "Engine/InstancedStaticMesh.h"
 #include "DataDrivenShaderPlatformInfo.h"
@@ -44,7 +44,7 @@ static TAutoConsoleVariable<int32> CVarInstanceDataMinInstanceCountToOptimize(
 namespace RenderingSpatialHash
 {
 	template <typename ScalarType>
-	FArchive& operator<<(FArchive& Ar, TLocation<ScalarType>& Item)
+	inline FArchive& operator<<(FArchive& Ar, TLocation<ScalarType>& Item)
 	{
 		Ar << Item.Coord;
 		Ar << Item.Level;
@@ -52,7 +52,7 @@ namespace RenderingSpatialHash
 	}
 }
 
-FArchive& operator<<(FArchive& Ar, FInstanceSceneDataBuffers::FCompressedSpatialHashItem& Item)
+inline FArchive& operator<<(FArchive& Ar, FInstanceSceneDataBuffers::FCompressedSpatialHashItem& Item)
 {
 	Ar << Item.Location;
 	Ar << Item.NumInstances;
@@ -392,86 +392,6 @@ union FChangeDesc
 
 };
 
-template <typename TaskLambdaType>
-void FPrimitiveInstanceDataManager::BeginUpdateTask(FInstanceDataUpdateTaskInfo &InstanceDataUpdateTaskInfo, TaskLambdaType &&TaskLambda, const FInstanceDataBufferHeader &InInstanceDataBufferHeader)
-{
-	// Make sure any previous tasks are done.
-	InstanceDataUpdateTaskInfo.WaitForUpdateCompletion();
-	InstanceDataUpdateTaskInfo.InstanceDataBufferHeader = InInstanceDataBufferHeader;
-
-#if IDPROXY_ENABLE_ASYNC_TASK
-	auto OuterTaskFunc = [TaskLambda = MoveTemp(TaskLambda)]() mutable
-	{
-		const float DebugDelay = CVarInstanceUpdateTaskDebugDelay.GetValueOnAnyThread();
-		if (DebugDelay != 0.0f)
-		{
-			UE_LOG(LogTemp, Warning, TEXT("Instance update debug delay %5.1fs (CVar r.InstanceUpdateTaskDebugDelay)"), DebugDelay);
-			FPlatformProcess::Sleep(DebugDelay);
-		}
-
-		TaskLambda();
-	};
-	InstanceDataUpdateTaskInfo.UpdateTaskHandle = UE::Tasks::Launch(TEXT("FInstanceDataUpdateTaskInfo::BeginUpdateTask"), MoveTemp(OuterTaskFunc));
-#else
-	InstanceDataUpdateTaskInfo.UpdateTaskHandle = UE::Tasks::FTask();
-	TaskLambda();
-#endif
-}
-
-template <typename TaskLambdaType>
-void FPrimitiveInstanceDataManager::DispatchUpdateTask(bool bIsUnattached, const FInstanceDataBufferHeader &InstanceDataBufferHeader, TaskLambdaType &&TaskLambda)
-{
-	FInstanceDataUpdateTaskInfo *InstanceDataUpdateTaskInfo = Proxy->GetUpdateTaskInfo();
-#if DO_CHECK
-	auto OuterTaskLambda = [InnerTaskLambda = MoveTemp(TaskLambda), InstanceDataBufferHeader, Proxy = Proxy]() mutable
-	{
-		FInstanceDataUpdateTaskInfo *InstanceDataUpdateTaskInfo = Proxy->GetUpdateTaskInfo();
-		check(!InstanceDataUpdateTaskInfo || InstanceDataUpdateTaskInfo->GetHeader() == InstanceDataBufferHeader);
-		InnerTaskLambda();
-		check(!InstanceDataUpdateTaskInfo || InstanceDataUpdateTaskInfo->GetHeader() == InstanceDataBufferHeader);
-		check(!InstanceDataUpdateTaskInfo || InstanceDataBufferHeader.NumInstances == Proxy->GetData().GetNumInstances());
-		// check(InstanceDataBufferHeader.Flags == Proxy->GetData().GetFlags());
-		const FInstanceDataFlags HeaderFlags = InstanceDataBufferHeader.Flags;
-		const bool bHasAnyPayloadData = HeaderFlags.bHasPerInstanceHierarchyOffset || HeaderFlags.bHasPerInstanceLocalBounds || HeaderFlags.bHasPerInstanceDynamicData || HeaderFlags.bHasPerInstanceLMSMUVBias || HeaderFlags.bHasPerInstanceCustomData || HeaderFlags.bHasPerInstancePayloadExtension || HeaderFlags.bHasPerInstanceEditorData;
-		check(bHasAnyPayloadData || InstanceDataBufferHeader.PayloadDataStride == 0);
-		check(!bHasAnyPayloadData || InstanceDataBufferHeader.PayloadDataStride != 0);
-	};
-#else
-	TaskLambdaType OuterTaskLambda = MoveTemp(TaskLambda);
-#endif
-	// Dispatch from any thread.
-	if (bIsUnattached)
-	{
-		if (InstanceDataUpdateTaskInfo)
-		{
-			BeginUpdateTask(*InstanceDataUpdateTaskInfo, MoveTemp(OuterTaskLambda), InstanceDataBufferHeader);
-		}
-		else
-		{
-			OuterTaskLambda();
-		}
-	}
-	else
-	{
-		// Mutating an existing data, must dispatch from RT (such that it does not happen mid-frame).
-		// (One could imagine other scheduling mechanisms)
-		ENQUEUE_RENDER_COMMAND(UpdateInstanceProxyData)(
-			[InstanceDataUpdateTaskInfo, 
-			InstanceDataBufferHeader, 
-			OuterTaskLambda = MoveTemp(OuterTaskLambda)](FRHICommandList& RHICmdList) mutable
-		{
-			if (InstanceDataUpdateTaskInfo)
-			{
-				BeginUpdateTask(*InstanceDataUpdateTaskInfo, MoveTemp(OuterTaskLambda), InstanceDataBufferHeader);
-			}
-			else
-			{
-				OuterTaskLambda();
-			}
-		});
-	}
-}
-
 /**
  * Data captured that may be flushed to rebuild the proxy from legacy data AKA FStaticMeshInstanceData.
  * Is marshalled across via the FLegacyRebuildChangeSet.
@@ -513,11 +433,7 @@ void FPrimitiveInstanceDataManager::InitChangeSet(const FChangeDesc &ChangeDesc,
 		Gather(ChangeSet.GetIndexChangedDelta(), ChangeSet.IndexToIdMapDeltaData, IndexToIdMap);
 	}
 
-	ChangeSet.bNumCustomDataChanged = bNumCustomDataChanged;
-	ChangeSet.bBakedLightingDataChanged = bBakedLightingDataChanged;
-
 #if WITH_EDITOR
-	ChangeSet.bAnyEditorDataChanged = bAnyEditorDataChanged;
 	bAnyEditorDataChanged = false;
 #endif	
 
@@ -534,7 +450,7 @@ void FPrimitiveInstanceDataManager::InitChangeSet(const FChangeDesc &ChangeDesc,
 	ChangeSet.NumCustomDataFloats = NumCustomDataFloats;
 
 	// This is the odd one out
-	ChangeSet.SetInstanceLocalBounds(ComponentData.StaticMeshBounds);
+	ChangeSet.SetSharedLocalBounds(ComponentData.StaticMeshBounds);
 }
 
 bool FPrimitiveInstanceDataManager::FlushChanges(FInstanceUpdateComponentDesc &&ComponentData, bool bNewPrimitiveProxy)
@@ -616,7 +532,7 @@ bool FPrimitiveInstanceDataManager::FlushChanges(FInstanceUpdateComponentDesc &&
 		InstanceDataBufferHeader.PayloadDataStride = FInstanceSceneDataBuffers::CalcPayloadDataStride(ExternalChangeSet.Flags, NumCustomDataFloats, 0);
 		InstanceDataBufferHeader.Flags = ExternalChangeSet.Flags;
 
-		DispatchUpdateTask(bIsUnattached, InstanceDataBufferHeader, [ExternalChangeSet = MoveTemp(ExternalChangeSet), Proxy = Proxy] () mutable 
+		DispatchInstanceDataUpdateTask(bIsUnattached, Proxy, InstanceDataBufferHeader, [ExternalChangeSet = MoveTemp(ExternalChangeSet), Proxy = Proxy] () mutable 
 		{
 			FISMCInstanceDataSceneProxy &ProxyRef = *Proxy;
 			
@@ -712,7 +628,7 @@ bool FPrimitiveInstanceDataManager::FlushChanges(FInstanceUpdateComponentDesc &&
 		bool bNeedFullUpdate = ChangeDesc.bUntrackedState || ChangeDesc.bMaterialUsageFlagsChanged || GetState() == ETrackingState::Optimized;
 
 		// NOTE: Moving the update tracker to the change set implicitly resets it.
-		FISMInstanceUpdateChangeSet ChangeSet(bNeedFullUpdate, MoveTemp(InstanceUpdateTracker));
+		FISMInstanceUpdateChangeSet ChangeSet(bNeedFullUpdate, MoveTemp(InstanceUpdateTracker), ComponentData.NumProxyInstances);
 		ChangeSet.bUpdateAllInstanceTransforms = ChangeDesc.bPrimitiveTransformChanged || bTransformChangedAllInstances;
 		ChangeSet.PostUpdateNumInstances = ComponentData.NumProxyInstances;
 		ChangeSet.NumSourceInstances = ComponentData.NumSourceInstances;
@@ -763,7 +679,7 @@ bool FPrimitiveInstanceDataManager::FlushChanges(FInstanceUpdateComponentDesc &&
 		{
 			LOG_INST_DATA(TEXT("Primitive Transform Update %s"), TEXT(""));
 			SuccessorTrackingState = ETrackingState::Initial;
-			DispatchUpdateTask(bIsUnattached, InstanceDataBufferHeader, [ChangeSet = MoveTemp(ChangeSet), Proxy = Proxy, NumHeaderInstances = InstanceDataBufferHeader.NumInstances] () mutable 
+			DispatchInstanceDataUpdateTask(bIsUnattached, Proxy, InstanceDataBufferHeader, [ChangeSet = MoveTemp(ChangeSet), Proxy = Proxy, NumHeaderInstances = InstanceDataBufferHeader.NumInstances] () mutable 
 			{
 				Proxy->UpdatePrimitiveTransform(MoveTemp(ChangeSet));
 				check(NumHeaderInstances == Proxy->GeInstanceSceneDataBuffers()->GetNumInstances());
@@ -783,7 +699,7 @@ bool FPrimitiveInstanceDataManager::FlushChanges(FInstanceUpdateComponentDesc &&
 				SuccessorTrackingState = ETrackingState::Optimized;
 			}
 
-			DispatchUpdateTask(bIsUnattached, InstanceDataBufferHeader, [ChangeSet = MoveTemp(ChangeSet), Proxy = Proxy, NumHeaderInstances = InstanceDataBufferHeader.NumInstances] () mutable 
+			DispatchInstanceDataUpdateTask(bIsUnattached, Proxy, InstanceDataBufferHeader, [ChangeSet = MoveTemp(ChangeSet), Proxy = Proxy, NumHeaderInstances = InstanceDataBufferHeader.NumInstances] () mutable 
 			{
 				Proxy->Build(MoveTemp(ChangeSet));
 				check(NumHeaderInstances == Proxy->GeInstanceSceneDataBuffers()->GetNumInstances());
@@ -793,7 +709,7 @@ bool FPrimitiveInstanceDataManager::FlushChanges(FInstanceUpdateComponentDesc &&
 		{
 			LOG_INST_DATA(TEXT("Delta Update %s"), TEXT(""));
 			check(!ChangeDesc.bUntrackedState);
-			DispatchUpdateTask(bIsUnattached, InstanceDataBufferHeader, [ChangeSet = MoveTemp(ChangeSet), Proxy = Proxy, NumHeaderInstances = InstanceDataBufferHeader.NumInstances] () mutable 
+			DispatchInstanceDataUpdateTask(bIsUnattached, Proxy, InstanceDataBufferHeader, [ChangeSet = MoveTemp(ChangeSet), Proxy = Proxy, NumHeaderInstances = InstanceDataBufferHeader.NumInstances] () mutable 
 			{
 				Proxy->Update(MoveTemp(ChangeSet));
 				check(NumHeaderInstances == Proxy->GeInstanceSceneDataBuffers()->GetNumInstances());
@@ -933,11 +849,6 @@ bool FPrimitiveInstanceDataManager::HasIdentityMapping() const
 	return IndexToIdMap.IsEmpty();
 }
 
-bool FPrimitiveInstanceDataManager::ShouldTrackIds() const
-{
-	return Proxy != nullptr;
-}
-
 void FPrimitiveInstanceDataManager::FreeInstanceId(FPrimitiveInstanceId InstanceId)
 {
 	LOG_INST_DATA(TEXT("FreeInstanceId(Id: %d)"), InstanceId.Id);
diff --git a/Engine/Source/Runtime/Engine/Private/InstancedStaticMesh/ISMInstanceDataSceneProxy.cpp b/Engine/Source/Runtime/Engine/Private/InstancedStaticMesh/ISMInstanceDataSceneProxy.cpp
--- a/Engine/Source/Runtime/Engine/Private/InstancedStaticMesh/ISMInstanceDataSceneProxy.cpp
+++ b/Engine/Source/Runtime/Engine/Private/InstancedStaticMesh/ISMInstanceDataSceneProxy.cpp
@@ -2,7 +2,7 @@
 
 #include "InstancedStaticMesh/ISMInstanceDataSceneProxy.h"
 #include "InstancedStaticMesh/ISMInstanceUpdateChangeSet.h"
-#include "InstancedStaticMesh/ISMScatterGatherUtil.h"
+#include "InstanceData/InstanceDataUpdateUtils.h"
 #include "Engine/InstancedStaticMesh.h"
 #include "DataDrivenShaderPlatformInfo.h"
 #include "Rendering/RenderingSpatialHash.h"
@@ -18,172 +18,6 @@ DEFINE_LOG_CATEGORY(LogInstanceProxy);
 	#define LOG_INST_DATA(_Format_, ...) 
 #endif
 
-/**
- * Vector register version of FRenderTransform, used to preload the primitive to world transform into registers
- */
-struct FRenderTransformVectorRegister
-{
-	VectorRegister4f R0;
-	VectorRegister4f R1;
-	VectorRegister4f R2;
-	VectorRegister4f Origin;
-
-	FORCEINLINE FRenderTransformVectorRegister(const FRenderTransform &RenderTransform)
-	{
-		//  we can use unaligmed vectorized load since we know there is data beyond the three rows (the origin), so it is ok to load whatever into the 4th component.
-		R0 = VectorLoad(&RenderTransform.TransformRows[0].X);
-		R1 = VectorLoad(&RenderTransform.TransformRows[1].X);
-		R2 = VectorLoad(&RenderTransform.TransformRows[2].X);
-		// But not for the origin
-		Origin = VectorLoadFloat3(&RenderTransform.Origin);
-	}
-};
-
-FORCEINLINE_DEBUGGABLE FRenderTransform VectorMatrixMultiply(const FRenderTransform& LocalToPrimitive, const FRenderTransformVectorRegister& PrimitiveToWorld)
-{
-	FRenderTransform Result;
-
-	// First row of result (Matrix1[0] * Matrix2).
-	{
-		//  we can use unaligmed vectorized load since we know there is data beyond the three rows (the origin), so it is ok to load whatever into the 4th component.
-		const VectorRegister4Float ARow = VectorLoad(&LocalToPrimitive.TransformRows[0].X);
-		VectorRegister4Float R0 = VectorMultiply(VectorReplicate(ARow, 0), PrimitiveToWorld.R0);
-		R0 = VectorMultiplyAdd(VectorReplicate(ARow, 1), PrimitiveToWorld.R1, R0);
-		R0 = VectorMultiplyAdd(VectorReplicate(ARow, 2), PrimitiveToWorld.R2, R0);
-	
-		// We can use unaligmed vectorized store since we know there is data beyond the three floats that is written later
-		// Note: stomps the X of the TransformRows[1]
-		VectorStore(R0, &Result.TransformRows[0].X);		
-	}
-
-	// Second row of result (Matrix1[1] * Matrix2).
-	{
-		//  we can use unaligmed vectorized load since we know there is data beyond the three rows (the origin), so it is ok to load whatever into the 4th component.
-		const VectorRegister4Float ARow = VectorLoad(&LocalToPrimitive.TransformRows[1].X);
-		VectorRegister4Float R1 = VectorMultiply(VectorReplicate(ARow, 0), PrimitiveToWorld.R0);
-		R1 = VectorMultiplyAdd(VectorReplicate(ARow, 1), PrimitiveToWorld.R1, R1);
-		R1 = VectorMultiplyAdd(VectorReplicate(ARow, 2), PrimitiveToWorld.R2, R1);
-
-		// We can use unaligmed vectorized store since we know there is data beyond the three floats that is written later
-		// Note: stomps the X of the TransformRows[2]
-		VectorStore(R1, &Result.TransformRows[1].X);
-	}
-
-	// Third row of result (Matrix1[2] * Matrix2).
-	{
-		//  we can use unaligmed vectorized load since we know there is data beyond the three rows (the origin), so it is ok to load whatever into the 4th component.
-		const VectorRegister4Float ARow = VectorLoad(&LocalToPrimitive.TransformRows[2].X);
-		VectorRegister4Float R2 = VectorMultiply(VectorReplicate(ARow, 0), PrimitiveToWorld.R0);
-		R2 = VectorMultiplyAdd(VectorReplicate(ARow, 1), PrimitiveToWorld.R1, R2);
-		R2 = VectorMultiplyAdd(VectorReplicate(ARow, 2), PrimitiveToWorld.R2, R2);
-
-		// We can use unaligmed vectorized store since we know there is data beyond the three floats that is written later
-		// Note: stomps the X of the Origin
-		VectorStore(R2, &Result.TransformRows[2].X);
-	}
-
-	// Fourth row of result (Matrix1[3] * Matrix2).
-	{
-		//  can _NOT_ use VectorLoad, or we'll run off the end of the FRenderTransform struct.
-		const VectorRegister4Float ARow = VectorLoadFloat3(&LocalToPrimitive.Origin);
-		
-		// Add B3 at once (instead of mult by 1.0 which would have been the fourth value in the 4x4 version of the matrix)
-		VectorRegister4Float R3 = VectorMultiplyAdd(VectorReplicate(ARow, 0), PrimitiveToWorld.R0, PrimitiveToWorld.Origin);
-		R3 = VectorMultiplyAdd(VectorReplicate(ARow, 1), PrimitiveToWorld.R1, R3);
-		R3 = VectorMultiplyAdd(VectorReplicate(ARow, 2), PrimitiveToWorld.R2, R3);
-
-		VectorStoreFloat3(R3, &Result.Origin);
-	}
-	return Result;
-}
-
-
-/**
- * Helper function to apply transform update that selectively performs Orthogonalize only if the primitive transform has any non-uniform scale.
- */
-template <typename DeltaType, typename IndexRemapType, typename TransformOutputType>
-FORCEINLINE void ApplyTransformUpdatesEx(const DeltaType &DeltaRange, const IndexRemapType &IndexRemap, const FRenderTransform &PrimitiveToRelativeWorld, const TArray<FRenderTransform> &InstanceTransforms, int32 PostUpdateNumTransforms, TransformOutputType &OutInstanceToPrimitiveRelative)
-{
-	OutInstanceToPrimitiveRelative.SetNum(PostUpdateNumTransforms);
-
-	if (DeltaRange.IsEmpty())
-	{
-		return;
-	}
-
-	if (PrimitiveToRelativeWorld.IsScaleNonUniform())
-	{
-		FRenderTransformVectorRegister PrimitiveToRelativeWorldVR(PrimitiveToRelativeWorld);
-		for (auto It = DeltaRange.GetIterator(); It; ++It)
-		{
-			int32 ItemIndex = It.GetItemIndex();
-			int32 InstanceIndex = It.GetIndex();
-
-			if (IndexRemap.Remap(ItemIndex, InstanceIndex))
-			{
-				FRenderTransform LocalToPrimitiveRelativeWorld = VectorMatrixMultiply(InstanceTransforms[ItemIndex], PrimitiveToRelativeWorldVR);
-				// Remove shear
-				LocalToPrimitiveRelativeWorld.Orthogonalize();
-				OutInstanceToPrimitiveRelative.Set(InstanceIndex, LocalToPrimitiveRelativeWorld);
-			}
-		}
-	}
-	else
-	{
-		FRenderTransformVectorRegister PrimitiveToRelativeWorldVR(PrimitiveToRelativeWorld);
-		for (auto It = DeltaRange.GetIterator(); It; ++It)
-		{
-			int32 ItemIndex = It.GetItemIndex();
-			int32 InstanceIndex = It.GetIndex();
-
-			if (IndexRemap.Remap(ItemIndex, InstanceIndex))
-			{
-				OutInstanceToPrimitiveRelative.Set(InstanceIndex, VectorMatrixMultiply(InstanceTransforms[ItemIndex], PrimitiveToRelativeWorldVR));
-			}
-		}
-	}
-}
-
-template <typename DeltaType, typename IndexRemapType>
-FORCEINLINE void ApplyTransformUpdates(const DeltaType &DeltaRange, const IndexRemapType &IndexRemap, const FRenderTransform &PrimitiveToRelativeWorld, const TArray<FRenderTransform> &InstanceTransforms, int32 PostUpdateNumTransforms, TArray<FRenderTransform> &OutInstanceToPrimitiveRelative)
-{
-	struct FArrayTransformCollector
-	{
-		FORCEINLINE void SetNum(int32 Num) { OutArray.SetNumUninitialized(Num); }
-		FORCEINLINE void Set(int32 Index, const FRenderTransform &Transform) { OutArray[Index] = Transform; }
-		TArray<FRenderTransform> &OutArray;
-	};
-
-	FArrayTransformCollector TransformCollector {OutInstanceToPrimitiveRelative};
-
-	ApplyTransformUpdatesEx(DeltaRange, IndexRemap, PrimitiveToRelativeWorld, InstanceTransforms, PostUpdateNumTransforms, TransformCollector);
-};
-/**
- * Helper class to apply transform concatenation that selectively performs Orthogonalize only if the primitive transform has any non-uniform scale.
- */
-struct FInstanceTransformApplyHelper
-{
-	FORCEINLINE_DEBUGGABLE FInstanceTransformApplyHelper(bool bHasWork, const FRenderTransform &InPrimitiveToRelativeWorld)
-	: PrimitiveToRelativeWorldVR(InPrimitiveToRelativeWorld)
-	, bNeedsOrthogonalization(bHasWork && InPrimitiveToRelativeWorld.IsScaleNonUniform())
-	{
-	}
-
-	FORCEINLINE_DEBUGGABLE FRenderTransform Apply(const FRenderTransform &InstanceTransform)
-	{
-		FRenderTransform LocalToPrimitiveRelativeWorld = VectorMatrixMultiply(InstanceTransform, PrimitiveToRelativeWorldVR);
-		if (bNeedsOrthogonalization)
-		{
-			// Remove shear
-			LocalToPrimitiveRelativeWorld.Orthogonalize();
-		}
-		return LocalToPrimitiveRelativeWorld;
-	}
-	FRenderTransformVectorRegister PrimitiveToRelativeWorldVR;
-	bool bNeedsOrthogonalization = false;
-};
-
-
 FISMCInstanceDataSceneProxy::FISMCInstanceDataSceneProxy(FStaticShaderPlatform InShaderPlatform, ERHIFeatureLevel::Type InFeatureLevel) 
 	: ShaderPlatform(InShaderPlatform)
 	, FeatureLevel(InFeatureLevel) 
@@ -191,32 +25,10 @@ FISMCInstanceDataSceneProxy::FISMCInstanceDataSceneProxy(FStaticShaderPlatform I
 	bUseLegacyRenderingPath = !UseGPUScene(ShaderPlatform, FeatureLevel);
 }
 
-struct FIdentityIndexRemap
-{
-	FORCEINLINE constexpr bool IsIdentity() const { return true; }
-
-	inline int32 operator[](int32 InIndex) const { return InIndex; }
-
-	template <typename DeltaType, typename ValueType>
-	inline void Scatter(bool bHasData, const DeltaType &Delta, TArray<ValueType> &DestData, int32 NumOutElements, TArray<ValueType> &&InData, int32 ElementStride = 1) const
-	{
-		if (bHasData)
-		{
-			::Scatter(Delta, DestData, NumOutElements, MoveTemp(InData), ElementStride);
-		}
-		else
-		{
-			DestData.Reset();
-		}
-	}
-
-	FORCEINLINE constexpr bool RemapDestIndex(int32 Index) const { return true; }
-
-	FORCEINLINE bool Remap(int32 &SrcIndex, int32 &DstIndex) const  { return true; }
-};
-
 struct FReorderTableIndexRemap
 {
+	FORCEINLINE constexpr bool IsIdentity() const { return false; }
+
 	FReorderTableIndexRemap(const TArray<int32> &InReorderTable, int32 InMaxValidIndex) : ReorderTable(InReorderTable), MaxValidIndex(InMaxValidIndex) {}
 
 	inline int32 ClampValidIndex(int32 Index) const
@@ -286,7 +98,7 @@ void FISMCInstanceDataSceneProxy::ApplyAttributeChanges(FISMInstanceUpdateChange
 	if (ChangeSet.Flags.bHasPerInstanceCustomData)
 	{
 		ProxyData.NumCustomDataFloats = ChangeSet.NumCustomDataFloats;
-		IndexRemap.Scatter(ChangeSet.Flags.bHasPerInstanceCustomData, ChangeSet.GetCustomDataDelta(), ProxyData.InstanceCustomData, ChangeSet.PostUpdateNumInstances, MoveTemp(ChangeSet.PerInstanceCustomData), ProxyData.NumCustomDataFloats);
+		Scatter(ChangeSet.GetCustomDataDelta(), ProxyData.InstanceCustomData, ChangeSet.PostUpdateNumInstances, MoveTemp(ChangeSet.PerInstanceCustomData), IndexRemap, ProxyData.NumCustomDataFloats);
 	}
 	else
 	{
@@ -294,9 +106,9 @@ void FISMCInstanceDataSceneProxy::ApplyAttributeChanges(FISMInstanceUpdateChange
 		ProxyData.InstanceCustomData.Reset();
 	}
 
-	IndexRemap.Scatter(ChangeSet.Flags.bHasPerInstanceLMSMUVBias, ChangeSet.GetInstanceLightShadowUVBiasDelta(), ProxyData.InstanceLightShadowUVBias, ChangeSet.PostUpdateNumInstances, MoveTemp(ChangeSet.InstanceLightShadowUVBias));
+	Scatter(ChangeSet.Flags.bHasPerInstanceLMSMUVBias, ChangeSet.GetInstanceLightShadowUVBiasDelta(), ProxyData.InstanceLightShadowUVBias, ChangeSet.PostUpdateNumInstances, MoveTemp(ChangeSet.InstanceLightShadowUVBias), IndexRemap);
 #if WITH_EDITOR
-	IndexRemap.Scatter(ChangeSet.Flags.bHasPerInstanceEditorData, ChangeSet.GetInstanceEditorDataDelta(), ProxyData.InstanceEditorData, ChangeSet.PostUpdateNumInstances, MoveTemp(ChangeSet.InstanceEditorData));
+	Scatter(ChangeSet.Flags.bHasPerInstanceEditorData, ChangeSet.GetInstanceEditorDataDelta(), ProxyData.InstanceEditorData, ChangeSet.PostUpdateNumInstances, MoveTemp(ChangeSet.InstanceEditorData), IndexRemap);
 
 	// replace the HP container.
 	if (ChangeSet.HitProxyContainer)
@@ -320,7 +132,7 @@ void FISMCInstanceDataSceneProxy::ApplyAttributeChanges(FISMInstanceUpdateChange
 			TmpInstanceRandomIDs.SetNumZeroed(ChangeSet.PostUpdateNumInstances);
 			ChangeSet.GeneratePerInstanceRandomIds(TmpInstanceRandomIDs);
 			FIdentityDeltaRange PerInstanceRandomDelta(TmpInstanceRandomIDs.Num());
-			IndexRemap.Scatter(true, PerInstanceRandomDelta, ProxyData.InstanceRandomIDs, ChangeSet.PostUpdateNumInstances, MoveTemp(TmpInstanceRandomIDs));
+			Scatter(PerInstanceRandomDelta, ProxyData.InstanceRandomIDs, ChangeSet.PostUpdateNumInstances, MoveTemp(TmpInstanceRandomIDs), IndexRemap);
 		}
 		//else 
 		//{
@@ -366,48 +178,6 @@ void FISMCInstanceDataSceneProxy::ApplyDataChanges(FISMInstanceUpdateChangeSet &
 	ApplyAttributeChanges(ChangeSet, IndexRemap, ProxyData);
 }
 
-template<typename ValueType>
-void CondMove(bool bCondition, TArray<ValueType> &Data, int32 FromIndex, int32 ToIndex, int32 NumElements = 1)
-{
-	if (bCondition)
-	{
-		FMemory::Memcpy(&Data[ToIndex * NumElements], &Data[FromIndex * NumElements], NumElements * sizeof(ValueType));
-	}
-}
-
-struct FSrcIndexRemap
-{
-	FORCEINLINE constexpr bool IsIdentity() const { return false; }
-
-	FSrcIndexRemap(const TArray<int32> &InIndexRemap) : IndexRemap(InIndexRemap) {}
-
-	FORCEINLINE bool RemapDestIndex(int32 &Index) const 
-	{ 
-		return true;
-	}
-
-	FORCEINLINE bool Remap(int32 &SrcIndex, int32 &DstIndex) const 
-	{ 
-		SrcIndex = IndexRemap[SrcIndex];
-		return true; 
-	}
-
-	template <typename DeltaType, typename ValueType>
-	FORCEINLINE void Scatter(bool bHasData, const DeltaType &Delta, TArray<ValueType> &DestData, int32 NumOutElements, TArray<ValueType> &&InData, int32 ElementStride = 1) const
-	{
-		if (bHasData)
-		{
-			::Scatter(Delta, DestData, NumOutElements, MoveTemp(InData), *this, ElementStride);
-		}
-		else
-		{
-			DestData.Reset();
-		}
-	}
-
-	const TArray<int32> &IndexRemap;
-};
-
 void FISMCInstanceDataSceneProxy::BuildFromOptimizedDataBuffers(FISMInstanceUpdateChangeSet& ChangeSet, FInstanceIdIndexMap &OutInstanceIdIndexMap, FInstanceSceneDataBuffers::FWriteView &ProxyData)
 {
 	SCOPED_NAMED_EVENT(FISMCInstanceDataSceneProxy_BuildFromOptimizedDataBuffers, FColor::Emerald);
@@ -556,40 +326,7 @@ void FISMCInstanceDataSceneProxy::Build(FISMInstanceUpdateChangeSet&& ChangeSet)
 template <typename IndexRemapType>
 void FISMCInstanceDataSceneProxy::UpdateIdMapping(FISMInstanceUpdateChangeSet& ChangeSet, const IndexRemapType &IndexRemap)
 {
-	// update mapping, create explicit mapping if needed
-	if (ChangeSet.bIdentityIdMap && IndexRemap.IsIdentity())
-	{
-		// Reset to identity mapping with the new number of instances
-		InstanceIdIndexMap.Reset(ChangeSet.PostUpdateNumInstances);
-	}
-	else
-	{
-		InstanceIdIndexMap.ResizeExplicit(ChangeSet.PostUpdateNumInstances, ChangeSet.MaxInstanceId);
-
-		// If any were removed, we need to clear the associated IDs, before updating (since they may have been added again)
-		for (TConstSetBitIterator<> It(ChangeSet.InstanceAttributeTracker.GetRemovedIterator()); It; ++It)
-		{
-			// There may be more bits set as things that are marked as removed may no longer be in the map
-			if (It.GetIndex() >= InstanceIdIndexMap.GetMaxInstanceId())
-			{
-				break;
-			}
-			InstanceIdIndexMap.SetInvalid(FPrimitiveInstanceId{It.GetIndex()});
-		}
-
-		// Update index mappings (if not identity)
-		auto IndexDelta = ChangeSet.GetIndexChangedDelta();
-		for (auto It = IndexDelta.GetIterator(); It; ++It)
-		{
-			int32 NewInstanceIndex = It.GetIndex();
-			int32 ItemIndex = It.GetItemIndex();
-			
-			IndexRemap.Remap(ItemIndex, NewInstanceIndex);
-
-			FPrimitiveInstanceId InstanceId = ChangeSet.bIdentityIdMap ? FPrimitiveInstanceId{ItemIndex} : ChangeSet.IndexToIdMapDeltaData[ItemIndex];
-			InstanceIdIndexMap.Update(InstanceId, NewInstanceIndex);
-		}
-	}
+	::UpdateIdMapping(ChangeSet, IndexRemap, InstanceIdIndexMap);
 }
 
 void FISMCInstanceDataSceneProxy::Update(FISMInstanceUpdateChangeSet&& ChangeSet)
diff --git a/Engine/Source/Runtime/Engine/Private/InstancedStaticMesh/ISMInstanceUpdateChangeSet.cpp b/Engine/Source/Runtime/Engine/Private/InstancedStaticMesh/ISMInstanceUpdateChangeSet.cpp
--- a/Engine/Source/Runtime/Engine/Private/InstancedStaticMesh/ISMInstanceUpdateChangeSet.cpp
+++ b/Engine/Source/Runtime/Engine/Private/InstancedStaticMesh/ISMInstanceUpdateChangeSet.cpp
@@ -1,142 +1,63 @@
 // Copyright Epic Games, Inc. All Rights Reserved.
 
 #include "InstancedStaticMesh/ISMInstanceUpdateChangeSet.h"
-#include "InstancedStaticMesh/ISMScatterGatherUtil.h"
+#include "InstanceData/InstanceDataUpdateUtils.h"
 #include "Engine/InstancedStaticMesh.h"
 #include "DataDrivenShaderPlatformInfo.h"
 #include "Rendering/RenderingSpatialHash.h"
 #include "Rendering/MotionVectorSimulation.h"
 #include "SceneInterface.h"
 
-
-
-#if WITH_EDITOR
-
-/**
- * Ugly little wrapper to make sure hitproxies that are kept alive by the proxy on the RT are deleted on the GT
- */
-class FOpaqueHitProxyContainer
-{
-public:
-	FOpaqueHitProxyContainer(const TArray<TRefCountPtr<HHitProxy>>& InHitProxies) : HitProxies(InHitProxies) {}
-	~FOpaqueHitProxyContainer();
-
-private:
-	FOpaqueHitProxyContainer(const FOpaqueHitProxyContainer &) = delete;
-	void operator=(const FOpaqueHitProxyContainer &) = delete;
-
-	TArray<TRefCountPtr<HHitProxy>> HitProxies;
-};
-
-FOpaqueHitProxyContainer::~FOpaqueHitProxyContainer()
-{
-	struct DeferDeleteHitProxies : FDeferredCleanupInterface
-	{
-		DeferDeleteHitProxies(TArray<TRefCountPtr<HHitProxy>>&& InHitProxies) : HitProxies(MoveTemp(InHitProxies)) {}
-		TArray<TRefCountPtr<HHitProxy>> HitProxies;
-	};
-
-	BeginCleanup(new DeferDeleteHitProxies(MoveTemp(HitProxies)));
-}
-
-
-void FISMInstanceUpdateChangeSet::SetEditorData(const TArray<TRefCountPtr<HHitProxy>>& HitProxies, const TBitArray<> &InSelectedInstances)//, bool bWasHitProxiesReallocated)
-{
-	HitProxyContainer = MakeOpaqueHitProxyContainer(HitProxies);
-	for (int32 Index = 0; Index < HitProxies.Num(); ++Index)
-	{
-		// Record if the instance is selected
-		FColor HitProxyColor(ForceInit);
-		bool bSelected = InSelectedInstances.IsValidIndex(Index) && InSelectedInstances[Index];
-		if (HitProxies.IsValidIndex(Index))
-		{
-			HitProxyColor = HitProxies[Index]->Id.GetColor();
-		}
-		InstanceEditorData.Add(FInstanceEditorData::Pack(HitProxyColor, bSelected));
-	}
-	SelectedInstances = InSelectedInstances;
-}
-
-#endif
-
 void FISMInstanceUpdateChangeSet::SetInstanceTransforms(TStridedView<FMatrix> InInstanceTransforms, const FVector Offset)
 {
-
-	GatherTransform(GetTransformDelta(), Transforms, InInstanceTransforms, [Offset](const FMatrix &M) -> FRenderTransform { return FRenderTransform(M.ConcatTranslation(Offset)); });
+	GetTransformWriter().Gather(
+		[&, Offset](int32 InstanceIndex) -> FRenderTransform 
+		{ 
+			return FRenderTransform(InInstanceTransforms[InstanceIndex].ConcatTranslation(Offset));
+		});
 }
 
 void FISMInstanceUpdateChangeSet::SetInstanceTransforms(TStridedView<FMatrix> InInstanceTransforms)
 {
-	GatherTransform(GetTransformDelta(), Transforms, InInstanceTransforms, [](const FMatrix &M) -> FRenderTransform { return FRenderTransform(M); });
+	GetTransformWriter().Gather(
+		[&](int32 InstanceIndex) -> FRenderTransform 
+		{ 
+			return FRenderTransform(InInstanceTransforms[InstanceIndex]);
+		});
 }
 
 void FISMInstanceUpdateChangeSet::SetInstanceTransforms(TStridedView<FMatrix> InInstanceTransforms, FBox const& InInstanceBounds, FBox& OutGatheredBounds)
 {
-	GatherTransform(GetTransformDelta(),Transforms, InInstanceTransforms, [&InInstanceBounds, &OutGatheredBounds](const FMatrix& M) -> FRenderTransform
-	{
-		FRenderTransform Transform(M);
-		OutGatheredBounds += InInstanceBounds.TransformBy(Transform.ToMatrix());
-		return Transform;
-	});
+	GetTransformWriter().Gather(
+		[&](int32 InstanceIndex) -> FRenderTransform
+		{
+			const FMatrix& M = InInstanceTransforms[InstanceIndex];
+			FRenderTransform Transform(M);
+			OutGatheredBounds += InInstanceBounds.TransformBy(Transform.ToMatrix());
+			return Transform;
+		});
 }
 
 void FISMInstanceUpdateChangeSet::SetInstancePrevTransforms(TArrayView<FMatrix> InPrevInstanceTransforms, const FVector &Offset)
 {
-	if (Flags.bHasPerInstanceDynamicData)
-	{
-		if (InPrevInstanceTransforms.IsEmpty())
-		{
-			Flags.bHasPerInstanceDynamicData = false;
-		}
-		else
-		{
-			GatherTransform(GetTransformDelta(), PrevTransforms, InPrevInstanceTransforms, [Offset](const FMatrix &M) -> FRenderTransform { return FRenderTransform(M.ConcatTranslation(Offset)); });
-		}
-	}
+	GetPrevTransformWriter().Gather(
+		[&, Offset](int32 InstanceIndex) -> FRenderTransform 
+		{ 
+			return FRenderTransform(InPrevInstanceTransforms[InstanceIndex].ConcatTranslation(Offset));
+		});
 }
 
 void FISMInstanceUpdateChangeSet::SetInstancePrevTransforms(TArrayView<FMatrix> InPrevInstanceTransforms)
 {
-	if (Flags.bHasPerInstanceDynamicData)
-	{
-		if (InPrevInstanceTransforms.IsEmpty())
-		{
-			Flags.bHasPerInstanceDynamicData = false;
-		}
-		else
-		{
-			GatherTransform(GetTransformDelta(), PrevTransforms, InPrevInstanceTransforms, [](const FMatrix &M) -> FRenderTransform { return FRenderTransform(M); });
-		}
-	}
+	GetPrevTransformWriter().Gather(
+		[&](int32 InstanceIndex) -> FRenderTransform 
+		{ 
+			return FRenderTransform(InPrevInstanceTransforms[InstanceIndex]);
+		});
 }
 
-void FISMInstanceUpdateChangeSet::SetCustomData(const TArrayView<const float> &InPerInstanceCustomData, int32 InNumCustomDataFloats)
-{
-	if (Flags.bHasPerInstanceCustomData)
-	{
-		check(InNumCustomDataFloats == NumCustomDataFloats);
-		check(NumCustomDataFloats > 0);
-		//check(InPerInstanceCustomData.Num() == InstanceIdIndexMap.GetMaxInstanceIndex() * NumCustomDataFloats);
-
-		Gather(GetCustomDataDelta(), PerInstanceCustomData, InPerInstanceCustomData, InNumCustomDataFloats);
-	}
-	else
-	{
-		checkSlow(GetCustomDataDelta().IsEmpty());
-	}
-}
-
-void FISMInstanceUpdateChangeSet::SetInstanceLocalBounds(const FRenderBounds &Bounds)
-{
-	InstanceLocalBounds.SetNum(1);
-	InstanceLocalBounds[0] = Bounds;
-}
-
-#if WITH_EDITOR
-
-TPimplPtr<FOpaqueHitProxyContainer> MakeOpaqueHitProxyContainer(const TArray<TRefCountPtr<HHitProxy>>& InHitProxies)
+void FISMInstanceUpdateChangeSet::SetCustomData(const TArrayView<const float>& InPerInstanceCustomData, int32 InNumCustomDataFloats)
 {
-	return MakePimpl<FOpaqueHitProxyContainer>(InHitProxies);
+	GetCustomDataWriter().Gather(InPerInstanceCustomData, InNumCustomDataFloats);
 }
 
-#endif
diff --git a/Engine/Source/Runtime/Engine/Private/InstancedStaticMesh/ISMScatterGatherUtil.h b/Engine/Source/Runtime/Engine/Private/InstancedStaticMesh/ISMScatterGatherUtil.h
deleted file mode 100644
diff --git a/Engine/Source/Runtime/Engine/Private/ObjectCacheContext.cpp b/Engine/Source/Runtime/Engine/Private/ObjectCacheContext.cpp
--- a/Engine/Source/Runtime/Engine/Private/ObjectCacheContext.cpp
+++ b/Engine/Source/Runtime/Engine/Private/ObjectCacheContext.cpp
@@ -840,8 +840,8 @@ TObjectCacheIterator<UAnimBankMeshComponent> FObjectCacheContext::GetAnimBankMes
 TObjectCacheIterator<UAnimBankMeshComponent> FObjectCacheContext::GetAnimBankMeshComponents(UAnimBank* InAnimBank)
 {
 	return AnimBankMeshCache.Get<TObjectIterator<UAnimBankMeshComponent>>(InAnimBank, 
-		[](UAnimBankMeshComponent* Component){ return Component->AnimBankItems.Num();  },
-		[](UAnimBankMeshComponent* Component, int32 Index)->const UAnimBank* { return Component->AnimBankItems[Index].BankAsset; }
+		[](UAnimBankMeshComponent* Component){ return Component->GetAnimBankItems().Num();  },
+		[](UAnimBankMeshComponent* Component, int32 Index)->const UAnimBank* { return Component->GetAnimBankItems()[Index].BankAsset; }
 	);
 }
 
diff --git a/Engine/Source/Runtime/Engine/Public/InstanceData/InstanceDataManager.h b/Engine/Source/Runtime/Engine/Public/InstanceData/InstanceDataManager.h
new file mode 100644
--- /dev/null
+++ b/Engine/Source/Runtime/Engine/Public/InstanceData/InstanceDataManager.h
@@ -0,0 +1,234 @@
+// Copyright Epic Games, Inc. All Rights Reserved.
+
+#pragma once
+
+#include "SceneTypes.h"
+#include "InstanceDataSceneProxy.h"
+#include "InstancedStaticMesh/InstanceAttributeTracker.h"
+
+class FPrimitiveSceneProxy;
+class HHitProxy;
+class FPrecomputedInstanceSpatialHashData;
+class FInstanceUpdateChangeSet;
+using FPrecomputedInstanceSpatialHashDataPtr = TSharedPtr<const FPrecomputedInstanceSpatialHashData, ESPMode::ThreadSafe>;
+
+/**
+ * Data descriptor representing the component state abstracting the UPrimitiveComponent, needs to be passed into the change flushing.
+ * The intention is to decouple the manager from the UComponent or any other supplier of instance data & scene proxies.
+ */
+struct FInstanceDataManagerSourceDataDesc
+{
+	FMatrix PrimitiveLocalToWorld;
+	EComponentMobility::Type ComponentMobility = EComponentMobility::Movable;
+	FRenderBounds MeshBounds;
+	FInstanceDataFlags Flags;
+
+	FPrimitiveMaterialPropertyDescriptor PrimitiveMaterialDesc;
+
+	// Number of instances in the source arrays (e.g., in the component)
+	int32 NumInstances = -1;
+
+	// 
+	int32 NumCustomDataFloats = 0;
+
+	// Callback to fill in the required change set with source data.
+	TFunction<void(FInstanceUpdateChangeSet& ChangeSet)> BuildChangeSet;
+};
+
+class InstanceDataManagerSourceInterface
+{
+public:
+
+	enum class FDirtyType
+	{
+		Incremental,
+		Full,
+	};
+
+	/**
+	 * Called to let the source (UComponent or whatever) know there has been a change that will needs to be flushed.
+	 */
+	virtual void InstanceDataManagerMarkDirty(FDirtyType DirtyFlags) = 0;
+
+	/**
+	 * Called to retrieve the data for 
+	 */
+	virtual FInstanceDataManagerSourceDataDesc GetInstanceDataSourceDesc() = 0;
+};
+
+/**
+ * Manager class that tracks changes to instance data within the component, and is responsible for dispatching updates of the proxy.
+ * Tracks instance index changes to be able to maintain a persistent ID mapping for use on the render thread.
+ * The ID mapping is not serialized and will be reset when the proxy is recreated.
+ * Not responsible for storing the component representation of the instance data.
+ * NOTE/TODO: This is tied to the ISM use-case, mostly because of legacy (HISM) interactions. Will be refactored and sub-classed or something.
+ *            Also: Still somewhat tied to the UComponent, which also can be refactored a bit to make it more general.
+ */
+class FInstanceDataManager : public FInstanceIdIndexMap
+{
+public:
+	//ENGINE_API FInstanceDataManager(IInstanceDataProvider* InInstanceDataProvider);
+	ENGINE_API FInstanceDataManager(UPrimitiveComponent* InPrimitiveComponent);
+
+	/**
+	 * Current tracking state, 
+	 */
+	enum class ETrackingState : uint8
+	{
+		// In the initial state, there is no proxy and therefore changes do not need to be tracked, e.g., during initial setup of an ISM component.
+		Initial, 
+		// Tracking changes to send on next flush
+		Tracked,
+		// Prevent any changes from being tracked (e.g., if we have no renderer)
+		Disabled,
+		// In the optimized state there's no need to track any delta changes, but if anything changes at all we must rebuild.
+		Optimized, 
+	};
+
+	/**
+	 * Tracking functions that mirror what is done to each instance in the source instance data array.
+	 */
+
+	FPrimitiveInstanceId Add(int32 InInstanceAddAtIndex);
+	void RemoveAtSwap(int32 InstanceIndex);
+	void RemoveAt(int32 InstanceIndex);
+
+	void TransformChanged(int32 InstanceIndex);
+	void TransformChanged(FPrimitiveInstanceId InstanceId);
+	void TransformsChangedAll();
+
+	void CustomDataChanged(int32 InstanceIndex);
+
+	void BakedLightingDataChanged(int32 InstanceIndex);
+	void BakedLightingDataChangedAll();
+
+	void NumCustomDataChanged();
+
+#if WITH_EDITOR
+	void EditorDataChangedAll();
+#endif
+
+	void PrimitiveTransformChanged();
+
+	ENGINE_API bool HasAnyInstanceChanges() const;
+
+	/** 
+	 * Returns true if there are explicitly tracked instance changes, or the state is not tracked (because no proxy has been created yet),
+	 * and the tracking state is not Disabled.
+	 */
+	inline bool HasAnyChanges() const { return GetState() != ETrackingState::Disabled && (GetState() != ETrackingState::Tracked || HasAnyInstanceChanges());}
+
+	/**
+	 * Queries the tracker for changes and builds an update build command to enqueue to the render thread. 
+	 * The ComponentData supplies source data through a callback as needed.
+	 */
+	bool FlushChanges(FInstanceDataManagerSourceDataDesc&& ComponentData);
+
+	/**
+	 * Clear all tracked changes (will result in a full update when next one is flushed)
+	 */
+	void ClearChangeTracking();
+
+	int32 GetMaxAllocatedInstanceId() const;
+
+	ETrackingState GetState() const { return TrackingState; }
+
+	ENGINE_API TSharedPtr<FInstanceDataSceneProxy, ESPMode::ThreadSafe> GetOrCreateProxy();
+	ENGINE_API TSharedPtr<FInstanceDataSceneProxy, ESPMode::ThreadSafe> GetProxy();
+
+#if DO_GUARD_SLOW
+	void ValidateMapping() const;
+#else
+	FORCEINLINE void ValidateMapping() const {};
+#endif
+
+#if WITH_EDITOR
+	void BeginCacheForCookedPlatformData(const ITargetPlatform* TargetPlatform, FInstanceDataManagerSourceDataDesc&& ComponentData);
+	void WriteCookedRenderData(FArchive& Ar, FInstanceDataManagerSourceDataDesc&& ComponentData);
+#endif
+	void ReadCookedRenderData(FArchive& Ar);
+
+	/**
+	 */
+	void Serialize(FArchive& Ar, bool bCooked);
+
+	/**
+	 * Compute the instance order used to build the instance hierarchy (spatial hashes) and return the reordering needed for this.
+	 * This should be used by the data source (i.e., UComponent etc) to reorder the source instance buffers.
+	 * Returns a reordering table where each index stores the old index for the given new index.
+	 */
+	TArray<int32> Optimize(FInstanceDataManagerSourceDataDesc&& ComponentData, bool bShouldRetainIdMap);
+
+	/**
+	 */
+	ENGINE_API SIZE_T GetAllocatedSize() const;
+
+protected:
+
+	void CreateExplicitIdentityMapping();
+
+	using EChangeFlag = FInstanceAttributeTracker::EFlag;
+
+	template<EChangeFlag Flag>
+	inline void MarkChangeHelper(int32 InstanceIndex);
+	template<EChangeFlag Flag>
+	inline void MarkChangeHelper(FPrimitiveInstanceId InstanceId);
+
+	void MarkComponentRenderInstancesDirty();
+
+	bool HasIdentityMapping() const;
+
+	void FreeInstanceId(FPrimitiveInstanceId InstanceId);
+
+	/** */
+	void GatherDefaultData(const FInstanceDataManagerSourceDataDesc& ComponentData, FInstanceUpdateChangeSet& ChangeSet) const;
+
+	/**
+	 * Initialize a chage set from the component data & manager state but not using any delta information or updating tracked state.
+	 */
+	void InitChangeSet(const FInstanceDataManagerSourceDataDesc &ComponentData, FInstanceUpdateChangeSet& ChangeSet) const;
+
+	void InitChangeSet(const struct FChangeDesc2 &ChangeDesc, const FInstanceDataManagerSourceDataDesc& ComponentData, FInstanceUpdateChangeSet& ChangeSet);
+
+#if WITH_EDITOR
+	// TODO: implement
+	bool ShouldWriteCookedData(const ITargetPlatform* TargetPlatform, int32 NumInstancesToBuildFor);
+
+	/**
+	 * Build precomputed data from the ComponentData.
+	 */
+	FPrecomputedInstanceSpatialHashData PrecomputeOptimizationData(FInstanceDataManagerSourceDataDesc&& ComponentData);
+
+	/**
+	 * Build precomputed data from a ChangeSet.
+	 */
+	static FPrecomputedInstanceSpatialHashData PrecomputeOptimizationData(FInstanceUpdateChangeSet& ChangeSet);
+
+#endif
+	TSharedPtr<class FUpdatableInstanceDataSceneProxy, ESPMode::ThreadSafe> GetOrCreateProxyInternal();
+
+	// Change set.
+	FInstanceAttributeTracker InstanceUpdateTracker;
+
+	// Id allocation tracking
+	TBitArray<> ValidInstanceIdMask;
+	int32 IdSearchStartIndex = 0;
+
+	ETrackingState TrackingState = ETrackingState::Initial;
+	FInstanceDataFlags AllChangedFlags;
+	uint8 bNumCustomDataChanged : 1;
+	uint8 bTransformChangedAllInstances : 1;
+	uint8 bPrimitiveTransformChanged : 1;
+	uint8 bAnyInstanceChange : 1;
+
+	TSharedPtr<class FUpdatableInstanceDataSceneProxy, ESPMode::ThreadSafe> InstanceDataProxy;
+	TWeakObjectPtr<UPrimitiveComponent> PrimitiveComponent = nullptr;
+
+	FPrecomputedInstanceSpatialHashDataPtr PrecomputedOptimizationData;
+
+	// Must track this to detect changes 
+	FInstanceDataFlags Flags;
+	int32 NumCustomDataFloats = 0;
+	float AbsMaxDisplacement = 0.0f;
+	FRenderBounds MeshBounds;
+};
\ No newline at end of file
diff --git a/Engine/Source/Runtime/Engine/Public/InstanceData/InstanceUpdateChangeSet.h b/Engine/Source/Runtime/Engine/Public/InstanceData/InstanceUpdateChangeSet.h
new file mode 100644
--- /dev/null
+++ b/Engine/Source/Runtime/Engine/Public/InstanceData/InstanceUpdateChangeSet.h
@@ -0,0 +1,416 @@
+// Copyright Epic Games, Inc. All Rights Reserved.
+
+#pragma once
+
+#include "InstanceDataSceneProxy.h"
+#include "Containers/StridedView.h"
+#include "InstancedStaticMesh/InstanceAttributeTracker.h"
+
+class HHitProxy;
+class FOpaqueHitProxyContainer;
+
+/**
+ * Helper to make it possible to use the same paths for gather / scatter even if the per-instance delta is not tracked.
+ */
+class FIdentityDeltaRange
+{
+public:
+	FIdentityDeltaRange(int32 InNum) : NumItems(InNum) {}
+
+	/**
+		*/
+	inline bool IsEmpty() const { return NumItems == 0; }
+
+	/**
+		*/
+	inline bool IsDelta() const { return false; }
+
+	/**
+		* Returns the number of items in this range - i.e., the number of items that need to be copied to collect an update.
+		*/
+	inline int32 GetNumItems() const { return NumItems; }
+
+	/**
+		*/
+	struct FConstIterator
+	{
+		int32 ItemIndex = 0;
+		int32 MaxNum = 0;
+
+		FConstIterator(int32 InIndex, int32 InMaxNum)
+		:	ItemIndex(InIndex),
+			MaxNum(InMaxNum)
+		{
+		}
+
+		void operator++() {  ++ItemIndex; }
+
+		/**
+			* Get the index of the data in the source / destination arrays. 
+			*/
+		int32 GetIndex() const { return ItemIndex; }
+
+		/**
+			* Get the continuous index of the data item in the collected item array.
+			*/
+		int32 GetItemIndex() const { return ItemIndex; }
+
+		explicit operator bool() const {  return ItemIndex < MaxNum; }
+	};
+
+	FConstIterator GetIterator() const
+	{
+		return FConstIterator(0, NumItems);
+	}
+
+private:
+	int32 NumItems = 0;
+};
+
+template <typename ElementType, typename DeltaType>
+struct TDeltaSetup
+{
+	using FElement = ElementType;
+	using FDelta = DeltaType;
+
+	TArray<FElement>& DeltaDataArray;
+	FDelta Delta;
+	int32 ElementStride = 1;
+	bool bIsEnabled = true;
+	int32 NumInstances = 0;
+
+
+	/**
+	 * Binds a specific delta type to an array for writing.
+	 */
+	struct FWriter
+	{
+		TDeltaSetup Setup;
+
+		void Gather(TFunctionRef<FElement(int32)> DataSourceFunc)
+		{
+			if (!Setup.bIsEnabled)
+			{
+				return;
+			}
+
+			if (Setup.Delta.IsEmpty())
+			{
+				Setup.DeltaDataArray.Reset();
+			}
+			else
+			{
+				Setup.DeltaDataArray.Reset(Setup.Delta.GetNumItems());
+				for (auto It = Setup.Delta.GetIterator(); It; ++It)
+				{
+					check(Setup.DeltaDataArray.Num() < Setup.Delta.GetNumItems());
+					check(Setup.DeltaDataArray.Num() == It.GetItemIndex());
+					Setup.DeltaDataArray.Emplace(DataSourceFunc(It.GetIndex()));
+				}
+			}
+		}
+
+		template <typename InElementType>
+		void Gather(const TArrayView<InElementType> SourceData, int32 InElementStride = 1)
+		{
+			if (!Setup.bIsEnabled)
+			{
+				return;
+			}
+
+			if (Setup.Delta.IsEmpty())
+			{
+				Setup.DeltaDataArray.Reset();
+			}
+			// strides & element count matches - just copy the data
+			else 
+			{
+				check(InElementStride == Setup.ElementStride);
+				// It is a full update if either it is not a delta, or we're sending everything anyway
+				bool bIsFull = !Setup.Delta.IsDelta() || SourceData.Num() == Setup.Delta.GetNumItems() * Setup.ElementStride;
+
+				if (bIsFull)
+				{
+					// full update, bulk-copy
+					Setup.DeltaDataArray = SourceData;
+				}
+				else
+				{
+					Setup.DeltaDataArray.Reset(Setup.Delta.GetNumItems() * Setup.ElementStride);
+					for (auto It = Setup.Delta.GetIterator(); It; ++It)
+					{
+						check(Setup.DeltaDataArray.Num() < Setup.Delta.GetNumItems() * Setup.ElementStride);
+						Setup.DeltaDataArray.Append(&SourceData[It.GetIndex() * Setup.ElementStride], Setup.ElementStride);
+					}
+				}
+			}
+		}
+	};
+
+	/**
+	 * Binds a specific delta type to an array for writing.
+	 */
+	struct FReader
+	{
+		TDeltaSetup Setup;
+
+		template <typename IndexRemapType>
+		void Scatter(TArray<ElementType>& OutDataArray, const IndexRemapType& IndexRemap)
+		{
+			static_assert(std::is_same_v<decltype(OutDataArray), decltype(Setup.DeltaDataArray)>, "The types should match or it won't be able to use MoveTemp properly");
+			if (!Setup.bIsEnabled)
+			{
+				OutDataArray.Reset();
+			}
+			else 
+			{
+				check(Setup.DeltaDataArray.Num() == Setup.Delta.GetNumItems() * Setup.ElementStride);
+
+				bool bIsFull = !Setup.Delta.IsDelta() || Setup.NumInstances * Setup.ElementStride == Setup.DeltaDataArray.Num();;
+				if (bIsFull && IndexRemap.IsIdentity())
+				{
+					check(Setup.DeltaDataArray.Num() == Setup.NumInstances * Setup.ElementStride);
+					OutDataArray = MoveTemp(Setup.DeltaDataArray);
+				}
+				else
+				{
+					OutDataArray.SetNumUninitialized(Setup.NumInstances * Setup.ElementStride);
+					for (auto It = Setup.Delta.GetIterator(); It; ++It)
+					{
+						int32 ItemIndex = It.GetItemIndex();
+						int32 DestIndex = It.GetIndex();
+						IndexRemap.Remap(ItemIndex, DestIndex);
+						FMemory::Memcpy(&OutDataArray[DestIndex * Setup.ElementStride], &Setup.DeltaDataArray[ItemIndex * Setup.ElementStride], Setup.ElementStride * sizeof(FElement));
+					}
+				}
+			}
+		}
+
+		template <typename ElementTransformFuncType, typename IndexRemapType>
+		void Scatter(TArray<ElementType>& OutDataArray, ElementTransformFuncType ElementTransformFunc, const IndexRemapType& IndexRemap)
+		{
+			static_assert(std::is_same_v<decltype(OutDataArray), decltype(Setup.DeltaDataArray)>, "The types should match or it won't be able to use MoveTemp properly");
+			if (!Setup.bIsEnabled)
+			{
+				OutDataArray.Reset();
+			}
+			else 
+			{
+				check(Setup.ElementStride == 1);
+				bool bIsFull = !Setup.Delta.IsDelta() || Setup.NumInstances * Setup.ElementStride == Setup.DeltaDataArray.Num();;
+				if (bIsFull && IndexRemap.IsIdentity())
+				{
+					check(Setup.DeltaDataArray.Num() == Setup.NumInstances * Setup.ElementStride);
+					// Just change ownership of the array
+					OutDataArray = MoveTemp(Setup.DeltaDataArray);
+					// Then apply transform in-place
+					for (int32 Index = 0; Index < Setup.NumInstances; ++Index)
+					{
+						ElementTransformFunc(OutDataArray[Index]);
+					}
+				}
+				else
+				{
+					check(Setup.DeltaDataArray.Num() == Setup.Delta.GetNumItems() * Setup.ElementStride);
+
+					OutDataArray.SetNumUninitialized(Setup.NumInstances);
+					for (auto It = Setup.Delta.GetIterator(); It; ++It)
+					{
+						int32 ItemIndex = It.GetItemIndex();
+						int32 DestIndex = It.GetIndex();
+						IndexRemap.Remap(ItemIndex, DestIndex);
+						ElementType Element = Setup.DeltaDataArray[ItemIndex];
+						ElementTransformFunc(Element);
+						OutDataArray[DestIndex] = Element;
+					}
+				}
+			}
+		}
+	};
+
+	FReader GetReader()
+	{
+		check(!bIsEnabled || DeltaDataArray.Num() == Delta.GetNumItems() * ElementStride);
+		return FReader{ *this };
+	}
+
+	FWriter GetWriter()
+	{
+		check(!bIsEnabled || DeltaDataArray.IsEmpty());
+		return FWriter{ *this };
+	}
+};
+
+/**
+ * Precomputed optimization data that descrives the spatial hashes and reordering needed.
+ */
+class FPrecomputedInstanceSpatialHashData
+{
+public:
+	TArray<FInstanceSceneDataBuffers::FCompressedSpatialHashItem> Hashes;
+	TArray<int32> ProxyIndexToComponentIndexRemap;
+};
+
+using FPrecomputedInstanceSpatialHashDataPtr = TSharedPtr<const FPrecomputedInstanceSpatialHashData, ESPMode::ThreadSafe>;
+
+
+
+/**
+ * Collects changed instance data (and what else is needed to update the instance data proxy) from the source and 
+ */
+class FInstanceUpdateChangeSet
+{
+public:
+	/**
+	 * Cosntruct a full change set, with no delta (which will collect all enabled data).
+	 */
+	FInstanceUpdateChangeSet(int32 InNumSourceInstances, FInstanceDataFlags InFlags) 
+		: bNeedFullUpdate(true) 
+		, Flags(InFlags)
+		, NumSourceInstances(InNumSourceInstances)
+	{
+	}
+
+	/**
+	 * Cosntruct a delta change set, but which can be forced to full using bInNeedFullUpdate.
+	 */
+	FInstanceUpdateChangeSet(bool bInNeedFullUpdate, FInstanceAttributeTracker&& InInstanceAttributeTracker, int32 InNumSourceInstances) 
+		: InstanceAttributeTracker(MoveTemp(InInstanceAttributeTracker))
+		, bNeedFullUpdate(bInNeedFullUpdate) 
+		, NumSourceInstances(InNumSourceInstances)
+	{
+	}
+
+	template <FInstanceAttributeTracker::EFlag Flag>
+	FInstanceAttributeTracker::FDeltaRange<Flag> GetDelta(bool bForceEmpty, bool bForceFull = false) const 
+	{ 
+		if (bForceEmpty)
+		{
+			return FInstanceAttributeTracker::FDeltaRange<Flag>();
+		}
+		return InstanceAttributeTracker.GetDeltaRange<Flag>(IsFullUpdate() || bForceFull, NumSourceInstances); 
+	}
+
+	FInstanceAttributeTracker::FDeltaRange<FInstanceAttributeTracker::EFlag::TransformChanged> GetTransformDelta() const 
+	{ 
+		return GetDelta<FInstanceAttributeTracker::EFlag::TransformChanged>(false, bUpdateAllInstanceTransforms);
+	}
+
+	FInstanceAttributeTracker::FDeltaRange<FInstanceAttributeTracker::EFlag::IndexChanged> GetIndexChangedDelta() const 
+	{ 
+		return GetDelta<FInstanceAttributeTracker::EFlag::IndexChanged>(false);
+	}
+
+	template <FInstanceAttributeTracker::EFlag DeltaFlag, typename ElementType>
+	TDeltaSetup<ElementType, FInstanceAttributeTracker::FDeltaRange<DeltaFlag>> GetSetup(bool bEnabledFlag, bool bForceFullFlag, TArray<ElementType> &DataArray, int32 ElementStride = 1)
+	{
+		FInstanceAttributeTracker::FDeltaRange<DeltaFlag> Delta = GetDelta<DeltaFlag>(!bEnabledFlag, bForceFullFlag);
+		return TDeltaSetup<ElementType, FInstanceAttributeTracker::FDeltaRange<DeltaFlag>>{ DataArray, Delta, ElementStride, bEnabledFlag, NumSourceInstances };
+	}
+
+	template <typename ElementType>
+	TDeltaSetup<ElementType, FIdentityDeltaRange> GetSetup(bool bEnabledFlag, TArray<ElementType> &DataArray, int32 ElementStride = 1)
+	{
+		FIdentityDeltaRange Delta = FIdentityDeltaRange(bEnabledFlag ? NumSourceInstances : 0);
+		return TDeltaSetup<ElementType, FIdentityDeltaRange>{ DataArray, Delta, ElementStride, bEnabledFlag, NumSourceInstances };
+	}
+
+	// These setups define the mapping from a delta attribute bit in the tracker to the data array. 
+	// This is not a simple 1:1 mapping as we only track a few bits, and there is also special overrides to take into account.
+	auto GetTransformSetup()		{ return GetSetup<FInstanceAttributeTracker::EFlag::TransformChanged>(true, bUpdateAllInstanceTransforms, Transforms); }
+	auto GetPrevTransformSetup()	{ return GetSetup<FInstanceAttributeTracker::EFlag::TransformChanged>(Flags.bHasPerInstanceDynamicData, ForceFullFlags.bHasPerInstanceDynamicData, PrevTransforms); }
+	auto GetCustomDataSetup()		{ return GetSetup<FInstanceAttributeTracker::EFlag::CustomDataChanged>(Flags.bHasPerInstanceCustomData, ForceFullFlags.bHasPerInstanceCustomData, PerInstanceCustomData, NumCustomDataFloats); }
+	// These use an identity delta, which means they send all or nothing.
+	auto GetLocalBoundsSetup()		{ return GetSetup(Flags.bHasPerInstanceLocalBounds, InstanceLocalBounds); }
+	auto GetSkinningDataSetup()		{ return GetSetup(Flags.bHasPerInstanceSkinningData, InstanceSkinningData); }
+
+	// Convenience functions to get a reader / writer for a given array
+	auto GetTransformWriter()		{ return GetTransformSetup().GetWriter(); }
+	auto GetPrevTransformWriter()	{ return GetPrevTransformSetup().GetWriter(); }
+	auto GetCustomDataWriter()		{ return GetCustomDataSetup().GetWriter(); }
+	auto GetLocalBoundsWriter()		{ return GetLocalBoundsSetup().GetWriter(); }
+	auto GetSkinningDataWriter()	{ return GetSkinningDataSetup().GetWriter(); }
+
+	auto GetTransformReader()		{ return GetTransformSetup().GetReader(); }
+	auto GetPrevTransformReader()	{ return GetPrevTransformSetup().GetReader(); }
+	auto GetCustomDataReader()		{ return GetCustomDataSetup().GetReader(); }
+	auto GetLocalBoundsReader()		{ return GetLocalBoundsSetup().GetReader(); }
+	auto GetSkinningDataReader()	{ return GetSkinningDataSetup().GetReader(); }
+
+	auto GetLightShadowUVBiasReader() { return GetSetup(Flags.bHasPerInstanceLMSMUVBias, InstanceLightShadowUVBias).GetReader(); }
+
+#if WITH_EDITOR
+	auto GetEditorDataReader() { return GetSetup(Flags.bHasPerInstanceEditorData, InstanceEditorData).GetReader(); }
+	// Set editor data 
+	void SetEditorData(const TArray<TRefCountPtr<HHitProxy>>& HitProxies, const TBitArray<> &SelectedInstances);
+#endif
+
+	/**
+	 * Used to set a single, shared, instance local bounds, only allowed when Flags.bHasPerInstanceLocalBounds is false.
+	 */
+	ENGINE_API void SetSharedLocalBounds(const FRenderBounds &Bounds);
+
+	bool IsFullUpdate() const
+	{
+		return bNeedFullUpdate;
+	}
+
+	FInstanceAttributeTracker InstanceAttributeTracker;
+	bool bNeedFullUpdate = false;
+
+	/**
+	 * Flags that dictate what attributes will be gathered from the source (and deposited at the destination). These must be present during gather (or they should have been disabled earlier).
+	 * Constructed as the intersection of the flags indicated by the component & what is supported by the proxy(?).
+	 */
+	FInstanceDataFlags Flags;
+	
+	/**
+	 * Flags that can be set to force individual attributes to use a full update has effect IFF Flags is set for the given attribute.
+	 * If an attribute was enabled after not being used (e.g., on a material change, perhaps) a full update must be sent.
+	 */
+	FInstanceDataFlags ForceFullFlags;
+
+	// Needs its own bool because it is always present and thus doesn't have a flag in FInstanceDataFlags
+	bool bUpdateAllInstanceTransforms = false;
+	bool bIdentityIdMap = false;
+	TArray<FPrimitiveInstanceId> IndexToIdMapDeltaData;
+
+	int32 NumCustomDataFloats = 0;
+	TArray<FRenderTransform> Transforms;
+	TArray<FRenderTransform> PrevTransforms;
+	TArray<float> PerInstanceCustomData;
+	TArray<uint32> InstanceSkinningData;
+	TArray<FVector4f> InstanceLightShadowUVBias;
+	TArray<FRenderBounds> InstanceLocalBounds;
+
+#if WITH_EDITOR
+	TArray<uint32> InstanceEditorData;
+	TBitArray<> SelectedInstances;
+	TPimplPtr<FOpaqueHitProxyContainer> HitProxyContainer;
+#endif
+
+	// Function that can generate all the random IDs on demand, if this is supplied is must generate all, otherwise it is zero-filled if they are requested.
+	TFunction<void(TArray<float> &InstanceRandomIDs)> GeneratePerInstanceRandomIds;
+
+	FRenderTransform PrimitiveToRelativeWorld;
+	FVector PrimitiveWorldSpaceOffset;
+	TOptional<FRenderTransform> PreviousPrimitiveToRelativeWorld;
+	float AbsMaxDisplacement = 0.0f;
+	int32 NumSourceInstances = 0;
+	int32 MaxInstanceId = 0;
+
+	/**
+	 * Describes precomputed spatial hashes and the instance reordering that is needed to use these.
+	 * If set, the update must reorder the data appropriately.
+	 * Should only ever be used together with a full rebuild.
+	 */
+	FPrecomputedInstanceSpatialHashDataPtr PrecomputedOptimizationData;
+};
+
+
+#if WITH_EDITOR
+
+ENGINE_API TPimplPtr<FOpaqueHitProxyContainer> MakeOpaqueHitProxyContainer(const TArray<TRefCountPtr<HHitProxy>>& InHitProxies);
+
+
+#endif
diff --git a/Engine/Source/Runtime/Engine/Public/InstanceDataSceneProxy.h b/Engine/Source/Runtime/Engine/Public/InstanceDataSceneProxy.h
--- a/Engine/Source/Runtime/Engine/Public/InstanceDataSceneProxy.h
+++ b/Engine/Source/Runtime/Engine/Public/InstanceDataSceneProxy.h
@@ -14,6 +14,8 @@ class FStaticMeshInstanceBuffer;
 class FInstanceIdIndexMap
 {
 public:
+	FInstanceIdIndexMap() = default;
+
 	/**
 	 * Returns true if the mapping is an identity mapping. I.e., each instance ID corresponds to the same index.
 	 */
@@ -60,6 +62,7 @@ public:
 
 	/**
 	 * Update the mapping from ID to index (and vice versa).
+	 * Does _not_ reset any previous mapping for the given ID.
 	 */
 	FORCEINLINE void Update(FPrimitiveInstanceId InstanceId, int32 InstanceIndex)
 	{
@@ -83,14 +86,25 @@ public:
 	 */
 	ENGINE_API void CreateExplicitIdentityMapping();
 
-	FInstanceIdIndexMap() = default;
+	/**
+	 */
+	ENGINE_API void Serialize(FArchive& Ar);
+
+	/**
+	 */	
+	ENGINE_API void RebuildFromIndexToIdMap(TArray<FPrimitiveInstanceId>&& InIndexToIdMap, int32 MaxInstanceId);
+
 protected:
+
+	/**
+	 */	
+	ENGINE_API void RebuildFromIndexToIdMap(int32 MaxInstanceId);
+
 	// Bidirectional mapping to / from ID.
 	TArray<FPrimitiveInstanceId> IndexToIdMap;
 	TArray<int32> IdToIndexMap;
 	// used when the mapping is implicit (i.e., identity)
 	int32 NumInstances = 0;
-	bool bIsValid = false;
 };
 
 #if WITH_EDITOR
@@ -114,7 +128,7 @@ struct FInstanceEditorData
 #endif
 
 class FInstanceSceneDataImmutable;
-
+using FInstanceSceneDataImmutablePtr = TSharedPtr<FInstanceSceneDataImmutable, ESPMode::ThreadSafe>;
 /**
  */
 class FInstanceSceneDataBuffers
@@ -207,26 +221,26 @@ public:
 	struct FWriteView
 	{
 		FInstanceSceneDataBuffers::FAccessTag AccessTag;
-		FRenderTransform &PrimitiveToRelativeWorld;
-		FVector &PrimitiveWorldSpaceOffset;
-		TArray<FRenderBounds, TInlineAllocator<1>> &InstanceLocalBounds;
-		TArray<float> &InstanceCustomData;
-		TArray<float> &InstanceRandomIDs;
-		TArray<FVector4f> &InstanceLightShadowUVBias;
-		TArray<uint32> &InstanceSkinningData;
-		TArray<uint32> &InstanceHierarchyOffset;
-		TArray<FVector4f> &InstancePayloadExtension;
-		TArray<FRenderTransform> &InstanceToPrimitiveRelative;
-		TArray<FRenderTransform> &PrevInstanceToPrimitiveRelative;
+		FRenderTransform& PrimitiveToRelativeWorld;
+		FVector& PrimitiveWorldSpaceOffset;
+		TArray<FRenderBounds>& InstanceLocalBounds;
+		TArray<float>& InstanceCustomData;
+		TArray<float>& InstanceRandomIDs;
+		TArray<FVector4f>& InstanceLightShadowUVBias;
+		TArray<uint32>& InstanceSkinningData;
+		TArray<uint32>& InstanceHierarchyOffset;
+		TArray<FVector4f>& InstancePayloadExtension;
+		TArray<FRenderTransform>& InstanceToPrimitiveRelative;
+		TArray<FRenderTransform>& PrevInstanceToPrimitiveRelative;
 	#if WITH_EDITOR
-		TArray<uint32> &InstanceEditorData;
-		TBitArray<> &SelectedInstances;
+		TArray<uint32>& InstanceEditorData;
+		TBitArray<>& SelectedInstances;
 	#endif
-		TBitArray<> &VisibleInstances;
+		TBitArray<>& VisibleInstances;
 
-		int32 &NumInstancesGPUOnly;
-		int32 &NumCustomDataFloats;
-		FInstanceDataFlags &Flags;
+		int32& NumInstancesGPUOnly;
+		int32& NumCustomDataFloats;
+		FInstanceDataFlags& Flags;
 	};
 
 	FWriteView BeginWriteAccess(FAccessTag AccessTag)
@@ -269,22 +283,22 @@ public:
 
 	struct FReadView
 	{
-		const FRenderTransform &PrimitiveToRelativeWorld;
-		const FVector &PrimitiveWorldSpaceOffset;
-		const TArray<FRenderBounds, TInlineAllocator<1>> &InstanceLocalBounds;
-		const TArray<float> &InstanceCustomData;
-		const TArray<float> &InstanceRandomIDs;
-		const TArray<FVector4f> &InstanceLightShadowUVBias;
-		const TArray<uint32> &InstanceSkinningData;
-		const TArray<uint32> &InstanceHierarchyOffset;
-		const TArray<FVector4f> &InstancePayloadExtension;
-		const TArray<FRenderTransform> &InstanceToPrimitiveRelative;
-		const TArray<FRenderTransform> &PrevInstanceToPrimitiveRelative;
+		const FRenderTransform& PrimitiveToRelativeWorld;
+		const FVector& PrimitiveWorldSpaceOffset;
+		const TArray<FRenderBounds>& InstanceLocalBounds;
+		const TArray<float>& InstanceCustomData;
+		const TArray<float>& InstanceRandomIDs;
+		const TArray<FVector4f>& InstanceLightShadowUVBias;
+		const TArray<uint32>& InstanceSkinningData;
+		const TArray<uint32>& InstanceHierarchyOffset;
+		const TArray<FVector4f>& InstancePayloadExtension;
+		const TArray<FRenderTransform>& InstanceToPrimitiveRelative;
+		const TArray<FRenderTransform>& PrevInstanceToPrimitiveRelative;
 	#if WITH_EDITOR
-		const TArray<uint32> &InstanceEditorData;
-		const TBitArray<> &SelectedInstances;
+		const TArray<uint32>& InstanceEditorData;
+		const TBitArray<>& SelectedInstances;
 	#endif
-		const TBitArray<> &VisibleInstances;
+		const TBitArray<>& VisibleInstances;
 
 		int32 NumInstancesGPUOnly;
 		int32 NumCustomDataFloats;
@@ -322,12 +336,12 @@ public:
 	}
 
 	ENGINE_API void SetImmutable(FInstanceSceneDataImmutable &&ImmutableData, FAccessTag AccessTag);
-	inline TSharedPtr<FInstanceSceneDataImmutable, ESPMode::ThreadSafe> GetImmutable() const { return Immutable; }
+	inline FInstanceSceneDataImmutablePtr GetImmutable() const { return Immutable; }
 
 protected:
 	FRenderTransform PrimitiveToRelativeWorld;
 	FVector PrimitiveWorldSpaceOffset;
-	TArray<FRenderBounds, TInlineAllocator<1>> InstanceLocalBounds;
+	TArray<FRenderBounds> InstanceLocalBounds;
 	TArray<float> InstanceCustomData;
 	TArray<float> InstanceRandomIDs;
 	TArray<FVector4f> InstanceLightShadowUVBias;
@@ -371,7 +385,7 @@ protected:
 	FORCEINLINE void ValidateAccess(const FAccessTag& AccessTag) const {}
 #endif
 
-	TSharedPtr<FInstanceSceneDataImmutable, ESPMode::ThreadSafe> Immutable;
+	FInstanceSceneDataImmutablePtr Immutable;
 };
 
 /**
@@ -418,8 +432,9 @@ public:
 	ENGINE_API void WaitForUpdateCompletion();
 
 private:
-	// TODO: break this link
-	friend class FPrimitiveInstanceDataManager;
+	template <typename TaskLambdaType>
+	friend void BeginInstanceDataUpdateTask(FInstanceDataUpdateTaskInfo& InstanceDataUpdateTaskInfo, TaskLambdaType&& TaskLambda, const FInstanceDataBufferHeader& InInstanceDataBufferHeader);
+
 
 	UE::Tasks::FTask UpdateTaskHandle;
 	FInstanceDataBufferHeader InstanceDataBufferHeader;
@@ -452,4 +467,3 @@ protected:
 
 	FInstanceSceneDataBuffers InstanceSceneDataBuffers;
 };
-
diff --git a/Engine/Source/Runtime/Engine/Public/InstanceDataTypes.h b/Engine/Source/Runtime/Engine/Public/InstanceDataTypes.h
--- a/Engine/Source/Runtime/Engine/Public/InstanceDataTypes.h
+++ b/Engine/Source/Runtime/Engine/Public/InstanceDataTypes.h
@@ -4,6 +4,8 @@
 
 #include "CoreMinimal.h"
 
+#include "InstanceDataTypes.generated.h"
+
 class FInstanceSceneDataBuffers;
 class FInstanceDataUpdateTaskInfo;
 class FInstanceDataSceneProxy;
@@ -12,14 +14,21 @@ class FInstanceDataSceneProxy;
  * ID of an instance within a primitive, persistent for the lifetime of the instance. 
  * Can be used as a linear index as the maximum possible value is never higher than the high-water-mark for the number of instances in the component.
  */
+USTRUCT(BlueprintType)
 struct FPrimitiveInstanceId
 {
-	//FPrimitiveInstanceId(int32 InId) : Id(InId) {}
+	GENERATED_BODY()
 	inline bool IsValid() const { return Id != INDEX_NONE; }
 	inline bool operator==(const FPrimitiveInstanceId &Other) const { return Other.Id == Id;}
 
 	inline int32 GetAsIndex() const { return Id; }
-//private:
+
+	friend FArchive& operator<<(FArchive& Ar, FPrimitiveInstanceId& InstanceData)
+	{
+		Ar << InstanceData.Id;
+		return Ar;
+	}
+
 	int32 Id = INDEX_NONE;
 };
 
diff --git a/Engine/Source/Runtime/Engine/Public/InstancedStaticMesh/ISMInstanceDataManager.h b/Engine/Source/Runtime/Engine/Public/InstancedStaticMesh/ISMInstanceDataManager.h
--- a/Engine/Source/Runtime/Engine/Public/InstancedStaticMesh/ISMInstanceDataManager.h
+++ b/Engine/Source/Runtime/Engine/Public/InstancedStaticMesh/ISMInstanceDataManager.h
@@ -174,12 +174,6 @@ public:
 	void ReadCookedRenderData(FArchive& Ar);
 
 private:
-	template <typename TaskLambdaType>
-	static void BeginUpdateTask(FInstanceDataUpdateTaskInfo &InstanceDataUpdateTaskInfo, TaskLambdaType &&TaskLambda, const FInstanceDataBufferHeader &InInstanceDataBufferHeader);
-
-	template <typename TaskLambdaType>
-	void DispatchUpdateTask(bool bUnattached, const FInstanceDataBufferHeader &InstanceDataBufferHeader, TaskLambdaType &&TaskLambda);
-
 	using EChangeFlag = FInstanceAttributeTracker::EFlag;
 
 	template<EChangeFlag Flag>
@@ -191,8 +185,6 @@ private:
 
 	bool HasIdentityMapping() const;
 
-	bool ShouldTrackIds() const;
-
 	void FreeInstanceId(FPrimitiveInstanceId InstanceId);
 
 	void InitChangeSet(const union FChangeDesc &ChangeDesc, const FInstanceUpdateComponentDesc &ComponentData, FISMInstanceUpdateChangeSet &ChangeSet);
diff --git a/Engine/Source/Runtime/Engine/Public/InstancedStaticMesh/ISMInstanceUpdateChangeSet.h b/Engine/Source/Runtime/Engine/Public/InstancedStaticMesh/ISMInstanceUpdateChangeSet.h
--- a/Engine/Source/Runtime/Engine/Public/InstancedStaticMesh/ISMInstanceUpdateChangeSet.h
+++ b/Engine/Source/Runtime/Engine/Public/InstancedStaticMesh/ISMInstanceUpdateChangeSet.h
@@ -4,97 +4,16 @@
 
 #include "InstanceDataSceneProxy.h"
 #include "Containers/StridedView.h"
-#include "InstanceAttributeTracker.h"
-
-class HHitProxy;
-class FOpaqueHitProxyContainer;
+#include "InstanceData/InstanceUpdateChangeSet.h"
 
 /**
- * Helper to make it possible to use the same paths for scattering through the remap...
+ * Extends the change set with ISM specifics.
  */
-class FIdentityDeltaRange
-{
-public:
-	FIdentityDeltaRange(int32 InNum) : NumItems(InNum) {}
-
-	/**
-		*/
-	inline bool IsEmpty() const { return NumItems == 0; }
-
-	/**
-		*/
-	inline bool IsDelta() const { return false; }
-
-	/**
-		* Returns the number of items in this range - i.e., the number of items that need to be copied to collect an update.
-		*/
-	inline int32 GetNumItems() const { return NumItems; }
-
-	/**
-		*/
-	struct FConstIterator
-	{
-		int32 ItemIndex = 0;
-		int32 MaxNum = 0;
-
-		FConstIterator(int32 InIndex, int32 InMaxNum)
-		:	ItemIndex(InIndex),
-			MaxNum(InMaxNum)
-		{
-		}
-
-		void operator++() {  ++ItemIndex; }
-
-		/**
-			* Get the index of the data in the source / destination arrays. 
-			*/
-		int32 GetIndex() const { return ItemIndex; }
-
-		/**
-			* Get the continuous index of the data item in the collected item array.
-			*/
-		int32 GetItemIndex() const { return ItemIndex; }
-
-		explicit operator bool() const {  return ItemIndex < MaxNum; }
-	};
-
-	FConstIterator GetIterator() const
-	{
-		return FConstIterator(0, NumItems);
-	}
-
-private:
-	int32 NumItems = 0;
-};
-
-class FISMInstanceUpdateChangeSet
+class FISMInstanceUpdateChangeSet : public FInstanceUpdateChangeSet
 {
 public:
-	FISMInstanceUpdateChangeSet(bool bInNeedFullUpdate, FInstanceAttributeTracker &&InInstanceAttributeTracker) 
-		: InstanceAttributeTracker(MoveTemp(InInstanceAttributeTracker))
-		, bNeedFullUpdate(bInNeedFullUpdate) 
-	{
-	}
-#if WITH_EDITOR
-	// Set editor data 
-	void SetEditorData(const TArray<TRefCountPtr<HHitProxy>>& HitProxies, const TBitArray<> &SelectedInstances);//, bool bWasHitProxiesReallocated);
-#endif
-
-	template <FInstanceAttributeTracker::EFlag Flag>
-	FInstanceAttributeTracker::FDeltaRange<Flag> GetDelta(bool bForceEmpty, bool bForceFull = false) const 
-	{ 
-		if (bForceEmpty)
-		{
-			return FInstanceAttributeTracker::FDeltaRange<Flag>();
-		}
-		return InstanceAttributeTracker.GetDeltaRange<Flag>(bNeedFullUpdate || bForceFull, NumSourceInstances); 
-	}
-
-
-	FInstanceAttributeTracker::FDeltaRange<FInstanceAttributeTracker::EFlag::TransformChanged> GetTransformDelta() const 
-	{ 
-		return GetDelta<FInstanceAttributeTracker::EFlag::TransformChanged>(false, bUpdateAllInstanceTransforms);
-	}
+	// Upcycle constructors
+	using FInstanceUpdateChangeSet::FInstanceUpdateChangeSet;
 
 	FInstanceAttributeTracker::FDeltaRange<FInstanceAttributeTracker::EFlag::CustomDataChanged> GetCustomDataDelta() const 
 	{ 
@@ -102,11 +21,6 @@ public:
 		return GetDelta<FInstanceAttributeTracker::EFlag::CustomDataChanged>(NumCustomDataFloats == 0 || !Flags.bHasPerInstanceCustomData);
 	}
 
-	FInstanceAttributeTracker::FDeltaRange<FInstanceAttributeTracker::EFlag::IndexChanged> GetIndexChangedDelta() const 
-	{ 
-		return GetDelta<FInstanceAttributeTracker::EFlag::IndexChanged>(false);
-	}
-
 	FIdentityDeltaRange GetInstanceLightShadowUVBiasDelta() const
 	{
 		return FIdentityDeltaRange(InstanceLightShadowUVBias.Num());
@@ -119,8 +33,6 @@ public:
 	}
 #endif
 
-	TFunction<void(TArray<float> &InstanceRandomIDs)> GeneratePerInstanceRandomIds;
-
 	/**
 	 * Add a value, must be done in the order represented in the InstanceLightShadowUVBiasDelta.
 	 */
@@ -140,52 +52,9 @@ public:
 	ENGINE_API void SetInstancePrevTransforms(TArrayView<FMatrix> InPrevInstanceTransforms, const FVector &Offset);
 	ENGINE_API void SetInstancePrevTransforms(TArrayView<FMatrix> InPrevInstanceTransforms);
 	ENGINE_API void SetCustomData(const TArrayView<const float> &InPerInstanceCustomData, int32 InNumCustomDataFloats);
-	ENGINE_API void SetInstanceLocalBounds(const FRenderBounds &Bounds);
-
-	bool IsFullUpdate() const
-	{
-		return bNeedFullUpdate;
-	}
-
-	FInstanceAttributeTracker InstanceAttributeTracker;
-	bool bNeedFullUpdate = false;
-	bool bUpdateAllInstanceTransforms = false;
 	
-	bool bNumCustomDataChanged = false;
-	bool bBakedLightingDataChanged = false;
-	bool bAnyEditorDataChanged = false;
-
-	bool bIdentityIdMap = false;
-	TArray<FPrimitiveInstanceId> IndexToIdMapDeltaData;
-	int32 NumCustomDataFloats = 0;
-
-	TArray<FRenderTransform> Transforms;
-	TArray<FRenderTransform> PrevTransforms;
-
-	TArray<float> PerInstanceCustomData;
-
-	TArray<FVector4f> InstanceLightShadowUVBias;
-	FInstanceDataFlags Flags;
-#if WITH_EDITOR
-	TArray<uint32> InstanceEditorData;
-	TBitArray<> SelectedInstances;
-	TPimplPtr<FOpaqueHitProxyContainer> HitProxyContainer;
-#endif
-	TArray<FRenderBounds, TInlineAllocator<1>> InstanceLocalBounds;
+	UE_DEPRECATED(5.6, "Use SetSharedLocalBounds instead")
+	void SetInstanceLocalBounds(const FRenderBounds &Bounds) { SetSharedLocalBounds(Bounds); }
 	TArray<int32> LegacyInstanceReorderTable;
-	FRenderTransform PrimitiveToRelativeWorld;
-	FVector PrimitiveWorldSpaceOffset;
-	TOptional<FRenderTransform> PreviousPrimitiveToRelativeWorld;
-	float AbsMaxDisplacement = 0.0f;
-	int32 NumSourceInstances = 0;
 	int32 PostUpdateNumInstances = 0;
-	int32 MaxInstanceId = 0;
 };
-
-
-#if WITH_EDITOR
-
-ENGINE_API TPimplPtr<FOpaqueHitProxyContainer> MakeOpaqueHitProxyContainer(const TArray<TRefCountPtr<HHitProxy>>& InHitProxies);
-
-
-#endif
