Fixed fallout from the changes done to support animbank.
 * static analysis errors.
 * Some unintentional difference in ID mapping update for the ISMC proxy, reverted to old code.

diff --git a/Engine/Source/Runtime/Engine/Private/InstanceData/InstanceDataManager.cpp b/Engine/Source/Runtime/Engine/Private/InstanceData/InstanceDataManager.cpp
--- a/Engine/Source/Runtime/Engine/Private/InstanceData/InstanceDataManager.cpp
+++ b/Engine/Source/Runtime/Engine/Private/InstanceData/InstanceDataManager.cpp
@@ -210,28 +210,6 @@ protected:
 #endif
 };
 
-
-#if 1
-namespace RenderingSpatialHash
-{
-	template <typename ScalarType>
-	inline FArchive& operator<<(FArchive& Ar, TLocation<ScalarType>& Item)
-	{
-		Ar << Item.Coord;
-		Ar << Item.Level;
-		return Ar;
-	}
-}
-
-inline FArchive& operator<<(FArchive& Ar, FInstanceSceneDataBuffers::FCompressedSpatialHashItem& Item)
-{
-	Ar << Item.Location;
-	Ar << Item.NumInstances;
-
-	return Ar;
-}
-#endif
-
 FInstanceDataManager::FInstanceDataManager(UPrimitiveComponent* InPrimitiveComponent) 
 	: bNumCustomDataChanged(false)
 	, bPrimitiveTransformChanged(false)
diff --git a/Engine/Source/Runtime/Engine/Private/InstanceData/InstanceDataUpdateUtils.h b/Engine/Source/Runtime/Engine/Private/InstanceData/InstanceDataUpdateUtils.h
--- a/Engine/Source/Runtime/Engine/Private/InstanceData/InstanceDataUpdateUtils.h
+++ b/Engine/Source/Runtime/Engine/Private/InstanceData/InstanceDataUpdateUtils.h
@@ -434,3 +434,22 @@ void DispatchInstanceDataUpdateTask(bool bIsUnattached, const TSharedPtr<ProxyTy
 		});
 	}
 }
+
+namespace RenderingSpatialHash
+{
+	template <typename ScalarType>
+	inline FArchive& operator<<(FArchive& Ar, TLocation<ScalarType>& Item)
+	{
+		Ar << Item.Coord;
+		Ar << Item.Level;
+		return Ar;
+	}
+}
+
+inline FArchive& operator<<(FArchive& Ar, FInstanceSceneDataBuffers::FCompressedSpatialHashItem& Item)
+{
+	Ar << Item.Location;
+	Ar << Item.NumInstances;
+
+	return Ar;
+}
diff --git a/Engine/Source/Runtime/Engine/Private/InstancedStaticMesh/ISMInstanceDataManager.cpp b/Engine/Source/Runtime/Engine/Private/InstancedStaticMesh/ISMInstanceDataManager.cpp
--- a/Engine/Source/Runtime/Engine/Private/InstancedStaticMesh/ISMInstanceDataManager.cpp
+++ b/Engine/Source/Runtime/Engine/Private/InstancedStaticMesh/ISMInstanceDataManager.cpp
@@ -41,25 +41,6 @@ static TAutoConsoleVariable<int32> CVarInstanceDataMinInstanceCountToOptimize(
 	2,
 	TEXT("Minimum number of instances to perform optimized build for (if enabled), can be used to disable optimized build for small ISMs as there is some overhead from doing so."));
 
-namespace RenderingSpatialHash
-{
-	template <typename ScalarType>
-	inline FArchive& operator<<(FArchive& Ar, TLocation<ScalarType>& Item)
-	{
-		Ar << Item.Coord;
-		Ar << Item.Level;
-		return Ar;
-	}
-}
-
-inline FArchive& operator<<(FArchive& Ar, FInstanceSceneDataBuffers::FCompressedSpatialHashItem& Item)
-{
-	Ar << Item.Location;
-	Ar << Item.NumInstances;
-
-	return Ar;
-}
-
 FPrimitiveInstanceDataManager::FPrimitiveInstanceDataManager(UPrimitiveComponent* InPrimitiveComponent) 
 	: bNumCustomDataChanged(false)
 	, bBakedLightingDataChanged(false)
diff --git a/Engine/Source/Runtime/Engine/Private/InstancedStaticMesh/ISMInstanceDataSceneProxy.cpp b/Engine/Source/Runtime/Engine/Private/InstancedStaticMesh/ISMInstanceDataSceneProxy.cpp
--- a/Engine/Source/Runtime/Engine/Private/InstancedStaticMesh/ISMInstanceDataSceneProxy.cpp
+++ b/Engine/Source/Runtime/Engine/Private/InstancedStaticMesh/ISMInstanceDataSceneProxy.cpp
@@ -326,7 +326,40 @@ void FISMCInstanceDataSceneProxy::Build(FISMInstanceUpdateChangeSet&& ChangeSet)
 template <typename IndexRemapType>
 void FISMCInstanceDataSceneProxy::UpdateIdMapping(FISMInstanceUpdateChangeSet& ChangeSet, const IndexRemapType &IndexRemap)
 {
-	::UpdateIdMapping(ChangeSet, IndexRemap, InstanceIdIndexMap);
+	// update mapping, create explicit mapping if needed
+	if (ChangeSet.bIdentityIdMap && IndexRemap.IsIdentity())
+	{
+		// Reset to identity mapping with the new number of instances
+		InstanceIdIndexMap.Reset(ChangeSet.PostUpdateNumInstances);
+	}
+	else
+	{
+		InstanceIdIndexMap.ResizeExplicit(ChangeSet.PostUpdateNumInstances, ChangeSet.MaxInstanceId);
+
+		// If any were removed, we need to clear the associated IDs, before updating (since they may have been added again)
+		for (TConstSetBitIterator<> It(ChangeSet.InstanceAttributeTracker.GetRemovedIterator()); It; ++It)
+		{
+			// There may be more bits set as things that are marked as removed may no longer be in the map
+			if (It.GetIndex() >= InstanceIdIndexMap.GetMaxInstanceId())
+			{
+				break;
+			}
+			InstanceIdIndexMap.SetInvalid(FPrimitiveInstanceId{It.GetIndex()});
+		}
+
+		// Update index mappings (if not identity)
+		auto IndexDelta = ChangeSet.GetIndexChangedDelta();
+		for (auto It = IndexDelta.GetIterator(); It; ++It)
+		{
+			int32 NewInstanceIndex = It.GetIndex();
+			int32 ItemIndex = It.GetItemIndex();
+			
+			IndexRemap.Remap(ItemIndex, NewInstanceIndex);
+
+			FPrimitiveInstanceId InstanceId = ChangeSet.bIdentityIdMap ? FPrimitiveInstanceId{ItemIndex} : ChangeSet.IndexToIdMapDeltaData[ItemIndex];
+			InstanceIdIndexMap.Update(InstanceId, NewInstanceIndex);
+		}
+	}
 }
 
 void FISMCInstanceDataSceneProxy::Update(FISMInstanceUpdateChangeSet&& ChangeSet)
