Render Asset Streaming Optimization: Allow for RenderAssetStreamingManager to process components to add with a ParallelFor to use more cores.
Optimization can be turned on using 'r.Streaming.AllowParallelRenderAssetStreamingManagerProcessComponents'.

diff --git a/Engine/Source/Runtime/Engine/Classes/Engine/TextureStreamingTypes.h b/Engine/Source/Runtime/Engine/Classes/Engine/TextureStreamingTypes.h
--- a/Engine/Source/Runtime/Engine/Classes/Engine/TextureStreamingTypes.h
+++ b/Engine/Source/Runtime/Engine/Classes/Engine/TextureStreamingTypes.h
@@ -259,6 +259,7 @@ class FStreamingTextureLevelContext
 	EMaterialQualityLevel::Type QualityLevel;
 	ERHIFeatureLevel::Type FeatureLevel;
 	bool bIsBuiltDataValid;
+	bool bForceNoUseBuiltData;
 
 	ENGINE_API int32* GetBuildDataIndexRef(UTexture* Texture, bool bForceUpdate = false);
 	ENGINE_API void UpdateQualityAndFeatureLevel(EMaterialQualityLevel::Type InQualityLevel, ERHIFeatureLevel::Type InFeatureLevel, const ULevel* InLevel = nullptr);
@@ -271,6 +272,7 @@ public:
 	ENGINE_API FStreamingTextureLevelContext(EMaterialQualityLevel::Type InQualityLevel, const UPrimitiveComponent* Primitive);
 	ENGINE_API ~FStreamingTextureLevelContext();
 
+	ENGINE_API void SetForceNoUseBuiltData(bool bValue) { bForceNoUseBuiltData = bValue; }
 	ENGINE_API void UpdateContext(EMaterialQualityLevel::Type InQualityLevel, const ULevel* InLevel, const TMap<FGuid, int32>* InTextureGuidToLevelIndex);
 	ENGINE_API void BindBuildData(const TArray<FStreamingTextureBuildInfo>* PreBuiltData);
 	ENGINE_API void ProcessMaterial(const FBoxSphereBounds& ComponentBounds, const FPrimitiveMaterialInfo& MaterialData, float ComponentScaling, TArray<FStreamingRenderAssetPrimitiveInfo>& OutStreamingTextures, bool bIsComponentBuildDataValid = false, const UPrimitiveComponent* DebugComponent = nullptr);
diff --git a/Engine/Source/Runtime/Engine/Private/Components/PrimitiveComponent.cpp b/Engine/Source/Runtime/Engine/Private/Components/PrimitiveComponent.cpp
--- a/Engine/Source/Runtime/Engine/Private/Components/PrimitiveComponent.cpp
+++ b/Engine/Source/Runtime/Engine/Private/Components/PrimitiveComponent.cpp
@@ -515,7 +515,7 @@ void UPrimitiveComponent::GetStreamingRenderAssetInfo(FStreamingTextureLevelCont
 					LevelContext.ProcessMaterial(Bounds, MaterialData, Bounds.SphereRadius, OutStreamingRenderAssets, bIsValidTextureStreamingBuiltData, this);
 				}
 				// Remove all instances of this material in case there were duplicates.
-				UsedMaterials.RemoveSwap(MaterialInterface);
+				UsedMaterials.RemoveSwap(MaterialInterface, EAllowShrinking::No);
 			}
 		}
 	}
@@ -566,7 +566,7 @@ void UPrimitiveComponent::GetStreamingRenderAssetInfoWithNULLRemoval(FStreamingT
 			const FStreamingRenderAssetPrimitiveInfo& Info = OutStreamingRenderAssets[Index];
 			if (!Info.RenderAsset || !Info.RenderAsset->IsStreamable())
 			{
-				OutStreamingRenderAssets.RemoveAtSwap(Index--);
+				OutStreamingRenderAssets.RemoveAtSwap(Index--, EAllowShrinking::No);
 			}
 			else
 			{
@@ -582,7 +582,7 @@ void UPrimitiveComponent::GetStreamingRenderAssetInfoWithNULLRemoval(FStreamingT
 
 				if (!bForceMipStreaming && !bCanBeStreamedByDistance && Info.TexelFactor >= 0.f)
 				{
-					OutStreamingRenderAssets.RemoveAtSwap(Index--);
+					OutStreamingRenderAssets.RemoveAtSwap(Index--, EAllowShrinking::No);
 				}
 			}
 		}
diff --git a/Engine/Source/Runtime/Engine/Private/Materials/MaterialInstance.cpp b/Engine/Source/Runtime/Engine/Private/Materials/MaterialInstance.cpp
--- a/Engine/Source/Runtime/Engine/Private/Materials/MaterialInstance.cpp
+++ b/Engine/Source/Runtime/Engine/Private/Materials/MaterialInstance.cpp
@@ -998,7 +998,7 @@ UMaterialInstance::~UMaterialInstance()
 
 const UMaterial* UMaterialInstance::GetMaterial() const
 {
-	check(IsInGameThread() || IsAsyncLoading());
+	check(IsInGameThread() || IsInParallelGameThread() || IsAsyncLoading());
 	if(GetReentrantFlag())
 	{
 		return UMaterial::GetDefaultMaterial(MD_Surface);
diff --git a/Engine/Source/Runtime/Engine/Private/Materials/MaterialInstanceSupport.h b/Engine/Source/Runtime/Engine/Private/Materials/MaterialInstanceSupport.h
--- a/Engine/Source/Runtime/Engine/Private/Materials/MaterialInstanceSupport.h
+++ b/Engine/Source/Runtime/Engine/Private/Materials/MaterialInstanceSupport.h
@@ -39,15 +39,18 @@ public:
 #else
 	FMICReentranceGuard(const UMaterialInstance* InMaterial)
 	{
-		bIsInGameThread = IsInGameThread();
-		Material = const_cast<UMaterialInstance*>(InMaterial);
-
-		if (Material->GetReentrantFlag(bIsInGameThread) == true)
+		if (!IsInParallelGameThread())
 		{
-			UE_LOG(LogMaterial, Warning, TEXT("InMaterial: %s GameThread: %d RenderThread: %d"), *InMaterial->GetFullName(), IsInGameThread(), IsInRenderingThread());
-			check(!Material->GetReentrantFlag(bIsInGameThread));
+			bIsInGameThread = IsInGameThread();
+			Material = const_cast<UMaterialInstance*>(InMaterial);
+
+			if (Material->GetReentrantFlag(bIsInGameThread) == true)
+			{
+				UE_LOG(LogMaterial, Warning, TEXT("InMaterial: %s GameThread: %d RenderThread: %d"), *InMaterial->GetFullName(), IsInGameThread(), IsInRenderingThread());
+				check(!Material->GetReentrantFlag(bIsInGameThread));
+			}
+			Material->SetReentrantFlag(true, bIsInGameThread);
 		}
-		Material->SetReentrantFlag(true, bIsInGameThread);
 	}
 
 	~FMICReentranceGuard()
diff --git a/Engine/Source/Runtime/Engine/Private/Materials/MaterialShared.cpp b/Engine/Source/Runtime/Engine/Private/Materials/MaterialShared.cpp
--- a/Engine/Source/Runtime/Engine/Private/Materials/MaterialShared.cpp
+++ b/Engine/Source/Runtime/Engine/Private/Materials/MaterialShared.cpp
@@ -1075,7 +1075,7 @@ const FMaterialShaderMap* FMaterial::GetShaderMapToUse() const
 { 
 	const FMaterialShaderMap* ShaderMapToUse = NULL;
 
-	if (IsInGameThread())
+	if (IsInGameThread() || IsInParallelGameThread())
 	{
 		// If we are accessing uniform texture expressions on the game thread, use results from a shader map whose compile is in flight that matches this material
 		// This allows querying what textures a material uses even when it is being asynchronously compiled
diff --git a/Engine/Source/Runtime/Engine/Private/Materials/MaterialUniformExpressions.cpp b/Engine/Source/Runtime/Engine/Private/Materials/MaterialUniformExpressions.cpp
--- a/Engine/Source/Runtime/Engine/Private/Materials/MaterialUniformExpressions.cpp
+++ b/Engine/Source/Runtime/Engine/Private/Materials/MaterialUniformExpressions.cpp
@@ -224,7 +224,7 @@ void FUniformParameterOverrides::SetTextureOverride(EMaterialTextureParameterTyp
 
 UTexture* FUniformParameterOverrides::GetTextureOverride_GameThread(EMaterialTextureParameterType Type, const FMaterialTextureParameterInfo& ParameterInfo) const
 {
-	check(IsInGameThread());
+	check(IsInGameThread() || IsInParallelGameThread());
 	const uint32 TypeIndex = (uint32)Type;
 	const FTextureParameterKey Key{ ParameterInfo.ParameterInfo, ParameterInfo.TextureIndex };
 	UTexture* const* Result = GameThreadTextureOverides[TypeIndex].Find(Key);
@@ -723,7 +723,7 @@ void FUniformExpressionSet::SetVTLayer(int32 StackIndex, int32 VTLayerIndex, int
 
 void FUniformExpressionSet::GetGameThreadTextureValue(EMaterialTextureParameterType Type, int32 Index, const UMaterialInterface* MaterialInterface, const FMaterial& Material, UTexture*& OutValue, bool bAllowOverride) const
 {
-	check(IsInGameThread());
+	check(IsInGameThread() || IsInParallelGameThread());
 	OutValue = NULL;
 	const FMaterialTextureParameterInfo& Parameter = GetTextureParameter(Type, Index);
 #if WITH_EDITOR
diff --git a/Engine/Source/Runtime/Engine/Private/Streaming/DynamicTextureInstanceManager.cpp b/Engine/Source/Runtime/Engine/Private/Streaming/DynamicTextureInstanceManager.cpp
--- a/Engine/Source/Runtime/Engine/Private/Streaming/DynamicTextureInstanceManager.cpp
+++ b/Engine/Source/Runtime/Engine/Private/Streaming/DynamicTextureInstanceManager.cpp
@@ -5,7 +5,9 @@
 =============================================================================*/
 
 #include "Streaming/DynamicTextureInstanceManager.h"
+#include "Streaming/StreamingManagerTexture.h"
 #include "Components/PrimitiveComponent.h"
+#include "Async/ParallelFor.h"
 
 // TODO: Remove this cvar and old code path once effectiveness and correctness is verified
 static int32 GStreamingDeferredRemoveDyanmicInstances = 1;
@@ -79,6 +81,7 @@ void FDynamicRenderAssetInstanceManager::RegisterTasks(RenderAssetInstanceTask::
 
 void FDynamicRenderAssetInstanceManager::IncrementalUpdate(FRemovedRenderAssetArray& RemovedRenderAssets, float Percentage)
 {
+	TRACE_CPUPROFILER_EVENT_SCOPE(FDynamicRenderAssetInstanceManager::IncrementalUpdate);
 	FRenderAssetInstanceState* State = StateSync.SyncAndGetState();
 
 	// First try to apply the pending defrag.
@@ -104,22 +107,57 @@ void FDynamicRenderAssetInstanceManager::IncrementalUpdate(FRemovedRenderAssetAr
 	}
 
 	// Now insert everything, checking for duplicates through bAttachedToStreamingManagerAsDynamic
-	for (const UPrimitiveComponent* Component : PendingComponents)
+	if (PendingComponents.Num())
 	{
-		check(Component);
-
-		// Check bAttachedToStreamingManagerAsDynamic here to prevent handling the component twice.
-		if (!Component->bAttachedToStreamingManagerAsDynamic && CanManage(Component) && Component->IsRegistered() && (!Component->IsRenderStateCreated() || Component->SceneProxy))
+		const int32 ToProcessComponentCount = PendingComponents.Num();
+		if (FRenderAssetStreamingManager::IsParallelProcessComponentsAllowed(ToProcessComponentCount))
 		{
-			FStreamingTextureLevelContext LevelContext(EMaterialQualityLevel::Num, Component);
-			if (State->AddComponentIgnoreBounds(Component, LevelContext) == EAddComponentResult::Success)
+			TArray<FRenderAssetInstanceState::FPreAddComponentIgnoreBoundsPayload> Payloads;
+			Payloads.SetNum(ToProcessComponentCount);
+			ParallelFor(TEXT("PreAddComponentsIgnoreBounds"), ToProcessComponentCount, 1, [&](int32 Index)
+			{
+				FOptionalTaskTagScope Scope(ETaskTag::EParallelGameThread);
+				const UPrimitiveComponent* Component = PendingComponents[Index];
+				check(Component);
+				if (!Component->bAttachedToStreamingManagerAsDynamic && CanManage(Component) && Component->IsRegistered() && (!Component->IsRenderStateCreated() || Component->SceneProxy))
+				{
+					FStreamingTextureLevelContext LevelContext(EMaterialQualityLevel::Num, Component);
+					FRenderAssetInstanceState::PreAddComponentIgnoreBounds(Component, LevelContext, Payloads[Index]);
+				}
+				else
+				{
+					Payloads[Index].Result = EAddComponentResult::Fail;
+				}
+			});
 			{
-				Component->bAttachedToStreamingManagerAsDynamic = true;
+				for (FRenderAssetInstanceState::FPreAddComponentIgnoreBoundsPayload& Payload : Payloads)
+				{
+					if (State->AddComponentIgnoreBounds(Payload) == EAddComponentResult::Success)
+					{
+						Payload.Component->bAttachedToStreamingManagerAsDynamic = true;
+					}
+				}
 			}
 		}
-	}
+		else
+		{
+			for (const UPrimitiveComponent* Component : PendingComponents)
+			{
+				check(Component);
 
-	PendingComponents.Reset();
+				// Check bAttachedToStreamingManagerAsDynamic here to prevent handling the component twice.
+				if (!Component->bAttachedToStreamingManagerAsDynamic && CanManage(Component) && Component->IsRegistered() && (!Component->IsRenderStateCreated() || Component->SceneProxy))
+				{
+					FStreamingTextureLevelContext LevelContext(EMaterialQualityLevel::Num, Component);
+					if (State->AddComponentIgnoreBounds(Component, LevelContext) == EAddComponentResult::Success)
+					{
+						Component->bAttachedToStreamingManagerAsDynamic = true;
+					}
+				}
+			}
+		}
+		PendingComponents.Reset();
+	}
 
 	// Now update the bounds
 	Refresh(Percentage);
diff --git a/Engine/Source/Runtime/Engine/Private/Streaming/LevelTextureManager.cpp b/Engine/Source/Runtime/Engine/Private/Streaming/LevelTextureManager.cpp
--- a/Engine/Source/Runtime/Engine/Private/Streaming/LevelTextureManager.cpp
+++ b/Engine/Source/Runtime/Engine/Private/Streaming/LevelTextureManager.cpp
@@ -11,6 +11,8 @@
 #include "SceneInterface.h"
 #include "Streaming/DynamicTextureInstanceManager.h"
 #include "Streaming/StaticTextureInstanceManager.h"
+#include "Streaming/StreamingManagerTexture.h"
+#include "Async/ParallelFor.h"
 
 FLevelRenderAssetManager::FLevelRenderAssetManager(ULevel* InLevel, RenderAssetInstanceTask::FDoWorkTask& AsyncTask)
 	: Level(InLevel)
@@ -110,6 +112,7 @@ void FLevelRenderAssetManager::SetAsDynamic(FDynamicRenderAssetInstanceManager&
 
 void FLevelRenderAssetManager::IncrementalBuild(FDynamicRenderAssetInstanceManager& DynamicComponentManager, FStreamingTextureLevelContext& LevelContext, bool bForceCompletion, int64& NumStepsLeft)
 {
+	QUICK_SCOPE_CYCLE_COUNTER(FLevelRenderAssetManager_IncrementalBuild);
 	check(Level);
 
 	const float MaxTextureUVDensity = CVarStreamingMaxTextureUVDensity.GetValueOnAnyThread();
@@ -118,6 +121,7 @@ void FLevelRenderAssetManager::IncrementalBuild(FDynamicRenderAssetInstanceManag
 	{
 	case EStaticBuildStep::BuildTextureLookUpMap:
 	{
+		TRACE_CPUPROFILER_EVENT_SCOPE(FLevelRenderAssetManager::IncrementalBuild_BuildTextureLookUpMap);
 		// Build the map to convert from a guid to the level index. 
 		TextureGuidToLevelIndex.Reserve(Level->StreamingTextureGuids.Num());
 		for (int32 TextureIndex = 0; TextureIndex < Level->StreamingTextureGuids.Num(); ++TextureIndex)
@@ -133,6 +137,7 @@ void FLevelRenderAssetManager::IncrementalBuild(FDynamicRenderAssetInstanceManag
 	}
 	case EStaticBuildStep::ProcessActors:
 	{
+		TRACE_CPUPROFILER_EVENT_SCOPE(FLevelRenderAssetManager::IncrementalBuild_ProcessActors);
 		// All actors needs to be processed at once here because of the logic around bStaticComponentsRegisteredInStreamingManager.
 		// All components must have either bHandledByStreamingManagerAsDynamic or bAttachedToStreamingManagerAsStatic set once bStaticComponentsRegisteredInStreamingManager gets set.
 		// If any component gets created after, the logic in UPrimitiveComponent::CreateRenderState_Concurrent() will detect it as a new component and put it through the dynamic path.
@@ -171,25 +176,67 @@ void FLevelRenderAssetManager::IncrementalBuild(FDynamicRenderAssetInstanceManag
 	}
 	case EStaticBuildStep::ProcessComponents:
 	{
+		TRACE_CPUPROFILER_EVENT_SCOPE(FLevelRenderAssetManager::IncrementalBuild_ProcessComponents);
 		const bool bLevelIsVisible = Level->bIsVisible;
 
-		while ((bForceCompletion || NumStepsLeft > 0) && UnprocessedComponents.Num())
+		const int32 UnprocessedComponentCount = UnprocessedComponents.Num();
+		const int32 ToProcessComponentCount = bForceCompletion ? UnprocessedComponentCount : FMath::Min<int32>(NumStepsLeft, UnprocessedComponentCount);
+		if (FRenderAssetStreamingManager::IsParallelProcessComponentsAllowed(ToProcessComponentCount) && !StaticInstances.CanAddComponent())
 		{
-			const UPrimitiveComponent* Primitive = UnprocessedComponents.Pop(EAllowShrinking::No);
-
-			const EAddComponentResult AddResult = StaticInstances.Add(Primitive, LevelContext, MaxTextureUVDensity);
-			if (AddResult == EAddComponentResult::Fail && !bLevelIsVisible)
+			TArray<FRenderAssetInstanceState::FPreAddComponentPayload> Payloads;
+			Payloads.SetNum(ToProcessComponentCount);
+			// Because LevelContext is not thread-safe when ComponentBuildData is valid, we force not to use built data
 			{
-				// Retry once the level becomes visible.
-				PendingComponents.Add(Primitive);
+				LevelContext.SetForceNoUseBuiltData(true);
+				ParallelFor(TEXT("PreAddComponents"), ToProcessComponentCount, 1, [&](int32 Index)
+				{
+					FOptionalTaskTagScope Scope(ETaskTag::EParallelGameThread);
+					const UPrimitiveComponent* Primitive = UnprocessedComponents[UnprocessedComponentCount - Index - 1];
+					FRenderAssetInstanceState::PreAddComponent(Primitive, LevelContext, MaxTextureUVDensity, Payloads[Index]);
+				});
+				LevelContext.SetForceNoUseBuiltData(false);
 			}
-			else if (AddResult != EAddComponentResult::Success)
+
+			TArray<EAddComponentResult> AddResults;
+			StaticInstances.Add(Payloads, AddResults);
+			for (int32 i=0; i < ToProcessComponentCount; ++i)
 			{
-				// Here we also handle the case for Fail_UIDensityConstraint.
-				SetAsDynamic(DynamicComponentManager, LevelContext, Primitive);
+				EAddComponentResult AddResult = AddResults[i];
+				if (AddResult == EAddComponentResult::Fail && !bLevelIsVisible)
+				{
+					// Retry once the level becomes visible.
+					PendingComponents.Add(Payloads[i].Component);
+				}
+				else if (AddResult != EAddComponentResult::Success)
+				{
+					// Here we also handle the case for Fail_UIDensityConstraint.
+					SetAsDynamic(DynamicComponentManager, LevelContext, Payloads[i].Component);
+				}
 			}
 
-			--NumStepsLeft;
+			NumStepsLeft -= ToProcessComponentCount;
+			UnprocessedComponents.SetNum(UnprocessedComponentCount - ToProcessComponentCount, EAllowShrinking::No);
+		}
+		else
+		{
+			while ((bForceCompletion || NumStepsLeft > 0) && UnprocessedComponents.Num())
+			{
+				const UPrimitiveComponent* Primitive = UnprocessedComponents.Pop(EAllowShrinking::No);
+
+				const EAddComponentResult AddResult = StaticInstances.Add(Primitive, LevelContext, MaxTextureUVDensity);
+				if (AddResult == EAddComponentResult::Fail && !bLevelIsVisible)
+				{
+					// Retry once the level becomes visible.
+					PendingComponents.Add(Primitive);
+				}
+				else if (AddResult != EAddComponentResult::Success)
+				{
+					// Here we also handle the case for Fail_UIDensityConstraint.
+					SetAsDynamic(DynamicComponentManager, LevelContext, Primitive);
+				}
+
+				--NumStepsLeft;
+			}
 		}
 
 		if (!UnprocessedComponents.Num())
@@ -201,6 +248,7 @@ void FLevelRenderAssetManager::IncrementalBuild(FDynamicRenderAssetInstanceManag
 	}
 	case EStaticBuildStep::NormalizeLightmapTexelFactors:
 	{
+		TRACE_CPUPROFILER_EVENT_SCOPE(FLevelRenderAssetManager::IncrementalBuild_NormalizeLightmapTexelFactors);
 		// Unfortunately, PendingInsertionStaticPrimtivComponents won't be taken into account here.
 		StaticInstances.NormalizeLightmapTexelFactor();
 		BuildStep = EStaticBuildStep::CompileElements;
@@ -208,6 +256,7 @@ void FLevelRenderAssetManager::IncrementalBuild(FDynamicRenderAssetInstanceManag
 	}
 	case EStaticBuildStep::CompileElements:
 	{
+		TRACE_CPUPROFILER_EVENT_SCOPE(FLevelRenderAssetManager::IncrementalBuild_CompileElements);
 		// Compile elements (to optimize runtime) for what is there.
 		// PendingInsertionStaticPrimitives will be added after.
 		NumStepsLeft -= (int64)StaticInstances.CompileElements();
@@ -217,6 +266,7 @@ void FLevelRenderAssetManager::IncrementalBuild(FDynamicRenderAssetInstanceManag
 	case EStaticBuildStep::WaitForRegistration:
 		if (Level->bIsVisible)
 		{
+			TRACE_CPUPROFILER_EVENT_SCOPE(FLevelRenderAssetManager::IncrementalBuild_WaitForRegistration);
 			// Remove unregistered component and resolve the bounds using the packed relative boxes.
 			TArray<const UPrimitiveComponent*> RemovedPrimitives;
 			NumStepsLeft -= (int64)StaticInstances.CheckRegistrationAndUnpackBounds(RemovedPrimitives);
diff --git a/Engine/Source/Runtime/Engine/Private/Streaming/StaticTextureInstanceManager.cpp b/Engine/Source/Runtime/Engine/Private/Streaming/StaticTextureInstanceManager.cpp
--- a/Engine/Source/Runtime/Engine/Private/Streaming/StaticTextureInstanceManager.cpp
+++ b/Engine/Source/Runtime/Engine/Private/Streaming/StaticTextureInstanceManager.cpp
@@ -81,6 +81,21 @@ void FStaticRenderAssetInstanceManager::Refresh(float Percentage)
 	}
 }
 
+bool FStaticRenderAssetInstanceManager::CanAddComponent() const
+{
+	return !!AsyncView;
+}
+
+EAddComponentResult FStaticRenderAssetInstanceManager::Add(const TArray<FRenderAssetInstanceState::FPreAddComponentPayload>& Payloads, TArray<EAddComponentResult>& OutResults)
+{
+	if (!AsyncView)
+	{
+		FRenderAssetInstanceState* State = StateSync.SyncAndGetState();
+		return State->AddComponents(Payloads, OutResults);
+	}
+	return EAddComponentResult::Fail;
+}
+
 EAddComponentResult FStaticRenderAssetInstanceManager::Add(const UPrimitiveComponent* Component, FStreamingTextureLevelContext& LevelContext, float MaxAllowedUIDensity)
 {
 	if (!AsyncView)
diff --git a/Engine/Source/Runtime/Engine/Private/Streaming/StaticTextureInstanceManager.h b/Engine/Source/Runtime/Engine/Private/Streaming/StaticTextureInstanceManager.h
--- a/Engine/Source/Runtime/Engine/Private/Streaming/StaticTextureInstanceManager.h
+++ b/Engine/Source/Runtime/Engine/Private/Streaming/StaticTextureInstanceManager.h
@@ -51,6 +51,9 @@ public:
 	/** Add a component streaming data, the LevelContext gives support for precompiled data. */
 	EAddComponentResult Add(const UPrimitiveComponent* Component, FStreamingTextureLevelContext& LevelContext, float MaxAllowedUIDensity) final override;
 
+	/** Add a batch of components streaming data. */
+	EAddComponentResult Add(const TArray<FRenderAssetInstanceState::FPreAddComponentPayload>& Payloads, TArray<EAddComponentResult>& OutResults);
+
 	/** Remove a component, the RemoveTextures is the list of textures not referred anymore. */
 	 void Remove(const UPrimitiveComponent* Component, FRemovedRenderAssetArray* RemovedRenderAssets) final override;
 
@@ -60,6 +63,9 @@ public:
 	/** Return a view of the data that has to be 100% thread safe. The content is allowed to be updated, but not memory must be reallocated. */
 	const FRenderAssetInstanceView* GetAsyncView(bool bCreateIfNull) final override;
 
+	/** Returns whether we can add components to the manager. */
+	bool CanAddComponent() const;
+
 	/** Return the size taken for sub-allocation. */
 	uint32 GetAllocatedSize() const final override;
 
diff --git a/Engine/Source/Runtime/Engine/Private/Streaming/StreamingManagerTexture.cpp b/Engine/Source/Runtime/Engine/Private/Streaming/StreamingManagerTexture.cpp
--- a/Engine/Source/Runtime/Engine/Private/Streaming/StreamingManagerTexture.cpp
+++ b/Engine/Source/Runtime/Engine/Private/Streaming/StreamingManagerTexture.cpp
@@ -97,6 +97,17 @@ static TAutoConsoleVariable<int32> CVarProcessAddedRenderAssetsAfterAsyncWork(
 	TEXT("Whether to call ProcessAddedRenderAssets in subsqequent UpdateResourceStreaming stages after Async work has completed."),
 	ECVF_Default);
 
+namespace UE::Private::RenderAssetStreamingManager
+{
+	static bool bAllowParallelProcessComponents = false;
+	static FAutoConsoleVariableRef CVarAllowParallelProcessComponents(
+		TEXT("r.Streaming.AllowParallelRenderAssetStreamingManagerProcessComponents"),
+		bAllowParallelProcessComponents,
+		TEXT("Whether it is allowed for RenderAssetStreamingManager to process components to add with a ParallelFor to use more cores. ")
+		TEXT("When enabled, texture streaming built data is not used to retrieve the primitive streaming render asset information."),
+		ECVF_Default);
+}
+
 bool TrackRenderAsset( const FString& AssetName );
 bool UntrackRenderAsset( const FString& AssetName );
 void ListTrackedRenderAssets( FOutputDevice& Ar, int32 NumTextures );
@@ -414,6 +425,12 @@ int64 FRenderAssetStreamingManager::GetPoolSize() const
 	return GTexturePoolSize;
 }
 
+bool FRenderAssetStreamingManager::IsParallelProcessComponentsAllowed(int32 NumComponentsToProcess)
+{
+	static int32 MinComponentsToProcessForParallelProcessing = 8;
+	return UE::Private::RenderAssetStreamingManager::bAllowParallelProcessComponents && FApp::ShouldUseThreadingForPerformance() && (NumComponentsToProcess >= MinComponentsToProcessForParallelProcessing);
+}
+
 void FRenderAssetStreamingManager::IncrementalUpdate(float Percentage, bool bUpdateDynamicComponents)
 {
 	TRACE_CPUPROFILER_EVENT_SCOPE(FRenderAssetStreamingManager::IncrementalUpdate);
@@ -1306,6 +1323,7 @@ bool FRenderAssetStreamingManager::FastForceFullyResident(UStreamableRenderAsset
  */
 void FRenderAssetStreamingManager::UpdateStreamingRenderAssets(int32 StageIndex, int32 NumUpdateStages, bool bWaitForMipFading, bool bAsync)
 {
+	TRACE_CPUPROFILER_EVENT_SCOPE(UpdateStreamingRenderAssets);
 	if ( StageIndex == 0 )
 	{
 		CurrentUpdateStreamingRenderAssetIndex = 0;
diff --git a/Engine/Source/Runtime/Engine/Private/Streaming/StreamingManagerTexture.h b/Engine/Source/Runtime/Engine/Private/Streaming/StreamingManagerTexture.h
--- a/Engine/Source/Runtime/Engine/Private/Streaming/StreamingManagerTexture.h
+++ b/Engine/Source/Runtime/Engine/Private/Streaming/StreamingManagerTexture.h
@@ -192,6 +192,9 @@ struct FRenderAssetStreamingManager final : public IRenderAssetStreamingManager
 	 */
 	void SetRenderAssetsRemovedTimestamp(const FRemovedRenderAssetArray& RemovedRenderAssets);
 
+	/** Whether it is allowed for this manager to process components to add with a ParallelFor to use more cores. */
+	static bool IsParallelProcessComponentsAllowed(int32 NumComponentsToProcess);
+
 private:
 //BEGIN: Thread-safe functions and data
 		friend class FRenderAssetStreamingMipCalcTask;
diff --git a/Engine/Source/Runtime/Engine/Private/Streaming/TextureInstanceState.cpp b/Engine/Source/Runtime/Engine/Private/Streaming/TextureInstanceState.cpp
--- a/Engine/Source/Runtime/Engine/Private/Streaming/TextureInstanceState.cpp
+++ b/Engine/Source/Runtime/Engine/Private/Streaming/TextureInstanceState.cpp
@@ -10,6 +10,7 @@
 #include "Streaming/TextureInstanceView.inl"
 #include "Components/PrimitiveComponent.h"
 #include "Streaming/TextureInstanceView.h"
+#include "Async/ParallelFor.h"
 
 FRenderAssetInstanceState::FRenderAssetInstanceState(bool bForDynamicInstances)
 	: bIsDynamicInstanceState(bForDynamicInstances)
@@ -250,21 +251,31 @@ FORCEINLINE bool operator<(const FBoxSphereBounds& Lhs, const FBoxSphereBounds&
 	return FMemory::Memcmp(&Lhs, &Rhs, sizeof(FBoxSphereBounds)) < 0;
 }
 
-void FRenderAssetInstanceState::AddRenderAssetElements(const UPrimitiveComponent* Component, const TArrayView<FStreamingRenderAssetPrimitiveInfo>& RenderAssetInstanceInfos, int32 BoundsIndex, int32*& ComponentLink)
+void FRenderAssetInstanceState::AddRenderAssetElements(const UPrimitiveComponent* Component, const TArray<FPreAddRenderAssetElement>& RenderAssetElements, int32 BoundsIndex, int32*& ComponentLink)
+{
+	for (const FPreAddRenderAssetElement& Element : RenderAssetElements)
+	{
+		AddElement(Component, Element.RenderAsset, BoundsIndex, Element.MergedTexelFactor, Component->bForceMipStreaming, ComponentLink);
+	}
+}
+
+namespace UE::Private
+{
+static void ForEachRenderAssetTexelFactorGroup(const TArrayView<FStreamingRenderAssetPrimitiveInfo>& InRenderAssetInstanceInfos, TFunctionRef<void(const UStreamableRenderAsset*, float)> Func)
 {
 	// Loop for each render asset - texel factor group (a group being of same texel factor sign)
-	for (int32 InfoIndex = 0; InfoIndex < RenderAssetInstanceInfos.Num();)
+	for (int32 InfoIndex = 0; InfoIndex < InRenderAssetInstanceInfos.Num();)
 	{
-		const FStreamingRenderAssetPrimitiveInfo& Info = RenderAssetInstanceInfos[InfoIndex];
+		const FStreamingRenderAssetPrimitiveInfo& Info = InRenderAssetInstanceInfos[InfoIndex];
 		float MergedTexelFactor = Info.TexelFactor;
 		int32 NumOfMergedElements = 1;
 
 		// Merge all texel factor >= 0 for the same texture, or all those < 0
 		if (Info.TexelFactor >= 0)
 		{
-			for (int32 NextInfoIndex = InfoIndex + 1; NextInfoIndex < RenderAssetInstanceInfos.Num(); ++NextInfoIndex)
+			for (int32 NextInfoIndex = InfoIndex + 1; NextInfoIndex < InRenderAssetInstanceInfos.Num(); ++NextInfoIndex)
 			{
-				const FStreamingRenderAssetPrimitiveInfo& NextInfo = RenderAssetInstanceInfos[NextInfoIndex];
+				const FStreamingRenderAssetPrimitiveInfo& NextInfo = InRenderAssetInstanceInfos[NextInfoIndex];
 				if (NextInfo.RenderAsset == Info.RenderAsset && NextInfo.TexelFactor >= 0)
 				{
 					MergedTexelFactor = FMath::Max(MergedTexelFactor, NextInfo.TexelFactor);
@@ -278,9 +289,9 @@ void FRenderAssetInstanceState::AddRenderAssetElements(const UPrimitiveComponent
 		}
 		else // Info.TexelFactor < 0
 		{
-			for (int32 NextInfoIndex = InfoIndex + 1; NextInfoIndex < RenderAssetInstanceInfos.Num(); ++NextInfoIndex)
+			for (int32 NextInfoIndex = InfoIndex + 1; NextInfoIndex < InRenderAssetInstanceInfos.Num(); ++NextInfoIndex)
 			{
-				const FStreamingRenderAssetPrimitiveInfo& NextInfo = RenderAssetInstanceInfos[NextInfoIndex];
+				const FStreamingRenderAssetPrimitiveInfo& NextInfo = InRenderAssetInstanceInfos[NextInfoIndex];
 				if (NextInfo.RenderAsset == Info.RenderAsset && NextInfo.TexelFactor < 0)
 				{
 					MergedTexelFactor = FMath::Min(MergedTexelFactor, NextInfo.TexelFactor);
@@ -292,17 +303,83 @@ void FRenderAssetInstanceState::AddRenderAssetElements(const UPrimitiveComponent
 				}
 			}
 		}
-		AddElement(Component, Info.RenderAsset, BoundsIndex, MergedTexelFactor, Component->bForceMipStreaming, ComponentLink);
+		Func(Info.RenderAsset, MergedTexelFactor);
 
 		InfoIndex += NumOfMergedElements;
 	}
 }
 
-EAddComponentResult FRenderAssetInstanceState::AddComponent(const UPrimitiveComponent* Component, FStreamingTextureLevelContext& LevelContext, float MaxAllowedUIDensity)
+static void PreAddRenderAssetElements(const TArrayView<FStreamingRenderAssetPrimitiveInfo>& RenderAssetInstanceInfos, TArray<FRenderAssetInstanceState::FPreAddRenderAssetElement>& OutRenderAssets)
 {
-	check(Component);
+	ForEachRenderAssetTexelFactorGroup(RenderAssetInstanceInfos, [&OutRenderAssets](const UStreamableRenderAsset* RenderAsset, float TexelFactor)
+	{
+		OutRenderAssets.Emplace(RenderAsset, TexelFactor);
+	});
+}
+}
+
+void FRenderAssetInstanceState::AddRenderAssetElements(const UPrimitiveComponent* Component, const TArrayView<FStreamingRenderAssetPrimitiveInfo>& RenderAssetInstanceInfos, int32 BoundsIndex, int32*& ComponentLink)
+{
+	UE::Private::ForEachRenderAssetTexelFactorGroup(RenderAssetInstanceInfos, [this, Component, BoundsIndex, &ComponentLink](const UStreamableRenderAsset* RenderAsset, float TexelFactor)
+	{
+		AddElement(Component, RenderAsset, BoundsIndex, TexelFactor, Component->bForceMipStreaming, ComponentLink);
+	});
+}
+
+EAddComponentResult FRenderAssetInstanceState::AddComponents(const TArray<FRenderAssetInstanceState::FPreAddComponentPayload>& Payloads, TArray<EAddComponentResult>& OutResults)
+{
+	TRACE_CPUPROFILER_EVENT_SCOPE(FRenderAssetInstanceState::AddComponents);
 	checkf(!bIsDynamicInstanceState, TEXT("Error: trying to add component to dynamic instance manager as static"));
+	const int32 PayloadCount = Payloads.Num();
+	OutResults.SetNumUninitialized(PayloadCount);
+	for (int i=0; i<PayloadCount; ++i)
+	{
+		OutResults[i] = AddComponent(Payloads[i]);
+	}
+	return EAddComponentResult::Success;
+}
+
+EAddComponentResult FRenderAssetInstanceState::AddComponent(const FPreAddComponentPayload& Payload)
+{
+	checkf(!bIsDynamicInstanceState, TEXT("Error: trying to add component to dynamic instance manager as static"));
+	if (Payload.Result == EAddComponentResult::Success)
+	{
+		int32* ComponentLink = ComponentMap.Find(Payload.Component);
+		for (const FPreAddBounds& Bounds : Payload.Bounds)
+		{
+			const int32 BoundsIndex = AddBounds(Bounds.InfoBounds, Bounds.PackedRelativeBox, Payload.Component, Payload.Component->GetLastRenderTimeOnScreen(), Bounds.RangeOrigin, Bounds.MinDistanceSq, Bounds.MinRangeSq, Bounds.MaxRangeSq);
+			AddRenderAssetElements(Payload.Component, Bounds.RenderAssets, BoundsIndex, ComponentLink);
+			if (Bounds.bIsPacked)
+			{
+				BoundsToUnpack.Push(BoundsIndex);
+			}
+		}
+	}
+	return Payload.Result;
+}
+
+void FRenderAssetInstanceState::PreAddComponent(const UPrimitiveComponent* Component, FStreamingTextureLevelContext& LevelContext, float MaxAllowedUIDensity, FPreAddComponentPayload& OutPreAddPayload)
+{
+	OutPreAddPayload.Result = AddComponentInternal(Component, LevelContext, MaxAllowedUIDensity, nullptr, &OutPreAddPayload);
+	check(OutPreAddPayload.Component == Component);
+}
+
+EAddComponentResult FRenderAssetInstanceState::AddComponent(const UPrimitiveComponent* Component, FStreamingTextureLevelContext& LevelContext, float MaxAllowedUIDensity)
+{
+	return AddComponentInternal(Component, LevelContext, MaxAllowedUIDensity, this, nullptr);
+}
+
+EAddComponentResult FRenderAssetInstanceState::AddComponentInternal(const UPrimitiveComponent* Component, FStreamingTextureLevelContext& LevelContext, float MaxAllowedUIDensity, FRenderAssetInstanceState* State, FPreAddComponentPayload* OutPreAddPayload)
+{
+	check(State || OutPreAddPayload);
+	checkf(!State || !State->bIsDynamicInstanceState, TEXT("Error: trying to add component to dynamic instance manager as static"));
+
+	if (OutPreAddPayload)
+	{
+		OutPreAddPayload->Component = Component;
+	}
 
+	check(Component);
 	TArray<FStreamingRenderAssetPrimitiveInfo> RenderAssetInstanceInfos;
 	Component->GetStreamingRenderAssetInfoWithNULLRemoval(LevelContext, RenderAssetInstanceInfos);
 
@@ -362,7 +439,7 @@ EAddComponentResult FRenderAssetInstanceState::AddComponent(const UPrimitiveComp
 			return Lhs.PackedRelativeBox < Rhs.PackedRelativeBox;
 		});
 
-		int32* ComponentLink = ComponentMap.Find(Component);
+		int32* ComponentLink = State ? State->ComponentMap.Find(Component) : nullptr;
 
 		// Loop for each bound.
 		for (int32 InfoIndex = 0; InfoIndex < RenderAssetInstanceInfos.Num();)
@@ -375,9 +452,18 @@ EAddComponentResult FRenderAssetInstanceState::AddComponent(const UPrimitiveComp
 				++NumOfBoundReferences;
 			}
 
-			const int32 BoundsIndex = AddBounds(FBoxSphereBounds(ForceInit), Info.PackedRelativeBox, Component, Component->GetLastRenderTimeOnScreen(), FVector(ForceInit), 0, 0, FLT_MAX);
-			AddRenderAssetElements(Component, TArrayView<FStreamingRenderAssetPrimitiveInfo>(RenderAssetInstanceInfos.GetData() + InfoIndex, NumOfBoundReferences), BoundsIndex, ComponentLink);
-			BoundsToUnpack.Push(BoundsIndex);
+			if (State)
+			{
+				const int32 BoundsIndex = State->AddBounds(FBoxSphereBounds(ForceInit), Info.PackedRelativeBox, Component, Component->GetLastRenderTimeOnScreen(), FVector(ForceInit), 0, 0, FLT_MAX);
+				State->AddRenderAssetElements(Component, TArrayView<FStreamingRenderAssetPrimitiveInfo>(RenderAssetInstanceInfos.GetData() + InfoIndex, NumOfBoundReferences), BoundsIndex, ComponentLink);
+				State->BoundsToUnpack.Push(BoundsIndex);
+			}
+			else
+			{
+				check(OutPreAddPayload);
+				FPreAddBounds& OutBounds = OutPreAddPayload->Bounds.Emplace_GetRef(FBoxSphereBounds(ForceInit), Info.PackedRelativeBox, true, FVector(ForceInit), 0, 0, FLT_MAX);
+				UE::Private::PreAddRenderAssetElements(TArrayView<FStreamingRenderAssetPrimitiveInfo>(RenderAssetInstanceInfos.GetData() + InfoIndex, NumOfBoundReferences), OutBounds.RenderAssets);
+			}
 
 			InfoIndex += NumOfBoundReferences;
 		}
@@ -400,7 +486,7 @@ EAddComponentResult FRenderAssetInstanceState::AddComponent(const UPrimitiveComp
 			return Lhs.Bounds < Rhs.Bounds;
 		});
 
-		int32* ComponentLink = ComponentMap.Find(Component);
+		int32* ComponentLink = State ? State->ComponentMap.Find(Component) : nullptr;
 		float MinDistanceSq = 0, MinRangeSq = 0, MaxRangeSq = FLT_MAX;
 
 		// Loop for each bound.
@@ -415,8 +501,17 @@ EAddComponentResult FRenderAssetInstanceState::AddComponent(const UPrimitiveComp
 			}
 
 			GetDistanceAndRange(Component, Info.Bounds, MinDistanceSq, MinRangeSq, MaxRangeSq);
-			const int32 BoundsIndex = AddBounds(Info.Bounds, PackedRelativeBox_Identity, Component, Component->GetLastRenderTimeOnScreen(), Component->Bounds.Origin, MinDistanceSq, MinRangeSq, MaxRangeSq);
-			AddRenderAssetElements(Component, TArrayView<FStreamingRenderAssetPrimitiveInfo>(RenderAssetInstanceInfos.GetData() + InfoIndex, NumOfBoundReferences), BoundsIndex, ComponentLink);
+			if (State)
+			{
+				const int32 BoundsIndex = State->AddBounds(Info.Bounds, PackedRelativeBox_Identity, Component, Component->GetLastRenderTimeOnScreen(), Component->Bounds.Origin, MinDistanceSq, MinRangeSq, MaxRangeSq);
+				State->AddRenderAssetElements(Component, TArrayView<FStreamingRenderAssetPrimitiveInfo>(RenderAssetInstanceInfos.GetData() + InfoIndex, NumOfBoundReferences), BoundsIndex, ComponentLink);
+			}
+			else
+			{
+				check(OutPreAddPayload);
+				FPreAddBounds& OutBounds = OutPreAddPayload->Bounds.Emplace_GetRef(Info.Bounds, PackedRelativeBox_Identity, false, Component->Bounds.Origin, MinDistanceSq, MinRangeSq, MaxRangeSq);
+				UE::Private::PreAddRenderAssetElements(RenderAssetInstanceInfos, OutBounds.RenderAssets);
+			}
 
 			InfoIndex += NumOfBoundReferences;
 		}
@@ -424,10 +519,38 @@ EAddComponentResult FRenderAssetInstanceState::AddComponent(const UPrimitiveComp
 	return EAddComponentResult::Success;
 }
 
+void FRenderAssetInstanceState::PreAddComponentIgnoreBounds(const UPrimitiveComponent* Component, FStreamingTextureLevelContext& LevelContext, FPreAddComponentIgnoreBoundsPayload& OutPayload)
+{
+	OutPayload.Result = AddComponentIgnoreBoundsInternal(Component, LevelContext, nullptr, &OutPayload);
+	check(OutPayload.Component == Component);
+}
+
+EAddComponentResult FRenderAssetInstanceState::AddComponentIgnoreBounds(const FPreAddComponentIgnoreBoundsPayload& Payload)
+{
+	checkf(bIsDynamicInstanceState, TEXT("Error: trying to add component to static instance manager as dynamic"));
+	if (Payload.Result == EAddComponentResult::Success)
+	{
+		int32* ComponentLink = ComponentMap.Find(Payload.Component);
+		AddRenderAssetElements(Payload.Component, Payload.RenderAssets, AddBounds(Payload.Component), ComponentLink);
+	}
+	return Payload.Result;
+}
+
 EAddComponentResult FRenderAssetInstanceState::AddComponentIgnoreBounds(const UPrimitiveComponent* Component, FStreamingTextureLevelContext& LevelContext)
+{
+	return AddComponentIgnoreBoundsInternal(Component, LevelContext, this, nullptr);
+}
+
+EAddComponentResult FRenderAssetInstanceState::AddComponentIgnoreBoundsInternal(const UPrimitiveComponent* Component, FStreamingTextureLevelContext& LevelContext, FRenderAssetInstanceState* State, FPreAddComponentIgnoreBoundsPayload* OutPreAddPayload)
 {
 	check(Component->IsRegistered()); // Must be registered otherwise bounds are invalid.
-	checkf(bIsDynamicInstanceState, TEXT("Error: trying to add component to static instance manager as dynamic"));
+	check(State || OutPreAddPayload);
+	checkf(!State || State->bIsDynamicInstanceState, TEXT("Error: trying to add component to static instance manager as dynamic"));
+
+	if (OutPreAddPayload)
+	{
+		OutPreAddPayload->Component = Component;
+	}
 
 	TArray<FStreamingRenderAssetPrimitiveInfo> RenderAssetInstanceInfos;
 	Component->GetStreamingRenderAssetInfoWithNULLRemoval(LevelContext, RenderAssetInstanceInfos);
@@ -448,8 +571,16 @@ EAddComponentResult FRenderAssetInstanceState::AddComponentIgnoreBounds(const UP
 		return Lhs.RenderAsset < Rhs.RenderAsset;
 	});
 
-	int32* ComponentLink = ComponentMap.Find(Component);
-	AddRenderAssetElements(Component, RenderAssetInstanceInfos, AddBounds(Component), ComponentLink);
+	int32* ComponentLink = State ? State->ComponentMap.Find(Component) : nullptr;
+	if (State)
+	{
+		State->AddRenderAssetElements(Component, RenderAssetInstanceInfos, State->AddBounds(Component), ComponentLink);
+	}
+	else
+	{
+		check(OutPreAddPayload);
+		UE::Private::PreAddRenderAssetElements(RenderAssetInstanceInfos, OutPreAddPayload->RenderAssets);
+	}
 	return EAddComponentResult::Success;
 }
 
@@ -623,6 +754,7 @@ bool FRenderAssetInstanceState::ConditionalUpdateBounds(int32 BoundIndex)
 
 void FRenderAssetInstanceState::UpdateLastRenderTimeAndMaxDrawDistance(int32 BoundIndex)
 {
+	TRACE_CPUPROFILER_EVENT_SCOPE(UpdateLastRenderTimeAndMaxDrawDistance);
 	const UPrimitiveComponent* Component = ensure(Bounds4Components.IsValidIndex(BoundIndex)) ? Bounds4Components[BoundIndex] : nullptr;
 	if (Component)
 	{
diff --git a/Engine/Source/Runtime/Engine/Private/Streaming/TextureInstanceState.h b/Engine/Source/Runtime/Engine/Private/Streaming/TextureInstanceState.h
--- a/Engine/Source/Runtime/Engine/Private/Streaming/TextureInstanceState.h
+++ b/Engine/Source/Runtime/Engine/Private/Streaming/TextureInstanceState.h
@@ -27,13 +27,66 @@ enum class EAddComponentResult : uint8
 class FRenderAssetInstanceState : public FRenderAssetInstanceView
 {
 public:
+	struct FPreAddRenderAssetElement
+	{
+		FPreAddRenderAssetElement(const UStreamableRenderAsset* InRenderAsset, float InMergedTexelFactor)
+			: RenderAsset(InRenderAsset)
+			, MergedTexelFactor(InMergedTexelFactor)
+		{}
+
+		const UStreamableRenderAsset* RenderAsset;
+		float MergedTexelFactor;
+	};
+
+	struct FPreAddBounds
+	{
+		FPreAddBounds(FBoxSphereBounds InInfoBounds, uint32 InPackedRelativeBox, bool bInIsPacked, const FVector& InRangeOrigin, float InMinDistanceSq, float InMinRangeSq, float InMaxRangeSq)
+			: InfoBounds(InInfoBounds)
+			, PackedRelativeBox(InPackedRelativeBox)
+			, RangeOrigin(InRangeOrigin)
+			, MinDistanceSq(InMinDistanceSq)
+			, MinRangeSq(InMinRangeSq)
+			, MaxRangeSq(InMaxRangeSq)
+			, bIsPacked(bInIsPacked)
+		{}
+
+		FBoxSphereBounds InfoBounds;
+		uint32 PackedRelativeBox;
+		FVector RangeOrigin;
+		float MinDistanceSq;
+		float MinRangeSq;
+		float MaxRangeSq;
+		bool bIsPacked;
+		TArray<FPreAddRenderAssetElement> RenderAssets;
+	};
+
+	struct FPreAddComponentPayload
+	{
+		const UPrimitiveComponent* Component;
+		EAddComponentResult Result;
+		TArray<FPreAddBounds> Bounds;
+	};
+
+	struct FPreAddComponentIgnoreBoundsPayload
+	{
+		const UPrimitiveComponent* Component;
+		EAddComponentResult Result;
+		TArray<FPreAddRenderAssetElement> RenderAssets;
+	};
+
 	FRenderAssetInstanceState(bool bForDynamicInstances);
 
+	static void PreAddComponent(const UPrimitiveComponent* Component, FStreamingTextureLevelContext& LevelContext, float MaxAllowedUIDensity, FPreAddComponentPayload& OutPayload);
+	static void PreAddComponentIgnoreBounds(const UPrimitiveComponent* Component, FStreamingTextureLevelContext& LevelContext, FPreAddComponentIgnoreBoundsPayload& OutPayload);
+
 	// Will also remove bounds
 	EAddComponentResult AddComponent(const UPrimitiveComponent* Component, FStreamingTextureLevelContext& LevelContext, float MaxAllowedUIDensity);
+	EAddComponentResult AddComponent(const FPreAddComponentPayload& Payload);
+	EAddComponentResult AddComponents(const TArray<FRenderAssetInstanceState::FPreAddComponentPayload>& Payloads, TArray<EAddComponentResult>& OutResults);
 
 	// Similar to AddComponent, but ignore the streaming data bounds. Used for dynamic components. A faster implementation that does less processing.
 	EAddComponentResult AddComponentIgnoreBounds(const UPrimitiveComponent* Component, FStreamingTextureLevelContext& LevelContext);
+	EAddComponentResult AddComponentIgnoreBounds(const FPreAddComponentIgnoreBoundsPayload& Payload);
 
 	FORCEINLINE bool HasComponentReferences(const UPrimitiveComponent* Component) const { return ComponentMap.Contains(Component); }
 
@@ -82,6 +135,10 @@ private:
 	FORCEINLINE int32 AddBounds(const UPrimitiveComponent* Component);
 	void RemoveBounds(int32 Index);
 
+	static EAddComponentResult AddComponentInternal(const UPrimitiveComponent* Component, FStreamingTextureLevelContext& LevelContext, float MaxAllowedUIDensity, FRenderAssetInstanceState* State, FPreAddComponentPayload* OutPreAddPayload);
+	static EAddComponentResult AddComponentIgnoreBoundsInternal(const UPrimitiveComponent* Component, FStreamingTextureLevelContext& LevelContext, FRenderAssetInstanceState* State, FPreAddComponentIgnoreBoundsPayload* OutPreAddPayload);
+
+	void AddRenderAssetElements(const UPrimitiveComponent* Component, const TArray<FPreAddRenderAssetElement>& RenderAssetElements, int32 BoundsIndex, int32*& ComponentLink);
 	void AddRenderAssetElements(const UPrimitiveComponent* Component, const TArrayView<FStreamingRenderAssetPrimitiveInfo>& RenderAssetInstanceInfos, int32 BoundsIndex, int32*& ComponentLink);
 
 private:
diff --git a/Engine/Source/Runtime/Engine/Private/Streaming/TextureStreamingBuild.cpp b/Engine/Source/Runtime/Engine/Private/Streaming/TextureStreamingBuild.cpp
--- a/Engine/Source/Runtime/Engine/Private/Streaming/TextureStreamingBuild.cpp
+++ b/Engine/Source/Runtime/Engine/Private/Streaming/TextureStreamingBuild.cpp
@@ -450,6 +450,7 @@ FStreamingTextureLevelContext::FStreamingTextureLevelContext(EMaterialQualityLev
 	, bUseRelativeBoxes(false)
 	, BuildDataTimestamp(0)
 	, ComponentBuildData(nullptr)
+	, bForceNoUseBuiltData(false)
 {
 	ULevel* Level = Primitive && Primitive->GetOwner() ? Primitive->GetOwner()->GetLevel() : nullptr;
 	UpdateQualityAndFeatureLevel(InQualityLevel, GMaxRHIFeatureLevel, Level);
@@ -460,6 +461,7 @@ FStreamingTextureLevelContext::FStreamingTextureLevelContext(EMaterialQualityLev
 	, bUseRelativeBoxes(InUseRelativeBoxes)
 	, BuildDataTimestamp(0)
 	, ComponentBuildData(nullptr)
+	, bForceNoUseBuiltData(false)
 {
 	UpdateQualityAndFeatureLevel(InQualityLevel, InFeatureLevel);
 }
@@ -469,12 +471,14 @@ FStreamingTextureLevelContext::FStreamingTextureLevelContext(EMaterialQualityLev
 , bUseRelativeBoxes(false)
 , BuildDataTimestamp(0)
 , ComponentBuildData(nullptr)
+, bForceNoUseBuiltData(false)
 {
 	UpdateContext(InQualityLevel, InLevel, InTextureGuidToLevelIndex);
 }
 
 void FStreamingTextureLevelContext::UpdateContext(EMaterialQualityLevel::Type InQualityLevel, const ULevel* InLevel, const TMap<FGuid, int32>* InTextureGuidToLevelIndex)
 {
+	TRACE_CPUPROFILER_EVENT_SCOPE(FStreamingTextureLevelContext::UpdateContext);
 	TextureGuidToLevelIndex = nullptr;
 	bUseRelativeBoxes = false;
 	BuildDataTimestamp = 0;
@@ -546,7 +550,7 @@ void FStreamingTextureLevelContext::UpdateQualityAndFeatureLevel(EMaterialQualit
 
 bool FStreamingTextureLevelContext::CanUseTextureStreamingBuiltData() const
 {
-	return bIsBuiltDataValid && GUseTextureStreamingBuiltData != 0;
+	return !bForceNoUseBuiltData && bIsBuiltDataValid && GUseTextureStreamingBuiltData != 0;
 }
 
 FStreamingTextureLevelContext::~FStreamingTextureLevelContext()
diff --git a/Engine/Source/Runtime/Engine/Public/Materials/MaterialInstance.h b/Engine/Source/Runtime/Engine/Public/Materials/MaterialInstance.h
--- a/Engine/Source/Runtime/Engine/Public/Materials/MaterialInstance.h
+++ b/Engine/Source/Runtime/Engine/Public/Materials/MaterialInstance.h
@@ -679,16 +679,21 @@ public:
 	FORCEINLINE bool GetReentrantFlag(bool bIsInGameThread = IsInGameThread()) const
 	{
 #if WITH_EDITOR
-		return ReentrantFlag[bIsInGameThread ? 0 : 1];
-#else
-		return false;
+		if (!IsInParallelGameThread())
+		{
+			return ReentrantFlag[bIsInGameThread ? 0 : 1];
+		}
 #endif
+		return false;
 	}
 
 	FORCEINLINE void SetReentrantFlag(const bool bValue, bool bIsInGameThread = IsInGameThread())
 	{
 #if WITH_EDITOR
-		ReentrantFlag[bIsInGameThread ? 0 : 1] = bValue;
+		if (!IsInParallelGameThread())
+		{
+			ReentrantFlag[bIsInGameThread ? 0 : 1] = bValue;
+		}
 #endif
 	}
 
