Fixed some unity build errors

diff --git a/Engine/Source/Runtime/Engine/Private/InstanceData/InstanceDataManager.cpp b/Engine/Source/Runtime/Engine/Private/InstanceData/InstanceDataManager.cpp
--- a/Engine/Source/Runtime/Engine/Private/InstanceData/InstanceDataManager.cpp
+++ b/Engine/Source/Runtime/Engine/Private/InstanceData/InstanceDataManager.cpp
@@ -903,7 +903,7 @@ SIZE_T FInstanceDataManager::GetAllocatedSize() const
 		InstanceUpdateTracker.GetAllocatedSize();
 }
 
-static bool ShouldUsePrecomputed()
+bool FInstanceDataManager::ShouldUsePrecomputed()
 {
 	static const auto CVarPrecomputed = IConsoleManager::Get().FindTConsoleVariableDataInt(TEXT("r.SceneCulling.Precomputed"));
 	static const auto CVarSceneCull = IConsoleManager::Get().FindTConsoleVariableDataInt(TEXT("r.SceneCulling"));
@@ -954,69 +954,6 @@ void FInstanceDataManager::BeginCacheForCookedPlatformData(const ITargetPlatform
 	}	
 }
 
-class FSpatialHashSortBuilder
-{
-public:
-	//int32 FirstLevel;
-	struct FSortedInstanceItem
-	{
-		RenderingSpatialHash::FLocation64 InstanceLoc;
-		int32 InstanceIndex;
-	};
-
-	template <typename GetWorldSpaceInstanceSphereFuncType>
-	void BuildOptimizedSpatialHashOrder(int32 NumInstances, GetWorldSpaceInstanceSphereFuncType &&GetWorldSpaceInstanceSphere)
-	{
-		// TODO: Fix the cvar interaction here, need to unify code path with SceneCulling, or something maybe even come up with some heuristic to allow using fewer bits?
-		int32 FirstLevel = 0;
-		{
-			static const auto CVarInstanceHierarchyMinCellSize = IConsoleManager::Get().FindTConsoleVariableDataFloat(TEXT("r.SceneCulling.MinCellSize"));
-			if (CVarInstanceHierarchyMinCellSize)
-			{
-				// TODO: only one code path to compute this value!!!
-				FirstLevel = RenderingSpatialHash::CalcLevel(CVarInstanceHierarchyMinCellSize->GetValueOnAnyThread() - 1.0);
-			}
-		}
-
-		SortedInstances.Reserve(NumInstances);
-		for (int32 InstanceIndex = 0; InstanceIndex < NumInstances; ++InstanceIndex)
-		{
-			FSphere InstanceWorldSpaceSphere = GetWorldSpaceInstanceSphere(InstanceIndex);
-
-			RenderingSpatialHash::FLocation64 InstanceLoc = RenderingSpatialHash::CalcLevelAndLocationClamped(InstanceWorldSpaceSphere.Center, InstanceWorldSpaceSphere.W, FirstLevel);
-
-			FSortedInstanceItem& Item = SortedInstances.AddDefaulted_GetRef();
-			Item.InstanceLoc = InstanceLoc;
-			Item.InstanceIndex = InstanceIndex;
-		}
-	
-		// Sort the instances according to hash location (first level, then coordinate) and last on instance Index.
-		SortedInstances.Sort(
-			[](const FSortedInstanceItem& A, const FSortedInstanceItem& B) -> bool
-			{
-				if (A.InstanceLoc.Level != B.InstanceLoc.Level)
-				{
-					return A.InstanceLoc.Level < B.InstanceLoc.Level;
-				}
-				if (A.InstanceLoc.Coord.X != B.InstanceLoc.Coord.X)
-				{
-					return A.InstanceLoc.Coord.X < B.InstanceLoc.Coord.X;
-				}
-				if (A.InstanceLoc.Coord.Y != B.InstanceLoc.Coord.Y)
-				{
-					return A.InstanceLoc.Coord.Y < B.InstanceLoc.Coord.Y;
-				}
-				if (A.InstanceLoc.Coord.Z != B.InstanceLoc.Coord.Z)
-				{
-					return A.InstanceLoc.Coord.Z < B.InstanceLoc.Coord.Z;
-				}
-				return A.InstanceIndex < B.InstanceIndex;
-			}
-		);
-	}
-	TArray<FSortedInstanceItem> SortedInstances;
-};
-
 FPrecomputedInstanceSpatialHashData FInstanceDataManager::PrecomputeOptimizationData(FInstanceDataManagerSourceDataDesc&& ComponentData)
 {
 	FInstanceUpdateChangeSet ChangeSet(ComponentData.NumInstances, ComponentData.Flags);
@@ -1033,7 +970,15 @@ FPrecomputedInstanceSpatialHashData FInstanceDataManager::PrecomputeOptimization
 {
 	FSpatialHashSortBuilder SortBuilder;
 
-	SortBuilder.BuildOptimizedSpatialHashOrder(ChangeSet.NumSourceInstances,
+	int32 MinLevel = 0;
+	static const auto CVarInstanceHierarchyMinCellSize = IConsoleManager::Get().FindTConsoleVariableDataFloat(TEXT("r.SceneCulling.MinCellSize"));
+	if (CVarInstanceHierarchyMinCellSize)
+	{
+		// TODO: only one code path to compute this value!!!
+		MinLevel = RenderingSpatialHash::CalcLevel(CVarInstanceHierarchyMinCellSize->GetValueOnAnyThread() - 1.0);
+	}
+
+	SortBuilder.BuildOptimizedSpatialHashOrder(ChangeSet.NumSourceInstances, MinLevel,
 		[&](int32 InstanceIndex) -> FSphere
 		{
 			int32 BoundsIndex = FMath::Min(ChangeSet.InstanceLocalBounds.Num() - 1, InstanceIndex);
diff --git a/Engine/Source/Runtime/Engine/Private/InstanceData/InstanceDataUpdateUtils.h b/Engine/Source/Runtime/Engine/Private/InstanceData/InstanceDataUpdateUtils.h
--- a/Engine/Source/Runtime/Engine/Private/InstanceData/InstanceDataUpdateUtils.h
+++ b/Engine/Source/Runtime/Engine/Private/InstanceData/InstanceDataUpdateUtils.h
@@ -453,3 +453,60 @@ inline FArchive& operator<<(FArchive& Ar, FInstanceSceneDataBuffers::FCompressed
 
 	return Ar;
 }
+
+#if WITH_EDITOR
+
+class FSpatialHashSortBuilder
+{
+public:
+	struct FSortedInstanceItem
+	{
+		RenderingSpatialHash::FLocation64 InstanceLoc;
+		int32 InstanceIndex;
+	};
+
+	template <typename GetWorldSpaceInstanceSphereFuncType>
+	void BuildOptimizedSpatialHashOrder(int32 NumInstances, int32 MinLevel, GetWorldSpaceInstanceSphereFuncType&& GetWorldSpaceInstanceSphere)
+	{
+		int32 FirstLevel = MinLevel;
+
+		SortedInstances.Reserve(NumInstances);
+		for (int32 InstanceIndex = 0; InstanceIndex < NumInstances; ++InstanceIndex)
+		{
+			FSphere InstanceWorldSpaceSphere = GetWorldSpaceInstanceSphere(InstanceIndex);
+
+			RenderingSpatialHash::FLocation64 InstanceLoc = RenderingSpatialHash::CalcLevelAndLocationClamped(InstanceWorldSpaceSphere.Center, InstanceWorldSpaceSphere.W, FirstLevel);
+
+			FSortedInstanceItem& Item = SortedInstances.AddDefaulted_GetRef();
+			Item.InstanceLoc = InstanceLoc;
+			Item.InstanceIndex = InstanceIndex;
+		}
+	
+		// Sort the instances according to hash location (first level, then coordinate) and last on instance Index.
+		SortedInstances.Sort(
+			[](const FSortedInstanceItem& A, const FSortedInstanceItem& B) -> bool
+			{
+				if (A.InstanceLoc.Level != B.InstanceLoc.Level)
+				{
+					return A.InstanceLoc.Level < B.InstanceLoc.Level;
+				}
+				if (A.InstanceLoc.Coord.X != B.InstanceLoc.Coord.X)
+				{
+					return A.InstanceLoc.Coord.X < B.InstanceLoc.Coord.X;
+				}
+				if (A.InstanceLoc.Coord.Y != B.InstanceLoc.Coord.Y)
+				{
+					return A.InstanceLoc.Coord.Y < B.InstanceLoc.Coord.Y;
+				}
+				if (A.InstanceLoc.Coord.Z != B.InstanceLoc.Coord.Z)
+				{
+					return A.InstanceLoc.Coord.Z < B.InstanceLoc.Coord.Z;
+				}
+				return A.InstanceIndex < B.InstanceIndex;
+			}
+		);
+	}
+	TArray<FSortedInstanceItem> SortedInstances;
+};
+
+#endif //WITH_EDITOR
diff --git a/Engine/Source/Runtime/Engine/Private/InstancedStaticMesh/ISMInstanceDataManager.cpp b/Engine/Source/Runtime/Engine/Private/InstancedStaticMesh/ISMInstanceDataManager.cpp
--- a/Engine/Source/Runtime/Engine/Private/InstancedStaticMesh/ISMInstanceDataManager.cpp
+++ b/Engine/Source/Runtime/Engine/Private/InstancedStaticMesh/ISMInstanceDataManager.cpp
@@ -976,7 +976,7 @@ void FPrimitiveInstanceDataManager::OnRegister(int32 InNumInstances)
 	}
 }
 
-static bool ShouldUsePrecomputed()
+bool FPrimitiveInstanceDataManager::ShouldUsePrecomputed()
 {
 	static const auto CVarPrecomputed = IConsoleManager::Get().FindTConsoleVariableDataInt(TEXT("r.SceneCulling.Precomputed"));
 	static const auto CVarSceneCull = IConsoleManager::Get().FindTConsoleVariableDataInt(TEXT("r.SceneCulling"));
@@ -1022,68 +1022,6 @@ void FPrimitiveInstanceDataManager::BeginCacheForCookedPlatformData(const ITarge
 	}	
 }
 
-class FSpatialHashSortBuilder
-{
-public:
-	//int32 FirstLevel;
-	struct FSortedInstanceItem
-	{
-		RenderingSpatialHash::FLocation64 InstanceLoc;
-		int32 InstanceIndex;
-	};
-
-	template <typename GetWorldSpaceInstanceSphereFuncType>
-	void BuildOptimizedSpatialHashOrder(int32 NumInstances, GetWorldSpaceInstanceSphereFuncType &&GetWorldSpaceInstanceSphere)
-	{
-		// TODO: Fix the cvar interaction here, need to unify code path with SceneCulling, or something maybe even come up with some heuristic to allow using fewer bits?
-		int32 FirstLevel = 0;
-		{
-			static const auto CVarInstanceHierarchyMinCellSize = IConsoleManager::Get().FindTConsoleVariableDataFloat(TEXT("r.SceneCulling.MinCellSize"));
-			if (CVarInstanceHierarchyMinCellSize)
-			{
-				// TODO: only one code path to compute this value!!!
-				FirstLevel = RenderingSpatialHash::CalcLevel(CVarInstanceHierarchyMinCellSize->GetValueOnAnyThread() - 1.0);
-			}
-		}
-
-		SortedInstances.Reserve(NumInstances);
-		for (int32 InstanceIndex = 0; InstanceIndex < NumInstances; ++InstanceIndex)
-		{
-			FSphere InstanceWorldSpaceSphere = GetWorldSpaceInstanceSphere(InstanceIndex);
-
-			RenderingSpatialHash::FLocation64 InstanceLoc = RenderingSpatialHash::CalcLevelAndLocationClamped(InstanceWorldSpaceSphere.Center, InstanceWorldSpaceSphere.W, FirstLevel);
-
-			FSortedInstanceItem& Item = SortedInstances.AddDefaulted_GetRef();
-			Item.InstanceLoc = InstanceLoc;
-			Item.InstanceIndex = InstanceIndex;
-		}
-	
-		// Sort the instances according to hash location (first level, then coordinate) and last on instance Index.
-		SortedInstances.Sort(
-			[](const FSortedInstanceItem& A, const FSortedInstanceItem& B) -> bool
-			{
-				if (A.InstanceLoc.Level != B.InstanceLoc.Level)
-				{
-					return A.InstanceLoc.Level < B.InstanceLoc.Level;
-				}
-				if (A.InstanceLoc.Coord.X != B.InstanceLoc.Coord.X)
-				{
-					return A.InstanceLoc.Coord.X < B.InstanceLoc.Coord.X;
-				}
-				if (A.InstanceLoc.Coord.Y != B.InstanceLoc.Coord.Y)
-				{
-					return A.InstanceLoc.Coord.Y < B.InstanceLoc.Coord.Y;
-				}
-				if (A.InstanceLoc.Coord.Z != B.InstanceLoc.Coord.Z)
-				{
-					return A.InstanceLoc.Coord.Z < B.InstanceLoc.Coord.Z;
-				}
-				return A.InstanceIndex < B.InstanceIndex;
-			}
-		);
-	}
-	TArray<FSortedInstanceItem> SortedInstances;
-};
 
 FISMPrecomputedSpatialHashData FPrimitiveInstanceDataManager::PrecomputeOptimizationData(FInstanceUpdateComponentDesc &&ComponentData, TStridedView<FMatrix> InstanceTransforms)
 {
@@ -1091,6 +1029,7 @@ FISMPrecomputedSpatialHashData FPrimitiveInstanceDataManager::PrecomputeOptimiza
 	const float LocalAbsMaxDisplacement = FMath::Max(-ComponentData.PrimitiveMaterialDesc.MinMaxMaterialDisplacement.X, ComponentData.PrimitiveMaterialDesc.MinMaxMaterialDisplacement.Y)
 									+ ComponentData.PrimitiveMaterialDesc.MaxWorldPositionOffsetDisplacement;
 
+
 	const FVector3f PadExtent = FISMCInstanceDataSceneProxy::GetLocalBoundsPadExtent(ComponentData.PrimitiveLocalToWorld, LocalAbsMaxDisplacement);
 	FRenderBounds InstanceLocalBounds = ComponentData.StaticMeshBounds;
 	InstanceLocalBounds.Min -= PadExtent;
@@ -1098,7 +1037,15 @@ FISMPrecomputedSpatialHashData FPrimitiveInstanceDataManager::PrecomputeOptimiza
 
 	FSphere LocalInstanceSphere = InstanceLocalBounds.ToBoxSphereBounds().GetSphere();
 
-	SortBuilder.BuildOptimizedSpatialHashOrder(InstanceTransforms.Num(),
+	int32 MinLevel = 0;
+	static const auto CVarInstanceHierarchyMinCellSize = IConsoleManager::Get().FindTConsoleVariableDataFloat(TEXT("r.SceneCulling.MinCellSize"));
+	if (CVarInstanceHierarchyMinCellSize)
+	{
+		// TODO: only one code path to compute this value!!!
+		MinLevel = RenderingSpatialHash::CalcLevel(CVarInstanceHierarchyMinCellSize->GetValueOnAnyThread() - 1.0);
+	}
+
+	SortBuilder.BuildOptimizedSpatialHashOrder(InstanceTransforms.Num(), MinLevel, 
 		[&](int32 InstanceIndex) -> FSphere
 		{
 			FMatrix InstanceLocalToWorld = InstanceTransforms[InstanceIndex] * ComponentData.PrimitiveLocalToWorld;
diff --git a/Engine/Source/Runtime/Engine/Public/InstanceData/InstanceDataManager.h b/Engine/Source/Runtime/Engine/Public/InstanceData/InstanceDataManager.h
--- a/Engine/Source/Runtime/Engine/Public/InstanceData/InstanceDataManager.h
+++ b/Engine/Source/Runtime/Engine/Public/InstanceData/InstanceDataManager.h
@@ -164,6 +164,7 @@ public:
 	ENGINE_API SIZE_T GetAllocatedSize() const;
 
 protected:
+	static bool ShouldUsePrecomputed();
 
 	void CreateExplicitIdentityMapping();
 
diff --git a/Engine/Source/Runtime/Engine/Public/InstancedStaticMesh/ISMInstanceDataManager.h b/Engine/Source/Runtime/Engine/Public/InstancedStaticMesh/ISMInstanceDataManager.h
--- a/Engine/Source/Runtime/Engine/Public/InstancedStaticMesh/ISMInstanceDataManager.h
+++ b/Engine/Source/Runtime/Engine/Public/InstancedStaticMesh/ISMInstanceDataManager.h
@@ -174,6 +174,8 @@ public:
 	void ReadCookedRenderData(FArchive& Ar);
 
 private:
+	static bool ShouldUsePrecomputed();
+
 	using EChangeFlag = FInstanceAttributeTracker::EFlag;
 
 	template<EChangeFlag Flag>
