Chaos - Fix PT spherecast query

diff --git a/Engine/Source/Runtime/Engine/Private/Collision/SceneQuery.cpp b/Engine/Source/Runtime/Engine/Private/Collision/SceneQuery.cpp
--- a/Engine/Source/Runtime/Engine/Private/Collision/SceneQuery.cpp
+++ b/Engine/Source/Runtime/Engine/Private/Collision/SceneQuery.cpp
@@ -1415,11 +1415,13 @@ namespace Chaos::Private
 
 		if (QueryRadius > UE_SMALL_NUMBER)
 		{
-			return TSceneCastCommon<ESweepOrRay::Sweep, ESingleMultiOrTest::Multi>(World, OutHits, FGeomSQAdditionalInputs(FCollisionShape::MakeSphere(QueryRadius), FQuat::Identity), Start, End, TraceChannel, Params, ResponseParams, ObjectParams, FDefaultAccelContainer{});
+			using PTTraits = TSQTraits<ChaosInterface::FPTSweepHit, ESweepOrRay::Sweep, ESingleMultiOrTest::Multi>;
+			return TSceneCastCommonImp<PTTraits, FGeomSQAdditionalInputs>(World, OutHits, FGeomSQAdditionalInputs(FCollisionShape::MakeSphere(QueryRadius), FQuat::Identity), Start, End, TraceChannel, Params, ResponseParams, ObjectParams, FDefaultAccelContainer{});
 		}
 		else
 		{
-			return TSceneCastCommon<ESweepOrRay::Raycast, ESingleMultiOrTest::Multi>(World, OutHits, FRaycastSQAdditionalInputs(), Start, End, TraceChannel, Params, ResponseParams, ObjectParams, FDefaultAccelContainer{});
+			using PTTraits = TSQTraits<ChaosInterface::FPTRaycastHit, ESweepOrRay::Raycast, ESingleMultiOrTest::Multi>;
+			return TSceneCastCommonImp<PTTraits, FRaycastSQAdditionalInputs>(World, OutHits, FRaycastSQAdditionalInputs(), Start, End, TraceChannel, Params, ResponseParams, ObjectParams, FDefaultAccelContainer{});
 		}
 	};
 
