Level Streaming : Detect unsupported re-entrance for UWorld::UpdateLevelStreaming.
This doesn't fix the JIRA but it will detect the issue earlier.

diff --git a/Engine/Source/Runtime/Engine/Private/World.cpp b/Engine/Source/Runtime/Engine/Private/World.cpp
--- a/Engine/Source/Runtime/Engine/Private/World.cpp
+++ b/Engine/Source/Runtime/Engine/Private/World.cpp
@@ -4472,6 +4472,10 @@ void UWorld::UpdateLevelStreaming()
 	CSV_SCOPED_TIMING_STAT_EXCLUSIVE(UpdateLevelStreaming);
 	LLM_SCOPE(ELLMTag::LoadMapMisc);
 
+	static bool bIsReentrant = false;
+	UE_CLOG(bIsReentrant, LogWorld, Fatal, TEXT("Attempted to call UWorld::UpdateLevelStreaming while already inside UWorld::UpdateLevelStreaming."));
+	TGuardValue<bool> ReentrantGuard(bIsReentrant, true);
+
 	// Reset counters
 	GAddToWorldTimeCumul = 0.0;
 	GRemoveFromWorldUnregisterComponentTimeCumul = 0.0;
