[VisualLogger] improved use of multithread access detectors to detect scenarios where flush is called manually and can conflict with worker thread adding new entries
- todo: refactor methods exposing FVisualLogEntry* to user code (e.g., CheckVisualLogInputInternal)

diff --git a/Engine/Plugins/Runtime/MassGameplay/Source/MassSpawner/Private/MassEntityEQSSpawnPointsGenerator.cpp b/Engine/Plugins/Runtime/MassGameplay/Source/MassSpawner/Private/MassEntityEQSSpawnPointsGenerator.cpp
--- a/Engine/Plugins/Runtime/MassGameplay/Source/MassSpawner/Private/MassEntityEQSSpawnPointsGenerator.cpp
+++ b/Engine/Plugins/Runtime/MassGameplay/Source/MassSpawner/Private/MassEntityEQSSpawnPointsGenerator.cpp
@@ -87,24 +87,24 @@ void UMassEntityEQSSpawnPointsGenerator::OnEQSQueryFinished(TSharedPtr<FEnvQuery
 	UE_VLOG(this, LogMassSpawner, Log, TEXT("Spawning at %d locations"), LocationIndex);
 	if (GetDefault<UMassGameplaySettings>()->bLogSpawnLocations)
 	{
-		if (FVisualLogEntry* LogEntry = FVisualLogger::Get().GetLastEntryForObject(this))
-		{
-			FVisualLogShapeElement Element(TEXT(""), FColor::Orange, /*Thickness*/20, LogMassSpawner.GetCategoryName());
-
-			Element.Points.Reserve(LocationIndex);
-			for (const FMassEntitySpawnDataGeneratorResult& Result : Results)
+		FVisualLogger::Get().ExecuteOnLastEntryForObject(this, [LocationIndex, &Results](FVisualLogEntry& LogEntry)
 			{
-				const FMassTransformsSpawnData& Transforms = Result.SpawnData.Get<FMassTransformsSpawnData>();
-				for (int i = 0; i < Result.NumEntities; i++)
+				FVisualLogShapeElement Element(TEXT(""), FColor::Orange, /*Thickness*/20, LogMassSpawner.GetCategoryName());
+
+				Element.Points.Reserve(LocationIndex);
+				for (const FMassEntitySpawnDataGeneratorResult& Result : Results)
 				{
-					Element.Points.Add(Transforms.Transforms[i].GetLocation());
+					const FMassTransformsSpawnData& Transforms = Result.SpawnData.Get<FMassTransformsSpawnData>();
+					for (int i = 0; i < Result.NumEntities; i++)
+					{
+						Element.Points.Add(Transforms.Transforms[i].GetLocation());
+					}
 				}
-			}
-			
-			Element.Type = EVisualLoggerShapeElement::SinglePoint;
-			Element.Verbosity = ELogVerbosity::Display;
-			LogEntry->AddElement(Element);
-		}
+
+				Element.Type = EVisualLoggerShapeElement::SinglePoint;
+				Element.Verbosity = ELogVerbosity::Display;
+				LogEntry.AddElement(Element);
+			});
 	}
 #endif // ENABLE_VISUAL_LOG
 
diff --git a/Engine/Plugins/Runtime/MassGameplay/Source/MassSpawner/Private/MassEntityZoneGraphSpawnPointsGenerator.cpp b/Engine/Plugins/Runtime/MassGameplay/Source/MassSpawner/Private/MassEntityZoneGraphSpawnPointsGenerator.cpp
--- a/Engine/Plugins/Runtime/MassGameplay/Source/MassSpawner/Private/MassEntityZoneGraphSpawnPointsGenerator.cpp
+++ b/Engine/Plugins/Runtime/MassGameplay/Source/MassSpawner/Private/MassEntityZoneGraphSpawnPointsGenerator.cpp
@@ -91,7 +91,7 @@ void UMassEntityZoneGraphSpawnPointsGenerator::Generate(UObject& QueryOwner, TCo
 	UE_VLOG(this, LogMassSpawner, Log, TEXT("Spawning at %d locations"), LocationIndex);
 	if (GetDefault<UMassGameplaySettings>()->bLogSpawnLocations)
 	{
-		if (FVisualLogEntry* LogEntry = FVisualLogger::Get().GetLastEntryForObject(this))
+		FVisualLogger::Get().ExecuteOnLastEntryForObject(this, [LocationIndex, &Results](FVisualLogEntry& LogEntry)
 		{
 			FVisualLogShapeElement Element(TEXT(""), FColor::Orange, /*Thickness*/20, LogMassSpawner.GetCategoryName());
 
@@ -107,8 +107,8 @@ void UMassEntityZoneGraphSpawnPointsGenerator::Generate(UObject& QueryOwner, TCo
 			
 			Element.Type = EVisualLoggerShapeElement::SinglePoint;
 			Element.Verbosity = ELogVerbosity::Display;
-			LogEntry->AddElement(Element);
-		}
+			LogEntry.AddElement(Element);
+		});
 	}
 #endif // ENABLE_VISUAL_LOG
 
diff --git a/Engine/Source/Runtime/Engine/Private/VisualLogger/VisualLogger.cpp b/Engine/Source/Runtime/Engine/Private/VisualLogger/VisualLogger.cpp
--- a/Engine/Source/Runtime/Engine/Private/VisualLogger/VisualLogger.cpp
+++ b/Engine/Source/Runtime/Engine/Private/VisualLogger/VisualLogger.cpp
@@ -333,19 +333,21 @@ void FVisualLogger::AddObjectToAllowList(const UObject& InObject)
 
 	if (!bIsAlreadyInSet)
 	{
-		FVisualLogEntry* CurrentEntry = CurrentEntryPerObject.Find(&InObject);
-		if (CurrentEntry)
 		{
-			CurrentEntry->SetPassedObjectAllowList(true);
+			UE_MT_SCOPED_WRITE_ACCESS(CurrentEntryPerObject.EntriesAccessDetector);
+			if (FVisualLogEntry* CurrentEntry = CurrentEntryPerObject.Find(&InObject))
+			{
+				CurrentEntry->SetPassedObjectAllowList(true);
+			}
 		}
 
-	    // We are not locking the thread entry map, we do not expect other threads adding new entries while editing allow list
-	    UE_MT_SCOPED_READ_ACCESS(CurrentEntryMapsAccessDetector);
+		// We are not locking the thread entry map, we do not expect other threads adding new entries while editing allow list
+		UE_MT_SCOPED_READ_ACCESS(ThreadLocalMapListAccessDetector);
 		for (FVisualLoggerObjectEntryMap* ThreadCurrentEntryMap : ThreadCurrentEntryMaps)
 		{
 			checkf(ThreadCurrentEntryMap, TEXT("Unexpected null map entry pointer"));
-			FVisualLogEntry* ThreadCurrentEntry = ThreadCurrentEntryMap->Find(&InObject);
-			if (ThreadCurrentEntry)
+			UE_MT_SCOPED_WRITE_ACCESS(ThreadCurrentEntryMap->EntriesAccessDetector);
+			if (FVisualLogEntry* ThreadCurrentEntry = ThreadCurrentEntryMap->Find(&InObject))
 			{
 				ThreadCurrentEntry->SetPassedObjectAllowList(true);
 			}
@@ -358,19 +360,21 @@ void FVisualLogger::ClearObjectAllowList()
 	FWriteScopeLock Lock(EntryRWLock);
 	for (FObjectKey It : ObjectAllowList)
 	{
-		FVisualLogEntry* CurrentEntry = CurrentEntryPerObject.Find(It);
-		if (CurrentEntry)
 		{
-			CurrentEntry->SetPassedObjectAllowList(false);
+			UE_MT_SCOPED_WRITE_ACCESS(CurrentEntryPerObject.EntriesAccessDetector);
+			if (FVisualLogEntry* CurrentEntry = CurrentEntryPerObject.Find(It))
+			{
+				CurrentEntry->SetPassedObjectAllowList(false);
+			}
 		}
 
-	    // We are not locking the thread entry map, we do not expect other threads adding new entries while editing allow list
-	    UE_MT_SCOPED_READ_ACCESS(CurrentEntryMapsAccessDetector);
+		// We are not locking the thread entry map, we do not expect other threads adding new entries while editing allow list
+		UE_MT_SCOPED_READ_ACCESS(ThreadLocalMapListAccessDetector);
 		for (FVisualLoggerObjectEntryMap* ThreadCurrentEntryMap : ThreadCurrentEntryMaps)
 		{
 			checkf(ThreadCurrentEntryMap, TEXT("Unexpected null map entry pointer"));
-			FVisualLogEntry* ThreadCurrentEntry = ThreadCurrentEntryMap->Find(It);
-			if (ThreadCurrentEntry)
+			UE_MT_SCOPED_WRITE_ACCESS(ThreadCurrentEntryMap->EntriesAccessDetector);
+			if (FVisualLogEntry* ThreadCurrentEntry = ThreadCurrentEntryMap->Find(It))
 			{
 				ThreadCurrentEntry->SetPassedObjectAllowList(false);
 			}
@@ -386,19 +390,36 @@ bool FVisualLogger::IsObjectAllowed(const UObject* InObject) const
 }
 
 FVisualLogEntry* FVisualLogger::GetLastEntryForObject(const UObject* Object)
+{
+	FVisualLogEntry* FoundEntry = nullptr;
+	ExecuteOnLastEntryForObject(Object, [&FoundEntry](FVisualLogEntry& LogEntry)
+	{
+		FoundEntry = &LogEntry;
+	});
+
+	return FoundEntry;
+}
+
+bool FVisualLogger::ExecuteOnLastEntryForObject(const UObject* Object, TFunctionRef<void(FVisualLogEntry&)> Function)
 {
 	const UObject* LogOwner = nullptr;
 	{
 		FTransactionallySafeReadScopeLock Lock(RedirectRWLock);
 		LogOwner = FindRedirectionInternal(Object);
 	}
-	if (LogOwner == nullptr)
+
+	if (LogOwner != nullptr)
 	{
-		return nullptr;
+		FVisualLoggerObjectEntryMap& ThreadCurrentEntryMap = GetThreadCurrentEntryMap();
+		UE_MT_SCOPED_WRITE_ACCESS(ThreadCurrentEntryMap.EntriesAccessDetector);
+		if (FVisualLogEntry* LogEntry = ThreadCurrentEntryMap.Find(LogOwner))
+		{
+			Function(*LogEntry);
+			return true;
+		}
 	}
 
-	FVisualLoggerObjectEntryMap& ThreadCurrentEntryPerObject = GetThreadCurrentEntryMap();
-	return ThreadCurrentEntryPerObject.Find(LogOwner);
+	return false;
 }
 
 FVisualLogEntry* FVisualLogger::GetEntryToWrite(const UObject* Object, const double TimeStamp, ECreateIfNeeded ShouldCreate)
@@ -414,6 +435,8 @@ FVisualLogEntry* FVisualLogger::GetEntryToWrite(const UObject* Object, const dou
 	}
 
 	FVisualLoggerObjectEntryMap& ThreadCurrentEntryPerObject = GetThreadCurrentEntryMap();
+	// Using the Write access since it might flush and reset an entry and/or initialize one.
+	UE_MT_SCOPED_WRITE_ACCESS(ThreadCurrentEntryPerObject.EntriesAccessDetector);
 	FVisualLogEntry* ThreadCurrentEntry = ThreadCurrentEntryPerObject.Find(LogOwner);
 	bool bInitializeEntry = false;
 	if (ThreadCurrentEntry)
@@ -423,6 +446,7 @@ FVisualLogEntry* FVisualLogger::GetEntryToWrite(const UObject* Object, const dou
 			if (ThreadCurrentEntry->ShouldFlush(TimeStamp))
 			{
 				FWriteScopeLock Lock(EntryRWLock);
+				UE_MT_SCOPED_WRITE_ACCESS(CurrentEntryPerObject.EntriesAccessDetector);
 				if (FVisualLogEntry* GlobalCurrentEntry = GetEntryToWriteInternal(LogOwner, ThreadCurrentEntry->TimeStamp, ShouldCreate))
 				{
 					ThreadCurrentEntry->MoveTo(*GlobalCurrentEntry);
@@ -439,6 +463,7 @@ FVisualLogEntry* FVisualLogger::GetEntryToWrite(const UObject* Object, const dou
 		if (ShouldCreate != ECreateIfNeeded::Create)
 		{
 			FReadScopeLock Lock(EntryRWLock);
+			UE_MT_SCOPED_READ_ACCESS(CurrentEntryPerObject.EntriesAccessDetector);
 			if (CurrentEntryPerObject.Find(LogOwner))
 			{
 				ShouldCreate = ECreateIfNeeded::Create;
@@ -473,7 +498,6 @@ FVisualLogEntry* FVisualLogger::GetEntryToWrite(const UObject* Object, const dou
 	return ThreadCurrentEntry;
 }
 
-
 FVisualLogEntry* FVisualLogger::GetEntryToWriteInternal(const UObject* Object, const double TimeStamp, const ECreateIfNeeded ShouldCreate)
 {
 	using namespace UE::VisLog::Private;
@@ -487,8 +511,10 @@ FVisualLogEntry* FVisualLogger::GetEntryToWriteInternal(const UObject* Object, c
 
 	// Entry can be created or reused (after being flushed) and will need to be initialized
 	bool bInitializeEntry = false;
-	FVisualLogEntry* CurrentEntry = CurrentEntryPerObject.Find(LogOwner);
 
+	// Using the Write access since it might flush and reset an entry and/or initialize one.
+	UE_MT_SCOPED_WRITE_ACCESS(CurrentEntryPerObject.EntriesAccessDetector);
+	FVisualLogEntry* CurrentEntry = CurrentEntry = CurrentEntryPerObject.Find(LogOwner);
 	if (CurrentEntry)
 	{
 		if (CurrentEntry->ShouldLog(ShouldCreate))
@@ -588,7 +614,7 @@ FVisualLogEntry* FVisualLogger::GetEntryToWriteInternal(const UObject* Object, c
 	return CurrentEntry;
 }
 
-void FVisualLogger::CalculateEntryAllowLogging(FVisualLogEntry* CurrentEntry, const UObject* LogOwner, const UObject* Object)
+void FVisualLogger::CalculateEntryAllowLogging(FVisualLogEntry* CurrentEntry, const UObject* LogOwner, const UObject* Object) const
 {
 	// IsClassAllowed isn't super fast, but this gets calculated only once for every 
 	// object trying to log something
@@ -614,11 +640,12 @@ void FVisualLogger::Tick(float DeltaTime)
 
 void FVisualLogger::FlushThreadsEntries()
 {
-	// We are not locking the thread entry map, we do not expect other threads adding new entries while flushing
-	UE_MT_SCOPED_READ_ACCESS(CurrentEntryMapsAccessDetector);
+	UE_MT_SCOPED_READ_ACCESS(ThreadLocalMapListAccessDetector);
+	UE_MT_SCOPED_WRITE_ACCESS(CurrentEntryPerObject.EntriesAccessDetector);
 	for (FVisualLoggerObjectEntryMap* ThreadCurrentEntryMap : ThreadCurrentEntryMaps)
 	{
 		checkf(ThreadCurrentEntryMap, TEXT("Unexpected null map entry pointer"));
+		UE_MT_SCOPED_READ_ACCESS(ThreadCurrentEntryMap->EntriesAccessDetector);
 		for (auto& ThreadCurrentEntry : *ThreadCurrentEntryMap)
 		{
 			if (ThreadCurrentEntry.Value.bIsInitialized)
@@ -639,6 +666,8 @@ void FVisualLogger::Flush()
 
 	FlushThreadsEntries();
 
+	// Using the Write access since it might flush and reset entries.
+	UE_MT_SCOPED_WRITE_ACCESS(CurrentEntryPerObject.EntriesAccessDetector);
 	for (auto &CurrentEntry : CurrentEntryPerObject)
 	{
 		if (CurrentEntry.Value.bIsInitialized)
@@ -667,7 +696,7 @@ FVisualLogger::FVisualLoggerObjectEntryMap& FVisualLogger::GetThreadCurrentEntry
 	{
 		GThreadCurrentEntryPerObject = new FVisualLoggerObjectEntryMap();
 		FWriteScopeLock Lock(EntryRWLock);
-		UE_MT_SCOPED_WRITE_ACCESS(CurrentEntryMapsAccessDetector);
+		UE_MT_SCOPED_WRITE_ACCESS(ThreadLocalMapListAccessDetector);
 		ThreadCurrentEntryMaps.Add(GThreadCurrentEntryPerObject);
 	}
 	checkf(GThreadCurrentEntryPerObject, TEXT("Unexpected null map entry pointer"));
@@ -869,50 +898,53 @@ void FVisualLogger::Cleanup(UWorld* OldWorld, const bool bReleaseMemory)
 		// perform cleanup only if provided world is valid and was registered
 		if (WorldToRedirectionMap.Remove(OldWorld))
 		{
-		    if (WorldToRedirectionMap.Num() == 0)
-            {
-                WorldToRedirectionMap.Reset();
-                ObjectToWorldMap.Reset();
-                ChildToOwnerMap.Reset();
-                CurrentEntryPerObject.Reset();
-
-				UE_MT_SCOPED_WRITE_ACCESS(CurrentEntryMapsAccessDetector);
+			if (WorldToRedirectionMap.Num() == 0)
+			{
+				WorldToRedirectionMap.Reset();
+				ObjectToWorldMap.Reset();
+				ChildToOwnerMap.Reset();
+
+				UE_MT_SCOPED_WRITE_ACCESS(CurrentEntryPerObject.EntriesAccessDetector);
+				CurrentEntryPerObject.Reset();
+
 				for (FVisualLoggerObjectEntryMap* ThreadCurrentEntryMap : ThreadCurrentEntryMaps)
 				{
 					checkf(ThreadCurrentEntryMap, TEXT("Unexpected null map entry pointer"));
+					UE_MT_SCOPED_WRITE_ACCESS(ThreadCurrentEntryMap->EntriesAccessDetector);
 					ThreadCurrentEntryMap->Reset();
 				}
 				ObjectToNamesMap.Reset();
-            }
-            else
-            {
-                for (auto It = ObjectToWorldMap.CreateIterator(); It; ++It)
-                {
-                    if (It.Value() == OldWorld)
-                    {
+			}
+			else
+			{
+				UE_MT_SCOPED_WRITE_ACCESS(CurrentEntryPerObject.EntriesAccessDetector);
+				for (auto It = ObjectToWorldMap.CreateIterator(); It; ++It)
+				{
+					if (It.Value() == OldWorld)
+					{
 						FObjectKey Obj = It.Key();
-                        ObjectToWorldMap.Remove(Obj);
-                        CurrentEntryPerObject.Remove(Obj);
+						ObjectToWorldMap.Remove(Obj);
+						CurrentEntryPerObject.Remove(Obj);
 
-						UE_MT_SCOPED_WRITE_ACCESS(CurrentEntryMapsAccessDetector);
 						for (FVisualLoggerObjectEntryMap* ThreadCurrentEntryMap : ThreadCurrentEntryMaps)
 						{
 							checkf(ThreadCurrentEntryMap, TEXT("Unexpected null map entry pointer"));
+							UE_MT_SCOPED_WRITE_ACCESS(ThreadCurrentEntryMap->EntriesAccessDetector);
 							ThreadCurrentEntryMap->Remove(Obj);
 						}
 						ObjectToNamesMap.Remove(Obj);
-                    }
-                }
+					}
+				}
 
-                for (FChildToOwnerRedirectionMap::TIterator It = ChildToOwnerMap.CreateIterator(); It; ++It)
-                {
+				for (FChildToOwnerRedirectionMap::TIterator It = ChildToOwnerMap.CreateIterator(); It; ++It)
+				{
 					UObject* Object = It->Key.ResolveObjectPtrEvenIfGarbage();
 					if (Object == nullptr || Object->GetWorld() == OldWorld)
-                    {
-                        It.RemoveCurrent();
-                    }
-                }
-            }
+					{
+						It.RemoveCurrent();
+					}
+				}
+			}
 		}
 	}
 	else
@@ -920,12 +952,14 @@ void FVisualLogger::Cleanup(UWorld* OldWorld, const bool bReleaseMemory)
 		WorldToRedirectionMap.Reset();
 		ObjectToWorldMap.Reset();
 		ChildToOwnerMap.Reset();
+
+		UE_MT_SCOPED_WRITE_ACCESS(CurrentEntryPerObject.EntriesAccessDetector);
 		CurrentEntryPerObject.Reset();
 
-		UE_MT_SCOPED_WRITE_ACCESS(CurrentEntryMapsAccessDetector);
 		for (FVisualLoggerObjectEntryMap* ThreadCurrentEntryMap : ThreadCurrentEntryMaps)
 		{
 			checkf(ThreadCurrentEntryMap, TEXT("Unexpected null map entry pointer"));
+			UE_MT_SCOPED_WRITE_ACCESS(ThreadCurrentEntryMap->EntriesAccessDetector);
 			ThreadCurrentEntryMap->Reset();
 		}
 		ObjectToNamesMap.Reset();
diff --git a/Engine/Source/Runtime/Engine/Public/VisualLogger/VisualLogger.h b/Engine/Source/Runtime/Engine/Public/VisualLogger/VisualLogger.h
--- a/Engine/Source/Runtime/Engine/Public/VisualLogger/VisualLogger.h
+++ b/Engine/Source/Runtime/Engine/Public/VisualLogger/VisualLogger.h
@@ -805,9 +805,21 @@ public:
 	 * @param LogCategory - The LogCategory we are logging about.  This function will only return a valid log entry if the visual logging is enabled for the category.
 	 */
 	[[nodiscard]] static ENGINE_API FVisualLogEntry* GetEntryToWrite(const UObject* LogOwner, const FLogCategoryBase& LogCategory);
+
 	/** Retrieves last used entry for given UObject
 	 *	@note this function can return null */
+	UE_DEPRECATED(5.6, "This method will no longer be exposed, use ExecuteOnLastEntryForObject instead.")
 	ENGINE_API FVisualLogEntry* GetLastEntryForObject(const UObject* Object);
+
+	/**
+	 * Executes the provided function using the last used entry for given UObject if any.
+	 * Method provides validation for thread safe access when adding information to an existing LogEntry.
+	 * @param Object The object to find a LogEntry for.
+	 * @param Function Function to execute if a valid view can be created.
+	 * @return Whether a LogEntry was found and the function was called
+	 */
+	ENGINE_API bool ExecuteOnLastEntryForObject(const UObject* Object, TFunctionRef<void(FVisualLogEntry&)> Function);
+
 	/** flush and serialize data if timestamp allows it */
 	ENGINE_API virtual void Flush() override;
 	/** Moves all threads entries into the global entry map */
@@ -850,15 +862,19 @@ public:
 	ENGINE_API void ClearObjectAllowList();
 	ENGINE_API bool IsObjectAllowed(const UObject* InObject) const;
 
-	typedef TMap<FObjectKey, FVisualLogEntry> FVisualLoggerObjectEntryMap;
+	struct FVisualLoggerObjectEntryMap : TMap<FObjectKey, FVisualLogEntry>
+	{
+		// Multithread access detector for map entries (adding/removing to/from the map or modifying an entry)
+		UE_MT_DECLARE_RW_RECURSIVE_ACCESS_DETECTOR(EntriesAccessDetector);
+	};
 
 private:
 	ENGINE_API FVisualLogger();
-	virtual void Serialize(const TCHAR* V, ELogVerbosity::Type Verbosity, const FName& Category) override { ensureMsgf(0, TEXT("Regular serialize is forbiden for visual logs")); }
+	virtual void Serialize(const TCHAR* V, ELogVerbosity::Type Verbosity, const FName& Category) override { ensureMsgf(0, TEXT("Regular serialize is forbidden for visual logs")); }
 
 	ENGINE_API FVisualLoggerObjectEntryMap& GetThreadCurrentEntryMap();
 
-	/** Get global entry to write where all logs are combined together, not thread safe */
+	/** Get global entry to write where all logs are combined, not thread safe */
 	ENGINE_API FVisualLogEntry* GetEntryToWriteInternal(const UObject* Object, double TimeStamp, ECreateIfNeeded ShouldCreate);
 
 	/** Redirect internal implementation, not thread safe */
@@ -868,7 +884,7 @@ private:
 	/** Cleanup invalid redirects */
 	ENGINE_API void CleanupRedirects();
 	/** Figure out all conditions if this entry is allowed to log */
-	ENGINE_API void CalculateEntryAllowLogging(FVisualLogEntry* CurrentEntry, const UObject* LogOwner, const UObject* Object);
+	ENGINE_API void CalculateEntryAllowLogging(FVisualLogEntry* CurrentEntry, const UObject* LogOwner, const UObject* Object) const;
 
 protected:
 	/** Flushes entries recorded in the frame */
@@ -912,8 +928,8 @@ protected:
 	// Threads current entry maps
 	TArray<FVisualLoggerObjectEntryMap*> ThreadCurrentEntryMaps;
 
-	// Multithread access detector for ThreadCurrentEntryMaps operations (adding/removing map to/from the list or in the map itself)
-	UE_MT_DECLARE_RW_ACCESS_DETECTOR(CurrentEntryMapsAccessDetector);
+	// Multithread access detector for the list of thread local FVisualLoggerObjectEntryMap. Each map has its own MT detector for its entries.
+	UE_MT_DECLARE_RW_ACCESS_DETECTOR(ThreadLocalMapListAccessDetector);
 
 	// Read Write lock protecting object entries
 	mutable FRWLock EntryRWLock;
diff --git a/Engine/Source/Runtime/NavigationSystem/Private/NavMesh/PImplRecastNavMesh.cpp b/Engine/Source/Runtime/NavigationSystem/Private/NavMesh/PImplRecastNavMesh.cpp
--- a/Engine/Source/Runtime/NavigationSystem/Private/NavMesh/PImplRecastNavMesh.cpp
+++ b/Engine/Source/Runtime/NavigationSystem/Private/NavMesh/PImplRecastNavMesh.cpp
@@ -1255,11 +1255,10 @@ ENavigationQueryResult::Type FPImplRecastNavMesh::PostProcessPathInternal(dtStat
 	if (dtStatusDetail(FindPathStatus, DT_INVALID_CYCLE_PATH))
 	{
 		UE_VLOG(NavMeshOwner, LogNavigation, Error, TEXT("FPImplRecastNavMesh::FindPath resulted in a cyclic path!"));
-		FVisualLogEntry* Entry = FVisualLogger::Get().GetLastEntryForObject(NavMeshOwner);
-		if (Entry)
+		FVisualLogger::Get().ExecuteOnLastEntryForObject(NavMeshOwner, [&Path](FVisualLogEntry& LogEntry)
 		{
-			Path.DescribeSelfToVisLog(Entry);
-		}
+			Path.DescribeSelfToVisLog(&LogEntry);
+		});
 	}
 #endif // ENABLE_VISUAL_LOG
 
