Fixed CIS Issue 689807

diff --git a/Engine/Source/Runtime/Engine/Private/VisualLogger/VisualLogger.cpp b/Engine/Source/Runtime/Engine/Private/VisualLogger/VisualLogger.cpp
--- a/Engine/Source/Runtime/Engine/Private/VisualLogger/VisualLogger.cpp
+++ b/Engine/Source/Runtime/Engine/Private/VisualLogger/VisualLogger.cpp
@@ -514,7 +514,7 @@ FVisualLogEntry* FVisualLogger::GetEntryToWriteInternal(const UObject* Object, c
 
 	// Using the Write access since it might flush and reset an entry and/or initialize one.
 	UE_MT_SCOPED_WRITE_ACCESS(CurrentEntryPerObject.EntriesAccessDetector);
-	FVisualLogEntry* CurrentEntry = CurrentEntry = CurrentEntryPerObject.Find(LogOwner);
+	FVisualLogEntry* CurrentEntry = CurrentEntryPerObject.Find(LogOwner);
 	if (CurrentEntry)
 	{
 		if (CurrentEntry->ShouldLog(ShouldCreate))
