Allow UMaterialExpressionSign to be exported to other modules

diff --git a/Engine/Source/Runtime/Engine/Public/Materials/MaterialExpressionSign.h b/Engine/Source/Runtime/Engine/Public/Materials/MaterialExpressionSign.h
--- a/Engine/Source/Runtime/Engine/Public/Materials/MaterialExpressionSign.h
+++ b/Engine/Source/Runtime/Engine/Public/Materials/MaterialExpressionSign.h
@@ -8,7 +8,7 @@
 #include "Materials/MaterialExpression.h"
 #include "MaterialExpressionSign.generated.h"
 
-UCLASS(collapsecategories, hidecategories=Object)
+UCLASS(MinimalAPI, collapsecategories, hidecategories=Object)
 class UMaterialExpressionSign : public UMaterialExpression
 {
 	GENERATED_UCLASS_BODY()
