Fix a crash in some of the editor adhoc serialization tests I was running locally.

diff --git a/Engine/Source/Runtime/Engine/Private/Camera/CameraActor.cpp b/Engine/Source/Runtime/Engine/Private/Camera/CameraActor.cpp
--- a/Engine/Source/Runtime/Engine/Private/Camera/CameraActor.cpp
+++ b/Engine/Source/Runtime/Engine/Private/Camera/CameraActor.cpp
@@ -57,8 +57,8 @@ void ACameraActor::Serialize(FArchive& Ar)
 void ACameraActor::PostLoadSubobjects(FObjectInstancingGraph* OuterInstanceGraph)
 {
 	USceneComponent* OldRoot = RootComponent;
-	USceneComponent* OldAttachParent = OldRoot->GetAttachParent();
-	const FName OldSocketName = OldRoot->GetAttachSocketName();
+	USceneComponent* OldAttachParent = OldRoot ? OldRoot->GetAttachParent() : nullptr;
+	const FName OldSocketName = OldRoot ? OldRoot->GetAttachSocketName() : FName();
 
 	Super::PostLoadSubobjects(OuterInstanceGraph);
 	
