Fix indeterminism in FPerQualityLevelProperty : when being serialized save the state of the object in the harvest phase. Then request a post save serialization. Finally restore the state of the object.

New structure FSavedData used to store the state of FPerQualityLevelProperty object.
FPerQualityLevelProperty :
- add new member variable FSavedData to store the state of the object.
- operator<< : during the harvest phase, save the state of the object and request a post save serialization. During PostSave phase restore the state.

StreamableRenderAsset : revert cl 37375599. It was  workaround for the interminism bug.

#rnx

diff --git a/Engine/Source/Runtime/Engine/Classes/Engine/StreamableRenderAsset.h b/Engine/Source/Runtime/Engine/Classes/Engine/StreamableRenderAsset.h
--- a/Engine/Source/Runtime/Engine/Classes/Engine/StreamableRenderAsset.h
+++ b/Engine/Source/Runtime/Engine/Classes/Engine/StreamableRenderAsset.h
@@ -42,10 +42,6 @@ public:
 	/** Destructor */
 	ENGINE_API virtual ~UStreamableRenderAsset();
 
-#if WITH_EDITOR
-	ENGINE_API virtual void PreSave(FObjectPreSaveContext SaveContext);
-#endif
-
 	/** Get an integer representation of the LOD group */
 	virtual int32 GetLODGroupForStreaming() const
 	{
diff --git a/Engine/Source/Runtime/Engine/Private/PerQualityLevelProperties.cpp b/Engine/Source/Runtime/Engine/Private/PerQualityLevelProperties.cpp
--- a/Engine/Source/Runtime/Engine/Private/PerQualityLevelProperties.cpp
+++ b/Engine/Source/Runtime/Engine/Private/PerQualityLevelProperties.cpp
@@ -9,6 +9,7 @@
 #include "PlatformInfo.h"
 #include "DeviceProfiles/DeviceProfile.h"
 #include "DeviceProfiles/DeviceProfileManager.h"
+#include "UObject/ObjectSaveContext.h"
 #endif
 
 
@@ -174,6 +175,20 @@ namespace QualityLevelProperty
 			return *CachedMappingQualitLevelInfo;
 		}
 	}
+
+	template<typename StructType, typename ValueType>
+	void SaveQualityLevel(const StructType& Property, FSavedData<ValueType>& OutSavedData)
+	{
+		OutSavedData.Default = Property.Default;
+		OutSavedData.PerQuality = Property.PerQuality;
+	}
+
+	template<typename StructType, typename ValueType>
+	void RestoreQualityLevel(const FSavedData<ValueType>& SavedProperty, StructType& OutProperty)
+	{
+		OutProperty.Default = SavedProperty.Default;
+		OutProperty.PerQuality = SavedProperty.PerQuality;
+	}
 #endif
 }
 
@@ -428,8 +443,24 @@ ENGINE_API FArchive& operator<<(FArchive& Ar, FPerQualityLevelProperty<StructTyp
 	if (Ar.IsCooking())
 	{
 		bCooked = true;
-		const FDataDrivenPlatformInfo& PlatformInfo = Ar.CookingTarget()->GetPlatformInfo();
-		This->StripQualityLevelForCooking(*(PlatformInfo.IniPlatformName.ToString()));
+
+		FObjectSavePackageSerializeContext& ObjectSavePackageSerializeContext = Ar.GetSavePackageData()->SavePackageContext;
+		EObjectSaveContextPhase CurrentSavePhase = ObjectSavePackageSerializeContext.GetPhase();
+		if (CurrentSavePhase == EObjectSaveContextPhase::Harvest)
+		{
+			// Save data so they can be restored in the PostSave phase.
+			QualityLevelProperty::SaveQualityLevel(*This, This->SavedValue);
+
+			const FDataDrivenPlatformInfo& PlatformInfo = Ar.CookingTarget()->GetPlatformInfo();
+			This->StripQualityLevelForCooking(*(PlatformInfo.IniPlatformName.ToString()));
+
+			// Request post save serialization on this object so its state can be restored.
+			ObjectSavePackageSerializeContext.RequestPostSaveSerialization();
+		}
+		else if (CurrentSavePhase == EObjectSaveContextPhase::PostSave)
+		{
+			QualityLevelProperty::RestoreQualityLevel(This->SavedValue, *This);
+		}
 	}
 #endif
 	{
diff --git a/Engine/Source/Runtime/Engine/Private/StreamableRenderAsset.cpp b/Engine/Source/Runtime/Engine/Private/StreamableRenderAsset.cpp
--- a/Engine/Source/Runtime/Engine/Private/StreamableRenderAsset.cpp
+++ b/Engine/Source/Runtime/Engine/Private/StreamableRenderAsset.cpp
@@ -2,11 +2,9 @@
 
 #include "Engine/StreamableRenderAsset.h"
 #include "Containers/Ticker.h"
-#include "Interfaces/ITargetPlatform.h"
 #include "RenderAssetUpdate.h"
 #include "Rendering/NaniteCoarseMeshStreamingManager.h"
 #include "RenderingThread.h"
-#include "UObject/ObjectSaveContext.h"
 #include "UObject/Package.h"
 #include "UObject/UnrealType.h"
 
@@ -42,33 +40,6 @@ UStreamableRenderAsset::UStreamableRenderAsset(const FObjectInitializer& ObjectI
 
 UStreamableRenderAsset::~UStreamableRenderAsset() = default;
 
-#if WITH_EDITOR
-void UStreamableRenderAsset::PreSave(FObjectPreSaveContext SaveContext)
-{
-	Super::PreSave(SaveContext);
-	// Workaround cook indeterminism in FPerQualityLevelProperty. It writes to its data during cooking
-	// save serialization, by calling StripQualityLevelForCooking. But this function can make the struct
-	// equal to its archetype, and delta-serialization after that then skips serializing the property fname.
-	// This makes a second in-memory save of its package different than the first in-memory save. In some of
-	// our processes (-diffonly cooks and -iterativevalidate cooks) we need the first and second in-memory
-	// saves to be identical.
-	// We don't yet have a way to remove the modification of Property in the save path of
-	// operator<<(FArchive& Ar, FPerQualityLevelProperty<...>& Property).
-	// Other objects within the savepackage call can depend on the modified values, which is surprising since
-	// the values modified by StripQualityLevelForCooking do not change the value read for the targetplatform.
-	// We need to find those dependencies and remove them.
-	// In the meantime, we workaround the problem by triggering the call to StripQualityLevelForCooking before
-	// the first operator<< save path, here in PreSave. We can not do this for all occurrences of the
-	// struct unfortunately because Unreal's savepackage system does not support a PreSave call for structs.
-	// So we currently do it only for the one class using it that is currently causing a problem in our engine tests.
-	if (SaveContext.IsCooking())
-	{
-		const FDataDrivenPlatformInfo& PlatformInfo = SaveContext.GetTargetPlatform()->GetPlatformInfo();
-		NoRefStreamingLODBias.StripQualityLevelForCooking(*(PlatformInfo.IniPlatformName.ToString()));
-	}
-}
-#endif
-
 void UStreamableRenderAsset::RegisterMipLevelChangeCallback(UPrimitiveComponent* Component, int32 LODIndex, float TimeoutSecs, bool bOnStreamIn, FLODStreamingCallback&& Callback)
 {
 	check(IsInGameThread());
diff --git a/Engine/Source/Runtime/Engine/Public/PerQualityLevelProperties.h b/Engine/Source/Runtime/Engine/Public/PerQualityLevelProperties.h
--- a/Engine/Source/Runtime/Engine/Public/PerQualityLevelProperties.h
+++ b/Engine/Source/Runtime/Engine/Public/PerQualityLevelProperties.h
@@ -64,6 +64,13 @@ namespace QualityLevelProperty
 		return ConvertQualityLevelData(Data);
 	}
 
+	template <typename ValueType>
+	struct FSavedData
+	{
+		ValueType Default = 0;
+		TMap<int32, ValueType> PerQuality;
+	};
+
 #if WITH_EDITOR
 	ENGINE_API TArray<FName> GetEnginePlatformsForPlatformOrGroupName(const FString& InPlatformName);
 	ENGINE_API FSupportedQualityLevelArray PerPlatformOverrideMapping(FString& InPlatformName, UObject* RequestingAsset = nullptr);
@@ -196,6 +203,19 @@ struct FPerQualityLevelProperty
 	FString ScalabilitySection;
 #endif
 	FString CVarName;
+
+private:
+
+#if WITH_EDITOR
+	
+	/* Make this operator friend so it can access private member SavedValue. */
+	template<typename __StructType, typename __ValueType, EName __BasePropertyName>
+	friend ENGINE_API FArchive& operator<<(FArchive& Ar, FPerQualityLevelProperty<__StructType, __ValueType, __BasePropertyName>& Property);
+
+	/* Store values here during harvesting so they can be restored during PostSave. */
+	QualityLevelProperty::FSavedData<ValueType> SavedValue;
+#endif
+
 };
 
 template< typename _StructType, typename _ValueType, EName _BasePropertyName>
