- Fixed getini crash due to new ^ command

diff --git a/Engine/Source/Runtime/Engine/Private/UnrealEngine.cpp b/Engine/Source/Runtime/Engine/Private/UnrealEngine.cpp
--- a/Engine/Source/Runtime/Engine/Private/UnrealEngine.cpp
+++ b/Engine/Source/Runtime/Engine/Private/UnrealEngine.cpp
@@ -10232,6 +10232,7 @@ bool UEngine::HandleGetIniCommand(const TCHAR* Cmd, FOutputDevice& Ar)
 			TEXT(" ADD"),
 			TEXT("UNIQ"),
 			TEXT("REMV"),
+			TEXT("EMPT"),
 			TEXT(" CLR"),
 			TEXT(" AOS"),
 			TEXT(" POC"),
