World Partition - Fix crash caused by changing the type of an existing partition with matching HLOD partitions

diff --git a/Engine/Source/Runtime/Engine/Private/WorldPartition/RuntimeHashSet/WorldPartitionRuntimeHashSet.cpp b/Engine/Source/Runtime/Engine/Private/WorldPartition/RuntimeHashSet/WorldPartitionRuntimeHashSet.cpp
--- a/Engine/Source/Runtime/Engine/Private/WorldPartition/RuntimeHashSet/WorldPartitionRuntimeHashSet.cpp
+++ b/Engine/Source/Runtime/Engine/Private/WorldPartition/RuntimeHashSet/WorldPartitionRuntimeHashSet.cpp
@@ -941,6 +941,7 @@ void UWorldPartitionRuntimeHashSet::PostEditChangeChainProperty(FPropertyChanged
 		
 		RuntimePartitionDesc.MainLayer = nullptr;
 
+		// Recreate the partition with the new class
 		if (RuntimePartitionDesc.Class)
 		{
 			RuntimePartitionDesc.Name = *FString::Printf(TEXT("%s_%d"), *RuntimePartitionDesc.Class->GetName(), RuntimePartitionIndex);
@@ -948,6 +949,12 @@ void UWorldPartitionRuntimeHashSet::PostEditChangeChainProperty(FPropertyChanged
 			RuntimePartitionDesc.MainLayer->SetDefaultValues();
 			RuntimePartitionDesc.MainLayer->Name = RuntimePartitionDesc.Name;
 		}
+		
+		// Recreate the HLOD partitions as they must use the same class
+		if (RuntimePartitionDesc.MainLayer)
+		{
+			FixupHLODSetup(RuntimePartitionDesc);
+		}
 	}
 	else if (PropertyName == GET_MEMBER_NAME_CHECKED(FRuntimePartitionDesc, Name))
 	{
@@ -1113,15 +1120,33 @@ void UWorldPartitionRuntimeHashSet::PostLoad()
 {
 	Super::PostLoad();
 
-	// Make sure that HLOD partitions settings are updated to match the main layer settings
-	for (const FRuntimePartitionDesc& RuntimePartition : RuntimePartitions)
+	// Perform fixup for potentially wrong HLOD setups
+	for (FRuntimePartitionDesc& RuntimePartition : RuntimePartitions)
 	{
-		for (const FRuntimePartitionHLODSetup& HLODSetup : RuntimePartition.HLODSetups)
+		FixupHLODSetup(RuntimePartition);
+	}
+}
+
+void UWorldPartitionRuntimeHashSet::FixupHLODSetup(FRuntimePartitionDesc& RuntimePartition)
+{
+	for (int32 HLODSetupIndex = 0; HLODSetupIndex < RuntimePartition.HLODSetups.Num(); HLODSetupIndex++)
+	{
+		FRuntimePartitionHLODSetup& HLODSetup = RuntimePartition.HLODSetups[HLODSetupIndex];
+		if (HLODSetup.PartitionLayer)
 		{
-			if (HLODSetup.PartitionLayer)
+			// If needed, recreate the HLOD partition to use the same class as the main partition
+			if (HLODSetup.PartitionLayer->GetClass() != RuntimePartition.MainLayer->GetClass())
 			{
-				HLODSetup.PartitionLayer->UpdateHLODRuntimePartitionFrom(RuntimePartition.MainLayer);
+				if (HLODSetup.bIsSpatiallyLoaded)
+				{
+					URuntimePartition* ParentRuntimePartition = HLODSetupIndex != 0 ? RuntimePartition.HLODSetups[HLODSetupIndex - 1].PartitionLayer : RuntimePartition.MainLayer;
+					HLODSetup.PartitionLayer = ParentRuntimePartition->CreateHLODRuntimePartition(HLODSetupIndex);
+					HLODSetup.PartitionLayer->Name = HLODSetup.Name;
+				}
 			}
+
+			// Make sure that HLOD partitions settings are updated to match the main layer settings
+			HLODSetup.PartitionLayer->UpdateHLODRuntimePartitionFrom(RuntimePartition.MainLayer);
 		}
 	}
 }
diff --git a/Engine/Source/Runtime/Engine/Public/WorldPartition/RuntimeHashSet/WorldPartitionRuntimeHashSet.h b/Engine/Source/Runtime/Engine/Public/WorldPartition/RuntimeHashSet/WorldPartitionRuntimeHashSet.h
--- a/Engine/Source/Runtime/Engine/Public/WorldPartition/RuntimeHashSet/WorldPartitionRuntimeHashSet.h
+++ b/Engine/Source/Runtime/Engine/Public/WorldPartition/RuntimeHashSet/WorldPartitionRuntimeHashSet.h
@@ -224,6 +224,8 @@ private:
 	const URuntimePartition* ResolveRuntimePartition(FName GridName, bool bMainPartitionLayer = false) const;
 	const URuntimePartition* ResolveRuntimePartitionForHLODLayer(FName GridName, const FSoftObjectPath& HLODLayerPath) const;
 
+	void FixupHLODSetup(FRuntimePartitionDesc& RuntimePartition);
+
 	/** Array of runtime partition descriptors */
 	UPROPERTY(EditAnywhere, Category = RuntimeSettings, Meta = (TitleProperty = "Name"))
 	TArray<FRuntimePartitionDesc> RuntimePartitions;
