Iris - Silence warnings when running ReplicationSystemLowLevelTests.

#rnx

diff --git a/Engine/Plugins/Runtime/ReplicationSystemTestPlugin/Source/Private/Tests/ReplicationSystem/TestDirtyNetObjectTracker.cpp b/Engine/Plugins/Runtime/ReplicationSystemTestPlugin/Source/Private/Tests/ReplicationSystem/TestDirtyNetObjectTracker.cpp
--- a/Engine/Plugins/Runtime/ReplicationSystemTestPlugin/Source/Private/Tests/ReplicationSystem/TestDirtyNetObjectTracker.cpp
+++ b/Engine/Plugins/Runtime/ReplicationSystemTestPlugin/Source/Private/Tests/ReplicationSystem/TestDirtyNetObjectTracker.cpp
@@ -2,6 +2,7 @@
 
 #include "NetworkAutomationTest.h"
 #include "NetworkAutomationTestMacros.h"
+#include "Logging/LogScopedVerbosityOverride.h"
 #include "Iris/ReplicationSystem/DirtyNetObjectTracker.h"
 #include "Iris/ReplicationSystem/ReplicationSystemInternal.h"
 #include "Tests/ReplicationSystem/ReplicationSystemTestFixture.h"
@@ -52,6 +53,8 @@ UE_NET_TEST_FIXTURE(FDirtyNetObjectTrackerTestFixture, TestNoObjectIsDirtyFromSt
 
 UE_NET_TEST_FIXTURE(FDirtyNetObjectTrackerTestFixture, CannotMarkInvalidObjectAsDirty)
 {
+	LOG_SCOPE_VERBOSITY_OVERRIDE(LogIrisDirtyTracker, ELogVerbosity::Error);
+
 	MarkNetObjectStateDirty(ReplicationSystemId, FNetRefHandleManager::InvalidInternalIndex);
 	MarkNetObjectStateDirty(ReplicationSystemId, NetObjectIndexRangeEnd + 1);
 
diff --git a/Engine/Source/Runtime/Engine/Private/Net/Iris/ReplicationSystem/NetActorFactory.cpp b/Engine/Source/Runtime/Engine/Private/Net/Iris/ReplicationSystem/NetActorFactory.cpp
--- a/Engine/Source/Runtime/Engine/Private/Net/Iris/ReplicationSystem/NetActorFactory.cpp
+++ b/Engine/Source/Runtime/Engine/Private/Net/Iris/ReplicationSystem/NetActorFactory.cpp
@@ -54,11 +54,14 @@ namespace UE::Net::Private
 	
 	EActorNetSpawnInfoFlags GetSpawnInfoFlags()
 	{
+		// Disable performance warnings on FindConsoleVariable. We we call it once per NetActorFactory that is created, but in tests this can be thousands of times.
+		constexpr bool bWarnOnFrequentFindCVar = false;
+
 		// Init spawninfo flags from CVARs
 		EActorNetSpawnInfoFlags Flags = EActorNetSpawnInfoFlags::None;
 		{
 			bool bQuantizeActorScaleOnSpawn = false;
-			IConsoleVariable* CVar = IConsoleManager::Get().FindConsoleVariable(TEXT("net.QuantizeActorScaleOnSpawn"));
+			IConsoleVariable* CVar = IConsoleManager::Get().FindConsoleVariable(TEXT("net.QuantizeActorScaleOnSpawn"), bWarnOnFrequentFindCVar);
 			if (ensure(CVar))
 			{
 				bQuantizeActorScaleOnSpawn = CVar->GetBool();
@@ -68,7 +71,7 @@ namespace UE::Net::Private
 
 		{
 			bool bQuantizeActorLocationOnSpawn = true;
-			IConsoleVariable* CVar = IConsoleManager::Get().FindConsoleVariable(TEXT("net.QuantizeActorLocationOnSpawn"));
+			IConsoleVariable* CVar = IConsoleManager::Get().FindConsoleVariable(TEXT("net.QuantizeActorLocationOnSpawn"), bWarnOnFrequentFindCVar);
 			if (ensure(CVar))
 			{
 				bQuantizeActorLocationOnSpawn = CVar->GetBool();
@@ -78,7 +81,7 @@ namespace UE::Net::Private
 
 		{
 			bool bQuantizeActorVelocityOnSpawn = true;
-			IConsoleVariable* CVar = IConsoleManager::Get().FindConsoleVariable(TEXT("net.QuantizeActorVelocityOnSpawn"));
+			IConsoleVariable* CVar = IConsoleManager::Get().FindConsoleVariable(TEXT("net.QuantizeActorVelocityOnSpawn"), bWarnOnFrequentFindCVar);
 			if (ensure(CVar))
 			{
 				bQuantizeActorVelocityOnSpawn = CVar->GetBool();
diff --git a/Engine/Source/Runtime/Experimental/Iris/Core/Private/Iris/ReplicationSystem/DirtyNetObjectTracker.cpp b/Engine/Source/Runtime/Experimental/Iris/Core/Private/Iris/ReplicationSystem/DirtyNetObjectTracker.cpp
--- a/Engine/Source/Runtime/Experimental/Iris/Core/Private/Iris/ReplicationSystem/DirtyNetObjectTracker.cpp
+++ b/Engine/Source/Runtime/Experimental/Iris/Core/Private/Iris/ReplicationSystem/DirtyNetObjectTracker.cpp
@@ -3,7 +3,6 @@
 #include "DirtyNetObjectTracker.h"
 #include "HAL/PlatformAtomics.h"
 #include "Iris/IrisConstants.h"
-#include "Iris/Core/IrisLog.h"
 #include "Iris/Core/IrisProfiler.h"
 #include "Iris/ReplicationSystem/ReplicationSystem.h"
 #include "Iris/ReplicationSystem/ReplicationSystemInternal.h"
@@ -13,14 +12,8 @@
 
 #include <atomic>
 
-// Don't compile verbose logs in Shipping builds
-#if (UE_BUILD_SHIPPING)
-#	define UE_NET_DIRTYOBJECTTRACKER_LOG_COMPIL_VERBOSITY Log
-#else
-#	define UE_NET_DIRTYOBJECTTRACKER_LOG_COMPIL_VERBOSITY All
-#endif 
 
-DEFINE_LOG_CATEGORY_STATIC(LogIrisDirtyTracker, Log, UE_NET_DIRTYOBJECTTRACKER_LOG_COMPIL_VERBOSITY)
+DEFINE_LOG_CATEGORY(LogIrisDirtyTracker)
 
 namespace UE::Net::Private
 {
diff --git a/Engine/Source/Runtime/Experimental/Iris/Core/Private/Iris/ReplicationSystem/DirtyNetObjectTracker.h b/Engine/Source/Runtime/Experimental/Iris/Core/Private/Iris/ReplicationSystem/DirtyNetObjectTracker.h
--- a/Engine/Source/Runtime/Experimental/Iris/Core/Private/Iris/ReplicationSystem/DirtyNetObjectTracker.h
+++ b/Engine/Source/Runtime/Experimental/Iris/Core/Private/Iris/ReplicationSystem/DirtyNetObjectTracker.h
@@ -2,6 +2,8 @@
 
 #pragma once
 
+#include "Logging/LogMacros.h"
+
 #include "Net/Core/NetBitArray.h"
 #include "Net/Core/DirtyNetObjectTracker/GlobalDirtyNetObjectTracker.h"
 
@@ -16,6 +18,17 @@ namespace UE::Net::Private
 	typedef uint32 FInternalNetRefIndex;
 }
 
+#ifndef UE_NET_DIRTYOBJECTTRACKER_LOG_COMPILE_VERBOSITY
+// Don't compile verbose logs in Shipping builds	
+#if UE_BUILD_SHIPPING
+#	define UE_NET_DIRTYOBJECTTRACKER_LOG_COMPILE_VERBOSITY Log
+#else
+#	define UE_NET_DIRTYOBJECTTRACKER_LOG_COMPILE_VERBOSITY All
+#endif
+#endif
+
+IRISCORE_API DECLARE_LOG_CATEGORY_EXTERN(LogIrisDirtyTracker, Log, UE_NET_DIRTYOBJECTTRACKER_LOG_COMPILE_VERBOSITY);
+
 namespace UE::Net::Private
 {
 
