Don't run UWorldPartition::OnGCPostReachabilityAnalysis() during the engine shutdown step to avoid slowing down the exit process and reporting false errors.

diff --git a/Engine/Source/Runtime/Engine/Private/WorldPartition/WorldPartition.cpp b/Engine/Source/Runtime/Engine/Private/WorldPartition/WorldPartition.cpp
--- a/Engine/Source/Runtime/Engine/Private/WorldPartition/WorldPartition.cpp
+++ b/Engine/Source/Runtime/Engine/Private/WorldPartition/WorldPartition.cpp
@@ -390,38 +390,41 @@ UWorldPartition::UWorldPartition(const FObjectInitializer& ObjectInitializer)
 #if WITH_EDITOR
 void UWorldPartition::OnGCPostReachabilityAnalysis()
 {
-	TRACE_CPUPROFILER_EVENT_SCOPE(UWorldPartition::OnGCPostReachabilityAnalysis);
-	const TIndirectArray<FWorldContext>& WorldContextList = GEngine->GetWorldContexts();
-
-	// Avoid running this process while a game world is live
-	for (const FWorldContext& WorldContext : WorldContextList)
+	if (!IsEngineExitRequested())
 	{
-		if (WorldContext.World() != nullptr && WorldContext.World()->IsGameWorld())
+		TRACE_CPUPROFILER_EVENT_SCOPE(UWorldPartition::OnGCPostReachabilityAnalysis);
+		const TIndirectArray<FWorldContext>& WorldContextList = GEngine->GetWorldContexts();
+
+		// Avoid running this process while a game world is live
+		for (const FWorldContext& WorldContext : WorldContextList)
 		{
-			return;
+			if (WorldContext.World() != nullptr && WorldContext.World()->IsGameWorld())
+			{
+				return;
+			}
 		}
-	}
 
-	for (FRawObjectIterator It; It; ++It)
-	{
-		if (AActor* Actor = Cast<AActor>(static_cast<UObject*>(It->Object)))
+		for (FRawObjectIterator It; It; ++It)
 		{
-			if (Actor->IsUnreachable() && !Actor->GetClass()->HasAnyClassFlags(CLASS_NewerVersionExists) && Actor->IsMainPackageActor())
+			if (AActor* Actor = Cast<AActor>(static_cast<UObject*>(It->Object)))
 			{
-				ForEachObjectWithPackage(Actor->GetPackage(), [this, Actor](UObject* Object)
+				if (Actor->IsUnreachable() && !Actor->GetClass()->HasAnyClassFlags(CLASS_NewerVersionExists) && Actor->IsMainPackageActor())
 				{
-					if (Object->HasAnyFlags(RF_Standalone))
+					ForEachObjectWithPackage(Actor->GetPackage(), [this, Actor](UObject* Object)
 					{
-						UE_LOG(LogWorldPartition, Log, TEXT("Actor %s is unreachable without properly detaching object %s in its package"), *Actor->GetPathName(), *Object->GetPathName());
+						if (Object->HasAnyFlags(RF_Standalone))
+						{
+							UE_LOG(LogWorldPartition, Log, TEXT("Actor %s is unreachable without properly detaching object %s in its package"), *Actor->GetPathName(), *Object->GetPathName());
 
-						Object->ClearFlags(RF_Standalone);
+							Object->ClearFlags(RF_Standalone);
 
-						// Make sure we trigger a second GC at the next tick to properly destroy packages there were fixed in this pass
-						bForceGarbageCollection = true;
-						bForceGarbageCollectionPurge = true;
-					}
-					return true;
-				}, false);
+							// Make sure we trigger a second GC at the next tick to properly destroy packages there were fixed in this pass
+							bForceGarbageCollection = true;
+							bForceGarbageCollectionPurge = true;
+						}
+						return true;
+					}, false);
+				}
 			}
 		}
 	}
