Auto enable IsThin during Substrate's material auto-conversion with foliage material.

diff --git a/Engine/Source/Runtime/Engine/Private/Materials/Material.cpp b/Engine/Source/Runtime/Engine/Private/Materials/Material.cpp
--- a/Engine/Source/Runtime/Engine/Private/Materials/Material.cpp
+++ b/Engine/Source/Runtime/Engine/Private/Materials/Material.cpp
@@ -3720,6 +3720,12 @@ bool UMaterial::ConvertMaterialToSubstrateMaterial(bool bAllowEmptyMaterialUpdat
 				// Store Substrate shading model of the converted material. 
 				check(ShadingModels.CountShadingModels() == 1);
 				ConvertAttributeNode->ShadingModelOverride = ShadingModel;
+
+				// Apply is thin surface if needed
+				if (ShadingModel == MSM_TwoSidedFoliage)
+				{
+					bIsThinSurface = true;
+				}
 			}
 
 			if (MaterialDomain == MD_DeferredDecal)
@@ -3876,7 +3882,7 @@ bool UMaterial::ConvertMaterialToSubstrateMaterial(bool bAllowEmptyMaterialUpdat
 						bRelinkCustomOutputNodes = false;	// We do not want that to happen in this case
 					}
 				}
-			
+
 				if (!bClearCoatConversionDone)
 				{
 					ConvertNode = NewObject<UMaterialExpressionSubstrateShadingModels>(this);
@@ -3927,6 +3933,11 @@ bool UMaterial::ConvertMaterialToSubstrateMaterial(bool bAllowEmptyMaterialUpdat
 					{
 						ConvertNode->ShadingModelOverride = ShadingModel;
 						check(ShadingModels.CountShadingModels() == 1);
+
+						if (ShadingModel == MSM_TwoSidedFoliage)
+						{
+							bIsThinSurface = true;
+						}
 					}
 
 					EditorOnly->FrontMaterial.Connect(0, ConvertNode);
