Export URuntimePartition so it can be overloaded in other projects.

diff --git a/Engine/Source/Runtime/Engine/Public/WorldPartition/RuntimeHashSet/RuntimePartition.h b/Engine/Source/Runtime/Engine/Public/WorldPartition/RuntimeHashSet/RuntimePartition.h
--- a/Engine/Source/Runtime/Engine/Public/WorldPartition/RuntimeHashSet/RuntimePartition.h
+++ b/Engine/Source/Runtime/Engine/Public/WorldPartition/RuntimeHashSet/RuntimePartition.h
@@ -15,7 +15,7 @@ enum class ERuntimePartitionCellBoundsMethod : uint8
 };
 
 UCLASS(Abstract, CollapseCategories)
-class URuntimePartition : public UObject
+class ENGINE_API URuntimePartition : public UObject
 {
 	GENERATED_UCLASS_BODY()
 
