Removed unused dead code in the material translator.

diff --git a/Engine/Source/Runtime/Engine/Private/Materials/HLSLMaterialTranslator.cpp b/Engine/Source/Runtime/Engine/Private/Materials/HLSLMaterialTranslator.cpp
--- a/Engine/Source/Runtime/Engine/Private/Materials/HLSLMaterialTranslator.cpp
+++ b/Engine/Source/Runtime/Engine/Private/Materials/HLSLMaterialTranslator.cpp
@@ -660,12 +660,6 @@ void FHLSLMaterialTranslator::FSubstrateCompilationContext::Initialise()
 	SubstrateThicknessStack.SetNum(0);
 }
 
-void FHLSLMaterialTranslator::AppendVersionString(FString& Output, EShaderPlatform Platform)
-{
-	FShaderKeyGenerator KeyGen(Output);
-	AppendVersion(KeyGen, Platform);
-}
-
 void FHLSLMaterialTranslator::AppendVersion(FShaderKeyGenerator& KeyGen, EShaderPlatform Platform)
 {
 	static const FGuid DDCVersion = FDevSystemGuids::GetSystemGuid(FDevSystemGuids::Get().MaterialTranslationDDCVersion);
diff --git a/Engine/Source/Runtime/Engine/Private/Materials/HLSLMaterialTranslator.h b/Engine/Source/Runtime/Engine/Private/Materials/HLSLMaterialTranslator.h
--- a/Engine/Source/Runtime/Engine/Private/Materials/HLSLMaterialTranslator.h
+++ b/Engine/Source/Runtime/Engine/Private/Materials/HLSLMaterialTranslator.h
@@ -598,7 +598,6 @@ public:
 	/**
 	 * Returns a string representation that identify the translator version. Used to keep shader DDC keys valid when material translation internals change.
 	 */
-	static void AppendVersionString(FString& Output, EShaderPlatform Platform);
 	static void AppendVersion(FShaderKeyGenerator& KeyGen, EShaderPlatform Platform);
 
 	FHLSLMaterialTranslator(FMaterial* InMaterial,
