EngineTest
* Prevent FTestWorlds server to clone the loaded editor world before transitioning to the proper unit test map.
* Instead it clones the 'Entry' map and transitions to the requested world.
* Made FTestWorlds use the simple GameMode class instead of the default game mode configured in the current project's properties

diff --git a/Engine/Source/Runtime/Engine/Private/Tests/NetTestHelpers.cpp b/Engine/Source/Runtime/Engine/Private/Tests/NetTestHelpers.cpp
--- a/Engine/Source/Runtime/Engine/Private/Tests/NetTestHelpers.cpp
+++ b/Engine/Source/Runtime/Engine/Private/Tests/NetTestHelpers.cpp
@@ -10,7 +10,6 @@
 #include "Engine/GameViewportClient.h"
 #include "Engine/PackageMapClient.h"
 #include "EngineUtils.h"
-#include "GameFramework/PlayerController.h"
 #include "GameFramework/PlayerState.h"
 #include "HAL/IConsoleManager.h"
 #include "HAL/PlatformMath.h"
@@ -45,6 +44,7 @@ FTestWorldInstance FTestWorldInstance::CreateServer(const TCHAR* InURL)
 	ServerParams.WorldFeatureLevel = ERHIFeatureLevel::SM5;
 	ServerParams.EditorPlaySettings = nullptr;
 	ServerParams.NetMode = PIE_ListenServer;
+	ServerParams.OverrideMapURL = TEXT("/Engine/Maps/Entry"); // always start from an empty map. Otherwise it will clone the current editor map to start the server in.
 
 	FTestWorldInstance NewInstance(ServerParams);
 
@@ -99,6 +99,7 @@ FTestWorldInstance FTestWorldInstance::CreateProxy()
 	ServerParams.WorldFeatureLevel = ERHIFeatureLevel::SM5;
 	ServerParams.EditorPlaySettings = nullptr;
 	ServerParams.NetMode = PIE_ListenServer;
+	ServerParams.OverrideMapURL = TEXT("/Engine/Maps/Entry"); // always start from an empty map. Otherwise it will clone the current editor map to start the server in.
 		
 	FTestWorldInstance NewInstance(ServerParams);
 
@@ -149,10 +150,7 @@ void FTestWorldInstance::Shutdown()
 
 	if (World)
 	{
-		for (FActorIterator ActorIt(World); ActorIt; ++ActorIt)
-		{
-			ActorIt->RouteEndPlay(EEndPlayReason::EndPlayInEditor);
-		}
+		World->EndPlay(EEndPlayReason::EndPlayInEditor);
 	}
 
 	if (GameInstance)
@@ -163,7 +161,6 @@ void FTestWorldInstance::Shutdown()
 	
 	if (World)
 	{
-		World->BeginTearingDown();
 		GEngine->ShutdownWorldNetDriver(World);
 		GEngine->DestroyWorldContext(World);
 		World->CleanupWorld();
@@ -265,11 +262,28 @@ void FTestWorldInstance::UnloadStreamingLevel(FName LevelName)
 // FTestWorlds
 //============================================================================
 
-FTestWorlds::FTestWorlds(const TCHAR* ServerURL, float DeltaSeconds)
+FTestWorlds::FTestWorlds()
 	: Server(true)
-	, TickDeltaSeconds(DeltaSeconds)
 {
-	NetDriverCreatedHandle = FWorldDelegates::OnNetDriverCreated.AddRaw(this, &FTestWorlds::OnNetDriverCreated);
+	InitDelegates();
+
+	Server = FTestWorldInstance::CreateServer(TEXT("/Engine/Maps/Entry?listen?Game=/Script/Engine.GameMode"));
+}
+
+FTestWorlds::FTestWorlds(const FString& MapName, const FString& GameModeName)
+	: Server(true)
+{
+	InitDelegates();
+
+	FString URL = FString::Printf(TEXT("%s?listen?Game=%s"), *MapName, *GameModeName);
+
+	Server = FTestWorldInstance::CreateServer(*URL);
+}
+
+FTestWorlds::FTestWorlds(const TCHAR* ServerURL)
+	: Server(true)
+{
+	InitDelegates();
 
 	Server = FTestWorldInstance::CreateServer(ServerURL);
 }
@@ -279,6 +293,16 @@ FTestWorlds::~FTestWorlds()
 	FWorldDelegates::OnNetDriverCreated.Remove(NetDriverCreatedHandle);
 }
 
+void FTestWorlds::SetTickInSeconds(float TickInSeconds)
+{
+	TickDeltaSeconds = TickInSeconds;
+}
+
+void FTestWorlds::InitDelegates()
+{
+	NetDriverCreatedHandle = FWorldDelegates::OnNetDriverCreated.AddRaw(this, &FTestWorlds::OnNetDriverCreated);
+}
+
 void FTestWorlds::OnNetDriverCreated(UWorld* InWorld, UNetDriver* InNetDriver)
 {
 	// Make sure NetDriver will tick every engine frame
diff --git a/Engine/Source/Runtime/Engine/Public/Tests/NetTestHelpers.h b/Engine/Source/Runtime/Engine/Public/Tests/NetTestHelpers.h
--- a/Engine/Source/Runtime/Engine/Public/Tests/NetTestHelpers.h
+++ b/Engine/Source/Runtime/Engine/Public/Tests/NetTestHelpers.h
@@ -154,10 +154,24 @@ private:
  */
 struct FTestWorlds
 {
+	/** Creates a server world using the default map (Entry) and default GameMode class (GameMode) */
+	ENGINE_API FTestWorlds();
+
+	/** 
+	 * Creates a server world using a specific map name and GameMode class
+	 * @param MapName Full path name of the map. ex: /Engine/Maps/Entry
+	 * @param GameModeName Full path name of the game mode. ex: /Script/EngineTest.MyGameMode
+	 */
+	ENGINE_API FTestWorlds(const FString& MapName, const FString& GameModeName);
+
 	/** Creates a server world using the given URL. */
-	ENGINE_API explicit FTestWorlds(const TCHAR* ServerURL, float DeltaSeconds = 0.0166f);
+	ENGINE_API explicit FTestWorlds(const TCHAR* ServerURL);
+
 	ENGINE_API ~FTestWorlds();
 
+	ENGINE_API void SetTickInSeconds(float TickInSeconds);
+	
+	/** Create a client and fully connect it to the server before returning */
 	ENGINE_API bool CreateAndConnectClient();
 
 	/** Ticks all server & client worlds NumTick times synchronously. */
@@ -202,6 +216,8 @@ public:
 
 private:
 
+	void InitDelegates();
+
 	void OnNetDriverCreated(UWorld* InWorld, UNetDriver* InNetDriver);
 	FDelegateHandle NetDriverCreatedHandle;
 
