Remove old code forcing Windowed Fullscreen for DX12, the condition is broken and native Fullscreen is supported by DX12 RHI.

#rnx

diff --git a/Engine/Source/Runtime/Engine/Private/GameEngine.cpp b/Engine/Source/Runtime/Engine/Private/GameEngine.cpp
--- a/Engine/Source/Runtime/Engine/Private/GameEngine.cpp
+++ b/Engine/Source/Runtime/Engine/Private/GameEngine.cpp
@@ -323,17 +323,6 @@ void UGameEngine::ConditionallyOverrideSettings(int32& ResolutionX, int32& Resol
 			static auto CVar = IConsoleManager::Get().FindTConsoleVariableDataInt(TEXT("r.FullScreenMode"));
 			check(CVar);
 			WindowMode = CVar->GetValueOnGameThread() == 0 ? EWindowMode::Fullscreen : EWindowMode::WindowedFullscreen;
-
-			if (PLATFORM_WINDOWS && WindowMode == EWindowMode::Fullscreen)
-			{
-				// Handle fullscreen mode differently for D3D11/D3D12
-				static const bool bD3D12 = FParse::Param(FCommandLine::Get(), TEXT("d3d12")) || FParse::Param(FCommandLine::Get(), TEXT("dx12"));
-				if (bD3D12)
-				{
-					// Force D3D12 RHI to use windowed fullscreen mode
-					WindowMode = EWindowMode::WindowedFullscreen;
-				}
-			}
 		}
 	}
 
diff --git a/Engine/Source/Runtime/Engine/Private/GameViewportClient.cpp b/Engine/Source/Runtime/Engine/Private/GameViewportClient.cpp
--- a/Engine/Source/Runtime/Engine/Private/GameViewportClient.cpp
+++ b/Engine/Source/Runtime/Engine/Private/GameViewportClient.cpp
@@ -3845,17 +3845,6 @@ bool UGameViewportClient::HandleToggleFullscreenCommand()
 	auto FullScreenMode = CVar->GetValueOnGameThread() == 0 ? EWindowMode::Fullscreen : EWindowMode::WindowedFullscreen;
 	FullScreenMode = Viewport->IsFullscreen() ? EWindowMode::Windowed : FullScreenMode;
 
-	if (PLATFORM_WINDOWS && FullScreenMode == EWindowMode::Fullscreen)
-	{
-		// Handle fullscreen mode differently for D3D11/D3D12
-		static const bool bD3D12 = FParse::Param(FCommandLine::Get(), TEXT("d3d12")) || FParse::Param(FCommandLine::Get(), TEXT("dx12"));
-		if (bD3D12)
-		{
-			// Force D3D12 RHI to use windowed fullscreen mode
-			FullScreenMode = EWindowMode::WindowedFullscreen;
-		}
-	}
-
 	int32 ResolutionX = GSystemResolution.ResX;
 	int32 ResolutionY = GSystemResolution.ResY;
 	bool bNewModeApplied = false;
