Changed verse DataLayerManager into WorldPartitionHelper since we expanded it's capacities
Added support for StreamingSource/Queries into this helper
Only intended for experimental internal use, so still behind an <epic_internal>

#okforversepublic

diff --git a/Engine/Source/Runtime/Engine/Private/WorldPartition/DataLayer/DataLayerManager.cpp b/Engine/Source/Runtime/Engine/Private/WorldPartition/DataLayer/DataLayerManager.cpp
--- a/Engine/Source/Runtime/Engine/Private/WorldPartition/DataLayer/DataLayerManager.cpp
+++ b/Engine/Source/Runtime/Engine/Private/WorldPartition/DataLayer/DataLayerManager.cpp
@@ -244,16 +244,6 @@ void UDataLayerManager::DeInitialize()
 #endif
 }
 
-void UDataLayerManager::AddReferencedObject(UObject* InObject)
-{
-	ReferencedObjects.Add(InObject);
-}
-
-void UDataLayerManager::RemoveReferencedObject(UObject* InObject)
-{
-	ReferencedObjects.Remove(InObject);
-}
-
 AWorldDataLayers* UDataLayerManager::GetWorldDataLayers() const
 {
 	return GetTypedOuter<UWorld>()->GetWorldDataLayers();
diff --git a/Engine/Source/Runtime/Engine/Private/WorldPartition/WorldPartition.cpp b/Engine/Source/Runtime/Engine/Private/WorldPartition/WorldPartition.cpp
--- a/Engine/Source/Runtime/Engine/Private/WorldPartition/WorldPartition.cpp
+++ b/Engine/Source/Runtime/Engine/Private/WorldPartition/WorldPartition.cpp
@@ -2532,4 +2532,15 @@ TArray<UWorldPartitionRuntimeCell*> UWorldPartition::GetWorldAssetStreamingCells
 	return RuntimeHash ? RuntimeHash->GetWorldAssetStreamingCells(InWorldAssetStreamingGuid) : TArray<UWorldPartitionRuntimeCell*>();
 }
 
+void UWorldPartition::AddReferencedObject(UObject* InObject)
+{
+	ReferencedObjects.Add(InObject);
+}
+
+void UWorldPartition::RemoveReferencedObject(UObject* InObject)
+{
+	ReferencedObjects.Remove(InObject);
+}
+
+
 #undef LOCTEXT_NAMESPACE
diff --git a/Engine/Source/Runtime/Engine/Public/WorldPartition/DataLayer/DataLayerManager.h b/Engine/Source/Runtime/Engine/Public/WorldPartition/DataLayer/DataLayerManager.h
--- a/Engine/Source/Runtime/Engine/Public/WorldPartition/DataLayer/DataLayerManager.h
+++ b/Engine/Source/Runtime/Engine/Public/WorldPartition/DataLayer/DataLayerManager.h
@@ -140,22 +140,12 @@ private:
 	ENGINE_API TArray<UDataLayerInstance*> ConvertArgsToDataLayers(const TArray<FString>& InArgs);
 	//~ End Debugging
 
-	//~ Begin Verse support
-	ENGINE_API void AddReferencedObject(UObject* InObject);
-	ENGINE_API void RemoveReferencedObject(UObject* InObject);
-	const TSet<TObjectPtr<UObject>>& GetReferencedObjects() const { return ReferencedObjects; }
-	//~ End Verse support
-
 	ENGINE_API AWorldDataLayers* GetWorldDataLayers() const;
 	ENGINE_API const UDataLayerInstance* GetDataLayerInstanceFromAssetName(const FName& InDataLayerAssetPathName) const;
 	ENGINE_API void BroadcastOnDataLayerInstanceRuntimeStateChanged(const UDataLayerInstance* InDataLayer, EDataLayerRuntimeState InState);
 	ENGINE_API static bool IsAnyDataLayerInEffectiveRuntimeState(TArrayView<const FName> InDataLayerNames, EDataLayerRuntimeState InState, const FWorldDataLayersEffectiveStates& InEffectiveStates);
 	ENGINE_API static bool IsAllDataLayerInEffectiveRuntimeState(TArrayView<const FName> InDataLayerNames, EDataLayerRuntimeState InState, const FWorldDataLayersEffectiveStates& InEffectiveStates);
 
-	/** Referenced objects (used by verse) */
-	UPROPERTY(Transient)
-	TSet<TObjectPtr<UObject>> ReferencedObjects;
-
 	/** Console command used to toggle activation of a Data Layer */
 	static ENGINE_API class FAutoConsoleCommand ToggleDataLayerActivation;
 
@@ -178,7 +168,6 @@ private:
 	friend class UWorldPartitionRuntimeCell;
 	friend class UWorldPartitionSubsystem;
 	friend class FWorldPartitionDebugHelper;
-	friend class UVerseDataLayerManagerBase;
 	friend struct FWorldPartitionStreamingContext;
 
 #if WITH_EDITOR
diff --git a/Engine/Source/Runtime/Engine/Public/WorldPartition/WorldPartition.h b/Engine/Source/Runtime/Engine/Public/WorldPartition/WorldPartition.h
--- a/Engine/Source/Runtime/Engine/Public/WorldPartition/WorldPartition.h
+++ b/Engine/Source/Runtime/Engine/Public/WorldPartition/WorldPartition.h
@@ -684,6 +684,18 @@ public:
 private:
 #endif
 
+	//~ Begin Verse support
+	ENGINE_API void AddReferencedObject(UObject* InObject);
+	ENGINE_API void RemoveReferencedObject(UObject* InObject);
+	const TSet<TObjectPtr<UObject>>& GetReferencedObjects() const { return ReferencedObjects; }
+	//~ End Verse support
+
+	/** Referenced objects (used by verse) */
+	UPROPERTY(Transient)
+	TSet<TObjectPtr<UObject>> ReferencedObjects;
+
+	friend class UVerseWorldPartitionHelperBase;
+
 	EWorldPartitionInitState InitState;
 	TOptional<FTransform> InstanceTransform;
 
