Fix UEngineElementsLibrary registering invalid actors, resulting in a crash. This can happen during PIE for unsaved actors where we create copies of these actors around for duplication instead of streaming them from disk to include unsaved changes.
These actors don't belong to a level, they are just templates for duplication.

diff --git a/Engine/Source/Runtime/Engine/Private/Elements/Framework/EngineElementsLibrary.cpp b/Engine/Source/Runtime/Engine/Private/Elements/Framework/EngineElementsLibrary.cpp
--- a/Engine/Source/Runtime/Engine/Private/Elements/Framework/EngineElementsLibrary.cpp
+++ b/Engine/Source/Runtime/Engine/Private/Elements/Framework/EngineElementsLibrary.cpp
@@ -401,7 +401,10 @@ void UEngineElementsLibrary::UnregisterActorElement(const AActor* InActor)
 #if WITH_EDITOR
 void UEngineElementsLibrary::CreateEditorActorElement(const AActor* Actor)
 {
-	EngineElementsLibraryUtil::CreateEditorTypedElement<AActor, FActorElementData>(Actor, GActorElementOwnerStore, &UEngineElementsLibrary::CreateActorElement);
+	if (Actor->GetLevel())
+	{
+		EngineElementsLibraryUtil::CreateEditorTypedElement<AActor, FActorElementData>(Actor, GActorElementOwnerStore, &UEngineElementsLibrary::CreateActorElement);
+	}
 }
 
 void UEngineElementsLibrary::DestroyEditorActorElement(const AActor* Actor)
