Add Mirroring to steering.

#tests PIE

diff --git a/Engine/Plugins/Animation/AnimationWarping/Source/Runtime/Private/BoneControllers/AnimNode_Steering.cpp b/Engine/Plugins/Animation/AnimationWarping/Source/Runtime/Private/BoneControllers/AnimNode_Steering.cpp
--- a/Engine/Plugins/Animation/AnimationWarping/Source/Runtime/Private/BoneControllers/AnimNode_Steering.cpp
+++ b/Engine/Plugins/Animation/AnimationWarping/Source/Runtime/Private/BoneControllers/AnimNode_Steering.cpp
@@ -4,6 +4,7 @@
 #include "Animation/AnimInstanceProxy.h"
 #include "Animation/AnimNodeFunctionRef.h"
 #include "Animation/AnimRootMotionProvider.h"
+#include "Animation/AnimSequenceHelpers.h"
 #include "HAL/IConsoleManager.h"
 #include "Animation/AnimTrace.h"
 #include "BoneControllers/AnimNode_OffsetRootBone.h"
@@ -85,7 +86,16 @@ void FAnimNode_Steering::EvaluateSkeletalControl_AnyThread(FComponentSpacePoseCo
 						{
 							if (UAnimSequenceBase* AnimSequence = Cast<UAnimSequenceBase>(CurrentAnimAsset))
 							{
-								FTransform PredictedRootMotionTransform = AnimSequence->ExtractRootMotion(CurrentAnimAssetTime, AnimatedTargetTime, true);
+								FTransform PredictedRootMotionTransform;
+								if (bMirrored && MirrorDataTable)
+								{
+									PredictedRootMotionTransform = UE::Anim::ExtractRootMotionFromAnimationAsset(AnimSequence, MirrorDataTable, CurrentAnimAssetTime, AnimatedTargetTime, true);
+								}
+								else
+								{
+									PredictedRootMotionTransform = AnimSequence->ExtractRootMotion(CurrentAnimAssetTime, AnimatedTargetTime, true);
+								}
+
 								FQuat PredictedRootMotionQuat = PredictedRootMotionTransform.GetRotation();
 								float PredictedRootMotionYaw = FRotator(PredictedRootMotionQuat).Yaw;
 							
diff --git a/Engine/Plugins/Animation/AnimationWarping/Source/Runtime/Public/BoneControllers/AnimNode_Steering.h b/Engine/Plugins/Animation/AnimationWarping/Source/Runtime/Public/BoneControllers/AnimNode_Steering.h
--- a/Engine/Plugins/Animation/AnimationWarping/Source/Runtime/Public/BoneControllers/AnimNode_Steering.h
+++ b/Engine/Plugins/Animation/AnimationWarping/Source/Runtime/Public/BoneControllers/AnimNode_Steering.h
@@ -23,6 +23,10 @@ struct ANIMATIONWARPINGRUNTIME_API FAnimNode_Steering : public FAnimNode_Skeleta
 
 	UPROPERTY(EditAnywhere, BlueprintReadWrite, Category=Filtering)
 	bool bEnableTargetSmoothing = false;
+
+	// True if input animation is mirrored
+	UPROPERTY(EditAnywhere, BlueprintReadWrite, Category=Filtering, meta = (PinShownByDefault))
+	bool bMirrored = false;
 	
 	UPROPERTY(EditAnywhere, BlueprintReadWrite, Category=Filtering)
 	float SmoothTargetStiffness = 300;
@@ -50,6 +54,10 @@ struct ANIMATIONWARPINGRUNTIME_API FAnimNode_Steering : public FAnimNode_Skeleta
 	// below this movement speed (based on the root motion in the animation) disable steering
 	UPROPERTY(EditAnywhere, Category=Evaluation)
 	float DisableSteeringBelowSpeed = 1.0f;
+
+	// If bMirrored is set, MirrorDataTable will be used for mirroring the CurrentAnimAsset during prediction
+	UPROPERTY(EditAnywhere, Transient, BlueprintReadWrite, Category = Evaluation, meta = (PinShownByDefault))
+	TObjectPtr<UMirrorDataTable> MirrorDataTable;
 	
 	// Animation Asset for incorporating root motion data. If CurrentAnimAsset is set, and the animation has root motion rotation within the TargetTime, then those rotations will be scaled to reach the TargetOrientation
 	UPROPERTY(EditAnywhere, Transient, BlueprintReadWrite, Category=Evaluation, meta=(PinShownByDefault))
diff --git a/Engine/Plugins/Animation/BlendStack/Source/Runtime/Private/BlendStackAnimNodeLibrary.cpp b/Engine/Plugins/Animation/BlendStack/Source/Runtime/Private/BlendStackAnimNodeLibrary.cpp
--- a/Engine/Plugins/Animation/BlendStack/Source/Runtime/Private/BlendStackAnimNodeLibrary.cpp
+++ b/Engine/Plugins/Animation/BlendStack/Source/Runtime/Private/BlendStackAnimNodeLibrary.cpp
@@ -34,6 +34,30 @@ float UBlendStackAnimNodeLibrary::GetCurrentBlendStackAnimAssetTime(const FAnimN
 	return 0.f;
 }
 
+bool UBlendStackAnimNodeLibrary::GetCurrentBlendStackAnimAssetMirrored(const FAnimNodeReference& Node)
+{
+	if (FAnimNode_BlendStackInput* BlendStackInput = Node.GetAnimNodePtr<FAnimNode_BlendStackInput>())
+	{
+		if (BlendStackInput->Player && *BlendStackInput->Player)
+		{
+			return (*BlendStackInput->Player)->GetMirror();
+		}
+	}
+	return false;
+}
+
+UMirrorDataTable* UBlendStackAnimNodeLibrary::GetCurrentBlendStackAnimAssetMirrorTable(const FAnimNodeReference& Node)
+{
+	if (FAnimNode_BlendStackInput* BlendStackInput = Node.GetAnimNodePtr<FAnimNode_BlendStackInput>())
+	{
+		if (BlendStackInput->Player && *BlendStackInput->Player)
+		{
+			return (*BlendStackInput->Player)->GetMirrorNode().GetMirrorDataTable();
+		}
+	}
+	return nullptr;
+}
+
 void UBlendStackAnimNodeLibrary::BlendTo(const FAnimUpdateContext& Context, 
 										const FBlendStackAnimNodeReference& BlendStackNode, 
 										UAnimationAsset* AnimationAsset,
diff --git a/Engine/Plugins/Animation/BlendStack/Source/Runtime/Public/BlendStack/BlendStackAnimNodeLibrary.h b/Engine/Plugins/Animation/BlendStack/Source/Runtime/Public/BlendStack/BlendStackAnimNodeLibrary.h
--- a/Engine/Plugins/Animation/BlendStack/Source/Runtime/Public/BlendStack/BlendStackAnimNodeLibrary.h
+++ b/Engine/Plugins/Animation/BlendStack/Source/Runtime/Public/BlendStack/BlendStackAnimNodeLibrary.h
@@ -30,7 +30,6 @@ public:
 	UFUNCTION(BlueprintCallable, Category = "Animation|BlendStack", meta = (BlueprintThreadSafe, ExpandEnumAsExecs = "Result"))
 	static FBlendStackAnimNodeReference ConvertToBlendStackNode(const FAnimNodeReference& Node, EAnimNodeReferenceConversionResult& Result);
 	
-	
 	/** Get the current AnimationAsset that is playing from a Blend Stack Input node */
     UFUNCTION(BlueprintPure, Category = "Animation|BlendStack", meta = (BlueprintThreadSafe))
     static UAnimationAsset* GetCurrentBlendStackAnimAsset(const FAnimNodeReference& Node);
@@ -39,6 +38,14 @@ public:
     UFUNCTION(BlueprintPure, Category = "Animation|BlendStack", meta = (BlueprintThreadSafe))
     static float GetCurrentBlendStackAnimAssetTime(const FAnimNodeReference& Node);
 
+	/** Get if we are currently mirrored from a Blend Stack Input node */
+	UFUNCTION(BlueprintPure, Category = "Animation|BlendStack", meta = (BlueprintThreadSafe))
+	static bool GetCurrentBlendStackAnimAssetMirrored(const FAnimNodeReference& Node);
+
+	/** Get current AssetMirrorTable from a Blend Stack Input node */
+	UFUNCTION(BlueprintPure, Category = "Animation|BlendStack", meta = (BlueprintThreadSafe))
+	static UMirrorDataTable* GetCurrentBlendStackAnimAssetMirrorTable(const FAnimNodeReference& Node);
+
 	/** Get a blend stack node context from an anim node context (pure) */
 	UFUNCTION(BlueprintPure, Category = "Animation|BlendStack", meta = (BlueprintThreadSafe, DisplayName = "Convert to Blend Stack Node"))
 	static void ConvertToBlendStackNodePure(const FAnimNodeReference& Node, FBlendStackAnimNodeReference& BlendStackNode, bool& Result)
diff --git a/Engine/Source/Runtime/Engine/Classes/Animation/AnimSequenceHelpers.h b/Engine/Source/Runtime/Engine/Classes/Animation/AnimSequenceHelpers.h
--- a/Engine/Source/Runtime/Engine/Classes/Animation/AnimSequenceHelpers.h
+++ b/Engine/Source/Runtime/Engine/Classes/Animation/AnimSequenceHelpers.h
@@ -194,6 +194,18 @@ namespace Anim {
 	 */
 	ENGINE_API FTransform MirrorTransform(const FTransform& Transform, const UMirrorDataTable& MirrorDataTable);
 
+	/**
+	 * Extract relative root motion between start and end position of the animation asset. With optional looping if applicable.
+	 * NOTE WELL: Our second time param is a delta time, not an end like like other variations of this API.
+	 *
+	 * @param	Animation		Animation asset to extract the root motion from.
+	 * @param	MirrorDataTable	[optional] Animation data mirroring data table.
+	 * @param	StartTime		Start Position of the animation.
+	 * @param	DeltaTime		Delta time to Advance, can be negative.
+	 * @param	bAllowLooping	True if we should accumulate and loop.
+	 */
+	ENGINE_API FTransform ExtractRootMotionFromAnimationAsset(const UAnimationAsset* Animation, const UMirrorDataTable* MirrorDataTable, float StartTime, float DeltaTime, bool bAllowLooping);
+
 	/**
 	 * Extract relative root motion between start and end position of the animation asset.
 	 *
diff --git a/Engine/Source/Runtime/Engine/Private/Animation/AnimSequenceHelpers.cpp b/Engine/Source/Runtime/Engine/Private/Animation/AnimSequenceHelpers.cpp
--- a/Engine/Source/Runtime/Engine/Private/Animation/AnimSequenceHelpers.cpp
+++ b/Engine/Source/Runtime/Engine/Private/Animation/AnimSequenceHelpers.cpp
@@ -207,6 +207,54 @@ FTransform MirrorTransform(const FTransform& Transform, const UMirrorDataTable&
 	return FTransform(Q, T, S);
 }
 
+FTransform ExtractRootMotionFromAnimationAsset(const UAnimationAsset* Animation, const UMirrorDataTable* MirrorDataTable, float StartTime, float DeltaTime, bool bAllowLooping)
+{
+	if (!Animation)
+	{
+		return FTransform::Identity;
+	}
+
+	// Note: Adapted from 'UAnimSequence::ExtractRootMotion' - since anim sequence doesn't know about MirrorDataTable
+	FRootMotionMovementParams RootMotionParams;
+
+	if (DeltaTime != 0.f)
+	{
+		bool const bPlayingBackwards = (DeltaTime < 0.f);
+
+		float PreviousPosition = StartTime;
+		float CurrentPosition = StartTime;
+		float DesiredDeltaMove = DeltaTime;
+
+		do
+		{
+			// Disable looping here. Advance to desired position, or beginning / end of animation 
+			const ETypeAdvanceAnim AdvanceType = FAnimationRuntime::AdvanceTime(false, DesiredDeltaMove, CurrentPosition, Animation->GetPlayLength());
+
+			// Verify position assumptions
+			ensureMsgf(bPlayingBackwards ? (CurrentPosition <= PreviousPosition) : (CurrentPosition >= PreviousPosition), TEXT("in Animation %s(Skeleton %s) : bPlayingBackwards(%d), PreviousPosition(%0.2f), Current Position(%0.2f)"),
+				*Animation->GetName(), *GetNameSafe(Animation->GetSkeleton()), bPlayingBackwards, PreviousPosition, CurrentPosition);
+			
+			RootMotionParams.Accumulate(UE::Anim::ExtractRootMotionFromAnimationAsset(Animation, MirrorDataTable, PreviousPosition, CurrentPosition));
+
+			// If we've hit the end of the animation, and we're allowed to loop, keep going.
+			if ((AdvanceType == ETAA_Finished) && bAllowLooping)
+			{
+				const double ActualDeltaMove = (CurrentPosition - PreviousPosition);
+				DesiredDeltaMove -= ActualDeltaMove;
+
+				PreviousPosition = bPlayingBackwards ? Animation->GetPlayLength() : 0.f;
+				CurrentPosition = PreviousPosition;
+			}
+			else
+			{
+				break;
+			}
+		} while (true);
+	}
+
+	return RootMotionParams.GetRootMotionTransform();
+}
+
 FTransform ExtractRootMotionFromAnimationAsset(const UAnimationAsset* Animation, const UMirrorDataTable* MirrorDataTable, float StartPosition, float EndPosition)
 {
 	FTransform Result = FTransform::Identity;
