Remove check from ObjectWorldPosition to unblock materials failing to compile. Will just fail silently for certain materials for now (return 0.0) until a full fix is re-introduced.

diff --git a/Engine/Source/Runtime/Engine/Private/Materials/HLSLMaterialTranslator.cpp b/Engine/Source/Runtime/Engine/Private/Materials/HLSLMaterialTranslator.cpp
--- a/Engine/Source/Runtime/Engine/Private/Materials/HLSLMaterialTranslator.cpp
+++ b/Engine/Source/Runtime/Engine/Private/Materials/HLSLMaterialTranslator.cpp
@@ -11748,7 +11748,7 @@ int32 FHLSLMaterialTranslator::Noise(int32 Position, EPositionOrigin PositionOri
 
 	return AddCodeChunk(MCT_Float, 
 		TEXT("MaterialExpressionNoise(%s,%s,%s,%s,%s,%s,%s,%s,%s,%s,%s,%s)"), 
-		*CoerceParameter(Position, MCT_Float3),
+		*GetParameterCode(Position),
 		*GetParameterCode(ScaleConst),
 		*GetParameterCode(QualityConst),
 		*GetParameterCode(NoiseFunctionConst),
