Fix for GFP device profile fragments cvars parsing order

diff --git a/Engine/Source/Runtime/Engine/Private/DeviceProfiles/DeviceProfileManager.cpp b/Engine/Source/Runtime/Engine/Private/DeviceProfiles/DeviceProfileManager.cpp
--- a/Engine/Source/Runtime/Engine/Private/DeviceProfiles/DeviceProfileManager.cpp
+++ b/Engine/Source/Runtime/Engine/Private/DeviceProfiles/DeviceProfileManager.cpp
@@ -240,8 +240,10 @@ TMap<FName, FString> UDeviceProfileManager::GatherDeviceProfileCVars(const FStri
 		{
 			TArray<FString> FragmentCVars;
 			GetFragmentCVars(*SelectedFragment.Fragment, TEXT("CVars"), FragmentCVars, ConfigSystem);
-			for (const FString& FragCVar : FragmentCVars)
+			for (int j = FragmentCVars.Num() - 1; j >= 0; j--)
 			{
+				const FString& FragCVar = FragmentCVars[j];
+
 				FString CVarKey, CVarValue;
 				if (FragCVar.Split(TEXT("="), &CVarKey, &CVarValue))
 				{
