Adding missing GatherDebugData calls.

diff --git a/Engine/Plugins/Animation/AnimationWarping/Source/Runtime/Private/BoneControllers/AnimNode_OffsetRootBone.cpp b/Engine/Plugins/Animation/AnimationWarping/Source/Runtime/Private/BoneControllers/AnimNode_OffsetRootBone.cpp
--- a/Engine/Plugins/Animation/AnimationWarping/Source/Runtime/Private/BoneControllers/AnimNode_OffsetRootBone.cpp
+++ b/Engine/Plugins/Animation/AnimationWarping/Source/Runtime/Private/BoneControllers/AnimNode_OffsetRootBone.cpp
@@ -85,6 +85,8 @@ void FAnimNode_OffsetRootBone::GatherDebugData(FNodeDebugData& DebugData)
 	DebugLine += FString::Printf(TEXT("\n - Rotation Halflife: (%.3fd)"), GetRotationHalfLife());
 #endif
 	DebugData.AddDebugItem(DebugLine);
+
+	Source.GatherDebugData(DebugData);
 }
 
 void FAnimNode_OffsetRootBone::Initialize_AnyThread(const FAnimationInitializeContext& Context)
diff --git a/Engine/Plugins/Animation/AnimationWarping/Source/Runtime/Private/BoneControllers/AnimNode_OverrideRootMotion.cpp b/Engine/Plugins/Animation/AnimationWarping/Source/Runtime/Private/BoneControllers/AnimNode_OverrideRootMotion.cpp
--- a/Engine/Plugins/Animation/AnimationWarping/Source/Runtime/Private/BoneControllers/AnimNode_OverrideRootMotion.cpp
+++ b/Engine/Plugins/Animation/AnimationWarping/Source/Runtime/Private/BoneControllers/AnimNode_OverrideRootMotion.cpp
@@ -25,6 +25,23 @@ void FAnimNode_OverrideRootMotion::CacheBones_AnyThread(const FAnimationCacheBon
 	Source.CacheBones(Context);
 }
 
+void FAnimNode_OverrideRootMotion::GatherDebugData(FNodeDebugData& DebugData)
+{
+	DECLARE_SCOPE_HIERARCHICAL_COUNTER_ANIMNODE(GatherDebugData)
+
+	// Folded variables.
+	FVector StaticOverrideVelocity = GetOverrideVelocity();
+	float StaticAlpha = GetAlpha();
+	
+	FString DebugLine = DebugData.GetNodeName(this);
+	DebugLine += FString::Printf(TEXT("(Override Velocity: %s, Alpha: %.3f"),
+		*StaticOverrideVelocity.ToString(),
+		StaticAlpha
+		);
+	
+	Source.GatherDebugData(DebugData);
+}
+
 void FAnimNode_OverrideRootMotion::Update_AnyThread(const FAnimationUpdateContext& Context)
 {
 	Super::Update_AnyThread(Context);
diff --git a/Engine/Plugins/Animation/AnimationWarping/Source/Runtime/Private/BoneControllers/AnimNode_Steering.cpp b/Engine/Plugins/Animation/AnimationWarping/Source/Runtime/Private/BoneControllers/AnimNode_Steering.cpp
--- a/Engine/Plugins/Animation/AnimationWarping/Source/Runtime/Private/BoneControllers/AnimNode_Steering.cpp
+++ b/Engine/Plugins/Animation/AnimationWarping/Source/Runtime/Private/BoneControllers/AnimNode_Steering.cpp
@@ -34,6 +34,19 @@ void FAnimNode_Steering::Initialize_AnyThread(const FAnimationInitializeContext&
 	bResetFilter = true;
 }
 
+void FAnimNode_Steering::GatherDebugData(FNodeDebugData& DebugData)
+{
+	DECLARE_SCOPE_HIERARCHICAL_COUNTER_ANIMNODE(GatherDebugData)
+
+	FString DebugLine = DebugData.GetNodeName(this);
+	// Just track alpha, there are clearer tools for the visualization of steering elsewhere.  
+	DebugLine += FString::Printf(TEXT("Alpha:%.3f"),
+		Alpha
+		);
+	
+	ComponentPose.GatherDebugData(DebugData);
+}
+
 
 void FAnimNode_Steering::EvaluateSkeletalControl_AnyThread(FComponentSpacePoseContext& Output, TArray<FBoneTransform>& OutBoneTransforms)
 {
diff --git a/Engine/Plugins/Animation/AnimationWarping/Source/Runtime/Public/BoneControllers/AnimNode_OverrideRootMotion.h b/Engine/Plugins/Animation/AnimationWarping/Source/Runtime/Public/BoneControllers/AnimNode_OverrideRootMotion.h
--- a/Engine/Plugins/Animation/AnimationWarping/Source/Runtime/Public/BoneControllers/AnimNode_OverrideRootMotion.h
+++ b/Engine/Plugins/Animation/AnimationWarping/Source/Runtime/Public/BoneControllers/AnimNode_OverrideRootMotion.h
@@ -34,6 +34,7 @@ public:
 	virtual void Evaluate_AnyThread(FPoseContext& Output) override;
 	virtual void Initialize_AnyThread(const FAnimationInitializeContext& Context) override;
 	virtual void CacheBones_AnyThread(const FAnimationCacheBonesContext& Context) override;
+	virtual void GatherDebugData(FNodeDebugData& DebugData) override;
 	// End of FAnimNode_Base interface
 
 	// Folded property accesors
diff --git a/Engine/Plugins/Animation/AnimationWarping/Source/Runtime/Public/BoneControllers/AnimNode_Steering.h b/Engine/Plugins/Animation/AnimationWarping/Source/Runtime/Public/BoneControllers/AnimNode_Steering.h
--- a/Engine/Plugins/Animation/AnimationWarping/Source/Runtime/Public/BoneControllers/AnimNode_Steering.h
+++ b/Engine/Plugins/Animation/AnimationWarping/Source/Runtime/Public/BoneControllers/AnimNode_Steering.h
@@ -69,6 +69,7 @@ struct ANIMATIONWARPINGRUNTIME_API FAnimNode_Steering : public FAnimNode_Skeleta
 
 	// FAnimNodeBase interface
 	virtual void Initialize_AnyThread(const FAnimationInitializeContext& Context) override;
+	virtual void GatherDebugData(FNodeDebugData& DebugData) override;
 	// End of FAnimNodeBase interface
 	
 	// FAnimNode_SkeletalControlBase interface
diff --git a/Engine/Plugins/Experimental/EvaluationNotifies/Source/Runtime/Private/EvaluationNotifies/AnimNode_EvaluationNotifies.cpp b/Engine/Plugins/Experimental/EvaluationNotifies/Source/Runtime/Private/EvaluationNotifies/AnimNode_EvaluationNotifies.cpp
--- a/Engine/Plugins/Experimental/EvaluationNotifies/Source/Runtime/Private/EvaluationNotifies/AnimNode_EvaluationNotifies.cpp
+++ b/Engine/Plugins/Experimental/EvaluationNotifies/Source/Runtime/Private/EvaluationNotifies/AnimNode_EvaluationNotifies.cpp
@@ -70,6 +70,20 @@ void FAnimNode_EvaluationNotifies::Initialize_AnyThread(const FAnimationInitiali
 	FAnimNode_SkeletalControlBase::Initialize_AnyThread(Context);
 }
 
+void FAnimNode_EvaluationNotifies::GatherDebugData(FNodeDebugData& DebugData)
+{
+	DECLARE_SCOPE_HIERARCHICAL_COUNTER_ANIMNODE(GatherDebugData)
+
+	FString DebugLine = DebugData.GetNodeName(this);
+	// Just tracking the animation asset for now. Tracking the map of notifies being evaluated would be better suited in rewind debugger.
+	DebugLine += FString::Printf(TEXT("(Animation: %s, Animation Time: %.3f"),
+		*CurrentAnimAsset.GetName(),
+		CurrentAnimAssetTime
+		);
+
+	ComponentPose.GatherDebugData(DebugData);
+}
+
 
 void FAnimNode_EvaluationNotifies::EvaluateSkeletalControl_AnyThread(FComponentSpacePoseContext& Output, TArray<FBoneTransform>& OutBoneTransforms)
 {
diff --git a/Engine/Plugins/Experimental/EvaluationNotifies/Source/Runtime/Public/EvaluationNotifies/AnimNode_EvaluationNotifies.h b/Engine/Plugins/Experimental/EvaluationNotifies/Source/Runtime/Public/EvaluationNotifies/AnimNode_EvaluationNotifies.h
--- a/Engine/Plugins/Experimental/EvaluationNotifies/Source/Runtime/Public/EvaluationNotifies/AnimNode_EvaluationNotifies.h
+++ b/Engine/Plugins/Experimental/EvaluationNotifies/Source/Runtime/Public/EvaluationNotifies/AnimNode_EvaluationNotifies.h
@@ -51,6 +51,7 @@ struct EVALUATIONNOTIFIESRUNTIME_API FAnimNode_EvaluationNotifies : public FAnim
 
 	// FAnimNodeBase interface
 	virtual void Initialize_AnyThread(const FAnimationInitializeContext& Context) override;
+	virtual void GatherDebugData(FNodeDebugData& DebugData) override;
 	// End of FAnimNodeBase interface
 	
 	// FAnimNode_SkeletalControlBase interface
diff --git a/Engine/Source/Runtime/Engine/Classes/Animation/AnimNode_DeadBlending.h b/Engine/Source/Runtime/Engine/Classes/Animation/AnimNode_DeadBlending.h
--- a/Engine/Source/Runtime/Engine/Classes/Animation/AnimNode_DeadBlending.h
+++ b/Engine/Source/Runtime/Engine/Classes/Animation/AnimNode_DeadBlending.h
@@ -167,7 +167,7 @@ public: // FAnimNode_Base
 	ENGINE_API virtual void CacheBones_AnyThread(const FAnimationCacheBonesContext& Context) override;
 	ENGINE_API virtual void Update_AnyThread(const FAnimationUpdateContext& Context) override;
 	ENGINE_API virtual void Evaluate_AnyThread(FPoseContext& Output) override;
-
+	ENGINE_API virtual void GatherDebugData(FNodeDebugData& DebugData) override;
 	ENGINE_API virtual bool NeedsDynamicReset() const override;
 
 private:
diff --git a/Engine/Source/Runtime/Engine/Private/Animation/AnimNode_DeadBlending.cpp b/Engine/Source/Runtime/Engine/Private/Animation/AnimNode_DeadBlending.cpp
--- a/Engine/Source/Runtime/Engine/Private/Animation/AnimNode_DeadBlending.cpp
+++ b/Engine/Source/Runtime/Engine/Private/Animation/AnimNode_DeadBlending.cpp
@@ -1051,6 +1051,39 @@ void FAnimNode_DeadBlending::Evaluate_AnyThread(FPoseContext& Output)
 	TRACE_ANIM_INERTIALIZATION(*Output.AnimInstanceProxy, GetNodeIndex(), InertializationWeight, FAnimTrace::EInertializationType::DeadBlending);
 }
 
+void FAnimNode_DeadBlending::GatherDebugData(FNodeDebugData& DebugData)
+{
+	DECLARE_SCOPE_HIERARCHICAL_COUNTER_ANIMNODE(GatherDebugData)
+
+	FString DebugLine = DebugData.GetNodeName(this);
+	
+	float InertializationWeight = InertializationState == EInertializationState::Active ?
+		1.0f - FAlphaBlend::AlphaToBlendOption(
+			FMath::Clamp(InertializationTime / FMath::Max(InertializationDuration, UE_SMALL_NUMBER), 0.0f, 1.0f),
+			InertializationBlendMode, InertializationCustomBlendCurve) : 0.0f;
+	
+	if (InertializationDuration > UE_KINDA_SMALL_NUMBER)
+	{
+		DebugLine += FString::Printf(TEXT("('%s' Time: %.3f / %.3f (%.0f%%) Weight: %.3f"),
+			*UEnum::GetValueAsString(InertializationState),
+			InertializationTime,
+			InertializationDuration,
+			100.0f * InertializationTime / InertializationDuration,
+			InertializationWeight);
+	}
+	else
+	{
+		DebugLine += FString::Printf(TEXT("('%s' Time: %.3f / %.3f) Weight: %.3f"),
+			*UEnum::GetValueAsString(InertializationState),
+			InertializationTime,
+			InertializationDuration,
+			InertializationWeight);
+	}
+
+	DebugData.AddDebugItem(DebugLine);
+	Source.GatherDebugData(DebugData);
+}
+
 bool FAnimNode_DeadBlending::NeedsDynamicReset() const
 {
 	return true;
