Update FSceneViewport to ensure that the VR PIE spectator view is properly updated to fit the preview window after it's resized.
Previously, the spectator view would remain at a fixed resolution, displaying only part of the view if the window was made smaller, or black bars along the bottom right if the window was made bigger. The actual size of the underlying render target (dependent on the HMD resolution) still won't be changed, just how it's displayed in the preview spectator window.

diff --git a/Engine/Source/Runtime/Engine/Private/Slate/SceneViewport.cpp b/Engine/Source/Runtime/Engine/Private/Slate/SceneViewport.cpp
--- a/Engine/Source/Runtime/Engine/Private/Slate/SceneViewport.cpp
+++ b/Engine/Source/Runtime/Engine/Private/Slate/SceneViewport.cpp
@@ -440,6 +440,19 @@ void FSceneViewport::OnDrawViewport( const FGeometry& AllottedGeometry, const FS
 				}
 			}
 		}
+
+		// Because DrawSize and GetSizeXY() are fixed for VR PIE after launch, we need to manually check if a viewport resize is necessary due to window resizes.
+		// This ensures the spectator view will be stretched to fit to the newly resized window without changing the size of the HMD render target.
+		else if (GIsEditor && IsStereoRenderingAllowed())
+		{
+			TSharedPtr<SWindow> Window = FSlateApplication::Get().FindWidgetWindow(ViewportWidget.Pin().ToSharedRef());
+			if (Window.IsValid())
+			{
+				ensure(PaintWindow == Window.Get());
+
+				FSlateApplication::Get().GetRenderer()->UpdateFullscreenState(Window.ToSharedRef());
+			}
+		}
 	}
 	
 }
