Skeletal Mesh: Fix race condition when triggering a skin weight profile change in the editor.

#rnx

diff --git a/Engine/Source/Runtime/Engine/Private/Animation/SkinWeightProfile.cpp b/Engine/Source/Runtime/Engine/Private/Animation/SkinWeightProfile.cpp
--- a/Engine/Source/Runtime/Engine/Private/Animation/SkinWeightProfile.cpp
+++ b/Engine/Source/Runtime/Engine/Private/Animation/SkinWeightProfile.cpp
@@ -10,9 +10,9 @@
 #include "UObject/AnimObjectVersion.h"
 #include "Engine/SkeletalMesh.h"
 #include "Engine/World.h"
+#include "Rendering/RenderCommandPipes.h"
 #include "Rendering/SkeletalMeshRenderData.h"
 #include "SkeletalMeshTypes.h"
-#include "UObject/AnimObjectVersion.h"
 #include "UObject/ObjectVersion.h"
 #include "UObject/UE5MainStreamObjectVersion.h"
 #include "UObject/UObjectIterator.h"
@@ -506,7 +506,8 @@ void FSkinWeightProfilesData::ReleaseBuffer(const FSkinWeightProfileStack& InPro
 		{
 			DEC_DWORD_STAT_BY(STAT_SkeletalMeshVertexMemory, Buffer->GetVertexDataSize());
 			ENQUEUE_RENDER_COMMAND(ReleaseSkinSkinWeightProfilesDataBufferCommand)(
-				[Buffer](FRHICommandListImmediate& RHICmdList)
+				UE::RenderCommandPipe::SkeletalMesh,
+				[Buffer](FRHICommandList& RHICmdList)
 			{			
 				Buffer->ReleaseResources();
 				delete Buffer;		
@@ -536,7 +537,8 @@ void FSkinWeightProfilesData::ReleaseResources()
 	if (Buffers.Num())
 	{
 		ENQUEUE_RENDER_COMMAND(ReleaseSkinSkinWeightProfilesDataBufferCommand)(
-			[Buffers](FRHICommandListImmediate& RHICmdList)
+			UE::RenderCommandPipe::SkeletalMesh,
+			[Buffers](FRHICommandList& RHICmdList)
 		{
 			for (FSkinWeightVertexBuffer* Buffer : Buffers)
 			{	
@@ -628,7 +630,8 @@ void FSkinWeightProfilesData::EnqueueGPUReadback()
 			static const FName ReadbackName = TEXT("ReadbackSkinWeightBuffer");
 			FSkinweightReadbackData* DataPtr = &ReadbackData;
 			ENQUEUE_RENDER_COMMAND(FSkinWeightProfilesData_EnqueueGPUReadback)(
-				[DataPtr, this](FRHICommandListImmediate& RHICmdList)
+				UE::RenderCommandPipe::SkeletalMesh,
+				[DataPtr, this](FRHICommandList& RHICmdList)
 			{
 				if (BaseBuffer->GetDataVertexBuffer()->VertexBufferRHI->GetSize())
 				{
@@ -647,7 +650,8 @@ void FSkinWeightProfilesData::EnqueueGPUReadback()
 			static const FName ReadbackName = TEXT("ReadbackSkinWeightBuffer");
 			FSkinweightReadbackData* DataPtr = &ReadbackData;
 			ENQUEUE_RENDER_COMMAND(FSkinWeightProfilesData_EnqueueGPUReadback)(
-				[DataPtr, this](FRHICommandListImmediate& RHICmdList)
+			UE::RenderCommandPipe::SkeletalMesh,
+				[DataPtr, this](FRHICommandList& RHICmdList)
 			{
 				DataPtr->BufferReadback.Reset(new FRHIGPUBufferReadback(ReadbackName));
 				if (DataPtr->BufferReadback.IsValid() && BaseBuffer->GetDataVertexBuffer()->VertexBufferRHI->GetSize())
@@ -674,7 +678,8 @@ void FSkinWeightProfilesData::EnqueueDataReadback()
 
 		FSkinweightReadbackData* DataPtr = &ReadbackData;
 		ENQUEUE_RENDER_COMMAND(FSkinWeightProfilesData_EnqueueDataReadback)(
-				[DataPtr](FRHICommandListImmediate& RHICmdList)
+			UE::RenderCommandPipe::SkeletalMesh,
+				[DataPtr](FRHICommandList& RHICmdList)
 		{
 			if (DataPtr && DataPtr->BufferReadback.IsValid())
 			{
