Cleanup SceneViewState cached UTexture2D (bloom and filmgrain) when cleaning world. This fixes GC leak due to rendering caching post fx textures

diff --git a/Engine/Source/Runtime/Engine/Private/LocalPlayer.cpp b/Engine/Source/Runtime/Engine/Private/LocalPlayer.cpp
--- a/Engine/Source/Runtime/Engine/Private/LocalPlayer.cpp
+++ b/Engine/Source/Runtime/Engine/Private/LocalPlayer.cpp
@@ -61,6 +61,12 @@ static FAutoConsoleVariableRef CVarCalcLocalPlayerCachedLODDistanceFactor(
 	TEXT("Should we calculate a LOD Distance Factor based on the current FOV.  Should not be necessary since LOD is already based on screen size.\n")
 	);
 
+static bool GLocalPlayerCleanUpSceneViewStateTransientReferences = true;
+static FAutoConsoleVariableRef CVarLocalPlayerCleanUpSceneViewStateTransientReferences(
+	TEXT("r.CleanUpSceneViewStateTransientReferences"),
+	GLocalPlayerCleanUpSceneViewStateTransientReferences,
+	TEXT("clean up scene view state transient references"));
+
 DECLARE_CYCLE_STAT(TEXT("CalcSceneView"), STAT_CalcSceneView, STATGROUP_Engine);
 
 //////////////////////////////////////////////////////////////////////////
@@ -1846,6 +1852,10 @@ void ULocalPlayer::CleanupViewState(FStringView MidParentRootPath /*= {}*/)
 		if (Ref)
 		{
 			Ref->ClearMIDPool(MidParentRootPath);
+			if (GLocalPlayerCleanUpSceneViewStateTransientReferences)
+			{
+				Ref->CleanUpTransientReferences();
+			}
 		}
 	}
 }
diff --git a/Engine/Source/Runtime/Engine/Public/SceneManagement.h b/Engine/Source/Runtime/Engine/Public/SceneManagement.h
--- a/Engine/Source/Runtime/Engine/Public/SceneManagement.h
+++ b/Engine/Source/Runtime/Engine/Public/SceneManagement.h
@@ -138,6 +138,7 @@ public:
 	virtual FSceneViewState* GetConcreteViewState () = 0;
 
 	virtual void AddReferencedObjects(FReferenceCollector& Collector) = 0;
+	virtual void CleanUpTransientReferences() = 0;
 
 	virtual SIZE_T GetSizeBytes() const { return 0; }
 
diff --git a/Engine/Source/Runtime/Renderer/Private/PostProcess/PostProcessFFTBloom.cpp b/Engine/Source/Runtime/Renderer/Private/PostProcess/PostProcessFFTBloom.cpp
--- a/Engine/Source/Runtime/Renderer/Private/PostProcess/PostProcessFFTBloom.cpp
+++ b/Engine/Source/Runtime/Renderer/Private/PostProcess/PostProcessFFTBloom.cpp
@@ -337,7 +337,7 @@ void InitDomainAndGetKernel(
 		FClearValueBinding::None,
 		TexCreate_ShaderResource | TexCreate_UAV);
 
-	if (ViewState && ViewState->BloomFFTKernel.Spectral && CVarBloomCacheKernel.GetValueOnRenderThread() != 0)
+	if (ViewState && ViewState->BloomFFTKernel.Physical != nullptr && ViewState->BloomFFTKernel.Spectral && CVarBloomCacheKernel.GetValueOnRenderThread() != 0)
 	{
 		auto& FFTKernel = ViewState->BloomFFTKernel;
 
diff --git a/Engine/Source/Runtime/Renderer/Private/PostProcess/PostProcessTonemap.cpp b/Engine/Source/Runtime/Renderer/Private/PostProcess/PostProcessTonemap.cpp
--- a/Engine/Source/Runtime/Renderer/Private/PostProcess/PostProcessTonemap.cpp
+++ b/Engine/Source/Runtime/Renderer/Private/PostProcess/PostProcessTonemap.cpp
@@ -727,7 +727,7 @@ FScreenPassTexture AddTonemapPass(FRDGBuilder& GraphBuilder, const FViewInfo& Vi
 
 
 		FRDGBufferRef FilmGrainConstantsBuffer;
-		if (View.ViewState && View.ViewState->FilmGrainCache.TextureRHI == FilmGrainTextureRHI && CVarFilmGrainCacheTextureConstants.GetValueOnRenderThread() != 0)
+		if (View.ViewState && View.ViewState->FilmGrainCache.Texture != nullptr && View.ViewState->FilmGrainCache.TextureRHI == FilmGrainTextureRHI && CVarFilmGrainCacheTextureConstants.GetValueOnRenderThread() != 0)
 		{
 			FilmGrainConstantsBuffer = GraphBuilder.RegisterExternalBuffer(View.ViewState->FilmGrainCache.ConstantsBuffer, TEXT("FilmGrain.TextureConstants"));
 		}
diff --git a/Engine/Source/Runtime/Renderer/Private/ScenePrivate.h b/Engine/Source/Runtime/Renderer/Private/ScenePrivate.h
--- a/Engine/Source/Runtime/Renderer/Private/ScenePrivate.h
+++ b/Engine/Source/Runtime/Renderer/Private/ScenePrivate.h
@@ -1462,6 +1462,16 @@ PRAGMA_ENABLE_DEPRECATION_WARNINGS
 		}
 	}
 
+	virtual void CleanUpTransientReferences()
+	{
+		ENQUEUE_RENDER_COMMAND(Cleanup)(
+			[this](FRHICommandList& RHICmdList)
+			{
+				BloomFFTKernel.Physical = nullptr;
+				FilmGrainCache.Texture = nullptr;
+			});
+	}
+
 	// needed for GetReusableMID()
 	virtual void OnStartPostProcessing(FSceneView& CurrentView) override
 	{
