Substrate - marked the advanced debug option as editor/win64/dx12 only for now until we make it general as for the other Substrate debug views.

#tests PS / DX12 and DX11 with substrate on/off.
#rnx

diff --git a/Engine/Source/Runtime/Engine/Classes/Engine/RendererSettings.h b/Engine/Source/Runtime/Engine/Classes/Engine/RendererSettings.h
--- a/Engine/Source/Runtime/Engine/Classes/Engine/RendererSettings.h
+++ b/Engine/Source/Runtime/Engine/Classes/Engine/RendererSettings.h
@@ -1222,8 +1222,8 @@ class URendererSettings : public UDeveloperSettings
 	"Enable advanced Substrate material debug visualization shaders. Base pas shaders can output such advanced data."
 	*/
 	UPROPERTY(config, EditAnywhere, Category = Substrate, meta = (
-		ConsoleVariable = "r.Substrate.Debug.AdvancedVisualizationShaders", DisplayName = "Substrate advanced visualization shaders",
-		ToolTip = "Enable advanced Substrate material debug visualization shaders. Base pass shaders can output such advanced data.",
+		ConsoleVariable = "r.Substrate.Debug.AdvancedVisualizationShaders", DisplayName = "Substrate advanced visualization shaders (Editor/Win64/DX12 Only)",
+		ToolTip = "Enable advanced Substrate material debug visualization shaders. Base pass shaders can output such advanced data. Only available for Editors built for Win64 and running with the DX12 graphic API for now.",
 		ConfigRestartRequired = true))
 		uint32 SubstrateDebugAdvancedVisualizationShaders : 1;
 
