Fix deep copy of MeshPaintTexture when duplicating the owning actor.
Deep copy was implemented for PostDuplicate() which handles component duplication. But it also needed implementing for PostEditImport() which handles actor duplications.

diff --git a/Engine/Source/Runtime/Engine/Private/Components/StaticMeshComponent.cpp b/Engine/Source/Runtime/Engine/Private/Components/StaticMeshComponent.cpp
--- a/Engine/Source/Runtime/Engine/Private/Components/StaticMeshComponent.cpp
+++ b/Engine/Source/Runtime/Engine/Private/Components/StaticMeshComponent.cpp
@@ -758,31 +758,47 @@ void UStaticMeshComponent::InitializeComponent()
 	Super::InitializeComponent();
 }
 
-void UStaticMeshComponent::PostDuplicate(bool bDuplicateForPIE)
+/** Call this after component duplication to perform the deep copy of the MeshPaintTexture. */
+static UTexture* RecreateMeshPaintTexture(UTexture* InTexture, UPrimitiveComponent* InComponent)
 {
-	if (!bDuplicateForPIE && MeshPaintTexture)
+	if (InTexture == nullptr || InComponent == nullptr)
 	{
-		UMeshPaintVirtualTexture* NewTexture = nullptr;
-		
-		FImage Image;
-		if (MeshPaintTexture->Source.GetMipImage(Image, 0))
-		{
-			NewTexture = NewObject<UMeshPaintVirtualTexture>(GetOutermost());
-			NewTexture->Source.Init(Image);
-			NewTexture->OwningComponent = MakeWeakObjectPtr(this);
-			NewTexture->UpdateResource();
-		}
+		return nullptr;
+	}
 
-		MeshPaintTexture = NewTexture;
+	UMeshPaintVirtualTexture* NewTexture = nullptr;
+	
+	FImage Image;
+	if (InTexture->Source.GetMipImage(Image, 0))
+	{
+		NewTexture = NewObject<UMeshPaintVirtualTexture>(InComponent->GetOutermost());
+		NewTexture->Source.Init(Image);
+		NewTexture->OwningComponent = MakeWeakObjectPtr(InComponent);
+		NewTexture->UpdateResource();
 	}
 
-	NotifyIfStaticMeshChanged();
+	return NewTexture;
+}
+
+void UStaticMeshComponent::PostDuplicate(bool bDuplicateForPIE)
+{
+ 	if (!bDuplicateForPIE && MeshPaintTexture != nullptr)
+ 	{
+		MeshPaintTexture = RecreateMeshPaintTexture(MeshPaintTexture, this);
+ 	}
+ 
+ 	NotifyIfStaticMeshChanged();
 
 	Super::PostDuplicate(bDuplicateForPIE);
 }
 
 void UStaticMeshComponent::PostEditImport()
 {
+	if (MeshPaintTexture != nullptr)
+	{
+		MeshPaintTexture = RecreateMeshPaintTexture(MeshPaintTexture, this);
+	}
+
 	NotifyIfStaticMeshChanged();
 
 	Super::PostEditImport();
