WorldPartition - Mark STAT_WorldPartitionAppendStreamingSourceInfo as verbose

There can be many of these scopes and they're typically very small and don't provide much value, rather they inflate the time of UpdateStreamingState.

They can be re-enabled with -verbosenamedevents

#rnx

diff --git a/Engine/Source/Runtime/Engine/Private/WorldPartition/WorldPartitionRuntimeCellData.cpp b/Engine/Source/Runtime/Engine/Private/WorldPartition/WorldPartitionRuntimeCellData.cpp
--- a/Engine/Source/Runtime/Engine/Private/WorldPartition/WorldPartitionRuntimeCellData.cpp
+++ b/Engine/Source/Runtime/Engine/Private/WorldPartition/WorldPartitionRuntimeCellData.cpp
@@ -46,7 +46,7 @@ void UWorldPartitionRuntimeCellData::ResetStreamingSourceInfo(const FWorldPartit
 	CachedSourceInfoEpoch = Context.GetUpdateStreamingStateEpoch();
 }
 
-DECLARE_CYCLE_STAT(TEXT("Append Streaming Source Info"), STAT_WorldPartitionAppendStreamingSourceInfo, STATGROUP_WorldPartition);
+DECLARE_CYCLE_STAT_WITH_FLAGS(TEXT("Append Streaming Source Info"), STAT_WorldPartitionAppendStreamingSourceInfo, STATGROUP_WorldPartition, EStatFlags::Verbose);
 void UWorldPartitionRuntimeCellData::AppendStreamingSourceInfo(const FWorldPartitionStreamingSource& Source, const FSphericalSector& SourceShape, const FWorldPartitionStreamingContext& Context) const
 {
 	SCOPE_CYCLE_COUNTER(STAT_WorldPartitionAppendStreamingSourceInfo);
