Anim BP: Fix race condition when copying sparse class data in base subsystem

diff --git a/Engine/Source/Runtime/Engine/Private/Animation/AnimSubsystem_Base.cpp b/Engine/Source/Runtime/Engine/Private/Animation/AnimSubsystem_Base.cpp
--- a/Engine/Source/Runtime/Engine/Private/Animation/AnimSubsystem_Base.cpp
+++ b/Engine/Source/Runtime/Engine/Private/Animation/AnimSubsystem_Base.cpp
@@ -44,6 +44,7 @@ FAnimSubsystem_Base& FAnimSubsystem_Base::operator =(const FAnimSubsystem_Base&
 {
 	// Need to perform a lock as copying can race with PatchValueHandlers in async loading thread-enabled builds
 	FScopeLock Lock(&ValueHandlersCriticalSection);
+	FScopeLock OtherLock(&Other.ValueHandlersCriticalSection);
 
 	ExposedValueHandlers = Other.ExposedValueHandlers;
 
diff --git a/Engine/Source/Runtime/Engine/Public/Animation/AnimSubsystem_Base.h b/Engine/Source/Runtime/Engine/Public/Animation/AnimSubsystem_Base.h
--- a/Engine/Source/Runtime/Engine/Public/Animation/AnimSubsystem_Base.h
+++ b/Engine/Source/Runtime/Engine/Public/Animation/AnimSubsystem_Base.h
@@ -28,6 +28,6 @@ struct FAnimSubsystem_Base : public FAnimSubsystem
 	ENGINE_API void PatchValueHandlers(UClass* InClass);
 
 private:
-	FCriticalSection ValueHandlersCriticalSection;
+	mutable FCriticalSection ValueHandlersCriticalSection;
 	TArray<FExposedValueHandler> ExposedValueHandlers;
 };
