HLOD Outliner Incorrect Unbuilt State with bSaveLODActorsToHLODPackage (Legacy HLODs)
* Remove the name of the first sub actor from the CRC computation in UHLODProxy::GenerateKeyForActor (name does not contribute to the final HLOD result anyway)

diff --git a/Engine/Source/Runtime/Engine/Private/HLODProxy.cpp b/Engine/Source/Runtime/Engine/Private/HLODProxy.cpp
--- a/Engine/Source/Runtime/Engine/Private/HLODProxy.cpp
+++ b/Engine/Source/Runtime/Engine/Private/HLODProxy.cpp
@@ -607,11 +607,6 @@ FName UHLODProxy::GenerateKeyForActor(const ALODActor* LODActor, bool bMustUndoL
 	// We get the CRC of the first actor name and various static mesh components
 	{
 		uint32 CRC = 0;
-		const AActor* FirstActor = FindFirstActor(LODActor);
-		if(FirstActor)
-		{
-			CRC = FCrc::StrCrc32(*FirstActor->GetName(), CRC);
-		}
 
 		TArray<UPrimitiveComponent*> Components;
 		ExtractComponents(LODActor, Components);
