Fix crash in FRecastTileGenerator::AddReferencedObjects by removing invalid check in FNavigationElement::GetWeakUObject.

diff --git a/Engine/Source/Runtime/Engine/Classes/AI/Navigation/NavigationElement.h b/Engine/Source/Runtime/Engine/Classes/AI/Navigation/NavigationElement.h
--- a/Engine/Source/Runtime/Engine/Classes/AI/Navigation/NavigationElement.h
+++ b/Engine/Source/Runtime/Engine/Classes/AI/Navigation/NavigationElement.h
@@ -371,7 +371,6 @@ bool FNavigationElementHandle::IsValid() const
 
 TWeakObjectPtr<const UObject> FNavigationElement::GetWeakUObject() const
 {
-	checkf(IsInGameThread(), TEXT("Currenlty not allowing this to be called from non-game thread."));
 	return OwnerUObject;
 }
 
