TEDS-Outliner: Revert change to call PostEditorChange when setting editor visibility on an actor in favor of doing it in the Teds Outliner locally

diff --git a/Engine/Source/Runtime/Engine/Private/ActorEditor.cpp b/Engine/Source/Runtime/Engine/Private/ActorEditor.cpp
--- a/Engine/Source/Runtime/Engine/Private/ActorEditor.cpp
+++ b/Engine/Source/Runtime/Engine/Private/ActorEditor.cpp
@@ -983,9 +983,6 @@ void AActor::SetIsTemporarilyHiddenInEditor( bool bIsHidden )
 	{
 		bHiddenEdTemporary = bIsHidden;
 		MarkComponentsRenderStateDirty();
-		
-		FPropertyChangedEvent HiddenEdTemporaryChanged(AActor::StaticClass()->FindPropertyByName(GET_MEMBER_NAME_CHECKED(AActor, bHiddenEdTemporary)), EPropertyChangeType::ValueSet);
-		PostEditChangeProperty(HiddenEdTemporaryChanged);
 	}
 }
 
