Cooker: StaticMesh: Fix soft lock in the cook if the DDC key changes during cook. Log an error and use the old key. For other unknown errors that lead to repeated failures of Cache, add error detection and fail the cook of the staticmesh.
#rnx

diff --git a/Engine/Source/Runtime/Engine/Private/StaticMesh.cpp b/Engine/Source/Runtime/Engine/Private/StaticMesh.cpp
--- a/Engine/Source/Runtime/Engine/Private/StaticMesh.cpp
+++ b/Engine/Source/Runtime/Engine/Private/StaticMesh.cpp
@@ -3735,7 +3735,22 @@ void FStaticMeshRenderData::Cache(const ITargetPlatform* TargetPlatform, UStatic
 		UE_TRACE_LOG_SCOPED_T(Cpu, StaticMeshCache, CpuChannel)
 			<< StaticMeshCache.KeySuffix(*KeySuffix);
 #endif
-		DerivedDataKey = BuildStaticMeshDerivedDataKey(KeySuffix);
+		FString NewDerivedDataKey = BuildStaticMeshDerivedDataKey(KeySuffix);
+		if (NumCacheAttempts == 0 || !IsRunningCookCommandlet())
+		{
+			DerivedDataKey = MoveTemp(NewDerivedDataKey);
+		}
+		else if (DerivedDataKey != NewDerivedDataKey)
+		{
+			UE_LOG(LogStaticMesh, Error,
+				TEXT("Indeterminism in FStaticMeshRenderData::Cache call to BuildStaticMeshDerivedDataKey. Using the previously computed RenderData even though the key is now different.")
+				TEXT("\n\tStaticMesh=%s, Platform=%s")
+				TEXT("\n\tPreviousKey=%s")
+				TEXT("\n\tCurrentKey =%s"),
+				*Owner->GetPathName(), *TargetPlatform->PlatformName(),
+				*DerivedDataKey, *NewDerivedDataKey);
+		}
+		++NumCacheAttempts;
 
 		using namespace UE;
 		using namespace UE::DerivedData;
@@ -5871,8 +5886,26 @@ FStaticMeshRenderData& UStaticMesh::GetPlatformStaticMeshRenderData(UStaticMesh*
 		return *PlatformRenderData;
 	}
 
-	while (PlatformRenderData && PlatformRenderData->DerivedDataKey != PlatformDerivedDataKey)
+	while (PlatformRenderData)
 	{
+		if (PlatformRenderData->DerivedDataKey == PlatformDerivedDataKey)
+		{
+			PlatformRenderData->SupportedPlatforms.Add(Platform);
+			break;
+		}
+
+		if (PlatformRenderData->SupportedPlatforms.Contains(Platform) && IsRunningCookCommandlet())
+		{
+			UE_LOG(LogStaticMesh, Error,
+				TEXT("Indeterminism in GetPlatformStaticMeshRenderData call to BuildStaticMeshDerivedDataKey. Using the previously computed RenderData even though the key is now different.")
+				TEXT("\n\tStaticMesh=%s, Platform=%s")
+				TEXT("\n\tPreviousKey=%s")
+				TEXT("\n\tCurrentKey =%s"),
+				*Mesh->GetPathName(), *Platform->PlatformName(),
+				*PlatformRenderData->DerivedDataKey, *PlatformDerivedDataKey);
+			break;
+		}
+
 		PlatformRenderData = PlatformRenderData->NextCachedRenderData.Get();
 	}
 
@@ -5882,6 +5915,7 @@ FStaticMeshRenderData& UStaticMesh::GetPlatformStaticMeshRenderData(UStaticMesh*
 		PlatformRenderData = new FStaticMeshRenderData();
 		PlatformRenderData->Cache(Platform, Mesh, PlatformLODSettings);
 		check(PlatformRenderData->DerivedDataKey == PlatformDerivedDataKey);
+		PlatformRenderData->SupportedPlatforms.Add(Platform);
 		Swap(PlatformRenderData->NextCachedRenderData, Mesh->GetRenderData()->NextCachedRenderData);
 		Mesh->GetRenderData()->NextCachedRenderData = TUniquePtr<FStaticMeshRenderData>(PlatformRenderData);
 	}
@@ -5909,6 +5943,7 @@ void UStaticMesh::ClearAllCachedCookedPlatformData()
 {
 	GetRenderData()->NextCachedRenderData.Reset();
 	GetRenderData()->NaniteResourcesPtr->DropBulkData();
+	GetRenderData()->NumCacheAttempts = 0;
 }
 
 void UStaticMesh::BeginCacheForCookedPlatformData(const ITargetPlatform* TargetPlatform)
@@ -5932,12 +5967,26 @@ bool UStaticMesh::IsCachedCookedPlatformDataLoaded(const ITargetPlatform* Target
 
 	if (bFailed)
 	{
+		if (PlatformRenderData.NumCacheAttempts >= 2)
+		{
+			UE_LOG(LogStaticMesh, Error,
+				TEXT("Failed to recover Nanite streaming from DDC for '%s', even after multiple attempts. ")
+				TEXT("Abandoning the attempt to build the static mesh renderdata; the mesh will be broken at runtime."), *GetPathName());
+			return true;
+		}
 		UE_LOG(LogStaticMesh, Log, TEXT("Failed to recover Nanite streaming from DDC for '%s'. Rebuilding and retrying."), *GetPathName());
 
 		// This should be a very rare event
-		// For simplicity, just rebuild the entire RenderData
+		// For simplicity, just rebuild the entire RenderData, but preserve the CacheRetryData
+		// TODO: Verify that move construction works for FPlatformRenderData, and use the move constructor rather than
+		// calling CopyCacheRetryData twice..
+		FStaticMeshRenderData CopyData;
+		CopyData.CopyCacheRetryData(PlatformRenderData);
+
 		PlatformRenderData.~FStaticMeshRenderData();
 		new (&PlatformRenderData) FStaticMeshRenderData();
+
+		PlatformRenderData.CopyCacheRetryData(CopyData);
 		if (PlatformRenderData.NaniteResourcesPtr.IsValid())
 		{
 			PlatformRenderData.NaniteResourcesPtr->SetHasBuildFromDDCError(true);
@@ -5949,6 +5998,15 @@ bool UStaticMesh::IsCachedCookedPlatformDataLoaded(const ITargetPlatform* Target
 	return true;
 }
 
+#if WITH_EDITORONLY_DATA
+void FStaticMeshRenderData::CopyCacheRetryData(FStaticMeshRenderData& Other)
+{
+	SupportedPlatforms = Other.SupportedPlatforms;
+	DerivedDataKey = Other.DerivedDataKey;
+	NumCacheAttempts = Other.NumCacheAttempts;
+}
+#endif
+
 #if WITH_EDITORONLY_DATA
 
 bool UStaticMesh::LoadMeshDescription(int32 LodIndex, FMeshDescription& OutMeshDescription) const
diff --git a/Engine/Source/Runtime/Engine/Public/StaticMeshResources.h b/Engine/Source/Runtime/Engine/Public/StaticMeshResources.h
--- a/Engine/Source/Runtime/Engine/Public/StaticMeshResources.h
+++ b/Engine/Source/Runtime/Engine/Public/StaticMeshResources.h
@@ -8,6 +8,7 @@
 
 #include "CoreMinimal.h"
 #include "Containers/IndirectArray.h"
+#include "Containers/Set.h"
 #include "Misc/Guid.h"
 #include "Engine/EngineTypes.h"
 #include "UObject/UObjectIterator.h"
@@ -810,6 +811,14 @@ public:
 	/** The derived data key associated with this render data. */
 	FString DerivedDataKey;
 
+	/**
+	 * The platform(s) that are using using this RenderData. There may be more than one if multiple
+	 * platforms share the same key.
+	 */
+	TSet<const ITargetPlatform*> SupportedPlatforms;
+	/** The number of times Cache has been called. Used to log errors when we are calling it repeatedly. */
+	int32 NumCacheAttempts = 0;
+
 	/** Map of material index -> original material index at import time. */
 	TArray<int32> MaterialIndexToImportIndex;
 
@@ -841,6 +850,11 @@ public:
 	 * level of detail settings.
 	 */
 	ENGINE_API void Cache(const ITargetPlatform* TargetPlatform, UStaticMesh* Owner, const FStaticMeshLODSettings& LODSettings);
+	/**
+	 * Copies data from an earlier FStaticMeshRenderData for the same platform which should persist across
+	 * retries; the persistent data is used to track repeated failures and error exit.
+	 */
+	ENGINE_API void CopyCacheRetryData(FStaticMeshRenderData& Other);
 #endif // #if WITH_EDITORONLY_DATA
 
 	/** Count the number of LODs that are inlined and not streamable. Starting from the last LOD and stopping at the first non inlined LOD. */
