Moved PeriodicWorldPosition calculation into TransformBase in the material translator. This removes the special path that was taken before and makes it easier to extend the code.

diff --git a/Engine/Plugins/Enterprise/GLTFExporter/Source/GLTFExporter/Private/Converters/GLTFProxyMaterialCompiler.h b/Engine/Plugins/Enterprise/GLTFExporter/Source/GLTFExporter/Private/Converters/GLTFProxyMaterialCompiler.h
--- a/Engine/Plugins/Enterprise/GLTFExporter/Source/GLTFExporter/Private/Converters/GLTFProxyMaterialCompiler.h
+++ b/Engine/Plugins/Enterprise/GLTFExporter/Source/GLTFExporter/Private/Converters/GLTFProxyMaterialCompiler.h
@@ -97,16 +97,16 @@ public:
 		return Compiler->CustomPrimitiveData(OutputIndex, Type);
 	}
 
-	virtual int32 TransformVector(EMaterialCommonBasis SourceCoordBasis, EMaterialCommonBasis DestCoordBasis, int32 A) override
+	virtual int32 TransformVector(EMaterialCommonBasis SourceCoordBasis, EMaterialCommonBasis DestCoordBasis, FTransformParameters& Parameters, int32 A) override
 	{
 		if (SourceCoordBasis != DestCoordBasis) bUsesTransformVector = true; // TODO: optimize later by ignoring some irrelevant basis (like MCB_View and MCB_Camera)
-		return Compiler->TransformVector(SourceCoordBasis, DestCoordBasis, A);
+		return Compiler->TransformVector(SourceCoordBasis, DestCoordBasis, Parameters, A);
 	}
 
-	virtual int32 TransformPosition(EMaterialCommonBasis SourceCoordBasis, EMaterialCommonBasis DestCoordBasis, int32 A) override
+	virtual int32 TransformPosition(EMaterialCommonBasis SourceCoordBasis, EMaterialCommonBasis DestCoordBasis, FTransformParameters& Parameters, int32 A) override
 	{
 		if (SourceCoordBasis != DestCoordBasis) bUsesTransformPosition = true; // TODO: optimize later by ignoring some irrelevant basis (like MCB_View and MCB_Camera)
-		return Compiler->TransformPosition(SourceCoordBasis, DestCoordBasis, A);
+		return Compiler->TransformPosition(SourceCoordBasis, DestCoordBasis, Parameters, A);
 	}
 
 	virtual int32 GIReplace(int32 Direct, int32 StaticIndirect, int32 DynamicIndirect) override
diff --git a/Engine/Source/Editor/UnrealEd/Private/Lightmass/LightmassRender.cpp b/Engine/Source/Editor/UnrealEd/Private/Lightmass/LightmassRender.cpp
--- a/Engine/Source/Editor/UnrealEd/Private/Lightmass/LightmassRender.cpp
+++ b/Engine/Source/Editor/UnrealEd/Private/Lightmass/LightmassRender.cpp
@@ -177,13 +177,13 @@ struct FLightmassMaterialCompiler : public FProxyMaterialCompiler
 		return Compiler->Constant3(0.0f,0.0f,-1.0f);
 	}
 
-	virtual int32 TransformVector(EMaterialCommonBasis SourceCoordBasis, EMaterialCommonBasis DestCoordBasis, int32 A) override
+	virtual int32 TransformVector(EMaterialCommonBasis SourceCoordBasis, EMaterialCommonBasis DestCoordBasis, FTransformParameters& Parameters, int32 A) override
 	{
 		//UE_LOG(LogLightmassRender, Log, TEXT("Lightmass material compiler has encountered TransformVector... Passing thru source vector untouched."));
 		return A;
 	}
 
-	virtual int32 TransformPosition(EMaterialCommonBasis SourceCoordBasis, EMaterialCommonBasis DestCoordBasis, int32 A) override
+	virtual int32 TransformPosition(EMaterialCommonBasis SourceCoordBasis, EMaterialCommonBasis DestCoordBasis, FTransformParameters& Parameters, int32 A) override
 	{
 		//UE_LOG(LogLightmassRender, Log, TEXT("Lightmass material compiler has encountered TransformPosition... Passing thru source vector untouched."));
 		return A;
diff --git a/Engine/Source/Runtime/Engine/Private/Materials/HLSLMaterialTranslator.cpp b/Engine/Source/Runtime/Engine/Private/Materials/HLSLMaterialTranslator.cpp
--- a/Engine/Source/Runtime/Engine/Private/Materials/HLSLMaterialTranslator.cpp
+++ b/Engine/Source/Runtime/Engine/Private/Materials/HLSLMaterialTranslator.cpp
@@ -6286,7 +6286,8 @@ int32 FHLSLMaterialTranslator::LocalPosition(EPositionIncludedOffsets IncludedOf
 	{
 		// Generate derivative based on WorldPosition_DDX
 		int32 WorldPositionCode = WorldPosition(IncludedOffsets == EPositionIncludedOffsets::IncludeOffsets ? WPT_CameraRelative : WPT_CameraRelativeNoOffsets);
-		int32 TransformedPositionCode = TransformPosition(MCB_TranslatedWorld, OriginType == ELocalPositionOrigin::Instance ? MCB_Instance : MCB_Local, WorldPositionCode);
+		FTransformParameters Parameters{};
+		int32 TransformedPositionCode = TransformPosition(MCB_TranslatedWorld, OriginType == ELocalPositionOrigin::Instance ? MCB_Instance : MCB_Local, Parameters, WorldPositionCode);
 		if (TransformedPositionCode < 0) // TransformPosition failed
 		{
 			return TransformedPositionCode;
@@ -10821,7 +10822,7 @@ static FString LWCMultiplyMatrix(const TCHAR* Vector, const TCHAR* Matrix, int A
 	}
 }
 
-int32 FHLSLMaterialTranslator::TransformBase(EMaterialCommonBasis SourceCoordBasis, EMaterialCommonBasis DestCoordBasis, int32 A, int AWComponent)
+int32 FHLSLMaterialTranslator::TransformBase(EMaterialCommonBasis SourceCoordBasis, EMaterialCommonBasis DestCoordBasis, FTransformParameters& Parameters, int32 A, int AWComponent)
 {
 	if (A == INDEX_NONE)
 	{
@@ -10903,6 +10904,10 @@ int32 FHLSLMaterialTranslator::TransformBase(EMaterialCommonBasis SourceCoordBas
 					CodeDerivStr = MultiplyTranslatedMatrix(TEXT("<A>"), TEXT("Get<PREV>LocalToWorldDF(Parameters)"), 0, bCompilingPreviousFrame);
 				}
 			}
+			else if (DestCoordBasis == MCB_PeriodicWorld)
+			{
+				IntermediaryBasis = MCB_TranslatedWorld;
+			}
 			// else use MCB_World as intermediary basis
 			break;
 		}
@@ -10964,6 +10969,16 @@ int32 FHLSLMaterialTranslator::TransformBase(EMaterialCommonBasis SourceCoordBas
 				CodeDerivStr = MultiplyTranslatedMatrix(TEXT("<A>"), TEXT("GetWorldToInstanceDF(Parameters)"), 0, bCompilingPreviousFrame);
 				bUsesInstanceWorldToLocalPS |= ShaderFrequency == SF_Pixel;
 			}
+			else if (DestCoordBasis == MCB_PeriodicWorld)
+			{
+				if (Parameters.PeriodicWorldTileSizeIndex == INDEX_NONE)
+				{
+					return Errorf(TEXT("Missing periodic world tile size"));
+				}
+				int32 PeriodicWorldOrigin = CalculatePeriodicWorldPositionOrigin(Parameters.PeriodicWorldTileSizeIndex);
+				CodeStr = FString::Printf(TEXT("(<A> - (%s))"), *GetParameterCode(PeriodicWorldOrigin));
+				CodeDerivStr = TEXT("<A>");
+			}
 			// else use MCB_World as intermediary basis
 			break;
 		}
@@ -11059,6 +11074,10 @@ int32 FHLSLMaterialTranslator::TransformBase(EMaterialCommonBasis SourceCoordBas
 				CodeDerivStr = MultiplyTranslatedMatrix(TEXT("<A>"), TEXT("Parameters.Particle.ParticleToWorld"), 0, bCompilingPreviousFrame);
 				bUsesParticleLocalToWorld = true;
 			}
+			else if (DestCoordBasis == MCB_PeriodicWorld)
+			{
+				IntermediaryBasis = MCB_TranslatedWorld;
+			}
 			// use World as an intermediary base
 			break;
 		}
@@ -11077,9 +11096,30 @@ int32 FHLSLMaterialTranslator::TransformBase(EMaterialCommonBasis SourceCoordBas
 				CodeDerivStr = MultiplyTranslatedMatrix(TEXT("<A>"), TEXT("Get<PREV>InstanceToWorldDF(Parameters)"), 0, bCompilingPreviousFrame);
 				bUsesInstanceLocalToWorldPS |= ShaderFrequency == SF_Pixel;
 			}
+			else if (DestCoordBasis == MCB_PeriodicWorld)
+			{
+				IntermediaryBasis = MCB_TranslatedWorld;
+			}
 			// use World as an intermediary base
 			break;
 		}
+		case MCB_PeriodicWorld:
+		{
+			if (DestCoordBasis == MCB_TranslatedWorld)
+			{
+				if (Parameters.PeriodicWorldTileSizeIndex == INDEX_NONE)
+				{
+					return Errorf(TEXT("Missing periodic world tile size"));
+				}
+				int32 PeriodicWorldOrigin = CalculatePeriodicWorldPositionOrigin(Parameters.PeriodicWorldTileSizeIndex);
+				CodeStr = FString::Printf(TEXT("(<A> + (%s))"), *GetParameterCode(PeriodicWorldOrigin));
+				CodeDerivStr = TEXT("<A>");
+			}
+
+			// else use MCB_TranslatedWorld as intermediary basis
+			IntermediaryBasis = MCB_TranslatedWorld;
+			break;
+		}
 
 		default:
 			check(0);
@@ -11093,9 +11133,9 @@ int32 FHLSLMaterialTranslator::TransformBase(EMaterialCommonBasis SourceCoordBas
 		check(IntermediaryBasis != DestCoordBasis);
 
 		// use intermediary basis
-		const int32 IntermediaryA = TransformBase(SourceCoordBasis, IntermediaryBasis, A, AWComponent);
+		const int32 IntermediaryA = TransformBase(SourceCoordBasis, IntermediaryBasis, Parameters, A, AWComponent);
 
-		return TransformBase(IntermediaryBasis, DestCoordBasis, IntermediaryA, AWComponent);
+		return TransformBase(IntermediaryBasis, DestCoordBasis, Parameters, IntermediaryA, AWComponent);
 	}
 		
 	if (bIsPositionTransform)
@@ -11168,101 +11208,27 @@ int32 FHLSLMaterialTranslator::TransformBase(EMaterialCommonBasis SourceCoordBas
 	return CastToNonLWCIfDisabled(Result);
 }
 
-int32 FHLSLMaterialTranslator::TransformVector(EMaterialCommonBasis SourceCoordBasis, EMaterialCommonBasis DestCoordBasis, int32 A)
+int32 FHLSLMaterialTranslator::TransformVector(EMaterialCommonBasis SourceCoordBasis, EMaterialCommonBasis DestCoordBasis, FTransformParameters& Parameters, int32 A)
 {
-	return TransformBase(SourceCoordBasis, DestCoordBasis, A, 0);
+	return TransformBase(SourceCoordBasis, DestCoordBasis, Parameters, A, 0);
 }
 
-int32 FHLSLMaterialTranslator::TransformPosition(EMaterialCommonBasis SourceCoordBasis, EMaterialCommonBasis DestCoordBasis, int32 A)
+int32 FHLSLMaterialTranslator::TransformPosition(EMaterialCommonBasis SourceCoordBasis, EMaterialCommonBasis DestCoordBasis, FTransformParameters& Parameters, int32 A)
 {
-	return TransformBase(SourceCoordBasis, DestCoordBasis, A, 1);
+	return TransformBase(SourceCoordBasis, DestCoordBasis, Parameters, A, 1);
 }
 
-int32 FHLSLMaterialTranslator::CalculatePeriodicWorldPositionOrigin(int TileScaleIdx)
+int32 FHLSLMaterialTranslator::CalculatePeriodicWorldPositionOrigin(int TileScaleIndex)
 {
 	// float3 PeriodicWorldOrigin = GetPeriodicWorldOrigin(Scale);
 	FString PeriodicWorldOrigin = CreateSymbolName(TEXT("PeriodicWorldOrigin"));
-	FString FuncName = IsConstFloatOfPow2Expression(TileScaleIdx) ? TEXT("GetPeriodicWorldOrigin_Pow2") : TEXT("GetPeriodicWorldOrigin");
-	FString Code = FString::Format(TEXT("	float3 {0} = {1}({2});{3}"), { PeriodicWorldOrigin, FuncName, *GetParameterCode(TileScaleIdx), HLSL_LINE_TERMINATOR });
+	FString FuncName = IsConstFloatOfPow2Expression(TileScaleIndex) ? TEXT("GetPeriodicWorldOrigin_Pow2") : TEXT("GetPeriodicWorldOrigin");
+	FString Code = FString::Format(TEXT("	float3 {0} = {1}({2});{3}"), { PeriodicWorldOrigin, FuncName, *GetParameterCode(TileScaleIndex), HLSL_LINE_TERMINATOR });
 
 	AddInlinedCodeChunk(EMaterialValueType::MCT_VoidStatement, *Code);
 	return AddCodeChunk(MCT_Float3, *PeriodicWorldOrigin);
 }
 
-int32 FHLSLMaterialTranslator::TransformFromPeriodicWorldPosition(EMaterialCommonBasis DestCoordBasis, int TileScaleIdx, int32 PeriodicWorldPosition)
-{
-	const EDerivativeStatus InputDerivStatus = GetDerivativeStatus(PeriodicWorldPosition);
-
-	int32 PeriodicWorldOrigin = CalculatePeriodicWorldPositionOrigin(TileScaleIdx);
-	FString TranslatedWorldPositionCode = FString::Format(TEXT("(({0}) + ({1}))"), { GetParameterCode(PeriodicWorldPosition), GetParameterCode(PeriodicWorldOrigin) });
-
-	int32 Result;
-	if (IsAnalyticDerivEnabled() && IsDerivativeValid(InputDerivStatus))
-	{
-		if (InputDerivStatus == EDerivativeStatus::Valid)
-		{
-			FString PositionDerivCode = GetParameterCodeDeriv(PeriodicWorldPosition, CompiledPDV_Analytic);
-
-			FString AnalyticCode = DerivativeAutogen.ConstructDeriv(
-				*TranslatedWorldPositionCode, 
-				*(PositionDerivCode + TEXT(".Ddx")), 
-				*(PositionDerivCode + TEXT(".Ddy")), 
-				GetDerivType(MCT_Float3));
-			Result = AddCodeChunkInnerDeriv(*TranslatedWorldPositionCode, *AnalyticCode, MCT_Float3, false, EDerivativeStatus::Valid);
-		}
-		else
-		{
-			Result = AddCodeChunkZeroDeriv(MCT_Float3, *TranslatedWorldPositionCode);
-		}
-	}
-	else
-	{
-		Result = AddCodeChunk(MCT_Float3, *TranslatedWorldPositionCode);
-	}
-
-	return TransformBase(EMaterialCommonBasis::MCB_TranslatedWorld, DestCoordBasis, Result, 1);
-}
-
-int32 FHLSLMaterialTranslator::TransformToPeriodicWorldPosition(EMaterialCommonBasis SourceCoordBasis, int TileScaleIdx, int32 Input)
-{
-	int32 TranslatedWorldPos = TransformBase(SourceCoordBasis, EMaterialCommonBasis::MCB_TranslatedWorld, Input, 1);
-	if (TranslatedWorldPos == INDEX_NONE) // Transform failed
-	{
-		return INDEX_NONE;
-	}
-
-	const EDerivativeStatus InputDerivStatus = GetDerivativeStatus(TranslatedWorldPos);
-
-	int32 PeriodicWorldOrigin = CalculatePeriodicWorldPositionOrigin(TileScaleIdx);
-	FString PeriodicWorldPositionCode = FString::Format(TEXT("(({0}) - ({1}))"), { GetParameterCode(TranslatedWorldPos), GetParameterCode(PeriodicWorldOrigin) });
-
-	int32 Result;
-	if (IsAnalyticDerivEnabled() && IsDerivativeValid(InputDerivStatus))
-	{
-		if (InputDerivStatus == EDerivativeStatus::Valid)
-		{
-			FString TranslatedWorldPosDerivCode = GetParameterCodeDeriv(TranslatedWorldPos, CompiledPDV_Analytic);
-
-			FString AnalyticCode = DerivativeAutogen.ConstructDeriv(
-				*PeriodicWorldPositionCode, 
-				*(TranslatedWorldPosDerivCode + TEXT(".Ddx")), 
-				*(TranslatedWorldPosDerivCode + TEXT(".Ddy")), 
-				GetDerivType(MCT_Float3));
-			Result = AddCodeChunkInnerDeriv(*PeriodicWorldPositionCode, *AnalyticCode, MCT_Float3, false, EDerivativeStatus::Valid);
-		}
-		else
-		{
-			Result = AddCodeChunkZeroDeriv(MCT_Float3, *PeriodicWorldPositionCode);
-		}
-	}
-	else
-	{
-		Result = AddCodeChunk(MCT_Float3, *PeriodicWorldPositionCode);
-	}
-
-	return Result;
-}
-
 int32 FHLSLMaterialTranslator::TransformNormalFromRequestedBasisToWorld(int32 NormalCodeChunk)
 {
 	if (IsTangentSpaceNormal())
diff --git a/Engine/Source/Runtime/Engine/Private/Materials/HLSLMaterialTranslator.h b/Engine/Source/Runtime/Engine/Private/Materials/HLSLMaterialTranslator.h
--- a/Engine/Source/Runtime/Engine/Private/Materials/HLSLMaterialTranslator.h
+++ b/Engine/Source/Runtime/Engine/Private/Materials/HLSLMaterialTranslator.h
@@ -1161,13 +1161,24 @@ protected:
 	virtual int32 HsvToRgb(int32 X) override;
 	virtual int32 RgbToHsv(int32 X) override;
 
-	int32 TransformBase(EMaterialCommonBasis SourceCoordBasis, EMaterialCommonBasis DestCoordBasis, int32 A, int AWComponent);
+	int32 TransformBase(EMaterialCommonBasis SourceCoordBasis, EMaterialCommonBasis DestCoordBasis, FTransformParameters& Parameters, int32 A, int AWComponent);
 	
-	virtual int32 TransformVector(EMaterialCommonBasis SourceCoordBasis, EMaterialCommonBasis DestCoordBasis, int32 A) override;
-	virtual int32 TransformPosition(EMaterialCommonBasis SourceCoordBasis, EMaterialCommonBasis DestCoordBasis, int32 A) override;
-	int32 CalculatePeriodicWorldPositionOrigin(int TileScaleIdx);
-	virtual int32 TransformFromPeriodicWorldPosition(EMaterialCommonBasis DestCoordBasis, int TileScaleIdx, int32 A) override;
-	virtual int32 TransformToPeriodicWorldPosition(EMaterialCommonBasis SourceCoordBasis, int TileScaleIdx, int32 A) override;
+	virtual int32 TransformVector(EMaterialCommonBasis SourceCoordBasis, EMaterialCommonBasis DestCoordBasis, FTransformParameters& Parameters, int32 A) override;
+	virtual int32 TransformPosition(EMaterialCommonBasis SourceCoordBasis, EMaterialCommonBasis DestCoordBasis, FTransformParameters& Parameters, int32 A) override;
+	
+	int32 TransformVector(EMaterialCommonBasis SourceCoordBasis, EMaterialCommonBasis DestCoordBasis, int32 A)
+	{
+		FTransformParameters Parameters{};
+		return TransformVector(SourceCoordBasis, DestCoordBasis, Parameters, A);
+	}
+	
+	int32 TransformPosition(EMaterialCommonBasis SourceCoordBasis, EMaterialCommonBasis DestCoordBasis, int32 A)
+	{
+		FTransformParameters Parameters{};
+		return TransformPosition(SourceCoordBasis, DestCoordBasis, Parameters, A);
+	}
+	
+	int32 CalculatePeriodicWorldPositionOrigin(int TileScaleIndex);
 	virtual int32 TransformNormalFromRequestedBasisToWorld(int32 NormalCodeChunk) override;
 	virtual int32 DynamicParameter(FLinearColor& DefaultValue, uint32 ParameterIndex = 0) override;
 	virtual int32 LightmapUVs() override;
diff --git a/Engine/Source/Runtime/Engine/Private/Materials/MaterialExpressions.cpp b/Engine/Source/Runtime/Engine/Private/Materials/MaterialExpressions.cpp
--- a/Engine/Source/Runtime/Engine/Private/Materials/MaterialExpressions.cpp
+++ b/Engine/Source/Runtime/Engine/Private/Materials/MaterialExpressions.cpp
@@ -13665,7 +13665,7 @@ static EMaterialCommonBasis GetMaterialCommonBasis(EMaterialPositionTransformSou
 	static const EMaterialCommonBasis ConversionTable[TRANSFORMPOSSOURCE_MAX] = {
 		MCB_Local,						// TRANSFORMPOSSOURCE_Local
 		MCB_World,						// TRANSFORMPOSSOURCE_World
-		MCB_MAX,						// No match for TRANSFORMPOSSOURCE_PeriodicWorld
+		MCB_PeriodicWorld,				// TRANSFORMPOSSOURCE_PeriodicWorld
 		MCB_TranslatedWorld,			// TRANSFORMPOSSOURCE_TranslatedWorld
 		MCB_View,						// TRANSFORMPOSSOURCE_View
 		MCB_Camera,						// TRANSFORMPOSSOURCE_Camera
@@ -13685,41 +13685,24 @@ int32 UMaterialExpressionTransformPosition::Compile(class FMaterialCompiler* Com
 
 	int32 VecInputIdx = Input.Compile(Compiler);
 
-	// Periodic World Position (see TRANSFORMPOSSOURCE_PeriodicWorld) takes a different path as the transform depends on a user-specified tile scale
+	FTransformParameters Parameters {};
+
 	if (TransformSourceType == TRANSFORMPOSSOURCE_PeriodicWorld || TransformType == TRANSFORMPOSSOURCE_PeriodicWorld)
 	{
-		int32 PeriodicWorldTileSizeIdx;
 		if (PeriodicWorldTileSize.GetTracedInput().Expression)
 		{
-			PeriodicWorldTileSizeIdx = PeriodicWorldTileSize.Compile(Compiler);
+			Parameters.PeriodicWorldTileSizeIndex = PeriodicWorldTileSize.Compile(Compiler);
 		}
 		else
 		{
-			PeriodicWorldTileSizeIdx = Compiler->Constant(ConstPeriodicWorldTileSize);
-		}
-
-		if (TransformSourceType == TRANSFORMPOSSOURCE_PeriodicWorld && TransformType == TRANSFORMPOSSOURCE_PeriodicWorld)
-		{
-			return VecInputIdx;
-		}
-		if (TransformSourceType == TRANSFORMPOSSOURCE_PeriodicWorld)
-		{
-			const auto TransformDestBasis = GetMaterialCommonBasis(TransformType);
-			return Compiler->TransformFromPeriodicWorldPosition(TransformDestBasis, PeriodicWorldTileSizeIdx, VecInputIdx);
-		}
-		else
-		{
-			const auto TransformSourceBasis = GetMaterialCommonBasis(TransformSourceType);
-			return Compiler->TransformToPeriodicWorldPosition(TransformSourceBasis, PeriodicWorldTileSizeIdx, VecInputIdx);
+			Parameters.PeriodicWorldTileSizeIndex = Compiler->Constant(ConstPeriodicWorldTileSize);
 		}
 	}
-	else
-	{
-		const auto TransformSourceBasis = GetMaterialCommonBasis(TransformSourceType);
-		const auto TransformDestBasis = GetMaterialCommonBasis(TransformType);
+	
+	const auto TransformSourceBasis = GetMaterialCommonBasis(TransformSourceType);
+	const auto TransformDestBasis = GetMaterialCommonBasis(TransformType);
 
-		return Compiler->TransformPosition(TransformSourceBasis, TransformDestBasis, VecInputIdx);
-	}
+	return Compiler->TransformPosition(TransformSourceBasis, TransformDestBasis, Parameters, VecInputIdx);
 }
 
 void UMaterialExpressionTransformPosition::GetCaption(TArray<FString>& OutCaptions) const
diff --git a/Engine/Source/Runtime/Engine/Public/MaterialCompiler.h b/Engine/Source/Runtime/Engine/Public/MaterialCompiler.h
--- a/Engine/Source/Runtime/Engine/Public/MaterialCompiler.h
+++ b/Engine/Source/Runtime/Engine/Public/MaterialCompiler.h
@@ -80,6 +80,13 @@ enum ESubstrateMaterialExportContext : uint8
 	SMEC_Translucent = 1
 };
 
+/** Additional parameters used in TransformPosition / TransformVector. */
+struct FTransformParameters
+{
+	/** Code index for the scale of the tiles used in Periodic World Space */
+	int PeriodicWorldTileSizeIndex = INDEX_NONE;
+};
+
 /** 
  * The interface used to translate material expressions into executable code. 
  * Note: Most member functions should be pure virtual to force a FProxyMaterialCompiler override!
@@ -501,10 +508,32 @@ public:
 
 	virtual int32 ComponentMask(int32 Vector,bool R,bool G,bool B,bool A) = 0;
 	virtual int32 AppendVector(int32 A,int32 B) = 0;
-	virtual int32 TransformVector(EMaterialCommonBasis SourceCoordBasis, EMaterialCommonBasis DestCoordBasis, int32 A) = 0;
-	virtual int32 TransformPosition(EMaterialCommonBasis SourceCoordBasis, EMaterialCommonBasis DestCoordBasis, int32 A) = 0;
-	virtual int32 TransformFromPeriodicWorldPosition(EMaterialCommonBasis DestCoordBasis, int TileScaleIdx, int32 A) = 0;
-	virtual int32 TransformToPeriodicWorldPosition(EMaterialCommonBasis SourceCoordBasis, int TileScaleIdx, int32 A) = 0;
+	virtual int32 TransformVector(EMaterialCommonBasis SourceCoordBasis, EMaterialCommonBasis DestCoordBasis, FTransformParameters& Parameters, int32 A) = 0;
+	virtual int32 TransformPosition(EMaterialCommonBasis SourceCoordBasis, EMaterialCommonBasis DestCoordBasis, FTransformParameters& Parameters, int32 A) = 0;
+	int32 TransformVector(EMaterialCommonBasis SourceCoordBasis, EMaterialCommonBasis DestCoordBasis, int32 A)
+	{
+		FTransformParameters Parameters {};
+		return TransformVector(SourceCoordBasis, DestCoordBasis, Parameters, A);
+	}
+	int32 TransformPosition(EMaterialCommonBasis SourceCoordBasis, EMaterialCommonBasis DestCoordBasis, int32 A)
+	{
+		FTransformParameters Parameters {};
+		return TransformPosition(SourceCoordBasis, DestCoordBasis, Parameters, A);
+	}
+	UE_DEPRECATED(5.6, "Use TransformPosition(EMaterialCommonBasis, EMaterialCommonBasis, FTransformParameters, int32) instead")
+	int32 TransformFromPeriodicWorldPosition(EMaterialCommonBasis DestCoordBasis, int TileScaleIndex, int32 A)
+	{
+		FTransformParameters Parameters{};
+		Parameters.PeriodicWorldTileSizeIndex = TileScaleIndex;
+		return TransformPosition(EMaterialCommonBasis::MCB_PeriodicWorld, DestCoordBasis, Parameters, A);
+	}
+	UE_DEPRECATED(5.6, "Use TransformPosition(EMaterialCommonBasis, EMaterialCommonBasis, FTransformParameters, int32) instead")
+	int32 TransformToPeriodicWorldPosition(EMaterialCommonBasis SourceCoordBasis, int TileScaleIndex, int32 A)
+	{
+		FTransformParameters Parameters{};
+		Parameters.PeriodicWorldTileSizeIndex = TileScaleIndex;
+		return TransformPosition(SourceCoordBasis, EMaterialCommonBasis::MCB_PeriodicWorld, Parameters, A);
+	}
 	virtual int32 TransformNormalFromRequestedBasisToWorld(int32 NormalCodeChunk) = 0;
 
 	virtual int32 DynamicParameter(FLinearColor& DefaultValue, uint32 ParameterIndex = 0) = 0;
@@ -1057,21 +1086,13 @@ public:
 
 	virtual int32 ComponentMask(int32 Vector, bool R, bool G, bool B, bool A) override { return Compiler->ComponentMask(Vector, R, G, B, A); }
 	virtual int32 AppendVector(int32 A, int32 B) override { return Compiler->AppendVector(A, B); }
-	virtual int32 TransformVector(EMaterialCommonBasis SourceCoordBasis, EMaterialCommonBasis DestCoordBasis, int32 A) override
-	{
-		return Compiler->TransformVector(SourceCoordBasis, DestCoordBasis, A);
-	}
-	virtual int32 TransformPosition(EMaterialCommonBasis SourceCoordBasis, EMaterialCommonBasis DestCoordBasis, int32 A) override
-	{
-		return Compiler->TransformPosition(SourceCoordBasis, DestCoordBasis, A);
-	}
-	virtual int32 TransformFromPeriodicWorldPosition(EMaterialCommonBasis DestCoordBasis, int TileScaleIdx, int32 A) override
+	virtual int32 TransformVector(EMaterialCommonBasis SourceCoordBasis, EMaterialCommonBasis DestCoordBasis, FTransformParameters& Parameters, int32 A) override
 	{
-		return Compiler->TransformFromPeriodicWorldPosition(DestCoordBasis, TileScaleIdx, A);
+		return Compiler->TransformVector(SourceCoordBasis, DestCoordBasis, Parameters, A);
 	}
-	virtual int32 TransformToPeriodicWorldPosition(EMaterialCommonBasis SourceCoordBasis, int TileScaleIdx, int32 A) override
+	virtual int32 TransformPosition(EMaterialCommonBasis SourceCoordBasis, EMaterialCommonBasis DestCoordBasis, FTransformParameters& Parameters, int32 A) override
 	{
-		return Compiler->TransformToPeriodicWorldPosition(SourceCoordBasis, TileScaleIdx, A);
+		return Compiler->TransformPosition(SourceCoordBasis, DestCoordBasis, Parameters, A);
 	}
 	virtual int32 TransformNormalFromRequestedBasisToWorld(int32 NormalCodeChunk) override { return Compiler->TransformNormalFromRequestedBasisToWorld(NormalCodeChunk); }
 
diff --git a/Engine/Source/Runtime/Engine/Public/MaterialShared.h b/Engine/Source/Runtime/Engine/Public/MaterialShared.h
--- a/Engine/Source/Runtime/Engine/Public/MaterialShared.h
+++ b/Engine/Source/Runtime/Engine/Public/MaterialShared.h
@@ -326,6 +326,7 @@ enum EMaterialCommonBasis
 	MCB_Camera,
 	MCB_MeshParticle,
 	MCB_Instance,
+	MCB_PeriodicWorld,
 	MCB_MAX,
 };
 
