Wait for additive base to finish compressing before we compress

diff --git a/Engine/Source/Runtime/Engine/Private/Animation/AnimSequence.cpp b/Engine/Source/Runtime/Engine/Private/Animation/AnimSequence.cpp
--- a/Engine/Source/Runtime/Engine/Private/Animation/AnimSequence.cpp
+++ b/Engine/Source/Runtime/Engine/Private/Animation/AnimSequence.cpp
@@ -4724,7 +4724,7 @@ FIoHash UAnimSequence::BeginCacheDerivedData(const ITargetPlatform* TargetPlatfo
 
 	// Wait for any in-flight requests to finish
 	// Once a compression request finishes, it might modify the compression settings below and the key hash as well
-	UE::Anim::FAnimSequenceCompilingManager::Get().FinishCompilation({ this });
+	UE::Anim::FAnimSequenceCompilingManager::Get().FinishCompilation({ this, RefPoseSeq });
 
 	// Ensure that there are valid compression settings
 	if (BoneCompressionSettings == nullptr || !BoneCompressionSettings->AreSettingsValid())
