First Person: Added First Person Space as a transform destination to the TransformPosition material graph node.

diff --git a/Engine/Shaders/Private/MaterialTemplate.ush b/Engine/Shaders/Private/MaterialTemplate.ush
--- a/Engine/Shaders/Private/MaterialTemplate.ush
+++ b/Engine/Shaders/Private/MaterialTemplate.ush
@@ -3794,6 +3794,20 @@ bool IsFirstPerson(FMaterialPixelParameters Parameters)
 	return IsFirstPerson_FromFlags(GetPrimitiveData(Parameters).Flags);
 }
 
+float3 TransformToFirstPerson(float3 TranslatedWorldPosition, float FirstPersonInterpolationAlpha)
+{
+	const float3 FirstPersonPosition = mul(TranslatedWorldPosition, (float3x3)ResolvedView.FirstPersonTransform);
+	const float3 InterpolatedPosition = lerp(TranslatedWorldPosition, FirstPersonPosition, FirstPersonInterpolationAlpha);
+	return InterpolatedPosition;
+}
+
+float3 TransformToPreviousFirstPerson(float3 TranslatedWorldPosition, float FirstPersonInterpolationAlpha)
+{
+	const float3 FirstPersonPosition = mul(TranslatedWorldPosition, (float3x3)ResolvedView.PrevFirstPersonTransform);
+	const float3 InterpolatedPosition = lerp(TranslatedWorldPosition, FirstPersonPosition, FirstPersonInterpolationAlpha);
+	return InterpolatedPosition;
+}
+
 #if SUPPORT_FIRST_PERSON_RENDERING
 
 float GetMaterialFirstPersonInterpolationAlpha(FMaterialVertexParameters Parameters)
@@ -3863,10 +3877,9 @@ float3 GetMaterialWorldPositionOffset(FMaterialVertexParameters Parameters)
 	BRANCH
 	if (IsFirstPerson(Parameters))
 	{
-		const float3 FirstPersonPosition = mul(Parameters.WorldPosition + WPO, (float3x3)ResolvedView.FirstPersonTransform);
-		const float3 FirstPersonWPO = FirstPersonPosition - Parameters.WorldPosition;
 		const float LerpAlpha = GetMaterialFirstPersonInterpolationAlpha(Parameters);
-		WPO = lerp(WPO, FirstPersonWPO, LerpAlpha);
+		const float3 FirstPersonPosition = TransformToFirstPerson(Parameters.WorldPosition + WPO, LerpAlpha);
+		WPO = FirstPersonPosition - Parameters.WorldPosition;
 	}
 #endif // SUPPORT_FIRST_PERSON_RENDERING
 	return WPO;
@@ -3892,10 +3905,9 @@ float3 GetMaterialPreviousWorldPositionOffset(FMaterialVertexParameters Paramete
 	BRANCH
 	if (IsFirstPerson(Parameters))
 	{
-		const float3 FirstPersonPosition = mul(Parameters.WorldPosition + WPO, (float3x3)ResolvedView.PrevFirstPersonTransform);
-		const float3 FirstPersonWPO = FirstPersonPosition - Parameters.WorldPosition;
 		const float LerpAlpha = GetMaterialPreviousFirstPersonInterpolationAlpha(Parameters);
-		WPO = lerp(WPO, FirstPersonWPO, LerpAlpha);
+		const float3 FirstPersonPosition = TransformToPreviousFirstPerson(Parameters.WorldPosition + WPO, LerpAlpha);
+		WPO = FirstPersonPosition - Parameters.WorldPosition;
 	}
 #endif // SUPPORT_FIRST_PERSON_RENDERING
 	return WPO;
diff --git a/Engine/Source/Runtime/Engine/Private/Materials/HLSLMaterialTranslator.cpp b/Engine/Source/Runtime/Engine/Private/Materials/HLSLMaterialTranslator.cpp
--- a/Engine/Source/Runtime/Engine/Private/Materials/HLSLMaterialTranslator.cpp
+++ b/Engine/Source/Runtime/Engine/Private/Materials/HLSLMaterialTranslator.cpp
@@ -10851,6 +10851,11 @@ int32 FHLSLMaterialTranslator::TransformBase(EMaterialCommonBasis SourceCoordBas
 		{
 			return Errorf(TEXT("Tangent basis not available for position transformations"));
 		}
+
+		if (SourceCoordBasis == MCB_FirstPerson && DestCoordBasis != MCB_FirstPerson)
+		{
+			return Errorf(TEXT("Transforming from First Person Space is not supported"));
+		}
 		
 		// Construct float3(0,0,x) out of the input if it is a scalar
 		// This way artists can plug in a scalar and it will be treated as height, or a vector displacement
@@ -10904,7 +10909,7 @@ int32 FHLSLMaterialTranslator::TransformBase(EMaterialCommonBasis SourceCoordBas
 					CodeDerivStr = MultiplyTranslatedMatrix(TEXT("<A>"), TEXT("Get<PREV>LocalToWorldDF(Parameters)"), 0, bCompilingPreviousFrame);
 				}
 			}
-			else if (DestCoordBasis == MCB_PeriodicWorld)
+			else if (DestCoordBasis == MCB_PeriodicWorld || DestCoordBasis == MCB_FirstPerson)
 			{
 				IntermediaryBasis = MCB_TranslatedWorld;
 			}
@@ -10979,6 +10984,17 @@ int32 FHLSLMaterialTranslator::TransformBase(EMaterialCommonBasis SourceCoordBas
 				CodeStr = FString::Printf(TEXT("(<A> - (%s))"), *GetParameterCode(PeriodicWorldOrigin));
 				CodeDerivStr = TEXT("<A>");
 			}
+			else if (DestCoordBasis == MCB_FirstPerson)
+			{
+				if (Parameters.FirstPersonInterpolationAlphaIndex == INDEX_NONE)
+				{
+					return Errorf(TEXT("Missing first person interpolation alpha"));
+				}
+				int32 LerpAlphaFloat1Index = ForceCast(Parameters.FirstPersonInterpolationAlphaIndex, MCT_Float1);
+				int32 LerpAlphaClampedIndex = Saturate(LerpAlphaFloat1Index);
+				CodeStr = FString::Printf(TEXT("TransformTo<PREVIOUS>FirstPerson(<A>, %s)"), *GetParameterCode(LerpAlphaClampedIndex));
+				CodeDerivStr = CodeStr; // The first person transform is actually a 3x3 matrix and can therefore be used for derivatives as well.
+			}
 			// else use MCB_World as intermediary basis
 			break;
 		}
@@ -11074,7 +11090,7 @@ int32 FHLSLMaterialTranslator::TransformBase(EMaterialCommonBasis SourceCoordBas
 				CodeDerivStr = MultiplyTranslatedMatrix(TEXT("<A>"), TEXT("Parameters.Particle.ParticleToWorld"), 0, bCompilingPreviousFrame);
 				bUsesParticleLocalToWorld = true;
 			}
-			else if (DestCoordBasis == MCB_PeriodicWorld)
+			else if (DestCoordBasis == MCB_PeriodicWorld || DestCoordBasis == MCB_FirstPerson)
 			{
 				IntermediaryBasis = MCB_TranslatedWorld;
 			}
@@ -11096,7 +11112,7 @@ int32 FHLSLMaterialTranslator::TransformBase(EMaterialCommonBasis SourceCoordBas
 				CodeDerivStr = MultiplyTranslatedMatrix(TEXT("<A>"), TEXT("Get<PREV>InstanceToWorldDF(Parameters)"), 0, bCompilingPreviousFrame);
 				bUsesInstanceLocalToWorldPS |= ShaderFrequency == SF_Pixel;
 			}
-			else if (DestCoordBasis == MCB_PeriodicWorld)
+			else if (DestCoordBasis == MCB_PeriodicWorld || DestCoordBasis == MCB_FirstPerson)
 			{
 				IntermediaryBasis = MCB_TranslatedWorld;
 			}
@@ -11120,6 +11136,11 @@ int32 FHLSLMaterialTranslator::TransformBase(EMaterialCommonBasis SourceCoordBas
 			IntermediaryBasis = MCB_TranslatedWorld;
 			break;
 		}
+		case MCB_FirstPerson:
+		{
+			checkNoEntry(); // MCB_FirstPerson is not supported as a source basis. This should've been caught earlier in validation.
+			break;
+		}
 
 		default:
 			check(0);
diff --git a/Engine/Source/Runtime/Engine/Private/Materials/MaterialExpressionHLSL.cpp b/Engine/Source/Runtime/Engine/Private/Materials/MaterialExpressionHLSL.cpp
--- a/Engine/Source/Runtime/Engine/Private/Materials/MaterialExpressionHLSL.cpp
+++ b/Engine/Source/Runtime/Engine/Private/Materials/MaterialExpressionHLSL.cpp
@@ -3834,6 +3834,7 @@ bool UMaterialExpressionTransformPosition::GenerateHLSLExpression(FMaterialHLSLG
 		MCB_World,						// TRANSFORMPOSSOURCE_World
 		MCB_MAX,						// No match for TRANSFORMPOSSOURCE_PeriodicWorld
 		MCB_TranslatedWorld,			// TRANSFORMPOSSOURCE_TranslatedWorld
+		MCB_MAX,						// No match for TRANSFORMPOSSOURCE_FirstPersonTranslatedWorld
 		MCB_View,						// TRANSFORMPOSSOURCE_View
 		MCB_Camera,						// TRANSFORMPOSSOURCE_Camera
 		MCB_MeshParticle,				// TRANSFORMPOSSOURCE_Particle
@@ -3866,6 +3867,10 @@ bool UMaterialExpressionTransformPosition::GenerateHLSLExpression(FMaterialHLSLG
 			return true;
 		}
 	}
+	else if (TransformSourceType == TRANSFORMPOSSOURCE_FirstPersonTranslatedWorld || TransformType == TRANSFORMPOSSOURCE_FirstPersonTranslatedWorld)
+	{
+		return false; // TODO: Support transforming to FirstPerson? HLSLTree is not used at all at the moment and might be removed in the future.
+	}
 	else
 	{
 		OutExpression = TransformBase(Generator.GetTree(), kTable[TransformSourceType], kTable[TransformType], ExpressionInput, true);
diff --git a/Engine/Source/Runtime/Engine/Private/Materials/MaterialExpressions.cpp b/Engine/Source/Runtime/Engine/Private/Materials/MaterialExpressions.cpp
--- a/Engine/Source/Runtime/Engine/Private/Materials/MaterialExpressions.cpp
+++ b/Engine/Source/Runtime/Engine/Private/Materials/MaterialExpressions.cpp
@@ -13667,6 +13667,7 @@ static EMaterialCommonBasis GetMaterialCommonBasis(EMaterialPositionTransformSou
 		MCB_World,						// TRANSFORMPOSSOURCE_World
 		MCB_PeriodicWorld,				// TRANSFORMPOSSOURCE_PeriodicWorld
 		MCB_TranslatedWorld,			// TRANSFORMPOSSOURCE_TranslatedWorld
+		MCB_FirstPerson,				// TRANSFORMPOSSOURCE_FirstPersonTranslatedWorld
 		MCB_View,						// TRANSFORMPOSSOURCE_View
 		MCB_Camera,						// TRANSFORMPOSSOURCE_Camera
 		MCB_MeshParticle,				// TRANSFORMPOSSOURCE_Particle
@@ -13689,14 +13690,11 @@ int32 UMaterialExpressionTransformPosition::Compile(class FMaterialCompiler* Com
 
 	if (TransformSourceType == TRANSFORMPOSSOURCE_PeriodicWorld || TransformType == TRANSFORMPOSSOURCE_PeriodicWorld)
 	{
-		if (PeriodicWorldTileSize.GetTracedInput().Expression)
-		{
-			Parameters.PeriodicWorldTileSizeIndex = PeriodicWorldTileSize.Compile(Compiler);
-		}
-		else
-		{
-			Parameters.PeriodicWorldTileSizeIndex = Compiler->Constant(ConstPeriodicWorldTileSize);
-		}
+		Parameters.PeriodicWorldTileSizeIndex = PeriodicWorldTileSize.IsConnected() ? PeriodicWorldTileSize.Compile(Compiler) : Compiler->Constant(ConstPeriodicWorldTileSize);
+	}
+	if (TransformSourceType == TRANSFORMPOSSOURCE_FirstPersonTranslatedWorld || TransformType == TRANSFORMPOSSOURCE_FirstPersonTranslatedWorld)
+	{
+		Parameters.FirstPersonInterpolationAlphaIndex = FirstPersonInterpolationAlpha.IsConnected() ? FirstPersonInterpolationAlpha.Compile(Compiler) : Compiler->Constant(ConstFirstPersonInterpolationAlpha);
 	}
 	
 	const auto TransformSourceBasis = GetMaterialCommonBasis(TransformSourceType);
@@ -13705,6 +13703,11 @@ int32 UMaterialExpressionTransformPosition::Compile(class FMaterialCompiler* Com
 	return Compiler->TransformPosition(TransformSourceBasis, TransformDestBasis, Parameters, VecInputIdx);
 }
 
+void UMaterialExpressionTransformPosition::GetConnectorToolTip(int32 InputIndex, int32 OutputIndex, TArray<FString>& OutToolTip)
+{
+	Super::GetConnectorToolTip(GetAbsoluteInputIndex(InputIndex), OutputIndex, OutToolTip);
+}
+
 void UMaterialExpressionTransformPosition::GetCaption(TArray<FString>& OutCaptions) const
 {
 #if WITH_EDITOR
@@ -13744,6 +13747,10 @@ FExpressionInput* UMaterialExpressionTransformPosition::GetInput(int32 InputInde
 	{
 		IF_INPUT_RETURN(PeriodicWorldTileSize);
 	}
+	if (bUsesFirstPersonInterpolationAlpha)
+	{
+		IF_INPUT_RETURN(FirstPersonInterpolationAlpha);
+	}
 
 	return nullptr;
 }
@@ -13757,6 +13764,10 @@ FName UMaterialExpressionTransformPosition::GetInputName(int32 InputIndex) const
 	{
 		return TEXT("Periodic World Tile Size");
 	}
+	else if (FoundInput == &FirstPersonInterpolationAlpha)
+	{
+		return TEXT("First Person Interpolation Alpha");
+	}
 
 	return Super::GetInputName(InputIndex);
 }
@@ -13769,6 +13780,7 @@ void UMaterialExpressionTransformPosition::PostEditChangeProperty(struct FProper
 		PropertyChangedEvent.GetMemberPropertyName() == GET_MEMBER_NAME_CHECKED(ThisClass, TransformType))
 	{
 		bUsesPeriodicWorldPosition = TransformSourceType == TRANSFORMPOSSOURCE_PeriodicWorld || TransformType == TRANSFORMPOSSOURCE_PeriodicWorld;
+		bUsesFirstPersonInterpolationAlpha = TransformSourceType == TRANSFORMPOSSOURCE_FirstPersonTranslatedWorld || TransformType == TRANSFORMPOSSOURCE_FirstPersonTranslatedWorld;
 		if (GraphNode)
 		{
 			GraphNode->ReconstructNode();
@@ -13778,6 +13790,22 @@ void UMaterialExpressionTransformPosition::PostEditChangeProperty(struct FProper
 	// Need to update expression properties before super call (which triggers recompile)
 	Super::PostEditChangeProperty(PropertyChangedEvent);
 }
+
+#define IF_INPUT_RETURN(Value) if(!InputIndex) return (Value); --InputIndex
+int32 UMaterialExpressionTransformPosition::GetAbsoluteInputIndex(int32 InputIndex)
+{
+	IF_INPUT_RETURN(0);
+	if (bUsesPeriodicWorldPosition)
+	{
+		IF_INPUT_RETURN(1);
+	}
+	if (bUsesFirstPersonInterpolationAlpha)
+	{
+		IF_INPUT_RETURN(2);
+	}
+	return INDEX_NONE;
+}
+#undef IF_INPUT_RETURN
 #endif // WITH_EDITOR
 
 ///////////////////////////////////////////////////////////////////////////////
diff --git a/Engine/Source/Runtime/Engine/Public/MaterialCompiler.h b/Engine/Source/Runtime/Engine/Public/MaterialCompiler.h
--- a/Engine/Source/Runtime/Engine/Public/MaterialCompiler.h
+++ b/Engine/Source/Runtime/Engine/Public/MaterialCompiler.h
@@ -83,8 +83,11 @@ enum ESubstrateMaterialExportContext : uint8
 /** Additional parameters used in TransformPosition / TransformVector. */
 struct FTransformParameters
 {
-	/** Code index for the scale of the tiles used in Periodic World Space */
-	int PeriodicWorldTileSizeIndex = INDEX_NONE;
+	/** Code index for the scale of the tiles used in Periodic World Space. */
+	int32 PeriodicWorldTileSizeIndex = INDEX_NONE;
+
+	/** Code index for the interpolation alpha value used in First Person Space. */
+	int32 FirstPersonInterpolationAlphaIndex = INDEX_NONE;
 };
 
 /** 
diff --git a/Engine/Source/Runtime/Engine/Public/MaterialShared.h b/Engine/Source/Runtime/Engine/Public/MaterialShared.h
--- a/Engine/Source/Runtime/Engine/Public/MaterialShared.h
+++ b/Engine/Source/Runtime/Engine/Public/MaterialShared.h
@@ -327,6 +327,7 @@ enum EMaterialCommonBasis
 	MCB_MeshParticle,
 	MCB_Instance,
 	MCB_PeriodicWorld,
+	MCB_FirstPerson,
 	MCB_MAX,
 };
 
diff --git a/Engine/Source/Runtime/Engine/Public/Materials/MaterialExpressionTransformPosition.h b/Engine/Source/Runtime/Engine/Public/Materials/MaterialExpressionTransformPosition.h
--- a/Engine/Source/Runtime/Engine/Public/Materials/MaterialExpressionTransformPosition.h
+++ b/Engine/Source/Runtime/Engine/Public/Materials/MaterialExpressionTransformPosition.h
@@ -30,6 +30,9 @@ enum EMaterialPositionTransformSource : int
 	/** Translated world space, i.e. world space rotation and scale but with a position relative to the camera */
 	TRANSFORMPOSSOURCE_TranslatedWorld  UMETA(DisplayName="Camera Relative World Space"),
 
+	/** First person "space", which can be thought of as a transform that is applied to a position in translated world space. */
+	TRANSFORMPOSSOURCE_FirstPersonTranslatedWorld  UMETA(DisplayName = "First Person Space (Camera Relative World Space)"),
+
 	/** View space (differs from camera space in the shadow passes) */
 	TRANSFORMPOSSOURCE_View  UMETA(DisplayName="View Space"),
 
@@ -55,7 +58,7 @@ class UMaterialExpressionTransformPosition : public UMaterialExpression
 	FExpressionInput Input;
 
 	/** source format of the position that will be transformed */
-	UPROPERTY(EditAnywhere, Category=MaterialExpressionTransformPosition, meta=(DisplayName = "Source", ShowAsInputPin = "Advanced"))
+	UPROPERTY(EditAnywhere, Category=MaterialExpressionTransformPosition, meta=(DisplayName = "Source", ShowAsInputPin = "Advanced", InvalidEnumValues="TRANSFORMPOSSOURCE_FirstPersonTranslatedWorld"))
 	TEnumAsByte<enum EMaterialPositionTransformSource> TransformSourceType;
 
 	/** type of transform to apply to the input expression */
@@ -66,13 +69,24 @@ class UMaterialExpressionTransformPosition : public UMaterialExpression
 	UPROPERTY(meta=(DisplayName = "Periodic World Tile Size", ToolTip = "Distance the camera can move before the world origin is moved", RequiredInput = "false"))
 	FExpressionInput PeriodicWorldTileSize;
 
+	/** Interpolates between translated world space and first person space. Valid range is [0, 1], from translated world space to first person space. */
+	UPROPERTY(meta = (DisplayName = "First Person Interpolation Alpha", ToolTip = "Interpolates between world space and first person space. Valid range is [0, 1], from world space to first person space. Defaults to 'ConstFirstPersonInterpolationAlpha' if not specified.", RequiredInput = "false"))
+	FExpressionInput FirstPersonInterpolationAlpha;
+
 	/** only used if PeriodicWorldTileSize is not hooked up */
-	UPROPERTY(EditAnywhere, Category=MaterialExpressionTransformPosition, meta=(DisplayName = "Periodic World Tile Size", RequiredInput="false", ClampMin=0.0001, UIMax=524288.0, OverridingInputProperty = "PeriodicWorldTileSize", EditCondition=bUsesPeriodicWorldPosition, EditConditionHides))
+	UPROPERTY(EditAnywhere, Category=MaterialExpressionTransformPosition, meta=(DisplayName = "Periodic World Tile Size", ClampMin=0.0001, UIMax=524288.0, OverridingInputProperty = "PeriodicWorldTileSize", EditCondition=bUsesPeriodicWorldPosition, EditConditionHides, HideEditConditionToggle))
 	float ConstPeriodicWorldTileSize = 32.0;
 
+	/** Only used if FirstPersonInterpolationAlpha is not hooked up. Interpolates between translated world space and first person space. Valid range is [0, 1], from translated world space to first person space. */
+	UPROPERTY(EditAnywhere, Category=MaterialExpressionTransformPosition, meta=(DisplayName = "First Person Interpolation Alpha", UIMin = 0.0f, UIMax = 1.0f, ClampMin = 0.0f, ClampMax = 1.0f, OverridingInputProperty = "FirstPersonInterpolationAlpha", EditCondition = bUsesFirstPersonInterpolationAlpha, EditConditionHides, HideEditConditionToggle))
+	float ConstFirstPersonInterpolationAlpha = 1.0f;
+
 private:
 	UPROPERTY()
 	bool bUsesPeriodicWorldPosition = false;
+	
+	UPROPERTY()
+	bool bUsesFirstPersonInterpolationAlpha = false;
 
 public:
 
@@ -80,6 +94,7 @@ public:
 #if WITH_EDITOR
 	virtual bool GenerateHLSLExpression(FMaterialHLSLGenerator& Generator, UE::HLSLTree::FScope& Scope, int32 OutputIndex, UE::HLSLTree::FExpression const*& OutExpression) const override;
 	virtual int32 Compile(class FMaterialCompiler* Compiler, int32 OutputIndex) override;
+	virtual void GetConnectorToolTip(int32 InputIndex, int32 OutputIndex, TArray<FString>& OutToolTip) override;
 	virtual void GetCaption(TArray<FString>& OutCaptions) const override;
 	virtual TArrayView<FExpressionInput*> GetInputsView() override;
 	virtual FExpressionInput* GetInput(int32 InputIndex) override;
@@ -87,5 +102,8 @@ public:
 	virtual void PostEditChangeProperty(struct FPropertyChangedEvent& PropertyChangedEvent) override;
 #endif
 	//~ End UMaterialExpression Interface
+
+private:
+	int32 GetAbsoluteInputIndex(int32 InputIndex);
 };
 
