Blueprint Logging has to be Transactionally Safe when using AutoRTFM

diff --git a/Engine/Source/Runtime/Engine/Private/UnrealEngine.cpp b/Engine/Source/Runtime/Engine/Private/UnrealEngine.cpp
--- a/Engine/Source/Runtime/Engine/Private/UnrealEngine.cpp
+++ b/Engine/Source/Runtime/Engine/Private/UnrealEngine.cpp
@@ -11793,7 +11793,7 @@ void UEngine::AddTextureStreamingLoc(FVector InLoc, float BoostFactor, bool bOve
 
 #if !(UE_BUILD_SHIPPING || UE_BUILD_TEST)
 // Critical section used in AddOnScreenDebugMessage to handle concurrent insert.
-static FCriticalSection GOnScreenMessageCS;
+static FTransactionallySafeCriticalSection GOnScreenMessageCS;
 #endif
 
 /**
@@ -11811,7 +11811,7 @@ void UEngine::AddOnScreenDebugMessage(uint64 Key, float TimeToDisplay, FColor Di
 	if (bEnableOnScreenDebugMessages == true)
 	{
 		// Because some components add their message in concurrent work, we need a CS here.
-		FScopeLock ScopeLock(&GOnScreenMessageCS);
+		UE::TScopeLock ScopeLock(GOnScreenMessageCS);
 		if (Key == (uint64)-1)
 		{
 			if (bNewerOnTop)
@@ -11872,7 +11872,7 @@ void UEngine::RemoveOnScreenDebugMessage(uint64 Key)
 #if !(UE_BUILD_SHIPPING || UE_BUILD_TEST)
 	{
 		// Because some components add their message in concurrent work, we need a CS here.
-		FScopeLock ScopeLock(&GOnScreenMessageCS);
+		UE::TScopeLock ScopeLock(GOnScreenMessageCS);
 		ScreenMessages.Remove((int32)Key);
 	}
 #endif
@@ -12017,7 +12017,7 @@ void UEngine::ClearOnScreenDebugMessages()
 
 	{	
 		// Because some components add their message in concurrent work, we need a CS here.
-		FScopeLock ScopeLock(&GOnScreenMessageCS);
+		UE::TScopeLock ScopeLock(GOnScreenMessageCS);
 		ScreenMessages.Empty();
 		PriorityScreenMessages.Empty();
 	}
@@ -12834,7 +12834,7 @@ float UEngine::DrawOnscreenDebugMessages(UWorld* World, FViewport* Viewport, FCa
 	const int32 MaxYPos = CanvasObject ? CanvasObject->SizeY : 700;
 	{
 		// Because some components add their message in concurrent work, we need a CS here.
-		FScopeLock ScopeLock(&GOnScreenMessageCS);
+		UE::TScopeLock ScopeLock(GOnScreenMessageCS);
 
 		// Calculate the final amount of vertices for the batch or else we will end up resizing many times as the batching system only reserves the first draw call.
 		TSharedPtr<FSlateFontCache> FontCache = FEngineFontServices::Get().GetFontCache();
