Fixing some cases that prevent Overlaps from happening after Actor Migration.

diff --git a/Engine/Source/Runtime/Engine/Private/Actor.cpp b/Engine/Source/Runtime/Engine/Private/Actor.cpp
--- a/Engine/Source/Runtime/Engine/Private/Actor.cpp
+++ b/Engine/Source/Runtime/Engine/Private/Actor.cpp
@@ -1210,6 +1210,7 @@ void AActor::PostMigrate(FRemoteServerId DestinationServerId)
 		const USceneComponent* SceneRootComponent = FixupNativeActorComponents(this);
 		ensureMsgf(SceneRootComponent == RootComponent, TEXT("%hs: We were expecting RootComponent (%s) to be equal to FixupNativeActorComponents (%s). Transform fix-up code would be required to handle this case."), __func__, *GetNameSafe(RootComponent), *GetNameSafe(SceneRootComponent));
 
+		ResetOwnedComponents();
 		RegisterAllComponents();
 
 		bHasFinishedSpawning = true;
@@ -1234,6 +1235,12 @@ void AActor::PostMigrate(FRemoteServerId DestinationServerId)
 			if (Component->IsRegistered())
 			{
 				Component->RegisterAllComponentTickFunctions(true);
+
+				if (USceneComponent* SceneComponent = Cast<USceneComponent>(Component))
+				{
+					// UpdateInitialOverlaps needs this flag cleared or it won't run
+					SceneComponent->ClearSkipUpdateOverlaps();
+				}
 			}
 		}
 
