[WorldPartitionActorDesc] Make Relative Bounds work with classes derived from FWorldPartitionActorDesc which set custom Editor/Runtime Bounds
Add SetEditorBounds / SetRuntimeBounds methods, which set Relative Bounds at the same time

diff --git a/Engine/Source/Runtime/Engine/Private/WorldPartition/ActorPartition/PartitionActorDesc.cpp b/Engine/Source/Runtime/Engine/Private/WorldPartition/ActorPartition/PartitionActorDesc.cpp
--- a/Engine/Source/Runtime/Engine/Private/WorldPartition/ActorPartition/PartitionActorDesc.cpp
+++ b/Engine/Source/Runtime/Engine/Private/WorldPartition/ActorPartition/PartitionActorDesc.cpp
@@ -30,7 +30,7 @@ void FPartitionActorDesc::Init(const AActor* InActor)
 		const FVector ActorLocation = InActor->GetActorLocation();
 		SetGridIndices(ActorLocation.X, ActorLocation.Y, ActorLocation.Z);
 
-		EditorBounds = UActorPartitionSubsystem::FCellCoord::GetCellBounds({ GridIndexX, GridIndexY, GridIndexZ, 0 }, GridSize);
+		SetEditorBounds(UActorPartitionSubsystem::FCellCoord::GetCellBounds({ GridIndexX, GridIndexY, GridIndexZ, 0 }, GridSize));
 	}
 }
 
@@ -47,7 +47,7 @@ void FPartitionActorDesc::Serialize(FArchive& Ar)
 
 		if (Ar.CustomVer(FFortniteMainBranchObjectVersion::GUID) < FFortniteMainBranchObjectVersion::WorldPartitionActorDescSerializeEditorBounds)
 		{
-			EditorBounds = UActorPartitionSubsystem::FCellCoord::GetCellBounds({ GridIndexX, GridIndexY, GridIndexZ, 0 }, GridSize);
+			SetEditorBounds(UActorPartitionSubsystem::FCellCoord::GetCellBounds({ GridIndexX, GridIndexY, GridIndexZ, 0 }, GridSize));
 		}
 
 		if (Ar.CustomVer(FUE5MainStreamObjectVersion::GUID) >= FUE5MainStreamObjectVersion::PartitionActorDescSerializeGridGuid)
diff --git a/Engine/Source/Runtime/Engine/Private/WorldPartition/HLOD/HLODActorDesc.cpp b/Engine/Source/Runtime/Engine/Private/WorldPartition/HLOD/HLODActorDesc.cpp
--- a/Engine/Source/Runtime/Engine/Private/WorldPartition/HLOD/HLODActorDesc.cpp
+++ b/Engine/Source/Runtime/Engine/Private/WorldPartition/HLOD/HLODActorDesc.cpp
@@ -65,7 +65,7 @@ void FHLODActorDesc::Init(const AActor* InActor)
 
 	FVector Origin, Extent;
 	HLODActor->GetActorBounds(false, Origin, Extent);
-	EditorBounds = FBox::BuildAABB(Origin, Extent);
+	SetEditorBounds(FBox::BuildAABB(Origin, Extent));
 }
 
 struct FHLODSubActorDescDeprecated
diff --git a/Engine/Source/Runtime/Engine/Private/WorldPartition/Landscape/LandscapeActorDesc.cpp b/Engine/Source/Runtime/Engine/Private/WorldPartition/Landscape/LandscapeActorDesc.cpp
--- a/Engine/Source/Runtime/Engine/Private/WorldPartition/Landscape/LandscapeActorDesc.cpp
+++ b/Engine/Source/Runtime/Engine/Private/WorldPartition/Landscape/LandscapeActorDesc.cpp
@@ -27,7 +27,11 @@ void FLandscapeActorDesc::Init(const AActor* InActor)
 
 		// FLandscapeActorDesc derives from FPartitionActorDesc but doesn't use the cell bounds as the parent class was designed for.
 		// @todo_ow: make FLandscapeActorDesc derives from FWorldPartitionActorDesc instead?
-		InActor->GetStreamingBounds(RuntimeBounds, EditorBounds);
+		FBox NewRuntimeBounds = RuntimeBounds;
+		FBox NewEditorBounds = EditorBounds;
+		InActor->GetStreamingBounds(NewRuntimeBounds, NewEditorBounds);
+		SetRuntimeBounds(NewRuntimeBounds);
+		SetEditorBounds(NewEditorBounds);
 	}
 }
 
diff --git a/Engine/Source/Runtime/Engine/Private/WorldPartition/LevelInstance/LevelInstanceActorDesc.cpp b/Engine/Source/Runtime/Engine/Private/WorldPartition/LevelInstance/LevelInstanceActorDesc.cpp
--- a/Engine/Source/Runtime/Engine/Private/WorldPartition/LevelInstance/LevelInstanceActorDesc.cpp
+++ b/Engine/Source/Runtime/Engine/Private/WorldPartition/LevelInstance/LevelInstanceActorDesc.cpp
@@ -130,8 +130,8 @@ void FLevelInstanceActorDesc::UpdateBounds()
 {
 	if (UActorDescContainer* ChildContainerPtr = ChildContainer.Get())
 	{
-		RuntimeBounds = UActorDescContainerSubsystem::GetChecked().GetContainerBounds(GetChildContainerName(), false).TransformBy(GetChildContainerTransform());
-		EditorBounds = UActorDescContainerSubsystem::GetChecked().GetContainerBounds(GetChildContainerName(), true).TransformBy(GetChildContainerTransform());
+		SetRuntimeBounds(UActorDescContainerSubsystem::GetChecked().GetContainerBounds(GetChildContainerName(), false).TransformBy(GetChildContainerTransform()));
+		SetEditorBounds(UActorDescContainerSubsystem::GetChecked().GetContainerBounds(GetChildContainerName(), true).TransformBy(GetChildContainerTransform()));
 	}
 }
 
@@ -402,7 +402,9 @@ void FLevelInstanceActorDesc::Serialize(FArchive& Ar)
 			{
 				if (!IsChildContainerInstance())
 				{
-					ULevelInstanceSubsystem::GetLevelInstanceBoundsFromPackage(ActorTransform, GetChildContainerPackage(), RuntimeBounds);
+					FBox NewRuntimeBounds = RuntimeBounds;
+					ULevelInstanceSubsystem::GetLevelInstanceBoundsFromPackage(ActorTransform, GetChildContainerPackage(), NewRuntimeBounds);
+					SetRuntimeBounds(NewRuntimeBounds);
 				}
 			}
 		}
diff --git a/Engine/Source/Runtime/Engine/Private/WorldPartition/WorldPartitionActorDesc.cpp b/Engine/Source/Runtime/Engine/Private/WorldPartition/WorldPartitionActorDesc.cpp
--- a/Engine/Source/Runtime/Engine/Private/WorldPartition/WorldPartitionActorDesc.cpp
+++ b/Engine/Source/Runtime/Engine/Private/WorldPartition/WorldPartitionActorDesc.cpp
@@ -891,6 +891,18 @@ FBox FWorldPartitionActorDesc::GetRuntimeBounds() const
 	return RuntimeBounds;
 }
 
+void FWorldPartitionActorDesc::SetEditorBounds(const FBox& InEditorBounds)
+{
+	EditorBounds = InEditorBounds;
+	EditorBoundsRelative = FWorldPartitionRelativeBounds(EditorBounds).InverseTransformBy(ActorTransform);
+}
+
+void FWorldPartitionActorDesc::SetRuntimeBounds(const FBox& InRuntimeBounds)
+{
+	RuntimeBounds = InRuntimeBounds;
+	RuntimeBoundsRelative = FWorldPartitionRelativeBounds(RuntimeBounds).InverseTransformBy(ActorTransform);
+}
+
 FName FWorldPartitionActorDesc::GetActorName() const
 {
 	return ActorName;
diff --git a/Engine/Source/Runtime/Engine/Public/WorldPartition/WorldPartitionActorDesc.h b/Engine/Source/Runtime/Engine/Public/WorldPartition/WorldPartitionActorDesc.h
--- a/Engine/Source/Runtime/Engine/Public/WorldPartition/WorldPartitionActorDesc.h
+++ b/Engine/Source/Runtime/Engine/Public/WorldPartition/WorldPartitionActorDesc.h
@@ -530,6 +530,9 @@ protected:
 
 	ENGINE_API virtual void Serialize(FArchive& Ar);
 
+	ENGINE_API void SetEditorBounds(const FBox& InEditorBounds);
+	ENGINE_API void SetRuntimeBounds(const FBox& InRuntimeBounds);
+
 	// Persistent
 	FGuid							Guid;
 	FTopLevelAssetPath				BaseClass;
