FActorDescArchive: properly handle serializing actor descriptors derived classes when the base type is unknown.

diff --git a/Engine/Source/Runtime/Engine/Public/WorldPartition/WorldPartitionActorDescArchive.h b/Engine/Source/Runtime/Engine/Public/WorldPartition/WorldPartitionActorDescArchive.h
--- a/Engine/Source/Runtime/Engine/Public/WorldPartition/WorldPartitionActorDescArchive.h
+++ b/Engine/Source/Runtime/Engine/Public/WorldPartition/WorldPartitionActorDescArchive.h
@@ -58,9 +58,8 @@ public:
 			{
 				const UPTRINT PropertyOffset = (UPTRINT)&V.Value - *(UPTRINT*)&ActorDescAr.ActorDesc;
 
-				if (PropertyOffset < ActorDescAr.BaseDescSizeof)
+				if ((PropertyOffset + sizeof(V.Value)) <= ActorDescAr.BaseDescSizeof)
 				{
-					check((PropertyOffset + sizeof(V.Value)) <= ActorDescAr.BaseDescSizeof);
 					const DestPropertyType* RefValue = (const DestPropertyType*)(*(UPTRINT*)&ActorDescAr.BaseDesc + PropertyOffset);
 					return RefValue;
 				}
