Fixed missing tracking of ClearInstances in AnimBank

diff --git a/Engine/Source/Runtime/Engine/Private/Components/AnimBankMeshComponent.cpp b/Engine/Source/Runtime/Engine/Private/Components/AnimBankMeshComponent.cpp
--- a/Engine/Source/Runtime/Engine/Private/Components/AnimBankMeshComponent.cpp
+++ b/Engine/Source/Runtime/Engine/Private/Components/AnimBankMeshComponent.cpp
@@ -1519,6 +1519,7 @@ void UAnimBankMeshComponent::ClearInstances()
 #if WITH_EDITOR
 	SelectedInstances.Empty();
 #endif
+	InstanceDataManager.ClearInstances();
 }
 
 int32 UAnimBankMeshComponent::AddAnimBankItem(const FAnimBankItem& BankItem)
diff --git a/Engine/Source/Runtime/Engine/Private/InstanceData/InstanceDataManager.cpp b/Engine/Source/Runtime/Engine/Private/InstanceData/InstanceDataManager.cpp
--- a/Engine/Source/Runtime/Engine/Private/InstanceData/InstanceDataManager.cpp
+++ b/Engine/Source/Runtime/Engine/Private/InstanceData/InstanceDataManager.cpp
@@ -487,6 +487,21 @@ void FInstanceDataManager::PrimitiveTransformChanged()
 	MarkComponentRenderInstancesDirty();
 }
 
+void FInstanceDataManager::ClearInstances()
+{
+	if (GetState() == ETrackingState::Disabled)
+	{
+		return;
+	}
+	// Reset any instance allocations
+	Reset(0);
+	ValidInstanceIdMask.Empty();
+	IdSearchStartIndex = 0;
+	// reset the change tracking back to initial state
+	ClearChangeTracking();
+	MarkComponentRenderInstancesDirty();
+}
+
 bool FInstanceDataManager::HasAnyInstanceChanges() const
 {
 	return bAnyInstanceChange || bNumCustomDataChanged || AllChangedFlags.Packed != 0u || bTransformChangedAllInstances
diff --git a/Engine/Source/Runtime/Engine/Public/InstanceData/InstanceDataManager.h b/Engine/Source/Runtime/Engine/Public/InstanceData/InstanceDataManager.h
--- a/Engine/Source/Runtime/Engine/Public/InstanceData/InstanceDataManager.h
+++ b/Engine/Source/Runtime/Engine/Public/InstanceData/InstanceDataManager.h
@@ -110,6 +110,8 @@ public:
 
 	void PrimitiveTransformChanged();
 
+	void ClearInstances();
+
 	ENGINE_API bool HasAnyInstanceChanges() const;
 
 	/** 
