Changed sound occlusion traces to run synchronously in the audio device update task by default. This audio work already runs on a task thread, so it's less overhead to run the trace synchronously.

When the traces are run asynchronously, the audio task queues a tick on the game thread which then queues an async trace back on the task threads. The async trace callback runs on the GT again afterwards. None of this work is actually required to run on the GT, so we can save a bit of time by using the synchronous path.

This will remove GT cost from TickActors, GT Audio, and ResetAsyncTraceTime.
Controlled by cvar au.UseAsyncOcclusionTraces.

diff --git a/Engine/Source/Runtime/Engine/Private/ActiveSound.cpp b/Engine/Source/Runtime/Engine/Private/ActiveSound.cpp
--- a/Engine/Source/Runtime/Engine/Private/ActiveSound.cpp
+++ b/Engine/Source/Runtime/Engine/Private/ActiveSound.cpp
@@ -21,6 +21,13 @@ FAutoConsoleVariableRef CVarAudioOcclusionEnabled(
 	TEXT("Disables (1) or enables (0) audio occlusion.\n"),
 	ECVF_Default);
 
+static bool bUseAsyncOcclusionTracesCVar = false;
+FAutoConsoleVariableRef CVarUseAsyncOcclusionTraces(
+	TEXT("au.UseAsyncOcclusionTraces"),
+	bUseAsyncOcclusionTracesCVar,
+	TEXT("0: Run sound occlusion traces from the audio device update, 1: Run sound occlusion traces through the gt async trace api."),
+	ECVF_Default);
+
 static int32 GatherInteriorDataFromAudioVolumesCVar = 1;
 FAutoConsoleVariableRef CVarGatherInteriorDataFromAudioVolumes(
 	TEXT("au.InteriorData.UseAudioVolumes"),
@@ -1359,7 +1366,7 @@ void FActiveSound::CheckOcclusion(const FVector ListenerLocation, const FVector
 			const bool bUseComplexCollisionForOcclusion = AttenuationSettingsPtr->bUseComplexCollisionForOcclusion;
 			const ECollisionChannel OcclusionTraceChannel = AttenuationSettingsPtr->OcclusionTraceChannel;
 
-			if (!bHasCheckedOcclusion)
+			if (!bHasCheckedOcclusion || !bUseAsyncOcclusionTracesCVar)
 			{
 				FCollisionQueryParams Params(SCENE_QUERY_STAT(SoundOcclusion), bUseComplexCollisionForOcclusion);
 				if (OwnerID > 0)
