Fix support for loading multiple copies of the same partitioned worlds from LoadLevelInstance Blueprint function.

diff --git a/Engine/Source/Runtime/Engine/Private/WorldPartition/WorldPartitionLevelStreamingPolicy.cpp b/Engine/Source/Runtime/Engine/Private/WorldPartition/WorldPartitionLevelStreamingPolicy.cpp
--- a/Engine/Source/Runtime/Engine/Private/WorldPartition/WorldPartitionLevelStreamingPolicy.cpp
+++ b/Engine/Source/Runtime/Engine/Private/WorldPartition/WorldPartitionLevelStreamingPolicy.cpp
@@ -66,12 +66,10 @@ FString UWorldPartitionLevelStreamingPolicy::GetCellPackagePath(const FName& InC
 	if (InWorld->IsGameWorld())
 	{
 		// Set as memory package to avoid wasting time in FPackageName::DoesPackageExist
-		return FString::Printf(TEXT("/Memory/%s"), *InCellName.ToString());
-	}
-	else
-	{
-		return FString::Printf(TEXT("/%s"), *InCellName.ToString());
+		return FString::Printf(TEXT("/Memory/%s_%s"), *FPackageName::GetShortName(InWorld->GetPackage()), *InCellName.ToString());
 	}
+
+	return FString::Printf(TEXT("/%s"), *InCellName.ToString());
 }
 
 TSubclassOf<UWorldPartitionRuntimeCell> UWorldPartitionLevelStreamingPolicy::GetRuntimeCellClass() const
