Created missing Get/Set functions for the following member variables:
ACharacter.ReplicatedBasedMovement
ACharacter.ReplicatedServerLastTransformUpdateTimeStamp
ACharacter.ReplayLastTransformUpdateTimeStamp
ACharacter.ReplicatedGravityDirection
ACharacter.ReplicatedMovementMode
ACharacter.bProxyIsJumpForceApplied
ACharacter.bIsCrouched
ACharacter.RepRootMotion
APawn.RemoteViewPitch
APawn.Controller

And implemented use of Get/Set functions when referencing member variables of APawn and ACharacter in Engine code

It is recommended to use the Get/Set API instead of directly referencing these member variables.
This is to enable a later change where these replicated properties are Dirtied inside the Set functions, which allows for an implementation of Push Model Networking on these classes.

[RN] minor Foundation/Core
#tests Tested against Dedicated Server Cooked Data and PIE

[RN] minor Foundation/Core

diff --git a/Engine/Plugins/Experimental/Animation/MotionTrajectory/Source/MotionTrajectory/Private/MotionTrajectoryLibrary.cpp b/Engine/Plugins/Experimental/Animation/MotionTrajectory/Source/MotionTrajectory/Private/MotionTrajectoryLibrary.cpp
--- a/Engine/Plugins/Experimental/Animation/MotionTrajectory/Source/MotionTrajectory/Private/MotionTrajectoryLibrary.cpp
+++ b/Engine/Plugins/Experimental/Animation/MotionTrajectory/Source/MotionTrajectory/Private/MotionTrajectoryLibrary.cpp
@@ -57,7 +57,7 @@ void FCharacterTrajectoryData::UpdateDataFromCharacter(float DeltaSeconds, const
 	}
 
 	// @todo: Simulated proxies don't have controllers, so they'll need some other mechanism to account for controller rotation rate.
-	const AController* Controller = Character->Controller;
+	const AController* Controller = Character->GetController();
 	if (Controller)
 	{
 		float DesiredControllerYaw = Controller->GetDesiredRotation().Yaw;
diff --git a/Engine/Plugins/Experimental/ChaosVehiclesPlugin/Source/ChaosVehicles/Private/ChaosVehicleMovementComponent.cpp b/Engine/Plugins/Experimental/ChaosVehiclesPlugin/Source/ChaosVehicles/Private/ChaosVehicleMovementComponent.cpp
--- a/Engine/Plugins/Experimental/ChaosVehiclesPlugin/Source/ChaosVehicles/Private/ChaosVehicleMovementComponent.cpp
+++ b/Engine/Plugins/Experimental/ChaosVehiclesPlugin/Source/ChaosVehicles/Private/ChaosVehicleMovementComponent.cpp
@@ -1368,7 +1368,7 @@ AController* UChaosVehicleMovementComponent::GetController() const
 	{
 		if (APawn* Pawn = Cast<APawn>(UpdatedComponent->GetOwner()))
 		{
-			return Pawn->Controller;
+			return Pawn->GetController();
 		}
 	}
 
diff --git a/Engine/Plugins/Experimental/MoverExamples/Source/MoverExamples/Private/MoverExamplesCharacter.cpp b/Engine/Plugins/Experimental/MoverExamples/Source/MoverExamples/Private/MoverExamplesCharacter.cpp
--- a/Engine/Plugins/Experimental/MoverExamples/Source/MoverExamples/Private/MoverExamplesCharacter.cpp
+++ b/Engine/Plugins/Experimental/MoverExamples/Source/MoverExamples/Private/MoverExamplesCharacter.cpp
@@ -56,7 +56,7 @@ void AMoverExamplesCharacter::Tick(float DeltaTime)
 
 
 	// Spin camera based on input
-	if (APlayerController* PC = Cast<APlayerController>(Controller))
+	if (APlayerController* PC = Cast<APlayerController>(GetController()))
 	{
 		// Simple input scaling. A real game will probably map this to an acceleration curve
 		static float LookRateYaw = 100.f;	// degs/sec
@@ -74,7 +74,7 @@ void AMoverExamplesCharacter::BeginPlay()
 {
 	Super::BeginPlay();
 
-	if (APlayerController* PC = Cast<APlayerController>(Controller))
+	if (APlayerController* PC = Cast<APlayerController>(GetController()))
 	{
 		PC->PlayerCameraManager->ViewPitchMax = 89.0f;
 		PC->PlayerCameraManager->ViewPitchMin = -89.0f;
@@ -159,7 +159,7 @@ void AMoverExamplesCharacter::OnProduceInput(float DeltaMs, FMoverInputCmdContex
 
 	FCharacterDefaultInputs& CharacterInputs = OutInputCmd.InputCollection.FindOrAddMutableDataByType<FCharacterDefaultInputs>();
 
-	if (Controller == nullptr)
+	if (GetController() == nullptr)
 	{
 		if (GetLocalRole() == ENetRole::ROLE_Authority && GetRemoteRole() == ENetRole::ROLE_SimulatedProxy)
 		{
@@ -180,7 +180,7 @@ void AMoverExamplesCharacter::OnProduceInput(float DeltaMs, FMoverInputCmdContex
 
 	CharacterInputs.ControlRotation = FRotator::ZeroRotator;
 
-	APlayerController* PC = Cast<APlayerController>(Controller);
+	APlayerController* PC = Cast<APlayerController>(GetController());
 	if (PC)
 	{
 		CharacterInputs.ControlRotation = PC->GetControlRotation();
diff --git a/Engine/Plugins/Experimental/XRCreativeFramework/Source/XRCreative/Private/ITF/XRCreativeITFComponent.cpp b/Engine/Plugins/Experimental/XRCreativeFramework/Source/XRCreative/Private/ITF/XRCreativeITFComponent.cpp
--- a/Engine/Plugins/Experimental/XRCreativeFramework/Source/XRCreative/Private/ITF/XRCreativeITFComponent.cpp
+++ b/Engine/Plugins/Experimental/XRCreativeFramework/Source/XRCreative/Private/ITF/XRCreativeITFComponent.cpp
@@ -713,7 +713,7 @@ void UXRCreativeITFComponent::ToolsTick(float InDeltaTime)
 		.SetRealtimeUpdate(true));
 
 	APawn* OwnerPawn = CastChecked<APawn>(GetOwner());
-	APlayerController* PC = CastChecked<APlayerController>(OwnerPawn->Controller);
+	APlayerController* PC = CastChecked<APlayerController>(OwnerPawn->GetController());
 	ULocalPlayer* LocalPlayer = Cast<ULocalPlayer>(PC->Player);
 	FVector ViewLocation;
 	FRotator ViewRotation;
diff --git a/Engine/Plugins/Runtime/ArchVisCharacter/Source/ArchVisCharacter/Private/ArchVisCharMovementComponent.cpp b/Engine/Plugins/Runtime/ArchVisCharacter/Source/ArchVisCharacter/Private/ArchVisCharMovementComponent.cpp
--- a/Engine/Plugins/Runtime/ArchVisCharacter/Source/ArchVisCharacter/Private/ArchVisCharMovementComponent.cpp
+++ b/Engine/Plugins/Runtime/ArchVisCharacter/Source/ArchVisCharacter/Private/ArchVisCharMovementComponent.cpp
@@ -114,7 +114,7 @@ void UArchVisCharMovementComponent::PhysWalking(float DeltaTime, int32 Iteration
 
 		FHitResult Hit(1.f);
 		SafeMoveUpdatedComponent(FVector::ZeroVector, FRotator(0.f, NewRot.Yaw, NewRot.Roll), false, Hit);
-		CharacterOwner->Controller->SetControlRotation(FRotator(NewRot.Pitch, 0.f, 0.f));
+		CharacterOwner->GetController()->SetControlRotation(FRotator(NewRot.Pitch, 0.f, 0.f));
 
 		// TODO: Do Cosine scaling?
 		float const PercentSpeed = (90.f - fabs(NewRot.Pitch)) / 90.f;
diff --git a/Engine/Plugins/Runtime/ArchVisCharacter/Source/ArchVisCharacter/Private/ArchVisCharacter.cpp b/Engine/Plugins/Runtime/ArchVisCharacter/Source/ArchVisCharacter/Private/ArchVisCharacter.cpp
--- a/Engine/Plugins/Runtime/ArchVisCharacter/Source/ArchVisCharacter/Private/ArchVisCharacter.cpp
+++ b/Engine/Plugins/Runtime/ArchVisCharacter/Source/ArchVisCharacter/Private/ArchVisCharacter.cpp
@@ -103,7 +103,7 @@ void AArchVisCharacter::MoveRight(float Val)
 {
 	if (Val != 0.f)
 	{
-		if (Controller)
+		if (GetController())
 		{
 			FRotator const ControlSpaceRot = GetActorRotation();
 
@@ -117,7 +117,7 @@ void AArchVisCharacter::MoveForward(float Val)
 {
 	if (Val != 0.f)
 	{
-		if (Controller)
+		if (GetController())
 		{
 			FRotator const ControlSpaceRot = GetActorRotation();
 
diff --git a/Engine/Plugins/Runtime/GameplayAbilities/Source/GameplayAbilities/Private/GameplayCueNotifyTypes.cpp b/Engine/Plugins/Runtime/GameplayAbilities/Source/GameplayAbilities/Private/GameplayCueNotifyTypes.cpp
--- a/Engine/Plugins/Runtime/GameplayAbilities/Source/GameplayAbilities/Private/GameplayCueNotifyTypes.cpp
+++ b/Engine/Plugins/Runtime/GameplayAbilities/Source/GameplayAbilities/Private/GameplayCueNotifyTypes.cpp
@@ -53,7 +53,7 @@ static APlayerController* FindPlayerControllerFromActor(AActor* Actor)
 		Pawn = Actor->GetInstigator<APawn>();
 	}
 
-	PC = (Pawn ? Cast<APlayerController>(Pawn->Controller) : nullptr);
+	PC = (Pawn ? Cast<APlayerController>(Pawn->GetController()) : nullptr);
 	if (PC)
 	{
 		return PC;
diff --git a/Engine/Plugins/Runtime/GameplayAbilities/Source/GameplayAbilities/Private/GameplayEffectTypes.cpp b/Engine/Plugins/Runtime/GameplayAbilities/Source/GameplayAbilities/Private/GameplayEffectTypes.cpp
--- a/Engine/Plugins/Runtime/GameplayAbilities/Source/GameplayAbilities/Private/GameplayEffectTypes.cpp
+++ b/Engine/Plugins/Runtime/GameplayAbilities/Source/GameplayAbilities/Private/GameplayEffectTypes.cpp
@@ -352,9 +352,9 @@ bool FGameplayEffectContext::IsLocallyControlledPlayer() const
 	{
 		Pawn = Cast<APawn>(EffectCauser.Get());
 	}
-	if (Pawn && Pawn->Controller)
+	if (Pawn && Pawn->GetController())
 	{
-		return Pawn->Controller->IsLocalPlayerController();
+		return Pawn->GetController()->IsLocalPlayerController();
 	}
 	return false;
 }
@@ -1266,9 +1266,9 @@ bool FGameplayCueParameters::IsInstigatorLocallyControlledPlayer(AActor* Fallbac
 			}
 		}
 
-		if (Pawn && Pawn->Controller)
+		if (Pawn && Pawn->GetController())
 		{
-			return Pawn->Controller->IsLocalPlayerController();
+			return Pawn->GetController()->IsLocalPlayerController();
 		}
 	}
 
diff --git a/Engine/Plugins/Runtime/GeoReferencing/Source/GeoReferencing/Private/RoundPlanetPawn.cpp b/Engine/Plugins/Runtime/GeoReferencing/Source/GeoReferencing/Private/RoundPlanetPawn.cpp
--- a/Engine/Plugins/Runtime/GeoReferencing/Source/GeoReferencing/Private/RoundPlanetPawn.cpp
+++ b/Engine/Plugins/Runtime/GeoReferencing/Source/GeoReferencing/Private/RoundPlanetPawn.cpp
@@ -205,7 +205,7 @@ void ARoundPlanetPawn::MoveForward(float Val)
 {
 	if (Val != 0.f)
 	{
-		if (Controller)
+		if (GetController())
 		{
 			FRotator const ControlSpaceRot = this->GetViewRotation();
 
@@ -234,7 +234,7 @@ void ARoundPlanetPawn::MoveRight(float Val)
 {
 	if (Val != 0.f)
 	{
-		if (Controller)
+		if (GetController())
 		{
 			FRotator const ControlSpaceRot = this->GetViewRotation();
 
diff --git a/Engine/Plugins/Runtime/NetworkPredictionExtras/Source/NetworkPredictionExtras/Private/NetworkPredictionExtrasCharacter.cpp b/Engine/Plugins/Runtime/NetworkPredictionExtras/Source/NetworkPredictionExtras/Private/NetworkPredictionExtrasCharacter.cpp
--- a/Engine/Plugins/Runtime/NetworkPredictionExtras/Source/NetworkPredictionExtras/Private/NetworkPredictionExtrasCharacter.cpp
+++ b/Engine/Plugins/Runtime/NetworkPredictionExtras/Source/NetworkPredictionExtras/Private/NetworkPredictionExtrasCharacter.cpp
@@ -99,7 +99,7 @@ void ANetworkPredictionExtrasCharacter::SetupPlayerInputComponent(UInputComponen
 	// This also is not necessary. This would usually be defined in DefaultInput.ini.
 	// But the nature of NetworkPredictionExtras is to not assume or require anything of the project its being run in.
 	// So in order to make this nice and self contained, manually do input bindings here.
-	if (APlayerController* PC = Cast<APlayerController>(Controller))
+	if (APlayerController* PC = Cast<APlayerController>(GetController()))
 	{
 		if (PC->PlayerInput)
 		{
@@ -244,7 +244,7 @@ void ANetworkPredictionExtrasCharacter::ProduceInput(const int32 DeltaMS, FChara
 	}
 
 
-	if (Controller == nullptr)
+	if (GetController() == nullptr)
 	{
 		if (GetLocalRole() == ENetRole::ROLE_Authority && GetRemoteRole() == ENetRole::ROLE_SimulatedProxy)
 		{
@@ -301,7 +301,7 @@ void ANetworkPredictionExtrasCharacter::ProduceInput(const int32 DeltaMS, FChara
 		case 1:
 		{
 			// Free camera Restricted 2D movement on XY plane.
-			APlayerController* PC = Cast<APlayerController>(Controller);
+			APlayerController* PC = Cast<APlayerController>(GetController());
 			if (ensure(PC)) // Requires player controller for now
 			{
 				// Camera yaw rotation
@@ -333,7 +333,7 @@ void ANetworkPredictionExtrasCharacter::ProduceInput(const int32 DeltaMS, FChara
 		case 2:
 		{
 			// Free camera on yaw and pitch, camera-relative movement.
-			APlayerController* PC = Cast<APlayerController>(Controller);
+			APlayerController* PC = Cast<APlayerController>(GetController());
 			if (ensure(PC)) // Requires player controller for now
 			{
 				// Camera yaw rotation
@@ -379,7 +379,7 @@ void ANetworkPredictionExtrasCharacter::ProduceInput(const int32 DeltaMS, FChara
 		case 3:
 		{
 			// Free camera on the yaw, camera-relative motion
-			APlayerController* PC = Cast<APlayerController>(Controller);
+			APlayerController* PC = Cast<APlayerController>(GetController());
 			if (ensure(PC)) // Requires player controller for now
 			{
 				// Camera yaw rotation
@@ -507,7 +507,7 @@ void ANetworkPredictionExtrasCharacter_MockAbility::SetupPlayerInputComponent(UI
 {
 	Super::SetupPlayerInputComponent(PlayerInputComponent);
 
-	if (APlayerController* PC = Cast<APlayerController>(Controller))
+	if (APlayerController* PC = Cast<APlayerController>(GetController()))
 	{
 		if (PC->PlayerInput)
 		{
diff --git a/Engine/Plugins/Runtime/NetworkPredictionExtras/Source/NetworkPredictionExtras/Private/NetworkPredictionExtrasFlyingPawn.cpp b/Engine/Plugins/Runtime/NetworkPredictionExtras/Source/NetworkPredictionExtras/Private/NetworkPredictionExtrasFlyingPawn.cpp
--- a/Engine/Plugins/Runtime/NetworkPredictionExtras/Source/NetworkPredictionExtras/Private/NetworkPredictionExtrasFlyingPawn.cpp
+++ b/Engine/Plugins/Runtime/NetworkPredictionExtras/Source/NetworkPredictionExtras/Private/NetworkPredictionExtrasFlyingPawn.cpp
@@ -98,7 +98,7 @@ void ANetworkPredictionExtrasFlyingPawn::SetupPlayerInputComponent(UInputCompone
 	// This also is not necessary. This would usually be defined in DefaultInput.ini.
 	// But the nature of NetworkPredictionExtras is to not assume or require anything of the project its being run in.
 	// So in order to make this nice and self contained, manually do input bindings here.
-	if (APlayerController* PC = Cast<APlayerController>(Controller))
+	if (APlayerController* PC = Cast<APlayerController>(GetController()))
 	{
 		if (PC->PlayerInput)
 		{
@@ -242,7 +242,7 @@ void ANetworkPredictionExtrasFlyingPawn::ProduceInput(const int32 DeltaMS, FFlyi
 		return;
 	}
 
-	if (Controller == nullptr)
+	if (GetController() == nullptr)
 	{
 		if (GetLocalRole() == ENetRole::ROLE_Authority && GetRemoteRole() == ENetRole::ROLE_SimulatedProxy)
 		{
@@ -297,7 +297,7 @@ void ANetworkPredictionExtrasFlyingPawn::ProduceInput(const int32 DeltaMS, FFlyi
 		case 1:
 		{
 			// Free camera Restricted 2D movement on XY plane.
-			APlayerController* PC = Cast<APlayerController>(Controller);
+			APlayerController* PC = Cast<APlayerController>(GetController());
 			if (ensure(PC)) // Requires player controller for now
 			{
 				// Camera yaw rotation
@@ -329,7 +329,7 @@ void ANetworkPredictionExtrasFlyingPawn::ProduceInput(const int32 DeltaMS, FFlyi
 		case 2:
 		{
 			// Free camera on yaw and pitch, camera-relative movement.
-			APlayerController* PC = Cast<APlayerController>(Controller);
+			APlayerController* PC = Cast<APlayerController>(GetController());
 			if (ensure(PC)) // Requires player controller for now
 			{
 				// Camera yaw rotation
@@ -375,7 +375,7 @@ void ANetworkPredictionExtrasFlyingPawn::ProduceInput(const int32 DeltaMS, FFlyi
 		case 3:
 		{
 			// Free camera on the yaw, camera-relative motion
-			APlayerController* PC = Cast<APlayerController>(Controller);
+			APlayerController* PC = Cast<APlayerController>(GetController());
 			if (ensure(PC)) // Requires player controller for now
 			{
 				// Camera yaw rotation
@@ -504,7 +504,7 @@ void ANetworkPredictionExtrasFlyingPawn_MockAbility::SetupPlayerInputComponent(U
 {
 	Super::SetupPlayerInputComponent(PlayerInputComponent);
 
-	if (APlayerController* PC = Cast<APlayerController>(Controller))
+	if (APlayerController* PC = Cast<APlayerController>(GetController()))
 	{
 		if (PC->PlayerInput)
 		{
diff --git a/Engine/Source/Runtime/AIModule/Private/BehaviorTree/BlueprintNodeHelpers.cpp b/Engine/Source/Runtime/AIModule/Private/BehaviorTree/BlueprintNodeHelpers.cpp
--- a/Engine/Source/Runtime/AIModule/Private/BehaviorTree/BlueprintNodeHelpers.cpp
+++ b/Engine/Source/Runtime/AIModule/Private/BehaviorTree/BlueprintNodeHelpers.cpp
@@ -269,9 +269,9 @@ namespace BlueprintNodeHelpers
 		bool bFound = false;
 
 		APawn* OwningPawn = Cast<APawn>(OwningActor);
-		if (OwningPawn && OwningPawn->Controller)
+		if (OwningPawn && OwningPawn->GetController())
 		{
-			bFound = FindNodeOwner(OwningPawn->Controller, Node, OwningComp, OwningInstanceIdx);
+			bFound = FindNodeOwner(OwningPawn->GetController(), Node, OwningComp, OwningInstanceIdx);
 		}
 
 		if (OwningActor && !bFound)
diff --git a/Engine/Source/Runtime/AIModule/Private/Blueprint/AIBlueprintHelperLibrary.cpp b/Engine/Source/Runtime/AIModule/Private/Blueprint/AIBlueprintHelperLibrary.cpp
--- a/Engine/Source/Runtime/AIModule/Private/Blueprint/AIBlueprintHelperLibrary.cpp
+++ b/Engine/Source/Runtime/AIModule/Private/Blueprint/AIBlueprintHelperLibrary.cpp
@@ -199,14 +199,14 @@ APawn* UAIBlueprintHelperLibrary::SpawnAIFromClass(UObject* WorldContextObject,
 
 		if (NewPawn != NULL)
 		{
-			if (NewPawn->Controller == NULL)
+			if (NewPawn->GetController() == NULL)
 			{	// NOTE: SpawnDefaultController ALSO calls Possess() to possess the pawn (if a controller is successfully spawned).
 				NewPawn->SpawnDefaultController();
 			}
 
 			if (BehaviorTree != NULL)
 			{
-				AAIController* AIController = Cast<AAIController>(NewPawn->Controller);
+				AAIController* AIController = Cast<AAIController>(NewPawn->GetController());
 
 				if (AIController != NULL)
 				{
@@ -260,7 +260,7 @@ void UAIBlueprintHelperLibrary::LockAIResourcesWithAnimation(UAnimInstance* Anim
 	APawn* PawnOwner = AnimInstance->TryGetPawnOwner();
 	if (PawnOwner)
 	{
-		AAIController* OwningAI = Cast<AAIController>(PawnOwner->Controller);
+		AAIController* OwningAI = Cast<AAIController>(PawnOwner->GetController());
 		if (OwningAI)
 		{
 			if (bLockMovement && OwningAI->GetPathFollowingComponent())
@@ -285,7 +285,7 @@ void UAIBlueprintHelperLibrary::UnlockAIResourcesWithAnimation(UAnimInstance* An
 	APawn* PawnOwner = AnimInstance->TryGetPawnOwner();
 	if (PawnOwner)
 	{
-		AAIController* OwningAI = Cast<AAIController>(PawnOwner->Controller);
+		AAIController* OwningAI = Cast<AAIController>(PawnOwner->GetController());
 		if (OwningAI)
 		{
 			if (bUnlockMovement && OwningAI->GetPathFollowingComponent())
diff --git a/Engine/Source/Runtime/AIModule/Private/BrainComponent.cpp b/Engine/Source/Runtime/AIModule/Private/BrainComponent.cpp
--- a/Engine/Source/Runtime/AIModule/Private/BrainComponent.cpp
+++ b/Engine/Source/Runtime/AIModule/Private/BrainComponent.cpp
@@ -34,9 +34,9 @@ static UBrainComponent* FindBrainComponentHelper(const APawn* Pawn)
 
 	UBrainComponent* Comp = NULL;
 
-	if (Pawn->Controller)
+	if (Pawn->GetController())
 	{
-		Comp = FindBrainComponentHelper(Pawn->Controller);
+		Comp = FindBrainComponentHelper(Pawn->GetController());
 	}
 
 	if (Pawn && Comp == NULL)
diff --git a/Engine/Source/Runtime/AIModule/Private/EnvironmentQuery/EQSTestingPawn.cpp b/Engine/Source/Runtime/AIModule/Private/EnvironmentQuery/EQSTestingPawn.cpp
--- a/Engine/Source/Runtime/AIModule/Private/EnvironmentQuery/EQSTestingPawn.cpp
+++ b/Engine/Source/Runtime/AIModule/Private/EnvironmentQuery/EQSTestingPawn.cpp
@@ -243,7 +243,7 @@ void AEQSTestingPawn::Reset()
 	StepToDebugDraw = 0;
 	StepResults.Reset();
 
-	if (Controller == NULL)
+	if (GetController() == nullptr)
 	{
 		SpawnDefaultController();
 	}
diff --git a/Engine/Source/Runtime/AIModule/Private/GameplayDebugger/GameplayDebuggerCategory_AI.cpp b/Engine/Source/Runtime/AIModule/Private/GameplayDebugger/GameplayDebuggerCategory_AI.cpp
--- a/Engine/Source/Runtime/AIModule/Private/GameplayDebugger/GameplayDebuggerCategory_AI.cpp
+++ b/Engine/Source/Runtime/AIModule/Private/GameplayDebugger/GameplayDebuggerCategory_AI.cpp
@@ -136,7 +136,7 @@ void FGameplayDebuggerCategory_AI::CollectData(APlayerController* OwnerPC, AActo
 	DataPack.PawnName = MyPawn ? MyPawn->GetHumanReadableName() : FString(TEXT("{red}No selected pawn."));
 	DataPack.bIsUsingCharacter = (MyChar != nullptr);
 
-	AAIController* MyController = MyPawn ? Cast<AAIController>(MyPawn->Controller) : nullptr;
+	AAIController* MyController = MyPawn ? Cast<AAIController>(MyPawn->GetController()) : nullptr;
 	DataPack.bHasController = (MyController != nullptr);
 	if (MyController)
 	{
diff --git a/Engine/Source/Runtime/AIModule/Private/GameplayDebugger/GameplayDebuggerCategory_BehaviorTree.cpp b/Engine/Source/Runtime/AIModule/Private/GameplayDebugger/GameplayDebuggerCategory_BehaviorTree.cpp
--- a/Engine/Source/Runtime/AIModule/Private/GameplayDebugger/GameplayDebuggerCategory_BehaviorTree.cpp
+++ b/Engine/Source/Runtime/AIModule/Private/GameplayDebugger/GameplayDebuggerCategory_BehaviorTree.cpp
@@ -29,7 +29,7 @@ void FGameplayDebuggerCategory_BehaviorTree::FRepData::Serialize(FArchive& Ar)
 void FGameplayDebuggerCategory_BehaviorTree::CollectData(APlayerController* OwnerPC, AActor* DebugActor)
 {
 	APawn* MyPawn = Cast<APawn>(DebugActor);
-	AAIController* MyController = MyPawn ? Cast<AAIController>(MyPawn->Controller) : nullptr;
+	AAIController* MyController = MyPawn ? Cast<AAIController>(MyPawn->GetController()) : nullptr;
 	UBrainComponent* BrainComp = GetValid(MyController ? MyController->GetBrainComponent() : nullptr);
 	
 	if (BrainComp)
diff --git a/Engine/Source/Runtime/AIModule/Private/Perception/PawnSensingComponent.cpp b/Engine/Source/Runtime/AIModule/Private/Perception/PawnSensingComponent.cpp
--- a/Engine/Source/Runtime/AIModule/Private/Perception/PawnSensingComponent.cpp
+++ b/Engine/Source/Runtime/AIModule/Private/Perception/PawnSensingComponent.cpp
@@ -178,9 +178,9 @@ AController* UPawnSensingComponent::GetSensorController() const
 	else
 	{
 		APawn* Pawn = Cast<APawn>(SensorActor);
-		if (IsValid(Pawn) && IsValid(Pawn->Controller))
+		if (IsValid(Pawn) && IsValid(Pawn->GetController()))
 		{
-			return Pawn->Controller;
+			return Pawn->GetController();
 		}
 	}
 
@@ -442,7 +442,7 @@ bool UPawnSensingComponent::CanHear(const FVector& NoiseLoc, float Loudness, boo
 
 bool UPawnSensingComponent::ShouldCheckVisibilityOf(APawn *Pawn) const
 {
-	const bool bPawnIsPlayer = (Pawn->Controller && Pawn->Controller->PlayerState);
+	const bool bPawnIsPlayer = (Pawn->GetController() && Pawn->GetController()->PlayerState);
 	if (!bSeePawns || (bOnlySensePlayers && !bPawnIsPlayer))
 	{
 		return false;
@@ -464,7 +464,7 @@ bool UPawnSensingComponent::ShouldCheckVisibilityOf(APawn *Pawn) const
 
 bool UPawnSensingComponent::ShouldCheckAudibilityOf(APawn* Pawn) const
 {
-	const bool bPawnIsPlayer = (Pawn->Controller && Pawn->Controller->PlayerState);
+	const bool bPawnIsPlayer = (Pawn->GetController() && Pawn->GetController()->PlayerState);
 	if (!bHearNoises || (bOnlySensePlayers && !bPawnIsPlayer))
 	{
 		return false;
diff --git a/Engine/Source/Runtime/Engine/Classes/GameFramework/Character.h b/Engine/Source/Runtime/Engine/Classes/GameFramework/Character.h
--- a/Engine/Source/Runtime/Engine/Classes/GameFramework/Character.h
+++ b/Engine/Source/Runtime/Engine/Classes/GameFramework/Character.h
@@ -434,6 +434,11 @@ protected:
 	UPROPERTY(ReplicatedUsing=OnRep_ReplicatedBasedMovement)
 	struct FBasedMovementInfo ReplicatedBasedMovement;
 
+	/** This is the Replicated version of relative movement. Do not modify on simulated proxies! Returns non-const ReplicatedBasedMovement. */
+	ENGINE_API FBasedMovementInfo& GetReplicatedBasedMovement_Mutable();
+	/** This is the Replicated version of relative movement. Do not modify on simulated proxies! Sets ReplicatedBasedMovement*/
+	ENGINE_API void SetReplicatedBasedMovement(const FBasedMovementInfo& InReplicatedBasedMovement);
+
 public:
 	/** Rep notify for ReplicatedBasedMovement */
 	UFUNCTION()
@@ -475,13 +480,24 @@ protected:
 	UPROPERTY(Replicated)
 	float AnimRootMotionTranslationScale;
 
+	/** This is the value of the CharacterMovement timestamp, and is replicated down to SimulatedProxies to assist with visual smoothing of the character. May not be updated every frame to reduce bandwidth. This function sets ReplicatedServerLastTransformUpdateTimeStamp. */
+	ENGINE_API void SetReplicatedServerLastTransformUpdateTimeStamp(const float InReplicatedServerLastTransformUpdateTimeStamp);
+
+	/** This is updated every frame to the latest timestamp, but only replicated to Replays, this is because we want to save positions to the replay replication connection each frame to ensure smooth movement. Sets ReplayLastTransformUpdateTimeStamp.*/
+	ENGINE_API void SetReplayLastTransformUpdateTimeStamp(const float InReplayLastTransformUpdateTimeStamp);
+
+	/** CharacterMovement Custom gravity direction replicated for simulated proxies. This function sets ReplicatedGravityDirection and also marks property as dirty. */
+	ENGINE_API void SetReplicatedGravityDirection(const FVector& InReplicatedGravityDirection);
 public:
 	UFUNCTION()
 	ENGINE_API void OnRep_ReplayLastTransformUpdateTimeStamp();
 
-	/** Accessor for ReplicatedServerLastTransformUpdateTimeStamp. */
+	/** This is the value of the CharacterMovement timestamp, and is replicated down to SimulatedProxies to assist with visual smoothing of the character. May not be updated every frame to reduce bandwidth. This is the Accessor for ReplicatedServerLastTransformUpdateTimeStamp. */
 	FORCEINLINE float GetReplicatedServerLastTransformUpdateTimeStamp() const { return ReplicatedServerLastTransformUpdateTimeStamp; }
 
+	/** This is updated every frame to the latest timestamp, but only replicated to Replays, this is because we want to save positions to the replay replication connection each frame to ensure smooth movement. This is the Accessor for ReplayLastTransformUpdateTimeStamp. */
+	FORCEINLINE float GetReplayLastTransformUpdateTimeStamp() const { return ReplayLastTransformUpdateTimeStamp; }
+
 	/** Accessor for BasedMovement */
 	FORCEINLINE const FBasedMovementInfo& GetBasedMovement() const { return BasedMovement; }
 	
@@ -491,8 +507,8 @@ public:
 	/** Save a new relative location in BasedMovement and a new rotation with is either relative or absolute. */
 	ENGINE_API void SaveRelativeBasedMovement(const FVector& NewRelativeLocation, const FRotator& NewRotation, bool bRelativeRotation);
 
-	/** Returns ReplicatedMovementMode */
-	uint8 GetReplicatedMovementMode() const { return ReplicatedMovementMode; }
+	/** CharacterMovement MovementMode (and custom mode) replicated for simulated proxies. Use CharacterMovementComponent::UnpackNetworkMovementMode() to translate it. Returns ReplicatedMovementMode */
+	FORCEINLINE uint8 GetReplicatedMovementMode() const { return ReplicatedMovementMode; }
 
 	/** Get the saved translation offset of mesh. This is how much extra offset is applied from the center of the capsule. */
 	UFUNCTION(BlueprintCallable, Category=Character)
@@ -576,6 +592,9 @@ protected:
 	UPROPERTY(Replicated)
 	uint8 ReplicatedMovementMode;
 
+	/** Sets CharacterMovement MovementMode (and custom mode) replicated for simulated proxies. Use CharacterMovementComponent::UnpackNetworkMovementMode() to translate it. This function sets the compressed value of ReplicatedMovementMode.*/
+	ENGINE_API void SetReplicatedMovementMode(const uint8 InReplicatedMovementMode);
+
 public:
 	/** 
 	 * Jump key Held Time.
@@ -710,6 +729,16 @@ public:
 	UFUNCTION(BlueprintCallable, Category=Character)
 	ENGINE_API bool CanJump() const;
 
+	/** Indicates that this Character is currently under the force of a jump (if JumpMaxHoldTime is non-zero). IsJumpProvidingForce() handles this as well. This function returns bProxyIsJumpForceApplied */
+	FORCEINLINE bool GetProxyIsJumpForceApplied() const { return bProxyIsJumpForceApplied; }
+	/** Set to indicate that this Character is currently under the force of a jump (if JumpMaxHoldTime is non-zero). IsJumpProvidingForce() handles this as well. This function sets bProxyIsJumpForceApplied*/
+	ENGINE_API void SetProxyIsJumpForceApplied(const bool bInProxyIsJumpForceApplied);
+
+	/** Set by character movement to specify that this Character is currently crouched. This function returns bIsCrouched.*/
+	ENGINE_API bool IsCrouched() const;
+	/** Specifies whether this Character is currently crouched or not. This function sets bIsCrouched and marks property as dirty.*/
+	ENGINE_API void SetIsCrouched(const bool bInIsCrouched);
+
 protected:
 	/**
 	 * Customizable event to check if the character can jump in the current state.
@@ -1006,6 +1035,14 @@ public:
 	/** Replicated Root Motion montage */
 	UPROPERTY(ReplicatedUsing=OnRep_RootMotion)
 	struct FRepRootMotionMontage RepRootMotion;
+
+	/** Replicated Root Motion montage. Returns RepRootMotion */
+	FORCEINLINE const FRepRootMotionMontage& GetRepRootMotion() const { return RepRootMotion; }
+
+	/** Replicated Root Motion montage. Returns non-const RepRootMotion*/
+	ENGINE_API FRepRootMotionMontage& GetRepRootMotion_Mutable();
+	/** Replicated Root Motion montage. Sets RepRootMotion*/
+	ENGINE_API void SetRepRootMotion(const FRepRootMotionMontage& InRepRootMotion);
 	
 	/** Handles replicated root motion properties on simulated proxies and position correction. */
 	UFUNCTION()
@@ -1035,7 +1072,7 @@ public:
 	/** Sets scale to apply to root motion translation on this Character */
 	ENGINE_API void SetAnimRootMotionTranslationScale(float InAnimRootMotionTranslationScale = 1.f);
 
-	/** Returns current value of AnimRootMotionScale */
+	/** Scale to apply to root motion translation on this Character. Returns current value of AnimRootMotionScale. */
 	UFUNCTION(BlueprintCallable, Category=Animation)
 	ENGINE_API float GetAnimRootMotionTranslationScale() const;
 
diff --git a/Engine/Source/Runtime/Engine/Classes/GameFramework/Pawn.h b/Engine/Source/Runtime/Engine/Classes/GameFramework/Pawn.h
--- a/Engine/Source/Runtime/Engine/Classes/GameFramework/Pawn.h
+++ b/Engine/Source/Runtime/Engine/Classes/GameFramework/Pawn.h
@@ -205,6 +205,9 @@ public:
 	/** Handle StartFire() passed from PlayerController */
 	ENGINE_API virtual void PawnStartFire(uint8 FireModeNum = 0);
 
+	/** Replicated pitch value so we can see where remote clients are looking. */
+	FORCEINLINE uint8 GetRemoteViewPitch() const { return RemoteViewPitch; }
+
 	/**
 	 * Set Pawn ViewPitch, so we can see where remote clients are looking.
 	 * Maps 360.0 degrees into a byte
@@ -249,6 +252,9 @@ public:
 		return Cast<T>(GetController());
 	}
 
+	/** Sets controller for this actor */
+	ENGINE_API void SetController(AController* InController);
+
 	/** Get the rotation of the Controller, often the 'view' rotation of this Pawn. */
 	UFUNCTION(BlueprintCallable, Category=Pawn)
 	ENGINE_API FRotator GetControlRotation() const;
diff --git a/Engine/Source/Runtime/Engine/Private/Actor.cpp b/Engine/Source/Runtime/Engine/Private/Actor.cpp
--- a/Engine/Source/Runtime/Engine/Private/Actor.cpp
+++ b/Engine/Source/Runtime/Engine/Private/Actor.cpp
@@ -3434,7 +3434,7 @@ APawn* AActor::GetInstigator() const
 
 AController* AActor::GetInstigatorController() const
 {
-	return Instigator ? Instigator->Controller : nullptr;
+	return Instigator ? Instigator->GetController() : nullptr;
 }
 
 void AActor::CalcCamera(float DeltaTime, FMinimalViewInfo& OutResult)
diff --git a/Engine/Source/Runtime/Engine/Private/Character.cpp b/Engine/Source/Runtime/Engine/Private/Character.cpp
--- a/Engine/Source/Runtime/Engine/Private/Character.cpp
+++ b/Engine/Source/Runtime/Engine/Private/Character.cpp
@@ -91,7 +91,7 @@ ACharacter::ACharacter(const FObjectInitializer& ObjectInitializer)
 	JumpCurrentCountPreJump = 0;
     bWasJumping = false;
 
-	AnimRootMotionTranslationScale = 1.0f;
+	SetAnimRootMotionTranslationScale(1.0f);
 
 #if WITH_EDITORONLY_DATA
 	ArrowComponent = CreateEditorOnlyDefaultSubobject<UArrowComponent>(TEXT("Arrow"));
@@ -133,7 +133,7 @@ ACharacter::ACharacter(const FObjectInitializer& ObjectInitializer)
 	}
 
 	BaseRotationOffset = FQuat::Identity;
-	ReplicatedGravityDirection = UCharacterMovementComponent::DefaultGravityDirection;
+	SetReplicatedGravityDirection(UCharacterMovementComponent::DefaultGravityDirection);
 }
 
 void ACharacter::PostInitializeComponents()
@@ -160,7 +160,7 @@ void ACharacter::PostInitializeComponents()
 			CharacterMovement->UpdateNavAgent(*CapsuleComponent);
 		}
 
-		if (Controller == nullptr && GetNetMode() != NM_Client)
+		if (GetController() == nullptr && GetNetMode() != NM_Client)
 		{
 			if (CharacterMovement && CharacterMovement->bRunPhysicsWithNoController)
 			{
@@ -301,7 +301,7 @@ bool ACharacter::CanJump() const
 
 bool ACharacter::CanJumpInternal_Implementation() const
 {
-	return !bIsCrouched && JumpIsAllowedInternal();
+	return !IsCrouched() && JumpIsAllowedInternal();
 }
 
 bool ACharacter::JumpIsAllowedInternal() const
@@ -337,6 +337,21 @@ bool ACharacter::JumpIsAllowedInternal() const
 	return bJumpIsAllowed;
 }
 
+void ACharacter::SetProxyIsJumpForceApplied(const bool bInProxyIsJumpForceApplied)
+{
+	bProxyIsJumpForceApplied = bInProxyIsJumpForceApplied;
+}
+
+bool ACharacter::IsCrouched() const
+{
+	return bIsCrouched;
+}
+
+void ACharacter::SetIsCrouched(const bool bInIsCrouched)
+{
+	bIsCrouched = bInIsCrouched;
+}
+
 void ACharacter::ResetJumpState()
 {
 	bPressedJump = false;
@@ -361,7 +376,7 @@ bool ACharacter::IsJumpProvidingForce() const
 	{
 		return true;
 	}
-	else if (bProxyIsJumpForceApplied && (GetLocalRole() == ROLE_SimulatedProxy))
+	else if (GetProxyIsJumpForceApplied() && (GetLocalRole() == ROLE_SimulatedProxy))
 	{
 		return GetWorld()->TimeSince(ProxyJumpForceStartedTime) <= GetJumpMaxHoldTime();
 	}
@@ -371,7 +386,7 @@ bool ACharacter::IsJumpProvidingForce() const
 
 void ACharacter::RecalculateBaseEyeHeight()
 {
-	if (!bIsCrouched)
+	if (!IsCrouched())
 	{
 		Super::RecalculateBaseEyeHeight();
 	}
@@ -386,7 +401,7 @@ void ACharacter::OnRep_IsCrouched()
 {
 	if (CharacterMovement)
 	{
-		if (bIsCrouched)
+		if (IsCrouched())
 		{
 			CharacterMovement->bWantsToCrouch = true;
 			CharacterMovement->Crouch(true);
@@ -419,7 +434,7 @@ void ACharacter::SetReplicateMovement(bool bInReplicateMovement)
 
 bool ACharacter::CanCrouch() const
 {
-	return !bIsCrouched && CharacterMovement && CharacterMovement->CanEverCrouch() && GetRootComponent() && !GetRootComponent()->IsSimulatingPhysics();
+	return !IsCrouched() && CharacterMovement && CharacterMovement->CanEverCrouch() && GetRootComponent() && !GetRootComponent()->IsSimulatingPhysics();
 }
 
 void ACharacter::Crouch(bool bClientSimulation)
@@ -841,6 +856,30 @@ void ACharacter::SetBase( UPrimitiveComponent* NewBaseComponent, const FName InB
 	}
 }
 
+void ACharacter::SetReplicatedServerLastTransformUpdateTimeStamp(float InReplicatedServerLastTransformUpdateTimeStamp)
+{
+	ReplicatedServerLastTransformUpdateTimeStamp = InReplicatedServerLastTransformUpdateTimeStamp; 
+}
+
+void ACharacter::SetReplayLastTransformUpdateTimeStamp(float InReplayLastTransformUpdateTimeStamp)
+{
+	ReplayLastTransformUpdateTimeStamp = InReplayLastTransformUpdateTimeStamp;
+}
+
+FBasedMovementInfo& ACharacter::GetReplicatedBasedMovement_Mutable()
+{
+	return ReplicatedBasedMovement;
+}
+
+void ACharacter::SetReplicatedBasedMovement(const FBasedMovementInfo& InReplicatedBasedMovement)
+{
+	ReplicatedBasedMovement = InReplicatedBasedMovement;
+}
+
+void ACharacter::SetReplicatedMovementMode(const uint8 InReplicatedMovementMode)
+{
+	ReplicatedMovementMode = InReplicatedMovementMode; 
+}
 
 void ACharacter::SaveRelativeBasedMovement(const FVector& NewRelativeLocation, const FRotator& NewRotation, bool bRelativeRotation)
 {
@@ -879,6 +918,11 @@ FVector ACharacter::GetReplicatedGravityDirection() const
 	return ReplicatedGravityDirection;
 }
 
+void ACharacter::SetReplicatedGravityDirection(const FVector& InReplicatedGravityDirection)
+{
+	ReplicatedGravityDirection = InReplicatedGravityDirection;
+}
+
 FVector ACharacter::GetNavAgentLocation() const
 {
 	FVector AgentLocation = FNavigationSystem::InvalidLocation;
@@ -1085,7 +1129,7 @@ void ACharacter::OnMovementModeChanged(EMovementMode PrevMovementMode, uint8 Pre
 	}
 
 	// Record jump force start time for proxies. Allows us to expire the jump even if not continually ticking down a timer.
-	if (bProxyIsJumpForceApplied && CharacterMovement->IsFalling())
+	if (GetProxyIsJumpForceApplied() && CharacterMovement->IsFalling())
 	{
 		ProxyJumpForceStartedTime = GetWorld()->GetTimeSeconds();
 	}
@@ -1188,8 +1232,8 @@ static uint8 SavedMovementMode;
 
 void ACharacter::PreNetReceive()
 {
-	SavedMovementMode = ReplicatedMovementMode;
-	PreNetReceivedGravityDirection = ReplicatedGravityDirection;
+	SavedMovementMode = GetReplicatedMovementMode();
+	PreNetReceivedGravityDirection = GetReplicatedGravityDirection();
 	Super::PreNetReceive();
 }
 
@@ -1197,8 +1241,8 @@ void ACharacter::PostNetReceive()
 {
 	if (GetLocalRole() == ROLE_SimulatedProxy)
 	{
-		CharacterMovement->bNetworkGravityDirectionChanged = !PreNetReceivedGravityDirection.Equals(ReplicatedGravityDirection);
-		CharacterMovement->bNetworkMovementModeChanged |= ((SavedMovementMode != ReplicatedMovementMode) || (CharacterMovement->PackNetworkMovementMode() != ReplicatedMovementMode));
+		CharacterMovement->bNetworkGravityDirectionChanged = !PreNetReceivedGravityDirection.Equals(GetReplicatedGravityDirection());
+		CharacterMovement->bNetworkMovementModeChanged |= ((SavedMovementMode != GetReplicatedMovementMode()) || (CharacterMovement->PackNetworkMovementMode() != GetReplicatedMovementMode()));
 		CharacterMovement->bNetworkUpdateReceived |= CharacterMovement->bNetworkMovementModeChanged || CharacterMovement->bJustTeleported || CharacterMovement->bNetworkGravityDirectionChanged;
 	}
 
@@ -1227,30 +1271,31 @@ void ACharacter::OnRep_ReplicatedBasedMovement()
 	CharacterMovement->bNetworkUpdateReceived = true;
 	FGuardValue_Bitfield(bInBaseReplication, true);
 
-	const bool bBaseChanged = (BasedMovement.MovementBase != ReplicatedBasedMovement.MovementBase || BasedMovement.BoneName != ReplicatedBasedMovement.BoneName);
+	const FBasedMovementInfo& RepBasedMovement = GetReplicatedBasedMovement();
+	const bool bBaseChanged = (BasedMovement.MovementBase != RepBasedMovement.MovementBase || BasedMovement.BoneName != RepBasedMovement.BoneName);
 	if (bBaseChanged)
 	{
 		// Even though we will copy the replicated based movement info, we need to use SetBase() to set up tick dependencies and trigger notifications.
-		SetBase(ReplicatedBasedMovement.MovementBase, ReplicatedBasedMovement.BoneName);
+		SetBase(RepBasedMovement.MovementBase, RepBasedMovement.BoneName);
 	}
 
 	// Make sure to use the values of relative location/rotation etc from the server.
-	BasedMovement = ReplicatedBasedMovement;
+	BasedMovement = RepBasedMovement;
 
-	if (ReplicatedBasedMovement.HasRelativeLocation())
+	if (RepBasedMovement.HasRelativeLocation())
 	{
 		// Update transform relative to movement base
 		const FVector OldLocation = GetActorLocation();
 		const FQuat OldRotation = GetActorQuat();
-		MovementBaseUtility::GetMovementBaseTransform(ReplicatedBasedMovement.MovementBase, ReplicatedBasedMovement.BoneName, CharacterMovement->OldBaseLocation, CharacterMovement->OldBaseQuat);
+		MovementBaseUtility::GetMovementBaseTransform(RepBasedMovement.MovementBase, RepBasedMovement.BoneName, CharacterMovement->OldBaseLocation, CharacterMovement->OldBaseQuat);
 		const FTransform BaseTransform(CharacterMovement->OldBaseQuat, CharacterMovement->OldBaseLocation);
-		const FVector NewLocation = BaseTransform.TransformPositionNoScale(ReplicatedBasedMovement.Location);
+		const FVector NewLocation = BaseTransform.TransformPositionNoScale(RepBasedMovement.Location);
 		FRotator NewRotation;
 
-		if (ReplicatedBasedMovement.HasRelativeRotation())
+		if (RepBasedMovement.HasRelativeRotation())
 		{
 			// Relative location, relative rotation
-			NewRotation = (FRotationMatrix(ReplicatedBasedMovement.Rotation) * FQuatRotationMatrix(CharacterMovement->OldBaseQuat)).Rotator();
+			NewRotation = (FRotationMatrix(RepBasedMovement.Rotation) * FQuatRotationMatrix(CharacterMovement->OldBaseQuat)).Rotator();
 			
 			if (CharacterMovement->ShouldRemainVertical())
 			{
@@ -1261,7 +1306,7 @@ void ACharacter::OnRep_ReplicatedBasedMovement()
 		else
 		{
 			// Relative location, absolute rotation
-			NewRotation = ReplicatedBasedMovement.Rotation;
+			NewRotation = RepBasedMovement.Rotation;
 		}
 
 		// When position or base changes, movement mode will need to be updated. This assumes rotation changes don't affect that.
@@ -1306,7 +1351,7 @@ void ACharacter::OnRep_ReplicatedMovement()
 
 void ACharacter::OnRep_ReplayLastTransformUpdateTimeStamp()
 {
-	ReplicatedServerLastTransformUpdateTimeStamp = ReplayLastTransformUpdateTimeStamp;
+	SetReplicatedServerLastTransformUpdateTimeStamp(GetReplayLastTransformUpdateTimeStamp());
 }
 
 /** Get FAnimMontageInstance playing RootMotion */
@@ -1315,6 +1360,15 @@ FAnimMontageInstance * ACharacter::GetRootMotionAnimMontageInstance() const
 	return (Mesh && Mesh->GetAnimInstance()) ? Mesh->GetAnimInstance()->GetRootMotionMontageInstance() : nullptr;
 }
 
+FRepRootMotionMontage& ACharacter::GetRepRootMotion_Mutable()
+{
+	return RepRootMotion;
+}
+void ACharacter::SetRepRootMotion(const FRepRootMotionMontage& InRepRootMotion)
+{
+	RepRootMotion = InRepRootMotion;
+}
+
 void ACharacter::OnRep_RootMotion()
 {
 	// Following the same pattern in AActor::OnRep_ReplicatedMovement() just in case...
@@ -1329,12 +1383,12 @@ void ACharacter::OnRep_RootMotion()
 		UE_LOG(LogRootMotion, Log,  TEXT("ACharacter::OnRep_RootMotion"));
 
 		// Save received move in queue, we'll try to use it during Tick().
-		if( RepRootMotion.bIsActive )
+		if( GetRepRootMotion().bIsActive )
 		{
 			// Add new move
 			RootMotionRepMoves.AddZeroed(1);
 			FSimulatedRootMotionReplicatedMove& NewMove = RootMotionRepMoves.Last();
-			NewMove.RootMotion = RepRootMotion;
+			NewMove.RootMotion = GetRepRootMotion();
 			NewMove.Time = GetWorld()->GetTimeSeconds();
 		}
 		else
@@ -1537,7 +1591,7 @@ void ACharacter::PostNetReceiveLocationAndRotation()
 	if(GetLocalRole() == ROLE_SimulatedProxy)
 	{
 		// Don't change transform if using relative position (it should be nearly the same anyway, or base may be slightly out of sync)
-		if (!ReplicatedBasedMovement.HasRelativeLocation())
+		if (!GetReplicatedBasedMovement().HasRelativeLocation())
 		{
 			const FRepMovement& ConstRepMovement = GetReplicatedMovement();
 			const FVector OldLocation = GetActorLocation();
@@ -1561,58 +1615,62 @@ void ACharacter::PreReplication( IRepChangedPropertyTracker & ChangedPropertyTra
 	{
 		const FAnimMontageInstance* RootMotionMontageInstance = GetRootMotionAnimMontageInstance();
 
-		RepRootMotion.bIsActive = true;
+		FRepRootMotionMontage& RepRootMotionRef = GetRepRootMotion_Mutable();
+		RepRootMotionRef.bIsActive = true;
 		// Is position stored in local space?
-		RepRootMotion.bRelativePosition = BasedMovement.HasRelativeLocation();
-		RepRootMotion.bRelativeRotation = BasedMovement.HasRelativeRotation();
-		RepRootMotion.Location			= RepRootMotion.bRelativePosition ? BasedMovement.Location : FRepMovement::RebaseOntoZeroOrigin(GetActorLocation(), GetWorld()->OriginLocation);
-		RepRootMotion.Rotation			= RepRootMotion.bRelativeRotation ? BasedMovement.Rotation : GetActorRotation();
-		RepRootMotion.MovementBase		= BasedMovement.MovementBase;
-		RepRootMotion.MovementBaseBoneName = BasedMovement.BoneName;
+		RepRootMotionRef.bRelativePosition = BasedMovement.HasRelativeLocation();
+		RepRootMotionRef.bRelativeRotation = BasedMovement.HasRelativeRotation();
+		RepRootMotionRef.Location			= RepRootMotionRef.bRelativePosition ? BasedMovement.Location : FRepMovement::RebaseOntoZeroOrigin(GetActorLocation(), GetWorld()->OriginLocation);
+		RepRootMotionRef.Rotation			= RepRootMotionRef.bRelativeRotation ? BasedMovement.Rotation : GetActorRotation();
+		RepRootMotionRef.MovementBase		= BasedMovement.MovementBase;
+		RepRootMotionRef.MovementBaseBoneName = BasedMovement.BoneName;
 		if (RootMotionMontageInstance)
 		{
-			RepRootMotion.Animation = RootMotionMontageInstance->Montage->IsDynamicMontage() ? RootMotionMontageInstance->Montage->GetFirstAnimReference() : RootMotionMontageInstance->Montage;
-			RepRootMotion.Position			= RootMotionMontageInstance->GetPosition();
+			RepRootMotionRef.Animation = RootMotionMontageInstance->Montage->IsDynamicMontage() ? RootMotionMontageInstance->Montage->GetFirstAnimReference() : RootMotionMontageInstance->Montage;
+			RepRootMotionRef.Position			= RootMotionMontageInstance->GetPosition();
 		}
 		else
 		{
-			RepRootMotion.Animation = nullptr;
+			RepRootMotionRef.Animation = nullptr;
 		}
 
-		RepRootMotion.AuthoritativeRootMotion = CharacterMovement->CurrentRootMotion;
-		RepRootMotion.Acceleration = CharacterMovement->GetCurrentAcceleration();
-		RepRootMotion.LinearVelocity = CharacterMovement->Velocity;
+		RepRootMotionRef.AuthoritativeRootMotion = CharacterMovement->CurrentRootMotion;
+		RepRootMotionRef.Acceleration = CharacterMovement->GetCurrentAcceleration();
+		RepRootMotionRef.LinearVelocity = CharacterMovement->Velocity;
 
-		DOREPLIFETIME_ACTIVE_OVERRIDE_FAST( ACharacter, RepRootMotion, true );
+		DOREPLIFETIME_ACTIVE_OVERRIDE_FAST(ACharacter, RepRootMotion, true);
 	}
 	else
 	{
-		const bool bWasRootMotionPreviouslyActive = RepRootMotion.bIsActive;
-		RepRootMotion.Clear();
+		const bool bWasRootMotionPreviouslyActive = GetRepRootMotion().bIsActive;
 
+		FRepRootMotionMontage& RepRootMotionRef = GetRepRootMotion_Mutable();
+		RepRootMotionRef.Clear();
+		
 		// Replicate RepRootMotion one last time when root motion ends, so that clients see the change.
 		// Then deactivate subsequent property comparisons and replication updates until root motion starts again.
-		DOREPLIFETIME_ACTIVE_OVERRIDE_FAST( ACharacter, RepRootMotion, bWasRootMotionPreviouslyActive );
+		DOREPLIFETIME_ACTIVE_OVERRIDE_FAST(ACharacter, RepRootMotion, bWasRootMotionPreviouslyActive);
 	}
 
-	bProxyIsJumpForceApplied = (JumpForceTimeRemaining > 0.0f);
-	ReplicatedMovementMode = CharacterMovement->PackNetworkMovementMode();
-	ReplicatedGravityDirection = CharacterMovement->GetGravityDirection();
+	SetProxyIsJumpForceApplied(JumpForceTimeRemaining > 0.0f);
+	SetReplicatedMovementMode(CharacterMovement->PackNetworkMovementMode());
+	SetReplicatedGravityDirection(CharacterMovement->GetGravityDirection());
 
 	if(IsReplicatingMovement())
 	{
-		ReplicatedBasedMovement = BasedMovement;
+		FBasedMovementInfo& ReplicatedBasedMovementRef = GetReplicatedBasedMovement_Mutable();
+		ReplicatedBasedMovementRef = BasedMovement;
 
 		// Optimization: only update and replicate these values if they are actually going to be used.
 		if (BasedMovement.HasRelativeLocation())
 		{
 			// When velocity becomes zero, force replication so the position is updated to match the server (it may have moved due to simulation on the client).
-			ReplicatedBasedMovement.bServerHasVelocity = !CharacterMovement->Velocity.IsZero();
+			ReplicatedBasedMovementRef.bServerHasVelocity = !CharacterMovement->Velocity.IsZero();
 
 			// Make sure absolute rotations are updated in case rotation occurred after the base info was saved.
 			if (!BasedMovement.HasRelativeRotation())
 			{
-				ReplicatedBasedMovement.Rotation = GetActorRotation();
+				ReplicatedBasedMovementRef.Rotation = GetActorRotation();
 			}
 		}
 	}
@@ -1622,11 +1680,11 @@ void ACharacter::PreReplication( IRepChangedPropertyTracker & ChangedPropertyTra
 	// Save bandwidth by not replicating this value unless it is necessary, since it changes every update.
 	if ((CharacterMovement->NetworkSmoothingMode == ENetworkSmoothingMode::Linear) || CharacterMovement->bNetworkAlwaysReplicateTransformUpdateTimestamp)
 	{
-		ReplicatedServerLastTransformUpdateTimeStamp = CharacterMovement->GetServerLastTransformUpdateTimeStamp();
+		SetReplicatedServerLastTransformUpdateTimeStamp(CharacterMovement->GetServerLastTransformUpdateTimeStamp());
 	}
 	else
 	{
-		ReplicatedServerLastTransformUpdateTimeStamp = 0.f;
+		SetReplicatedServerLastTransformUpdateTimeStamp(0.f);
 	}
 }
 
@@ -1646,12 +1704,12 @@ void ACharacter::PreReplicationForReplay(IRepChangedPropertyTracker& ChangedProp
 	{
 		// On client replays, our view pitch will be set to 0 as by default we do not replicate
 		// pitch for owners, just for simulated. So instead push our rotation into the sampler
-		if (World->IsRecordingClientReplay() && Controller != nullptr && GetLocalRole() == ROLE_AutonomousProxy && GetNetMode() == NM_Client)
+		if (World->IsRecordingClientReplay() && GetController() != nullptr && GetLocalRole() == ROLE_AutonomousProxy && GetNetMode() == NM_Client)
 		{
-			SetRemoteViewPitch(Controller->GetControlRotation().Pitch);
+			SetRemoteViewPitch(GetController()->GetControlRotation().Pitch);
 		}
 
-		ReplayLastTransformUpdateTimeStamp = World->GetTimeSeconds();
+		SetReplayLastTransformUpdateTimeStamp(World->GetTimeSeconds());
 	}
 }
 
@@ -1926,7 +1984,7 @@ void ACharacter::FillAsyncInput(FCharacterAsyncInput& Input) const
 	Input.bUseControllerRotationPitch = bUseControllerRotationPitch;
 	Input.bUseControllerRotationYaw = bUseControllerRotationYaw;
 	Input.bUseControllerRotationRoll = bUseControllerRotationRoll;
-	Input.ControllerDesiredRotation = Controller->GetDesiredRotation();
+	Input.ControllerDesiredRotation = GetController()->GetDesiredRotation();
 }
 
 void ACharacter::InitializeAsyncOutput(FCharacterAsyncOutput& Output) const
diff --git a/Engine/Source/Runtime/Engine/Private/CheatManager.cpp b/Engine/Source/Runtime/Engine/Private/CheatManager.cpp
--- a/Engine/Source/Runtime/Engine/Private/CheatManager.cpp
+++ b/Engine/Source/Runtime/Engine/Private/CheatManager.cpp
@@ -324,10 +324,10 @@ void UCheatManager::DestroyTarget()
         APawn* Pawn = Cast<APawn>(TargetActor);
 		if (Pawn != NULL)
 		{
-			if ((Pawn->Controller != NULL) && (Cast<APlayerController>(Pawn->Controller) == NULL))
+			if ((Pawn->GetController() != NULL) && (Cast<APlayerController>(Pawn->GetController()) == NULL))
 			{
 				// Destroy any associated controller as long as it's not a player controller.
-				Pawn->Controller->Destroy();
+				Pawn->GetController()->Destroy();
 			}
 		}
         
@@ -345,10 +345,10 @@ void UCheatManager::DestroyAll(TSubclassOf<AActor> aClass)
 			APawn* Pawn = Cast<APawn>(A);
 			if (Pawn != NULL)
 			{
-				if ((Pawn->Controller != NULL) && (Cast<APlayerController>(Pawn->Controller) == NULL))
+				if ((Pawn->GetController() != NULL) && (Cast<APlayerController>(Pawn->GetController()) == NULL))
 				{
 					// Destroy any associated controller as long as it's not a player controller.
-					Pawn->Controller->Destroy();
+					Pawn->GetController()->Destroy();
 				}
 			}
 			A->Destroy();
@@ -369,11 +369,11 @@ void UCheatManager::DestroyAllPawnsExceptTarget()
 			APawn* Pawn = *It;
 			if (IsValidChecked(Pawn))
 			{
-				if ((Pawn != HitPawnTarget) && Cast<APlayerController>(Pawn->Controller) == NULL)
+				if ((Pawn != HitPawnTarget) && Cast<APlayerController>(Pawn->GetController()) == NULL)
 				{
-					if (Pawn->Controller != NULL)
+					if (Pawn->GetController() != NULL)
 					{
-						Pawn->Controller->Destroy();
+						Pawn->GetController()->Destroy();
 					}
 					Pawn->Destroy();
 				}
@@ -391,11 +391,11 @@ void UCheatManager::DestroyPawns(TSubclassOf<APawn> aClass)
 	for (TActorIterator<APawn> It(GetWorld(), aClass); It; ++It)
 	{
 		APawn* Pawn = *It;
-		if ( Cast<APlayerController>(Pawn->Controller) == NULL )
+		if ( Cast<APlayerController>(Pawn->GetController()) == NULL )
 		{
-			if ( Pawn->Controller != NULL )
+			if ( Pawn->GetController() != NULL )
 			{
-				Pawn->Controller->Destroy();
+				Pawn->GetController()->Destroy();
 			}
 			Pawn->Destroy();
 		}
diff --git a/Engine/Source/Runtime/Engine/Private/Components/CharacterMovementComponent.cpp b/Engine/Source/Runtime/Engine/Private/Components/CharacterMovementComponent.cpp
--- a/Engine/Source/Runtime/Engine/Private/Components/CharacterMovementComponent.cpp
+++ b/Engine/Source/Runtime/Engine/Private/Components/CharacterMovementComponent.cpp
@@ -1682,8 +1682,8 @@ void UCharacterMovementComponent::TickComponent(float DeltaTime, enum ELevelTick
 		// Perform input-driven move for any locally-controlled character, and also
 		// allow animation root motion or physics to move characters even if they have no controller
 		const bool bShouldPerformControlledCharMove = CharacterOwner->IsLocallyControlled() 
-													  || (!CharacterOwner->Controller && bRunPhysicsWithNoController)		
-													  || (!CharacterOwner->Controller && CharacterOwner->IsPlayingRootMotion());
+													  || (!CharacterOwner->GetController() && bRunPhysicsWithNoController)		
+													  || (!CharacterOwner->GetController() && CharacterOwner->IsPlayingRootMotion());
 
 		if (bShouldPerformControlledCharMove)
 		{
@@ -2473,7 +2473,7 @@ void UCharacterMovementComponent::UpdateBasedMovement(float DeltaSeconds)
 			{
 				// Nothing changed. This means we probably are using another rotation mechanism (bOrientToMovement etc). We should still follow the base object.
 				// @todo: This assumes only Yaw is used, currently a valid assumption. This is the only reason FaceRotation() is used above really, aside from being a virtual hook.
-				if (bOrientRotationToMovement || (bUseControllerDesiredRotation && CharacterOwner->Controller))
+				if (bOrientRotationToMovement || (bUseControllerDesiredRotation && CharacterOwner->GetController()))
 				{
 					// Custom gravity automatically aligns the character to the gravity direction, so we shouldn't zero out pitch and roll.
 					if (!HasCustomGravity())
@@ -2487,7 +2487,7 @@ void UCharacterMovementComponent::UpdateBasedMovement(float DeltaSeconds)
 			}
 
 			// Pipe through ControlRotation, to affect camera.
-			if (CharacterOwner->Controller)
+			if (CharacterOwner->GetController())
 			{
 				const FQuat PawnDeltaRotation = FinalQuat * PawnOldQuat.Inverse();
 				FRotator FinalRotation = FinalQuat.Rotator();
@@ -2594,7 +2594,7 @@ void UCharacterMovementComponent::OnUnableToFollowBaseMove(const FVector& DeltaP
 
 void UCharacterMovementComponent::UpdateBasedRotation(FRotator& FinalRotation, const FRotator& ReducedRotation)
 {
-	AController* Controller = CharacterOwner ? CharacterOwner->Controller : nullptr;
+	AController* Controller = CharacterOwner ? CharacterOwner->GetController() : nullptr;
 	float ControllerRoll = 0.f;
 
 	if ((Controller != nullptr) && !bIgnoreBaseRotation)
@@ -3140,7 +3140,7 @@ void UCharacterMovementComponent::Crouch(bool bClientSimulation)
 	{
 		if (!bClientSimulation)
 		{
-			CharacterOwner->bIsCrouched = true;
+			CharacterOwner->SetIsCrouched(true);
 		}
 		CharacterOwner->OnStartCrouch( 0.f, 0.f );
 		return;
@@ -3189,7 +3189,7 @@ void UCharacterMovementComponent::Crouch(bool bClientSimulation)
 			UpdatedComponent->MoveComponent(ScaledHalfHeightAdjust * GetGravityDirection(), UpdatedComponent->GetComponentQuat(), true, nullptr, EMoveComponentFlags::MOVECOMP_NoFlags, ETeleportType::TeleportPhysics);
 		}
 
-		CharacterOwner->bIsCrouched = true;
+		CharacterOwner->SetIsCrouched(true);
 	}
 
 	bForceNextFloorCheck = true;
@@ -3230,7 +3230,7 @@ void UCharacterMovementComponent::UnCrouch(bool bClientSimulation)
 	{
 		if (!bClientSimulation)
 		{
-			CharacterOwner->bIsCrouched = false;
+			CharacterOwner->SetIsCrouched(false);
 		}
 		CharacterOwner->OnEndCrouch( 0.f, 0.f );
 		return;
@@ -3335,7 +3335,7 @@ void UCharacterMovementComponent::UnCrouch(bool bClientSimulation)
 			return;
 		}
 
-		CharacterOwner->bIsCrouched = false;
+		CharacterOwner->SetIsCrouched(false);
 	}	
 	else
 	{
@@ -3721,7 +3721,7 @@ bool UCharacterMovementComponent::IsSwimming() const
 
 bool UCharacterMovementComponent::IsCrouching() const
 {
-	return CharacterOwner && CharacterOwner->bIsCrouched;
+	return CharacterOwner && CharacterOwner->IsCrouched();
 }
 
 void UCharacterMovementComponent::CalcVelocity(float DeltaTime, float Friction, bool bFluid, float BrakingDeceleration)
@@ -5501,7 +5501,7 @@ void UCharacterMovementComponent::PhysWalking(float deltaTime, int32 Iterations)
 		return;
 	}
 
-	if (!CharacterOwner || (!CharacterOwner->Controller && !bRunPhysicsWithNoController && !HasAnimRootMotion() && !CurrentRootMotion.HasOverrideVelocity() && (CharacterOwner->GetLocalRole() != ROLE_SimulatedProxy)))
+	if (!CharacterOwner || (!CharacterOwner->GetController() && !bRunPhysicsWithNoController && !HasAnimRootMotion() && !CurrentRootMotion.HasOverrideVelocity() && (CharacterOwner->GetLocalRole() != ROLE_SimulatedProxy)))
 	{
 		Acceleration = FVector::ZeroVector;
 		Velocity = FVector::ZeroVector;
@@ -5525,7 +5525,7 @@ void UCharacterMovementComponent::PhysWalking(float deltaTime, int32 Iterations)
 	const uint8 StartingCustomMovementMode = CustomMovementMode;
 
 	// Perform the move
-	while ( (remainingTime >= MIN_TICK_TIME) && (Iterations < MaxSimulationIterations) && CharacterOwner && (CharacterOwner->Controller || bRunPhysicsWithNoController || HasAnimRootMotion() || CurrentRootMotion.HasOverrideVelocity() || (CharacterOwner->GetLocalRole() == ROLE_SimulatedProxy)) )
+	while ( (remainingTime >= MIN_TICK_TIME) && (Iterations < MaxSimulationIterations) && CharacterOwner && (CharacterOwner->GetController() || bRunPhysicsWithNoController || HasAnimRootMotion() || CurrentRootMotion.HasOverrideVelocity() || (CharacterOwner->GetLocalRole() == ROLE_SimulatedProxy)) )
 	{
 		Iterations++;
 		bJustTeleported = false;
@@ -5732,7 +5732,7 @@ void UCharacterMovementComponent::PhysNavWalking(float deltaTime, int32 Iteratio
 		return;
 	}
 
-	if ((!CharacterOwner || !CharacterOwner->Controller) && !bRunPhysicsWithNoController && !HasAnimRootMotion() && !CurrentRootMotion.HasOverrideVelocity())
+	if ((!CharacterOwner || !CharacterOwner->GetController()) && !bRunPhysicsWithNoController && !HasAnimRootMotion() && !CurrentRootMotion.HasOverrideVelocity())
 	{
 		Acceleration = FVector::ZeroVector;
 		Velocity = FVector::ZeroVector;
@@ -6425,7 +6425,7 @@ void UCharacterMovementComponent::PhysicsRotation(float DeltaTime)
 		return;
 	}
 
-	if (!HasValidData() || (!CharacterOwner->Controller && !bRunPhysicsWithNoController))
+	if (!HasValidData() || (!CharacterOwner->GetController() && !bRunPhysicsWithNoController))
 	{
 		return;
 	}
@@ -6441,11 +6441,11 @@ void UCharacterMovementComponent::PhysicsRotation(float DeltaTime)
 	{
 		DesiredRotation = ComputeOrientToMovementRotation(CurrentRotation, DeltaTime, DeltaRot);
 	}
-	else if (CharacterOwner->Controller && bUseControllerDesiredRotation)
+	else if (CharacterOwner->GetController() && bUseControllerDesiredRotation)
 	{
-		DesiredRotation = CharacterOwner->Controller->GetDesiredRotation();
+		DesiredRotation = CharacterOwner->GetController()->GetDesiredRotation();
 	}
-	else if (!CharacterOwner->Controller && bRunPhysicsWithNoController && bUseControllerDesiredRotation)
+	else if (!CharacterOwner->GetController() && bRunPhysicsWithNoController && bUseControllerDesiredRotation)
 	{
 		if (AController* ControllerOwner = Cast<AController>(CharacterOwner->GetOwner()))
 		{
@@ -13231,7 +13231,7 @@ void UCharacterMovementComponent::AccumulateRootMotionForAsync(float DeltaSecond
 
 void UCharacterMovementComponent::FillAsyncInput(const FVector& InputVector, FCharacterMovementComponentAsyncInput& AsyncInput)
 {
-	if (!CharacterOwner || !CharacterOwner->Controller)
+	if (!CharacterOwner || !CharacterOwner->GetController())
 	{
 		return;
 	}
@@ -13434,7 +13434,7 @@ void UCharacterMovementComponent::FillAsyncInput(const FVector& InputVector, FCh
 		AsyncSimState->bCrouchMaintainsBaseLocation = bCrouchMaintainsBaseLocation;
 		AsyncSimState->bJustTeleported = bJustTeleported;
 		CharacterOwner->GetCapsuleComponent()->GetScaledCapsuleSize(AsyncSimState->ScaledCapsuleRadius, AsyncSimState->ScaledCapsuleHalfHeight);
-		AsyncSimState->bIsCrouched = CharacterOwner->bIsCrouched;
+		AsyncSimState->bIsCrouched = CharacterOwner->IsCrouched();
 		AsyncSimState->bWantsToCrouch = bWantsToCrouch;
 		AsyncSimState->bMovementInProgress = bMovementInProgress;
 		AsyncSimState->CurrentFloor = CurrentFloor;
diff --git a/Engine/Source/Runtime/Engine/Private/Components/PawnNoiseEmitterComponent.cpp b/Engine/Source/Runtime/Engine/Private/Components/PawnNoiseEmitterComponent.cpp
--- a/Engine/Source/Runtime/Engine/Private/Components/PawnNoiseEmitterComponent.cpp
+++ b/Engine/Source/Runtime/Engine/Private/Components/PawnNoiseEmitterComponent.cpp
@@ -46,7 +46,7 @@ void UPawnNoiseEmitterComponent::MakeNoise(AActor* NoiseMaker, float Loudness, c
 	}
 
 	// only emit sounds from pawns with controllers
-	if (!PawnOwner || !IsValidChecked(PawnOwner) || !PawnOwner->Controller)
+	if (!PawnOwner || !IsValidChecked(PawnOwner) || !PawnOwner->GetController())
 	{
 		return;
 	}
diff --git a/Engine/Source/Runtime/Engine/Private/Components/PrimitiveComponent.cpp b/Engine/Source/Runtime/Engine/Private/Components/PrimitiveComponent.cpp
--- a/Engine/Source/Runtime/Engine/Private/Components/PrimitiveComponent.cpp
+++ b/Engine/Source/Runtime/Engine/Private/Components/PrimitiveComponent.cpp
@@ -2994,9 +2994,9 @@ bool UPrimitiveComponent::MoveComponentImpl( const FVector& Delta, const FQuat&
 				// this is sole debug purpose to find how capsule trace information was when hit 
 				// to resolve stuck or improve our movement system - To turn this on, use DebugCapsuleSweepPawn
 				APawn const* const ActorPawn = (Actor ? Cast<APawn>(Actor) : NULL);
-				if (ActorPawn && ActorPawn->Controller && ActorPawn->Controller->IsLocalPlayerController())
+				if (ActorPawn && ActorPawn->GetController() && ActorPawn->GetController()->IsLocalPlayerController())
 				{
-					APlayerController const* const PC = CastChecked<APlayerController>(ActorPawn->Controller);
+					APlayerController const* const PC = CastChecked<APlayerController>(ActorPawn->GetController());
 					if (PC->CheatManager)
 					{
 						FVector CylExtent = ActorPawn->GetSimpleCollisionCylinderExtent()*FVector(1.001f,1.001f,1.0f);							
diff --git a/Engine/Source/Runtime/Engine/Private/Controller.cpp b/Engine/Source/Runtime/Engine/Private/Controller.cpp
--- a/Engine/Source/Runtime/Engine/Private/Controller.cpp
+++ b/Engine/Source/Runtime/Engine/Private/Controller.cpp
@@ -364,10 +364,10 @@ void AController::OnPossess(APawn* InPawn)
 		return;
 	}
 
-	if (InPawn->Controller != nullptr)
+	if (InPawn->GetController() != nullptr)
 	{
-		UE_CLOG(InPawn->Controller == this, LogController, Warning, TEXT("Asking %s to possess pawn %s more than once; pawn will be restarted! Should call Unpossess first."), *GetNameSafe(this), *GetNameSafe(InPawn));
-		InPawn->Controller->UnPossess();
+		UE_CLOG(InPawn->GetController() == this, LogController, Warning, TEXT("Asking %s to possess pawn %s more than once; pawn will be restarted! Should call Unpossess first."), *GetNameSafe(this), *GetNameSafe(InPawn));
+		InPawn->GetController()->UnPossess();
 	}
 
 	InPawn->PossessedBy(this);
@@ -548,10 +548,10 @@ void AController::OnRep_Pawn()
 {
 	APawn* StrongOldPawn = OldPawn.Get();
 	// Detect when pawn changes, so we can NULL out the controller on the old pawn
-	if ((StrongOldPawn != nullptr) && (Pawn != StrongOldPawn) && (StrongOldPawn->Controller == this))
+	if ((StrongOldPawn != nullptr) && (Pawn != StrongOldPawn) && (StrongOldPawn->GetController() == this))
 	{
 		// Set the old controller to NULL, since we are no longer the owner, and can't rely on it replicating to us anymore
-		StrongOldPawn->Controller = nullptr;
+		StrongOldPawn->SetController(nullptr);
 		if (UE::Gameplay::CVars::bAlwaysNotifyClientOnControllerChange)
 		{
 			// This will notify other systems like the game instance
diff --git a/Engine/Source/Runtime/Engine/Private/DefaultPawn.cpp b/Engine/Source/Runtime/Engine/Private/DefaultPawn.cpp
--- a/Engine/Source/Runtime/Engine/Private/DefaultPawn.cpp
+++ b/Engine/Source/Runtime/Engine/Private/DefaultPawn.cpp
@@ -157,9 +157,9 @@ void ADefaultPawn::MoveRight(float Val)
 {
 	if (Val != 0.f)
 	{
-		if (Controller)
+		if (GetController())
 		{
-			FRotator const ControlSpaceRot = Controller->GetControlRotation();
+			FRotator const ControlSpaceRot = GetController()->GetControlRotation();
 
 			// transform to world space and add it
 			AddMovementInput( FRotationMatrix(ControlSpaceRot).GetScaledAxis( EAxis::Y ), Val );
@@ -171,9 +171,9 @@ void ADefaultPawn::MoveForward(float Val)
 {
 	if (Val != 0.f)
 	{
-		if (Controller)
+		if (GetController())
 		{
-			FRotator const ControlSpaceRot = Controller->GetControlRotation();
+			FRotator const ControlSpaceRot = GetController()->GetControlRotation();
 
 			// transform to world space and add it
 			AddMovementInput( FRotationMatrix(ControlSpaceRot).GetScaledAxis( EAxis::X ), Val );
diff --git a/Engine/Source/Runtime/Engine/Private/GameModeBase.cpp b/Engine/Source/Runtime/Engine/Private/GameModeBase.cpp
--- a/Engine/Source/Runtime/Engine/Private/GameModeBase.cpp
+++ b/Engine/Source/Runtime/Engine/Private/GameModeBase.cpp
@@ -1437,7 +1437,7 @@ bool AGameModeBase::SpawnPlayerFromSimulate(const FVector& NewLocation, const FR
 			{
 				if (Pawn->AutoPossessPlayer == EAutoReceiveInput::Player0)
 				{
-					if (Pawn->Controller == nullptr)
+					if (Pawn->GetController() == nullptr)
 					{
 						PC->Possess(Pawn);
 						bNeedsRestart = false;
diff --git a/Engine/Source/Runtime/Engine/Private/Pawn.cpp b/Engine/Source/Runtime/Engine/Private/Pawn.cpp
--- a/Engine/Source/Runtime/Engine/Private/Pawn.cpp
+++ b/Engine/Source/Runtime/Engine/Private/Pawn.cpp
@@ -142,7 +142,7 @@ void APawn::PostInitializeComponents()
 
 		// Automatically add Controller to AI Pawns if we are allowed to.
 		if (AutoPossessPlayer == EAutoReceiveInput::Disabled
-			&& AutoPossessAI != EAutoPossessAI::Disabled && Controller == nullptr && GetNetMode() != NM_Client
+			&& AutoPossessAI != EAutoPossessAI::Disabled && GetController() == nullptr && GetNetMode() != NM_Client
 #if WITH_EDITOR
 			&& (GIsEditor == false || World->IsGameWorld())
 #endif // WITH_EDITOR
@@ -246,12 +246,12 @@ FVector APawn::GetVelocity() const
 
 bool APawn::IsLocallyControlled() const
 {
-	return ( Controller && Controller->IsLocalController() );
+	return ( GetController() && GetController()->IsLocalController() );
 }
 
 FPlatformUserId APawn::GetPlatformUserId() const
 {
-	if (const APlayerController* PC = Cast<APlayerController>(Controller))
+	if (const APlayerController* PC = Cast<APlayerController>(GetController()))
 	{
 		return PC->GetPlatformUserId();
 	}
@@ -260,18 +260,18 @@ FPlatformUserId APawn::GetPlatformUserId() const
 
 bool APawn::IsPlayerControlled() const
 {
-	return PlayerState && !PlayerState->IsABot();
+	return GetPlayerState() && !GetPlayerState()->IsABot();
 }
 
 bool APawn::IsBotControlled() const
 {
-	return PlayerState && PlayerState->IsABot();
+	return GetPlayerState() && GetPlayerState()->IsABot();
 }
 
 bool APawn::ReachedDesiredRotation()
 {
 	// Only base success on Yaw 
-	FRotator DesiredRotation = Controller ? Controller->GetDesiredRotation() : GetActorRotation();
+	FRotator DesiredRotation = GetController() ? GetController()->GetDesiredRotation() : GetActorRotation();
 	float YawDiff = FMath::Abs(FRotator::ClampAxis(DesiredRotation.Yaw) - FRotator::ClampAxis(GetActorRotation().Yaw));
 	return ( (YawDiff < AllowedYawError) || (YawDiff > 360.f - AllowedYawError) );
 }
@@ -304,9 +304,9 @@ void APawn::SetRemoteViewPitch(float NewRemoteViewPitch)
 UPawnNoiseEmitterComponent* APawn::GetPawnNoiseEmitterComponent() const
 {
 	 UPawnNoiseEmitterComponent* NoiseEmitterComponent = FindComponentByClass<UPawnNoiseEmitterComponent>();
-	 if (!NoiseEmitterComponent && Controller)
+	 if (!NoiseEmitterComponent && GetController())
 	 {
-		 NoiseEmitterComponent = Controller->FindComponentByClass<UPawnNoiseEmitterComponent>();
+		 NoiseEmitterComponent = GetController()->FindComponentByClass<UPawnNoiseEmitterComponent>();
 	 }
 
 	 return NoiseEmitterComponent;
@@ -334,9 +334,9 @@ FVector APawn::GetPawnViewLocation() const
 
 FRotator APawn::GetViewRotation() const
 {
-	if (Controller != nullptr)
+	if (GetController() != nullptr)
 	{
-		return Controller->GetControlRotation();
+		return GetController()->GetControlRotation();
 	}
 	else if (GetLocalRole() < ROLE_Authority)
 	{
@@ -358,7 +358,7 @@ FRotator APawn::GetViewRotation() const
 
 void APawn::SpawnDefaultController()
 {
-	if ( Controller != nullptr || GetNetMode() == NM_Client)
+	if (GetController() != nullptr || GetNetMode() == NM_Client)
 	{
 		return;
 	}
@@ -472,7 +472,7 @@ void APawn::PawnClientRestart()
 {
 	Restart();
 
-	APlayerController* PC = Cast<APlayerController>(Controller);
+	APlayerController* PC = Cast<APlayerController>(GetController());
 	if (PC && PC->IsLocalController())
 	{
 		// Handle camera possession
@@ -560,7 +560,7 @@ float APawn::TakeDamage(float Damage, FDamageEvent const& DamageEvent, AControll
 	// respond to the damage
 	if (ActualDamage != 0.f)
 	{
-		if ( EventInstigator && EventInstigator != Controller )
+		if ( EventInstigator && EventInstigator != GetController())
 		{
 			LastHitBy = EventInstigator;
 		}
@@ -571,27 +571,28 @@ float APawn::TakeDamage(float Damage, FDamageEvent const& DamageEvent, AControll
 
 bool APawn::IsControlled() const
 {
-	APlayerController* const PC = Cast<APlayerController>(Controller);
+	APlayerController* const PC = Cast<APlayerController>(GetController());
 	return(PC != nullptr);
 }
 
 bool APawn::IsPawnControlled() const
 {
-	return (Controller != nullptr);
+	return (GetController() != nullptr);
 }
 
 FRotator APawn::GetControlRotation() const
 {
-	return Controller ? Controller->GetControlRotation() : FRotator::ZeroRotator;
+	return GetController() ? GetController()->GetControlRotation() : FRotator::ZeroRotator;
 }
 
 void APawn::OnRep_Controller()
 {
+	AController* const ThisController = GetController();
 	bool bNotifyControllerChange = UE::Gameplay::CVars::bAlwaysNotifyClientOnControllerChange ?
-		(Controller != PreviousController) :	// By default, notify whenever the PreviousController is out of date for any reason
-		(Controller == nullptr);				// In backward compatibility, only notify when changing from null or the edge case below
+		(ThisController != PreviousController) :	// By default, notify whenever the PreviousController is out of date for any reason
+		(ThisController == nullptr);				// In backward compatibility, only notify when changing from null or the edge case below
 	
-	if ( (Controller != nullptr) && (Controller->GetPawn() == nullptr) )
+	if ( (ThisController != nullptr) && (ThisController->GetPawn() == nullptr) )
 	{
 		// This ensures that AController::OnRep_Pawn is called. Since we cant ensure replication order of APawn::Controller and AController::Pawn,
 		// if APawn::Controller is repped first, it will set AController::Pawn locally. When AController::Pawn is repped, the rep value will not
@@ -599,10 +600,10 @@ void APawn::OnRep_Controller()
 		//
 		// It would be better to never ever set replicated properties locally, but this is pretty core in the gameplay framework and I think there are
 		// lots of assumptions made in the code base that the Pawn and Controller will always be linked both ways.
-		Controller->SetPawnFromRep(this); 
+		ThisController->SetPawnFromRep(this); 
 
-		APlayerController* const PC = Cast<APlayerController>(Controller);
-		if ( (PC != nullptr) && PC->bAutoManageActiveCameraTarget && (PC->PlayerCameraManager->ViewTarget.Target == Controller) )
+		APlayerController* const PC = Cast<APlayerController>(ThisController);
+		if ( (PC != nullptr) && PC->bAutoManageActiveCameraTarget && (PC->PlayerCameraManager->ViewTarget.Target == GetController()) )
 		{
 			PC->AutoManageActiveCameraTarget(this);
 		}
@@ -618,7 +619,7 @@ void APawn::OnRep_Controller()
 
 void APawn::OnRep_PlayerState()
 {
-	SetPlayerState(PlayerState);
+	SetPlayerState(GetPlayerState());
 }
 
 void APawn::SetPlayerState(APlayerState* NewPlayerState)
@@ -644,9 +645,9 @@ void APawn::PossessedBy(AController* NewController)
 {
 	SetOwner(NewController);
 	
-	AController* const OldController = Controller;
+	AController* const OldController = GetController();
 
-	Controller = NewController;
+	SetController(NewController);
 	ForceNetUpdate();
 
 #if UE_WITH_IRIS
@@ -654,12 +655,12 @@ void APawn::PossessedBy(AController* NewController)
 	UpdateOwningNetConnection();
 #endif
 
-	if (Controller->PlayerState != nullptr)
+	if (GetController()->PlayerState != nullptr)
 	{
-		SetPlayerState(Controller->PlayerState);
+		SetPlayerState(GetController()->PlayerState);
 	}
 
-	if (APlayerController* PlayerController = Cast<APlayerController>(Controller))
+	if (APlayerController* PlayerController = Cast<APlayerController>(GetController()))
 	{
 		if (GetNetMode() != NM_Standalone)
 		{
@@ -678,7 +679,7 @@ void APawn::PossessedBy(AController* NewController)
 	// dispatch Blueprint event if necessary
 	if (OldController != NewController)
 	{
-		ReceivePossessed(Controller);
+		ReceivePossessed(GetController());
 	
 		NotifyControllerChanged();
 	}
@@ -686,13 +687,13 @@ void APawn::PossessedBy(AController* NewController)
 
 void APawn::UnPossessed()
 {
-	AController* const OldController = Controller;
+	AController* const OldController = GetController();
 
 	ForceNetUpdate();
 
 	SetPlayerState(nullptr);
 	SetOwner(nullptr);
-	Controller = nullptr;
+	SetController(nullptr);
 
 #if UE_WITH_IRIS
 	// The owning connection depends on the Controller having the new value.
@@ -715,23 +716,23 @@ void APawn::UnPossessed()
 
 void APawn::NotifyControllerChanged()
 {
-	ReceiveControllerChanged(PreviousController, Controller);
-	ReceiveControllerChangedDelegate.Broadcast(this, PreviousController, Controller);
+	ReceiveControllerChanged(PreviousController, GetController());
+	ReceiveControllerChangedDelegate.Broadcast(this, PreviousController, GetController());
 	if (UGameInstance* GameInstance = GetGameInstance())
 	{
-		GameInstance->GetOnPawnControllerChanged().Broadcast(this, Controller);
+		GameInstance->GetOnPawnControllerChanged().Broadcast(this, GetController());
 	}
 
 	// Update the cached controller
-	PreviousController = Controller;
+	PreviousController = GetController();
 }
 
 class UNetConnection* APawn::GetNetConnection() const
 {
 	// if have a controller, it has the net connection
-	if ( Controller )
+	if (GetController())
 	{
-		return Controller->GetNetConnection();
+		return GetController()->GetNetConnection();
 	}
 	return Super::GetNetConnection();
 }
@@ -745,9 +746,9 @@ class UPlayer* APawn::GetNetOwningPlayer()
 {
 	if (GetLocalRole() == ROLE_Authority)
 	{
-		if (Controller)
+		if (GetController())
 		{
-			APlayerController* PC = Cast<APlayerController>(Controller);
+			APlayerController* PC = Cast<APlayerController>(GetController());
 			return PC ? PC->Player : nullptr;
 		}
 	}
@@ -757,9 +758,9 @@ class UPlayer* APawn::GetNetOwningPlayer()
 
 class UPlayer* APawn::GetNetOwningPlayerAnyRole()
 {
-	if (Controller)
+	if (GetController())
 	{
-		APlayerController* PC = Cast<APlayerController>(Controller);
+		APlayerController* PC = Cast<APlayerController>(GetController());
 		return PC ? PC->Player : nullptr;
 	}
 
@@ -785,7 +786,7 @@ void APawn::DestroyPlayerInputComponent()
 
 bool APawn::IsMoveInputIgnored() const
 {
-	return Controller != nullptr && Controller->IsMoveInputIgnored();
+	return GetController() != nullptr && GetController()->IsMoveInputIgnored();
 }
 
 TSubclassOf<UInputComponent> APawn::GetOverrideInputComponentClass() const
@@ -850,27 +851,27 @@ FVector APawn::Internal_ConsumeMovementInputVector()
 
 void APawn::AddControllerPitchInput(float Val)
 {
-	if (Val != 0.f && Controller && Controller->IsLocalPlayerController())
+	if (Val != 0.f && GetController() && GetController()->IsLocalPlayerController())
 	{
-		APlayerController* const PC = CastChecked<APlayerController>(Controller);
+		APlayerController* const PC = CastChecked<APlayerController>(GetController());
 		PC->AddPitchInput(Val);
 	}
 }
 
 void APawn::AddControllerYawInput(float Val)
 {
-	if (Val != 0.f && Controller && Controller->IsLocalPlayerController())
+	if (Val != 0.f && GetController() && GetController()->IsLocalPlayerController())
 	{
-		APlayerController* const PC = CastChecked<APlayerController>(Controller);
+		APlayerController* const PC = CastChecked<APlayerController>(GetController());
 		PC->AddYawInput(Val);
 	}
 }
 
 void APawn::AddControllerRollInput(float Val)
 {
-	if (Val != 0.f && Controller && Controller->IsLocalPlayerController())
+	if (Val != 0.f && GetController() && GetController()->IsLocalPlayerController())
 	{
-		APlayerController* const PC = CastChecked<APlayerController>(Controller);
+		APlayerController* const PC = CastChecked<APlayerController>(GetController());
 		PC->AddRollInput(Val);
 	}
 }
@@ -914,7 +915,7 @@ void APawn::RecalculateBaseEyeHeight()
 
 void APawn::Reset()
 {
-	if ( (Controller == nullptr) || (Controller->PlayerState != nullptr) )
+	if ( (GetController() == nullptr) || (GetController()->PlayerState != nullptr) )
 	{
 		DetachFromControllerPendingDestroy();
 		Destroy();
@@ -927,7 +928,7 @@ void APawn::Reset()
 
 FString APawn::GetHumanReadableName() const
 {
-	return PlayerState ? PlayerState->GetPlayerName() : Super::GetHumanReadableName();
+	return GetPlayerState() ? GetPlayerState()->GetPlayerName() : Super::GetHumanReadableName();
 }
 
 
@@ -945,13 +946,13 @@ void APawn::DisplayDebug(class UCanvas* Canvas, const FDebugDisplayInfo& DebugDi
 // 	}
 
 	FDisplayDebugManager& DisplayDebugManager = Canvas->DisplayDebugManager;
-	if ( PlayerState == nullptr )
+	if (GetPlayerState() == nullptr )
 	{
 		DisplayDebugManager.DrawString(TEXT("NO PlayerState"));
 	}
 	else
 	{
-		PlayerState->DisplayDebug(Canvas, DebugDisplay, YL, YPos);
+		GetPlayerState()->DisplayDebug(Canvas, DebugDisplay, YL, YPos);
 	}
 
 	Super::DisplayDebug(Canvas, DebugDisplay, YL, YPos);
@@ -964,14 +965,14 @@ void APawn::DisplayDebug(class UCanvas* Canvas, const FDebugDisplayInfo& DebugDi
 	}
 
 	// Controller
-	if ( Controller == nullptr )
+	if (GetController() == nullptr )
 	{
 		DisplayDebugManager.SetDrawColor(FColor(255, 0, 0));
 		DisplayDebugManager.DrawString(TEXT("NO Controller"));
 	}
 	else
 	{
-		Controller->DisplayDebug(Canvas, DebugDisplay, YL, YPos);
+		GetController()->DisplayDebug(Canvas, DebugDisplay, YL, YPos);
 	}
 }
 
@@ -989,9 +990,9 @@ FRotator APawn::GetBaseAimRotation() const
 	// that is by default Controller.Rotation for AI, and camera (crosshair) rotation for human players.
 	FVector POVLoc;
 	FRotator POVRot;
-	if( Controller != nullptr && !InFreeCam() )
+	if(GetController() != nullptr && !InFreeCam() )
 	{
-		Controller->GetPlayerViewPoint(POVLoc, POVRot);
+		GetController()->GetPlayerViewPoint(POVLoc, POVRot);
 		return POVRot;
 	}
 
@@ -1014,12 +1015,12 @@ FRotator APawn::GetBaseAimRotation() const
 
 			if (DemoNetDriver && DemoNetDriver->IsPlaying() && (DemoNetDriver->GetPlaybackEngineNetworkProtocolVersion() < FEngineNetworkCustomVersion::PawnRemoteViewPitch))
 			{
-				POVRot.Pitch = RemoteViewPitch;
+				POVRot.Pitch = GetRemoteViewPitch();
 				POVRot.Pitch = POVRot.Pitch * 360.0f / 255.0f;
 			}
 			else
 			{
-				POVRot.Pitch = FRotator::DecompressAxisFromByte(RemoteViewPitch);
+				POVRot.Pitch = FRotator::DecompressAxisFromByte(GetRemoteViewPitch());
 			}
 		}
 	}
@@ -1029,7 +1030,7 @@ FRotator APawn::GetBaseAimRotation() const
 
 bool APawn::InFreeCam() const
 {
-	const APlayerController* PC = Cast<const APlayerController>(Controller);
+	const APlayerController* PC = Cast<const APlayerController>(GetController());
 	static FName NAME_FreeCam =  FName(TEXT("FreeCam"));
 	static FName NAME_FreeCamDefault = FName(TEXT("FreeCam_Default"));
 	return (PC != nullptr && PC->PlayerCameraManager != nullptr && (PC->PlayerCameraManager->CameraStyle == NAME_FreeCam || PC->PlayerCameraManager->CameraStyle == NAME_FreeCamDefault) );
@@ -1041,7 +1042,7 @@ void APawn::OutsideWorldBounds()
 	{
 		bProcessingOutsideWorldBounds = true;
 		// AI pawns on the server just destroy
-		if (GetLocalRole() == ROLE_Authority && Cast<APlayerController>(Controller) == nullptr)
+		if (GetLocalRole() == ROLE_Authority && Cast<APlayerController>(GetController()) == nullptr)
 		{
 			Destroy();
 		}
@@ -1091,20 +1092,21 @@ void APawn::FaceRotation(FRotator NewControlRotation, float DeltaTime)
 
 void APawn::DetachFromControllerPendingDestroy()
 {
-	if ( Controller != nullptr && Controller->GetPawn() == this )
+	AController* const ThisController = GetController();
+	if (ThisController != nullptr && ThisController->GetPawn() == this )
 	{
-		Controller->PawnPendingDestroy(this);
-		if (Controller != nullptr)
+		ThisController->PawnPendingDestroy(this);
+		if (ThisController != nullptr)
 		{
-			Controller->UnPossess();
-			Controller = nullptr;
+			ThisController->UnPossess();
+			SetController(nullptr);
 		}
 	}
 }
 
 AController* APawn::GetDamageInstigator(AController* InstigatedBy, const UDamageType& DamageType) const
 {
-	if ( (InstigatedBy != nullptr) && (InstigatedBy != Controller) )
+	if ( (InstigatedBy != nullptr) && (InstigatedBy != GetController()) )
 	{
 		return InstigatedBy;
 	}
@@ -1122,9 +1124,9 @@ void APawn::EditorApplyRotation(const FRotator& DeltaRotation, bool bAltDown, bo
 	Super::EditorApplyRotation(DeltaRotation, bAltDown, bShiftDown, bCtrlDown);
 
 	// Forward new rotation on to the pawn's controller.
-	if( Controller )
+	if(GetController())
 	{
-		Controller->SetControlRotation( GetActorRotation() );
+		GetController()->SetControlRotation( GetActorRotation() );
 	}
 }
 
@@ -1143,7 +1145,7 @@ void APawn::PostEditChangeProperty(struct FPropertyChangedEvent& PropertyChanged
 
 void APawn::EnableInput(class APlayerController* PlayerController)
 {
-	if (PlayerController == Controller || PlayerController == nullptr)
+	if (PlayerController == GetController() || PlayerController == nullptr)
 	{
 		bInputEnabled = true;
 	}
@@ -1155,7 +1157,7 @@ void APawn::EnableInput(class APlayerController* PlayerController)
 
 void APawn::DisableInput(class APlayerController* PlayerController)
 {
-	if (PlayerController == Controller || PlayerController == nullptr)
+	if (PlayerController == GetController() || PlayerController == nullptr)
 	{
 		bInputEnabled = false;
 	}
@@ -1251,7 +1253,7 @@ bool APawn::IsBasedOnActor(const AActor* Other) const
 bool APawn::IsNetRelevantFor(const AActor* RealViewer, const AActor* ViewTarget, const FVector& SrcLocation) const
 {
 	CA_SUPPRESS(6011);
-	if (bAlwaysRelevant || RealViewer == Controller || IsOwnedBy(ViewTarget) || IsOwnedBy(RealViewer) || this == ViewTarget || ViewTarget == GetInstigator()
+	if (bAlwaysRelevant || RealViewer == GetController() || IsOwnedBy(ViewTarget) || IsOwnedBy(RealViewer) || this == ViewTarget || ViewTarget == GetInstigator()
 		|| IsBasedOnActor(ViewTarget) || (ViewTarget && ViewTarget->IsBasedOnActor(this)))
 	{
 		return true;
@@ -1284,6 +1286,11 @@ void APawn::GetLifetimeReplicatedProps( TArray< FLifetimeProperty > & OutLifetim
 	DOREPLIFETIME_CONDITION( APawn, RemoteViewPitch, 	COND_SkipOwner );
 }
 
+void APawn::SetController(AController* InController)
+{
+	Controller = InController;
+}
+
 void APawn::PreReplication( IRepChangedPropertyTracker & ChangedPropertyTracker )
 {
 	Super::PreReplication( ChangedPropertyTracker );
diff --git a/Engine/Source/Runtime/Engine/Private/PlayerController.cpp b/Engine/Source/Runtime/Engine/Private/PlayerController.cpp
--- a/Engine/Source/Runtime/Engine/Private/PlayerController.cpp
+++ b/Engine/Source/Runtime/Engine/Private/PlayerController.cpp
@@ -790,10 +790,10 @@ void APlayerController::ClientRetryClientRestart_Implementation(APawn* NewPawn)
 	UE_LOG(LogPlayerController, Verbose, TEXT("ClientRetryClientRestart_Implementation %s, AcknowledgedPawn: %s"), *GetNameSafe(NewPawn), *GetNameSafe(AcknowledgedPawn));
 
 	// Avoid calling ClientRestart if we have already accepted this pawn
-	if( (GetPawn() != NewPawn) || (NewPawn->Controller != this) || (NewPawn != AcknowledgedPawn) )
+	if( (GetPawn() != NewPawn) || (NewPawn->GetController() != this) || (NewPawn != AcknowledgedPawn) )
 	{
 		SetPawn(NewPawn);
-		NewPawn->Controller = this;
+		NewPawn->SetController(this);
 		NewPawn->OnRep_Controller();
 		ClientRestart(GetPawn());
 	}
@@ -826,8 +826,8 @@ void APlayerController::ClientRestart_Implementation(APawn* NewPawn)
 	// but we may receive the function call before Pawn is replicated over, so it will resolve to NULL.
 	AcknowledgePossession(GetPawn());
 
-	AController* OldController = GetPawn()->Controller;
-	GetPawn()->Controller = this;
+	AController* OldController = GetPawn()->GetController();
+	GetPawn()->SetController(this);
 	if (OldController != this)
 	{
 		// In case this is received before APawn::OnRep_Controller is called
@@ -868,9 +868,9 @@ void APlayerController::OnPossess(APawn* PawnToPossess)
 			UnPossess();
 		}
 
-		if (PawnToPossess->Controller != NULL)
+		if (PawnToPossess->GetController() != NULL)
 		{
-			PawnToPossess->Controller->UnPossess();
+			PawnToPossess->GetController()->UnPossess();
 		}
 
 		PawnToPossess->PossessedBy(this);
@@ -5625,9 +5625,9 @@ void APlayerController::EndSpectatingState()
 
 void APlayerController::BeginInactiveState()
 {
-	if ( (GetPawn() != nullptr) && (GetPawn()->Controller == this) )
+	if ( (GetPawn() != nullptr) && (GetPawn()->GetController() == this) )
 	{
-		GetPawn()->Controller = nullptr;
+		GetPawn()->SetController(nullptr);
 		if (UE::Gameplay::CVars::bAlwaysNotifyClientOnControllerChange)
 		{
 			if (HasAuthority())
diff --git a/Engine/Source/Runtime/Engine/Private/SpectatorPawn.cpp b/Engine/Source/Runtime/Engine/Private/SpectatorPawn.cpp
--- a/Engine/Source/Runtime/Engine/Private/SpectatorPawn.cpp
+++ b/Engine/Source/Runtime/Engine/Private/SpectatorPawn.cpp
@@ -38,13 +38,13 @@ void ASpectatorPawn::PossessedBy(class AController* NewController)
 	else
 	{
 		// We don't want the automatic changing of net role in Pawn code since we don't replicate, so don't call Super.
-		AController* const OldController = Controller;
-		Controller = NewController;
+		AController* const OldController = GetController();
+		SetController(NewController);
 
 		// dispatch Blueprint event if necessary
 		if (OldController != NewController)
 		{
-			ReceivePossessed(Controller);
+			ReceivePossessed(GetController());
 		}
 	}
 }
