[ISM]
- Fix stall/crash because of bounds pulling on static mesh property while it is still being compiled

#rnx

diff --git a/Engine/Source/Runtime/Engine/Private/InstancedStaticMesh.cpp b/Engine/Source/Runtime/Engine/Private/InstancedStaticMesh.cpp
--- a/Engine/Source/Runtime/Engine/Private/InstancedStaticMesh.cpp
+++ b/Engine/Source/Runtime/Engine/Private/InstancedStaticMesh.cpp
@@ -2825,6 +2825,17 @@ FBoxSphereBounds UInstancedStaticMeshComponent::CalcBoundsImpl(const FTransform&
 
 void UInstancedStaticMeshComponent::UpdateBounds()
 {
+#if WITH_EDITOR
+	// Skip updating bounds to avoid stall on static mesh if we're still compiling and because
+	// if this is called during registration, we can end up doing a flush that will then
+	// try to recreate the render state as part of the postcompilation before we are even ready to do so
+	// in the registration flow, causing asserts because InstancingRandomSeed has not yet been initialized.
+	if (IsCompiling())
+	{
+		// Bounds will be updated when we finish compiling.
+		return;
+	}
+#endif
 	FGuardValue_Bitfield(bIsUpdatingBounds, true);
 	
 	if (GetNumRenderInstances() == 0)
