ChildActorComponent:
- Change child actors from NonPIEDuplicateTransient to DuplicateTransient to avoid duplicating them when duplicating the world in PIE.
- We don't need to duplicate child actors in PIE because they will be recreated when the component will be registered to the world.
- Was causing issues for partitioned worlds when child actors were dirty and unsaved, resulting in child actors being recreated leaving child actors outered to the unsaved actor container.
- Change the duplicated unsaved actor container outer to be the level package instead of the level object, to avoid actors in it returning the outer level when calling AActor::GetLevel().
- AActor::GetLevel() should be changed to return Cast<ULevel>(GetOuter()) instead of GetTypedOuter<ULevel>(), will make this change in anotrher changelist to minimize risk.
- Dont't include child actors in the unsaved actor container since they will be recreated during child actor component world registration.

This fixes a couple of different crashes related to unsaved child actors and unsaved child actor Blueprint classes.

diff --git a/Engine/Source/Runtime/Engine/Classes/Components/ChildActorComponent.h b/Engine/Source/Runtime/Engine/Classes/Components/ChildActorComponent.h
--- a/Engine/Source/Runtime/Engine/Classes/Components/ChildActorComponent.h
+++ b/Engine/Source/Runtime/Engine/Classes/Components/ChildActorComponent.h
@@ -119,7 +119,7 @@ private:
 	TSubclassOf<AActor>	ChildActorClass;
 
 	/** The actor that we spawned and own */
-	UPROPERTY(Replicated, BlueprintReadOnly, ReplicatedUsing=OnRep_ChildActor, Category=ChildActorComponent, TextExportTransient, NonPIEDuplicateTransient, meta=(AllowPrivateAccess="true"))
+	UPROPERTY(Replicated, BlueprintReadOnly, ReplicatedUsing=OnRep_ChildActor, Category=ChildActorComponent, TextExportTransient, DuplicateTransient, meta=(AllowPrivateAccess="true"))
 	TObjectPtr<AActor>	ChildActor;
 
 	/** Property to point to the template child actor for details panel purposes */
diff --git a/Engine/Source/Runtime/Engine/Private/Components/ChildActorComponent.cpp b/Engine/Source/Runtime/Engine/Private/Components/ChildActorComponent.cpp
--- a/Engine/Source/Runtime/Engine/Private/Components/ChildActorComponent.cpp
+++ b/Engine/Source/Runtime/Engine/Private/Components/ChildActorComponent.cpp
@@ -620,7 +620,7 @@ void UChildActorComponent::PostLoad()
 		else
 		{
 			FActorParentComponentSetter::Set(ChildActor, this);
-			ChildActor->SetFlags(RF_TextExportTransient | RF_NonPIEDuplicateTransient);
+			ChildActor->SetFlags(RF_TextExportTransient | RF_DuplicateTransient);
 		}
 	}
 
@@ -813,7 +813,7 @@ void UChildActorComponent::CreateChildActor(TFunction<void(AActor*)> CustomizerF
 				{
 					Params.Template = ChildActorTemplate;
 				}
-				Params.ObjectFlags |= (RF_TextExportTransient | RF_NonPIEDuplicateTransient);
+				Params.ObjectFlags |= (RF_TextExportTransient | RF_DuplicateTransient);
 				if (!HasAllFlags(RF_Transactional))
 				{
 					Params.ObjectFlags &= ~RF_Transactional;
diff --git a/Engine/Source/Runtime/Engine/Private/WorldPartition/WorldPartitionLevelHelper.cpp b/Engine/Source/Runtime/Engine/Private/WorldPartition/WorldPartitionLevelHelper.cpp
--- a/Engine/Source/Runtime/Engine/Private/WorldPartition/WorldPartitionLevelHelper.cpp
+++ b/Engine/Source/Runtime/Engine/Private/WorldPartition/WorldPartitionLevelHelper.cpp
@@ -762,15 +762,6 @@ bool FWorldPartitionLevelHelper::LoadActorsInternal(FLoadActorsParams&& InParams
 					// Rename through UObject to avoid changing Actor's external packaging and folder properties
 					Actor->UObject::Rename(*AddActorContainerID(PackageObjectMapping->ContainerID, Actor->GetName()), DestLevel, REN_NonTransactional | REN_DoNotDirty | REN_DontCreateRedirectors);
 
-					// Handle child actors
-					Actor->ForEachComponent<UChildActorComponent>(true, [DestLevel = DestLevel, PackageObjectMapping](UChildActorComponent* ChildActorComponent)
-					{
-						if (AActor* ChildActor = ChildActorComponent->GetChildActor())
-						{
-							ChildActor->UObject::Rename(*AddActorContainerID(PackageObjectMapping->ContainerID, ChildActor->GetName()), DestLevel, REN_NonTransactional | REN_DoNotDirty | REN_DontCreateRedirectors);
-						}
-					});
-
 					// Apply Pre-ConstructionScript Properties
 					TArray<FActorPropertyOverride> ActorPropertyOverrides;
 					for (auto ActorOverrideMapping : PackageObjectMapping->PropertyOverrides)
@@ -848,16 +839,6 @@ bool FWorldPartitionLevelHelper::LoadActorsInternal(FLoadActorsParams&& InParams
 
 					DestLevel->Actors.Add(Actor);
 					checkf(Actor->GetLevel() == DestLevel, TEXT("Levels mismatch, got : %s, expected: %s\nActor: %s\nActorFullName: %s\nActorPackage: %s"), *DestLevel->GetFullName(), *Actor->GetLevel()->GetFullName(), *Actor->GetActorNameOrLabel(), *Actor->GetFullName(), *Actor->GetPackage()->GetFullName());
-
-					// Handle child actors
-					Actor->ForEachComponent<UChildActorComponent>(true, [DestLevel = DestLevel](UChildActorComponent* ChildActorComponent)
-					{
-						if (AActor* ChildActor = ChildActorComponent->GetChildActor())
-						{
-							DestLevel->Actors.Add(ChildActor);
-							check(ChildActor->GetLevel() == DestLevel);
-						}
-					});
 				}
 
 				UE_LOG(LogWorldPartition, Verbose, TEXT(" ==> Loaded %s (remaining: %d)"), *Actor->GetFullName(), LoadProgress->NumPendingLoadRequests);
diff --git a/Engine/Source/Runtime/Engine/Private/WorldPartition/WorldPartitionLevelStreamingDynamic.cpp b/Engine/Source/Runtime/Engine/Private/WorldPartition/WorldPartitionLevelStreamingDynamic.cpp
--- a/Engine/Source/Runtime/Engine/Private/WorldPartition/WorldPartitionLevelStreamingDynamic.cpp
+++ b/Engine/Source/Runtime/Engine/Private/WorldPartition/WorldPartitionLevelStreamingDynamic.cpp
@@ -423,7 +423,7 @@ bool UWorldPartitionLevelStreamingDynamic::IssueLoadRequests()
 	// Duplicate unsaved actors
 	if (UnsavedActorsContainer)
 	{
-		FObjectDuplicationParameters Parameters(UnsavedActorsContainer, RuntimeLevel);
+		FObjectDuplicationParameters Parameters(UnsavedActorsContainer, RuntimeLevel->GetPackage());
 		Parameters.DestClass = UnsavedActorsContainer->GetClass();
 		Parameters.FlagMask = RF_AllFlags & ~(RF_MarkAsRootSet | RF_MarkAsNative | RF_HasExternalPackage);
 		Parameters.InternalFlagMask = EInternalObjectFlags_AllFlags;
diff --git a/Engine/Source/Runtime/Engine/Private/WorldPartition/WorldPartitionRuntimeCell.cpp b/Engine/Source/Runtime/Engine/Private/WorldPartition/WorldPartitionRuntimeCell.cpp
--- a/Engine/Source/Runtime/Engine/Private/WorldPartition/WorldPartitionRuntimeCell.cpp
+++ b/Engine/Source/Runtime/Engine/Private/WorldPartition/WorldPartitionRuntimeCell.cpp
@@ -59,17 +59,7 @@ void UWorldPartitionRuntimeCell::PostDuplicate(bool bDuplicateForPIE)
 
 		for (auto& [ActorName, Actor] : UnsavedActorsContainer->Actors)
 		{
-			if (ensure(Actor))
-			{
-				const FName FolderPath = Actor->GetFolderPath();
-				
-				// Don't use AActor::Rename here since the actor is not part of the world, it's only a duplication template.	
-				Actor->UObject::Rename(nullptr, UnsavedActorsContainer);
-				
-				// Mimic what AActor::Rename does under the hood to convert ActorFolders when changing the outer
-				const bool bBroadcastChange = false;
-				FSetActorFolderPath SetFolderPath(Actor, FolderPath, bBroadcastChange);
-			}
+			Actor->Rename(nullptr, UnsavedActorsContainer);
 		}
 	}
 }
diff --git a/Engine/Source/Runtime/Engine/Private/WorldPartition/WorldPartitionRuntimeHash.cpp b/Engine/Source/Runtime/Engine/Private/WorldPartition/WorldPartitionRuntimeHash.cpp
--- a/Engine/Source/Runtime/Engine/Private/WorldPartition/WorldPartitionRuntimeHash.cpp
+++ b/Engine/Source/Runtime/Engine/Private/WorldPartition/WorldPartitionRuntimeHash.cpp
@@ -346,15 +346,6 @@ void UWorldPartitionRuntimeHash::PopulateRuntimeCell(UWorldPartitionRuntimeCell*
 			if (AActor* Actor = FindObject<AActor>(nullptr, *ActorDescView.GetActorSoftPath().ToString()))
 			{
 				RuntimeCell->UnsavedActorsContainer->Actors.Add(Actor->GetFName(), Actor);
-
-				// Handle child actors
-				Actor->ForEachComponent<UChildActorComponent>(true, [RuntimeCell](UChildActorComponent* ChildActorComponent)
-				{
-					if (AActor* ChildActor = ChildActorComponent->GetChildActor())
-					{
-						RuntimeCell->UnsavedActorsContainer->Actors.Add(ChildActor->GetFName(), ChildActor);
-					}
-				});
 			}
 		}
 	}
