Iris
* Added SubObjectDestroyedFromReplication to NetObjectFactory
* Renamed DestroyedFromReplication to DestroyReplicatedObject
* Removed unneeded Actor dependancy in the NetSubObjectFactory

#rnx

diff --git a/Engine/Plugins/Runtime/ReplicationSystemTestPlugin/Source/Private/Tests/ReplicationSystem/ReplicatedTestObjectFactory.cpp b/Engine/Plugins/Runtime/ReplicationSystemTestPlugin/Source/Private/Tests/ReplicationSystem/ReplicatedTestObjectFactory.cpp
--- a/Engine/Plugins/Runtime/ReplicationSystemTestPlugin/Source/Private/Tests/ReplicationSystem/ReplicatedTestObjectFactory.cpp
+++ b/Engine/Plugins/Runtime/ReplicationSystemTestPlugin/Source/Private/Tests/ReplicationSystem/ReplicatedTestObjectFactory.cpp
@@ -186,7 +186,7 @@ void UReplicatedTestObjectFactory::PostInit(const FPostInitContext& Context)
 	Instance->NetRefHandle = Context.Handle;
 }
 
-void UReplicatedTestObjectFactory::DestroyedFromReplication(const FDestroyedContext& Context)
+void UReplicatedTestObjectFactory::DestroyReplicatedObject(const FDestroyedContext& Context)
 {
 	if (!Context.DestroyedInstance)
 	{
diff --git a/Engine/Plugins/Runtime/ReplicationSystemTestPlugin/Source/Private/Tests/ReplicationSystem/ReplicatedTestObjectFactory.h b/Engine/Plugins/Runtime/ReplicationSystemTestPlugin/Source/Private/Tests/ReplicationSystem/ReplicatedTestObjectFactory.h
--- a/Engine/Plugins/Runtime/ReplicationSystemTestPlugin/Source/Private/Tests/ReplicationSystem/ReplicatedTestObjectFactory.h
+++ b/Engine/Plugins/Runtime/ReplicationSystemTestPlugin/Source/Private/Tests/ReplicationSystem/ReplicatedTestObjectFactory.h
@@ -40,7 +40,7 @@ public:
 	virtual FInstantiateResult InstantiateReplicatedObjectFromHeader(const FInstantiateContext& Context, const UE::Net::FNetObjectCreationHeader* InHeader) override;
 
 	virtual void PostInit(const FPostInitContext& Context) override;
-	virtual void DestroyedFromReplication(const FDestroyedContext& Context) override;
+	virtual void DestroyReplicatedObject(const FDestroyedContext& Context) override;
 
 	void SetCreatedObjectsOnNode(TArray<TStrongObjectPtr<UObject>>* InCreatedObjectsOnNode)
 	{
diff --git a/Engine/Source/Runtime/Engine/Private/Net/Iris/ReplicationSystem/NetActorFactory.cpp b/Engine/Source/Runtime/Engine/Private/Net/Iris/ReplicationSystem/NetActorFactory.cpp
--- a/Engine/Source/Runtime/Engine/Private/Net/Iris/ReplicationSystem/NetActorFactory.cpp
+++ b/Engine/Source/Runtime/Engine/Private/Net/Iris/ReplicationSystem/NetActorFactory.cpp
@@ -496,7 +496,7 @@ void UNetActorFactory::SubObjectCreatedFromReplication(UE::Net::FNetRefHandle Ro
 	}
 }
 
-void UNetActorFactory::DestroyedFromReplication(const FDestroyedContext& Context)
+void UNetActorFactory::DestroyReplicatedObject(const FDestroyedContext& Context)
 {
 	UEngineReplicationBridge* ActorBridge = CastChecked<UEngineReplicationBridge>(Bridge);
 	UNetDriver* NetDriver = ActorBridge->GetNetDriver();
@@ -509,7 +509,7 @@ void UNetActorFactory::DestroyedFromReplication(const FDestroyedContext& Context
 		}
 		else if (EnumHasAnyFlags(Context.DestroyFlags, EReplicationBridgeDestroyInstanceFlags::AllowDestroyInstanceFromRemote))
 		{
-			// Any subobjects have already been detached from ReplicationBridge
+			// Note that subobjects have already been detached by the ReplicationBridge
 			Actor->PreDestroyFromReplication();
 			Actor->Destroy(true);
 		}
diff --git a/Engine/Source/Runtime/Engine/Private/Net/Iris/ReplicationSystem/NetSubObjectFactory.cpp b/Engine/Source/Runtime/Engine/Private/Net/Iris/ReplicationSystem/NetSubObjectFactory.cpp
--- a/Engine/Source/Runtime/Engine/Private/Net/Iris/ReplicationSystem/NetSubObjectFactory.cpp
+++ b/Engine/Source/Runtime/Engine/Private/Net/Iris/ReplicationSystem/NetSubObjectFactory.cpp
@@ -137,7 +137,6 @@ UNetSubObjectFactory::FInstantiateResult UNetSubObjectFactory::InstantiateReplic
 	const FNetDynamicSubObjectCreationHeader* SubObjectHeader = static_cast<const FNetDynamicSubObjectCreationHeader*>(Header);
 		
 	UObject* RootObject = Bridge->GetReplicatedObject(Context.RootObjectOfSubObject);
-	AActor* RootActor = CastChecked<AActor>(RootObject);
 			
 	// Find the proper Outer
 	UObject* OuterObject = nullptr;
@@ -147,7 +146,7 @@ UNetSubObjectFactory::FInstantiateResult UNetSubObjectFactory::InstantiateReplic
 	}
 	else if (SubObjectHeader->bOuterIsRootObject)
 	{
-		OuterObject = RootActor;
+		OuterObject = RootObject;
 	}				
 	else
 	{
@@ -158,7 +157,7 @@ UNetSubObjectFactory::FInstantiateResult UNetSubObjectFactory::InstantiateReplic
 			UE_LOG(LogIris, Error, TEXT("BeginInstantiateFromRemote Failed to find Outer %s for dynamic subobject %s"), *Bridge->DescribeObjectReference(SubObjectHeader->OuterReference, Context.ResolveContext), *Context.Handle.ToString())
 
 			// Fallback to the rootobject instead
-			OuterObject = RootActor;
+			OuterObject = RootObject;
 		}
 	}
 
@@ -170,9 +169,8 @@ UNetSubObjectFactory::FInstantiateResult UNetSubObjectFactory::InstantiateReplic
 	UObject* SubObj = NewObject<UObject>(OuterObject, SubObjClass);
 
 	// Sanity check some things
-	checkf(SubObj != nullptr, TEXT("UNetSubObjectFactory::InstantiateNetObjectFromHeader: Subobject is NULL after instantiating. Class: %s, Outer %s, Actor %s"), *GetNameSafe(SubObjClass), *GetNameSafe(OuterObject), *GetNameSafe(RootActor));
-	checkf(!OuterObject || SubObj->IsIn(OuterObject), TEXT("UNetSubObjectFactory::InstantiateNetObjectFromHeader: Subobject is not in Outer. SubObject: %s, Outer %s, Actor %s"), *SubObj->GetName(), *GetNameSafe(OuterObject), *GetNameSafe(RootActor));
-	checkf(Cast<AActor>(SubObj) == nullptr, TEXT("UNetSubObjectFactory::InstantiateNetObjectFromHeader: Subobject is an Actor. SubObject: %s, Outer %s, Actor %s"), *SubObj->GetName(), *GetNameSafe(OuterObject), *GetNameSafe(RootActor));
+	checkf(SubObj != nullptr, TEXT("UNetSubObjectFactory::InstantiateNetObjectFromHeader: Subobject is NULL after instantiating. Class: %s, Outer %s, Actor %s"), *GetNameSafe(SubObjClass), *GetNameSafe(OuterObject), *GetNameSafe(RootObject));
+	checkf(!OuterObject || SubObj->IsIn(OuterObject), TEXT("UNetSubObjectFactory::InstantiateNetObjectFromHeader: Subobject is not in Outer. SubObject: %s, Outer %s, Actor %s"), *SubObj->GetName(), *GetNameSafe(OuterObject), *GetNameSafe(RootObject));
 
 	FInstantiateResult Result { .Instance = SubObj };
 
@@ -202,7 +200,7 @@ void UNetSubObjectFactory::SubObjectCreatedFromReplication(UE::Net::FNetRefHandl
 	ensureMsgf(false, TEXT("NetSubObjectFactory::SubObjectCreatedFromReplication should never be called since subobjects cannot have their own subobject list. RootObject: %s, SubObjectCreated: %s"), *Bridge->PrintObjectFromNetRefHandle(RootObject), *Bridge->PrintObjectFromNetRefHandle(SubObjectCreated));
 }
 
-void UNetSubObjectFactory::DestroyedFromReplication(const FDestroyedContext& Context)
+void UNetSubObjectFactory::DestroyReplicatedObject(const FDestroyedContext& Context)
 {
 	// If the SubObject is being torn off it is up to owning actor to clean it up properly
 	if (Context.DestroyReason == EReplicationBridgeDestroyInstanceReason::TearOff)
@@ -215,11 +213,12 @@ void UNetSubObjectFactory::DestroyedFromReplication(const FDestroyedContext& Con
 		return;
 	}
 
-	AActor* Owner = Cast<AActor>(Context.RootObject);
-
-	if (ensureMsgf(IsValid(Owner) && !Owner->IsUnreachable(), TEXT("UEngineReplicationBridge::DestroyInstanceFromRemote Destroyed subobject: %s has an invalid owner: %s"), *GetNameSafe(Context.DestroyedInstance), *GetPathNameSafe(Context.RootObject)))
+	if (AActor* ActorOwner = Cast<AActor>(Context.RootObject))
 	{
-		Owner->OnSubobjectDestroyFromReplication(Context.DestroyedInstance);
+		if (ensureMsgf(IsValid(ActorOwner) && !ActorOwner->IsUnreachable(), TEXT("UEngineReplicationBridge::DestroyInstanceFromRemote Destroyed subobject: %s has an invalid owner: %s"), *GetNameSafe(Context.DestroyedInstance), *GetPathNameSafe(Context.RootObject)))
+		{
+			ActorOwner->OnSubobjectDestroyFromReplication(Context.DestroyedInstance);
+		}
 	}
 
 	Context.DestroyedInstance->PreDestroyFromReplication();
diff --git a/Engine/Source/Runtime/Engine/Public/Net/Iris/ReplicationSystem/NetActorFactory.h b/Engine/Source/Runtime/Engine/Public/Net/Iris/ReplicationSystem/NetActorFactory.h
--- a/Engine/Source/Runtime/Engine/Public/Net/Iris/ReplicationSystem/NetActorFactory.h
+++ b/Engine/Source/Runtime/Engine/Public/Net/Iris/ReplicationSystem/NetActorFactory.h
@@ -125,7 +125,7 @@ public:
 
 	virtual void SubObjectCreatedFromReplication(UE::Net::FNetRefHandle RootObject, UE::Net::FNetRefHandle SubObjectCreated) override;
 
-	virtual void DestroyedFromReplication(const FDestroyedContext& Context) override;
+	virtual void DestroyReplicatedObject(const FDestroyedContext& Context) override;
 
 
 protected:
diff --git a/Engine/Source/Runtime/Engine/Public/Net/Iris/ReplicationSystem/NetSubObjectFactory.h b/Engine/Source/Runtime/Engine/Public/Net/Iris/ReplicationSystem/NetSubObjectFactory.h
--- a/Engine/Source/Runtime/Engine/Public/Net/Iris/ReplicationSystem/NetSubObjectFactory.h
+++ b/Engine/Source/Runtime/Engine/Public/Net/Iris/ReplicationSystem/NetSubObjectFactory.h
@@ -31,7 +31,7 @@ public:
 
 	virtual void SubObjectCreatedFromReplication(UE::Net::FNetRefHandle RootObject, UE::Net::FNetRefHandle SubObjectCreated) override;
 
-	virtual void DestroyedFromReplication(const FDestroyedContext& Context) override;
+	virtual void DestroyReplicatedObject(const FDestroyedContext& Context) override;
 
 protected:
 
diff --git a/Engine/Source/Runtime/Experimental/Iris/Core/Private/Iris/ReplicationSystem/ObjectReplicationBridge.cpp b/Engine/Source/Runtime/Experimental/Iris/Core/Private/Iris/ReplicationSystem/ObjectReplicationBridge.cpp
--- a/Engine/Source/Runtime/Experimental/Iris/Core/Private/Iris/ReplicationSystem/ObjectReplicationBridge.cpp
+++ b/Engine/Source/Runtime/Experimental/Iris/Core/Private/Iris/ReplicationSystem/ObjectReplicationBridge.cpp
@@ -763,26 +763,33 @@ void UObjectReplicationBridge::DetachInstanceFromRemote(FNetRefHandle Handle, ER
 	// Tell factory about the imminent destruction
 	if (DestroyReason != EReplicationBridgeDestroyInstanceReason::DoNotDestroy)
 	{
-		UNetObjectFactory::FDestroyedContext DestroyInstanceParams
+		UNetObjectFactory::FDestroyedContext DestroyContext
 		{
 			.DestroyedInstance = GetObjectFromReferenceHandle(Handle),
 			.DestroyReason = DestroyReason,
 			.DestroyFlags = DestroyFlags,
         };
 
-		// Tell the factory to handle the object if the instance exists
-        if (ensureMsgf(NetFactoryId != InvalidNetObjectFactoryId, TEXT("Received invalid netfactory for %s: %s"), *Handle.ToString(), *GetPathNameSafe(DestroyInstanceParams.DestroyedInstance)))
+		if (DestroyContext.DestroyedInstance)
 		{
-            if (DestroyInstanceParams.DestroyedInstance)
-            {
-                const FInternalNetRefIndex DestroyedObjectIndex = NetRefHandleManager->GetInternalIndex(Handle);
-                const FNetRefHandleManager::FReplicatedObjectData& DestroyedObjectData = NetRefHandleManager->GetReplicatedObjectData(DestroyedObjectIndex);
+			const FInternalNetRefIndex DestroyedObjectIndex = NetRefHandleManager->GetInternalIndex(Handle);
+			const FNetRefHandleManager::FReplicatedObjectData& DestroyedObjectData = NetRefHandleManager->GetReplicatedObjectData(DestroyedObjectIndex);
 
-                DestroyInstanceParams.RootObject = DestroyedObjectData.SubObjectRootIndex != FNetRefHandleManager::InvalidInternalIndex ? NetRefHandleManager->GetReplicatedObjectInstance(DestroyedObjectData.SubObjectRootIndex) : nullptr;
+			// If this is a subobject, tell it's RootObject about the imminent destruction
+			if (DestroyedObjectData.SubObjectRootIndex != FNetRefHandleManager::InvalidInternalIndex)
+			{
+				DestroyContext.RootObject = NetRefHandleManager->GetReplicatedObjectInstance(DestroyedObjectData.SubObjectRootIndex);
+
+				UNetObjectFactory* RootObjectFactory = GetNetFactory(NetRefHandleManager->GetReplicatedObjectData(DestroyedObjectData.SubObjectRootIndex).NetFactoryId);
+				RootObjectFactory->SubObjectDestroyedFromReplication(DestroyContext);
+			}			
 
-                UNetObjectFactory* Factory = GetNetFactory(NetFactoryId);
-                Factory->DestroyedFromReplication(DestroyInstanceParams);
-            }
+			// Now ask the object's factory to handle the destruction
+			if (ensureMsgf(NetFactoryId != InvalidNetObjectFactoryId, TEXT("Received invalid netfactory for %s: %s"), *Handle.ToString(), *GetPathNameSafe(DestroyContext.DestroyedInstance)))
+			{
+				UNetObjectFactory* Factory = GetNetFactory(NetFactoryId);
+				Factory->DestroyReplicatedObject(DestroyContext);
+			}
 		}
 	}
 	
diff --git a/Engine/Source/Runtime/Experimental/Iris/Core/Public/Iris/ReplicationSystem/NetObjectFactory.h b/Engine/Source/Runtime/Experimental/Iris/Core/Public/Iris/ReplicationSystem/NetObjectFactory.h
--- a/Engine/Source/Runtime/Experimental/Iris/Core/Public/Iris/ReplicationSystem/NetObjectFactory.h
+++ b/Engine/Source/Runtime/Experimental/Iris/Core/Public/Iris/ReplicationSystem/NetObjectFactory.h
@@ -133,6 +133,13 @@ public:
 	 */
 	virtual void PostInit(const FPostInitContext& Context) {}
 
+	/**
+	 * Destroy a replicated object when it is detached from the remote authority.
+	 * A factory is not obligated to destroy the object for exemple if a TearOff is requested or the object is static and should only be hidden.
+	 * @param Context Gives access to the instance that should be destroyed along with details on what type of destruction is requested.
+	*/
+	virtual void DestroyReplicatedObject(const FDestroyedContext& Context) PURE_VIRTUAL(UNetObjectFactory::DestroyReplicatedObject, );
+
 	/** 
 	 * Optional callback triggered when a root object managed by this factory gets assigned a dynamic subobject. 
 	 * This callback is called on remotes only. 
@@ -143,11 +150,12 @@ public:
 
 	virtual void SubObjectCreatedFromReplication(UE::Net::FNetRefHandle RootObject, UE::Net::FNetRefHandle SubObjectCreated) {}
 
-	/** 
-	 * Optional callback triggered when an remote object is about to be destroyed.
-	 * @param Context Gives access to the instance that is about to be destroyed and details on what type of destruction is requested.
-	*/
-	virtual void DestroyedFromReplication(const FDestroyedContext& Context) {}
+	/**
+	 * Optional callback triggered when a subobject is about to destroyed for the owner root object managed by this factory.
+	 * This callback is called on remotes only.
+	 * @param Context Gives access to the subobject about to be destroyed and the root object that owns him.
+	 */
+	virtual void SubObjectDestroyedFromReplication(const FDestroyedContext& Context) {}
 
 public:
 
@@ -198,7 +206,7 @@ public:
 		UE::Net::FNetRefHandle Handle;
 	};
 
-	/** Contextual information for the DestroyedFromReplication callback*/
+	/** Contextual information for the destroy callbacks */
 	struct FDestroyedContext
 	{
 		/** The object about to be destroyed */
