IncrementalCook: WorldPartion: Add the AssetRegistry query for external actors as a build dependency.
#rnx

diff --git a/Engine/Source/Runtime/AssetRegistry/Public/AssetRegistry/AssetDependencyGatherer.h b/Engine/Source/Runtime/AssetRegistry/Public/AssetRegistry/AssetDependencyGatherer.h
--- a/Engine/Source/Runtime/AssetRegistry/Public/AssetRegistry/AssetDependencyGatherer.h
+++ b/Engine/Source/Runtime/AssetRegistry/Public/AssetRegistry/AssetDependencyGatherer.h
@@ -42,6 +42,16 @@ public:
 	 * WARNING: For high performance these callbacks are called inside the critical section of the AssetRegistry. 
 	 * Attempting to call public functions on the AssetRegistry will deadlock. 
 	 * To send queries about what assets exist, used the passed-in interface functions instead.
+	 * 
+	 * @param AssetData The Asset to which the dependencies will be added
+	 * @param AssetRegistryState Helper object: The AssetRegistry's internal state which can be used to run queries
+	 *        in the callback. This must be used instead of the usual interface functions on
+	 *        IAssetRegistry::GetChecked(), because the callback executes inside the critical section.
+	 * @param CompileFilterFunc Helper object: Used to compile FARFilter into FARCompiledFilter for queries sent to
+	 *        AssetRegistryState. The same functionality that is normally provided by
+	 *        IAssetRegistry::GetChecked().CompileFilter.
+	 * @param OutDependencies List of directories that the callback queried. The AssetRegistry will call the callback
+	 *        again if any assets are added to any of these directories.
 	 */
 	virtual void GatherDependencies(const FAssetData& AssetData, const FAssetRegistryState& AssetRegistryState, 
 		TFunctionRef<FARCompiledFilter(const FARFilter&)> CompileFilterFunc, TArray<IAssetDependencyGatherer::FGathereredDependency>& OutDependencies, 
diff --git a/Engine/Source/Runtime/Engine/Classes/Engine/ExternalObjectAndActorDependencyGatherer.h b/Engine/Source/Runtime/Engine/Classes/Engine/ExternalObjectAndActorDependencyGatherer.h
--- a/Engine/Source/Runtime/Engine/Classes/Engine/ExternalObjectAndActorDependencyGatherer.h
+++ b/Engine/Source/Runtime/Engine/Classes/Engine/ExternalObjectAndActorDependencyGatherer.h
@@ -6,15 +6,24 @@
 
 #include "AssetRegistry/AssetDependencyGatherer.h"
 #include "AssetRegistry/AssetData.h"
+#include "Containers/Array.h"
+#include "Containers/UnrealString.h"
 #include "Engine/World.h"
 
+struct FARFilter;
+
 class ENGINE_API FExternalObjectAndActorDependencyGatherer : public IAssetDependencyGatherer
 {
 public:	
 	FExternalObjectAndActorDependencyGatherer() = default;
 	virtual ~FExternalObjectAndActorDependencyGatherer() = default;
 
-	virtual void GatherDependencies(const FAssetData& AssetData, const FAssetRegistryState& AssetRegistryState, TFunctionRef<FARCompiledFilter(const FARFilter&)> CompileFilterFunc, TArray<IAssetDependencyGatherer::FGathereredDependency>& OutDependencies, TArray<FString>& OutDependencyDirectories) const override;
+	static FARFilter GetQueryFilter(FName PackageName, TArray<FString>* OutQueryDirectories = nullptr);
+
+	virtual void GatherDependencies(const FAssetData& AssetData, const FAssetRegistryState& AssetRegistryState,
+		TFunctionRef<FARCompiledFilter(const FARFilter&)> CompileFilterFunc,
+		TArray<IAssetDependencyGatherer::FGathereredDependency>& OutDependencies,
+		TArray<FString>& OutDependencyDirectories) const override;
 };
 
 #endif
\ No newline at end of file
diff --git a/Engine/Source/Runtime/Engine/Private/ExternalObjectAndActorDependencyGatherer.cpp b/Engine/Source/Runtime/Engine/Private/ExternalObjectAndActorDependencyGatherer.cpp
--- a/Engine/Source/Runtime/Engine/Private/ExternalObjectAndActorDependencyGatherer.cpp
+++ b/Engine/Source/Runtime/Engine/Private/ExternalObjectAndActorDependencyGatherer.cpp
@@ -7,20 +7,12 @@
 #include "AssetRegistry/AssetRegistryState.h"
 #include "ExternalPackageHelper.h"
 
-void FExternalObjectAndActorDependencyGatherer::GatherDependencies(const FAssetData& AssetData, const FAssetRegistryState& AssetRegistryState, TFunctionRef<FARCompiledFilter(const FARFilter&)> CompileFilterFunc, TArray<IAssetDependencyGatherer::FGathereredDependency>& OutDependencies, TArray<FString>& OutDependencyDirectories) const
-{		
-	const FString ExternalActorsPath = ULevel::GetExternalActorsPath(AssetData.PackageName.ToString());
-	const FString ExternalObjectPath = FExternalPackageHelper::GetExternalObjectsPath(AssetData.PackageName.ToString());
-
-	OutDependencyDirectories.Add(ExternalActorsPath);
-	OutDependencyDirectories.Add(ExternalObjectPath);
-
-	FARFilter Filter;
-	Filter.bRecursivePaths = true;
-	Filter.bIncludeOnlyOnDiskAssets = true;
-
-	Filter.PackagePaths.Add(*ExternalActorsPath);
-	Filter.PackagePaths.Add(*ExternalObjectPath);
+void FExternalObjectAndActorDependencyGatherer::GatherDependencies(const FAssetData& AssetData,
+	const FAssetRegistryState& AssetRegistryState,TFunctionRef<FARCompiledFilter(const FARFilter&)> CompileFilterFunc,
+	TArray<IAssetDependencyGatherer::FGathereredDependency>& OutDependencies,
+	TArray<FString>& OutDependencyDirectories) const
+{
+	FARFilter Filter = GetQueryFilter(AssetData.PackageName, &OutDependencyDirectories);
 
 	TArray<FAssetData> FilteredAssets;
 	AssetRegistryState.GetAssets(CompileFilterFunc(Filter), {}, FilteredAssets, true);
@@ -31,6 +23,24 @@ void FExternalObjectAndActorDependencyGatherer::GatherDependencies(const FAssetD
 	}
 }
 
+FARFilter FExternalObjectAndActorDependencyGatherer::GetQueryFilter(FName PackageName, TArray<FString>* OutQueryDirectories)
+{
+	const FString ExternalActorsPath = ULevel::GetExternalActorsPath(PackageName.ToString());
+	const FString ExternalObjectPath = FExternalPackageHelper::GetExternalObjectsPath(PackageName.ToString());
+	if (OutQueryDirectories)
+	{
+		OutQueryDirectories->Add(ExternalActorsPath);
+		OutQueryDirectories->Add(ExternalObjectPath);
+	}
+
+	FARFilter Filter;
+	Filter.bRecursivePaths = true;
+	Filter.bIncludeOnlyOnDiskAssets = true;
+	Filter.PackagePaths.Add(FName(ExternalActorsPath));
+	Filter.PackagePaths.Add(FName(ExternalObjectPath));
+	return Filter;
+}
+
 REGISTER_ASSETDEPENDENCY_GATHERER(FExternalObjectAndActorDependencyGatherer, UWorld);
 
 #endif
diff --git a/Engine/Source/Runtime/Engine/Private/WorldPartition/WorldPartition.cpp b/Engine/Source/Runtime/Engine/Private/WorldPartition/WorldPartition.cpp
--- a/Engine/Source/Runtime/Engine/Private/WorldPartition/WorldPartition.cpp
+++ b/Engine/Source/Runtime/Engine/Private/WorldPartition/WorldPartition.cpp
@@ -37,6 +37,7 @@
 #include "LevelEditorViewport.h"
 #include "ScopedTransaction.h"
 #include "LocationVolume.h"
+#include "Engine/ExternalObjectAndActorDependencyGatherer.h"
 #include "Engine/LevelScriptBlueprint.h"
 #include "ActorReferencesUtils.h"
 #include "WorldPartition/DataLayer/WorldDataLayersActorDesc.h"
@@ -631,6 +632,16 @@ void UWorldPartition::PostDuplicate(EDuplicateMode::Type DuplicateMode)
 	}
 }
 
+void UWorldPartition::PreSave(FObjectPreSaveContext ObjectContext)
+{
+	Super::PreSave(ObjectContext);
+	if (ObjectContext.IsCooking())
+	{
+		FARFilter Filter = FExternalObjectAndActorDependencyGatherer::GetQueryFilter(GetPackage()->GetFName());
+		ObjectContext.AddCookBuildDependency(UE::Cook::FCookDependency::AssetRegistryQuery(MoveTemp(Filter)));
+	}
+}
+
 FName UWorldPartition::GetWorldPartitionEditorName() const
 {
 	if (SupportsStreaming())
diff --git a/Engine/Source/Runtime/Engine/Public/WorldPartition/WorldPartition.h b/Engine/Source/Runtime/Engine/Public/WorldPartition/WorldPartition.h
--- a/Engine/Source/Runtime/Engine/Public/WorldPartition/WorldPartition.h
+++ b/Engine/Source/Runtime/Engine/Public/WorldPartition/WorldPartition.h
@@ -210,8 +210,9 @@ public:
 	ENGINE_API virtual bool CanEditChange(const FProperty* InProperty) const override;
 	ENGINE_API virtual void PostEditChangeChainProperty(FPropertyChangedChainEvent& PropertyChangedEvent) override;
 	ENGINE_API virtual void PostDuplicate(EDuplicateMode::Type DuplicateMode) override;
+	ENGINE_API virtual void PreSave(FObjectPreSaveContext ObjectContext) override;
 #endif //WITH_EDITOR
-	ENGINE_API virtual void Serialize(FArchive& Ar) override;	
+	ENGINE_API virtual void Serialize(FArchive& Ar) override;
 	ENGINE_API virtual UWorld* GetWorld() const override;
 	ENGINE_API virtual bool ResolveSubobject(const TCHAR* SubObjectPath, UObject*& OutObject, bool bLoadIfExists) override;
 	ENGINE_API virtual void BeginDestroy() override;
