New material translator push. Added support for uniform scalar and vector parameters of type int and float.

diff --git a/Engine/Source/Runtime/Core/Private/UObject/FortniteMainBranchObjectVersion.cpp b/Engine/Source/Runtime/Core/Private/UObject/FortniteMainBranchObjectVersion.cpp
--- a/Engine/Source/Runtime/Core/Private/UObject/FortniteMainBranchObjectVersion.cpp
+++ b/Engine/Source/Runtime/Core/Private/UObject/FortniteMainBranchObjectVersion.cpp
@@ -14,6 +14,6 @@ TMap<FGuid, FGuid> FFortniteMainBranchObjectVersion::GetSystemGuids()
 	SystemGuids.Add(DevGuids.Niagara_LatestScriptCompileVersion, FGuid("393081597FA7234C92F00869DE262D73"));
 	SystemGuids.Add(DevGuids.SkeletalMeshDerivedDataVersion, FGuid("0829C9BAB43B4B25A63747FDA68E7197"));
 	SystemGuids.Add(DevGuids.STATICMESH_DERIVEDDATA_VER, FGuid("68C7D98BE48042E0A35351325741F70E"));
-	SystemGuids.Add(DevGuids.MaterialTranslationDDCVersion, FGuid("D318C2D188414B7284D232E41F77BB74"));
+	SystemGuids.Add(DevGuids.MaterialTranslationDDCVersion, FGuid("618C475161AF47B58FC057E13AD67D53"));
 	return SystemGuids;
 }
diff --git a/Engine/Source/Runtime/Engine/Private/Materials/MaterialExpressionsIR.cpp b/Engine/Source/Runtime/Engine/Private/Materials/MaterialExpressionsIR.cpp
--- a/Engine/Source/Runtime/Engine/Private/Materials/MaterialExpressionsIR.cpp
+++ b/Engine/Source/Runtime/Engine/Private/Materials/MaterialExpressionsIR.cpp
@@ -308,40 +308,7 @@ static MIR::FValue* BuildMaterialExpressionParameter(UMaterialExpressionParamete
 		return nullptr;
 	}
 
-	// Is parameter a custom primitive data?
-	if (Metadata.PrimitiveDataIndex != INDEX_NONE)
-	{
-		if (Metadata.Value.Type == EMaterialParameterType::Scalar)
-		{
-			return Emitter.EmitCustomPrimitiveData(Metadata.PrimitiveDataIndex);
-		}
-		else
-		{
-			MIR::FValue* X = Emitter.EmitCustomPrimitiveData(Metadata.PrimitiveDataIndex + 0);
-			MIR::FValue* Y = Emitter.EmitCustomPrimitiveData(Metadata.PrimitiveDataIndex + 1);
-			MIR::FValue* Z = Emitter.EmitCustomPrimitiveData(Metadata.PrimitiveDataIndex + 2);
-			MIR::FValue* W = Emitter.EmitCustomPrimitiveData(Metadata.PrimitiveDataIndex + 3);
-			return Emitter.EmitVector4(X, Y, Z, W);
-		}
-	}
-
-	switch (Metadata.Value.Type)
-	{
-		case EMaterialParameterType::Scalar:
-			return Emitter.EmitConstantFloatScalar(Metadata.Value.Float[0]);
-			
-		case EMaterialParameterType::Vector:
-			return Emitter.EmitConstantFloat4(FVector4f{ Metadata.Value.Float[0], Metadata.Value.Float[1], Metadata.Value.Float[2], Metadata.Value.Float[3] });
-
-		 case EMaterialParameterType::Texture:
-		 	return Emitter.EmitTextureObject(Metadata.Value.Texture, SAMPLERTYPE_Color);
-
-		case EMaterialParameterType::StaticSwitch:
-			return Emitter.EmitConstantBoolScalar(Metadata.Value.Bool[0]);
-
-		default:
-			UE_MIR_TODO();
-	}
+	return Emitter.EmitParameter(ParameterExpr->GetParameterName(), Metadata);
 }
 
 void UMaterialExpressionParameter::Build(UE::MIR::FEmitter& Emitter)
@@ -349,6 +316,16 @@ void UMaterialExpressionParameter::Build(UE::MIR::FEmitter& Emitter)
 	Emitter.Output(0, BuildMaterialExpressionParameter(this, Emitter));
 }
 
+void UMaterialExpressionVectorParameter::Build(UE::MIR::FEmitter& Emitter)
+{
+	MIR::FValue* Value = BuildMaterialExpressionParameter(this, Emitter);
+	Emitter.Output(0, Value);
+	Emitter.Output(1, Emitter.EmitSubscript(Value, 0));
+	Emitter.Output(2, Emitter.EmitSubscript(Value, 0));
+	Emitter.Output(3, Emitter.EmitSubscript(Value, 0));
+	Emitter.Output(4, Emitter.EmitSubscript(Value, 0));
+}
+
 void UMaterialExpressionStaticBoolParameter::Build(MIR::FEmitter& Emitter)
 {
 	MIR::FValue* Value = BuildMaterialExpressionParameter(this, Emitter);
@@ -495,9 +472,7 @@ static void BuildBinaryArithmeticOperator(
 	UE_MIR_CHECKPOINT();
 
 	// Determine operation input/output type by looking at the first connected input and picking float1 otherwise.
-	MIR::FTypePtr ResultType = LhsValue ? LhsValue->Type
-		: RhsValue ? RhsValue->Type
-		: MIR::FPrimitiveType::GetScalar(MIR::SK_Float);
+	MIR::FTypePtr ResultType = Emitter.GetCommonPrimitiveType(LhsValue->Type->AsPrimitive(), RhsValue->Type->AsPrimitive());
 
 	// Convert operand values to determined result type. 
 	LhsValue = Emitter.EmitConstruct(ResultType, LhsValue);
diff --git a/Engine/Source/Runtime/Engine/Private/Materials/MaterialIR.cpp b/Engine/Source/Runtime/Engine/Private/Materials/MaterialIR.cpp
--- a/Engine/Source/Runtime/Engine/Private/Materials/MaterialIR.cpp
+++ b/Engine/Source/Runtime/Engine/Private/Materials/MaterialIR.cpp
@@ -16,6 +16,7 @@ const TCHAR* ValueKindToString(EValueKind Kind)
 		case VK_Constant: return TEXT("Constant");
 		case VK_ExternalInput: return TEXT("ExternalInput");
 		case VK_TextureObject: return TEXT("TextureObject");
+		case VK_UniformParameter: return TEXT("Parameter");
 		case VK_Dimensional: return TEXT("Dimensional");
 		case VK_SetMaterialOutput: return TEXT("SetMaterialOutput");
 		case VK_Operator: return TEXT("Operator");
@@ -34,6 +35,7 @@ uint32 FValue::GetSizeInBytes() const
 		case VK_Constant: return sizeof(FConstant);
 		case VK_ExternalInput: return sizeof(FExternalInput);
 		case VK_TextureObject: return sizeof(FTextureObject);
+		case VK_UniformParameter: return sizeof(FUniformParameter);
 		case VK_SetMaterialOutput: return sizeof(FSetMaterialOutput);
 		case VK_Operator: return sizeof(FOperator);
 		case VK_Dimensional: return sizeof(FDimensional) + sizeof(FValue*) * static_cast<const FDimensional*>(this)->GetComponents().Num();
diff --git a/Engine/Source/Runtime/Engine/Private/Materials/MaterialIREmitter.cpp b/Engine/Source/Runtime/Engine/Private/Materials/MaterialIREmitter.cpp
--- a/Engine/Source/Runtime/Engine/Private/Materials/MaterialIREmitter.cpp
+++ b/Engine/Source/Runtime/Engine/Private/Materials/MaterialIREmitter.cpp
@@ -171,13 +171,6 @@ static FName SlowFindInputName(UMaterialExpression* Expression, const FExpressio
 
 // FEmitter API
 
-FEmitter::FEmitter(FMaterialIRModuleBuilderImpl* InBuilder, UMaterial* InMaterial, FMaterialIRModule* InModule)
-{
-	Builder = InBuilder;
-	Material = InMaterial;
-	Module = InModule;
-}
-
 FEmitter& FEmitter::DefaultInputToFloatScalarZero(const FExpressionInput* InInput)
 {
 	return DefaultInputToFloatScalar(InInput, 0.0f);
@@ -187,7 +180,7 @@ FEmitter& FEmitter::DefaultInputToBoolScalar(const FExpressionInput* InInput, bo
 {
 	if (!InInput->IsConnected())
 	{
-		Internal::SetInputValue(Builder, InInput, EmitConstantBoolScalar(Value));
+		Internal::SetInputValue(BuilderImpl, InInput, EmitConstantBoolScalar(Value));
 	}
 
 	return *this;
@@ -197,7 +190,7 @@ FEmitter& FEmitter::DefaultInputToFloatScalar(const FExpressionInput* InInput, T
 {
 	if (!InInput->IsConnected())
 	{
-		Internal::SetInputValue(Builder, InInput, EmitConstantFloatScalar(Value));
+		Internal::SetInputValue(BuilderImpl, InInput, EmitConstantFloatScalar(Value));
 	}
 
 	return *this;
@@ -205,7 +198,7 @@ FEmitter& FEmitter::DefaultInputToFloatScalar(const FExpressionInput* InInput, T
 
 FValue*  FEmitter::TryInput(const FExpressionInput* InInput)
 {
-	FValue* Value = MIR::Internal::GetInputValue(Builder, InInput);
+	FValue* Value = MIR::Internal::GetInputValue(BuilderImpl, InInput);
 	return Value ? Value : nullptr;
 }
 
@@ -294,7 +287,7 @@ FEmitter& FEmitter::Output(int OutputIndex, FValue* Value)
 
 FEmitter& FEmitter::Output(const FExpressionOutput* Output, FValue* Value)
 {
-	MIR::Internal::SetOutputValue(Builder, Output, Value);
+	MIR::Internal::SetOutputValue(BuilderImpl, Output, Value);
 	return *this;
 }
 
@@ -628,10 +621,62 @@ FValue* FEmitter::EmitTextureObject(UTexture* Texture, EMaterialSamplerType Samp
 	return EmitPrototype(this, Proto);
 }
 
-FValue* FEmitter::EmitParameter(FName Name, const FMaterialParameterMetadata& Metadata)
+FValue* FEmitter::EmitParameter(FName Name, FMaterialParameterMetadata& Metadata)
 {
-	// UE_MIR_TODO();
-	return nullptr;
+	switch (Metadata.Value.Type)
+	{
+		case EMaterialParameterType::Scalar:
+		{
+			if (Metadata.PrimitiveDataIndex != INDEX_NONE)
+			{
+				return EmitCustomPrimitiveData(Metadata.PrimitiveDataIndex);
+			}
+
+			MIR::FUniformParameter Prototype = MakePrototype<MIR::FUniformParameter>(FPrimitiveType::GetFloat1());
+			Prototype.Info.Name = Name;
+			Prototype.Metadata = Metadata;
+			return EmitPrototype(this, Prototype);
+		}
+
+		case EMaterialParameterType::Vector:
+		{
+			if (Metadata.PrimitiveDataIndex != INDEX_NONE)
+			{
+				MIR::FValue* X = EmitCustomPrimitiveData(Metadata.PrimitiveDataIndex + 0);
+				MIR::FValue* Y = EmitCustomPrimitiveData(Metadata.PrimitiveDataIndex + 1);
+				MIR::FValue* Z = EmitCustomPrimitiveData(Metadata.PrimitiveDataIndex + 2);
+				MIR::FValue* W = EmitCustomPrimitiveData(Metadata.PrimitiveDataIndex + 3);
+				return EmitVector4(X, Y, Z, W);
+			}
+
+			MIR::FUniformParameter Prototype = MakePrototype<MIR::FUniformParameter>(FPrimitiveType::GetFloat4());
+			Prototype.Info.Name = Name;
+			Prototype.Metadata = Metadata;
+			return EmitPrototype(this, Prototype);
+		}
+
+		case EMaterialParameterType::Texture:
+		{
+			return EmitTextureObject(Metadata.Value.Texture, SAMPLERTYPE_Color);
+		}
+
+		case EMaterialParameterType::StaticSwitch:
+		{
+			// Apply eventual parameter override
+			for (const FStaticSwitchParameter& Param : StaticParameterSet->GetRuntime().StaticSwitchParameters)
+			{
+				if (Param.IsOverride() && Param.ParameterInfo.Name == Name)
+				{
+					Metadata.Value.Bool[0] = Param.Value;
+					break;
+				}
+			}
+			return EmitConstantBoolScalar(Metadata.Value.Bool[0]);
+		}
+
+		default:
+			UE_MIR_TODO();
+	}
 }
 
 FValue* FEmitter::EmitCustomPrimitiveData(uint32 PrimitiveDataIndex)
diff --git a/Engine/Source/Runtime/Engine/Private/Materials/MaterialIRInternal.cpp b/Engine/Source/Runtime/Engine/Private/Materials/MaterialIRInternal.cpp
--- a/Engine/Source/Runtime/Engine/Private/Materials/MaterialIRInternal.cpp
+++ b/Engine/Source/Runtime/Engine/Private/Materials/MaterialIRInternal.cpp
@@ -67,11 +67,13 @@ MIR::FValue* CreateMaterialAttributeDefaultValue(FEmitter& Emitter, const UMater
 		case MCT_Float1: return Emitter.EmitConstantFloatScalar(  DefaultValue.X );
 		case MCT_Float2: return Emitter.EmitConstantFloat2({ DefaultValue.X, DefaultValue.Y });
 		case MCT_Float3: return Emitter.EmitConstantFloat3({ DefaultValue.X, DefaultValue.Y, DefaultValue.Z });
+		case MCT_Float:
 		case MCT_Float4: return Emitter.EmitConstantFloat4(DefaultValue);
 
 		case MCT_UInt1: return Emitter.EmitConstantIntScalar(  (int32)DefaultValue.X );
 		case MCT_UInt2: return Emitter.EmitConstantInt2({ (int32)DefaultValue.X, (int32)DefaultValue.Y });
 		case MCT_UInt3: return Emitter.EmitConstantInt3({ (int32)DefaultValue.X, (int32)DefaultValue.Y, (int32)DefaultValue.Z });
+		case MCT_UInt:
 		case MCT_UInt4: return Emitter.EmitConstantInt4({ (int32)DefaultValue.X, (int32)DefaultValue.Y, (int32)DefaultValue.Z, (int32)DefaultValue.W });
 
 		default: UE_MIR_UNREACHABLE();
diff --git a/Engine/Source/Runtime/Engine/Private/Materials/MaterialIRInternal.h b/Engine/Source/Runtime/Engine/Private/Materials/MaterialIRInternal.h
--- a/Engine/Source/Runtime/Engine/Private/Materials/MaterialIRInternal.h
+++ b/Engine/Source/Runtime/Engine/Private/Materials/MaterialIRInternal.h
@@ -39,6 +39,18 @@ void SetOutputValue(FMaterialIRModuleBuilderImpl* Builder, const FExpressionOutp
 //
 uint32 HashBytes(const char* Ptr, uint32 Size);
 
+/* Other helper functions */
+
+template <typename TKey, typename TValue>
+bool Find(const TMap<TKey, TValue>& Map, const TKey& Key, TValue& OutValue)
+{
+	if (auto ValuePtr = Map.Find(Key)) {
+		OutValue = *ValuePtr;
+		return true;
+	}
+	return false;
+}
+
 } // namespace UE::MIR::Internal
 
 #endif // #if WITH_EDITOR
diff --git a/Engine/Source/Runtime/Engine/Private/Materials/MaterialIRModuleBuilder.cpp b/Engine/Source/Runtime/Engine/Private/Materials/MaterialIRModuleBuilder.cpp
--- a/Engine/Source/Runtime/Engine/Private/Materials/MaterialIRModuleBuilder.cpp
+++ b/Engine/Source/Runtime/Engine/Private/Materials/MaterialIRModuleBuilder.cpp
@@ -4,6 +4,7 @@
 
 #if WITH_EDITOR
 
+#include "Materials/MaterialIRModuleBuilderAnalysis.h"
 #include "Materials/MaterialIRModule.h"
 #include "Materials/MaterialIRDebug.h"
 #include "Materials/MaterialIR.h"
@@ -19,8 +20,8 @@
 #include "MaterialShared.h"
 #include "Materials/MaterialExpression.h"
 #include "Materials/MaterialInsights.h"
-#include "Async/ParallelFor.h"
 #include "Engine/Texture.h"
+
 namespace MIR = UE::MIR;
 
 struct FAnalysisContext
@@ -61,6 +62,7 @@ struct FMaterialIRModuleBuilderImpl
 	MIR::FEmitter* Emitter;
 	TArray<FAnalysisContext> AnalysisContextStack;
 	TArray<MIR::FInstruction*> InstructionStack{};
+	FAnalysisData AnalysisData;
 
 	void Step_Initialize()
 	{
@@ -298,7 +300,7 @@ struct FMaterialIRModuleBuilderImpl
 				Output->Arg = Emitter->EmitConstruct(OutputArgType, *ValuePtr);
 			}
 
-			if (Builder->TargetInsight)
+			if (Builder->TargetInsights)
 			{
 				check(Output->Arg);
 				PushConnectionInsight(Builder->Material, (int)Output->Property, Input.Input->Expression, Input.Input->OutputIndex, Output->Arg->Type);
@@ -308,7 +310,7 @@ struct FMaterialIRModuleBuilderImpl
 
 	void Step_AnalyzeIRGraph()
 	{
-		InstructionStack.Reserve(64);
+		InstructionStack.Empty(64);
 
 		for (MIR::FSetMaterialOutput* Output : Module->Outputs)
 		{
@@ -317,51 +319,60 @@ struct FMaterialIRModuleBuilderImpl
 
 		while (!InstructionStack.IsEmpty())
 		{
-			MIR::FValue* Instr = InstructionStack.Pop();
+			MIR::FInstruction* Instr = InstructionStack.Last();
+
+			if (Instr->Flags & MIR::VF_Analyzed)
+			{
+				InstructionStack.Pop();
+				continue;
+			}
+
 			for (MIR::FValue* UseValue : Instr->GetUses())
 			{
-				if (UseValue && !(UseValue->Flags & MIR::VF_ValueAnalyzed))
+				// If this argument value is null, ignore it
+				if (!UseValue)
 				{
-					UseValue->SetFlags(MIR::VF_ValueAnalyzed);
-					AnalyzeValue(UseValue);
+					continue;
 				}
 
+				// This value is used. If it hasn't been analyzed yet, analyze it now.
 				MIR::FInstruction* Use = UseValue->AsInstruction();
 				if (!Use)
 				{
+					if (!(UseValue->Flags & MIR::VF_Analyzed))
+					{
+						// Only instructions can have arguments. Verify this here for safety.
+						check(UseValue->GetUses().Num() == 0);
+
+						UseValue->SetFlags(MIR::VF_Analyzed);
+						AnalyzeValue(AnalysisData, UseValue);
+					}
 					continue;
 				}
 
-				Use->NumUsers += 1;
-
-				if (!(Use->Flags & MIR::VF_InstructionAnalyzed))
+				if (!(Use->Flags & MIR::VF_Analyzed))
 				{
-					Use->SetFlags(MIR::VF_InstructionAnalyzed);
 					InstructionStack.Push(Use);
 				}
 			}
-		}
-	}
 
-	void AnalyzeValue(MIR::FValue* Value)
-	{
-		if (auto ExternalInput = Value->As<MIR::FExternalInput>())
-		{
-			Module->Statistics.ExternalInputUsedMask[SF_Vertex][(int)ExternalInput->Id] = true;
-			Module->Statistics.ExternalInputUsedMask[SF_Pixel][(int)ExternalInput->Id] = true;
-		}
-		else if (auto TextureSample = Value->As<MIR::FTextureSample>())
-		{
-			EMaterialTextureParameterType ParamType = MIR::Internal::TextureMaterialValueTypeToParameterType(TextureSample->Texture->GetMaterialType());
+			if (InstructionStack.Last() != Instr)
+			{
+				continue;
+			}
 
-			FMaterialTextureParameterInfo ParamInfo{};
-			ParamInfo.ParameterInfo = { "", EMaterialParameterAssociation::GlobalParameter, INDEX_NONE };
-			ParamInfo.SamplerSource = SSM_FromTextureAsset; // TODO - Is this needed?
+			InstructionStack.Pop();
 
-			ParamInfo.TextureIndex = Builder->Material->GetReferencedTextures().Find(TextureSample->Texture);
-			check(ParamInfo.TextureIndex != INDEX_NONE);
+			for (MIR::FValue* UseValue : Instr->GetUses())
+			{
+				if (MIR::FInstruction* UseAsInstruction = UseValue->AsInstruction())
+				{
+					UseAsInstruction->NumUsers += 1;
+				}
+			}
 
-			TextureSample->TextureParameterIndex = Module->CompilationOutput.UniformExpressionSet.FindOrAddTextureParameter(ParamType, ParamInfo);
+			Instr->SetFlags(MIR::VF_Analyzed);
+			AnalyzeValue(AnalysisData, Instr);
 		}
 	}
 
@@ -451,10 +462,13 @@ struct FMaterialIRModuleBuilderImpl
 		FMaterialCompilationOutput& CompilationOutput = Module->CompilationOutput;
 		CompilationOutput.NumUsedUVScalars = Module->Statistics.NumPixelTexCoords * 2;
 	}
+	
+	/* Other functions */
 
+	//
 	void AddExpressionConnectionInsights(UMaterialExpression* Expression)
 	{
-		if (!Builder->TargetInsight)
+		if (!Builder->TargetInsights)
 		{
 			return;
 		}
@@ -472,6 +486,7 @@ struct FMaterialIRModuleBuilderImpl
 		}
 	}
 	
+	//
 	void PushConnectionInsight(const UObject* InputObject, int InputIndex, const UMaterialExpression* OutputExpression, int OutputIndex, MIR::FTypePtr Type)
 	{
 		FMaterialInsights::FConnectionInsight Insight;
@@ -481,17 +496,28 @@ struct FMaterialIRModuleBuilderImpl
 		Insight.OutputIndex = OutputIndex,
 		Insight.ValueType = Type ? Type->ToValueType() : UE::Shader::EValueType::Any,
 		
-		Builder->TargetInsight->ConnectionInsights.Push(Insight);
+		Builder->TargetInsights->ConnectionInsights.Push(Insight);
 	}
 };
 
 bool FMaterialIRModuleBuilder::Build(FMaterialIRModule* TargetModule)
 {
 	FMaterialIRModuleBuilderImpl Impl{ this, TargetModule };
-	
-	MIR::FEmitter Emitter{ &Impl, Material, TargetModule };
-	Impl.Emitter = &Emitter;
 
+	// Setup the emitter
+	MIR::FEmitter Emitter;
+	Emitter.BuilderImpl = &Impl;
+	Emitter.Material = Material;
+	Emitter.Module = TargetModule;
+	Emitter.StaticParameterSet = &StaticParameters;
+
+	// Setup the Builder implementation
+	Impl.Emitter = &Emitter;
+	Impl.AnalysisData.Material = Material;
+	Impl.AnalysisData.Module = TargetModule;
+	Impl.AnalysisData.Insights = TargetInsights;
+	Impl.AnalysisData.CompilationOutput = &TargetModule->CompilationOutput;
+	
 	Impl.Step_Initialize();
 	Impl.Step_GenerateOutputInstructions();
 	Impl.Step_BuildMaterialExpressionsToIRGraph();
diff --git a/Engine/Source/Runtime/Engine/Private/Materials/MaterialIRModuleBuilderAnalysis.cpp b/Engine/Source/Runtime/Engine/Private/Materials/MaterialIRModuleBuilderAnalysis.cpp
new file mode 100644
--- /dev/null
+++ b/Engine/Source/Runtime/Engine/Private/Materials/MaterialIRModuleBuilderAnalysis.cpp
@@ -0,0 +1,164 @@
+// Copyright Epic Games, Inc. All Rights Reserved.
+
+#include "Runtime/Engine/Private/Materials/MaterialIRModuleBuilderAnalysis.h"
+
+#if WITH_EDITOR
+
+#include "Materials/MaterialIRModule.h"
+#include "Materials/MaterialIRDebug.h"
+#include "Materials/MaterialIR.h"
+#include "Materials/MaterialIRTypes.h"
+#include "Materials/MaterialIRInternal.h"
+#include "Materials/MaterialInsights.h"
+#include "MaterialShared.h"
+#include "Materials/MaterialInsights.h"
+#include "Engine/Texture.h"
+#include "Materials/Material.h"
+
+namespace MIR = UE::MIR;
+
+void AnalyzeExternalInput(FAnalysisData& Data, MIR::FExternalInput* ExternalInput)
+{
+	 Data.Module->GetStatistics().ExternalInputUsedMask[SF_Vertex][(int)ExternalInput->Id] = true;
+	 Data.Module->GetStatistics().ExternalInputUsedMask[SF_Pixel][(int)ExternalInput->Id] = true;
+}
+
+void AnalyzeTextureSample(FAnalysisData& Data, MIR::FTextureSample* TextureSample)
+{
+	EMaterialTextureParameterType ParamType = MIR::Internal::TextureMaterialValueTypeToParameterType(TextureSample->Texture->GetMaterialType());
+
+	FMaterialTextureParameterInfo ParamInfo {
+		.ParameterInfo = { "", EMaterialParameterAssociation::GlobalParameter, INDEX_NONE },
+		.TextureIndex = Data.Material->GetReferencedTextures().Find(TextureSample->Texture),
+		.SamplerSource = SSM_FromTextureAsset, // TODO - Is this needed?
+	};
+			
+	check(ParamInfo.TextureIndex != INDEX_NONE);
+
+	TextureSample->TextureParameterIndex = Data.CompilationOutput->UniformExpressionSet.FindOrAddTextureParameter(ParamType, ParamInfo);
+}
+
+void AnalyzeUniformParameter(FAnalysisData& Data, MIR::FUniformParameter* Parameter)
+{
+	FUniformExpressionSet& UniformExpressionSet = Data.CompilationOutput->UniformExpressionSet;
+
+	UE::Shader::FValue DefaultValue;
+	switch (Parameter->Metadata.Value.Type)
+	{
+		case EMaterialParameterType::Scalar: DefaultValue = Parameter->Metadata.Value.AsScalar(); break;
+		case EMaterialParameterType::Vector: DefaultValue = Parameter->Metadata.Value.AsLinearColor(); break;
+		default: UE_MIR_TODO();
+	}
+
+	// TODO: GetParameterOverrideValueForCurrentFunction?
+
+	uint32 DefaultValueOffset;
+	if (!MIR::Internal::Find(Data.UniformDefaultValueOffsets, DefaultValue, DefaultValueOffset))
+	{
+		DefaultValueOffset = UniformExpressionSet.AddDefaultParameterValue(DefaultValue);
+		Data.UniformDefaultValueOffsets.Add(DefaultValue, DefaultValueOffset);
+	}
+
+	Parameter->Index = UniformExpressionSet.FindOrAddNumericParameter(Parameter->Metadata.Value.Type, Parameter->Info, DefaultValueOffset);
+
+	// Get the parameter type as primitive
+	MIR::FPrimitiveTypePtr ParameterType = Parameter->Type->AsPrimitive();
+
+	// Only int and float parameters supported for now (TODO)
+	check(ParameterType->ScalarKind == MIR::SK_Int || ParameterType->ScalarKind == MIR::SK_Float);
+
+	// Depending on the number of compoents this parameter has, try to find a previous allocation that may fit it.
+	uint32 ParameterUniformBufferOffset = UINT32_MAX;
+	uint32 NumComponents = ParameterType->GetNumComponents();
+	check(NumComponents > 0);
+	check(NumComponents <= 4); // Only vectors supported for now
+
+	uint32 UsedNumComponents;
+
+	for (UsedNumComponents = NumComponents; UsedNumComponents < 4; ++UsedNumComponents)
+	{
+		auto& FreeOffsetsForCurrentNumComponents = Data.FreeOffsetsPerNumComponents[UsedNumComponents - 1];
+
+		// If there are no available slots at this nubmer of components, try with a larger one
+		if (FreeOffsetsForCurrentNumComponents.IsEmpty()) {
+			continue;
+		}
+		
+		// There is space for this number `NumComponents` in a a chunk with `i` free components
+		ParameterUniformBufferOffset = FreeOffsetsForCurrentNumComponents.Last();
+		FreeOffsetsForCurrentNumComponents.Pop();
+
+		break;
+	}
+	
+	// If UsedNumComponents is 4, it means we looked for all possible components sizes (1, 2 and 3) and
+	// could not find a chunk that would fit this parameter. Allocate a new chunk.
+	if (UsedNumComponents >= 4) {
+		check(UsedNumComponents == 4); 
+		ParameterUniformBufferOffset = UniformExpressionSet.AllocateFromUniformBuffer(16);
+	}
+
+	// Finally, check whether in the used slot there's any slack left, and if so, record it in the
+	// appropriate free-list.
+	int NumComponentsLeft = UsedNumComponents - NumComponents;
+	if (NumComponentsLeft > 0) {
+		int LeftOverOffset = ParameterUniformBufferOffset + NumComponents * 4;
+		Data.FreeOffsetsPerNumComponents[NumComponentsLeft - 1].Push(LeftOverOffset);
+	}
+
+	// Add the parameter evaluation to the uniform data
+	UniformExpressionSet.AddNumericParameterEvaluation(Parameter->Index, ParameterUniformBufferOffset);
+
+	if (Data.Insights)
+	{
+		// Push information about this parameter allocation to the insights.
+		// ParameterUniformBufferOffset is the byte offset of where the parameter will be inside
+		// the uniform buffer. Since the uniform buffer is an array of 4D vectors of 4-byte components,
+		// we first calculate UniformComponentIndexFromStart being the offset of the allocated
+		// parameter in the uniform buffer expressed in "components" by dividing the byte offset by 4 bytes.
+		// From this, we can work out the index of the 4D vector in the whole uniform buffer (stored in
+		// SlotInsight.UniformIndex) and the first component occupied (SlotInsight.UniformComponetOffset).
+		// Currently the allocation algorithm does not spread parameters across multiple uniform 4D constants,
+		// hence SlotInsight.ParameterComponentOffset is always zero.
+
+		uint32 UniformComponentIndexFromStart = ParameterUniformBufferOffset / 4;
+
+		FMaterialInsights::FUniformParameterAllocationInsight ParamInsight;
+		ParamInsight.UniformIndex = UniformComponentIndexFromStart / 4;
+		ParamInsight.FirstUniformComponent = UniformComponentIndexFromStart % 4;
+		ParamInsight.FirstParameterComponent = 0; // Explained above
+		ParamInsight.ComponentsCount = NumComponents;
+		ParamInsight.ParameterName = Parameter->Info.Name;
+
+		switch (ParameterType->ScalarKind)
+		{
+			case MIR::SK_Int: ParamInsight.ComponentType = FMaterialInsights::FUniformBufferSlotComponentType::CT_Int; break;
+			case MIR::SK_Float: ParamInsight.ComponentType = FMaterialInsights::FUniformBufferSlotComponentType::CT_Float; break;
+			default: UE_MIR_UNREACHABLE();
+		}
+
+		Data.Insights->UniformParameterAllocationInsights.Push(ParamInsight);
+	}
+}
+
+void AnalyzeValue(FAnalysisData& Data, MIR::FValue* Value)
+{
+	#define EXPAND_CASE(ValueType) \
+	case MIR::VK_##ValueType: \
+	{ \
+		Analyze##ValueType(Data, static_cast<MIR::F##ValueType*>(Value)); \
+		break; \
+	}
+
+	switch (Value->Kind)
+	{
+		EXPAND_CASE(ExternalInput);
+		EXPAND_CASE(TextureSample);
+		EXPAND_CASE(UniformParameter);
+		default: break;
+	}
+
+	#undef EXPAND_CASE
+}
+	
+#endif // #if WITH_EDITOR
diff --git a/Engine/Source/Runtime/Engine/Private/Materials/MaterialIRModuleBuilderAnalysis.h b/Engine/Source/Runtime/Engine/Private/Materials/MaterialIRModuleBuilderAnalysis.h
new file mode 100644
--- /dev/null
+++ b/Engine/Source/Runtime/Engine/Private/Materials/MaterialIRModuleBuilderAnalysis.h
@@ -0,0 +1,32 @@
+// Copyright Epic Games, Inc. All Rights Reserved.
+
+#pragma once
+
+#include "Materials/MaterialIRModuleBuilder.h"
+
+#if WITH_EDITOR
+
+#include "Materials/MaterialIRModule.h"
+#include "Materials/MaterialIRDebug.h"
+#include "Materials/MaterialIR.h"
+#include "Materials/MaterialIRTypes.h"
+#include "Materials/MaterialIRInternal.h"
+#include "MaterialShared.h"
+#include "Materials/MaterialInsights.h"
+#include "Engine/Texture.h"
+
+namespace MIR = UE::MIR;
+
+struct FAnalysisData
+{
+	UMaterial* Material;
+	FMaterialIRModule* Module;
+	FMaterialInsights* Insights; // Optional
+	FMaterialCompilationOutput* CompilationOutput;
+	TMap<UE::Shader::FValue, uint32> UniformDefaultValueOffsets;
+	TArray<uint32, TInlineAllocator<8>> FreeOffsetsPerNumComponents[3];
+};
+
+void AnalyzeValue(FAnalysisData& Data, MIR::FValue* Value);
+	
+#endif // #if WITH_EDITOR
diff --git a/Engine/Source/Runtime/Engine/Private/Materials/MaterialIRToHLSLTranslator.cpp b/Engine/Source/Runtime/Engine/Private/Materials/MaterialIRToHLSLTranslator.cpp
--- a/Engine/Source/Runtime/Engine/Private/Materials/MaterialIRToHLSLTranslator.cpp
+++ b/Engine/Source/Runtime/Engine/Private/Materials/MaterialIRToHLSLTranslator.cpp
@@ -302,6 +302,7 @@ struct FTranslator : FMaterialIRToHLSLTranslation
 		{
 			case MIR::VK_Constant: LowerConstant(static_cast<const MIR::FConstant*>(InValue)); break;
 			case MIR::VK_ExternalInput: LowerExternalInput(static_cast<const MIR::FExternalInput*>(InValue)); break;
+			case MIR::VK_UniformParameter: LowerParameter(static_cast<const MIR::FUniformParameter*>(InValue)); break;
 			default: UE_MIR_UNREACHABLE();
 		}
 
@@ -364,6 +365,16 @@ struct FTranslator : FMaterialIRToHLSLTranslation
 		}
 	}
 
+	void LowerParameter(const MIR::FUniformParameter* Parameter)
+	{
+		Printer << TEXT("Material.PreshaderBuffer[") << Parameter->Index << TEXT("].");
+		MIR::FPrimitiveTypePtr PrimitiveType = Parameter->Type->AsPrimitive();
+		check(PrimitiveType->IsScalar() || PrimitiveType->IsVector()); // no matrices yet
+		
+		static const TCHAR* Suffixes[] = { TEXT("x"), TEXT("xy"), TEXT("xyz"), TEXT("xyzw") };
+		Printer << Suffixes[PrimitiveType->NumRows - 1];
+	}
+
 	void LowerDimensional(const MIR::FDimensional* Dimensional) 
 	{
 		MIR::FPrimitiveTypePtr ArithmeticType = Dimensional->Type->AsPrimitive();
@@ -587,7 +598,7 @@ struct FTranslator : FMaterialIRToHLSLTranslation
 	
 	/* Finalization */
 
-	void SetMaterialParameters(TMap<FString, FString>& Params)
+	void GenerateTemplateStringParameters(TMap<FString, FString>& Params)
 	{
 		const FMaterialIRModule::FStatistics ModuleStatistics = Module->GetStatistics();
 
@@ -813,7 +824,7 @@ void FMaterialIRToHLSLTranslation::Run(TMap<FString, FString>& OutParameters, FS
 
 	FTranslator Translator{ *this };
 	Translator.GenerateHLSL();
-	Translator.SetMaterialParameters(OutParameters);
+	Translator.GenerateTemplateStringParameters(OutParameters);
 	Translator.GetShaderCompilerEnvironment(OutEnvironment);
 }
 
diff --git a/Engine/Source/Runtime/Engine/Private/Materials/MaterialShared.cpp b/Engine/Source/Runtime/Engine/Private/Materials/MaterialShared.cpp
--- a/Engine/Source/Runtime/Engine/Private/Materials/MaterialShared.cpp
+++ b/Engine/Source/Runtime/Engine/Private/Materials/MaterialShared.cpp
@@ -123,6 +123,7 @@ IMPLEMENT_TYPE_LAYOUT(FMaterialCompilationOutput);
 IMPLEMENT_TYPE_LAYOUT(FMeshMaterialShaderMap);
 IMPLEMENT_TYPE_LAYOUT(FMaterialProcessedSource);
 IMPLEMENT_TYPE_LAYOUT(FMaterialShaderMapContent);
+IMPLEMENT_TYPE_LAYOUT(FMaterialUniformParameterEvaluation);
 IMPLEMENT_TYPE_LAYOUT(FMaterialUniformPreshaderHeader);
 IMPLEMENT_TYPE_LAYOUT(FMaterialUniformPreshaderField);
 IMPLEMENT_TYPE_LAYOUT(FMaterialNumericParameterInfo);
@@ -3404,6 +3405,36 @@ bool FMaterial::Translate_Legacy(const FMaterialShaderMapId& ShaderMapId,
 	return true;
 }
 
+static void EmitDebugInfoComment(const FMaterialInsights& Insights, FString& ShaderCode)
+{
+	ShaderCode.Append("\n/* INSIGHTS\n");
+	ShaderCode.Append("Uniform Buffer Content:\n");
+
+	if (Insights.UniformParameterAllocationInsights.IsEmpty()) {
+		ShaderCode.Append(TEXT("\tNo uniform parameters used by the material\n"));
+	}
+	else {
+		for (auto& ParamInsight : Insights.UniformParameterAllocationInsights)
+		{
+			ShaderCode.Appendf(TEXT(" - UniformBuffer[%d]."), ParamInsight.UniformIndex);
+			for (uint32 i = 0; i < ParamInsight.ComponentsCount; ++i)
+			{
+				check(i >= 0 && i < 4);
+				ShaderCode.AppendChar(TEXT("xyzw")[ParamInsight.FirstUniformComponent + i]);
+			}
+			ShaderCode.Appendf(TEXT(" = %s."), *ParamInsight.ParameterName.ToString());
+			for (uint32 i = 0; i < ParamInsight.ComponentsCount; ++i)
+			{
+				check(i >= 0 && i < 4);
+				ShaderCode.AppendChar(TEXT("xyzw")[ParamInsight.FirstParameterComponent + i]);
+			}
+			ShaderCode.AppendChar('\n');
+		}
+	}
+
+	ShaderCode.Append("*/\n");
+}
+
 bool FMaterial::Translate_New(const FMaterialShaderMapId& InShaderMapId,
 	const FStaticParameterSet& InStaticParameters,
 	EShaderPlatform InShaderPlatform,
@@ -3427,7 +3458,7 @@ bool FMaterial::Translate_New(const FMaterialShaderMapId& InShaderMapId,
 		.ShaderPlatform = InShaderPlatform,
 		.TargetPlatform = InTargetPlatform,
 		.StaticParameters = InStaticParameters,
-		.TargetInsight = GetMaterialInterface()->MaterialInsight.Get(),
+		.TargetInsights = GetMaterialInterface()->MaterialInsight.Get(),
 	};
 
 	if (!Builder.Build(&Module))
@@ -3465,6 +3496,10 @@ bool FMaterial::Translate_New(const FMaterialShaderMapId& InShaderMapId,
 
 	// Interpolate the final material shader source string
 	FString MaterialShaderCode = Resolver.Finalize();
+
+	// Emit uniform data debug information to the end of the generated shader
+	EmitDebugInfoComment(*GetMaterialInterface()->MaterialInsight, MaterialShaderCode);
+
 	OutMaterialEnvironment->IncludeVirtualPathToContentsMap.Add(TEXT("/Engine/Generated/Material.ush"), MoveTemp(MaterialShaderCode));
 
 	return true;
diff --git a/Engine/Source/Runtime/Engine/Private/Materials/MaterialUniformExpressions.cpp b/Engine/Source/Runtime/Engine/Private/Materials/MaterialUniformExpressions.cpp
--- a/Engine/Source/Runtime/Engine/Private/Materials/MaterialUniformExpressions.cpp
+++ b/Engine/Source/Runtime/Engine/Private/Materials/MaterialUniformExpressions.cpp
@@ -704,6 +704,23 @@ uint32 FUniformExpressionSet::AddDefaultParameterValue(const UE::Shader::FValue&
 	return Offset;
 }
 
+void FUniformExpressionSet::AddNumericParameterEvaluation(uint32 ParameterIndex, uint32 BufferOffset)
+{
+	check(ParameterIndex <= UINT16_MAX);
+	check(BufferOffset <= UINT16_MAX);
+	UniformParameterEvaluations.Add({
+		.ParameterIndex = (uint16)ParameterIndex,
+		.BufferOffset = (uint16)BufferOffset
+	});
+}
+
+uint32 FUniformExpressionSet::AllocateFromUniformBuffer(uint32 SizeInBytes)
+{
+	uint32 Temp = UniformPreshaderBufferSize;
+	UniformPreshaderBufferSize += SizeInBytes;
+	return Temp;
+}
+
 int32 FUniformExpressionSet::AddVTStack(int32 InPreallocatedStackTextureIndex)
 {
 	return VTStacks.Add(FMaterialVirtualTextureStack(InPreallocatedStackTextureIndex));
@@ -792,6 +809,86 @@ void FUniformExpressionSet::FillUniformBuffer(const FMaterialRenderContext& Mate
 	FillUniformBuffer(MaterialRenderContext, UniformExpressionCache.AllocatedVTs, UniformBufferLayout, TempBuffer, TempBufferSize);
 }
 
+static void CopyValueToUniformBuffer(const UE::Shader::FValue& Value, float* PreshaderBuffer, uint32 BufferOffset)
+{
+	using namespace UE::Shader;
+	FValueTypeDescription UniformTypeDesc = GetValueTypeDescription(Value.GetType());
+
+	if ( UniformTypeDesc.ComponentType == EValueComponentType::Float)
+	{
+		FFloatValue FloatValue = Value.AsFloat();
+		float* DestAddress = PreshaderBuffer + BufferOffset;
+		for (int32 i = 0; i < Value.GetNumComponents(); ++i)
+		{
+			*DestAddress++ = FloatValue[i];
+		}
+	}
+	else if (UniformTypeDesc.ComponentType == EValueComponentType::Int)
+	{
+		FIntValue IntValue = Value.AsInt();
+		int32* DestAddress = (int32*)PreshaderBuffer + BufferOffset;
+		for (int32 i = 0; i < Value.GetNumComponents(); ++i)
+		{
+			*DestAddress++ = IntValue[i];
+		}
+	}
+	else if (UniformTypeDesc.ComponentType == EValueComponentType::Bool)
+	{
+		FBoolValue BoolValue = Value.AsBool();
+		uint32 Mask = 0u;
+		for (int32 i = 0; i < Value.GetNumComponents(); ++i)
+		{
+			if (BoolValue[i])
+			{
+				Mask |= (1u << i);
+			}
+		}
+
+		uint32 BufferOffsetAdjusted = BufferOffset / 32u;
+		uint32 BufferBitOffset = BufferOffsetAdjusted % 32u;
+		check(BufferBitOffset + Value.GetNumComponents ()<= 32u);
+
+		uint32* DestAddress = (uint32*)PreshaderBuffer + BufferOffsetAdjusted;
+		if (BufferBitOffset == 0u)
+		{
+			// First update to a group of bits needs to initialize memory
+			*DestAddress = Mask;
+		}
+		else
+		{
+			// Combine with any previous bits
+			*DestAddress |= (Mask << BufferBitOffset);
+		}
+	}
+	else if (UniformTypeDesc.ComponentType == EValueComponentType::Double)
+	{
+		FDoubleValue DoubleValue = Value.AsDouble();
+
+		float ValueHigh[4] = {};
+		float ValueLow[4] = {};
+		for (int32 i = 0; i < Value.GetNumComponents(); ++i)
+		{
+			const FDFScalar ScalarValue(DoubleValue[i]);
+			ValueHigh[i] = ScalarValue.High;
+			ValueLow[i] = ScalarValue.Low;
+		}
+
+		float* DestAddress = PreshaderBuffer + BufferOffset;
+		for (int32 i = 0; i < Value.GetNumComponents(); ++i)
+		{
+			*DestAddress++ = ValueHigh[i];
+		}
+		for (int32 i = 0; i < Value.GetNumComponents(); ++i)
+		{
+			*DestAddress++ = ValueLow[i];
+		}
+	}
+	else
+	{
+		ensure(false);
+	}
+}
+
 void FUniformExpressionSet::FillUniformBuffer(const FMaterialRenderContext& MaterialRenderContext, TConstArrayView<IAllocatedVirtualTexture*> AllocatedVTs, const FRHIUniformBufferLayout* UniformBufferLayout, uint8* TempBuffer, int TempBufferSize) const
 {
 	using namespace UE::Shader;
@@ -892,6 +989,44 @@ void FUniformExpressionSet::FillUniformBuffer(const FMaterialRenderContext& Mate
 			}
 		}
 
+		// Perform direct parameter evaluation
+		for (const FMaterialUniformParameterEvaluation& Evaluation : UniformParameterEvaluations) {
+			// Get the numeric parameter information
+			const FMaterialNumericParameterInfo& Parameter = GetNumericParameter(Evaluation.ParameterIndex);
+		
+			// Resolve the parameter value from the material-chain
+			// Note: it is really not ideal to perform this evaluation here, as this is somewhat static information.
+			//		 Ideally it should be the translator that performs this parameter resolution at translation time
+			//		 and exports the necessary data to *efficiently* access the actual parameter here.
+			bool bFoundParameter = false;
+			UE::Shader::FValue ParameterShaderValue;
+
+			// First allow proxy the chance to override parameter
+			FMaterialParameterValue ParameterValue;
+			if (MaterialRenderContext.MaterialRenderProxy 
+				&& MaterialRenderContext.MaterialRenderProxy->GetParameterValue(Parameter.ParameterType, Parameter.ParameterInfo, ParameterValue, MaterialRenderContext))
+			{
+				ParameterShaderValue = ParameterValue.AsShaderValue();
+				bFoundParameter = true;
+			}
+
+			// Then allow the editor to override the parameter value
+			#if WITH_EDITOR
+			if (!bFoundParameter
+				&& MaterialRenderContext.Material.TransientOverrides.GetNumericOverride(Parameter.ParameterType, Parameter.ParameterInfo, ParameterShaderValue))
+			{
+				bFoundParameter = true;
+			}
+			#endif // WITH_EDITOR
+
+			if (!bFoundParameter) {
+				ParameterShaderValue = GetDefaultParameterValue(Parameter.ParameterType, Parameter.DefaultValueOffset);
+			}
+
+			// Copy the value to the target position
+			CopyValueToUniformBuffer(ParameterShaderValue, (float*)BufferCursor, Evaluation.BufferOffset);
+		}
+
 		// Dump preshader results into buffer.
 		float* const PreshaderBuffer = (float*)BufferCursor;
 		FPreshaderStack PreshaderStack;
@@ -914,27 +1049,27 @@ void FUniformExpressionSet::FillUniformBuffer(const FMaterialRenderContext& Mate
 			int32 MinValueType = FMath::Min((int32)FirstFieldType, (int32)ResultType);
 			int32 MaxValueType = FMath::Max((int32)FirstFieldType, (int32)ResultType);
 
-			if (Preshader.NumFields == 1 && MinValueType >= (int32)EValueType::Float1 && MaxValueType <= (int32)EValueType::Float4)
+			for (uint32 FieldIndex = 0u; FieldIndex < Preshader.NumFields; ++FieldIndex)
 			{
-				float* DestAddress = PreshaderBuffer + UniformPreshaderFields[Preshader.FieldIndex].BufferOffset;
-
-				// Copy components that exist in both result and destination buffer
-				int32 NumCopyComponents = MinValueType - ((int32)EValueType::Float1 - 1);
-				for (int32 i = 0; i < NumCopyComponents; ++i)
+				if (MinValueType >= (int32)EValueType::Float1 && MaxValueType <= (int32)EValueType::Float4)
 				{
-					*DestAddress++ = Result.Component[i].Float;
-				}
+					float* DestAddress = PreshaderBuffer + UniformPreshaderFields[Preshader.FieldIndex].BufferOffset;
 
-				// Zero out any additional components that exist in the destination buffer
-				int32 NumDestComponents = (int32)FirstFieldType - ((int32)EValueType::Float1 - 1);
-				for (int32 i = NumCopyComponents; i < NumDestComponents; ++i)
-				{
-					*DestAddress++ = 0.0f;
+					// Copy components that exist in both result and destination buffer
+					int32 NumCopyComponents = MinValueType - ((int32)EValueType::Float1 - 1);
+					for (int32 i = 0; i < NumCopyComponents; ++i)
+					{
+						*DestAddress++ = Result.Component[i].Float;
+					}
+
+					// Zero out any additional components that exist in the destination buffer
+					int32 NumDestComponents = (int32)FirstFieldType - ((int32)EValueType::Float1 - 1);
+					for (int32 i = NumCopyComponents; i < NumDestComponents; ++i)
+					{
+						*DestAddress++ = 0.0f;
+					}
 				}
-			}
-			else
-			{
-				for (uint32 FieldIndex = 0u; FieldIndex < Preshader.NumFields; ++FieldIndex)
+				else
 				{
 					const FMaterialUniformPreshaderField& PreshaderField = UniformPreshaderFields[Preshader.FieldIndex + FieldIndex];
 					const FValueTypeDescription UniformTypeDesc = GetValueTypeDescription(PreshaderField.Type);
@@ -951,74 +1086,8 @@ void FUniformExpressionSet::FillUniformBuffer(const FMaterialRenderContext& Mate
 					{
 						FieldValue.Component[i] = Result.Component[PreshaderField.ComponentIndex + i];
 					}
-
-					if (UniformTypeDesc.ComponentType == EValueComponentType::Float)
-					{
-						const FFloatValue FloatValue = FieldValue.AsFloat();
-						float* DestAddress = PreshaderBuffer + PreshaderField.BufferOffset;
-						for (int32 i = 0; i < NumFieldComponents; ++i)
-						{
-							*DestAddress++ = FloatValue[i];
-						}
-					}
-					else if (UniformTypeDesc.ComponentType == EValueComponentType::Int)
-					{
-						const FIntValue IntValue = FieldValue.AsInt();
-						int32* DestAddress = (int32*)PreshaderBuffer + PreshaderField.BufferOffset;
-						for (int32 i = 0; i < NumFieldComponents; ++i)
-						{
-							*DestAddress++ = IntValue[i];
-						}
-					}
-					else if (UniformTypeDesc.ComponentType == EValueComponentType::Bool)
-					{
-						const FBoolValue BoolValue = FieldValue.AsBool();
-						uint32 Mask = 0u;
-						for (int32 i = 0; i < NumFieldComponents; ++i)
-						{
-							if (BoolValue[i])
-							{
-								Mask |= (1u << i);
-							}
-						}
-
-						const uint32 BufferOffset = PreshaderField.BufferOffset / 32u;
-						const uint32 BufferBitOffset = PreshaderField.BufferOffset % 32u;
-						check(BufferBitOffset + NumFieldComponents <= 32u);
-
-						uint32* DestAddress = (uint32*)PreshaderBuffer + BufferOffset;
-						if (BufferBitOffset == 0u)
-						{
-							// First update to a group of bits needs to initialize memory
-							*DestAddress = Mask;
-						}
-						else
-						{
-							// Combine with any previous bits
-							*DestAddress |= (Mask << BufferBitOffset);
-						}
-					}
-					else if (UniformTypeDesc.ComponentType == EValueComponentType::Double)
-					{
-						const FDoubleValue DoubleValue = FieldValue.AsDouble();
-
-						float ValueHigh[4];
-						float ValueLow[4];
-						for (int32 i = 0; i < NumFieldComponents; ++i)
-						{
-							const FDFScalar Value(DoubleValue[i]);
-							ValueHigh[i] = Value.High;
-							ValueLow[i] = Value.Low;
-						}
-
-						float* DestAddress = PreshaderBuffer + PreshaderField.BufferOffset;
-						for (int32 i = 0; i < NumFieldComponents; ++i) *DestAddress++ = ValueHigh[i];
-						for (int32 i = 0; i < NumFieldComponents; ++i) *DestAddress++ = ValueLow[i];
-					}
-					else
-					{
-						ensure(false);
-					}
+					
+					CopyValueToUniformBuffer(FieldValue, PreshaderBuffer, PreshaderField.BufferOffset);
 				}
 			}
 		}
diff --git a/Engine/Source/Runtime/Engine/Public/MaterialShared.h b/Engine/Source/Runtime/Engine/Public/MaterialShared.h
--- a/Engine/Source/Runtime/Engine/Public/MaterialShared.h
+++ b/Engine/Source/Runtime/Engine/Public/MaterialShared.h
@@ -452,6 +452,19 @@ public:
 	LAYOUT_FIELD(int32, PreallocatedStackTextureIndex);
 };
 
+/**
+* Represents the evaluation of a material parameter directly to a byte offset in the preshader uniform buffer.
+* A UniformParameterEvaluation avoids the overhead of executing an entire preshader when all it does is 
+* evaluate a parameter with a single instruction.
+*/
+class FMaterialUniformParameterEvaluation
+{
+	DECLARE_TYPE_LAYOUT(FMaterialUniformParameterEvaluation, NonVirtual);
+public:
+	LAYOUT_FIELD(uint16, ParameterIndex);	/** Index of the parameter for UniformExpressionSet::GetNumericParameter() */
+	LAYOUT_FIELD(uint16, BufferOffset); 	/** Offset in the preshader uniform buffer (in bytes) where result will be written */
+};
+
 class FMaterialUniformPreshaderField
 {
 	DECLARE_TYPE_LAYOUT(FMaterialUniformPreshaderField, NonVirtual);
@@ -691,7 +704,9 @@ public:
 	int32 FindOrAddExternalTextureParameter(const FMaterialExternalTextureParameterInfo& Info);
 	int32 FindOrAddNumericParameter(EMaterialParameterType Type, const FMaterialParameterInfo& ParameterInfo, uint32 DefaultValueOffset);
 	uint32 AddDefaultParameterValue(const UE::Shader::FValue& Value);
-
+	void AddNumericParameterEvaluation(uint32 ParameterIndex, uint32 BufferOffset);
+	uint32 AllocateFromUniformBuffer(uint32 SizeInBytes);
+	
 	TConstArrayView<FMaterialVirtualTextureStack> GetVTStacks() const { return VTStacks; }
 	const FMaterialVirtualTextureStack& GetVTStack(int32 Index) const { return VTStacks[Index]; }
 	int32 AddVTStack(int32 InPreallocatedStackTextureIndex);
@@ -723,6 +738,7 @@ protected:
 	friend class UE::HLSLTree::FEmitContext;
 	friend class UE::HLSLTree::Material::FExpressionStaticTerrainLayerWeight;
 
+	LAYOUT_FIELD(TMemoryImageArray<FMaterialUniformParameterEvaluation>, UniformParameterEvaluations);
 	LAYOUT_FIELD(TMemoryImageArray<FMaterialUniformPreshaderHeader>, UniformPreshaders);
 	LAYOUT_FIELD(TMemoryImageArray<FMaterialUniformPreshaderField>, UniformPreshaderFields);
 	LAYOUT_FIELD(TMemoryImageArray<FMaterialNumericParameterInfo>, UniformNumericParameters);
diff --git a/Engine/Source/Runtime/Engine/Public/Materials/MaterialExpressionVectorParameter.h b/Engine/Source/Runtime/Engine/Public/Materials/MaterialExpressionVectorParameter.h
--- a/Engine/Source/Runtime/Engine/Public/Materials/MaterialExpressionVectorParameter.h
+++ b/Engine/Source/Runtime/Engine/Public/Materials/MaterialExpressionVectorParameter.h
@@ -28,6 +28,7 @@ class UMaterialExpressionVectorParameter : public UMaterialExpressionParameter
 
 	//~ Begin UMaterialExpression Interface
 #if WITH_EDITOR
+	virtual void Build(UE::MIR::FEmitter& Emitter) override;
 	virtual int32 Compile(class FMaterialCompiler* Compiler, int32 OutputIndex) override;
 	virtual void GetCaption(TArray<FString>& OutCaptions) const override;
 	virtual bool GetParameterValue(FMaterialParameterMetadata& OutMeta) const override
diff --git a/Engine/Source/Runtime/Engine/Public/Materials/MaterialIR.h b/Engine/Source/Runtime/Engine/Public/Materials/MaterialIR.h
--- a/Engine/Source/Runtime/Engine/Public/Materials/MaterialIR.h
+++ b/Engine/Source/Runtime/Engine/Public/Materials/MaterialIR.h
@@ -16,6 +16,7 @@ enum EValueKind
 	VK_Constant,
 	VK_ExternalInput,
 	VK_TextureObject,
+	VK_UniformParameter,
 
 	/* Instructions */
 
@@ -39,8 +40,7 @@ const TCHAR* ValueKindToString(EValueKind Kind);
 enum EValueFlags
 {
 	VF_None = 0,
-	VF_ValueAnalyzed = 1,       // todo
-	VF_InstructionAnalyzed = 2, // todo
+	VF_Analyzed = 1,       // todo
 };
 
 //
@@ -228,6 +228,14 @@ struct FTextureObject : TValue<VK_TextureObject>
 	EMaterialSamplerType SamplerType;
 };
 
+//
+struct FUniformParameter: TValue<VK_UniformParameter>
+{
+	FMaterialParameterInfo Info;
+	FMaterialParameterMetadata Metadata;
+	int Index;
+};
+
 /* Instructions */
 
 //
diff --git a/Engine/Source/Runtime/Engine/Public/Materials/MaterialIREmitter.h b/Engine/Source/Runtime/Engine/Public/Materials/MaterialIREmitter.h
--- a/Engine/Source/Runtime/Engine/Public/Materials/MaterialIREmitter.h
+++ b/Engine/Source/Runtime/Engine/Public/Materials/MaterialIREmitter.h
@@ -42,8 +42,6 @@ struct FSwizzleMask
 class FEmitter
 {
 public:
-	FEmitter(FMaterialIRModuleBuilderImpl* InBuilder, UMaterial* InMaterial, FMaterialIRModule* InModule);
-	
 	/* Defaults */
 	
 	//
@@ -147,7 +145,7 @@ public:
 	FValue* EmitSubscript(FValue* Value, int ComponentIndex);
 	FValue* EmitSwizzle(FValue* Value, FSwizzleMask Mask);
 	FValue* EmitTextureObject(UTexture* Texture, EMaterialSamplerType SamplerType);
-	FValue* EmitParameter(FName Name, const FMaterialParameterMetadata& Metadata);
+	FValue* EmitParameter(FName Name, FMaterialParameterMetadata& Metadata);
 	FValue* EmitCustomPrimitiveData(uint32 PrimitiveDataIndex);
 
 	/* Instructions */
@@ -173,11 +171,12 @@ public:
 	}
 
 	void Error(FString Message);
-
+	
 	struct FPrivate;
-	friend FMaterialIRModuleBuilderImpl;
 
 private:
+	FEmitter() {}
+
 	void Initialize();
 
 	struct FValueKeyFuncs : DefaultKeyFuncs<FValue*>
@@ -186,16 +185,21 @@ private:
 		static uint32 GetKeyHash(KeyInitType Key);
 	};
 
+	FMaterialIRModuleBuilderImpl* BuilderImpl{};
 	UMaterial* Material{};
 	FMaterialIRModule* Module{};
+	const FStaticParameterSet* StaticParameterSet{};
 	UMaterialExpression* Expression{};
-	FMaterialIRModuleBuilderImpl* Builder{};
 	bool bExpressionHasError = false;
 	FValue* ConstantTrue;
 	FValue* ConstantFalse;
 	TSet<FValue*, FValueKeyFuncs> ValueSet;
+
+	friend FMaterialIRModuleBuilder;
+	friend FMaterialIRModuleBuilderImpl;
 };
 
+
 } // namespace UE::MIR
 
 #endif // #if WITH_EDITOR
diff --git a/Engine/Source/Runtime/Engine/Public/Materials/MaterialIRModule.h b/Engine/Source/Runtime/Engine/Public/Materials/MaterialIRModule.h
--- a/Engine/Source/Runtime/Engine/Public/Materials/MaterialIRModule.h
+++ b/Engine/Source/Runtime/Engine/Public/Materials/MaterialIRModule.h
@@ -6,6 +6,17 @@
 
 #if WITH_EDITOR
 
+
+/// This class represents the intermediate representation (IR) of a material build.
+/// The `IRModule` includes an IR value graph, produced through expression analysis,
+/// as well as metadata on resource usage and reflection. The IR graph serves as an
+/// abstract representation of the material and must be translated into a target backend
+/// such as HLSL or specific Preshader opcodes for execution.
+///
+/// This class is designed to be backend-agnostic, meaning it does not contain any
+/// HLSL code nor does it configure a `MaterialCompilationOutput` instance. The data
+/// stored within this class should be sufficient to enable translation to any supported
+/// backend without requiring additional processing or validation.
 class FMaterialIRModule
 {
 public:
@@ -14,7 +25,7 @@ public:
 		UMaterialExpression* Expression;
 		FString Message;
 	};
-
+	
 	struct FStatistics
 	{
 		TBitArray<> ExternalInputUsedMask[SF_NumFrequencies];
@@ -32,6 +43,7 @@ public:
 	const UE::MIR::FBlock& GetRootBlock() const { return *RootBlock; }
 	TArrayView<const FError> GetErrors() const { return Errors; }
 	const FStatistics& GetStatistics() const { return Statistics; }
+	FStatistics& GetStatistics() { return Statistics; }
 
 private:
 	EShaderPlatform ShaderPlatform;
@@ -44,6 +56,7 @@ private:
 	FStatistics Statistics;
 
 	friend UE::MIR::FEmitter;
+	friend FMaterialIRModuleBuilder;
 	friend FMaterialIRModuleBuilderImpl;
 };
 
diff --git a/Engine/Source/Runtime/Engine/Public/Materials/MaterialIRModuleBuilder.h b/Engine/Source/Runtime/Engine/Public/Materials/MaterialIRModuleBuilder.h
--- a/Engine/Source/Runtime/Engine/Public/Materials/MaterialIRModuleBuilder.h
+++ b/Engine/Source/Runtime/Engine/Public/Materials/MaterialIRModuleBuilder.h
@@ -12,7 +12,7 @@ struct FMaterialIRModuleBuilder
 	EShaderPlatform ShaderPlatform;
 	const ITargetPlatform* TargetPlatform;
 	const FStaticParameterSet& StaticParameters;
-	FMaterialInsights* TargetInsight{};
+	FMaterialInsights* TargetInsights{};
 
 	bool Build(FMaterialIRModule* TargetModule);
 };
diff --git a/Engine/Source/Runtime/Engine/Public/Materials/MaterialInsights.h b/Engine/Source/Runtime/Engine/Public/Materials/MaterialInsights.h
--- a/Engine/Source/Runtime/Engine/Public/Materials/MaterialInsights.h
+++ b/Engine/Source/Runtime/Engine/Public/Materials/MaterialInsights.h
@@ -31,7 +31,33 @@ struct FMaterialInsights
         UE::Shader::EValueType ValueType;             ///< Type of the value flowing through the connection.
     };
 
-    TArray<FConnectionInsight> ConnectionInsights;    ///< Array of connection insights.
+	enum FUniformBufferSlotComponentType {
+		CT_Unused,
+		CT_Int,
+		CT_Float,
+	};
+
+	/**
+	* @brief Wraps information about a uniform parameter (ScalarParametr, VectorParameter) allocation in
+	* the uniform expression data buffer.
+	* An instance of this struct specifies wheere the components of a uniform parameter in the material were
+	* allocated in the preshader uniform buffer, that is which vec4 slot and which target components of that vec4.
+	*/
+	struct FUniformParameterAllocationInsight
+	{
+		uint32 UniformIndex;							///< Index of the preshader buffer constant vec4 (e.g. PreshaderBuffer[0])
+		uint32 FirstUniformComponent;					///< First component of the vec4 above (e.g. 0 for .x, 1 for .y, etc)
+		FName ParameterName;							///< The parameter name
+		uint32 FirstParameterComponent;					///< First component of the parameter (eg. 0 for .x, etc)
+		uint32 ComponentsCount;							///< Number of components stored, same for the Uniform and the Parameter
+		FUniformBufferSlotComponentType ComponentType;  ///< Parameter component type
+	};
+
+	/// Array of connection insights.
+    TArray<FConnectionInsight> ConnectionInsights;
+	
+	/// Array of parameter allocation insights.
+    TArray<FUniformParameterAllocationInsight> UniformParameterAllocationInsights;
 
     FMaterialInsights();
 
