Optimize UWorldPartition::OnGCPostReachabilityAnalysis by retrieving the objects unreachable flag from the iterator object item instead of dereferencing the object pointer, per Ben Ziegler suggestion.
Test case repro in internal game went from ~85ms down to ~30ms.

diff --git a/Engine/Source/Runtime/Engine/Private/WorldPartition/WorldPartition.cpp b/Engine/Source/Runtime/Engine/Private/WorldPartition/WorldPartition.cpp
--- a/Engine/Source/Runtime/Engine/Private/WorldPartition/WorldPartition.cpp
+++ b/Engine/Source/Runtime/Engine/Private/WorldPartition/WorldPartition.cpp
@@ -394,6 +394,7 @@ void UWorldPartition::OnGCPostReachabilityAnalysis()
 	if (!IsEngineExitRequested())
 	{
 		TRACE_CPUPROFILER_EVENT_SCOPE(UWorldPartition::OnGCPostReachabilityAnalysis);
+		UE_SCOPED_TIMER(TEXT("UWorldPartition::OnGCPostReachabilityAnalysis"), LogWorldPartition, Display);
 		const TIndirectArray<FWorldContext>& WorldContextList = GEngine->GetWorldContexts();
 
 		// Avoid running this process while a game world is live
@@ -407,24 +408,27 @@ void UWorldPartition::OnGCPostReachabilityAnalysis()
 
 		for (FRawObjectIterator It; It; ++It)
 		{
-			if (AActor* Actor = Cast<AActor>(static_cast<UObject*>(It->Object)))
+			if (It->IsUnreachable())
 			{
-				if (Actor->IsUnreachable() && !Actor->GetClass()->HasAnyClassFlags(CLASS_NewerVersionExists) && Actor->IsMainPackageActor())
+				if (AActor* Actor = Cast<AActor>(static_cast<UObject*>(It->Object)))
 				{
-					ForEachObjectWithPackage(Actor->GetPackage(), [this, Actor](UObject* Object)
+					if (Actor->IsPackageExternal() && Actor->IsMainPackageActor() && !Actor->GetClass()->HasAnyClassFlags(CLASS_NewerVersionExists))
 					{
-						if (Object->HasAnyFlags(RF_Standalone))
+						ForEachObjectWithPackage(Actor->GetPackage(), [this, Actor](UObject* Object)
 						{
-							UE_LOG(LogWorldPartition, Log, TEXT("Actor %s is unreachable without properly detaching object %s in its package"), *Actor->GetPathName(), *Object->GetPathName());
+							if (Object->HasAnyFlags(RF_Standalone))
+							{
+								UE_LOG(LogWorldPartition, Log, TEXT("Actor %s is unreachable without properly detaching object %s in its package"), *Actor->GetPathName(), *Object->GetPathName());
 
-							Object->ClearFlags(RF_Standalone);
+								Object->ClearFlags(RF_Standalone);
 
-							// Make sure we trigger a second GC at the next tick to properly destroy packages there were fixed in this pass
-							bForceGarbageCollection = true;
-							bForceGarbageCollectionPurge = true;
-						}
-						return true;
-					}, false);
+								// Make sure we trigger a second GC at the next tick to properly destroy packages there were fixed in this pass
+								bForceGarbageCollection = true;
+								bForceGarbageCollectionPurge = true;
+							}
+							return true;
+						}, false);
+					}
 				}
 			}
 		}
