Optimization: During IncrementalPurgeGarbage, UnhashUnreachableObjects calls USceneComponent::BeginDestroy ends up calling USceneComponent::UpdateBounds() on itself or on an attached component. Since there is no known case where updating the bounds is necessary when the component is not valid, skip the update.

diff --git a/Engine/Source/Runtime/Engine/Private/Components/SceneComponent.cpp b/Engine/Source/Runtime/Engine/Private/Components/SceneComponent.cpp
--- a/Engine/Source/Runtime/Engine/Private/Components/SceneComponent.cpp
+++ b/Engine/Source/Runtime/Engine/Private/Components/SceneComponent.cpp
@@ -1444,6 +1444,13 @@ void USceneComponent::CalcBoundingCylinder(float& CylinderRadius, float& Cylinde
 
 void USceneComponent::UpdateBounds()
 {
+	// UpdateBounds can be called on components marked for garbage (from its BeginDestroy or from an attached component)
+	// In this case, early out as the component will get destroyed anyways
+	if (!IsValidChecked(this))
+	{
+		return;
+	}
+
 	// if use parent bound if attach parent exists, and the flag is set
 	// since parents tick first before child, this should work correctly
 	if ( bUseAttachParentBound && GetAttachParent() != nullptr )
