Fix device profile fragment includes overridden by base DP values

#rnx
#tests Confirmed working in 33.10 on PS4, XB1, fixing the issue as reported

diff --git a/Engine/Source/Runtime/Engine/Private/DeviceProfiles/DeviceProfileManager.cpp b/Engine/Source/Runtime/Engine/Private/DeviceProfiles/DeviceProfileManager.cpp
--- a/Engine/Source/Runtime/Engine/Private/DeviceProfiles/DeviceProfileManager.cpp
+++ b/Engine/Source/Runtime/Engine/Private/DeviceProfiles/DeviceProfileManager.cpp
@@ -138,10 +138,10 @@ static bool GetFragmentCVars(const FString& FragmentName, const FString& CVarArr
 	return true;
 }
 
-// read the requested fragment from within the +FragmentIncludes= array of a DP.
-static void GetCVarsFromDPFragmentIncludes(const FString& CurrentSectionName, const FString& CVarArrayName, TArray<FString>& FragmentCVarsINOUT, FConfigCacheIni* ConfigSystem)
+// Read the requested fragment from within the +FragmentIncludes= array of a DP. When bForAppending is set to true, we'll use OverrideFragmentIncludes instead.
+static void GetCVarsFromDPFragmentIncludes(const FString& CurrentSectionName, const FString& CVarArrayName, TArray<FString>& FragmentCVarsINOUT, FConfigCacheIni* ConfigSystem, bool bForAppending)
 {
-	FString FragmentIncludes = TEXT("FragmentIncludes");
+	const FString FragmentIncludes = bForAppending ? TEXT("OverrideFragmentIncludes") : TEXT("FragmentIncludes");
 	TArray<FString> FragmentIncludeArray;
 	ConfigSystem->GetArray(*CurrentSectionName, *FragmentIncludes, FragmentIncludeArray, GDeviceProfilesIni);
 
@@ -289,8 +289,9 @@ TMap<FName, FString> UDeviceProfileManager::GatherDeviceProfileCVars(const FStri
 					ArrayName += BucketNames[(int32)MemBucket];
 				}
 
-				TArray< FString > CurrentProfilesCVars, FragmentCVars;
-				GetCVarsFromDPFragmentIncludes(*CurrentSectionName, *ArrayName, FragmentCVars, ConfigSystem);
+				TArray< FString > CurrentProfilesCVars, FragmentCVars, OverrideFragmentCVars;
+				GetCVarsFromDPFragmentIncludes(*CurrentSectionName, *ArrayName, FragmentCVars, ConfigSystem, false);
+				GetCVarsFromDPFragmentIncludes(*CurrentSectionName, *ArrayName, OverrideFragmentCVars, ConfigSystem, true);
 				CurrentSection->MultiFind(*ArrayName, CurrentProfilesCVars, true);
 
 				if (FragmentCVars.Num())
@@ -300,6 +301,8 @@ TMap<FName, FString> UDeviceProfileManager::GatherDeviceProfileCVars(const FStri
 					CurrentProfilesCVars += FragmentCVars;
 				}
 
+				CurrentProfilesCVars += OverrideFragmentCVars;
+
 				// now add the selected fragments at the end so these override the DP.
 				CurrentProfilesCVars += SelectedFragmentCVars;
 				SelectedFragmentCVars.Empty();
@@ -312,6 +315,7 @@ TMap<FName, FString> UDeviceProfileManager::GatherDeviceProfileCVars(const FStri
 						FString CVarKey, CVarValue;
 						if ((*CVarIt).Split(TEXT("="), &CVarKey, &CVarValue))
 						{
+							UE_LOG(LogDeviceProfileManager, VeryVerbose, TEXT("ValidCVars: %s=%s"), *CVarKey, *CVarValue);
 							ValidCVars.Add(CVarKey, CVarValue);
 						}
 					}
@@ -321,6 +325,7 @@ TMap<FName, FString> UDeviceProfileManager::GatherDeviceProfileCVars(const FStri
 
 					for (TMap< FString, FString >::TConstIterator ProcessedCVarIt(ValidCVars); ProcessedCVarIt; ++ProcessedCVarIt)
 					{
+						UE_LOG(LogDeviceProfileManager, VeryVerbose, TEXT("CurrentProfilesCVars: %s=%s"), *ProcessedCVarIt.Key(), *ProcessedCVarIt.Value());
 						CurrentProfilesCVars.Add(FString::Printf(TEXT("%s=%s"), *ProcessedCVarIt.Key(), *ProcessedCVarIt.Value()));
 					}
 
@@ -335,6 +340,7 @@ TMap<FName, FString> UDeviceProfileManager::GatherDeviceProfileCVars(const FStri
 						FName CVarKeyName(*CVarKey);
 						if (!DeviceProfileCVars.Find(CVarKeyName))
 						{
+							UE_LOG(LogDeviceProfileManager, VeryVerbose, TEXT("DeviceProfileCVars: %s=%s"), *CVarKey, *CVarValue);
 							DeviceProfileCVars.Add(CVarKeyName, CVarValue);
 						}
 					}
