Optimization: Early out in AWorldDataLayers::GetDataLayerInstanceFromAssetName() if the provided asset name is None

diff --git a/Engine/Source/Runtime/Engine/Private/WorldPartition/DataLayer/WorldDataLayers.cpp b/Engine/Source/Runtime/Engine/Private/WorldPartition/DataLayer/WorldDataLayers.cpp
--- a/Engine/Source/Runtime/Engine/Private/WorldPartition/DataLayer/WorldDataLayers.cpp
+++ b/Engine/Source/Runtime/Engine/Private/WorldPartition/DataLayer/WorldDataLayers.cpp
@@ -1040,6 +1040,11 @@ const UDataLayerInstance* AWorldDataLayers::GetDataLayerInstance(const FName& In
 
 const UDataLayerInstance* AWorldDataLayers::GetDataLayerInstanceFromAssetName(const FName& InDataLayerAssetPathName) const
 {
+	if (InDataLayerAssetPathName.IsNone())
+	{
+		return nullptr;
+	}
+
 #if WITH_EDITOR	
 	const UDataLayerInstance* FoundDataLayerInstance = nullptr;
 	ForEachDataLayerInstance([InDataLayerAssetPathName, &FoundDataLayerInstance](UDataLayerInstance* DataLayerInstance)
