Cleanup spammy physics stats scopes.

+ Move STAT_UnweldPhysics below the conditions so it's only emitted if we actually unweld.
+ Mark STAT_UpdatePhysFilter and STAT_BodyInst_SetCollisionProfileName as verbose as they're very common and very short.
+ Mark STAT_UpdateShapesArrayFromGeometryImpl verbose is there can be many hundreds of these tiny scopes and they bloat the parent scope.
+ Mark STAT_ComponentCalcBounds verbose as it's pretty fast, very common and largely covered by STAT_UpdateComponentToWorld.
+ Mark STAT_SyncBodiesSingleParticlePhysicsProxy and STAT_SyncBodiesJointConstraintPhysicsProxy verbose.

#rnx

diff --git a/Engine/Source/Runtime/Engine/Private/Components/SceneComponent.cpp b/Engine/Source/Runtime/Engine/Private/Components/SceneComponent.cpp
--- a/Engine/Source/Runtime/Engine/Private/Components/SceneComponent.cpp
+++ b/Engine/Source/Runtime/Engine/Private/Components/SceneComponent.cpp
@@ -73,7 +73,7 @@ DEFINE_LOG_CATEGORY_STATIC(LogSceneComponent, Log, All);
 
 DECLARE_CYCLE_STAT(TEXT("UpdateComponentToWorld"), STAT_UpdateComponentToWorld, STATGROUP_Component);
 DECLARE_CYCLE_STAT(TEXT("UpdateChildTransforms"), STAT_UpdateChildTransforms, STATGROUP_Component);
-DECLARE_CYCLE_STAT(TEXT("Component CalcBounds"), STAT_ComponentCalcBounds, STATGROUP_Component);
+DECLARE_CYCLE_STAT_WITH_FLAGS(TEXT("Component CalcBounds"), STAT_ComponentCalcBounds, STATGROUP_Component, EStatFlags::Verbose);
 DECLARE_CYCLE_STAT(TEXT("Component UpdateNavData"), STAT_ComponentUpdateNavData, STATGROUP_Component);
 DECLARE_CYCLE_STAT(TEXT("Component PostUpdateNavData"), STAT_ComponentPostUpdateNavData, STATGROUP_Component);
 
diff --git a/Engine/Source/Runtime/Engine/Private/PhysicsEngine/BodyInstance.cpp b/Engine/Source/Runtime/Engine/Private/PhysicsEngine/BodyInstance.cpp
--- a/Engine/Source/Runtime/Engine/Private/PhysicsEngine/BodyInstance.cpp
+++ b/Engine/Source/Runtime/Engine/Private/PhysicsEngine/BodyInstance.cpp
@@ -53,7 +53,7 @@ DECLARE_CYCLE_STAT(TEXT("Init Body Post Add to Scene"), STAT_InitBodyPostAdd, ST
 DECLARE_CYCLE_STAT(TEXT("Term Body"), STAT_TermBody, STATGROUP_Physics);
 DECLARE_CYCLE_STAT(TEXT("Update Materials"), STAT_UpdatePhysMats, STATGROUP_Physics);
 DECLARE_CYCLE_STAT(TEXT("Update Materials Scene Interaction"), STAT_UpdatePhysMatsSceneInteraction, STATGROUP_Physics);
-DECLARE_CYCLE_STAT(TEXT("Filter Update"), STAT_UpdatePhysFilter, STATGROUP_Physics);
+DECLARE_CYCLE_STAT_WITH_FLAGS(TEXT("Filter Update"), STAT_UpdatePhysFilter, STATGROUP_Physics, EStatFlags::Verbose);
 DECLARE_CYCLE_STAT(TEXT("Filter Update (PhysX Code)"), STAT_UpdatePhysFilterPhysX, STATGROUP_Physics);
 DECLARE_CYCLE_STAT(TEXT("Init Bodies"), STAT_InitBodies, STATGROUP_Physics);
 DECLARE_CYCLE_STAT(TEXT("Bulk Body Scene Add"), STAT_BulkSceneAdd, STATGROUP_Physics);
@@ -62,7 +62,7 @@ DECLARE_CYCLE_STAT(TEXT("UpdateBodyScale"), STAT_BodyInstanceUpdateBodyScale, ST
 DECLARE_CYCLE_STAT(TEXT("CreatePhysicsShapesAndActors"), STAT_CreatePhysicsShapesAndActors, STATGROUP_Physics);
 DECLARE_CYCLE_STAT(TEXT("CreatePhysicsShapes"), STAT_CreatePhysicsShapes, STATGROUP_Physics);
 DECLARE_CYCLE_STAT(TEXT("CreatePhysicsActor"), STAT_CreatePhysicsActor, STATGROUP_Physics);
-DECLARE_CYCLE_STAT(TEXT("BodyInstance SetCollisionProfileName"), STAT_BodyInst_SetCollisionProfileName, STATGROUP_Physics);
+DECLARE_CYCLE_STAT_WITH_FLAGS(TEXT("BodyInstance SetCollisionProfileName"), STAT_BodyInst_SetCollisionProfileName, STATGROUP_Physics, EStatFlags::Verbose);
 DECLARE_CYCLE_STAT(TEXT("Phys SetBodyTransform"), STAT_SetBodyTransform, STATGROUP_Physics);
 
 // @HACK Guard to better encapsulate game related hacks introduced into UpdatePhysicsFilterData()
diff --git a/Engine/Source/Runtime/Engine/Private/PhysicsEngine/Experimental/PhysScene_Chaos.cpp b/Engine/Source/Runtime/Engine/Private/PhysicsEngine/Experimental/PhysScene_Chaos.cpp
--- a/Engine/Source/Runtime/Engine/Private/PhysicsEngine/Experimental/PhysScene_Chaos.cpp
+++ b/Engine/Source/Runtime/Engine/Private/PhysicsEngine/Experimental/PhysScene_Chaos.cpp
@@ -2222,7 +2222,7 @@ void FPhysScene_Chaos::KillVisualDebugger()
 }
 
 DECLARE_CYCLE_STAT(TEXT("FPhysScene_Chaos::OnSyncBodies-FSingleParticlePhysicsProxy"), STAT_SyncBodiesSingleParticlePhysicsProxy, STATGROUP_Chaos);
-DECLARE_CYCLE_STAT(TEXT("FPhysScene_Chaos::OnSyncBodies-FJointConstraintPhysicsProxy"), STAT_SyncBodiesJointConstraintPhysicsProxy, STATGROUP_Chaos);
+DECLARE_CYCLE_STAT_WITH_FLAGS(TEXT("FPhysScene_Chaos::OnSyncBodies-FJointConstraintPhysicsProxy"), STAT_SyncBodiesJointConstraintPhysicsProxy, STATGROUP_Chaos, EStatFlags::Verbose);
 DECLARE_CYCLE_STAT(TEXT("FPhysScene_Chaos::OnSyncBodies-FGeometryCollectionPhysicsProxy"), STAT_SyncBodiesGeometryCollectionPhysicsProxy, STATGROUP_Chaos);
 DECLARE_CYCLE_STAT(TEXT("FPhysScene_Chaos::OnSyncBodies-FClusterUnionPhysicsProxy"), STAT_SyncBodiesClusterUnionPhysicsProxy, STATGROUP_Chaos);
 
diff --git a/Engine/Source/Runtime/Engine/Private/PrimitiveComponentPhysics.cpp b/Engine/Source/Runtime/Engine/Private/PrimitiveComponentPhysics.cpp
--- a/Engine/Source/Runtime/Engine/Private/PrimitiveComponentPhysics.cpp
+++ b/Engine/Source/Runtime/Engine/Private/PrimitiveComponentPhysics.cpp
@@ -22,7 +22,7 @@
 
 DECLARE_CYCLE_STAT(TEXT("WeldPhysics"), STAT_WeldPhysics, STATGROUP_Physics);
 DECLARE_CYCLE_STAT(TEXT("UnweldPhysics"), STAT_UnweldPhysics, STATGROUP_Physics);
-DECLARE_CYCLE_STAT(TEXT("PrimComp SetCollisionProfileName"), STAT_PrimComp_SetCollisionProfileName, STATGROUP_Physics);
+DECLARE_CYCLE_STAT_WITH_FLAGS(TEXT("PrimComp SetCollisionProfileName"), STAT_PrimComp_SetCollisionProfileName, STATGROUP_Physics, EStatFlags::Verbose);
 
 
 #if !(UE_BUILD_SHIPPING || UE_BUILD_TEST)
@@ -1073,8 +1073,6 @@ void UPrimitiveComponent::WeldTo(USceneComponent* InParent, FName InSocketName /
 
 void UPrimitiveComponent::UnWeldFromParent()
 {
-	SCOPE_CYCLE_COUNTER(STAT_UnweldPhysics);
-
 	FBodyInstance* NewRootBI = GetBodyInstance(NAME_None, false);
 	UWorld* CurrentWorld = GetWorld();
 	if (NewRootBI == NULL || NewRootBI->WeldParent == nullptr || CurrentWorld == nullptr || CurrentWorld->GetPhysicsScene() == nullptr || !IsValidChecked(this) || IsUnreachable())
@@ -1089,6 +1087,8 @@ void UPrimitiveComponent::UnWeldFromParent()
 		return;
 	}
 
+	SCOPE_CYCLE_COUNTER(STAT_UnweldPhysics);
+
 	FName SocketName;
 	UPrimitiveComponent * RootComponent = GetRootWelded(this, GetAttachSocketName(), &SocketName);
 
diff --git a/Engine/Source/Runtime/Experimental/Chaos/Private/Chaos/ShapeInstance.cpp b/Engine/Source/Runtime/Experimental/Chaos/Private/Chaos/ShapeInstance.cpp
--- a/Engine/Source/Runtime/Experimental/Chaos/Private/Chaos/ShapeInstance.cpp
+++ b/Engine/Source/Runtime/Experimental/Chaos/Private/Chaos/ShapeInstance.cpp
@@ -8,7 +8,7 @@ namespace Chaos
 {
 	static_assert(sizeof(FShapeInstanceProxy) <= 192, "FShapeInstanceProxy was optimized to fit into 192b bin of MB3 to prevent excess memory waste");
 
-	DECLARE_CYCLE_STAT(TEXT("UpdateShapesArrayFromGeometryImpl"), STAT_UpdateShapesArrayFromGeometryImpl, STATGROUP_Chaos);
+	DECLARE_CYCLE_STAT_WITH_FLAGS(TEXT("UpdateShapesArrayFromGeometryImpl"), STAT_UpdateShapesArrayFromGeometryImpl, STATGROUP_Chaos, EStatFlags::Verbose);
 	// Create or reuse the shapes in the shapes array and populate them with the Geometry.
 	// If we have a Union it will be unpacked into the ShapesArray.
 	// On the Physics Thread we set bAllowCachedLeafInfo which caches the shapes world space state to optimize collision detection,
