Fix a check for the correct thread when calling FSkinWeightVertexBuffer::ReleaseResources

#rnx

diff --git a/Engine/Source/Runtime/Engine/Private/Rendering/SkinWeightVertexBuffer.cpp b/Engine/Source/Runtime/Engine/Private/Rendering/SkinWeightVertexBuffer.cpp
--- a/Engine/Source/Runtime/Engine/Private/Rendering/SkinWeightVertexBuffer.cpp
+++ b/Engine/Source/Runtime/Engine/Private/Rendering/SkinWeightVertexBuffer.cpp
@@ -798,7 +798,7 @@ void FSkinWeightVertexBuffer::BeginReleaseResources()
 
 void FSkinWeightVertexBuffer::ReleaseResources()
 {
-	check(IsInRenderingThread());
+	check(IsInParallelRenderingThread());
 	LookupVertexBuffer.ReleaseResource();
 	DataVertexBuffer.ReleaseResource();
 }
