Fix invalid condition that was causing anim task evaluation to be triggered on the game thread, on the frame a mesh is spawned

diff --git a/Engine/Source/Runtime/Engine/Private/Components/SkeletalMeshComponent.cpp b/Engine/Source/Runtime/Engine/Private/Components/SkeletalMeshComponent.cpp
--- a/Engine/Source/Runtime/Engine/Private/Components/SkeletalMeshComponent.cpp
+++ b/Engine/Source/Runtime/Engine/Private/Components/SkeletalMeshComponent.cpp
@@ -2580,7 +2580,10 @@ void USkeletalMeshComponent::RefreshBoneTransforms(FActorComponentTickFunction*
 				bShouldTickAnimation = !AnimScriptInstance->GetUpdateCounter().HasEverBeenUpdated();
 			}
 
-			bShouldTickAnimation = bShouldTickAnimation || (ShouldPostUpdatePostProcessInstance() && !PostProcessAnimInstance->GetUpdateCounter().HasEverBeenUpdated());
+			if (!bDisablePostProcessBlueprint && PostProcessAnimInstance && ShouldEvaluatePostProcessAnimBP() && !PostProcessAnimInstance->NeedsUpdate())
+			{
+				bShouldTickAnimation = bShouldTickAnimation || !PostProcessAnimInstance->GetUpdateCounter().HasEverBeenUpdated();
+			}
 
 			if (bShouldTickAnimation)
 			{
