Skeletal Mesh: Put a mutex around USkeletalMesh::InitMorphTargets to speculatively prevent threading issues until we move morph target remapping to a safer place.

#rnx

diff --git a/Engine/Source/Runtime/Engine/Private/SkeletalMesh.cpp b/Engine/Source/Runtime/Engine/Private/SkeletalMesh.cpp
--- a/Engine/Source/Runtime/Engine/Private/SkeletalMesh.cpp
+++ b/Engine/Source/Runtime/Engine/Private/SkeletalMesh.cpp
@@ -4373,6 +4373,10 @@ void USkeletalMesh::UnregisterMorphTarget(UMorphTarget* MorphTarget, bool bInval
 
 void USkeletalMesh::InitMorphTargets(bool bInKeepEmptyMorphTargets)
 {
+	// FIXME: Big hammer fix to prevent reduction + morph target mapping crashing when running on multiple threads. 
+	static FCriticalSection InitMorphTargetMutex;
+	FScopeLock ScopeLock(&InitMorphTargetMutex);
+
 	TRACE_CPUPROFILER_EVENT_SCOPE(USkeletalMesh::InitMorphTargets);
 	GetMorphTargetIndexMap().Empty();
 
