Interchange - FBX Export/Import - No Timecode information off of the root bone
#fix replace GuardValue with custom behavior to prevent resetting bPopulated flag to false

diff --git a/Engine/Source/Runtime/Engine/Classes/Animation/AnimData/IAnimationDataModel.h b/Engine/Source/Runtime/Engine/Classes/Animation/AnimData/IAnimationDataModel.h
--- a/Engine/Source/Runtime/Engine/Classes/Animation/AnimData/IAnimationDataModel.h
+++ b/Engine/Source/Runtime/Engine/Classes/Animation/AnimData/IAnimationDataModel.h
@@ -132,11 +132,26 @@ public:
 	// RAII helper to allow repopulating data model without warnings during reimporting
 	struct FReimportScope
 	{
-		FReimportScope(IAnimationDataModel* InModel) : PopulationFlagGuard(InModel->GetPopulationFlag(), false)
+		FReimportScope(IAnimationDataModel* InModel) : bPopulationFlagRef(InModel->GetPopulationFlag())
 		{
+			bOriginalValue = bPopulationFlagRef;
+			if (bOriginalValue)
+			{
+				bPopulationFlagRef = false;
+			}
+		}
+
+		~FReimportScope()
+		{
+			// Only reset back to true if that was the original state and currently is false			
+			if (!bPopulationFlagRef && bOriginalValue)
+			{
+				bPopulationFlagRef = bOriginalValue;
+			}
 		}
 	private:
-		TGuardValue<bool> PopulationFlagGuard;
+		bool& bPopulationFlagRef;
+		bool bOriginalValue;
 	};
 #endif // WITH_EDITOR
 
