Support Mesh Paint Texture Object being passed as an input to Material Functions.
The main part of the fix is to have the material expression declare it's output pin as MCT_Texture2D and support casting from MCT_TextureMeshPaint to MCT_Texture2D. This way we can pass MCT_TextureMeshPaint through a standard MCT_Texture2D input to a material function. This is the same approach that MCT_TextureVirtual already takes.
The other change is to not express the type as a combination of MCT_TextureVirtual and MCT_TextureMeshPaint. That was a very bad idea. Instead we simply express it as MCT_TextureMeshPaint only. We then need to recognize this type as a virtual texture at the (one) necessary point during material translation.

diff --git a/Engine/Source/Runtime/Engine/Private/Materials/HLSLMaterialTranslator.cpp b/Engine/Source/Runtime/Engine/Private/Materials/HLSLMaterialTranslator.cpp
--- a/Engine/Source/Runtime/Engine/Private/Materials/HLSLMaterialTranslator.cpp
+++ b/Engine/Source/Runtime/Engine/Private/Materials/HLSLMaterialTranslator.cpp
@@ -3481,6 +3481,7 @@ const TCHAR* FHLSLMaterialTranslator::DescribeType(EMaterialValueType Type) cons
 	case MCT_LWCVector3:			return TEXT("LWCVector3");
 	case MCT_LWCVector4:			return TEXT("LWCVector4");
 	case MCT_TextureCollection:		return TEXT("TextureCollection");
+	case MCT_TextureMeshPaint:		return TEXT("TextureMeshPaint");
 	default:						return TEXT("unknown");
 	};
 }
@@ -3519,6 +3520,7 @@ const TCHAR* FHLSLMaterialTranslator::HLSLTypeString(EMaterialValueType Type) co
 	case MCT_LWCVector3:			return TEXT("FWSVector3");
 	case MCT_LWCVector4:			return TEXT("FWSVector4");
 	case MCT_TextureCollection:		return TEXT("FResourceCollection");
+	case MCT_TextureMeshPaint:		return TEXT("TextureMeshPaint");
 	default:						return TEXT("unknown");
 	};
 }
@@ -3557,6 +3559,7 @@ const TCHAR* FHLSLMaterialTranslator::HLSLTypeStringDeriv(EMaterialValueType Typ
 	case MCT_LWCVector3:			return (DerivativeStatus == EDerivativeStatus::Valid) ? TEXT("FWSVector3Deriv") : TEXT("FWSVector3");
 	case MCT_LWCVector4:			return (DerivativeStatus == EDerivativeStatus::Valid) ? TEXT("FWSVector4Deriv") : TEXT("FWSVector4");
 	case MCT_TextureCollection:		return TEXT("FResourceCollection");
+	case MCT_TextureMeshPaint:		return TEXT("TextureMeshPaint");
 	default:						return TEXT("unknown");
 	};
 }
@@ -4786,6 +4789,10 @@ int32 FHLSLMaterialTranslator::ValidCast(int32 Code, EMaterialValueType DestType
 			return ValidCast(AccessUniformExpression(Code), DestType);
 		}
 	}
+	else if ((SourceType & MCT_TextureMeshPaint) && (DestType & MCT_Texture2D))
+	{
+		return Code;
+	}
 	else if(IsFloatNumericType(SourceType) && IsFloatNumericType(DestType))
 	{
 		const FDerivInfo CodeDerivInfo = GetDerivInfo(Code);
@@ -4897,7 +4904,7 @@ int32 FHLSLMaterialTranslator::ForceCast(int32 Code, EMaterialValueType DestType
 		StaticBoolToFloat = CastValue(StaticBoolToFloat, MCT_Float, DestType, GetForceCastFlags(ForceCastFlags));
 		return AddInlinedCodeChunk(DestType, *StaticBoolToFloat);
 	}
-	else if ((SourceType & MCT_TextureVirtual) && (DestType & MCT_Texture2D))
+	else if ((SourceType & (MCT_TextureVirtual|MCT_TextureMeshPaint)) && (DestType & MCT_Texture2D))
 	{
 		return Code;
 	}
@@ -7147,16 +7154,14 @@ int32 FHLSLMaterialTranslator::TextureSample(
 	const bool bFromCollection = (TextureType & MCT_TextureCollection) != 0;
 	TextureType = EMaterialValueType(TextureType & ~MCT_TextureCollection);
 
-	const bool bMeshPaint = (TextureType & MCT_TextureMeshPaint) != 0;
-	TextureType = EMaterialValueType(TextureType & ~MCT_TextureMeshPaint);
-
 	if (ShaderFrequency != SF_Pixel && MipValueMode == TMVM_MipBias)
 	{
 		Errorf(TEXT("MipBias is only supported in the pixel shader"));
 		return INDEX_NONE;
 	}
 
-	const bool bVirtualTexture = TextureType == MCT_TextureVirtual;
+	const bool bMeshPaint = TextureType == MCT_TextureMeshPaint;
+	const bool bVirtualTexture = bMeshPaint || TextureType == MCT_TextureVirtual;
 	if (bVirtualTexture)
 	{
 		if (Material->GetMaterialDomain() != MD_Surface && Material->GetMaterialDomain() != MD_DeferredDecal)
@@ -7217,7 +7222,7 @@ int32 FHLSLMaterialTranslator::TextureSample(
 	}
 
 	// If not 2D texture, disable AutomaticViewMipBias.
-	if (!(TextureType & (MCT_Texture2D | MCT_TextureVirtual)))
+	if (!(TextureType & (MCT_Texture2D | MCT_TextureVirtual | MCT_TextureMeshPaint)))
 	{
 		AutomaticViewMipBias = false;
 	}
@@ -8925,7 +8930,7 @@ int32 FHLSLMaterialTranslator::MeshPaintTextureCoordinateIndex()
 
 int32 FHLSLMaterialTranslator::MeshPaintTextureDescriptor()
 {
-	return AddInlinedCodeChunkZeroDeriv(EMaterialValueType(MCT_TextureVirtual | MCT_TextureMeshPaint), TEXT("GetMeshPaintTextureDescriptor(GetPrimitiveData(Parameters))"));
+	return AddInlinedCodeChunkZeroDeriv(MCT_TextureMeshPaint, TEXT("GetMeshPaintTextureDescriptor(GetPrimitiveData(Parameters))"));
 }
 
 int32 FHLSLMaterialTranslator::MeshPaintTextureReplace(int32 Invalid, int32 Valid)
diff --git a/Engine/Source/Runtime/Engine/Private/Materials/MaterialExpressions.cpp b/Engine/Source/Runtime/Engine/Private/Materials/MaterialExpressions.cpp
--- a/Engine/Source/Runtime/Engine/Private/Materials/MaterialExpressions.cpp
+++ b/Engine/Source/Runtime/Engine/Private/Materials/MaterialExpressions.cpp
@@ -4007,7 +4007,7 @@ void UMaterialExpressionMeshPaintTextureObject::GetExpressionToolTip(TArray<FStr
 
 uint32 UMaterialExpressionMeshPaintTextureObject::GetOutputType(int32 OutputIndex)
 {
-	return MCT_TextureVirtual | MCT_TextureMeshPaint;
+	return MCT_Texture2D;
 }
 
 int32 UMaterialExpressionMeshPaintTextureObject::Compile(class FMaterialCompiler* Compiler, int32 OutputIndex)
diff --git a/Engine/Source/Runtime/Engine/Public/MaterialValueType.h b/Engine/Source/Runtime/Engine/Public/MaterialValueType.h
--- a/Engine/Source/Runtime/Engine/Public/MaterialValueType.h
+++ b/Engine/Source/Runtime/Engine/Public/MaterialValueType.h
@@ -66,7 +66,7 @@ enum EMaterialValueType : int
 	MCT_TextureMeshPaint      = 1 << 30,
 
 	/** MCT_SparseVolumeTexture is intentionally not (yet) included here because it differs a lot from the other texture types and may not be supported/appropriate for all MCT_Texture use cases. */
-	MCT_Texture = MCT_Texture2D | MCT_TextureCube | MCT_Texture2DArray | MCT_TextureCubeArray | MCT_VolumeTexture | MCT_TextureExternal | MCT_TextureVirtual,
+	MCT_Texture = MCT_Texture2D | MCT_TextureCube | MCT_Texture2DArray | MCT_TextureCubeArray | MCT_VolumeTexture | MCT_TextureExternal | MCT_TextureVirtual | MCT_TextureMeshPaint,
 
 	MCT_Float   = MCT_Float1 | MCT_Float2 | MCT_Float3 | MCT_Float4,
 	MCT_UInt    = MCT_UInt1 | MCT_UInt2 | MCT_UInt3 | MCT_UInt4,
