Add full material function stack when reporting a material compilation error.
Previously we only showed the first function in the stack which isn't helpful for finding deeply nested errors.

diff --git a/Engine/Source/Runtime/Engine/Private/Materials/HLSLMaterialTranslator.cpp b/Engine/Source/Runtime/Engine/Private/Materials/HLSLMaterialTranslator.cpp
--- a/Engine/Source/Runtime/Engine/Private/Materials/HLSLMaterialTranslator.cpp
+++ b/Engine/Source/Runtime/Engine/Private/Materials/HLSLMaterialTranslator.cpp
@@ -4465,13 +4465,29 @@ int32 FHLSLMaterialTranslator::Error(const TCHAR* Text)
 	auto& CurrentFunctionStack = FunctionStacks[ShaderFrequency];
 	if (CurrentFunctionStack.Num() > 1)
 	{
-		// If we are inside a function, add that to the error message.  
 		// Only add the function call node to ErrorExpressions, since we can't add a reference to the expressions inside the function as they are private objects.
 		// Add the first function node on the stack because that's the one visible in the material being compiled, the rest are all nested functions.
-		UMaterialExpressionMaterialFunctionCall* ErrorFunction = CurrentFunctionStack[1]->FunctionCall;
-		check(ErrorFunction);
-		ExpressionToError = ErrorFunction;			
-		ErrorString = FString(TEXT("Function ")) + ErrorFunction->MaterialFunction->GetName() + TEXT(": ");
+		ExpressionToError = CurrentFunctionStack[1]->FunctionCall;
+		check(ExpressionToError);
+
+		// Build full callstack for error message.
+		ErrorString = FString(TEXT("(Function "));
+		for (int32 StackIndex = 1; StackIndex < CurrentFunctionStack.Num(); ++StackIndex)
+		{
+			FMaterialFunctionCompileState const* Function = CurrentFunctionStack[StackIndex];
+			if (Function == nullptr || Function->FunctionCall == nullptr || Function->FunctionCall->MaterialFunction == nullptr)
+			{
+				continue;
+			}
+
+			ErrorString += Function->FunctionCall->MaterialFunction->GetName();
+
+			if (StackIndex < CurrentFunctionStack.Num() - 1)
+			{
+				ErrorString += TEXT("|");
+			}
+		}
+		ErrorString += FString(TEXT(") "));
 	}
 
 	if (CurrentFunctionStack.Last()->ExpressionStack.Num() > 0)
