Change implemetation of VirtualTexture replace node so that it only compiles the desired path instead of compiling both input paths selecting one.
This is a precondition for later changes where we want to validate the virtual texture output path only.
We add a Push/Pop compiler state to set when we are in the VirtualTextureOutput.
We rely on the fact that Output Nodes have their own Chunk Scope in the translator so there is no way to compile the Replace node in two different contexts.

diff --git a/Engine/Plugins/Enterprise/GLTFExporter/Source/GLTFExporter/Private/Converters/GLTFProxyMaterialCompiler.h b/Engine/Plugins/Enterprise/GLTFExporter/Source/GLTFExporter/Private/Converters/GLTFProxyMaterialCompiler.h
--- a/Engine/Plugins/Enterprise/GLTFExporter/Source/GLTFExporter/Private/Converters/GLTFProxyMaterialCompiler.h
+++ b/Engine/Plugins/Enterprise/GLTFExporter/Source/GLTFExporter/Private/Converters/GLTFProxyMaterialCompiler.h
@@ -124,11 +124,6 @@ public:
 		return Normal; // ignore all non-default branches in this runtime switch
 	}
 
-	virtual int32 VirtualTextureOutputReplace(int32 Default, int32 VirtualTexture) override
-	{
-		return Default; // ignore all non-default branches in this runtime switch
-	}
-
 	virtual int32 ReflectionCapturePassSwitch(int32 Default, int32 Reflection) override
 	{
 		return Default; // ignore all non-default branches in this runtime switch
diff --git a/Engine/Shaders/Private/Common.ush b/Engine/Shaders/Private/Common.ush
--- a/Engine/Shaders/Private/Common.ush
+++ b/Engine/Shaders/Private/Common.ush
@@ -2064,17 +2064,6 @@ bool GetLightmassReplaceState()
 #endif
 }
 
-// Used for the RuntimeVirtualTextureReplace material expression
-// Returns true if rendering to a runtime virtual texture page
-bool GetRuntimeVirtualTextureOutputSwitch()
-{
-#if VIRTUAL_TEXTURE_PAGE_RENDER
-	return true;
-#else
-	return false;
-#endif
-}
-
 struct FWriteToSliceGeometryOutput
 {
 	FScreenVertexOutput Vertex;
diff --git a/Engine/Source/Runtime/Core/Private/UObject/FortniteMainBranchObjectVersion.cpp b/Engine/Source/Runtime/Core/Private/UObject/FortniteMainBranchObjectVersion.cpp
--- a/Engine/Source/Runtime/Core/Private/UObject/FortniteMainBranchObjectVersion.cpp
+++ b/Engine/Source/Runtime/Core/Private/UObject/FortniteMainBranchObjectVersion.cpp
@@ -14,6 +14,6 @@ TMap<FGuid, FGuid> FFortniteMainBranchObjectVersion::GetSystemGuids()
 	SystemGuids.Add(DevGuids.Niagara_LatestScriptCompileVersion, FGuid("393081597FA7234C92F00869DE262D73"));
 	SystemGuids.Add(DevGuids.SkeletalMeshDerivedDataVersion, FGuid("0829C9BAB43B4B25A63747FDA68E7197"));
 	SystemGuids.Add(DevGuids.STATICMESH_DERIVEDDATA_VER, FGuid("68C7D98BE48042E0A35351325741F70E"));
-	SystemGuids.Add(DevGuids.MaterialTranslationDDCVersion, FGuid("618C475161AF47B58FC057E13AD67D53"));
+	SystemGuids.Add(DevGuids.MaterialTranslationDDCVersion, FGuid("8CF065BF76124292A1A0546D532BEBDF"));
 	return SystemGuids;
 }
diff --git a/Engine/Source/Runtime/Engine/Private/Materials/HLSLMaterialTranslator.cpp b/Engine/Source/Runtime/Engine/Private/Materials/HLSLMaterialTranslator.cpp
--- a/Engine/Source/Runtime/Engine/Private/Materials/HLSLMaterialTranslator.cpp
+++ b/Engine/Source/Runtime/Engine/Private/Materials/HLSLMaterialTranslator.cpp
@@ -738,6 +738,7 @@ FHLSLMaterialTranslator::FHLSLMaterialTranslator(FMaterial* InMaterial,
 ,	bUsesPerInstanceFadeAmount(false)
 ,	bCullIntermediateUniformExpressions(GCullIntermediateUniformExpressions)
 ,	bUsesExplicitDerivatives(false)
+,	bIsInRuntimeVirtualTextureOutput(false)
 ,	AddingUniformExpression(0)
 ,	AllocatedUserTexCoords()
 ,	AllocatedUserVertexTexCoords()
@@ -11459,11 +11460,6 @@ int32 FHLSLMaterialTranslator::LightmassReplace(int32 Realtime, int32 Lightmass)
 	return GenericSwitch(TEXT("GetLightmassReplaceState()"), Lightmass, Realtime);
 }
 
-int32 FHLSLMaterialTranslator::VirtualTextureOutputReplace(int32 Default, int32 VirtualTexture)
-{
-	return GenericSwitch(TEXT("GetRuntimeVirtualTextureOutputSwitch()"), VirtualTexture, Default);
-}
-
 int32 FHLSLMaterialTranslator::ObjectOrientation()
 { 
 	return AddInlinedCodeChunkZeroDeriv(MCT_Float3,TEXT("GetObjectOrientation(Parameters)"));
@@ -15557,15 +15553,27 @@ int32 FHLSLMaterialTranslator::CustomOutput(class UMaterialExpressionCustomOutpu
 	return CustomOutput(Custom->GetDescription(), Custom->GetFunctionName(), OutputIndex, OutputCode);
 }
 
-int32 FHLSLMaterialTranslator::VirtualTextureOutput(uint8 AttributeMask)
+int32 FHLSLMaterialTranslator::PushRuntimeVirtualTextureOutput()
+{
+	bIsInRuntimeVirtualTextureOutput = true;
+	// return value is not used
+	return INDEX_NONE;
+}
+
+int32 FHLSLMaterialTranslator::PopRuntimeVirtualTextureOutput(uint8 AttributeMask)
 {
+	bIsInRuntimeVirtualTextureOutput = false;
 	MaterialCompilationOutput.bHasRuntimeVirtualTextureOutputNode |= AttributeMask != 0;
 	MaterialCompilationOutput.RuntimeVirtualTextureOutputAttributeMask |= AttributeMask;
-
 	// return value is not used
 	return INDEX_NONE;
 }
 
+bool FHLSLMaterialTranslator::IsInRuntimeVirtualTextureOutput() const
+{
+	return bIsInRuntimeVirtualTextureOutput;
+}
+
 /**
 	* Adds code to return a random value shared by all geometry for any given instanced static mesh
 	*
diff --git a/Engine/Source/Runtime/Engine/Private/Materials/HLSLMaterialTranslator.h b/Engine/Source/Runtime/Engine/Private/Materials/HLSLMaterialTranslator.h
--- a/Engine/Source/Runtime/Engine/Private/Materials/HLSLMaterialTranslator.h
+++ b/Engine/Source/Runtime/Engine/Private/Materials/HLSLMaterialTranslator.h
@@ -449,6 +449,9 @@ protected:
 	/** The generated shader source uses explicit ddx()/ddy() calls */
 	uint32 bUsesExplicitDerivatives : 1;
 
+	/** Flag that tracks if we are compiling a runtime virtual texture output subgraph. */
+	uint32 bIsInRuntimeVirtualTextureOutput : 1;
+
 	/** Incremented and decremented by FAddUniformExpressionScope.  See comments on that class below */
 	int32 AddingUniformExpression;
 
@@ -1193,8 +1196,6 @@ protected:
 	virtual int32 PathTracingRayTypeSwitch(int32 Main, int32 Shadow, int32 IndirectDiffuse, int32 IndirectSpecular, int32 IndirectVolume) override;
 	virtual int32 LightmassReplace(int32 Realtime, int32 Lightmass) override;
 
-	virtual int32 VirtualTextureOutputReplace(int32 Default, int32 VirtualTexture) override;
-
 	virtual int32 ObjectOrientation() override;
 
 	virtual int32 RotateAboutAxis(int32 NormalizedRotationAxisAndAngleIndex, int32 PositionOnAxisIndex, int32 PositionIndex) override;
@@ -1280,7 +1281,9 @@ protected:
 	virtual int32 CustomOutput(class UMaterialExpressionCustomOutput* Custom, int32 OutputIndex, int32 OutputCode) override;
 	int32 CustomOutput(const FString& CustomOutputName, const FString& FunctionName, int32 OutputIndex, int32 OutputCode); // Only used internally and intentionally not exposed on FMaterialCompiler.
 
-	virtual int32 VirtualTextureOutput(uint8 MaterialAttributeMask) override;
+	virtual int32 PushRuntimeVirtualTextureOutput() override;
+	virtual int32 PopRuntimeVirtualTextureOutput(uint8 MaterialAttributeMask) override;
+	virtual bool IsInRuntimeVirtualTextureOutput() const override;
 
 	// Material attributes
 	virtual int32 DefaultMaterialAttributes() override;
diff --git a/Engine/Source/Runtime/Engine/Private/Materials/MaterialExpressions.cpp b/Engine/Source/Runtime/Engine/Private/Materials/MaterialExpressions.cpp
--- a/Engine/Source/Runtime/Engine/Private/Materials/MaterialExpressions.cpp
+++ b/Engine/Source/Runtime/Engine/Private/Materials/MaterialExpressions.cpp
@@ -3092,67 +3092,74 @@ bool ValidateRuntimeVirtualTextureOutput(FMaterialCompiler* Compiler, FExpressio
 
 int32 UMaterialExpressionRuntimeVirtualTextureOutput::Compile(class FMaterialCompiler* Compiler, int32 OutputIndex)
 {
-	int32 CodeInput = INDEX_NONE;
+	int32 Result = INDEX_NONE;
 	uint8 OutputAttributeMask = 0;
 
-	// Order of outputs generates function names GetVirtualTextureOutput{index}
-	// These must match the function names called in VirtualTextureMaterial.usf
-	if (OutputIndex == 0)
-	{
-		CodeInput = BaseColor.IsConnected() ? BaseColor.Compile(Compiler) : Compiler->Constant3(0.f, 0.f, 0.f);
-		OutputAttributeMask |= BaseColor.IsConnected() ? (1 << (uint8)ERuntimeVirtualTextureAttributeType::BaseColor) : 0;
-		ValidateRuntimeVirtualTextureOutput(Compiler, BaseColor);
-	}
-	else if (OutputIndex == 1)
-	{
-		CodeInput = Specular.IsConnected() ? Specular.Compile(Compiler) : Compiler->Constant(0.5f);
-		OutputAttributeMask |= Specular.IsConnected() ? (1 << (uint8)ERuntimeVirtualTextureAttributeType::Specular) : 0;
-		ValidateRuntimeVirtualTextureOutput(Compiler, Specular);
-	}
-	else if (OutputIndex == 2)
-	{
-		CodeInput = Roughness.IsConnected() ? Roughness.Compile(Compiler) : Compiler->Constant(0.5f);
-		OutputAttributeMask |= Roughness.IsConnected() ? (1 << (uint8)ERuntimeVirtualTextureAttributeType::Roughness) : 0;
-		ValidateRuntimeVirtualTextureOutput(Compiler, Roughness);
-	}
-	else if (OutputIndex == 3)
-	{
-		CodeInput = Normal.IsConnected() ? Normal.Compile(Compiler) : Compiler->Constant3(0.f, 0.f, 1.f);
-		OutputAttributeMask |= Normal.IsConnected() ? (1 << (uint8)ERuntimeVirtualTextureAttributeType::Normal) : 0;
-		ValidateRuntimeVirtualTextureOutput(Compiler, Normal);
-	}
-	else if (OutputIndex == 4)
-	{
-		CodeInput = WorldHeight.IsConnected() ? WorldHeight.Compile(Compiler) : Compiler->Constant(0.f);
-		OutputAttributeMask |= WorldHeight.IsConnected() ? (1 << (uint8)ERuntimeVirtualTextureAttributeType::WorldHeight) : 0;
-		ValidateRuntimeVirtualTextureOutput(Compiler, WorldHeight);
-	}
-	else if (OutputIndex == 5)
-	{
-		CodeInput = Opacity.IsConnected() ? Opacity.Compile(Compiler) : Compiler->Constant(1.f);
-		ValidateRuntimeVirtualTextureOutput(Compiler, Opacity);
-	}
-	else if (OutputIndex == 6)
-	{
-		CodeInput = Mask.IsConnected() ? Mask.Compile(Compiler) : Compiler->Constant(1.f);
-		OutputAttributeMask |= Mask.IsConnected() ? (1 << (uint8)ERuntimeVirtualTextureAttributeType::Mask) : 0;
-		ValidateRuntimeVirtualTextureOutput(Compiler, Mask);
-	}
-	else if (OutputIndex == 7)
-	{
-		CodeInput = Displacement.IsConnected() ? Displacement.Compile(Compiler) : Compiler->Constant(0.f);
-		OutputAttributeMask |= Displacement.IsConnected() ? (1 << (uint8)ERuntimeVirtualTextureAttributeType::Displacement) : 0;
-		ValidateRuntimeVirtualTextureOutput(Compiler, Displacement);
-	}
-	else if (OutputIndex == 8)
+	Compiler->PushRuntimeVirtualTextureOutput();
 	{
-		CodeInput = Mask4.IsConnected() ? Mask4.Compile(Compiler) : Compiler->Constant4(0.f, 0.f, 0.f, 0.f);
-		OutputAttributeMask |= Mask4.IsConnected() ? (1 << (uint8)ERuntimeVirtualTextureAttributeType::Mask4) : 0;
-		ValidateRuntimeVirtualTextureOutput(Compiler, Mask4);
+		int32 CodeInput = INDEX_NONE;
+
+		// Order of outputs generates function names GetVirtualTextureOutput{index}
+		// These must match the function names called in VirtualTextureMaterial.usf
+		if (OutputIndex == 0)
+		{
+			CodeInput = BaseColor.IsConnected() ? BaseColor.Compile(Compiler) : Compiler->Constant3(0.f, 0.f, 0.f);
+			OutputAttributeMask |= BaseColor.IsConnected() ? (1 << (uint8)ERuntimeVirtualTextureAttributeType::BaseColor) : 0;
+			ValidateRuntimeVirtualTextureOutput(Compiler, BaseColor);
+		}
+		else if (OutputIndex == 1)
+		{
+			CodeInput = Specular.IsConnected() ? Specular.Compile(Compiler) : Compiler->Constant(0.5f);
+			OutputAttributeMask |= Specular.IsConnected() ? (1 << (uint8)ERuntimeVirtualTextureAttributeType::Specular) : 0;
+			ValidateRuntimeVirtualTextureOutput(Compiler, Specular);
+		}
+		else if (OutputIndex == 2)
+		{
+			CodeInput = Roughness.IsConnected() ? Roughness.Compile(Compiler) : Compiler->Constant(0.5f);
+			OutputAttributeMask |= Roughness.IsConnected() ? (1 << (uint8)ERuntimeVirtualTextureAttributeType::Roughness) : 0;
+			ValidateRuntimeVirtualTextureOutput(Compiler, Roughness);
+		}
+		else if (OutputIndex == 3)
+		{
+			CodeInput = Normal.IsConnected() ? Normal.Compile(Compiler) : Compiler->Constant3(0.f, 0.f, 1.f);
+			OutputAttributeMask |= Normal.IsConnected() ? (1 << (uint8)ERuntimeVirtualTextureAttributeType::Normal) : 0;
+			ValidateRuntimeVirtualTextureOutput(Compiler, Normal);
+		}
+		else if (OutputIndex == 4)
+		{
+			CodeInput = WorldHeight.IsConnected() ? WorldHeight.Compile(Compiler) : Compiler->Constant(0.f);
+			OutputAttributeMask |= WorldHeight.IsConnected() ? (1 << (uint8)ERuntimeVirtualTextureAttributeType::WorldHeight) : 0;
+			ValidateRuntimeVirtualTextureOutput(Compiler, WorldHeight);
+		}
+		else if (OutputIndex == 5)
+		{
+			CodeInput = Opacity.IsConnected() ? Opacity.Compile(Compiler) : Compiler->Constant(1.f);
+			ValidateRuntimeVirtualTextureOutput(Compiler, Opacity);
+		}
+		else if (OutputIndex == 6)
+		{
+			CodeInput = Mask.IsConnected() ? Mask.Compile(Compiler) : Compiler->Constant(1.f);
+			OutputAttributeMask |= Mask.IsConnected() ? (1 << (uint8)ERuntimeVirtualTextureAttributeType::Mask) : 0;
+			ValidateRuntimeVirtualTextureOutput(Compiler, Mask);
+		}
+		else if (OutputIndex == 7)
+		{
+			CodeInput = Displacement.IsConnected() ? Displacement.Compile(Compiler) : Compiler->Constant(0.f);
+			OutputAttributeMask |= Displacement.IsConnected() ? (1 << (uint8)ERuntimeVirtualTextureAttributeType::Displacement) : 0;
+			ValidateRuntimeVirtualTextureOutput(Compiler, Displacement);
+		}
+		else if (OutputIndex == 8)
+		{
+			CodeInput = Mask4.IsConnected() ? Mask4.Compile(Compiler) : Compiler->Constant4(0.f, 0.f, 0.f, 0.f);
+			OutputAttributeMask |= Mask4.IsConnected() ? (1 << (uint8)ERuntimeVirtualTextureAttributeType::Mask4) : 0;
+			ValidateRuntimeVirtualTextureOutput(Compiler, Mask4);
+		}
+
+		Result = Compiler->CustomOutput(this, OutputIndex, CodeInput);
 	}
+	Compiler->PopRuntimeVirtualTextureOutput(OutputAttributeMask);
 
-	Compiler->VirtualTextureOutput(OutputAttributeMask);
-	return Compiler->CustomOutput(this, OutputIndex, CodeInput);
+	return Result;
 }
 
 void UMaterialExpressionRuntimeVirtualTextureOutput::GetCaption(TArray<FString>& OutCaptions) const
@@ -3846,10 +3853,7 @@ int32 UMaterialExpressionRuntimeVirtualTextureReplace::Compile(class FMaterialCo
 		return Compiler->Errorf(TEXT("Missing RuntimeVirtualTextureReplace input 'VirtualTextureOutput'"));
 	}
 
-	int32 Arg1 = Default.Compile(Compiler);
-	int32 Arg2 = VirtualTextureOutput.Compile(Compiler);
-	return Compiler->VirtualTextureOutputReplace(Arg1, Arg2);
-}
+	return Compiler->IsInRuntimeVirtualTextureOutput() ? VirtualTextureOutput.Compile(Compiler) : Default.Compile(Compiler);}
 
 bool UMaterialExpressionRuntimeVirtualTextureReplace::IsResultMaterialAttributes(int32 OutputIndex)
 {
diff --git a/Engine/Source/Runtime/Engine/Public/MaterialCompiler.h b/Engine/Source/Runtime/Engine/Public/MaterialCompiler.h
--- a/Engine/Source/Runtime/Engine/Public/MaterialCompiler.h
+++ b/Engine/Source/Runtime/Engine/Public/MaterialCompiler.h
@@ -550,7 +550,6 @@ public:
 	virtual int32 PathTracingQualitySwitchReplace(int32 Normal, int32 PathTraced) = 0;
 	virtual int32 PathTracingRayTypeSwitch(int32 Main, int32 Shadow, int32 IndirectDiffuse, int32 IndirectSpecular, int32 IndirectVolume) = 0;
 	virtual int32 LightmassReplace(int32 Realtime, int32 Lightmass) = 0;
-	virtual int32 VirtualTextureOutputReplace(int32 Default, int32 VirtualTexture) = 0;
 	virtual int32 ReflectionCapturePassSwitch(int32 Default, int32 Reflection) = 0;
 
 	virtual int32 ObjectOrientation() = 0;
@@ -563,7 +562,10 @@ public:
 
 	virtual int32 CustomExpression(class UMaterialExpressionCustom* Custom, int32 OutputIndex, TArray<int32>& CompiledInputs) = 0;
 	virtual int32 CustomOutput(class UMaterialExpressionCustomOutput* Custom, int32 OutputIndex, int32 OutputCode) = 0;
-	virtual int32 VirtualTextureOutput(uint8 AttributeMask) = 0;
+	
+	virtual int32 PushRuntimeVirtualTextureOutput() = 0;
+	virtual int32 PopRuntimeVirtualTextureOutput(uint8 AttributeMask) = 0;
+	virtual bool IsInRuntimeVirtualTextureOutput() const = 0;
 
 	virtual int32 DDX(int32 X) = 0;
 	virtual int32 DDY(int32 X) = 0;
@@ -1110,7 +1112,6 @@ public:
 	virtual int32 PathTracingQualitySwitchReplace(int32 Normal, int32 PathTraced) override { return Compiler->PathTracingQualitySwitchReplace(Normal, PathTraced); }
 	virtual int32 PathTracingRayTypeSwitch(int32 Main, int32 Shadow, int32 IndirectDiffuse, int32 IndirectSpecular, int32 IndirectVolume) override { return Compiler->PathTracingRayTypeSwitch(Main, Shadow, IndirectDiffuse, IndirectSpecular, IndirectVolume); }
 	virtual int32 LightmassReplace(int32 Realtime, int32 Lightmass) override { return Compiler->LightmassReplace(Realtime, Lightmass); }
-	virtual int32 VirtualTextureOutputReplace(int32 Default, int32 VirtualTexture) override { return Compiler->VirtualTextureOutputReplace(Default, VirtualTexture); }
 	virtual int32 ReflectionCapturePassSwitch(int32 Default, int32 Reflection) override { return Compiler->ReflectionCapturePassSwitch(Default, Reflection); }
 
 	virtual int32 ObjectOrientation() override { return Compiler->ObjectOrientation(); }
@@ -1126,7 +1127,10 @@ public:
 
 	virtual int32 CustomExpression(class UMaterialExpressionCustom* Custom, int32 OutputIndex, TArray<int32>& CompiledInputs) override { return Compiler->CustomExpression(Custom, OutputIndex, CompiledInputs); }
 	virtual int32 CustomOutput(class UMaterialExpressionCustomOutput* Custom, int32 OutputIndex, int32 OutputCode) override { return Compiler->CustomOutput(Custom, OutputIndex, OutputCode); }
-	virtual int32 VirtualTextureOutput(uint8 AttributeMask) override { return Compiler->VirtualTextureOutput(AttributeMask); }
+	
+	virtual int32 PushRuntimeVirtualTextureOutput() override { return Compiler->PushRuntimeVirtualTextureOutput(); }
+	virtual int32 PopRuntimeVirtualTextureOutput(uint8 AttributeMask) override { return Compiler->PopRuntimeVirtualTextureOutput(AttributeMask); }
+	virtual bool IsInRuntimeVirtualTextureOutput() const override { return Compiler->IsInRuntimeVirtualTextureOutput(); }
 
 	virtual int32 DDX(int32 X) override { return Compiler->DDX(X); }
 	virtual int32 DDY(int32 X) override { return Compiler->DDY(X); }
diff --git a/Engine/Source/Runtime/Renderer/Private/VT/RuntimeVirtualTextureRender.cpp b/Engine/Source/Runtime/Renderer/Private/VT/RuntimeVirtualTextureRender.cpp
--- a/Engine/Source/Runtime/Renderer/Private/VT/RuntimeVirtualTextureRender.cpp
+++ b/Engine/Source/Runtime/Renderer/Private/VT/RuntimeVirtualTextureRender.cpp
@@ -153,7 +153,6 @@ namespace RuntimeVirtualTexture
 		static void ModifyCompilationEnvironment(const FMaterialShaderPermutationParameters& Parameters, FShaderCompilerEnvironment& OutEnvironment)
 		{
 			FMeshMaterialShader::ModifyCompilationEnvironment(Parameters, OutEnvironment);
-			OutEnvironment.SetDefine(TEXT("VIRTUAL_TEXTURE_PAGE_RENDER"), 1);
 			OutEnvironment.SetDefine(TEXT("IS_VIRTUAL_TEXTURE_MATERIAL"), 1);
 
 			static FShaderPlatformCachedIniValue<bool> HighQualityPerPixelHeightValue(TEXT("r.VT.RVT.HighQualityPerPixelHeight"));
