Add flag in the material translator to track if a virtual texture sample has happened when compiling the Norml output attribute.
This will be used in a later change to understand if using PixelNormalWS indirectly implies the sampling of virtual textures.

diff --git a/Engine/Source/Runtime/Engine/Private/Materials/HLSLMaterialTranslator.cpp b/Engine/Source/Runtime/Engine/Private/Materials/HLSLMaterialTranslator.cpp
--- a/Engine/Source/Runtime/Engine/Private/Materials/HLSLMaterialTranslator.cpp
+++ b/Engine/Source/Runtime/Engine/Private/Materials/HLSLMaterialTranslator.cpp
@@ -738,6 +738,7 @@ FHLSLMaterialTranslator::FHLSLMaterialTranslator(FMaterial* InMaterial,
 ,	bUsesPerInstanceFadeAmount(false)
 ,	bCullIntermediateUniformExpressions(GCullIntermediateUniformExpressions)
 ,	bUsesExplicitDerivatives(false)
+,	bUsesVirtualTextureSampleForNormalProperty(false)
 ,	bIsInRuntimeVirtualTextureOutput(false)
 ,	AddingUniformExpression(0)
 ,	AllocatedUserTexCoords()
@@ -7308,6 +7309,11 @@ int32 FHLSLMaterialTranslator::TextureSample(
 				, VirtualTextureIndex, VirtualTextureIndex);
 		}
 
+		if (MaterialProperty == MP_Normal)
+		{
+			bUsesVirtualTextureSampleForNormalProperty = true;
+		}
+
 		NumVtSamples++;
 	}
 	else // MCT_Texture2D
diff --git a/Engine/Source/Runtime/Engine/Private/Materials/HLSLMaterialTranslator.h b/Engine/Source/Runtime/Engine/Private/Materials/HLSLMaterialTranslator.h
--- a/Engine/Source/Runtime/Engine/Private/Materials/HLSLMaterialTranslator.h
+++ b/Engine/Source/Runtime/Engine/Private/Materials/HLSLMaterialTranslator.h
@@ -449,6 +449,11 @@ protected:
 	/** The generated shader source uses explicit ddx()/ddy() calls */
 	uint32 bUsesExplicitDerivatives : 1;
 
+	/** 
+	  * Tracks if any virtual texture samples are used in the subgraph for the normal property.
+	  * This is read later to know if the PixelNormalWS node is "connected" to virtual texture samples.
+	  */
+	uint32 bUsesVirtualTextureSampleForNormalProperty : 1;
 	/** Flag that tracks if we are compiling a runtime virtual texture output subgraph. */
 	uint32 bIsInRuntimeVirtualTextureOutput : 1;
 
