When duplicating a world for PIE, duplicated objects will have RF_NeedPostLoad, and GetUsedMaterials might invoke ConditionalPostLoad on its MID, triggering a PostLoad chain from MID->Component->Actor->Level, which will remove null entries from the level actors array, resulting in an assert in the iterator.

diff --git a/Engine/Source/Runtime/Engine/Private/World.cpp b/Engine/Source/Runtime/Engine/Private/World.cpp
--- a/Engine/Source/Runtime/Engine/Private/World.cpp
+++ b/Engine/Source/Runtime/Engine/Private/World.cpp
@@ -1232,7 +1232,12 @@ void UWorld::PostDuplicate(bool bDuplicateForPIE)
 					TSet<UMaterialInterface*> ProcessedMaterials; 
 					TArray<UMaterialInterface*> Materials;
 					TArray<UPrimitiveComponent*> Components;
-					for (const TObjectPtr<AActor>& Actor : Level->Actors)
+
+					// When duplicating for PIE, duplicated objects will have RF_NeedPostLoad, and GetUsedMaterials might invoke ConditionalPostLoad on its MID, triggering a
+					// PostLoad chain from MID->Component->Actor->Level, which will remove null entries from the Actors array, resulting in an assert in the iterator.
+					TArray<TObjectPtr<AActor>> LocalActors(Level->Actors);
+
+					for (const TObjectPtr<AActor>& Actor : LocalActors)
 					{
 						if (Actor != nullptr)
 						{
