Fixed FDynamicRenderAssetInstanceManager texture streaming when r.Streaming.AllowParallelRenderAssetStreamingManagerProcessComponents is enabled. FStreamingTextureLevelContext can't use prebuilt data because it is not thread safe (modifies the texture).

diff --git a/Engine/Source/Runtime/Engine/Private/Streaming/DynamicTextureInstanceManager.cpp b/Engine/Source/Runtime/Engine/Private/Streaming/DynamicTextureInstanceManager.cpp
--- a/Engine/Source/Runtime/Engine/Private/Streaming/DynamicTextureInstanceManager.cpp
+++ b/Engine/Source/Runtime/Engine/Private/Streaming/DynamicTextureInstanceManager.cpp
@@ -122,6 +122,8 @@ void FDynamicRenderAssetInstanceManager::IncrementalUpdate(FRemovedRenderAssetAr
 				if (!Component->bAttachedToStreamingManagerAsDynamic && CanManage(Component) && Component->IsRegistered() && (!Component->IsRenderStateCreated() || Component->SceneProxy))
 				{
 					FStreamingTextureLevelContext LevelContext(EMaterialQualityLevel::Num, Component);
+					// Because LevelContext is not thread-safe when ComponentBuildData is valid, we force not to use built data
+					LevelContext.SetForceNoUseBuiltData(true);
 					FRenderAssetInstanceState::PreAddComponentIgnoreBounds(Component, LevelContext, Payloads[Index]);
 				}
 				else
