OFPE: external entities now support the unsaved column in the scene outliner.

diff --git a/Engine/Source/Editor/SceneOutliner/Private/ActorTreeItem.cpp b/Engine/Source/Editor/SceneOutliner/Private/ActorTreeItem.cpp
--- a/Engine/Source/Editor/SceneOutliner/Private/ActorTreeItem.cpp
+++ b/Engine/Source/Editor/SceneOutliner/Private/ActorTreeItem.cpp
@@ -410,9 +410,9 @@ FString FActorTreeItem::GetPackageName() const
 {
 	if (const AActor* ActorPtr = Actor.Get())
 	{
-		if (ActorPtr->IsPackageExternal())
+		if (UPackage* ActorPackage = ActorPtr->GetSceneOutlinerPackage())
 		{
-			return ActorPtr->GetExternalPackage()->GetName();
+			return ActorPackage->GetName();
 		}
 	}
 	
diff --git a/Engine/Source/Runtime/Engine/Classes/GameFramework/Actor.h b/Engine/Source/Runtime/Engine/Classes/GameFramework/Actor.h
--- a/Engine/Source/Runtime/Engine/Classes/GameFramework/Actor.h
+++ b/Engine/Source/Runtime/Engine/Classes/GameFramework/Actor.h
@@ -2648,9 +2648,12 @@ public:
 	/** Returns true if this actor is allowed to be detached from the given actor */
 	ENGINE_API virtual bool EditorCanDetachFrom(const AActor* InParent, FText& OutReason) const { return true; }
 
-	/** Returns the actor attachement parent that should be used in editor */
+	/** Returns the actor attachement parent that should be used in the scene outliner */
 	ENGINE_API virtual AActor* GetSceneOutlinerParent() const;
 
+	/** Returns the actor package that should be used in the scene outliner */
+	ENGINE_API virtual UPackage* GetSceneOutlinerPackage() const;
+
 	/** Called before editor copy, true allow export */
 	virtual bool ShouldExport() { return true; }
 
diff --git a/Engine/Source/Runtime/Engine/Private/ActorEditor.cpp b/Engine/Source/Runtime/Engine/Private/ActorEditor.cpp
--- a/Engine/Source/Runtime/Engine/Private/ActorEditor.cpp
+++ b/Engine/Source/Runtime/Engine/Private/ActorEditor.cpp
@@ -1059,6 +1059,11 @@ AActor* AActor::GetSceneOutlinerParent() const
 	return GetAttachParentActor();
 }
 
+UPackage* AActor::GetSceneOutlinerPackage() const
+{
+	return GetExternalPackage();
+}
+
 class UHLODLayer* AActor::GetHLODLayer() const
 {
 	return HLODLayer;
