WorldPartitionRuntimeCellTransformerISM: split HISM based on collision profiles.

diff --git a/Engine/Source/Runtime/Engine/Private/WorldPartition/WorldPartitionRuntimeCellTransformerISM.cpp b/Engine/Source/Runtime/Engine/Private/WorldPartition/WorldPartitionRuntimeCellTransformerISM.cpp
--- a/Engine/Source/Runtime/Engine/Private/WorldPartition/WorldPartitionRuntimeCellTransformerISM.cpp
+++ b/Engine/Source/Runtime/Engine/Private/WorldPartition/WorldPartitionRuntimeCellTransformerISM.cpp
@@ -72,7 +72,7 @@ void UWorldPartitionRuntimeCellTransformerISM::Transform(ULevel* InLevel)
 				if (!StaticMeshComponent->IsEditorOnly() && StaticMeshComponent->IsVisible() && (StaticMeshComponent->Mobility == EComponentMobility::Static))
 				{
 					FISMComponentDescriptor ISMComponentDescriptor;
-					ISMComponentDescriptor.InitFrom(StaticMeshComponent, false);
+					ISMComponentDescriptor.InitFrom(StaticMeshComponent);
 					ISMComponentBatchers.FindOrAdd(ISMComponentDescriptor).ActorComponents.FindOrAdd(&Actor).Add(StaticMeshComponent);
 				}
 			});
