Avoid FName->FString->FName roundtrip when creating skinned mesh proxy objects

diff --git a/Engine/Source/Runtime/Engine/Private/SkinnedMeshSceneProxyDesc.cpp b/Engine/Source/Runtime/Engine/Private/SkinnedMeshSceneProxyDesc.cpp
--- a/Engine/Source/Runtime/Engine/Private/SkinnedMeshSceneProxyDesc.cpp
+++ b/Engine/Source/Runtime/Engine/Private/SkinnedMeshSceneProxyDesc.cpp
@@ -406,7 +406,7 @@ FSkinnedMeshSceneProxyDynamicData::FSkinnedMeshSceneProxyDynamicData(const USkin
 	: FSkinnedMeshSceneProxyDynamicData(InSkinnedMeshComponent, InSkinnedMeshComponent->LeaderPoseComponent.Get())
 {}
 FSkinnedMeshSceneProxyDynamicData::FSkinnedMeshSceneProxyDynamicData(const USkinnedMeshComponent* InSkinnedMeshComponent, const USkinnedMeshComponent* InLeaderPoseComponent)
-	: Name(InSkinnedMeshComponent->GetName())
+	: Name(InSkinnedMeshComponent->GetFName())
 	, ClothSimulDataProvider(InSkinnedMeshComponent)
 	, MeshDeformerInstances(&InSkinnedMeshComponent->GetMeshDeformerInstances())
 	, RefPoseOverride(InSkinnedMeshComponent->GetRefPoseOverride())
