Add new Custom Primitive Data setter functions in Primitive Component for float vectors and float arrays.
Includes small optimization by converting internal Custom Primitive Data function arguments from TArray to TArrayView.

[RN] minor

diff --git a/Engine/Source/Runtime/Engine/Classes/Components/PrimitiveComponent.h b/Engine/Source/Runtime/Engine/Classes/Components/PrimitiveComponent.h
--- a/Engine/Source/Runtime/Engine/Classes/Components/PrimitiveComponent.h
+++ b/Engine/Source/Runtime/Engine/Classes/Components/PrimitiveComponent.h
@@ -1134,14 +1134,30 @@ public:
 	UFUNCTION(BlueprintCallable, Category="Rendering|Material")
 	ENGINE_API void SetCustomPrimitiveDataVector2(int32 DataIndex, FVector2D Value);
 
+	/** Set custom primitive data, two floats at once, from index DataIndex to index DataIndex + 1. This sets the run-time data only, so it doesn't serialize. */
+	ENGINE_API void SetCustomPrimitiveDataVector2f(int32 DataIndex, FVector2f Value);
+
 	/** Set custom primitive data, three floats at once, from index DataIndex to index DataIndex + 2. This sets the run-time data only, so it doesn't serialize. */
 	UFUNCTION(BlueprintCallable, Category="Rendering|Material")
 	ENGINE_API void SetCustomPrimitiveDataVector3(int32 DataIndex, FVector Value);
 
+	/** Set custom primitive data, three floats at once, from index DataIndex to index DataIndex + 2. This sets the run-time data only, so it doesn't serialize. */
+	ENGINE_API void SetCustomPrimitiveDataVector3f(int32 DataIndex, FVector3f Value);
+
 	/** Set custom primitive data, four floats at once, from index DataIndex to index DataIndex + 3. This sets the run-time data only, so it doesn't serialize. */
 	UFUNCTION(BlueprintCallable, Category="Rendering|Material")
 	ENGINE_API void SetCustomPrimitiveDataVector4(int32 DataIndex, FVector4 Value);
 
+	/** Set custom primitive data, four floats at once, from index DataIndex to index DataIndex + 3. This sets the run-time data only, so it doesn't serialize. */
+	ENGINE_API void SetCustomPrimitiveDataVector4f(int32 DataIndex, FVector4f Value);
+
+	/** Set custom primitive data, an array of floats at once, from index DataIndex to index DataIndex + Values.Num(). This sets the run-time data only, so it doesn't serialize. */
+	UFUNCTION(BlueprintCallable, Category="Rendering|Material")
+	ENGINE_API void SetCustomPrimitiveDataFloatArray(int32 DataIndex, const TArray<float>& Values);
+
+	/** Set custom primitive data, an array of floats at once, from index DataIndex to index DataIndex + Values.Num(). This sets the run-time data only, so it doesn't serialize. */
+	ENGINE_API void SetCustomPrimitiveDataFloatArray(int32 DataIndex, TConstArrayView<float> Values);
+
 	/** 
 	 * Get the custom primitive data for this primitive component.
 	 * @return The payload of custom data that will be set on the primitive and accessible in the material through a material expression.
@@ -1175,14 +1191,30 @@ public:
 	UFUNCTION(BlueprintCallable, Category = "Rendering|Material")
 	ENGINE_API void SetDefaultCustomPrimitiveDataVector2(int32 DataIndex, FVector2D Value);
 
+	/** Set default custom primitive data, two floats at once, from index DataIndex to index DataIndex + 1, and marks the render state dirty */
+	ENGINE_API void SetDefaultCustomPrimitiveDataVector2f(int32 DataIndex, FVector2f Value);
+
 	/** Set default custom primitive data, three floats at once, from index DataIndex to index DataIndex + 2, and marks the render state dirty */
 	UFUNCTION(BlueprintCallable, Category = "Rendering|Material")
 	ENGINE_API void SetDefaultCustomPrimitiveDataVector3(int32 DataIndex, FVector Value);
 
+	/** Set default custom primitive data, three floats at once, from index DataIndex to index DataIndex + 2, and marks the render state dirty */
+	ENGINE_API void SetDefaultCustomPrimitiveDataVector3f(int32 DataIndex, FVector3f Value);
+
 	/** Set default custom primitive data, four floats at once, from index DataIndex to index DataIndex + 3, and marks the render state dirty */
 	UFUNCTION(BlueprintCallable, Category = "Rendering|Material")
 	ENGINE_API void SetDefaultCustomPrimitiveDataVector4(int32 DataIndex, FVector4 Value);
 
+	/** Set default custom primitive data, four floats at once, from index DataIndex to index DataIndex + 3, and marks the render state dirty */
+	ENGINE_API void SetDefaultCustomPrimitiveDataVector4f(int32 DataIndex, FVector4f Value);
+
+	/** Set default custom primitive data, an array floats at once, from index DataIndex to index DataIndex + Values.Num(), and marks the render state dirty */
+	UFUNCTION(BlueprintCallable, Category="Rendering|Material")
+	ENGINE_API void SetDefaultCustomPrimitiveDataFloatArray(int32 DataIndex, const TArray<float>& Values);
+
+	/** Set default custom primitive data, an array floats at once, from index DataIndex to index DataIndex + Values.Num(), and marks the render state dirty */
+	ENGINE_API void SetDefaultCustomPrimitiveDataFloatArray(int32 DataIndex, TConstArrayView<float> Values);
+
 	/**
 	 * Get the default custom primitive data for this primitive component.
 	 * @return The payload of custom data that will be set on the primitive and accessible in the material through a material expression.
@@ -1198,10 +1230,10 @@ public:
 protected:
 
 	/** Insert an array of floats into the CustomPrimitiveData, starting at the given index */
-	ENGINE_API void SetCustomPrimitiveDataInternal(int32 DataIndex, const TArray<float>& Values);
+	ENGINE_API void SetCustomPrimitiveDataInternal(int32 DataIndex, TConstArrayView<float> Values);
 
 	/** Insert an array of floats into the CustomPrimitiveData defaults, starting at the given index */
-	ENGINE_API void SetDefaultCustomPrimitiveData(int32 DataIndex, const TArray<float>& Values);
+	ENGINE_API void SetDefaultCustomPrimitiveData(int32 DataIndex, TConstArrayView<float> Values);
 
 	/** Set of components that this component is currently overlapping. */
 	TArray<FOverlapInfo> OverlappingComponents;
diff --git a/Engine/Source/Runtime/Engine/Private/Components/PrimitiveComponent.cpp b/Engine/Source/Runtime/Engine/Private/Components/PrimitiveComponent.cpp
--- a/Engine/Source/Runtime/Engine/Private/Components/PrimitiveComponent.cpp
+++ b/Engine/Source/Runtime/Engine/Private/Components/PrimitiveComponent.cpp
@@ -2324,7 +2324,7 @@ void UPrimitiveComponent::ResetCustomPrimitiveData()
 }
 
 /** Attempt to set the primitive data and return true if successful */
-bool SetPrimitiveData(FCustomPrimitiveData& PrimitiveData, int32 DataIndex, const TArray<float>& Values)
+bool SetPrimitiveData(FCustomPrimitiveData& PrimitiveData, int32 DataIndex, TConstArrayView<float> Values)
 {
 	// Can only set data on valid indices and only if there's actually any data to set
 	if (DataIndex >= 0 && DataIndex < FCustomPrimitiveData::NumCustomPrimitiveDataFloats && Values.Num() > 0)
@@ -2353,7 +2353,7 @@ bool SetPrimitiveData(FCustomPrimitiveData& PrimitiveData, int32 DataIndex, cons
 	return false;
 }
 
-void UPrimitiveComponent::SetCustomPrimitiveDataInternal(int32 DataIndex, const TArray<float>& Values)
+void UPrimitiveComponent::SetCustomPrimitiveDataInternal(int32 DataIndex, TConstArrayView<float> Values)
 {
 	if (SetPrimitiveData(CustomPrimitiveDataInternal, DataIndex, Values))
 	{
@@ -2365,7 +2365,7 @@ void UPrimitiveComponent::SetCustomPrimitiveDataInternal(int32 DataIndex, const
 	}
 }
 
-void UPrimitiveComponent::SetDefaultCustomPrimitiveData(int32 DataIndex, const TArray<float>& Values)
+void UPrimitiveComponent::SetDefaultCustomPrimitiveData(int32 DataIndex, TConstArrayView<float> Values)
 {
 	if (SetPrimitiveData(CustomPrimitiveData, DataIndex, Values))
 	{
@@ -2445,6 +2445,11 @@ void UPrimitiveComponent::SetCustomPrimitiveDataVector2(int32 DataIndex, FVector
 	SetCustomPrimitiveDataInternal(DataIndex, {ValueFlt.X, ValueFlt.Y});
 }
 
+void UPrimitiveComponent::SetCustomPrimitiveDataVector2f(int32 DataIndex, FVector2f Value)
+{
+	SetCustomPrimitiveDataInternal(DataIndex, { Value.X, Value.Y });
+}
+
 void UPrimitiveComponent::SetCustomPrimitiveDataVector3(int32 DataIndex, FVector Value)
 {
 	// LWC_TODO: precision loss
@@ -2452,6 +2457,11 @@ void UPrimitiveComponent::SetCustomPrimitiveDataVector3(int32 DataIndex, FVector
 	SetCustomPrimitiveDataInternal(DataIndex, {ValueFlt.X, ValueFlt.Y, ValueFlt.Z});
 }
 
+void UPrimitiveComponent::SetCustomPrimitiveDataVector3f(int32 DataIndex, FVector3f Value)
+{
+	SetCustomPrimitiveDataInternal(DataIndex, { Value.X, Value.Y, Value.Z });
+}
+
 void UPrimitiveComponent::SetCustomPrimitiveDataVector4(int32 DataIndex, FVector4 Value)
 {
 	// LWC_TODO: precision loss
@@ -2459,6 +2469,21 @@ void UPrimitiveComponent::SetCustomPrimitiveDataVector4(int32 DataIndex, FVector
 	SetCustomPrimitiveDataInternal(DataIndex, {ValueFlt.X, ValueFlt.Y, ValueFlt.Z, ValueFlt.W});
 }
 
+void UPrimitiveComponent::SetCustomPrimitiveDataVector4f(int32 DataIndex, FVector4f Value)
+{
+	SetCustomPrimitiveDataInternal(DataIndex, { Value.X, Value.Y, Value.Z, Value.W });
+}
+
+void UPrimitiveComponent::SetCustomPrimitiveDataFloatArray(int32 DataIndex, const TArray<float>& Values)
+{
+	SetCustomPrimitiveDataInternal(DataIndex, MakeConstArrayView(Values.GetData(), Values.Num()));
+}
+
+void UPrimitiveComponent::SetCustomPrimitiveDataFloatArray(int32 DataIndex, TConstArrayView<float> Values)
+{
+	SetCustomPrimitiveDataInternal(DataIndex, Values);
+}
+
 void UPrimitiveComponent::SetScalarParameterForDefaultCustomPrimitiveData(FName ParameterName, float Value)
 {
 	int32 PrimitiveDataIndex = GetCustomPrimitiveDataIndexForScalarParameter(ParameterName);
@@ -2492,18 +2517,43 @@ void UPrimitiveComponent::SetDefaultCustomPrimitiveDataVector2(int32 DataIndex,
 	SetDefaultCustomPrimitiveData(DataIndex, { ValueFlt.X, ValueFlt.Y });
 }
 
+void UPrimitiveComponent::SetDefaultCustomPrimitiveDataVector2f(int32 DataIndex, FVector2f Value)
+{
+	SetDefaultCustomPrimitiveData(DataIndex, { Value.X, Value.Y });
+}
+
 void UPrimitiveComponent::SetDefaultCustomPrimitiveDataVector3(int32 DataIndex, FVector Value)
 {
 	FVector3f ValueFlt(Value);
 	SetDefaultCustomPrimitiveData(DataIndex, { ValueFlt.X, ValueFlt.Y, ValueFlt.Z });
 }
 
+void UPrimitiveComponent::SetDefaultCustomPrimitiveDataVector3f(int32 DataIndex, FVector3f Value)
+{
+	SetDefaultCustomPrimitiveData(DataIndex, { Value.X, Value.Y, Value.Z });
+}
+
 void UPrimitiveComponent::SetDefaultCustomPrimitiveDataVector4(int32 DataIndex, FVector4 Value)
 {
 	FVector4f ValueFlt(Value);
 	SetDefaultCustomPrimitiveData(DataIndex, { ValueFlt.X, ValueFlt.Y, ValueFlt.Z, ValueFlt.W });
 }
 
+void UPrimitiveComponent::SetDefaultCustomPrimitiveDataVector4f(int32 DataIndex, FVector4f Value)
+{
+	SetDefaultCustomPrimitiveData(DataIndex, { Value.X, Value.Y, Value.Z, Value.W });
+}
+
+void UPrimitiveComponent::SetDefaultCustomPrimitiveDataFloatArray(int32 DataIndex, const TArray<float>& Values)
+{
+	SetDefaultCustomPrimitiveData(DataIndex, MakeConstArrayView(Values.GetData(), Values.Num()));
+}
+
+void UPrimitiveComponent::SetDefaultCustomPrimitiveDataFloatArray(int32 DataIndex, TConstArrayView<float> Values)
+{
+	SetDefaultCustomPrimitiveData(DataIndex, Values);
+}
+
 UMaterialInterface* UPrimitiveComponent::GetMaterialFromCollisionFaceIndex(int32 FaceIndex, int32& SectionIndex) const
 {
 	//This function should be overriden
