Iris
* Ensure and gracefully fail when subobject factory cannot load the UClass of a replicated object.
* Before this would cause a client crash

#rnx

diff --git a/Engine/Source/Runtime/Engine/Private/Net/Iris/ReplicationSystem/NetSubObjectFactory.cpp b/Engine/Source/Runtime/Engine/Private/Net/Iris/ReplicationSystem/NetSubObjectFactory.cpp
--- a/Engine/Source/Runtime/Engine/Private/Net/Iris/ReplicationSystem/NetSubObjectFactory.cpp
+++ b/Engine/Source/Runtime/Engine/Private/Net/Iris/ReplicationSystem/NetSubObjectFactory.cpp
@@ -161,10 +161,18 @@ UNetSubObjectFactory::FInstantiateResult UNetSubObjectFactory::InstantiateReplic
 		}
 	}
 
-	// We need to spawn the subobject
+	// Find the class of the subobject
 	UObject* SubObjClassObj = Bridge->ResolveObjectReference(SubObjectHeader->ObjectClassReference, Context.ResolveContext);
 	UClass * SubObjClass = Cast<UClass>(SubObjClassObj);
 
+	if (!SubObjClass)
+	{
+		UE_LOG(LogIris, Error, TEXT("UNetSubObjectFactory::InstantiateNetObjectFromHeader could not find UClass via reference: %s. Cannot spawn subobject for handle: %s"), 
+			*Bridge->DescribeObjectReference(SubObjectHeader->ObjectClassReference, Context.ResolveContext), *Context.Handle.ToString());
+		ensure(SubObjClass);
+		return FInstantiateResult();
+	}
+
 	// Try to spawn SubObject
 	UObject* SubObj = NewObject<UObject>(OuterObject, SubObjClass);
 
