USoundBase::PostLoad was listening on FCoreUObjectDelegates::OnAssetLoaded to handle loading of its AudioPropertiesSheet sheet member, even when this member was null.
Also, since the AudioPropertiesSheet member is a TObjectPtr, it will always be loaded when USoundBase::PostLoad is invoked, so there's no need to do that, simply call ConditionalPostLoad on it.

diff --git a/Engine/Source/Runtime/Engine/Private/SoundBase.cpp b/Engine/Source/Runtime/Engine/Private/SoundBase.cpp
--- a/Engine/Source/Runtime/Engine/Private/SoundBase.cpp
+++ b/Engine/Source/Runtime/Engine/Private/SoundBase.cpp
@@ -248,21 +248,11 @@ void USoundBase::PostLoad()
 		ConcurrencyOverrides.ResolutionRule = MaxConcurrentResolutionRule_DEPRECATED;
 	}
 
-	if (AudioPropertiesSheet && AudioPropertiesSheet->HasAnyFlags(RF_WasLoaded))
+	if (AudioPropertiesSheet)
 	{
+		AudioPropertiesSheet->ConditionalPostLoad();
 		BindToPropertySheetChanges();
 	}
-	else
-	{
-		FCoreUObjectDelegates::OnAssetLoaded.AddWeakLambda(this, [this](UObject* LoadedObject)
-			{
-				if (LoadedObject == AudioPropertiesSheet)
-				{
-					BindToPropertySheetChanges();
-					FCoreUObjectDelegates::OnAssetLoaded.RemoveAll(this);
-				}
-			});
-	}
 }
 #endif // WITH_EDITORONLY_DATA
 
