Removing r.ShaderCompiler.JobCache

- The job cache is so beneficial this has become the default.
- There aren't many scenarios where disabling this is a good idea.

#rnx

diff --git a/Engine/Config/ConsoleVariables.ini b/Engine/Config/ConsoleVariables.ini
--- a/Engine/Config/ConsoleVariables.ini
+++ b/Engine/Config/ConsoleVariables.ini
@@ -56,9 +56,6 @@ r.DumpShaderDebugInfo=2
 ;r.DumpShaderDebugWorkerCommandLine=1
 ; When this is enabled, shader compiler warnings are emitted to the log for all shaders as they are loaded (either from the DDC or from a shader compilation job).
 ;r.ShaderCompiler.EmitWarningsOnLoad=1
-; When this is enabled, shader compiler will cache jobs with the same input (works essentially like in-memory transient DDC for each shader).
-; Among other things, this reduces shader compilation in the material editor. Disable when you need to debug the shaders or the cache itself.
-;r.ShaderCompiler.JobCache=0
 
 ; Uncomment to disable parallel rendering
 ;r.RHICmdBypass=1
diff --git a/Engine/Source/Runtime/Engine/Private/ShaderCompiler/ShaderCompilerJobCache.cpp b/Engine/Source/Runtime/Engine/Private/ShaderCompiler/ShaderCompilerJobCache.cpp
--- a/Engine/Source/Runtime/Engine/Private/ShaderCompiler/ShaderCompilerJobCache.cpp
+++ b/Engine/Source/Runtime/Engine/Private/ShaderCompiler/ShaderCompilerJobCache.cpp
@@ -73,13 +73,6 @@ static TAutoConsoleVariable<bool> CVarShaderCompilerPerShaderDDCAsync(
 	TEXT("if != 0, Per-shader DDC queries will run async, instead of in the SubmitJobs task."),
 	ECVF_Default);
 
-int32 GShaderCompilerJobCache = 1;
-static FAutoConsoleVariableRef CVarShaderCompilerJobCache(
-	TEXT("r.ShaderCompiler.JobCache"),
-	GShaderCompilerJobCache,
-	TEXT("if != 0, shader compiler cache (based on the unpreprocessed input hash) will be disabled. By default, it is enabled."),
-	ECVF_Default);
-
 static TAutoConsoleVariable<bool> CVarJobCacheDDC(
 	TEXT("r.ShaderCompiler.JobCacheDDC"),
 	true,
@@ -147,14 +140,9 @@ namespace
 /** Helper functions for logging more debug info */
 namespace ShaderCompiler
 {
-	bool IsJobCacheEnabled()
-	{
-		return GShaderCompilerJobCache != 0;
-	}
-
 	bool IsJobCacheDebugValidateEnabled()
 	{
-		return IsJobCacheEnabled() && CVarShaderCompilerDebugValidateJobCache.GetValueOnAnyThread();
+		return CVarShaderCompilerDebugValidateJobCache.GetValueOnAnyThread();
 	}
 } // namespace ShaderCompiler
 
@@ -175,7 +163,7 @@ bool IsShaderJobCacheDDCEnabled()
 	if ((GIsEditor || IsRunningGame() || bEnablePerShaderDDCCook) && (!bIsNonCookCommandlet || bForceAllowShaderCompilerJobCache))
 	{
 		// job cache itself must be enabled first
-		return GShaderCompilerJobCache && CVarJobCacheDDC.GetValueOnAnyThread();
+		return CVarJobCacheDDC.GetValueOnAnyThread();
 	}
 
 	return false;
@@ -1158,7 +1146,7 @@ void FShaderJobCache::InternalSetPriority(FShaderCommonCompileJob* Job, EShaderC
 			check(Job->PendingPriority == Job->Priority);
 			UnlinkJobWithPriority(*Job);
 
-			ensure(!ShaderCompiler::IsJobCacheEnabled() || Job->bInputHashSet);
+			ensure(Job->bInputHashSet);
 			Job->Priority = InPriority;
 			LinkJobWithPriority(*Job);
 
@@ -1181,8 +1169,8 @@ void FShaderJobCache::InternalSetPriority(FShaderCommonCompileJob* Job, EShaderC
 			Job->bInputHashSet,
 			NewJob->bInputHashSet
 		);
-		ensureMsgf(!ShaderCompiler::IsJobCacheEnabled() || NewJob->GetInputHash() == Job->GetInputHash(),
-			TEXT("If shader jobs cache is enabled, cloned job should have the same input hash as the original, and it doesn't.")
+		ensureMsgf(NewJob->GetInputHash() == Job->GetInputHash(),
+			TEXT("Cloned job should have the same input hash as the original, and it doesn't.")
 		);
 
 		FWriteScopeLock Locker(JobLock);
@@ -1248,44 +1236,41 @@ int32 FShaderJobCache::RemoveAllPendingJobsWithId(uint32 InId)
 
 				if (Job.Id == InId)
 				{
-					if (ShaderCompiler::IsJobCacheEnabled())
+					if (Job.JobCacheRef.Block)
 					{
-						if (Job.JobCacheRef.Block)
-						{
-							FShaderJobData& JobData = GetShaderJobData(Job.JobCacheRef);
+						FShaderJobData& JobData = GetShaderJobData(Job.JobCacheRef);
 
-							check(JobData.JobInFlight == &Job);
+						check(JobData.JobInFlight == &Job);
 
-							// If we are removing an in-flight job, we need to promote a duplicate to be the new in-flight job, if present.
-							// Make sure the duplicate we choose doesn't have the same ID as what we're removing.
-							FShaderCommonCompileJob* DuplicateJob;
-							for (DuplicateJob = JobData.DuplicateJobsWaitList; DuplicateJob; DuplicateJob = DuplicateJob->NextLink)
+						// If we are removing an in-flight job, we need to promote a duplicate to be the new in-flight job, if present.
+						// Make sure the duplicate we choose doesn't have the same ID as what we're removing.
+						FShaderCommonCompileJob* DuplicateJob;
+						for (DuplicateJob = JobData.DuplicateJobsWaitList; DuplicateJob; DuplicateJob = DuplicateJob->NextLink)
+						{
+							if (DuplicateJob->Id != InId)
 							{
-								if (DuplicateJob->Id != InId)
-								{
-									break;
-								}
+								break;
 							}
+						}
 
-							if (DuplicateJob)
+						if (DuplicateJob)
+						{
+							// Advance head if we are unlinking the head, then remove
+							if (JobData.DuplicateJobsWaitList == DuplicateJob)
 							{
-								// Advance head if we are unlinking the head, then remove
-								if (JobData.DuplicateJobsWaitList == DuplicateJob)
-								{
-									JobData.DuplicateJobsWaitList = DuplicateJob->NextLink;
-								}
-								Unlink(*DuplicateJob);
-								RemoveDuplicateJob(DuplicateJob);
-
-								// Add it as pending at the appropriate priority
-								GShaderCompilerStats->RegisterNewPendingJob(*DuplicateJob);
-
-								LinkJobWithPriority(*DuplicateJob);
+								JobData.DuplicateJobsWaitList = DuplicateJob->NextLink;
 							}
+							Unlink(*DuplicateJob);
+							RemoveDuplicateJob(DuplicateJob);
 
-							// DuplicateJob will be nullptr if there was no duplicate to promote
-							JobData.JobInFlight = DuplicateJob;
+							// Add it as pending at the appropriate priority
+							GShaderCompilerStats->RegisterNewPendingJob(*DuplicateJob);
+
+							LinkJobWithPriority(*DuplicateJob);
 						}
+
+						// DuplicateJob will be nullptr if there was no duplicate to promote
+						JobData.JobInFlight = DuplicateJob;
 					}
 
 					check((int32)Job.PendingPriority == PriorityIndex);
@@ -1296,38 +1281,35 @@ int32 FShaderJobCache::RemoveAllPendingJobsWithId(uint32 InId)
 			}
 		}
 
-		if (ShaderCompiler::IsJobCacheEnabled())
+		// Also look into duplicate jobs that are cached -- we don't increment in the "for" loop because the current item may be deleted
+		for (int32 DuplicateIndex = 0; DuplicateIndex < DuplicateJobs.Num();)
 		{
-			// Also look into duplicate jobs that are cached -- we don't increment in the "for" loop because the current item may be deleted
-			for (int32 DuplicateIndex = 0; DuplicateIndex < DuplicateJobs.Num();)
+			FShaderCommonCompileJob* DuplicateJob = DuplicateJobs[DuplicateIndex];
+			check(DuplicateJob->JobCacheRef.DuplicateIndex == DuplicateIndex);
+
+			if (DuplicateJob->Id == InId)
 			{
-				FShaderCommonCompileJob* DuplicateJob = DuplicateJobs[DuplicateIndex];
-				check(DuplicateJob->JobCacheRef.DuplicateIndex == DuplicateIndex);
+				FShaderJobData& JobData = GetShaderJobData(DuplicateJob->JobCacheRef);
 
-				if (DuplicateJob->Id == InId)
+				// if we're removing the list head, we need to update it to the next
+				if (JobData.DuplicateJobsWaitList == DuplicateJob)
 				{
-					FShaderJobData& JobData = GetShaderJobData(DuplicateJob->JobCacheRef);
-
-					// if we're removing the list head, we need to update it to the next
-					if (JobData.DuplicateJobsWaitList == DuplicateJob)
-					{
-						JobData.DuplicateJobsWaitList = JobData.DuplicateJobsWaitList->NextLink;
-					}
+					JobData.DuplicateJobsWaitList = JobData.DuplicateJobsWaitList->NextLink;
+				}
 
-					// This removes the current job (at DuplicateIndex), so we don't increment in this case
-					RemoveDuplicateJob(DuplicateJob);
+				// This removes the current job (at DuplicateIndex), so we don't increment in this case
+				RemoveDuplicateJob(DuplicateJob);
 
-					// Duplicate jobs are in their own list, not one of the priority lists, so don't use UnlinkJobWithPriority
-					check(DuplicateJob->PendingPriority == EShaderCompileJobPriority::None);
-					Unlink(*DuplicateJob);
-					RemoveJob(DuplicateJob);
-					++NumRemoved;
-				}
-				else
-				{
-					// Didn't remove a job, increment!
-					DuplicateIndex++;
-				}
+				// Duplicate jobs are in their own list, not one of the priority lists, so don't use UnlinkJobWithPriority
+				check(DuplicateJob->PendingPriority == EShaderCompileJobPriority::None);
+				Unlink(*DuplicateJob);
+				RemoveJob(DuplicateJob);
+				++NumRemoved;
+			}
+			else
+			{
+				// Didn't remove a job, increment!
+				DuplicateIndex++;
 			}
 		}
 	}
@@ -1355,7 +1337,7 @@ void FShaderJobCache::SubmitJob(FShaderCommonCompileJob* Job)
 	bool bJobCacheLocked = false;
 
 	// check caches unless we're running in validation mode (which runs _all_ jobs and compares hashes of outputs)
-	if (ShaderCompiler::IsJobCacheEnabled() && !ShaderCompiler::IsJobCacheDebugValidateEnabled())
+	if (!ShaderCompiler::IsJobCacheDebugValidateEnabled())
 	{
 		bJobCacheLocked = true;
 
@@ -1449,7 +1431,7 @@ void FShaderJobCache::SubmitJob(FShaderCommonCompileJob* Job)
 	if (bNewJob)
 	{
 		GShaderCompilerStats->RegisterNewPendingJob(*Job);
-		ensure(!ShaderCompiler::IsJobCacheEnabled() || Job->bInputHashSet);
+		ensure(Job->bInputHashSet);
 
 		// If cache is disabled, we skipped the code that grabs the write lock above, so we need to do it here, before modifying the pending queue
 		if (bJobCacheLocked == false)
@@ -1516,12 +1498,8 @@ void FShaderJobCache::SubmitJobs(const TArray<FShaderCommonCompileJobPtr>& InJob
 						FPlatformProcess::Sleep(GShaderCompilerDebugStallSubmitJob * 0.001f);
 					}
 
-					bool bSubmitJob = true;
-					if (ShaderCompiler::IsJobCacheEnabled())
-					{
-						bSubmitJob = PreprocessShader(Job);
-						Job->GetInputHash();
-					}
+					const bool bSubmitJob = PreprocessShader(Job);
+					Job->GetInputHash();
 
 					if (bSubmitJob)
 					{
@@ -1541,7 +1519,7 @@ void FShaderJobCache::SubmitJobs(const TArray<FShaderCommonCompileJobPtr>& InJob
 double FShaderJobCache::ProcessFinishedJob(FShaderCommonCompileJob* FinishedJob, bool bCompilationSkipped)
 {
 	double StallTime;
-	if (!bCompilationSkipped && ShaderCompiler::IsJobCacheEnabled())
+	if (!bCompilationSkipped)
 	{
 		AddToCacheAndProcessPending(FinishedJob);
 	}
@@ -1573,7 +1551,7 @@ void FShaderJobCache::AddToCacheAndProcessPending(FShaderCommonCompileJob* Finis
 {
 	// Cloned jobs won't include an entry in the job cache, so skip the caching logic.  The non-cloned version of the same
 	// job will handle adding data to the cache when it completes.
-	if (!ShaderCompiler::IsJobCacheEnabled() || !FinishedJob->JobCacheRef.Block)
+	if (!FinishedJob->JobCacheRef.Block)
 	{
 		return;
 	}
@@ -1717,10 +1695,7 @@ int32 FShaderJobCache::GetPendingJobs(EShaderCompilerWorkerType InWorkerType, ES
 		FShaderCommonCompileJob& Job = *It;
 
 		GShaderCompilerStats->RegisterAssignedJob(Job);
-		// Temporary commented out until r.ShaderDevelopmentMode=1 shader error retry crash gets fixed
-		//check(Job.CurrentWorker == EShaderCompilerWorkerType::None);
-		//check(Job.PendingPriority == InPriority);
-		ensure(!ShaderCompiler::IsJobCacheEnabled() || Job.bInputHashSet);
+		ensure(Job.bInputHashSet);
 
 		It.Next();
 
@@ -1752,8 +1727,6 @@ FShaderCacheLoadContext FShaderJobCache::FindOrAdd(const FShaderCompilerInputHas
 {
 	LLM_SCOPE_BYTAG(ShaderCompiler);
 
-	check(ShaderCompiler::IsJobCacheEnabled());
-
 	++TotalSearchAttempts;
 	TRACE_COUNTER_INCREMENT(Shaders_JobCacheSearchAttempts);
 	TRACE_CPUPROFILER_EVENT_SCOPE(FShaderJobCache::Find);
@@ -2038,8 +2011,6 @@ FShaderCacheLoadContext FShaderJobCache::FindOrAdd(const FShaderCompilerInputHas
 
 FShaderJobData* FShaderJobCache::Find(const FShaderCompilerInputHash& Hash)
 {
-	check(ShaderCompiler::IsJobCacheEnabled());
-
 	return InputHashToJobData.Find(Hash);
 }
 
@@ -2101,11 +2072,6 @@ void FShaderJobCache::AddJobOutput(FShaderJobData& JobData, const FShaderCommonC
 {
 	TRACE_CPUPROFILER_EVENT_SCOPE(FShaderJobCache::Add);
 
-	if (!ShaderCompiler::IsJobCacheEnabled())
-	{
-		return;
-	}
-
 	if (JobData.HasOutput() && !ShaderCompiler::IsJobCacheDebugValidateEnabled())
 	{
 		return;
diff --git a/Engine/Source/Runtime/Engine/Private/ShaderCompiler/ShaderCompilerPrivate.h b/Engine/Source/Runtime/Engine/Private/ShaderCompiler/ShaderCompilerPrivate.h
--- a/Engine/Source/Runtime/Engine/Private/ShaderCompiler/ShaderCompilerPrivate.h
+++ b/Engine/Source/Runtime/Engine/Private/ShaderCompiler/ShaderCompilerPrivate.h
@@ -51,7 +51,6 @@ namespace ShaderCompilerCookStats
 namespace ShaderCompiler
 {
 	FString GetTargetPlatformName(const ITargetPlatform* TargetPlatform);
-	bool IsJobCacheEnabled();
 	bool IsJobCacheDebugValidateEnabled();
 	bool IsRemoteCompilingAllowed();
 	bool IsDumpShaderDebugInfoAlwaysEnabled();
diff --git a/Engine/Source/Runtime/Engine/Private/ShaderCompiler/ShaderCompilerThreadRunnable.cpp b/Engine/Source/Runtime/Engine/Private/ShaderCompiler/ShaderCompilerThreadRunnable.cpp
--- a/Engine/Source/Runtime/Engine/Private/ShaderCompiler/ShaderCompilerThreadRunnable.cpp
+++ b/Engine/Source/Runtime/Engine/Private/ShaderCompiler/ShaderCompilerThreadRunnable.cpp
@@ -576,7 +576,7 @@ void FShaderCompileThreadRunnable::WriteNewTasks()
 			{
 				FScopeLock Locker(&ArchiveLock);
 				static int ArchivedTransferFileNum = 0;
-				FString JobCacheDir = ShaderCompiler::IsJobCacheEnabled() ? TEXT("JobCache") : TEXT("NoJobCache");
+				FString JobCacheDir = TEXT("JobCache");
 				FString ArchiveDir = FPaths::ProjectSavedDir() / TEXT("ArchivedWorkerInputs") / JobCacheDir;
 				FString ArchiveName = FString::Printf(TEXT("Input-%d"), ArchivedTransferFileNum++);
 				FString ArchivePath = ArchiveDir / ArchiveName;
