[Interchange] fix a deadlock when re-importing skeletalmesh and sk component physics category detail panel needs to refresh.

diff --git a/Engine/Source/Runtime/Engine/Private/SkeletalMeshComponentPhysics.cpp b/Engine/Source/Runtime/Engine/Private/SkeletalMeshComponentPhysics.cpp
--- a/Engine/Source/Runtime/Engine/Private/SkeletalMeshComponentPhysics.cpp
+++ b/Engine/Source/Runtime/Engine/Private/SkeletalMeshComponentPhysics.cpp
@@ -197,7 +197,11 @@ UBodySetup* USkeletalMeshComponent::GetBodySetup()
 
 bool USkeletalMeshComponent::CanEditSimulatePhysics()
 {
+#if WITH_EDITOR
+	return (GetSkinnedAsset() != nullptr) && !GetSkinnedAsset()->IsCompiling() && (GetPhysicsAsset() != nullptr);
+#else
 	return GetPhysicsAsset() != nullptr;
+#endif
 }
 
 void USkeletalMeshComponent::SetSimulatePhysics(bool bSimulate)
