Adding Insights to inflight stream cache loads

diff --git a/Engine/Source/Runtime/Engine/Private/AudioStreamingCache.cpp b/Engine/Source/Runtime/Engine/Private/AudioStreamingCache.cpp
--- a/Engine/Source/Runtime/Engine/Private/AudioStreamingCache.cpp
+++ b/Engine/Source/Runtime/Engine/Private/AudioStreamingCache.cpp
@@ -13,11 +13,16 @@ AudioStreaming.cpp: Implementation of audio streaming classes.
 #include "Engine/Engine.h"
 #include "Stats/StatsTrace.h"
 #include "ProfilingDebugging/IoStoreTrace.h"
+#include "ProfilingDebugging/CountersTrace.h"
 #include "Audio/AudioDebug.h"
 #include "Audio/AudioDebugDraw.h"
 
 DEFINE_LOG_CATEGORY(LogAudioStreamCaching);
 
+TRACE_DECLARE_ATOMIC_INT_COUNTER(AudioStreamCache_Loads, TEXT("AudioStreamCache Loads"));
+TRACE_DECLARE_ATOMIC_INT_COUNTER(AudioStreamCache_Loads_Request, TEXT("AudioStreamCache Loads Request"));
+TRACE_DECLARE_ATOMIC_INT_COUNTER(AudioStreamCache_Loads_Playback, TEXT("AudioStreamCache Loads Playback"));
+
 // Defined in AudioDecompress.cpp off au.debug.extra_logging_filter
 extern FString ExtraTracingSoundWaveFilterCVar;
 
@@ -181,6 +186,39 @@ static FAutoConsoleCommand GFlushAudioCacheCommand(
 	})
 );
 
+static EAsyncIOPriorityAndFlags ToIoPriority(const int32 InValue)
+{
+	switch (InValue)
+	{
+		case 4: return AIOP_MIN;
+		case 3: return AIOP_Low;
+		case 2: return AIOP_BelowNormal;
+		case 1: return AIOP_Normal;
+		case 0: return AIOP_High;
+		default: break;
+	}
+	return AIOP_High;
+}
+
+static void AddLoadCounters(const EAsyncIOPriorityAndFlags InPriorityAndFlags, const int32 InValue)
+{
+	const EAsyncIOPriorityAndFlags Priority = InPriorityAndFlags & AIOP_PRIORITY_MASK;
+	const EAsyncIOPriorityAndFlags ReadPriority = ToIoPriority(ReadRequestPriorityCVar);
+	const EAsyncIOPriorityAndFlags PlaybackPriority = ToIoPriority(PlaybackRequestPriorityCVar);
+
+	if (ReadPriority == Priority)
+	{
+		TRACE_COUNTER_ADD(AudioStreamCache_Loads_Request, InValue);
+	}
+
+	if (PlaybackPriority == Priority)
+	{
+		TRACE_COUNTER_ADD(AudioStreamCache_Loads_Playback, InValue);
+	}
+	
+	TRACE_COUNTER_ADD(AudioStreamCache_Loads, InValue);
+}
+
 static FAutoConsoleCommand GResizeAudioCacheCommand(
 	TEXT("au.streamcaching.ResizeAudioCacheTo"),
 	TEXT("This will try to cull enough audio chunks to shrink the audio stream cache to the new size if neccessary, and keep the cache at that size."),
@@ -1846,6 +1884,7 @@ void FAudioChunkCache::KickOffAsyncLoad(FCacheElement* CacheElement, const FChun
 		};
 
 		NumberOfLoadsInFlight.Increment();
+		AddLoadCounters(AsyncIOPriority, 1);
 
 		CacheElement->DDCTask.Reset(new FAsyncStreamDerivedChunkTask(
 			Chunk->DerivedDataKey,
@@ -1878,8 +1917,9 @@ void FAudioChunkCache::KickOffAsyncLoad(FCacheElement* CacheElement, const FChun
 		//ensureAlwaysMsgf(AudioChunkSize == ChunkBulkDataSize, TEXT("For memory load on demand, we do not zero-pad to page sizes."));
 
 		NumberOfLoadsInFlight.Increment();
+		AddLoadCounters(AsyncIOPriority, 1);
 
-		FBulkDataIORequestCallBack AsyncFileCallBack = [this, OnLoadCompleted, CacheElement, InKey, ChunkDataSize, CallbackThread](bool bWasCancelled, IBulkDataIORequest*)
+		FBulkDataIORequestCallBack AsyncFileCallBack = [this, OnLoadCompleted, CacheElement, InKey, ChunkDataSize, CallbackThread, AsyncIOPriority](bool bWasCancelled, IBulkDataIORequest*)
 		{
 			// Take ownership of the read request and close the storage
 			IBulkDataIORequest* LocalReadRequest = (IBulkDataIORequest*)FPlatformAtomics::InterlockedExchangePtr((void* volatile*)&CacheElement->ReadRequest, (void*)0x1);
@@ -1920,6 +1960,7 @@ void FAudioChunkCache::KickOffAsyncLoad(FCacheElement* CacheElement, const FChun
 			ExecuteOnLoadCompleteCallback(LoadResult, OnLoadCompleted, CallbackThread);
 
 			NumberOfLoadsInFlight.Decrement();
+			AddLoadCounters(AsyncIOPriority, -1);
 		};
 
 #if DEBUG_STREAM_CACHE
@@ -1971,6 +2012,7 @@ void FAudioChunkCache::KickOffAsyncLoad(FCacheElement* CacheElement, const FChun
 				UE_LOG(LogAudioStreamCaching, Error, TEXT("Chunk load in audio LRU cache failed."));
 				OnLoadCompleted(EAudioChunkLoadResult::ChunkOutOfBounds);
 				NumberOfLoadsInFlight.Decrement();
+				AddLoadCounters(AsyncIOPriority, -1);
 			}
 			else if (FPlatformAtomics::InterlockedCompareExchangePointer((void* volatile*)&CacheElement->ReadRequest, LocalReadRequest, nullptr) == (void*)0x1)
 			{
@@ -1983,65 +2025,12 @@ void FAudioChunkCache::KickOffAsyncLoad(FCacheElement* CacheElement, const FChun
 
 EAsyncIOPriorityAndFlags FAudioChunkCache::GetAsyncPriorityForChunk(const FChunkKey& InKey, bool bNeededForPlayback)
 {
-
-	// TODO: In the future we can add an enum to USoundWaves to tweak load priority of individual assets.
-
 	if (bNeededForPlayback)
 	{
-		switch (PlaybackRequestPriorityCVar)
-		{
-		case 4:
-		{
-			return AIOP_MIN;
-		}
-		case 3:
-		{
-			return AIOP_Low;
-		}
-		case 2:
-		{
-			return AIOP_BelowNormal;
-		}
-		case 1:
-		{
-			return AIOP_Normal;
-		}
-		case 0:
-		default:
-		{
-			return AIOP_High;
-		}
-		}
-	}
-	else
-	{
-		switch (ReadRequestPriorityCVar)
-		{
-		case 4:
-		{
-			return AIOP_MIN;
-		}
-		case 3:
-		{
-			return AIOP_Low;
-		}
-		case 2:
-		{
-			return AIOP_BelowNormal;
-		}
-		case 1:
-		{
-			return AIOP_Normal;
-		}
-		case 0:
-		default:
-		{
-			return AIOP_High;
-		}
-		}
+		return ToIoPriority(PlaybackRequestPriorityCVar);
 	}
-
-}
+	return ToIoPriority(ReadRequestPriorityCVar);
+}	
 
 void FAudioChunkCache::ExecuteOnLoadCompleteCallback(EAudioChunkLoadResult Result, const TFunction<void(EAudioChunkLoadResult)>& OnLoadCompleted, const ENamedThreads::Type& CallbackThread)
 {
