Add a hash of the Nanite Assembly data to the DDC key instead of the entire serialized data to keep the DDC keys shorter.

#rnx

diff --git a/Engine/Source/Runtime/Engine/Private/EngineTypes.cpp b/Engine/Source/Runtime/Engine/Private/EngineTypes.cpp
--- a/Engine/Source/Runtime/Engine/Private/EngineTypes.cpp
+++ b/Engine/Source/Runtime/Engine/Private/EngineTypes.cpp
@@ -4,6 +4,8 @@
 #include "UObject/UnrealType.h"
 #include "GameFramework/Actor.h"
 #include "Engine/CollisionProfile.h"
+#include "Hash/xxhash.h"
+#include "Serialization/MemoryHasher.h"
 
 #if WITH_EDITOR
 #include "Engine/Texture2D.h"
@@ -192,8 +194,15 @@ void SerializeNaniteSettingsForDDC(FArchive& Ar, FMeshNaniteSettings& NaniteSett
 	Ar << NaniteSettings.MaxEdgeLengthFactor;
 	Ar << NaniteSettings.DisplacementUVChannel;
 	
-	// TODO: Nanite-Assemblies: Changes to referenced static meshes currently won't invalidate DDC
-	NaniteSettings.NaniteAssemblyData.SerializeForDDC(Ar);
+	{
+		// TODO: Nanite-Assemblies: Changes to referenced static meshes currently won't invalidate DDC
+		TMemoryHasher<FXxHash64Builder, FXxHash64> HashAr;
+		HashAr.SetIsPersistent(true);
+		NaniteSettings.NaniteAssemblyData.SerializeForDDC(HashAr);
+		
+		uint64 AssemblyDataHash = HashAr.Finalize().Hash;
+		Ar << AssemblyDataHash;
+	}
 
 	if (NaniteSettings.bPreserveArea)
 	{
