Removed deprecated ISMC interface and fixed a small bug in the implementation of UpdateInstanceTransformById (confusing bTeleport for bMarkRenderstateDirty).

diff --git a/Engine/Source/Runtime/Engine/Classes/Components/InstancedStaticMeshComponent.h b/Engine/Source/Runtime/Engine/Classes/Components/InstancedStaticMeshComponent.h
--- a/Engine/Source/Runtime/Engine/Classes/Components/InstancedStaticMeshComponent.h
+++ b/Engine/Source/Runtime/Engine/Classes/Components/InstancedStaticMeshComponent.h
@@ -323,24 +323,6 @@ public:
 	*/
 	ENGINE_API virtual bool BatchUpdateInstancesTransforms(int32 StartInstanceIndex, const TArray<FTransform>& NewInstancesTransforms, const TArray<FTransform>& NewInstancesPrevTransforms, bool bWorldSpace = false, bool bMarkRenderStateDirty = false, bool bTeleport = false);
 
-	/**
-	* Lightweight interface to add, remove and update instances.
-	*
-	* @param AddInstanceTransforms				The transforms of the new instances to add.
-	* @param RemoveInstanceIds					The ids of the instances to remove.
-	* @param UpdateInstanceIds					The ids of the new instances to update.
-	* @param UpdateInstanceTransforms			The transforms of the new instances to update.
-	* @param UpdateInstancePreviousTransforms	The transforms of the new instances to update.
-	* @return									True on success
-	*/
-	UE_DEPRECATED(5.4, "Use the new ID-based APIs instead as this enables persistence tracking for incremental updates.")
-	ENGINE_API virtual bool UpdateInstances(
-		const TArray<int32>& UpdateInstanceIds, 
-		const TArray<FTransform>& UpdateInstanceTransforms, 
-		const TArray<FTransform>& UpdateInstancePreviousTransforms,
-		int32 NumCustomFloats,
-		const TArray<float>& CustomFloatData);
-
 	/**
 	* Update the transform for a number of instances.
 	*
@@ -563,14 +545,6 @@ public:
 	/** Deselect all instances. */
 	ENGINE_API void ClearInstanceSelection();
 
-	/** Initialize the Per Instance Render Data */
-	UE_DEPRECATED(5.4, "This does not do anything as this has been refactored.")
-	ENGINE_API void InitPerInstanceRenderData(bool InitializeFromCurrentData, FStaticMeshInstanceData* InSharedInstanceBufferData = nullptr, bool InRequireCPUAccess = false);
-
-	/** Transfers ownership of instance render data to a render thread. Instance render data will be released in scene proxy destructor or on render thread task. */
-	UE_DEPRECATED(5.4, "This does not do anything as this has been refactored.")
-	ENGINE_API void ReleasePerInstanceRenderData();
-
 	/** Precache all PSOs which can be used by the component */
 	ENGINE_API virtual void CollectPSOPrecacheData(const FPSOPrecacheParams& BasePrecachePSOParams, FMaterialInterfacePSOPrecacheParamsList& OutParams) override;
 	
@@ -581,22 +555,9 @@ public:
 
 	ENGINE_API void GetInstancesMinMaxScale(FVector& MinScale, FVector& MaxScale) const;
 
-	UE_DEPRECATED(5.4, "This does not do anything as this is controlled by the deferred update system.")
-	ENGINE_API void FlushInstanceUpdateCommands(bool bFlushInstanceUpdateCmdBuffer);
 
 	UE_DEPRECATED(5.4, "This function has been added only for the purposes of moving LWI code outside of the engine. Don't use it, it will be removed soon.")
 	ENGINE_API void OnPostPopulatePerInstanceData() { OnPostLoadPerInstanceData(); }
-
-	UE_DEPRECATED(5.4, "Use the new ID-based APIs instead as this enables persistence tracking for incremental updates.")
-	TArray<int32> PerInstanceIds;
-
-	/** Used to cache a unique identifier for each instance.  These are provided
-	*	by the interface UpdateInstances.  This is a map from unique id to index
-	*	into the PerInstanceSMData array.
-	*/
-	UE_DEPRECATED(5.4, "Use the new ID-based APIs instead as this enables persistence tracking for incremental updates.")
-	TMap<int32, int32> InstanceIdToInstanceIndexMap;
-
 	/** Request to navigation system to update only part of navmesh occupied by specified instance. */
 	ENGINE_API virtual void PartialNavigationUpdate(int32 InstanceIdx);
 
diff --git a/Engine/Source/Runtime/Engine/Private/InstancedStaticMesh.cpp b/Engine/Source/Runtime/Engine/Private/InstancedStaticMesh.cpp
--- a/Engine/Source/Runtime/Engine/Private/InstancedStaticMesh.cpp
+++ b/Engine/Source/Runtime/Engine/Private/InstancedStaticMesh.cpp
@@ -2192,10 +2192,6 @@ void UInstancedStaticMeshComponent::ApplyComponentInstanceData(FInstancedStaticM
 #endif
 }
 
-void UInstancedStaticMeshComponent::FlushInstanceUpdateCommands(bool bFlushInstanceUpdateCmdBuffer)
-{
-}
-
 bool UInstancedStaticMeshComponent::IsHLODRelevant() const
 {
 	if (GetInstanceCount() == 0)
@@ -3093,10 +3089,6 @@ FBox UInstancedStaticMeshComponent::GetStreamingBounds() const
 }
 #endif
 
-void UInstancedStaticMeshComponent::ReleasePerInstanceRenderData()
-{
-}
-
 void UInstancedStaticMeshComponent::PropagateLightingScenarioChange()
 {
 	FComponentRecreateRenderStateContext Context(this);
@@ -3501,7 +3493,7 @@ void UInstancedStaticMeshComponent::RemoveInstancesById(const TArrayView<const F
 void UInstancedStaticMeshComponent::UpdateInstanceTransformById(FPrimitiveInstanceId InstanceId, const FTransform& NewInstanceTransform, bool bWorldSpace, bool bTeleport)
 {
 	check(PrimitiveInstanceDataManager.GetMode() != FPrimitiveInstanceDataManager::EMode::ExternalLegacyData);
-	UpdateInstanceTransform(PrimitiveInstanceDataManager.IdToIndex(InstanceId), NewInstanceTransform, bWorldSpace, bTeleport);
+	UpdateInstanceTransform(PrimitiveInstanceDataManager.IdToIndex(InstanceId), NewInstanceTransform, bWorldSpace, false, bTeleport);
 }
 
 void UInstancedStaticMeshComponent::SetPreviousTransformById(FPrimitiveInstanceId InstanceId, const FTransform& NewPrevInstanceTransform, bool bWorldSpace)
@@ -4089,253 +4081,6 @@ bool UInstancedStaticMeshComponent::BatchUpdateInstancesTransforms(int32 StartIn
 	return true;
 }
 
-bool UInstancedStaticMeshComponent::UpdateInstances(
-	const TArray<int32>& UpdateInstanceIds,
-	const TArray<FTransform>& UpdateInstanceTransforms,
-	const TArray<FTransform>& UpdateInstancePreviousTransforms,
-	int32 InNumCustomDataFloats,
-	const TArray<float>& CustomFloatData)
-{
-PRAGMA_DISABLE_DEPRECATION_WARNINGS
-	CSV_SCOPED_TIMING_STAT_EXCLUSIVE(UInstancedStaticMeshComponent_UpdateInstances);
-	TRACE_CPUPROFILER_EVENT_SCOPE_STR("UInstancedStaticMeshComponent::UpdateInstances");
-
-	const float EqualTolerance = 1e-6;
-
-	if (UpdateInstanceIds.Num() != UpdateInstanceTransforms.Num())
-	{
-		UE_LOG(LogStaticMesh, Warning, TEXT("Add instances to ISM has missmatched update arrays.  UpdateInstanceIds and UpdateInstanceTransforms should match."));
-		return false;
-	}
-
-	if (UpdateInstanceTransforms.Num() != UpdateInstancePreviousTransforms.Num())
-	{
-		UE_LOG(LogStaticMesh, Warning, TEXT("Add instances to ISM has missmatched update arrays.  Current transform and previous transform arrays should match."));
-		return false;
-	}
-
-	if (InNumCustomDataFloats > 0)
-	{
-		if ((CustomFloatData.Num() / InNumCustomDataFloats) != UpdateInstanceIds.Num())
-		{
-			UE_LOG(LogStaticMesh, Warning, TEXT("CustomFloatData only contains %d entires.  Which is not enough for NumCustomFloat count %d * %d instances."), CustomFloatData.Num(), InNumCustomDataFloats, UpdateInstanceIds.Num());
-			return false;
-		}
-	}
-
-	const int32 NewNumInstances = UpdateInstanceIds.Num();
-	const int32 PreviousNumInstances = PerInstanceSMData.Num();
-
-	// If we didn't have any instances and there are no new instances
-	// skip doing any work.
-	if (PreviousNumInstances == 0 && NewNumInstances == 0)
-	{
-		return false;
-	}
-
-	Modify();
-	InvalidateCachedBounds();
-
-#if (CSV_PROFILER_STATS)
-	int32 TotalSizeUpdateBytes = 0;
-#endif
-
-	// Note: this will reset the custom data for all instances meaning they will be updated below
-	if (InNumCustomDataFloats != NumCustomDataFloats)
-	{
-		SetNumCustomDataFloats(InNumCustomDataFloats);
-	}
-	const bool bHasCustomFloatData = NumCustomDataFloats > 0;
-
-	// if we already have values we need to update, remove, and add.
-	TArray<int32> OldInstanceIds(PerInstanceIds);
-	TArray<int32> AddedInstances;
-
-	// Navigation update required if the component is relevant to navigation and registered.
-	const bool bNavigationUpdateRequired = bNavigationRelevant && IsRegistered();
-	const bool bPartialNavigationUpdateRequired = bNavigationUpdateRequired && SupportsPartialNavigationUpdate();
-
-	TArray<FTransform> NavigationUpdateTransforms;
-	if (bPartialNavigationUpdateRequired)
-	{
-		// Reserve enough space for previous and new transforms
-		NavigationUpdateTransforms.Reserve(UpdateInstanceIds.Num() * 2);
-	}
-
-	// Apply updates
-	for (int32 i = 0; i < UpdateInstanceIds.Num(); ++i)
-	{
-		// This is an update.
-		const int32 InstanceId = UpdateInstanceIds[i];
-		int32* pInstanceIndex = InstanceIdToInstanceIndexMap.Find(InstanceId);
-		if (pInstanceIndex != nullptr)
-		{
-			const int32 InstanceIndex = *pInstanceIndex;
-
-			// Determine what data changed.
-
-			// Did the transform actually change?
-			// Explicitly check the component's previous transform for this instance because this 
-			// is the position we were previously rendered at.
-			FMatrix NewInstanceTransformMatrix = UpdateInstanceTransforms[i].ToMatrixWithScale();
-			if (Mobility == EComponentMobility::Movable && !PerInstanceSMData[InstanceIndex].Transform.Equals(NewInstanceTransformMatrix, EqualTolerance))
-			{
-				// Update the component's data in place.
-				PerInstanceSMData[InstanceIndex] = NewInstanceTransformMatrix;
-				PerInstancePrevTransform[InstanceIndex] = UpdateInstancePreviousTransforms[i].ToMatrixWithScale();
-
-
-				PrimitiveInstanceDataManager.TransformChanged(InstanceIndex);
-
-				// Append instance's previous and new transforms to dirty both areas
-				if (bPartialNavigationUpdateRequired)
-				{
-					NavigationUpdateTransforms.Append(
-						{
-							UpdateInstancePreviousTransforms[i] * GetComponentTransform(),
-							UpdateInstanceTransforms[i] * GetComponentTransform()
-						});
-				}
-
-#if (CSV_PROFILER_STATS)
-				// We are updating a current and previous transform.
-				TotalSizeUpdateBytes += (sizeof(FTransform) + sizeof(FTransform));
-#endif
-			}
-
-			// Did the custom data actually change?
-			if (bHasCustomFloatData && Mobility != EComponentMobility::Static)
-			{
-				const int32 CustomDataOffset = InstanceIndex * NumCustomDataFloats;
-				const int32 SrcCustomDataOffset = i * NumCustomDataFloats;
-				for (int32 FloatIndex = 0; FloatIndex < NumCustomDataFloats; ++FloatIndex)
-				{
-					if (FMath::Abs(CustomFloatData[SrcCustomDataOffset + FloatIndex] - PerInstanceSMCustomData[CustomDataOffset + FloatIndex]) > EqualTolerance)
-					{
-						// Update the component's data in place.
-						FMemory::Memcpy(&PerInstanceSMCustomData[CustomDataOffset], &CustomFloatData[SrcCustomDataOffset], NumCustomDataFloats * sizeof(float));
-
-						PrimitiveInstanceDataManager.CustomDataChanged(InstanceIndex);
-#if (CSV_PROFILER_STATS)
-						TotalSizeUpdateBytes += NumCustomDataFloats * sizeof(float);
-#endif
-						break;
-					}
-				}
-			}
-
-			// Mark this index as one we've updated.
-			OldInstanceIds[InstanceIndex] = INDEX_NONE;
-		}
-		else
-		{
-			AddedInstances.Add(i);
-		}
-	}
-
-	// Remove instances
-	if (bPartialNavigationUpdateRequired)
-	{
-		// Reserve more space if need for instances to remove
-		NavigationUpdateTransforms.Reserve(NavigationUpdateTransforms.Num() + OldInstanceIds.Num());
-	}
-
-	for (int32 InstanceIndex = 0; InstanceIndex < OldInstanceIds.Num();)
-	{
-		if (OldInstanceIds[InstanceIndex] != INDEX_NONE)
-		{
-			// Append old instances transform
-			if (bPartialNavigationUpdateRequired)
-			{
-				NavigationUpdateTransforms.Add(FTransform(PerInstanceSMData[InstanceIndex].Transform) * GetComponentTransform());
-			}
-
-			// TODO: Move this to common helper function such that all data remove goes through one place in the code.
-			PrimitiveInstanceDataManager.RemoveAtSwap(InstanceIndex);
-			PerInstanceSMData.RemoveAtSwap(InstanceIndex, EAllowShrinking::No);
-			PerInstancePrevTransform.RemoveAtSwap(InstanceIndex, EAllowShrinking::No);
-			PerInstanceIds.RemoveAtSwap(InstanceIndex, EAllowShrinking::No);
-
-			// Only remove the custom float data from this instance if it previously had it.
-			if (bHasCustomFloatData)
-			{
-				PerInstanceSMCustomData.RemoveAtSwap((InstanceIndex * NumCustomDataFloats), NumCustomDataFloats, EAllowShrinking::No);
-			}
-
-			OldInstanceIds.RemoveAtSwap(InstanceIndex, EAllowShrinking::No);
-		}
-		else
-		{
-			++InstanceIndex;
-		}
-	}
-
-	if (bPartialNavigationUpdateRequired)
-	{
-		// Reserve enough space for instances to add
-		NavigationUpdateTransforms.Reserve(NavigationUpdateTransforms.Num() + AddedInstances.Num());
-	}
-
-	// Add new instances to the component's data.
-	for (int32 Index : AddedInstances)
-	{
-		// TODO: Move this to common helper function such that all data add goes through one place in the code.
-		PrimitiveInstanceDataManager.Add(PerInstanceSMData.Num(), false);
-		PerInstanceSMData.Add(UpdateInstanceTransforms[Index].ToMatrixWithScale());
-		PerInstancePrevTransform.Add(UpdateInstancePreviousTransforms[Index].ToMatrixWithScale());
-		PerInstanceIds.Add(UpdateInstanceIds[Index]);
-
-		if (bHasCustomFloatData)
-		{
-			const int32 SrcCustomDataOffset = Index * NumCustomDataFloats;
-			const int32 CustomDataDestIndex = PerInstanceSMCustomData.AddUninitialized(NumCustomDataFloats);
-
-			FMemory::Memcpy(&PerInstanceSMCustomData[CustomDataDestIndex], &CustomFloatData[SrcCustomDataOffset], NumCustomDataFloats * sizeof(float));
-		}
-	}
-
-	// Rebuild the mapping from ID to InstanceIndex.
-	InstanceIdToInstanceIndexMap.Reset();
-	for (int32 InstanceIndex = 0; InstanceIndex < PerInstanceIds.Num(); ++InstanceIndex)
-	{
-		InstanceIdToInstanceIndexMap.Add(PerInstanceIds[InstanceIndex], InstanceIndex);
-	}
-
-	// Ensure our data is in sync.
-	check(PerInstanceSMData.Num() == PerInstancePrevTransform.Num());
-	check(PerInstanceSMData.Num() == PerInstanceIds.Num());
-	check(PerInstanceSMCustomData.Num() == (NumCustomDataFloats * PerInstanceSMData.Num()));
-
-	if (bNavigationUpdateRequired)
-	{
-		if (bPartialNavigationUpdateRequired)
-		{
-			PartialNavigationUpdates(NavigationUpdateTransforms);
-		}
-		else
-		{
-			FullNavigationUpdate();
-		}
-	}
-
-#if (CSV_PROFILER_STATS)
-	const int32 TotalSizeBytes = (UpdateInstanceTransforms.Num() * UpdateInstanceTransforms.GetTypeSize()) +
-								 (UpdateInstancePreviousTransforms.Num() * UpdateInstancePreviousTransforms.GetTypeSize()) +
-								 (CustomFloatData.Num() * CustomFloatData.GetTypeSize());
-
-	static const FName NumInstancesUpdated(TEXT("NumInstancesUpdated"));
-	static const FName NumBytesSubmitted(TEXT("NumBytesSubmitted"));
-	static const FName NumBytesUpdated(TEXT("NumBytesUpdated"));
-
-	FCsvProfiler::Get()->RecordCustomStat(NumInstancesUpdated, CSV_CATEGORY_INDEX(InstancedStaticMeshComponent), PerInstanceSMData.Num(), ECsvCustomStatOp::Accumulate);
-	FCsvProfiler::Get()->RecordCustomStat(NumBytesSubmitted, CSV_CATEGORY_INDEX(InstancedStaticMeshComponent), TotalSizeBytes, ECsvCustomStatOp::Accumulate);
-	FCsvProfiler::Get()->RecordCustomStat(NumBytesUpdated, CSV_CATEGORY_INDEX(InstancedStaticMeshComponent), TotalSizeUpdateBytes, ECsvCustomStatOp::Accumulate);
-#endif
-
-PRAGMA_ENABLE_DEPRECATION_WARNINGS
-	return true;
-}
-
 bool UInstancedStaticMeshComponent::BatchUpdateInstancesTransforms(int32 StartInstanceIndex, const TArray<FTransform>& NewInstancesTransforms, bool bWorldSpace, bool bMarkRenderStateDirty, bool bTeleport)
 {
 	return BatchUpdateInstancesTransformsInternal(StartInstanceIndex, MakeArrayView(NewInstancesTransforms), bWorldSpace, bMarkRenderStateDirty, bTeleport);
@@ -4873,10 +4618,6 @@ void UInstancedStaticMeshComponent::GetInstancesMinMaxScale(FVector& MinScale, F
 	}
 }
 
-void UInstancedStaticMeshComponent::InitPerInstanceRenderData(bool InitializeFromCurrentData, FStaticMeshInstanceData* InSharedInstanceBufferData, bool InRequireCPUAccess)
-{
-}
-
 void UInstancedStaticMeshComponent::OnRegister()
 {
 	TRACE_CPUPROFILER_EVENT_SCOPE(UInstancedStaticMeshComponent::OnRegister())
