Change runtime virtual texture material validation to catch more use cases.
The old validation tried to visit nodes to look for issues, but that didn't fully take into account reroutes/material functions etc.
The new validation leans on the translator to validate.
The new validation catches an error where virtual textures are sampled for the normal property and then fed to runtime virtual texture offset via the PixelNormalWS node.
Disabling validation for the condition is handled with a flag per material exposed to the user.
The validation is enabled by default BUT because it is stricter than the old validation it can't just be turned on for old materials.
So existing materials are grandfathered in with the validation disabled.

diff --git a/Engine/Source/Runtime/Core/Public/UObject/UE5MainStreamObjectVersion.h b/Engine/Source/Runtime/Core/Public/UObject/UE5MainStreamObjectVersion.h
--- a/Engine/Source/Runtime/Core/Public/UObject/UE5MainStreamObjectVersion.h
+++ b/Engine/Source/Runtime/Core/Public/UObject/UE5MainStreamObjectVersion.h
@@ -376,6 +376,9 @@ struct FUE5MainStreamObjectVersion
 		// Added support for InstanceDataManagerSerialization and changed format for the instances to FTransform3f
 		AnimBankInstanceDataSerializationV2,
 	
+		// Added new material compilation validation for runtime virtual textures
+		RuntimeVirtualTextureMaterialValidation,
+		
 		// -----<new versions can be added above this line>-------------------------------------------------
 		VersionPlusOne,
 		LatestVersion = VersionPlusOne - 1
diff --git a/Engine/Source/Runtime/Engine/Private/Materials/HLSLMaterialTranslator.cpp b/Engine/Source/Runtime/Engine/Private/Materials/HLSLMaterialTranslator.cpp
--- a/Engine/Source/Runtime/Engine/Private/Materials/HLSLMaterialTranslator.cpp
+++ b/Engine/Source/Runtime/Engine/Private/Materials/HLSLMaterialTranslator.cpp
@@ -7309,6 +7309,11 @@ int32 FHLSLMaterialTranslator::TextureSample(
 				, VirtualTextureIndex, VirtualTextureIndex);
 		}
 
+		const bool bIsStrictValidation = EnumHasAnyFlags(Material->GetMaterialTranslateValidationFlags(), EMaterialTranslateValidationFlags::Strict_RuntimeVirtualTexture);
+		if (bIsStrictValidation && bIsInRuntimeVirtualTextureOutput)
+		{
+			return Errorf(TEXT("Virtual Texture samples are not allowed during Runtime Virtual Texture Output."));
+		}
 		if (MaterialProperty == MP_Normal)
 		{
 			bUsesVirtualTextureSampleForNormalProperty = true;
@@ -11542,6 +11547,11 @@ int32 FHLSLMaterialTranslator::PixelNormalWS()
 	{
 		return Errorf(TEXT("Invalid node PixelNormalWS used for Normal input."));
 	}
+	const bool bIsStrictValidation = EnumHasAnyFlags(Material->GetMaterialTranslateValidationFlags(), EMaterialTranslateValidationFlags::Strict_RuntimeVirtualTexture);
+	if (bIsStrictValidation && bIsInRuntimeVirtualTextureOutput && bUsesVirtualTextureSampleForNormalProperty)
+	{
+		return Errorf(TEXT("PixelNormalWS samples Virtual Textures and so is not allowed during Runtime Virtual Texture Output."));
+	}
 	if (ShaderFrequency != SF_Vertex)
 	{
 		bUsesTransformVector = true;
diff --git a/Engine/Source/Runtime/Engine/Private/Materials/Material.cpp b/Engine/Source/Runtime/Engine/Private/Materials/Material.cpp
--- a/Engine/Source/Runtime/Engine/Private/Materials/Material.cpp
+++ b/Engine/Source/Runtime/Engine/Private/Materials/Material.cpp
@@ -1100,6 +1100,9 @@ void UMaterial::SetInitialValues()
 	bAllowDevelopmentShaderCompile = true;
 	bIsMaterialEditorStatsMaterial = false;
 
+	bRelaxRuntimeVirtualTextureRestrictions = false;
+	bRelaxRuntimeVirtualTextureRestrictionsForOldContent = false;
+
 	RefractionMethod = RM_None;
 	RefractionCoverageMode = RCM_CoverageAccountedFor;
 	PixelDepthOffsetMode = PDOM_AlongCameraVector;
@@ -2331,6 +2334,16 @@ void UMaterial::GetShaderTags(TArray<FName>& OutShaderTags) const
 
 	OutShaderTags.Append(CachedExpressionData->EditorOnlyData->ShaderTags);
 }
+
+EMaterialTranslateValidationFlags UMaterial::GetMaterialTranslateValidationFlags() const
+{
+	EMaterialTranslateValidationFlags Flags = EMaterialTranslateValidationFlags::None;
+	if (!bRelaxRuntimeVirtualTextureRestrictions && !bRelaxRuntimeVirtualTextureRestrictionsForOldContent)
+	{
+		Flags |= EMaterialTranslateValidationFlags::Strict_RuntimeVirtualTexture;
+	}
+	return Flags;
+}
 #endif
 
 UMaterialInterface* UMaterial::GetNaniteOverride(TMicRecursionGuard RecursionGuard) const
@@ -4323,6 +4336,9 @@ void UMaterial::PostLoad()
 		EditorOnly->ExpressionCollection.ExpressionExecEnd = MoveTemp(ExpressionExecEnd_DEPRECATED);
 	}
 
+	// Can't use new material compilation validations on existing materials.
+	bRelaxRuntimeVirtualTextureRestrictionsForOldContent = (UE5MainVer < FUE5MainStreamObjectVersion::RuntimeVirtualTextureMaterialValidation);
+
 #endif // WITH_EDITORONLY_DATA
 
 	if (!IsDefaultMaterial())
diff --git a/Engine/Source/Runtime/Engine/Private/Materials/MaterialExpressions.cpp b/Engine/Source/Runtime/Engine/Private/Materials/MaterialExpressions.cpp
--- a/Engine/Source/Runtime/Engine/Private/Materials/MaterialExpressions.cpp
+++ b/Engine/Source/Runtime/Engine/Private/Materials/MaterialExpressions.cpp
@@ -3067,29 +3067,6 @@ UMaterialExpressionRuntimeVirtualTextureOutput::UMaterialExpressionRuntimeVirtua
 
 #if WITH_EDITOR
 
-bool ValidateRuntimeVirtualTextureOutput(FMaterialCompiler* Compiler, FExpressionInput& InExpressionInput)
-{
-	// Note that this doesn't traverse through material layers and functions.
-	if (InExpressionInput.IsConnected())
-	{
-		TArray<UMaterialExpression*> InputExpressions;
-		InExpressionInput.Expression->GetAllInputExpressions(InputExpressions);
-		for (UMaterialExpression* Expression : InputExpressions)
-		{
-			UMaterialExpressionRuntimeVirtualTextureSample* RuntimeVirtualTextureSampleExpression = Cast<UMaterialExpressionRuntimeVirtualTextureSample>(Expression);
-			if (RuntimeVirtualTextureSampleExpression != nullptr)
-			{
-				const FString VirtualTextureName = RuntimeVirtualTextureSampleExpression->VirtualTexture ? RuntimeVirtualTextureSampleExpression->VirtualTexture->GetName() : RuntimeVirtualTextureSampleExpression->GetName();
-
-				Compiler->Errorf(TEXT("Sampling a Runtime Virtual Texture (%s) when writing Runtime Virtual Texture output is not supported."), *VirtualTextureName);
-				return false;
-			}
-		}
-	}
-
-	return true;
-}
-
 int32 UMaterialExpressionRuntimeVirtualTextureOutput::Compile(class FMaterialCompiler* Compiler, int32 OutputIndex)
 {
 	int32 Result = INDEX_NONE;
@@ -3105,54 +3082,45 @@ int32 UMaterialExpressionRuntimeVirtualTextureOutput::Compile(class FMaterialCom
 		{
 			CodeInput = BaseColor.IsConnected() ? BaseColor.Compile(Compiler) : Compiler->Constant3(0.f, 0.f, 0.f);
 			OutputAttributeMask |= BaseColor.IsConnected() ? (1 << (uint8)ERuntimeVirtualTextureAttributeType::BaseColor) : 0;
-			ValidateRuntimeVirtualTextureOutput(Compiler, BaseColor);
 		}
 		else if (OutputIndex == 1)
 		{
 			CodeInput = Specular.IsConnected() ? Specular.Compile(Compiler) : Compiler->Constant(0.5f);
 			OutputAttributeMask |= Specular.IsConnected() ? (1 << (uint8)ERuntimeVirtualTextureAttributeType::Specular) : 0;
-			ValidateRuntimeVirtualTextureOutput(Compiler, Specular);
 		}
 		else if (OutputIndex == 2)
 		{
 			CodeInput = Roughness.IsConnected() ? Roughness.Compile(Compiler) : Compiler->Constant(0.5f);
 			OutputAttributeMask |= Roughness.IsConnected() ? (1 << (uint8)ERuntimeVirtualTextureAttributeType::Roughness) : 0;
-			ValidateRuntimeVirtualTextureOutput(Compiler, Roughness);
 		}
 		else if (OutputIndex == 3)
 		{
 			CodeInput = Normal.IsConnected() ? Normal.Compile(Compiler) : Compiler->Constant3(0.f, 0.f, 1.f);
 			OutputAttributeMask |= Normal.IsConnected() ? (1 << (uint8)ERuntimeVirtualTextureAttributeType::Normal) : 0;
-			ValidateRuntimeVirtualTextureOutput(Compiler, Normal);
 		}
 		else if (OutputIndex == 4)
 		{
 			CodeInput = WorldHeight.IsConnected() ? WorldHeight.Compile(Compiler) : Compiler->Constant(0.f);
 			OutputAttributeMask |= WorldHeight.IsConnected() ? (1 << (uint8)ERuntimeVirtualTextureAttributeType::WorldHeight) : 0;
-			ValidateRuntimeVirtualTextureOutput(Compiler, WorldHeight);
 		}
 		else if (OutputIndex == 5)
 		{
 			CodeInput = Opacity.IsConnected() ? Opacity.Compile(Compiler) : Compiler->Constant(1.f);
-			ValidateRuntimeVirtualTextureOutput(Compiler, Opacity);
 		}
 		else if (OutputIndex == 6)
 		{
 			CodeInput = Mask.IsConnected() ? Mask.Compile(Compiler) : Compiler->Constant(1.f);
 			OutputAttributeMask |= Mask.IsConnected() ? (1 << (uint8)ERuntimeVirtualTextureAttributeType::Mask) : 0;
-			ValidateRuntimeVirtualTextureOutput(Compiler, Mask);
 		}
 		else if (OutputIndex == 7)
 		{
 			CodeInput = Displacement.IsConnected() ? Displacement.Compile(Compiler) : Compiler->Constant(0.f);
 			OutputAttributeMask |= Displacement.IsConnected() ? (1 << (uint8)ERuntimeVirtualTextureAttributeType::Displacement) : 0;
-			ValidateRuntimeVirtualTextureOutput(Compiler, Displacement);
 		}
 		else if (OutputIndex == 8)
 		{
 			CodeInput = Mask4.IsConnected() ? Mask4.Compile(Compiler) : Compiler->Constant4(0.f, 0.f, 0.f, 0.f);
 			OutputAttributeMask |= Mask4.IsConnected() ? (1 << (uint8)ERuntimeVirtualTextureAttributeType::Mask4) : 0;
-			ValidateRuntimeVirtualTextureOutput(Compiler, Mask4);
 		}
 
 		Result = Compiler->CustomOutput(this, OutputIndex, CodeInput);
diff --git a/Engine/Source/Runtime/Engine/Private/Materials/MaterialShared.cpp b/Engine/Source/Runtime/Engine/Private/Materials/MaterialShared.cpp
--- a/Engine/Source/Runtime/Engine/Private/Materials/MaterialShared.cpp
+++ b/Engine/Source/Runtime/Engine/Private/Materials/MaterialShared.cpp
@@ -2371,6 +2371,11 @@ void FMaterialResource::GetShaderTags(TArray<FName>& OutShaderTags) const
 	OutShaderTags.Append(CachedExpressionData.EditorOnlyData->ShaderTags);
 }
 
+EMaterialTranslateValidationFlags FMaterialResource::GetMaterialTranslateValidationFlags() const
+{
+	return Material->GetMaterialTranslateValidationFlags();
+}
+
 void FMaterialResource::NotifyCompilationFinished()
 {
 	UMaterial::NotifyCompilationFinished(MaterialInstance ? (UMaterialInterface*)MaterialInstance : (UMaterialInterface*)Material);
diff --git a/Engine/Source/Runtime/Engine/Public/MaterialShared.h b/Engine/Source/Runtime/Engine/Public/MaterialShared.h
--- a/Engine/Source/Runtime/Engine/Public/MaterialShared.h
+++ b/Engine/Source/Runtime/Engine/Public/MaterialShared.h
@@ -1806,7 +1806,13 @@ private:
 	friend class UE::Cook::FODSCClientData;
 };
 
-
+/** Flags that define the behavior of material translation validation. */
+enum class EMaterialTranslateValidationFlags : uint32
+{
+	None = 0,
+	Strict_RuntimeVirtualTexture = 1 << 0,
+};
+ENUM_CLASS_FLAGS(EMaterialTranslateValidationFlags)
 
 /** 
  * Enum that contains entries for the ways that material properties need to be compiled.
@@ -2284,6 +2290,7 @@ public:
 	virtual int32 GetNeuralProfileId() const { return INDEX_NONE; }
 	virtual bool IsVariableRateShadingAllowed() const { return true; }
 	virtual void GetShaderTags(TArray<FName>& OutShaderTags) const {};
+	virtual EMaterialTranslateValidationFlags GetMaterialTranslateValidationFlags() const { return EMaterialTranslateValidationFlags::None; }
 
 	/**
 	 * Should shaders compiled for this material be saved to disk?
@@ -3067,6 +3074,7 @@ public:
 
 #if WITH_EDITOR
 	ENGINE_API void GetShaderTags(TArray<FName>& OutShaderTags) const override;
+	ENGINE_API EMaterialTranslateValidationFlags GetMaterialTranslateValidationFlags() const override;
 	ENGINE_API virtual void NotifyCompilationFinished() override;
 	/** Allows to associate the shader resources with the asset for load order. */
 	ENGINE_API virtual FName GetAssetPath() const override;
diff --git a/Engine/Source/Runtime/Engine/Public/Materials/Material.h b/Engine/Source/Runtime/Engine/Public/Materials/Material.h
--- a/Engine/Source/Runtime/Engine/Public/Materials/Material.h
+++ b/Engine/Source/Runtime/Engine/Public/Materials/Material.h
@@ -25,6 +25,7 @@
 #include "Material.generated.h"
 
 enum EMaterialDomain : int;
+enum class EMaterialTranslateValidationFlags : uint32;
 class ITargetPlatform;
 class UMaterialExpressionComment;
 class UPhysicalMaterial;
@@ -1002,6 +1003,16 @@ public:
 	/** true if this is a special material used for stats by the material editor. */
 	UPROPERTY(transient, duplicatetransient)
 	uint8 bIsMaterialEditorStatsMaterial : 1;
+
+	/** 
+	 * Relax restrictions on Runtime Virtual Texture output compilation.
+	 * Normally we give a compilation error whenever a Runtime Virtual Texture Output reads from a Virtual Texture. 
+	 * This is because there is no feedback from Runtime Virtual Texture Output passes and no guarantee that Virtual Textures will be ready.
+	 * If allowed then Runtime Virtual Texture pages can contain stale low resolution data that is not evicted until the pages are flushed.
+	 */
+	UPROPERTY(EditAnywhere, Category = VirtualTexture)
+	uint8 bRelaxRuntimeVirtualTextureRestrictions : 1;
+	uint8 bRelaxRuntimeVirtualTextureRestrictionsForOldContent : 1;
 	
 	/** Where the node is inserted in the (post processing) graph, only used if domain is PostProcess */
 	UPROPERTY(EditAnywhere, BlueprintReadWrite, Category=PostProcessMaterial, meta=(DisplayName = "Blendable Location"))
@@ -1937,9 +1948,10 @@ public:
 	ENGINE_API void SetRenderTracePhysicalMaterialOutputs(TArrayView<TObjectPtr<UPhysicalMaterial>> PhysicalMaterials);
 
 #if WITH_EDITOR
-	/*
-	*/
-	ENGINE_API void GetShaderTags(TArray<FName>& OutShaderTags) const ;
+	/** Get collection of shader tags from this material's expressions. */
+	ENGINE_API void GetShaderTags(TArray<FName>& OutShaderTags) const;
+	/** Get user flags on the material that influence translation validation errors. */
+	ENGINE_API EMaterialTranslateValidationFlags GetMaterialTranslateValidationFlags() const;
 #endif
 
 private:
