Fix typo in SceneComponent.cpp
#rnx

diff --git a/Engine/Source/Runtime/Engine/Private/Components/SceneComponent.cpp b/Engine/Source/Runtime/Engine/Private/Components/SceneComponent.cpp
--- a/Engine/Source/Runtime/Engine/Private/Components/SceneComponent.cpp
+++ b/Engine/Source/Runtime/Engine/Private/Components/SceneComponent.cpp
@@ -494,7 +494,7 @@ static bool CheckDescendantsAreAlsoCulledForTarget(USceneComponent const* SceneC
 		}
 		else
 		{
-			UE_LOG(LogSceneComponent, Display, TEXT("Scene component %s will eot be culled"), *GetPathNameSafe(ChildSceneComponent));
+			UE_LOG(LogSceneComponent, Display, TEXT("Scene component %s will be culled"), *GetPathNameSafe(ChildSceneComponent));
 		}
 	}
 
