[RHI] Add resource name to FRenderResource to identify what buffer was released erroneously.

This improves the assertion when a render resource is destructed before it has been explicitly released.
The new FName is only added when RHI resource information is available (RHI_ENABLE_RESOURCE_INFO).

#rnx

diff --git a/Engine/Source/Runtime/Engine/Private/Animation/SkeletalMeshVertexAttribute.cpp b/Engine/Source/Runtime/Engine/Private/Animation/SkeletalMeshVertexAttribute.cpp
--- a/Engine/Source/Runtime/Engine/Private/Animation/SkeletalMeshVertexAttribute.cpp
+++ b/Engine/Source/Runtime/Engine/Private/Animation/SkeletalMeshVertexAttribute.cpp
@@ -117,6 +117,10 @@ int32 FSkeletalMeshVertexAttributeRenderData::GetResourceSize() const
 	return Size;
 }
 
+namespace SkeletalMeshResourceNames
+{
+	static const FName SkeletalMeshAttributeVertexBufferName(TEXT("FSkeletalMeshVertexAttributeRenderData.FSkeletalMeshAttributeVertexBuffer"));
+} // SkeletalMeshResourceNames
 
 void FSkeletalMeshVertexAttributeRenderData::Serialize(FArchive& Ar)
 {
@@ -132,6 +136,7 @@ void FSkeletalMeshVertexAttributeRenderData::Serialize(FArchive& Ar)
 			FSkeletalMeshAttributeVertexBuffer* Buffer = new FSkeletalMeshAttributeVertexBuffer;
 			Buffer->Serialize(Ar);
 			Buffers.Add(Name, Buffer);
+			Buffer->SetResourceName(SkeletalMeshResourceNames::SkeletalMeshAttributeVertexBufferName);
 		}
 	}
 	else
diff --git a/Engine/Source/Runtime/Engine/Private/Rendering/SkinWeightVertexBuffer.cpp b/Engine/Source/Runtime/Engine/Private/Rendering/SkinWeightVertexBuffer.cpp
--- a/Engine/Source/Runtime/Engine/Private/Rendering/SkinWeightVertexBuffer.cpp
+++ b/Engine/Source/Runtime/Engine/Private/Rendering/SkinWeightVertexBuffer.cpp
@@ -612,8 +612,16 @@ void FSkinWeightDataVertexBuffer::CopyDataFromBuffer(const uint8* InSkinWeightDa
 FSkinWeightVertexBuffer
 -----------------------------------------------------------------------------*/
 
+namespace SkeletalMeshResourceNames
+{
+	static const FName DataVertexBufferName(TEXT("FSkinWeightVertexBuffer.DataVertexBuffer"));
+	static const FName LookupVertexBufferName(TEXT("FSkinWeightVertexBuffer.LookupVertexBuffer"));
+} // SkeletalMeshResourceNames
+
 FSkinWeightVertexBuffer::FSkinWeightVertexBuffer()
 {
+	DataVertexBuffer.SetResourceName(SkeletalMeshResourceNames::DataVertexBufferName);
+	LookupVertexBuffer.SetResourceName(SkeletalMeshResourceNames::LookupVertexBufferName);
 }
 
 FSkinWeightVertexBuffer::FSkinWeightVertexBuffer(const FSkinWeightVertexBuffer& Other)
diff --git a/Engine/Source/Runtime/Engine/Private/SkeletalMeshLODRenderData.cpp b/Engine/Source/Runtime/Engine/Private/SkeletalMeshLODRenderData.cpp
--- a/Engine/Source/Runtime/Engine/Private/SkeletalMeshLODRenderData.cpp
+++ b/Engine/Source/Runtime/Engine/Private/SkeletalMeshLODRenderData.cpp
@@ -297,7 +297,7 @@ void FSkeletalMeshLODRenderData::InitResources(bool bNeedsVertexColors, int32 LO
 
 	VertexAttributeBuffers.InitResources();
 	
-	BeginInitResource(OwnerName, &HalfEdgeBuffer, &UE::RenderCommandPipe::SkeletalMesh);	
+	BeginInitResource(OwnerName, &HalfEdgeBuffer, &UE::RenderCommandPipe::SkeletalMesh);
 
 #if RHI_RAYTRACING
 	if (IsRayTracingAllowed())
@@ -726,6 +726,18 @@ private:
 	int64 Size;
 };
 
+namespace SkeletalMeshResourceNames
+{
+	static const FName MultiSizeIndexContainerName(TEXT("FSkeletalMeshLODRenderData.MultiSizeIndexContainer"));
+	static const FName StaticPositionVertexBufferName(TEXT("FSkeletalMeshLODRenderData.StaticVertexBuffers.PositionVertexBuffer"));
+	static const FName StaticMeshVertexBufferName(TEXT("FSkeletalMeshLODRenderData.StaticVertexBuffers.StaticMeshVertexBuffer"));
+	static const FName StaticColorVertexBufferName(TEXT("FSkeletalMeshLODRenderData.StaticVertexBuffers.ColorVertexBuffer"));
+	static const FName ClothVertexBufferName(TEXT("FSkeletalMeshLODRenderData.ClothVertexBuffer"));
+	static const FName SourceRayTracingGeometryName(TEXT("FSkeletalMeshLODRenderData.SourceRayTracingGeometry"));
+	static const FName MorphTargetVertexInfoBuffersName(TEXT("FSkeletalMeshLODRenderData.MorphTargetVertexInfoBuffers"));
+	static const FName HalfEdgeBufferName(TEXT("FSkeletalMeshLODRenderData.HalfEdgeBuffer"));
+} // SkeletalMeshResourceNames
+
 void FSkeletalMeshLODRenderData::SerializeStreamedData(FArchive& Ar, USkinnedAsset* Owner, int32 LODIdx, uint8 ClassDataStripFlags, bool bNeedsCPUAccess, bool bForceKeepCPUResources)
 {
 	Ar.UsingCustomVersion(FUE5PrivateFrostyStreamObjectVersion::GUID);
@@ -738,6 +750,7 @@ void FSkeletalMeshLODRenderData::SerializeStreamedData(FArchive& Ar, USkinnedAss
 	//Ar << RenderSections;
 
 	MultiSizeIndexContainer.Serialize(Ar, bNeedsCPUAccess);
+	MultiSizeIndexContainer.GetIndexBuffer()->SetResourceName(SkeletalMeshResourceNames::MultiSizeIndexContainerName);
 
 	if (Ar.IsLoading())
 	{
@@ -745,12 +758,15 @@ void FSkeletalMeshLODRenderData::SerializeStreamedData(FArchive& Ar, USkinnedAss
 	}
 
 	StaticVertexBuffers.PositionVertexBuffer.Serialize(Ar, bNeedsCPUAccess);
+	StaticVertexBuffers.PositionVertexBuffer.SetResourceName(SkeletalMeshResourceNames::StaticPositionVertexBufferName);
 	StaticVertexBuffers.StaticMeshVertexBuffer.Serialize(Ar, bNeedsCPUAccess);
+	StaticVertexBuffers.StaticMeshVertexBuffer.SetResourceName(SkeletalMeshResourceNames::StaticMeshVertexBufferName);
 	Ar << SkinWeightVertexBuffer;
 
 	if (Owner && Owner->GetHasVertexColors())
 	{
 		StaticVertexBuffers.ColorVertexBuffer.Serialize(Ar, bForceKeepCPUResources);
+		StaticVertexBuffers.ColorVertexBuffer.SetResourceName(SkeletalMeshResourceNames::StaticColorVertexBufferName);
 	}
 
 	if (Ar.IsLoading() && Ar.CustomVer(FUE5ReleaseStreamObjectVersion::GUID) < FUE5ReleaseStreamObjectVersion::RemovingTessellation && !StripFlags.IsClassDataStripped(CDSF_AdjacencyData_DEPRECATED))
@@ -762,6 +778,7 @@ void FSkeletalMeshLODRenderData::SerializeStreamedData(FArchive& Ar, USkinnedAss
 	if (HasClothData())
 	{
 		Ar << ClothVertexBuffer;
+		ClothVertexBuffer.SetResourceName(SkeletalMeshResourceNames::ClothVertexBufferName);
 	}
 
 	Ar << SkinWeightProfilesData;
@@ -772,6 +789,7 @@ void FSkeletalMeshLODRenderData::SerializeStreamedData(FArchive& Ar, USkinnedAss
 		Owner->SetSkinWeightProfilesData(LODIdx, SkinWeightProfilesData);
 	}
 	Ar << SourceRayTracingGeometry.RawData;
+	SourceRayTracingGeometry.SetResourceName(SkeletalMeshResourceNames::SourceRayTracingGeometryName);
 
 	// Determine if morph target data should be cooked out for this platform
 	EShaderPlatform MorphTargetShaderPlatform = GMaxRHIShaderPlatform;
@@ -840,6 +858,7 @@ void FSkeletalMeshLODRenderData::SerializeStreamedData(FArchive& Ar, USkinnedAss
 #endif
 			{
 				Ar << MorphTargetVertexInfoBuffers;
+				MorphTargetVertexInfoBuffers.SetResourceName(SkeletalMeshResourceNames::MorphTargetVertexInfoBuffersName);
 
 #if DO_CHECK
 				const FSkeletalMeshLODInfo* SkeletalMeshLODInfo = Owner->GetLODInfo(LODIdx);
@@ -868,6 +887,7 @@ void FSkeletalMeshLODRenderData::SerializeStreamedData(FArchive& Ar, USkinnedAss
 		if (!MeshDeformerStripFlags.IsClassDataStripped(MeshDeformerStripFlag))
 		{
 			Ar << HalfEdgeBuffer;
+			HalfEdgeBuffer.SetResourceName(SkeletalMeshResourceNames::HalfEdgeBufferName);
 		}
 	}
 	
diff --git a/Engine/Source/Runtime/Engine/Private/TextureRenderTarget2D.cpp b/Engine/Source/Runtime/Engine/Private/TextureRenderTarget2D.cpp
--- a/Engine/Source/Runtime/Engine/Private/TextureRenderTarget2D.cpp
+++ b/Engine/Source/Runtime/Engine/Private/TextureRenderTarget2D.cpp
@@ -529,12 +529,12 @@ void FTextureRenderTarget2DResource::InitRHI(FRHICommandListBase& RHICmdList)
 	{
 		const static FLazyName ClassName(TEXT("FTextureRenderTarget2DResource"));
 
-		FString ResourceName = Owner->GetName();
+		FString RHIResourceName = Owner->GetName();
 		ETextureCreateFlags TexCreateFlags = GetCreateFlags();
 		const int32 NumSamples = GetNumFromRenderTargetSampleCount(Owner->GetSampleCount());
 
 		FRHITextureCreateDesc Desc =
-			FRHITextureCreateDesc::Create2D(*ResourceName)
+			FRHITextureCreateDesc::Create2D(*RHIResourceName)
 			.SetExtent(TargetSizeX, TargetSizeY)
 			.SetFormat(Format)
 			.SetNumMips(TargetNumMips)
diff --git a/Engine/Source/Runtime/RenderCore/Private/RenderResource.cpp b/Engine/Source/Runtime/RenderCore/Private/RenderResource.cpp
--- a/Engine/Source/Runtime/RenderCore/Private/RenderResource.cpp
+++ b/Engine/Source/Runtime/RenderCore/Private/RenderResource.cpp
@@ -246,7 +246,7 @@ FRenderResource::~FRenderResource()
 	if (IsInitialized() && !GIsCriticalError)
 	{
 		// Deleting an initialized FRenderResource will result in a crash later since it is still linked
-		UE_LOG(LogRendererCore, Fatal,TEXT("A FRenderResource was deleted without being released first!"));
+		UE_LOG(LogRendererCore, Fatal, TEXT("A FRenderResource was deleted without being released first! [Resource=%s, Owner=%s]"), *GetResourceName().ToString(), *GetOwnerName().ToString());
 	}
 }
 
diff --git a/Engine/Source/Runtime/RenderCore/Private/VisualizeTexture.cpp b/Engine/Source/Runtime/RenderCore/Private/VisualizeTexture.cpp
--- a/Engine/Source/Runtime/RenderCore/Private/VisualizeTexture.cpp
+++ b/Engine/Source/Runtime/RenderCore/Private/VisualizeTexture.cpp
@@ -21,7 +21,7 @@ void FVisualizeTexture::ParseCommands(const TCHAR* Cmd, FOutputDevice &Ar)
 	ECommand Command = ECommand::Unknown;
 	int32 ViewId = 0;
 	FString ViewName;
-	FString ResourceName;
+	FString RDGResourceName;
 	TOptional<uint32> ResourceVersion;
 	TOptional<FWildcardString> ResourceListWildCard;
 	{
@@ -82,23 +82,23 @@ void FVisualizeTexture::ParseCommands(const TCHAR* Cmd, FOutputDevice &Ar)
 
 			if (*AfterAt == TCHAR('@'))
 			{
-				ResourceName = FirstParameter.Left(AfterAt - *FirstParameter);
+				RDGResourceName = FirstParameter.Left(AfterAt - *FirstParameter);
 				ResourceVersion = FCString::Atoi(AfterAt + 1);
 			}
 			else
 			{
-				ResourceName = FirstParameter;
+				RDGResourceName = FirstParameter;
 			}
 
-			if (ResourceName.Contains(TEXT("*")))
+			if (RDGResourceName.Contains(TEXT("*")))
 			{
-				ResourceListWildCard = FWildcardString(ResourceName);
+				ResourceListWildCard = FWildcardString(RDGResourceName);
 				Command = ECommand::DisplayResourceList;
 			}
 			else
 			{
 				Command = ECommand::VisualizeResource;
-				Visualize(ResourceName, ResourceVersion);
+				Visualize(RDGResourceName, ResourceVersion);
 			}
 		}
 	}
@@ -119,7 +119,7 @@ void FVisualizeTexture::ParseCommands(const TCHAR* Cmd, FOutputDevice &Ar)
 	else if (Command == ECommand::VisualizeResource)
 	{
 		Config = {};
-		Visualize(ResourceName, ResourceVersion);
+		Visualize(RDGResourceName, ResourceVersion);
 		for (;;)
 		{
 			FString Parameter = FParse::Token(Cmd, 0);
diff --git a/Engine/Source/Runtime/RenderCore/Public/RenderResource.h b/Engine/Source/Runtime/RenderCore/Public/RenderResource.h
--- a/Engine/Source/Runtime/RenderCore/Public/RenderResource.h
+++ b/Engine/Source/Runtime/RenderCore/Public/RenderResource.h
@@ -130,6 +130,22 @@ public:
 #endif
 	}
 
+	void SetResourceName(FName InResourceName)
+	{
+#if RHI_ENABLE_RESOURCE_INFO
+		ResourceName = InResourceName;
+#endif
+	}
+
+	FName GetResourceName() const
+	{
+#if RHI_ENABLE_RESOURCE_INFO
+		return ResourceName;
+#else
+		return NAME_None;
+#endif
+	}
+
 protected:
 	// This is used during mobile editor preview refactor, this will eventually be replaced with a parameter to InitRHI() etc..
 	void SetFeatureLevel(const FStaticFeatureLevel InFeatureLevel) { FeatureLevel = (ERHIFeatureLevel::Type)InFeatureLevel; }
@@ -181,6 +197,7 @@ private:
 
 #if RHI_ENABLE_RESOURCE_INFO
 	FName OwnerName;
+	FName ResourceName;
 #endif
 
 	int32 ListIndex;
