Fix for AspectRatioAxisConstraint skewing the scene when in Orthographic and when the viewport size Y > X

diff --git a/Engine/Source/Runtime/Engine/Private/Camera/CameraStackTypes.cpp b/Engine/Source/Runtime/Engine/Private/Camera/CameraStackTypes.cpp
--- a/Engine/Source/Runtime/Engine/Private/Camera/CameraStackTypes.cpp
+++ b/Engine/Source/Runtime/Engine/Private/Camera/CameraStackTypes.cpp
@@ -285,7 +285,7 @@ void FMinimalViewInfo::CalculateProjectionMatrixGivenViewRectangle(FMinimalViewI
 
 		// If x is bigger, and we're respecting x or major axis, AND mobile isn't forcing us to be Y axis aligned
 		const bool bMaintainXFOV = 
-			((SizeX > SizeY) && ((AspectRatioAxisConstraint == AspectRatio_MajorAxisFOV) || (ViewInfo.ProjectionMode == ECameraProjectionMode::Orthographic))) ||
+			((SizeX > SizeY) && (AspectRatioAxisConstraint == AspectRatio_MajorAxisFOV)) ||
 			(AspectRatioAxisConstraint == AspectRatio_MaintainXFOV);
 		if (bMaintainXFOV)
 		{
@@ -300,26 +300,9 @@ void FMinimalViewInfo::CalculateProjectionMatrixGivenViewRectangle(FMinimalViewI
 			YAxisMultiplier = 1.0f;
 		}
 		
-		float MatrixHalfFOV;
-		if (!bMaintainXFOV && ViewInfo.AspectRatio != 0.f && !CVarUseLegacyMaintainYFOV.GetValueOnGameThread())
-		{
-			// The view-info FOV is horizontal. But if we have a different aspect ratio constraint, we need to
-			// adjust this FOV value using the aspect ratio it was computed with, so we that we can compute the
-			// complementary FOV value (with the *effective* aspect ratio) correctly.
-			const float HalfXFOV = FMath::DegreesToRadians(FMath::Max(0.001f, ViewInfo.FOV) / 2.f);
-			const float HalfYFOV = FMath::Atan(FMath::Tan(HalfXFOV) / ViewInfo.AspectRatio);
-			MatrixHalfFOV = HalfYFOV;
-		}
-		else
-		{
-			// Avoid divide by zero in the projection matrix calculation by clamping FOV.
-			// Note the division by 360 instead of 180 because we want the half-FOV.
-			MatrixHalfFOV = FMath::Max(0.001f, ViewInfo.FOV) * (float)UE_PI / 360.0f;
-		}
-	
 		if (bOrthographic)
 		{
-			const float OrthoWidth = ViewInfo.OrthoWidth / 2.0f * XAxisMultiplier;
+			const float OrthoWidth = (ViewInfo.OrthoWidth / 2.0f) / XAxisMultiplier;
 			const float OrthoHeight = (ViewInfo.OrthoWidth / 2.0f) / YAxisMultiplier;
 
 			float FarPlane = ViewInfo.OrthoFarClipPlane;
@@ -339,6 +322,23 @@ void FMinimalViewInfo::CalculateProjectionMatrixGivenViewRectangle(FMinimalViewI
 		}
 		else
 		{
+			float MatrixHalfFOV;
+			if (!bMaintainXFOV && ViewInfo.AspectRatio != 0.f && !CVarUseLegacyMaintainYFOV.GetValueOnGameThread())
+			{
+				// The view-info FOV is horizontal. But if we have a different aspect ratio constraint, we need to
+				// adjust this FOV value using the aspect ratio it was computed with, so we that we can compute the
+				// complementary FOV value (with the *effective* aspect ratio) correctly.
+				const float HalfXFOV = FMath::DegreesToRadians(FMath::Max(0.001f, ViewInfo.FOV) / 2.f);
+				const float HalfYFOV = FMath::Atan(FMath::Tan(HalfXFOV) / ViewInfo.AspectRatio);
+				MatrixHalfFOV = HalfYFOV;
+			}
+			else
+			{
+				// Avoid divide by zero in the projection matrix calculation by clamping FOV.
+				// Note the division by 360 instead of 180 because we want the half-FOV.
+				MatrixHalfFOV = FMath::Max(0.001f, ViewInfo.FOV) * (float)UE_PI / 360.0f;
+			}
+
 			const float ClippingPlane = ViewInfo.GetFinalPerspectiveNearClipPlane();
 			InOutProjectionData.ProjectionMatrix = FReversedZPerspectiveMatrix(
 				MatrixHalfFOV,
