First Person: Exposed FP FOV and scale on ViewProperty material graph node.

diff --git a/Engine/Source/Runtime/Engine/Private/Materials/HLSLMaterialTranslator.cpp b/Engine/Source/Runtime/Engine/Private/Materials/HLSLMaterialTranslator.cpp
--- a/Engine/Source/Runtime/Engine/Private/Materials/HLSLMaterialTranslator.cpp
+++ b/Engine/Source/Runtime/Engine/Private/Materials/HLSLMaterialTranslator.cpp
@@ -5293,6 +5293,9 @@ int32 FHLSLMaterialTranslator::ViewProperty(EMaterialExposedViewProperty Propert
 		{MEVP_RuntimeVirtualTextureMaxLevel, MCT_Float1, TEXT("View.RuntimeVirtualTextureMipLevel.y"), nullptr},
 		{MEVP_ResolutionFraction, MCT_Float1, TEXT("View.ResolutionFractionAndInv.x"), TEXT("View.ResolutionFractionAndInv.y")},
 		{MEVP_PostVolumeUserFlags, MCT_Float1, TEXT("View.PostVolumeUserFlags"), nullptr},
+		{MEVP_FirstPersonFieldOfView, MCT_Float2, TEXT("View.<PREV>FirstPersonFieldOfViewWideAngles"), nullptr},
+		{MEVP_FirstPersonTanHalfFieldOfView, MCT_Float2, TEXT("View.<PREV>FirstPersonTanAndInvTanHalfFOV.xy"), TEXT("View.<PREV>FirstPersonTanAndInvTanHalfFOV.zw")},
+		{MEVP_FirstPersonScale, MCT_Float, TEXT("View.<PREV>FirstPersonScale"), nullptr},
 	};
 	static_assert((sizeof(ViewPropertyMetaArray) / sizeof(ViewPropertyMetaArray[0])) == MEVP_MAX, "incoherency between EMaterialExposedViewProperty and ViewPropertyMetaArray");
 
diff --git a/Engine/Source/Runtime/Engine/Private/SceneView.cpp b/Engine/Source/Runtime/Engine/Private/SceneView.cpp
--- a/Engine/Source/Runtime/Engine/Private/SceneView.cpp
+++ b/Engine/Source/Runtime/Engine/Private/SceneView.cpp
@@ -727,7 +727,9 @@ void FViewMatrices::Init(const FMinimalInitializer& Initializer)
 		Initializer.ConstrainedViewRect.Size().Y * 0.5f * ProjectionScale.Y
 	);
 
-	if (Initializer.FirstPersonParams.bUseParameters)
+	FirstPersonParams = Initializer.FirstPersonParams;
+
+	if (FirstPersonParams.bUseParameters)
 	{
 		// The projection matrix uses 1/tan(HalfFOV) on the X and Y components in view space, so in order to cancel out the existing FOV and achieve a different first person FOV, we do
 		//		(1 / tan(HalfFirstPersonFOV)) / (1 / tan(HalfFOV)) 
@@ -736,8 +738,8 @@ void FViewMatrices::Init(const FMinimalInitializer& Initializer)
 		// This correction factor should be independent of the actual aspect ratio used on the original projection matrix. As the FOV correction needs to be done in view space, we can conveniently
 		// also use this opportunity to scale down the first person geometry towards the camera. The geometry takes up the same solid angle from the perspective of the camera and therefore appears the
 		// same size while taking up a smaller range in the depth buffer, reducing the likelihood of intersections with the scene.
-		const float FOVCorrectionFactor = Initializer.FirstPersonParams.FOVCorrectionFactor;
-		FirstPersonTransform = TranslatedViewMatrix * FScaleMatrix(FVector(FOVCorrectionFactor, FOVCorrectionFactor, 1.0f) * Initializer.FirstPersonParams.Scale) * InvTranslatedViewMatrix;
+		const float FOVCorrectionFactor = FirstPersonParams.FOVCorrectionFactor;
+		FirstPersonTransform = TranslatedViewMatrix * FScaleMatrix(FVector(FOVCorrectionFactor, FOVCorrectionFactor, 1.0f) * FirstPersonParams.Scale) * InvTranslatedViewMatrix;
 	}
 }
 
@@ -2555,6 +2557,8 @@ void SetupViewRectUniformBufferParameters(
 		FVector4f TanAndInvTanFOV = InViewMatrices.GetTanAndInvTanHalfFOV();
 		ViewUniformShaderParameters.TanAndInvTanHalfFOV = TanAndInvTanFOV;
 		ViewUniformShaderParameters.PrevTanAndInvTanHalfFOV = InPrevViewMatrices.GetTanAndInvTanHalfFOV();
+		ViewUniformShaderParameters.FirstPersonTanAndInvTanHalfFOV = InViewMatrices.GetFirstPersonTanAndInvTanHalfFOV();
+		ViewUniformShaderParameters.PrevTanAndInvTanHalfFOV = InPrevViewMatrices.GetFirstPersonTanAndInvTanHalfFOV();
 
 		if (InViewMatrices.IsPerspectiveProjection())
 		{
@@ -2690,6 +2694,10 @@ void SetupCommonViewUniformBufferParameters(
 
 	ViewUniformShaderParameters.FieldOfViewWideAngles = FVector2f(2.f * InViewMatrices.ComputeHalfFieldOfViewPerAxis());	// LWC_TODO: Precision loss
 	ViewUniformShaderParameters.PrevFieldOfViewWideAngles = FVector2f(2.f * InPrevViewMatrices.ComputeHalfFieldOfViewPerAxis());	// LWC_TODO: Precision loss
+	ViewUniformShaderParameters.FirstPersonFieldOfViewWideAngles = FVector2f(2.f * InViewMatrices.GetFirstPersonHalfFieldOfViewPerAxis());
+	ViewUniformShaderParameters.PrevFirstPersonFieldOfViewWideAngles = FVector2f(2.f * InPrevViewMatrices.GetFirstPersonHalfFieldOfViewPerAxis());
+	ViewUniformShaderParameters.FirstPersonScale = InViewMatrices.GetFirstPersonScale();
+	ViewUniformShaderParameters.PrevFirstPersonScale = InPrevViewMatrices.GetFirstPersonScale();
 	ViewUniformShaderParameters.DiffuseOverrideParameter = LocalDiffuseOverrideParameter;
 	ViewUniformShaderParameters.SpecularOverrideParameter = Inputs.SpecularOverrideParameter;
 	ViewUniformShaderParameters.NormalOverrideParameter = Inputs.NormalOverrideParameter;
diff --git a/Engine/Source/Runtime/Engine/Public/Materials/MaterialExpressionViewProperty.h b/Engine/Source/Runtime/Engine/Public/Materials/MaterialExpressionViewProperty.h
--- a/Engine/Source/Runtime/Engine/Public/Materials/MaterialExpressionViewProperty.h
+++ b/Engine/Source/Runtime/Engine/Public/Materials/MaterialExpressionViewProperty.h
@@ -43,6 +43,12 @@ enum EMaterialExposedViewProperty : int
 	MEVP_ResolutionFraction UMETA(DisplayName = "ScreenPercentage / 100"),
 	/** Post process volume user flags, useful for varying the behavior of a material per view */
 	MEVP_PostVolumeUserFlags UMETA(DisplayName = "Post Process Volume User Flags"),
+	/** Horizontal and vertical first person field of view angles in radian */
+	MEVP_FirstPersonFieldOfView UMETA(DisplayName = "First Person Field Of View"),
+	/** Tan(FirstPersonFieldOfView * 0.5) */
+	MEVP_FirstPersonTanHalfFieldOfView UMETA(DisplayName = "Tan(0.5 * First Person Field Of View)"),
+	/** Scaling factor applied to first person primitives to keep them from intersecting with the scene. */
+	MEVP_FirstPersonScale UMETA(DisplayName = "First Person Scale"),
 
 	MEVP_MAX,
 };
diff --git a/Engine/Source/Runtime/Engine/Public/SceneView.h b/Engine/Source/Runtime/Engine/Public/SceneView.h
--- a/Engine/Source/Runtime/Engine/Public/SceneView.h
+++ b/Engine/Source/Runtime/Engine/Public/SceneView.h
@@ -399,6 +399,9 @@ private:
 	/** Depth test scaling; differs between perspective and orthographic  */
 	float PerProjectionDepthThicknessScale;
 
+	/** Parameters controlling the rendering (FOV and scale) of first person primitives. */
+	FFirstPersonParameters FirstPersonParams;
+
 	//
 	// World = TranslatedWorld - PreViewTranslation
 	// TranslatedWorld = World + PreViewTranslation
@@ -491,6 +494,11 @@ public:
 		return FirstPersonTransform;
 	}
 
+	inline float GetFirstPersonScale() const
+	{
+		return FirstPersonParams.bUseParameters ? FirstPersonParams.Scale : 1.0f;
+	}
+
 	inline const FVector& GetPreViewTranslation() const
 	{
 		return PreViewTranslation;
@@ -696,6 +704,36 @@ public:
 		return FVector2f::ZeroVector;
 	}
 
+	// @return in radians (horizontal,vertical)
+	FVector2f GetFirstPersonHalfFieldOfViewPerAxis() const
+	{
+		if (IsPerspectiveProjection())
+		{
+			const FVector4f FirstPersonTandAndInvTanHalfFOV = GetFirstPersonTanAndInvTanHalfFOV();
+			const FVector2f FirstPersonHalfFOV = FVector2f(FMath::Atan(FirstPersonTandAndInvTanHalfFOV.X), FMath::Atan(FirstPersonTandAndInvTanHalfFOV.Y));
+			return FirstPersonHalfFOV;
+		}
+		return FVector2f::Zero();
+	}
+
+	FVector4f GetFirstPersonTanAndInvTanHalfFOV() const
+	{
+		// No concept of FOV for orthographic projection so only return perspective related values or 1.0f
+		if (IsPerspectiveProjection())
+		{
+			const FVector4f SceneTanAndInvTanHalfFOV = GetTanAndInvTanHalfFOV();
+			const float FOVCorrectionFactor = FirstPersonParams.bUseParameters ? FirstPersonParams.FOVCorrectionFactor : 1.0f;
+			// FOVCorrectionFactor = SceneTanHalfFOV / FirstPersonTanHalfFov -> FirstPersonTanHalfFov = SceneTanHalfFOV / FOVCorrectionFactor
+			const FVector4f FirstPersonTanAndInvTanHalfFOV = FVector4f(
+				SceneTanAndInvTanHalfFOV.X / FOVCorrectionFactor,
+				SceneTanAndInvTanHalfFOV.Y / FOVCorrectionFactor,
+				SceneTanAndInvTanHalfFOV.Z * FOVCorrectionFactor,
+				SceneTanAndInvTanHalfFOV.W * FOVCorrectionFactor);
+			return FirstPersonTanAndInvTanHalfFOV;
+		}
+		return FVector4f(1.0f, 1.0f, 1.0f, 1.0f);
+	}
+
 private:
 	inline void RecomputeDerivedMatrices()
 	{
@@ -853,6 +891,8 @@ enum ETranslucencyVolumeCascade
 	VIEW_UNIFORM_BUFFER_MEMBER_PER_VIEW(FVector4f, GlobalClippingPlane) \
 	VIEW_UNIFORM_BUFFER_MEMBER_PER_VIEW(FVector2f, FieldOfViewWideAngles) \
 	VIEW_UNIFORM_BUFFER_MEMBER_PER_VIEW(FVector2f, PrevFieldOfViewWideAngles) \
+	VIEW_UNIFORM_BUFFER_MEMBER_PER_VIEW(FVector2f, FirstPersonFieldOfViewWideAngles) \
+	VIEW_UNIFORM_BUFFER_MEMBER_PER_VIEW(FVector2f, PrevFirstPersonFieldOfViewWideAngles) \
 	VIEW_UNIFORM_BUFFER_MEMBER_PER_VIEW_EX(FVector4f, ViewRectMin, EShaderPrecisionModifier::Half) \
 	VIEW_UNIFORM_BUFFER_MEMBER(FVector4f, ViewSizeAndInvSize) \
 	VIEW_UNIFORM_BUFFER_MEMBER(FUintVector4, ViewRectMinAndSize) \
@@ -1025,10 +1065,14 @@ enum ETranslucencyVolumeCascade
 	VIEW_UNIFORM_BUFFER_MEMBER(float, SubSurfaceColorAsTransmittanceAtDistanceInMeters) \
 	VIEW_UNIFORM_BUFFER_MEMBER_PER_VIEW(FVector4f, TanAndInvTanHalfFOV) \
 	VIEW_UNIFORM_BUFFER_MEMBER_PER_VIEW(FVector4f, PrevTanAndInvTanHalfFOV) \
+	VIEW_UNIFORM_BUFFER_MEMBER_PER_VIEW(FVector4f, FirstPersonTanAndInvTanHalfFOV) \
+	VIEW_UNIFORM_BUFFER_MEMBER_PER_VIEW(FVector4f, PrevFirstPersonTanAndInvTanHalfFOV) \
 	VIEW_UNIFORM_BUFFER_MEMBER_PER_VIEW(FVector2f, WorldDepthToPixelWorldRadius) \
 	VIEW_UNIFORM_BUFFER_MEMBER_PER_VIEW(FVector4f, ScreenRayLengthMultiplier) \
 	VIEW_UNIFORM_BUFFER_MEMBER_PER_VIEW(FVector4f, GlintLUTParameters0) \
 	VIEW_UNIFORM_BUFFER_MEMBER_PER_VIEW(FVector4f, GlintLUTParameters1) \
+	VIEW_UNIFORM_BUFFER_MEMBER_PER_VIEW(float, FirstPersonScale) \
+	VIEW_UNIFORM_BUFFER_MEMBER_PER_VIEW(float, PrevFirstPersonScale) \
 	VIEW_UNIFORM_BUFFER_MEMBER(float, MaterialMaxEmissiveValue) \
 	VIEW_UNIFORM_BUFFER_MEMBER(int32, PostVolumeUserFlags) \
 	VIEW_UNIFORM_BUFFER_MEMBER(FIntVector4, EnvironmentComponentsFlags) \
