Optimization: Added caching in material for texture streaming to optimize FStreamingTextureLevelContext::ProcessMaterial. Can be enabled using 'r.Streaming.EnableTexturesSamplingStreamingCache = true'.

diff --git a/Engine/Source/Runtime/Engine/Private/Materials/Material.cpp b/Engine/Source/Runtime/Engine/Private/Materials/Material.cpp
--- a/Engine/Source/Runtime/Engine/Private/Materials/Material.cpp
+++ b/Engine/Source/Runtime/Engine/Private/Materials/Material.cpp
@@ -4652,6 +4652,9 @@ void UMaterial::PostLoad()
 		UMaterialInterface::PrecachePSOs(&FLocalVertexFactory::StaticType, PSOPrecacheParams);
 	}
 
+#if !WITH_EDITOR
+	CacheTexturesSamplingInfo();
+#endif
 }
 
 #if WITH_EDITORONLY_DATA
diff --git a/Engine/Source/Runtime/Engine/Private/Materials/MaterialInstance.cpp b/Engine/Source/Runtime/Engine/Private/Materials/MaterialInstance.cpp
--- a/Engine/Source/Runtime/Engine/Private/Materials/MaterialInstance.cpp
+++ b/Engine/Source/Runtime/Engine/Private/Materials/MaterialInstance.cpp
@@ -3495,6 +3495,10 @@ void UMaterialInstance::PostLoad()
 		UMaterialInterface::PrecachePSOs(&FLocalVertexFactory::StaticType, PSOPrecacheParams);
 	}
 	//DumpDebugInfo(*GLog);
+
+#if !WITH_EDITOR
+	CacheTexturesSamplingInfo();
+#endif
 }
 
 #if WITH_EDITOR
@@ -5063,6 +5067,18 @@ void UMaterialInstance::PreSave(FObjectPreSaveContext ObjectSaveContext)
 #endif
 }
 
+#if !WITH_EDITOR
+FTextureSamplingInfo UMaterialInstance::CalculateTexturesSamplingInfo(UTexture* Texture)
+{
+	FTextureSamplingInfo SamplingInfo = Super::CalculateTexturesSamplingInfo(Texture);
+	if (!SamplingInfo.bIsValid && Parent)
+	{
+		return Parent->CalculateTexturesSamplingInfo(Texture);
+	}
+	return SamplingInfo;
+}
+#endif
+
 float UMaterialInstance::GetTextureDensity(FName TextureName, const struct FMeshUVChannelInfo& UVChannelData) const
 {
 	ensure(UVChannelData.bInitialized);
diff --git a/Engine/Source/Runtime/Engine/Private/Materials/MaterialInterface.cpp b/Engine/Source/Runtime/Engine/Private/Materials/MaterialInterface.cpp
--- a/Engine/Source/Runtime/Engine/Private/Materials/MaterialInterface.cpp
+++ b/Engine/Source/Runtime/Engine/Private/Materials/MaterialInterface.cpp
@@ -45,6 +45,7 @@
 #include "UObject/Package.h"
 #include "ShaderCompiler.h"
 #include "Materials/MaterialInsights.h"
+#include "UnrealEngine.h"
 
 #if WITH_EDITOR
 #include "ObjectCacheEventSink.h"
@@ -1847,6 +1848,25 @@ void UMaterialInterface::SetTextureStreamingData(const TArray<FMaterialTextureIn
 	SortTextureStreamingData(true, false);
 }
 
+#if !WITH_EDITOR
+float UMaterialInterface::GetTextureDensityWithCache(const FTextureSamplingInfo& TextureSamplingInfo, const FMeshUVChannelInfo& UVChannelData) const
+{
+	if (TextureSamplingInfo.bIsValid)
+	{
+		check(CachedTexturesSamplingInfo && Algo::FindByPredicate(CachedTexturesSamplingInfo->TextureSamplingInfos, [&TextureSamplingInfo](const FTextureSamplingInfo& Info) { return &Info == &TextureSamplingInfo; }));
+		float TextureDensity = 0;
+		for (int32 Index = 0; Index < MAX_TEXCOORDS; ++Index)
+		{
+			const float SamplingScale = TextureSamplingInfo.ChannelMinSamplingScale[Index];
+			TextureDensity = FMath::Max<float>(UVChannelData.LocalUVDensities[Index] / SamplingScale, TextureDensity);
+		}
+		// Here we behave like GetTextureDensity and make sure to return exactly 0 when nearly zero as the calling code can do a bool conversion with the result
+		return FMath::IsNearlyZero(TextureDensity) ? 0 : TextureDensity;
+	}
+	return 0;
+}
+#endif
+
 float UMaterialInterface::GetTextureDensity(FName TextureName, const FMeshUVChannelInfo& UVChannelData) const
 {
 	ensure(UVChannelData.bInitialized);
@@ -2048,6 +2068,62 @@ bool UMaterialInterface::Rename(const TCHAR* NewName, UObject* NewOuter, ERename
 	return bRenamed;
 }
 
+namespace UE::Private
+{
+	bool GEnableTexturesSamplingStreamingCache = false;
+	static FAutoConsoleVariableRef CVarEnableTexturesSamplingStreamingCache(
+		TEXT("r.Streaming.EnableTexturesSamplingStreamingCache"),
+		GEnableTexturesSamplingStreamingCache,
+		TEXT("Enables caching of material texture sampling (only used by texture streaming in a cooked build)."),
+		ECVF_Default);
+}
+
+#if !WITH_EDITOR
+
+bool UMaterialInterface::CanCacheTexturesSamplingInfo() const
+{
+	return UE::Private::GEnableTexturesSamplingStreamingCache;
+}
+
+void UMaterialInterface::CacheTexturesSamplingInfo()
+{
+	if (CanCacheTexturesSamplingInfo())
+	{
+		TRACE_CPUPROFILER_EVENT_SCOPE(UMaterialInterface::CacheTexturesSamplingInfo);
+
+		TArray<UTexture*> Textures;
+		GetUsedTextures(Textures, GetCachedScalabilityCVars().MaterialQualityLevel, false, GMaxRHIFeatureLevel, false);
+
+		FMaterialCachedTexturesSamplingInfo Cache;
+		for (UTexture* Texture : Textures)
+		{
+			Cache.TextureSamplingInfos.Add(CalculateTexturesSamplingInfo(Texture));
+		}
+
+		check(!CachedTexturesSamplingInfo.IsSet());
+		CachedTexturesSamplingInfo = MoveTemp(Cache);
+	}
+}
+
+FTextureSamplingInfo UMaterialInterface::CalculateTexturesSamplingInfo(UTexture* Texture)
+{
+	FTextureSamplingInfo SamplingInfo(Texture);
+	int32 LowerIndex = INDEX_NONE;
+	int32 UpperIndex = INDEX_NONE;
+	if (FindTextureStreamingDataIndexRange(Texture->GetFName(), LowerIndex, UpperIndex))
+	{
+		for (int32 Index = LowerIndex; Index <= UpperIndex; ++Index)
+		{
+			const FMaterialTextureInfo& Info = TextureStreamingData[Index];
+			SamplingInfo.ChannelMinSamplingScale[Info.UVChannelIndex] = FMath::Min<float>(Info.SamplingScale, SamplingInfo.ChannelMinSamplingScale[Info.UVChannelIndex]);
+			SamplingInfo.bIsValid |= true;
+		}
+	}
+	return SamplingInfo;
+}
+
+#endif
+
 UMaterialInterfaceEditorOnlyData::UMaterialInterfaceEditorOnlyData()
 {
 }
diff --git a/Engine/Source/Runtime/Engine/Private/Streaming/TextureInstanceState.cpp b/Engine/Source/Runtime/Engine/Private/Streaming/TextureInstanceState.cpp
--- a/Engine/Source/Runtime/Engine/Private/Streaming/TextureInstanceState.cpp
+++ b/Engine/Source/Runtime/Engine/Private/Streaming/TextureInstanceState.cpp
@@ -99,7 +99,6 @@ void FRenderAssetInstanceState::RemoveBounds(int32 BoundsIndex)
 
 void FRenderAssetInstanceState::AddElement(const UPrimitiveComponent* InComponent, const UStreamableRenderAsset* InAsset, int InBoundsIndex, float InTexelFactor, bool InForceLoad, int32*& ComponentLink)
 {
-	TRACE_CPUPROFILER_EVENT_SCOPE(FRenderAssetInstanceState::AddElement);
 	check(InComponent && InAsset);
 
 	// Keep Max texel factor up to date.
diff --git a/Engine/Source/Runtime/Engine/Private/Streaming/TextureStreamingBuild.cpp b/Engine/Source/Runtime/Engine/Private/Streaming/TextureStreamingBuild.cpp
--- a/Engine/Source/Runtime/Engine/Private/Streaming/TextureStreamingBuild.cpp
+++ b/Engine/Source/Runtime/Engine/Private/Streaming/TextureStreamingBuild.cpp
@@ -10,6 +10,7 @@ TextureStreamingBuild.cpp : Contains definitions to build texture streaming data
 #include "Engine/Level.h"
 #include "Engine/World.h"
 #include "Components/PrimitiveComponent.h"
+#include "Materials/Material.h"
 #include "Materials/MaterialInterface.h"
 #include "Misc/FeedbackContext.h"
 #include "Streaming/TextureStreamingHelpers.h"
@@ -645,12 +646,48 @@ void FStreamingTextureLevelContext::ProcessMaterial(const FBoxSphereBounds& Comp
 {
 	ensure(MaterialData.IsValid());
 
-	TArray<UTexture*> Textures;
+	auto ProcessTexture = [this, &ComponentBounds, &MaterialData, ComponentScaling, &OutStreamingTextures](UTexture* Texture, const FTextureSamplingInfo* TextureSamplingInfo = nullptr)
+	{
+#if WITH_EDITOR
+		if (GetTextureIsStreamable(*Texture))
+#else
+		if (Texture->IsStreamable())
+#endif
+		{
+			float TextureDensity = 0.f;
+#if !WITH_EDITOR
+			if (TextureSamplingInfo)
+			{
+				TextureDensity = MaterialData.Material->GetTextureDensityWithCache(*TextureSamplingInfo, *MaterialData.UVChannelData);
+			}
+			else
+#endif
+			{
+				TextureDensity = MaterialData.Material->GetTextureDensity(Texture->GetFName(), *MaterialData.UVChannelData);
+			}
+
+			if (!TextureDensity)
+			{
+				// Fallback assuming a sampling scale of 1 using the UV channel 0;
+				TextureDensity = MaterialData.UVChannelData->LocalUVDensities[0];
+			}
+
+			if (TextureDensity)
+			{
+				FStreamingRenderAssetPrimitiveInfo& Info = OutStreamingTextures.AddDefaulted_GetRef();
+				Info.RenderAsset = Texture;
+				Info.TexelFactor = TextureDensity * ComponentScaling;
+				Info.PackedRelativeBox = bUseRelativeBoxes ? MaterialData.PackedRelativeBox : PackedRelativeBox_Identity;
+				UnpackRelativeBox(ComponentBounds, Info.PackedRelativeBox, Info.Bounds);
+			}
+		}
+	};
 
 #if !WITH_EDITOR
 	// Use pre-built texture streaming data is possible
 	if (CanUseTextureStreamingBuiltData() && (ComponentBuildData || bIsComponentBuildDataValid))
 	{
+		TArray<UTexture*, TInlineAllocator<64>> StreamableTextures;
 		// bIsComponentBuildDataValid can be true, but no streamable textures were found when building texture streaming data
 		if (ComponentBuildData)
 		{
@@ -661,70 +698,63 @@ void FStreamingTextureLevelContext::ProcessMaterial(const FBoxSphereBounds& Comp
 					if (UTexture* Texture = LevelStreamingTextures[BuildInfo.TextureLevelIndex])
 					{
 						check(Texture->LevelIndex != INDEX_NONE);
-						Textures.Add(Texture);
+						if (Texture->IsStreamable())
+						{
+							StreamableTextures.Add(Texture);
+						}
 					}
 				}
 			}
 		}
-	}
-	else
-#endif
-	{
-		MaterialData.Material->GetUsedTextures(Textures, QualityLevel, false, FeatureLevel, false);
-	}
-
-	for (UTexture* Texture : Textures)
-	{
-#if WITH_EDITOR
-		bool bIsTextureStreamable = GetTextureIsStreamable(*Texture);
-#else
-		bool bIsTextureStreamable = Texture->IsStreamable();
-#endif
-
-		if (!bIsTextureStreamable)
-		{
-			continue;
-		}
 
-		int32* BuildDataIndex = GetBuildDataIndexRef(Texture);
-		if (BuildDataIndex)
+		for (UTexture* Texture : StreamableTextures)
 		{
-			if (*BuildDataIndex != INDEX_NONE)
+			if (int32* BuildDataIndex = GetBuildDataIndexRef(Texture))
 			{
-				FStreamingRenderAssetPrimitiveInfo& Info = OutStreamingTextures.AddDefaulted_GetRef();
-				const FStreamingTextureBuildInfo& BuildInfo = (*ComponentBuildData)[*BuildDataIndex];
+				if (*BuildDataIndex != INDEX_NONE)
+				{
+					FStreamingRenderAssetPrimitiveInfo& Info = OutStreamingTextures.AddDefaulted_GetRef();
+					const FStreamingTextureBuildInfo& BuildInfo = (*ComponentBuildData)[*BuildDataIndex];
 
-				Info.RenderAsset = Texture;
-				Info.TexelFactor = BuildInfo.TexelFactor * ComponentScaling;
-				Info.PackedRelativeBox = bUseRelativeBoxes ? BuildInfo.PackedRelativeBox : PackedRelativeBox_Identity;
-				UnpackRelativeBox(ComponentBounds, Info.PackedRelativeBox, Info.Bounds);
+					Info.RenderAsset = Texture;
+					Info.TexelFactor = BuildInfo.TexelFactor * ComponentScaling;
+					Info.PackedRelativeBox = bUseRelativeBoxes ? BuildInfo.PackedRelativeBox : PackedRelativeBox_Identity;
+					UnpackRelativeBox(ComponentBounds, Info.PackedRelativeBox, Info.Bounds);
 
-				// Indicate that this texture build data has already been processed.
-				// The build data use the merged results of all material so it only needs to be processed once.
-				*BuildDataIndex = INDEX_NONE;
+					// Indicate that this texture build data has already been processed.
+					// The build data use the merged results of all material so it only needs to be processed once.
+					*BuildDataIndex = INDEX_NONE;
+				}
+			}
+			else
+			{
+				// Otherwise create an entry using the available data.
+				ProcessTexture(Texture);
 			}
 		}
-		else // Otherwise create an entry using the available data.
+	}
+	else
+	{
+		if (const FMaterialCachedTexturesSamplingInfo* CachedTexturesSamplingInfo = MaterialData.Material->GetCachedTexturesSamplingInfo())
 		{
-			float TextureDensity = MaterialData.Material->GetTextureDensity(Texture->GetFName(), *MaterialData.UVChannelData);
-
-			if (!TextureDensity)
+			for (const FTextureSamplingInfo& SamplingInfo : CachedTexturesSamplingInfo->TextureSamplingInfos)
 			{
-				// Fallback assuming a sampling scale of 1 using the UV channel 0;
-				TextureDensity = MaterialData.UVChannelData->LocalUVDensities[0];
+				ProcessTexture(SamplingInfo.Texture, &SamplingInfo);
 			}
-
-			if (TextureDensity)
+		}
+		else
+#endif
+		{
+			TArray<UTexture*> Textures;
+			MaterialData.Material->GetUsedTextures(Textures, QualityLevel, false, FeatureLevel, false);
+			for (UTexture* Texture : Textures)
 			{
-				FStreamingRenderAssetPrimitiveInfo& Info = OutStreamingTextures.AddDefaulted_GetRef();
-
-				Info.RenderAsset = Texture;
-				Info.TexelFactor = TextureDensity * ComponentScaling;
-				Info.PackedRelativeBox = bUseRelativeBoxes ? MaterialData.PackedRelativeBox : PackedRelativeBox_Identity;
-				UnpackRelativeBox(ComponentBounds, Info.PackedRelativeBox, Info.Bounds);
+				ProcessTexture(Texture);
 			}
 		}
+#if !WITH_EDITOR
 	}
+#endif
 }
 
 void CheckTextureStreamingBuildValidity(UWorld* InWorld)
diff --git a/Engine/Source/Runtime/Engine/Public/Materials/MaterialInstance.h b/Engine/Source/Runtime/Engine/Public/Materials/MaterialInstance.h
--- a/Engine/Source/Runtime/Engine/Public/Materials/MaterialInstance.h
+++ b/Engine/Source/Runtime/Engine/Public/Materials/MaterialInstance.h
@@ -756,6 +756,9 @@ public:
 
 	ENGINE_API virtual void PreSave(FObjectPreSaveContext ObjectSaveContext) override;
 	ENGINE_API virtual float GetTextureDensity(FName TextureName, const struct FMeshUVChannelInfo& UVChannelData) const override;
+#if !WITH_EDITOR
+	ENGINE_API virtual FTextureSamplingInfo CalculateTexturesSamplingInfo(UTexture* Texture) override;
+#endif
 
 	ENGINE_API bool Equivalent(const UMaterialInstance* CompareTo) const;
 
diff --git a/Engine/Source/Runtime/Engine/Public/Materials/MaterialInstanceDynamic.h b/Engine/Source/Runtime/Engine/Public/Materials/MaterialInstanceDynamic.h
--- a/Engine/Source/Runtime/Engine/Public/Materials/MaterialInstanceDynamic.h
+++ b/Engine/Source/Runtime/Engine/Public/Materials/MaterialInstanceDynamic.h
@@ -216,6 +216,13 @@ protected:
 
 	ENGINE_API void InitializeMID(class UMaterialInterface* ParentMaterial);
 
+#if !WITH_EDITOR
+	virtual bool CanCacheTexturesSamplingInfo() const override
+	{
+		return false;
+	}
+#endif
+
 private:
 	
 	ENGINE_API void UpdateCachedDataDynamic();
diff --git a/Engine/Source/Runtime/Engine/Public/Materials/MaterialInterface.h b/Engine/Source/Runtime/Engine/Public/Materials/MaterialInterface.h
--- a/Engine/Source/Runtime/Engine/Public/Materials/MaterialInterface.h
+++ b/Engine/Source/Runtime/Engine/Public/Materials/MaterialInterface.h
@@ -8,9 +8,11 @@
 #include "Delegates/Delegate.h"
 #include "MaterialTypes.h"
 #include "Containers/ArrayView.h"
+#include "Containers/StaticArray.h"
 #include "UObject/ObjectMacros.h"
 #include "UObject/Object.h"
 #include "Misc/Guid.h"
+#include "Misc/Optional.h"
 #include "Templates/UniquePtr.h"
 #include "Templates/SharedPointer.h"
 #include "Engine/EngineTypes.h"
@@ -25,6 +27,7 @@
 #include "Materials/MaterialRelevance.h"
 #include "MaterialRecursionGuard.h"
 #include "MaterialShaderPrecompileMode.h"
+#include "MeshUVChannelInfo.h"
 #include "RHIFeatureLevel.h"
 #include "PSOPrecache.h"
 #include "StaticParameterSet.h"
@@ -227,6 +230,34 @@ struct FMaterialAnalysisResult
 	bool bRequiresVertexData = false;
 };
 
+USTRUCT()
+struct FTextureSamplingInfo
+{
+	GENERATED_BODY()
+
+	FTextureSamplingInfo(UTexture* InTexture = nullptr)
+		: bIsValid(false)
+		, Texture(InTexture)
+		, ChannelMinSamplingScale(InPlace, UE_MAX_FLT)
+	{}
+
+	bool bIsValid;
+	
+	UPROPERTY()
+	TObjectPtr<UTexture> Texture;
+
+	TStaticArray<float, MAX_TEXCOORDS> ChannelMinSamplingScale;
+};
+
+USTRUCT()
+struct FMaterialCachedTexturesSamplingInfo
+{
+	GENERATED_BODY()
+
+	UPROPERTY()
+	TArray<FTextureSamplingInfo> TextureSamplingInfos;
+};
+
 UCLASS(Optional)
 class UMaterialInterfaceEditorOnlyData : public UObject
 {
@@ -270,6 +301,14 @@ public:
 #endif // WITH_EDITORONLY_DATA
 
 public:
+#if !WITH_EDITOR
+	const FMaterialCachedTexturesSamplingInfo* GetCachedTexturesSamplingInfo() const
+	{
+		return CachedTexturesSamplingInfo.GetPtrOrNull();
+	}
+	virtual FTextureSamplingInfo CalculateTexturesSamplingInfo(UTexture* Texture);
+#endif
+
 	/** SubsurfaceProfile, for Screen Space Subsurface Scattering.. */
 	UPROPERTY(EditAnywhere, BlueprintReadOnly, Category = Material, meta = (DisplayName = "Subsurface Profile"))
 	TObjectPtr<class USubsurfaceProfile> SubsurfaceProfile;
@@ -301,7 +340,6 @@ protected:
 	UPROPERTY(EditAnywhere, Category=Lightmass)
 	struct FLightmassMaterialInterfaceSettings LightmassSettings;
 
-protected:
 #if WITH_EDITORONLY_DATA
 	/** Because of redirector, the texture names need to be resorted at each load in case they changed. */
 	UPROPERTY(transient)
@@ -318,6 +356,15 @@ protected:
 	UPROPERTY(EditAnywhere, AdvancedDisplay, Instanced, Category = Material)
 	TArray<TObjectPtr<UAssetUserData>> AssetUserData;
 
+#if !WITH_EDITOR
+	void CacheTexturesSamplingInfo();
+	virtual bool CanCacheTexturesSamplingInfo() const;
+#endif
+
+	/** Pre-cached texture sampling information used for texture streaming (calculated on load) **/
+	UPROPERTY(Transient)
+	TOptional<FMaterialCachedTexturesSamplingInfo> CachedTexturesSamplingInfo;
+
 private:
 	/** Feature levels to force to compile. */
 	uint32 FeatureLevelsToForceCompile;
@@ -1175,6 +1222,10 @@ public:
 	*/
 	ENGINE_API virtual float GetTextureDensity(FName TextureName, const struct FMeshUVChannelInfo& UVChannelData) const;
 
+#if !WITH_EDITOR
+	ENGINE_API float GetTextureDensityWithCache(const FTextureSamplingInfo& TextureSamplingInfo, const FMeshUVChannelInfo& UVChannelData) const;
+#endif
+
 	ENGINE_API virtual void PreSave(FObjectPreSaveContext ObjectSaveContext) override;
 
 	/**
