Fix hot reload for UTexture classes
* FVTableHelper constructor defined in UTexture must call its Super's FVTableHelper constructor.
* Failure to do so will cause the default constructor of UStreamableRenderAsset to be called with a default FObjectInitializer instead.  FObjectInitializer::Get() asserts with missing TopInitializer that would be pushed during object creation.

diff --git a/Engine/Source/Runtime/Engine/Private/Texture.cpp b/Engine/Source/Runtime/Engine/Private/Texture.cpp
--- a/Engine/Source/Runtime/Engine/Private/Texture.cpp
+++ b/Engine/Source/Runtime/Engine/Private/Texture.cpp
@@ -259,7 +259,8 @@ FTextureResource* UTexture::GetResource()
 }
 
 UTexture::UTexture(FVTableHelper& Helper)
-	: TextureReference(*new FTextureReference())
+	: Super(Helper)
+	, TextureReference(*new FTextureReference())
 {
 }
 
