Optimize component space to bind space conversion when sending matrices to render thread

diff --git a/Engine/Source/Runtime/Engine/Private/SkeletalRender.cpp b/Engine/Source/Runtime/Engine/Private/SkeletalRender.cpp
--- a/Engine/Source/Runtime/Engine/Private/SkeletalRender.cpp
+++ b/Engine/Source/Runtime/Engine/Private/SkeletalRender.cpp
@@ -366,7 +366,8 @@ void UpdateRefToLocalMatricesInner(TArray<FMatrix44f>& ReferenceToLocal, const T
 		const TArray<FBoneIndexType>& RequiredBoneIndices = *RequiredBoneSets[RequiredBoneSetIndex];
 
 		// Get the index of the bone in this skeleton, and loop up in table to find index in parent component mesh.
-		for (int32 BoneIndex = 0; BoneIndex < RequiredBoneIndices.Num(); BoneIndex++)
+		const int32 NumRequiredBoneIndices = RequiredBoneIndices.Num();
+		for (int32 BoneIndex = 0; BoneIndex < NumRequiredBoneIndices; BoneIndex++)
 		{
 			const int32 ThisBoneIndex = RequiredBoneIndices[BoneIndex];
 
@@ -442,18 +443,20 @@ void UpdateRefToLocalMatricesInner(TArray<FMatrix44f>& ReferenceToLocal, const T
 		}
 	}
 
+	const int32 NumReferenceToLocal = ReferenceToLocal.Num();
+
 	if (LeaderBoneMappedComponentSpaceTransform && bIsLeaderCompValid)
 	{
-		LeaderBoneMappedComponentSpaceTransform->SetNumUninitialized(ReferenceToLocal.Num());
-		for (int32 ThisBoneIndex = 0; ThisBoneIndex < ReferenceToLocal.Num(); ++ThisBoneIndex)
+		LeaderBoneMappedComponentSpaceTransform->SetNumUninitialized(NumReferenceToLocal);
+		for (int32 ThisBoneIndex = 0; ThisBoneIndex < NumReferenceToLocal; ++ThisBoneIndex)
 		{
 			(*LeaderBoneMappedComponentSpaceTransform)[ThisBoneIndex] = FTransform((FMatrix)ReferenceToLocal[ThisBoneIndex]);
 		}
 	}
 
-	for (int32 ThisBoneIndex = 0; ThisBoneIndex < ReferenceToLocal.Num(); ++ThisBoneIndex)
+	for (int32 ThisBoneIndex = 0; ThisBoneIndex < NumReferenceToLocal; ++ThisBoneIndex)
 	{
-		ReferenceToLocal[ThisBoneIndex] = (*RefBasesInvMatrix)[ThisBoneIndex] * ReferenceToLocal[ThisBoneIndex];
+		VectorMatrixMultiply(&ReferenceToLocal[ThisBoneIndex], &(*RefBasesInvMatrix)[ThisBoneIndex], &ReferenceToLocal[ThisBoneIndex]);
 	}
 }
 /**
