Moves System.Engine.Loading tests that rely on EngineTest assets into EngineTest to prevent them from failing when that project isn't in use.

#rnx

diff --git a/Engine/Source/Runtime/Engine/Private/Tests/Loading/AsyncLoadingTests.cpp b/Engine/Source/Runtime/Engine/Private/Tests/Loading/AsyncLoadingTests.cpp
--- a/Engine/Source/Runtime/Engine/Private/Tests/Loading/AsyncLoadingTests.cpp
+++ b/Engine/Source/Runtime/Engine/Private/Tests/Loading/AsyncLoadingTests.cpp
@@ -63,186 +63,5 @@ bool FThreadSafeAsyncLoadingTest::RunTest(const FString& Parameters)
 	return true;
 }
 
-/**
- * This test demonstrates that LoadPackage can load blueprints with circular dependencies who rely on dependencies with circular dependencies
- */
-IMPLEMENT_SIMPLE_AUTOMATION_TEST(FCircularDependencyLoadingTest, TEXT(TEST_NAME_ROOT ".LoadBlueprintWithCircularDependencyTest"), EAutomationTestFlags::EditorContext | EAutomationTestFlags::EngineFilter)
-bool FCircularDependencyLoadingTest::RunTest(const FString& Parameters)
-{
-	static constexpr const TCHAR* ActorWithCircularReferences = TEXT("/Game/Tests/Core/AssetLoading/RecursiveLoads/BlueprintActorWithCircularReferences");
-
-	UPackage* Package = LoadPackage(nullptr, ActorWithCircularReferences, LOAD_None);
-	TestTrue(TEXT("The object should have been properly loaded recursively"), Package != nullptr);
-
-	Package = FindPackage(nullptr, ActorWithCircularReferences);
-	TestTrue(TEXT("The object should have been properly loaded recursively"), Package != nullptr);
-
-	CollectGarbage(GARBAGE_COLLECTION_KEEPFLAGS);
-
-	return true;
-}
-
-/**
- * PackagePath1 leads to a user defined structure containing UFields which will be forcibly converted to FField when loading. The UFields should be marked as invalid and the loader should not consider them for loading.
- */
-IMPLEMENT_SIMPLE_AUTOMATION_TEST(FLoadingTests_ImportPackageConvertedOnLoad, TEXT(TEST_NAME_ROOT ".InvalidExports.ImportPackageConvertedOnLoad"), EAutomationTestFlags::EditorContext | EAutomationTestFlags::EngineFilter)
-bool FLoadingTests_ImportPackageConvertedOnLoad::RunTest(const FString& Parameters)
-{
-	static constexpr const TCHAR* PackagePath1 = TEXT("/Game/Tests/Core/AssetLoading/InvalidatingExports/OldProcMeshData");
-	static constexpr const TCHAR* PackagePath2 = TEXT("/Game/Tests/Core/AssetLoading/InvalidatingExports/BP_UsingUStructWithUnconvertedFields");
-	static constexpr const TCHAR* ObjectName = TEXT("ProcMeshData");
-	typedef TSet<TPair<FName, FName>> FObjPathAndClassNameSet;
-	FObjPathAndClassNameSet ExpectedObjects;
-	FObjPathAndClassNameSet ActualObjects;
-
-	auto VerifyLoadedObjects = [this](const FObjPathAndClassNameSet& ExpectedObjects, const FObjPathAndClassNameSet& ActualObjects)
-		{
-			TestTrue(FString::Printf(TEXT("Expected the same number of objects in package %s after reloading. Expected: %d != Actual: %d"), PackagePath1, ExpectedObjects.Num(), ActualObjects.Num()), ExpectedObjects.Num() == ActualObjects.Num());
-			for (const TPair<FName, FName>& Pair : ExpectedObjects)
-			{
-				TestTrue(FString::Printf(TEXT("Missing obj %s (type: %s) after reloading %s"), *Pair.Key.ToString(), *Pair.Value.ToString(), PackagePath1), ActualObjects.Find(Pair) != nullptr);
-			}
-		};
-
-	// Load our package and populate the ExpectedObjects set
-	{
-		UPackage* Package = LoadPackage(nullptr, PackagePath1, LOAD_None);
-		TestTrue(FString::Printf(TEXT("Failed to load package at %s"), PackagePath1), Package != nullptr);
-
-		ForEachObjectWithPackage(Package,
-			[&ExpectedObjects](UObject* Object)
-			{
-				ExpectedObjects.Add({ Object->GetFName(), Object->GetClass()->GetFName() });
-				return true;
-			}
-		);
-
-		FLoadingTestsScope::GarbageCollect({ {PackagePath1} }, *this);
-	}
-
-	// Load a package that imports the objects from the package at PackagePath1 and ensure we still load only the objects we expect
-	{
-		UPackage* Package = LoadPackage(nullptr, PackagePath2, LOAD_None);
-		TestTrue(FString::Printf(TEXT("Failed to load package at %s"), PackagePath2), Package != nullptr);
-
-		// Path2 imports PackagePath1 so we should be able to still find our Obj and our PackagePath1 package
-		UUserDefinedStruct* Obj = FindFirstObject<UUserDefinedStruct>(ObjectName);
-		TestTrue(FString::Printf(TEXT("Failed to find expected object %s in package %s"), ObjectName, PackagePath1), Obj != nullptr);
-		if (Obj == nullptr)
-		{
-			return false;
-		}
-
-		UPackage* PackagePath1Package = Obj->GetPackage();
-		ForEachObjectWithPackage(PackagePath1Package,
-			[&ActualObjects](UObject* Object)
-			{
-				ActualObjects.Add({ Object->GetFName(), Object->GetClass()->GetFName() });
-				return true;
-			}
-		);
-
-		VerifyLoadedObjects(ExpectedObjects, ActualObjects);
-		ActualObjects.Empty();
-
-		// Note, we are only GC'ing the importing package.
-		Package->ClearFlags(RF_Standalone);
-		FLoadingTestsScope::GarbageCollect({ { PackagePath2 } }, *this);
-	}
-
-	// Load the old package from memory and ensure we don't have any new objects created from importing
-	{
-		UPackage* Package = LoadPackage(nullptr, PackagePath1, LOAD_None);
-		TestTrue(FString::Printf(TEXT("Failed to load package at %s"), PackagePath1), Package != nullptr);
-
-		ForEachObjectWithPackage(Package,
-			[&ActualObjects](UObject* Object)
-			{
-				ActualObjects.Add({ Object->GetFName(), Object->GetClass()->GetFName() });
-				return true;
-			}
-		);
-
-		VerifyLoadedObjects(ExpectedObjects, ActualObjects);
-		ActualObjects.Empty();
-
-		FLoadingTestsScope::GarbageCollect({ { PackagePath2 } }, *this);
-	}
-
-	return true;
-}
-
-
-/**
- * PackagePath1 leads to a user defined structure containing UFields which will be forcibly converted to FField when loading.
- * When a sub-object is deleted, if the package is requested again we should reload the missing object and not create objects marked as invalid
- */
-IMPLEMENT_SIMPLE_AUTOMATION_TEST(FLoadingTests_ReloadDestroyedExport, TEXT(TEST_NAME_ROOT ".InvalidExports.ReloadDestroyedExport"), EAutomationTestFlags::EditorContext | EAutomationTestFlags::EngineFilter)
-bool FLoadingTests_ReloadDestroyedExport::RunTest(const FString& Parameters)
-{
-	static constexpr const TCHAR* PackagePath1 = TEXT("/Game/Tests/Core/AssetLoading/InvalidatingExports/OldProcMeshData");
-	static constexpr const TCHAR* PackagePath2 = TEXT("/Game/Tests/Core/AssetLoading/InvalidatingExports/BP_UsingUStructWithUnconvertedFields");
-	static constexpr const TCHAR* ObjectName = TEXT("ProcMeshData");
-	typedef TSet<TPair<FName, FName>> FObjPathAndClassNameSet;
-	FObjPathAndClassNameSet ExpectedObjects;
-	FObjPathAndClassNameSet ActualObjects;
-
-	auto VerifyLoadedObjects = [this](const FObjPathAndClassNameSet& ExpectedObjects, const FObjPathAndClassNameSet& ActualObjects)
-		{
-			TestTrue(FString::Printf(TEXT("Expected the same number of objects in package %s after reloading. Expected: %d != Actual: %d"), PackagePath1, ExpectedObjects.Num(), ActualObjects.Num()), ExpectedObjects.Num() == ActualObjects.Num());
-			for (const TPair<FName, FName>& Pair : ExpectedObjects)
-			{
-				TestTrue(FString::Printf(TEXT("Missing obj %s (type: %s) after reloading %s"), *Pair.Key.ToString(), *Pair.Value.ToString(), PackagePath1), ActualObjects.Find(Pair) != nullptr);
-			}
-		};
-
-	// Load our package and populate the ExpectedObjects set
-	{
-		UPackage* Package = LoadPackage(nullptr, PackagePath1, LOAD_None);
-		TestTrue(FString::Printf(TEXT("Failed to load package at %s"), PackagePath1), Package != nullptr);
-
-		ForEachObjectWithPackage(Package,
-			[&ExpectedObjects](UObject* Object)
-			{
-				ExpectedObjects.Add({ Object->GetFName(), Object->GetClass()->GetFName() });
-				return true;
-			}
-		);
-
-		// Find a sub-object from the package and mark it for garbage collection, but do not destroy the package entirely
-		UUserDefinedStruct* Obj = FindObject<UUserDefinedStruct>(Package, ObjectName);
-		TestTrue(FString::Printf(TEXT("Failed to find expected object %s in package %s"), ObjectName, PackagePath1), Obj != nullptr);
-		Obj->ClearFlags(RF_Standalone);
-
-		// Note we are not using FLoadingTestsScope::GarbageCollect as that will mark all sub-objects and the package for collection
-		CollectGarbage(RF_Standalone);
-	}
-
-	// Reloading the package should re-populate any deleted sub-objects as long as the objects weren't marked as invalid to load
-	{
-		UPackage* Package = LoadPackage(nullptr, PackagePath1, LOAD_None);
-		TestTrue(FString::Printf(TEXT("Failed to load package at %s"), PackagePath1), Package != nullptr);
-
-		ForEachObjectWithPackage(Package,
-			[&ActualObjects](UObject* Object)
-			{
-				ActualObjects.Add({ Object->GetFName(), Object->GetClass()->GetFName() });
-				return true;
-			}
-		);
-
-		// Our deleted object should have been restored
-		UUserDefinedStruct* Obj = FindObject<UUserDefinedStruct>(Package, ObjectName);
-		TestTrue(FString::Printf(TEXT("Failed to find expected object %s in package %s"), ObjectName, PackagePath1), Obj != nullptr);
-
-		VerifyLoadedObjects(ExpectedObjects, ActualObjects);
-		ActualObjects.Empty();
-
-		FLoadingTestsScope::GarbageCollect({ {PackagePath1} }, *this);
-	}
-
-	return true;
-}
-
 #undef TEST_NAME_ROOT
 #endif // WITH_DEV_AUTOMATION_TESTS
\ No newline at end of file
