fixing listener focus Focus feature was accidentally conflated with spatialize settings in attenution settings in CL 6940530.

issue described in jira was actually due to the dependency on spatialization checkmark not on if it was inside a sound cue node override, etc.

[RN] fixed bug where attenuation spatialization and focus settings were dependent. focus works independent of spatialize.

diff --git a/Engine/Source/Runtime/Engine/Private/ActiveSound.cpp b/Engine/Source/Runtime/Engine/Private/ActiveSound.cpp
--- a/Engine/Source/Runtime/Engine/Private/ActiveSound.cpp
+++ b/Engine/Source/Runtime/Engine/Private/ActiveSound.cpp
@@ -52,6 +52,15 @@ FAutoConsoleVariableRef CVarInitializeFocusFactorOnFirstUpdateCVar(
 	TEXT("0: Disabled, 1: Enabled (default)"),
 	ECVF_Default);
 
+
+static int32 FocusDependsOnSpatializationCVar = 0;
+FAutoConsoleVariableRef CVarFocusDependsOnSpatialization(
+	TEXT("au.FocusData.FocusDependsOnSpatialization"),
+	FocusDependsOnSpatializationCVar,
+	TEXT("When set to 1, focus factor on attenuation settings will not be enabled if spatialization is not enabled on the attenuation setting.\n")
+	TEXT("0: Focus and Spatialization are orthogonal, 1: Focus will be disabled if spatialization is not also enabled"),
+	ECVF_Default);
+
 FTraceDelegate FActiveSound::ActiveSoundTraceDelegate;
 TMap<FTraceHandle, FActiveSound::FAsyncTraceDetails> FActiveSound::TraceToActiveSoundMap;
 
@@ -1724,7 +1733,10 @@ void FActiveSound::UpdateFocusData(float DeltaTime, const FAttenuationListenerDa
 		return;
 	}
 
-	if (!ListenerData.AttenuationSettings->bSpatialize)
+	// This cvar is used to recover the incorrect state of spatialization and focus being dependent on each other
+	// The original design of this system was each attenuation setting category was orthogonal (i.e. focus should be applicable independent of spatialization)
+	// This cvar exists to allow a swap back to the old state if needed
+	if (FocusDependsOnSpatializationCVar == 1 && !ListenerData.AttenuationSettings->bSpatialize)
 	{
 		return;
 	}
