- Optimized FStaticRenderAssetInstanceManager::Add by replacing the allocated result array by a function param only called when the add fails.
- Adapted EStaticBuildStep::WaitForRegistration phase to also support r.Streaming.AllowParallelRenderAssetStreamingManagerProcessComponents = true (removes peaks)
- Added r.Streaming.WorkerCountForRenderAssetParallelProcessComponents to control the maximum number of workers to use when using r.Streaming.AllowParallelRenderAssetStreamingManagerProcessComponents.
- Fixed inverse logic in FStaticRenderAssetInstanceManager::CanAddComponent()

diff --git a/Engine/Source/Runtime/Engine/Private/Streaming/DynamicTextureInstanceManager.cpp b/Engine/Source/Runtime/Engine/Private/Streaming/DynamicTextureInstanceManager.cpp
--- a/Engine/Source/Runtime/Engine/Private/Streaming/DynamicTextureInstanceManager.cpp
+++ b/Engine/Source/Runtime/Engine/Private/Streaming/DynamicTextureInstanceManager.cpp
@@ -110,11 +110,12 @@ void FDynamicRenderAssetInstanceManager::IncrementalUpdate(FRemovedRenderAssetAr
 	if (PendingComponents.Num())
 	{
 		const int32 ToProcessComponentCount = PendingComponents.Num();
-		if (FRenderAssetStreamingManager::IsParallelProcessComponentsAllowed(ToProcessComponentCount))
+		int32 MinBatchSize = 1;
+		if (FRenderAssetStreamingManager::IsParallelProcessComponentsAllowed(ToProcessComponentCount, MinBatchSize))
 		{
 			TArray<FRenderAssetInstanceState::FPreAddComponentIgnoreBoundsPayload> Payloads;
 			Payloads.SetNum(ToProcessComponentCount);
-			ParallelFor(TEXT("PreAddComponentsIgnoreBounds"), ToProcessComponentCount, 1, [&](int32 Index)
+			ParallelFor(TEXT("PreAddComponentsIgnoreBounds"), ToProcessComponentCount, MinBatchSize, [&](int32 Index)
 			{
 				FOptionalTaskTagScope Scope(ETaskTag::EParallelGameThread);
 				const UPrimitiveComponent* Component = PendingComponents[Index];
diff --git a/Engine/Source/Runtime/Engine/Private/Streaming/LevelTextureManager.cpp b/Engine/Source/Runtime/Engine/Private/Streaming/LevelTextureManager.cpp
--- a/Engine/Source/Runtime/Engine/Private/Streaming/LevelTextureManager.cpp
+++ b/Engine/Source/Runtime/Engine/Private/Streaming/LevelTextureManager.cpp
@@ -181,14 +181,15 @@ void FLevelRenderAssetManager::IncrementalBuild(FDynamicRenderAssetInstanceManag
 
 		const int32 UnprocessedComponentCount = UnprocessedComponents.Num();
 		const int32 ToProcessComponentCount = bForceCompletion ? UnprocessedComponentCount : FMath::Min<int32>(NumStepsLeft, UnprocessedComponentCount);
-		if (FRenderAssetStreamingManager::IsParallelProcessComponentsAllowed(ToProcessComponentCount) && !StaticInstances.CanAddComponent())
+		int32 MinBatchSize = 1;
+		if (FRenderAssetStreamingManager::IsParallelProcessComponentsAllowed(ToProcessComponentCount, MinBatchSize) && StaticInstances.CanAddComponent())
 		{
 			TArray<FRenderAssetInstanceState::FPreAddComponentPayload> Payloads;
 			Payloads.SetNum(ToProcessComponentCount);
 			// Because LevelContext is not thread-safe when ComponentBuildData is valid, we force not to use built data
 			{
 				LevelContext.SetForceNoUseBuiltData(true);
-				ParallelFor(TEXT("PreAddComponents"), ToProcessComponentCount, 1, [&](int32 Index)
+				ParallelFor(TEXT("PreAddComponents"), ToProcessComponentCount, MinBatchSize, [&](int32 Index)
 				{
 					FOptionalTaskTagScope Scope(ETaskTag::EParallelGameThread);
 					const UPrimitiveComponent* Primitive = UnprocessedComponents[UnprocessedComponentCount - Index - 1];
@@ -197,22 +198,20 @@ void FLevelRenderAssetManager::IncrementalBuild(FDynamicRenderAssetInstanceManag
 				LevelContext.SetForceNoUseBuiltData(false);
 			}
 
-			TArray<EAddComponentResult> AddResults;
-			StaticInstances.Add(Payloads, AddResults);
-			for (int32 i=0; i < ToProcessComponentCount; ++i)
+			StaticInstances.Add(Payloads, [this, &DynamicComponentManager, &LevelContext, bLevelIsVisible](const FRenderAssetInstanceState::FPreAddComponentPayload& Payload, EAddComponentResult AddResult)
 			{
-				EAddComponentResult AddResult = AddResults[i];
+				check(AddResult != EAddComponentResult::Success);
 				if (AddResult == EAddComponentResult::Fail && !bLevelIsVisible)
 				{
 					// Retry once the level becomes visible.
-					PendingComponents.Add(Payloads[i].Component);
+					PendingComponents.Add(Payload.Component);
 				}
-				else if (AddResult != EAddComponentResult::Success)
+				else
 				{
 					// Here we also handle the case for Fail_UIDensityConstraint.
-					SetAsDynamic(DynamicComponentManager, LevelContext, Payloads[i].Component);
+					SetAsDynamic(DynamicComponentManager, LevelContext, Payload.Component);
 				}
-			}
+			});
 
 			NumStepsLeft -= ToProcessComponentCount;
 			UnprocessedComponents.SetNum(UnprocessedComponentCount - ToProcessComponentCount, EAllowShrinking::No);
@@ -278,13 +277,39 @@ void FLevelRenderAssetManager::IncrementalBuild(FDynamicRenderAssetInstanceManag
 			NumStepsLeft -= (int64)PendingComponents.Num();
 
 			// Reprocess the components that didn't have valid data.
-			while (PendingComponents.Num())
+			const int32 ToProcessComponentCount = PendingComponents.Num();
+			int32 MinBatchSize = 1;
+			if (FRenderAssetStreamingManager::IsParallelProcessComponentsAllowed(ToProcessComponentCount, MinBatchSize) && StaticInstances.CanAddComponent())
 			{
-				const UPrimitiveComponent* Primitive = PendingComponents.Pop(EAllowShrinking::No);
+				TArray<FRenderAssetInstanceState::FPreAddComponentPayload> Payloads;
+				Payloads.SetNum(ToProcessComponentCount);
+				// Because LevelContext is not thread-safe when ComponentBuildData is valid, we force not to use built data
+				{
+					LevelContext.SetForceNoUseBuiltData(true);
+					ParallelFor(TEXT("PreAddComponents"), ToProcessComponentCount, MinBatchSize, [&](int32 Index)
+					{
+						FOptionalTaskTagScope Scope(ETaskTag::EParallelGameThread);
+						FRenderAssetInstanceState::PreAddComponent(PendingComponents[Index], LevelContext, MaxTextureUVDensity, Payloads[Index]);
+					});
+					LevelContext.SetForceNoUseBuiltData(false);
+				}
 
-				if (StaticInstances.Add(Primitive, LevelContext, MaxTextureUVDensity) != EAddComponentResult::Success)
+				StaticInstances.Add(Payloads, [this, &DynamicComponentManager, &LevelContext](const FRenderAssetInstanceState::FPreAddComponentPayload& Payload, EAddComponentResult AddResult)
 				{
-					SetAsDynamic(DynamicComponentManager, LevelContext, Primitive);
+					check(AddResult != EAddComponentResult::Success);
+					SetAsDynamic(DynamicComponentManager, LevelContext, Payload.Component);
+				});
+			}
+			else
+			{
+				while (PendingComponents.Num())
+				{
+					const UPrimitiveComponent* Primitive = PendingComponents.Pop(EAllowShrinking::No);
+
+					if (StaticInstances.Add(Primitive, LevelContext, MaxTextureUVDensity) != EAddComponentResult::Success)
+					{
+						SetAsDynamic(DynamicComponentManager, LevelContext, Primitive);
+					}
 				}
 			}
 
diff --git a/Engine/Source/Runtime/Engine/Private/Streaming/StaticTextureInstanceManager.cpp b/Engine/Source/Runtime/Engine/Private/Streaming/StaticTextureInstanceManager.cpp
--- a/Engine/Source/Runtime/Engine/Private/Streaming/StaticTextureInstanceManager.cpp
+++ b/Engine/Source/Runtime/Engine/Private/Streaming/StaticTextureInstanceManager.cpp
@@ -83,15 +83,15 @@ void FStaticRenderAssetInstanceManager::Refresh(float Percentage)
 
 bool FStaticRenderAssetInstanceManager::CanAddComponent() const
 {
-	return !!AsyncView;
+	return !AsyncView;
 }
 
-EAddComponentResult FStaticRenderAssetInstanceManager::Add(const TArray<FRenderAssetInstanceState::FPreAddComponentPayload>& Payloads, TArray<EAddComponentResult>& OutResults)
+EAddComponentResult FStaticRenderAssetInstanceManager::Add(const TArray<FRenderAssetInstanceState::FPreAddComponentPayload>& Payloads, TFunctionRef<void(const FRenderAssetInstanceState::FPreAddComponentPayload&, EAddComponentResult)> OnAddFailedFunc)
 {
 	if (!AsyncView)
 	{
 		FRenderAssetInstanceState* State = StateSync.SyncAndGetState();
-		return State->AddComponents(Payloads, OutResults);
+		return State->AddComponents(Payloads, OnAddFailedFunc);
 	}
 	return EAddComponentResult::Fail;
 }
diff --git a/Engine/Source/Runtime/Engine/Private/Streaming/StaticTextureInstanceManager.h b/Engine/Source/Runtime/Engine/Private/Streaming/StaticTextureInstanceManager.h
--- a/Engine/Source/Runtime/Engine/Private/Streaming/StaticTextureInstanceManager.h
+++ b/Engine/Source/Runtime/Engine/Private/Streaming/StaticTextureInstanceManager.h
@@ -52,7 +52,7 @@ public:
 	EAddComponentResult Add(const UPrimitiveComponent* Component, FStreamingTextureLevelContext& LevelContext, float MaxAllowedUIDensity) final override;
 
 	/** Add a batch of components streaming data. */
-	EAddComponentResult Add(const TArray<FRenderAssetInstanceState::FPreAddComponentPayload>& Payloads, TArray<EAddComponentResult>& OutResults);
+	EAddComponentResult Add(const TArray<FRenderAssetInstanceState::FPreAddComponentPayload>& Payloads, TFunctionRef<void(const FRenderAssetInstanceState::FPreAddComponentPayload&, EAddComponentResult)> OnAddFailedFunc);
 
 	/** Remove a component, the RemoveTextures is the list of textures not referred anymore. */
 	 void Remove(const UPrimitiveComponent* Component, FRemovedRenderAssetArray* RemovedRenderAssets) final override;
diff --git a/Engine/Source/Runtime/Engine/Private/Streaming/StreamingManagerTexture.cpp b/Engine/Source/Runtime/Engine/Private/Streaming/StreamingManagerTexture.cpp
--- a/Engine/Source/Runtime/Engine/Private/Streaming/StreamingManagerTexture.cpp
+++ b/Engine/Source/Runtime/Engine/Private/Streaming/StreamingManagerTexture.cpp
@@ -106,6 +106,14 @@ namespace UE::Private::RenderAssetStreamingManager
 		TEXT("Whether it is allowed for RenderAssetStreamingManager to process components to add with a ParallelFor to use more cores. ")
 		TEXT("When enabled, texture streaming built data is not used to retrieve the primitive streaming render asset information."),
 		ECVF_Default);
+
+	static int32 GWorkerCountForRenderAssetParallelProcessComponents = 0;
+	static FAutoConsoleVariableRef CVarWorkerCountForRenderAssetParallelProcessComponents(
+		TEXT("r.Streaming.WorkerCountForRenderAssetParallelProcessComponents"),
+		GWorkerCountForRenderAssetParallelProcessComponents,
+		TEXT("Set the max number of workers when r.Streaming.AllowParallelRenderAssetStreamingManagerProcessComponents is enabled. ")
+		TEXT("Only taked into account if value > 0."),
+		ECVF_Default);
 }
 
 bool TrackRenderAsset( const FString& AssetName );
@@ -425,10 +433,12 @@ int64 FRenderAssetStreamingManager::GetPoolSize() const
 	return GTexturePoolSize;
 }
 
-bool FRenderAssetStreamingManager::IsParallelProcessComponentsAllowed(int32 NumComponentsToProcess)
+bool FRenderAssetStreamingManager::IsParallelProcessComponentsAllowed(int32 NumComponentsToProcess, int32& OutParallelForMinBatchSize)
 {
 	static int32 MinComponentsToProcessForParallelProcessing = 8;
-	return UE::Private::RenderAssetStreamingManager::bAllowParallelProcessComponents && FApp::ShouldUseThreadingForPerformance() && (NumComponentsToProcess >= MinComponentsToProcessForParallelProcessing);
+	bool bIsAllowed = UE::Private::RenderAssetStreamingManager::bAllowParallelProcessComponents && FApp::ShouldUseThreadingForPerformance() && (NumComponentsToProcess >= MinComponentsToProcessForParallelProcessing);
+	OutParallelForMinBatchSize = (bIsAllowed && (UE::Private::RenderAssetStreamingManager::GWorkerCountForRenderAssetParallelProcessComponents > 0) && (NumComponentsToProcess > UE::Private::RenderAssetStreamingManager::GWorkerCountForRenderAssetParallelProcessComponents)) ? FMath::DivideAndRoundUp(NumComponentsToProcess, UE::Private::RenderAssetStreamingManager::GWorkerCountForRenderAssetParallelProcessComponents) : 1;
+	return bIsAllowed;
 }
 
 void FRenderAssetStreamingManager::IncrementalUpdate(float Percentage, bool bUpdateDynamicComponents)
diff --git a/Engine/Source/Runtime/Engine/Private/Streaming/StreamingManagerTexture.h b/Engine/Source/Runtime/Engine/Private/Streaming/StreamingManagerTexture.h
--- a/Engine/Source/Runtime/Engine/Private/Streaming/StreamingManagerTexture.h
+++ b/Engine/Source/Runtime/Engine/Private/Streaming/StreamingManagerTexture.h
@@ -193,7 +193,7 @@ struct FRenderAssetStreamingManager final : public IRenderAssetStreamingManager
 	void SetRenderAssetsRemovedTimestamp(const FRemovedRenderAssetArray& RemovedRenderAssets);
 
 	/** Whether it is allowed for this manager to process components to add with a ParallelFor to use more cores. */
-	static bool IsParallelProcessComponentsAllowed(int32 NumComponentsToProcess);
+	static bool IsParallelProcessComponentsAllowed(int32 NumComponentsToProcess, int32& OutParallelForMinBatchSize);
 
 private:
 //BEGIN: Thread-safe functions and data
diff --git a/Engine/Source/Runtime/Engine/Private/Streaming/TextureInstanceState.cpp b/Engine/Source/Runtime/Engine/Private/Streaming/TextureInstanceState.cpp
--- a/Engine/Source/Runtime/Engine/Private/Streaming/TextureInstanceState.cpp
+++ b/Engine/Source/Runtime/Engine/Private/Streaming/TextureInstanceState.cpp
@@ -325,15 +325,18 @@ void FRenderAssetInstanceState::AddRenderAssetElements(const UPrimitiveComponent
 	});
 }
 
-EAddComponentResult FRenderAssetInstanceState::AddComponents(const TArray<FRenderAssetInstanceState::FPreAddComponentPayload>& Payloads, TArray<EAddComponentResult>& OutResults)
+EAddComponentResult FRenderAssetInstanceState::AddComponents(const TArray<FRenderAssetInstanceState::FPreAddComponentPayload>& Payloads, TFunctionRef<void(const FRenderAssetInstanceState::FPreAddComponentPayload&, EAddComponentResult)> OnAddFailedFunc)
 {
 	TRACE_CPUPROFILER_EVENT_SCOPE(FRenderAssetInstanceState::AddComponents);
 	checkf(!bIsDynamicInstanceState, TEXT("Error: trying to add component to dynamic instance manager as static"));
 	const int32 PayloadCount = Payloads.Num();
-	OutResults.SetNumUninitialized(PayloadCount);
 	for (int i=0; i<PayloadCount; ++i)
 	{
-		OutResults[i] = AddComponent(Payloads[i]);
+		EAddComponentResult Result = AddComponent(Payloads[i]);
+		if (Result != EAddComponentResult::Success)
+		{
+			OnAddFailedFunc(Payloads[i], Result);
+		}
 	}
 	return EAddComponentResult::Success;
 }
@@ -344,9 +347,10 @@ EAddComponentResult FRenderAssetInstanceState::AddComponent(const FPreAddCompone
 	if (Payload.Result == EAddComponentResult::Success)
 	{
 		int32* ComponentLink = ComponentMap.Find(Payload.Component);
+		const float LastRenderTime = Payload.Component->GetLastRenderTimeOnScreen();
 		for (const FPreAddBounds& Bounds : Payload.Bounds)
 		{
-			const int32 BoundsIndex = AddBounds(Bounds.InfoBounds, Bounds.PackedRelativeBox, Payload.Component, Payload.Component->GetLastRenderTimeOnScreen(), Bounds.RangeOrigin, Bounds.MinDistanceSq, Bounds.MinRangeSq, Bounds.MaxRangeSq);
+			const int32 BoundsIndex = AddBounds(Bounds.InfoBounds, Bounds.PackedRelativeBox, Payload.Component, LastRenderTime, Bounds.RangeOrigin, Bounds.MinDistanceSq, Bounds.MinRangeSq, Bounds.MaxRangeSq);
 			AddRenderAssetElements(Payload.Component, Bounds.RenderAssets, BoundsIndex, ComponentLink);
 			if (Bounds.bIsPacked)
 			{
diff --git a/Engine/Source/Runtime/Engine/Private/Streaming/TextureInstanceState.h b/Engine/Source/Runtime/Engine/Private/Streaming/TextureInstanceState.h
--- a/Engine/Source/Runtime/Engine/Private/Streaming/TextureInstanceState.h
+++ b/Engine/Source/Runtime/Engine/Private/Streaming/TextureInstanceState.h
@@ -82,7 +82,7 @@ public:
 	// Will also remove bounds
 	EAddComponentResult AddComponent(const UPrimitiveComponent* Component, FStreamingTextureLevelContext& LevelContext, float MaxAllowedUIDensity);
 	EAddComponentResult AddComponent(const FPreAddComponentPayload& Payload);
-	EAddComponentResult AddComponents(const TArray<FRenderAssetInstanceState::FPreAddComponentPayload>& Payloads, TArray<EAddComponentResult>& OutResults);
+	EAddComponentResult AddComponents(const TArray<FRenderAssetInstanceState::FPreAddComponentPayload>& Payloads, TFunctionRef<void(const FRenderAssetInstanceState::FPreAddComponentPayload&, EAddComponentResult)> OnAddFailedFunc);
 
 	// Similar to AddComponent, but ignore the streaming data bounds. Used for dynamic components. A faster implementation that does less processing.
 	EAddComponentResult AddComponentIgnoreBounds(const UPrimitiveComponent* Component, FStreamingTextureLevelContext& LevelContext);
