Add Default Button to Scalability Menu in Editor that will reset the scalability to what the device profile initially set

diff --git a/Engine/Source/Editor/UnrealEd/Private/EditorEngine.cpp b/Engine/Source/Editor/UnrealEd/Private/EditorEngine.cpp
--- a/Engine/Source/Editor/UnrealEd/Private/EditorEngine.cpp
+++ b/Engine/Source/Editor/UnrealEd/Private/EditorEngine.cpp
@@ -6441,7 +6441,7 @@ void UEditorEngine::SetPreviewPlatform(const FPreviewPlatformInfo& NewPreviewPla
 		}
 	}
 
-	Scalability::ApplyCachedQualityLevelForShaderPlatform(GetActiveShaderPlatform());
+	Scalability::ApplyCachedQualityLevels();
 
 	PreviewPlatformChanged.Broadcast();
 
@@ -6497,7 +6497,7 @@ void UEditorEngine::ToggleFeatureLevelPreview()
 		}
 	}
 
-	Scalability::ApplyCachedQualityLevelForShaderPlatform(GetActiveShaderPlatform());
+	Scalability::ApplyCachedQualityLevels();
 	OnEffectivePreviewShaderPlatformChange();
 
 	PreviewPlatformChanged.Broadcast();
diff --git a/Engine/Source/Editor/UnrealEd/Private/SScalabilitySettings.cpp b/Engine/Source/Editor/UnrealEd/Private/SScalabilitySettings.cpp
--- a/Engine/Source/Editor/UnrealEd/Private/SScalabilitySettings.cpp
+++ b/Engine/Source/Editor/UnrealEd/Private/SScalabilitySettings.cpp
@@ -170,6 +170,18 @@ TSharedRef<SWidget> SScalabilitySettings::MakeAutoButtonWidget()
 		];
 }
 
+TSharedRef<SWidget> SScalabilitySettings::MakeDefaultButtonWidget()
+{
+	return SNew(SButton)
+		.OnClicked(this, &SScalabilitySettings::OnDefaultClicked)
+		.ToolTipText(LOCTEXT("DefaultButtonTooltip", "Resets the Quality Levels to what the Device Profile sets"))
+		.Content()
+		[
+			SNew(STextBlock)
+				.Text(LOCTEXT("DefaultLabel", "Default"))
+		];
+}
+
 FReply SScalabilitySettings::OnHeaderClicked(int32 InQualityLevel, Scalability::EQualityLevelBehavior Behavior)
 {
 	Scalability::FQualityLevels LevelCounts = Scalability::GetQualityLevelCounts();
@@ -201,6 +213,16 @@ FReply SScalabilitySettings::OnAutoClicked()
 	return FReply::Handled();
 }
 
+FReply SScalabilitySettings::OnDefaultClicked()
+{
+	Scalability::ResetQualityLevelsToDefault();
+	CachedQualityLevels = Scalability::GetQualityLevels();
+	Scalability::SetQualityLevels(CachedQualityLevels);
+
+	GEditor->RedrawAllViewports();
+	return FReply::Handled();
+}
+
 SGridPanel::FSlot::FSlotArguments SScalabilitySettings::MakeGridSlot(int32 InCol, int32 InRow, int32 InColSpan /*= 1*/, int32 InRowSpan /*= 1*/)
 {
 	float PaddingH = 2.0f;
@@ -400,6 +422,7 @@ void SScalabilitySettings::Construct( const FArguments& InArgs )
 			SNew(SGridPanel)
 			.FillColumn(0, QualityColumnCoeff)
 
+			+MakeGridSlot(0,0) [ MakeDefaultButtonWidget() ]
 			+MakeGridSlot(1,0) [ MakeHeaderButtonWidget(NamesLow, 0, LOCTEXT("QualityLow", "Set all groups to low quality"), Scalability::EQualityLevelBehavior::EAbsolute) ]
 			+MakeGridSlot(2,0) [ MakeHeaderButtonWidget(NamesMedium, 3, LOCTEXT("QualityMedium", "Set all groups to medium quality"), Scalability::EQualityLevelBehavior::ERelativeToMax) ]
 			+MakeGridSlot(3,0) [ MakeHeaderButtonWidget(NamesHigh, 2, LOCTEXT("QualityHigh", "Set all groups to high quality"), Scalability::EQualityLevelBehavior::ERelativeToMax) ]
diff --git a/Engine/Source/Editor/UnrealEd/Public/SScalabilitySettings.h b/Engine/Source/Editor/UnrealEd/Public/SScalabilitySettings.h
--- a/Engine/Source/Editor/UnrealEd/Public/SScalabilitySettings.h
+++ b/Engine/Source/Editor/UnrealEd/Public/SScalabilitySettings.h
@@ -63,12 +63,18 @@ private:
 
 	// Makes the auto benchmark button
 	TSharedRef<SWidget> MakeAutoButtonWidget();
+	
+	// Makes the default button
+	TSharedRef<SWidget> MakeDefaultButtonWidget();
 
 	// Callback for when a quality level header button is pressed
 	FReply OnHeaderClicked(int32 InQualityLevel, Scalability::EQualityLevelBehavior Behavior);
 
 	// Callback for auto benchmark button
 	FReply OnAutoClicked();
+	
+	// Callback for default button
+	FReply OnDefaultClicked();
 
 	// Create a gridslot for the group quality level with all the required formatting
 	SGridPanel::FSlot::FSlotArguments MakeGridSlot(int32 InCol, int32 InRow, int32 InColSpan = 1, int32 InRowSpan = 1);
diff --git a/Engine/Source/Runtime/Engine/Private/Scalability.cpp b/Engine/Source/Runtime/Engine/Private/Scalability.cpp
--- a/Engine/Source/Runtime/Engine/Private/Scalability.cpp
+++ b/Engine/Source/Runtime/Engine/Private/Scalability.cpp
@@ -12,6 +12,8 @@
 #include "Interfaces/IProjectManager.h"
 
 #if WITH_EDITOR
+#include "DeviceProfiles/DeviceProfile.h"
+#include "DeviceProfiles/DeviceProfileManager.h"
 #endif
 
 Scalability::FOnScalabilitySettingsChanged Scalability::OnScalabilitySettingsChanged;
@@ -360,7 +362,8 @@ FString PlatformScalabilityIniFilename;
 TMap<IConsoleVariable*, FString> PlatformScalabilityCVarBackup;
 TSet<const IConsoleVariable*> PlatformScalabilityCVarAllowList;
 TSet<const IConsoleVariable*> PlatformScalabilityCVarDenyList;
-TMap<EShaderPlatform, Scalability::FQualityLevels> PreviewShaderPlatformToCachedQualityLevels;
+TMap<TTuple<EShaderPlatform, FName>, Scalability::FQualityLevels> PreviewShaderPlatformToCachedQualityLevels;
+TMap<TTuple<EShaderPlatform, FName>, Scalability::FQualityLevels> DefaultPreviewShaderPlatformToCachedQualityLevels;
 
 void UndoPlatformScalability()
 {
@@ -441,14 +444,38 @@ static void SetGroupQualityLevel(const TCHAR* InGroupName, int32 InQualityLevel,
 }
 
 #if WITH_EDITOR
-void ApplyCachedQualityLevelForShaderPlatform(const EShaderPlatform& ShaderPlatform)
+void ApplyCachedQualityLevels()
 {
 	if (ScalabilityShaderPlatform != EShaderPlatform::SP_NumPlatforms)
 	{
-		if (Scalability::FQualityLevels* CachedQualityLevel = PreviewShaderPlatformToCachedQualityLevels.Find(ScalabilityShaderPlatform))
+		UDeviceProfileManager& DeviceProfileManager = UDeviceProfileManager::Get();
+		const UDeviceProfile* PreviewDeviceProfile = DeviceProfileManager.GetPreviewDeviceProfile();
+		if (PreviewDeviceProfile)
 		{
-			Scalability::SetQualityLevels(*CachedQualityLevel);
-			Scalability::SaveState(GEditorSettingsIni);
+			TTuple<EShaderPlatform, FName> PreviewShaderPlatformKey = MakeTuple(ScalabilityShaderPlatform, PreviewDeviceProfile->GetFName());
+
+			DefaultPreviewShaderPlatformToCachedQualityLevels.FindOrAdd(PreviewShaderPlatformKey) = Scalability::GetQualityLevels();
+			if (Scalability::FQualityLevels* CachedQualityLevel = PreviewShaderPlatformToCachedQualityLevels.Find(PreviewShaderPlatformKey))
+			{
+				Scalability::SetQualityLevels(*CachedQualityLevel);
+				Scalability::SaveState(GEditorSettingsIni);
+			}
+		}
+	}
+}
+
+void ResetQualityLevelsToDefault()
+{
+	if (ScalabilityShaderPlatform != EShaderPlatform::SP_NumPlatforms)
+	{
+		UDeviceProfileManager& DeviceProfileManager = UDeviceProfileManager::Get();
+		const UDeviceProfile* PreviewDeviceProfile = DeviceProfileManager.GetPreviewDeviceProfile();
+		if (PreviewDeviceProfile)
+		{
+			if (Scalability::FQualityLevels* DefaultQualityLevels = DefaultPreviewShaderPlatformToCachedQualityLevels.Find(MakeTuple(ScalabilityShaderPlatform, PreviewDeviceProfile->GetFName())))
+			{
+				SetQualityLevels(*DefaultQualityLevels, false);
+			}
 		}
 	}
 }
@@ -915,7 +942,12 @@ void SetQualityLevels(const FQualityLevels& QualityLevels, bool bForce/* = false
 #if WITH_EDITOR
 	if (ScalabilityShaderPlatform != EShaderPlatform::SP_NumPlatforms)
 	{
-		PreviewShaderPlatformToCachedQualityLevels.FindOrAdd(ScalabilityShaderPlatform) = ClampedLevels;
+		UDeviceProfileManager& DeviceProfileManager = UDeviceProfileManager::Get();
+		const UDeviceProfile* PreviewDeviceProfile = DeviceProfileManager.GetPreviewDeviceProfile();
+		if (PreviewDeviceProfile)
+		{
+			PreviewShaderPlatformToCachedQualityLevels.FindOrAdd(MakeTuple(ScalabilityShaderPlatform, PreviewDeviceProfile->GetFName())) = ClampedLevels;
+		}
 	}
 	bScalabilityShaderPlatformHasBeenChanged = false;
 #endif
diff --git a/Engine/Source/Runtime/Engine/Public/Scalability.h b/Engine/Source/Runtime/Engine/Public/Scalability.h
--- a/Engine/Source/Runtime/Engine/Public/Scalability.h
+++ b/Engine/Source/Runtime/Engine/Public/Scalability.h
@@ -175,7 +175,9 @@ namespace Scalability
 	ENGINE_API void SetQualityLevels(const FQualityLevels& QualityLevels, bool bForce = false);
 
 #if WITH_EDITOR
-	ENGINE_API void ApplyCachedQualityLevelForShaderPlatform(const EShaderPlatform& ShaderPlatform);
+	ENGINE_API void ApplyCachedQualityLevels();
+	/** Resets the Quality Levels to what the Device Profile sets */
+	ENGINE_API void ResetQualityLevelsToDefault();
 #endif
 
 	/** This is the only suggested way to get the current state - don't get CVars directly */
