Substrate - added tooltip to specular profile.

#tests pc dx12 editor
#rnx

diff --git a/Engine/Source/Runtime/Engine/Classes/Engine/SpecularProfile.h b/Engine/Source/Runtime/Engine/Classes/Engine/SpecularProfile.h
--- a/Engine/Source/Runtime/Engine/Classes/Engine/SpecularProfile.h
+++ b/Engine/Source/Runtime/Engine/Classes/Engine/SpecularProfile.h
@@ -19,8 +19,8 @@ class FTextureReference;
 UENUM(BlueprintType)
 enum class ESpecularProfileFormat : uint8
 {
-	ViewLightVector UMETA(DisplayName = "View / Light"),
-	HalfVector UMETA(DisplayName = "Half Angle"),
+	ViewLightVector UMETA(DisplayName = "View / Light", ToolTip="The specular LUT will be sampled based on NoV (angle between the Normal and View direction) for the view color axis and NoL (angle between the Normal and Light direction) for the light color axis."),
+	HalfVector UMETA(DisplayName = "Half Angle", ToolTip = "The specular LUT will be sampled based on VoH (angle between the View and the Half vector) for the view color axis and NoH (angle between the Normal and the Half vector) for the light color axis."),
 };
 
 // struct with all the settings we want in USpecularProfile, separate to make it easer to pass this data around in the engine.
@@ -30,19 +30,21 @@ struct FSpecularProfileStruct
 	GENERATED_USTRUCT_BODY()
 	
 	/**
-	 * Define the texture used as a specular profile
+	 * Define the format driving the sampling of the specular LUT.
 	 */
 	UPROPERTY(Category = "Common", EditAnywhere, BlueprintReadOnly, meta=(DisplayName="LUT Format"))
 	ESpecularProfileFormat Format;
 
 	/**
-	* Define the view facing color
+	* Define the view facing color.
+	* Exemple with View/Light mode: color at 0 is applied when NoV=0 (view grazing angle)  while color at 1 is applied when NoV=1 (view facing angle).
 	*/
 	UPROPERTY(Category = "Procedural", EditAnywhere, meta = (AllowZoomOutput = "false", ShowZoomButtons = "false", ViewMinInput = "0", ViewMaxInput = "1", ViewMinOutput = "0", ViewMaxOutput = "1", TimelineLength = "1", ShowInputGridNumbers="false", ShowOutputGridNumbers="false"))
 	FRuntimeCurveLinearColor ViewColor;
 
 	/**
 	* Define the light facing color
+	* Exemple with View/Light mode: color at 0 is applied when NoL=0 (light hit the surface at grazing angle)  while color at 1 is applied when NoV=1 (light hit the surface at facing angle).
 	*/
 	UPROPERTY(Category = "Procedural", EditAnywhere, meta = (AllowZoomOutput = "false", ShowZoomButtons = "false", ViewMinInput = "0", ViewMaxInput = "1", ViewMinOutput = "0", ViewMaxOutput = "1", TimelineLength = "1", ShowInputGridNumbers="false", ShowOutputGridNumbers="false"))
 	FRuntimeCurveLinearColor LightColor;
