Fix pecular profiles issues.

* Add Specular Profile clear when force update is enabled
* Fix args ordering causing the wrong profiles to be fetched with the wrong parameterization
* Fix specular profile convolution being incorrect for low/mid roughness at grazing angle.

diff --git a/Engine/Shaders/Private/SpecularProfile.usf b/Engine/Shaders/Private/SpecularProfile.usf
--- a/Engine/Shaders/Private/SpecularProfile.usf
+++ b/Engine/Shaders/Private/SpecularProfile.usf
@@ -114,7 +114,7 @@ void MainCS(uint3 DispatchThreadId : SV_DispatchThreadID)
 	const uint ProfileIndex = DispatchThreadId.z;
 	if (all(PixelPos < uint2(Resolution)) && ProfileIndex < ProfileCount)
 	{
-		const float NoV = PixelPos.x / float(Resolution.x-1);
+		const float NoV = (PixelPos.x+0.5f) / float(Resolution.x);
 		const float Roughness = max(PixelPos.y / float(Resolution.y-1), 0.02f);
 		const float SpecularProfileId = ProfileSignBuffer[ProfileIndex] * ProfileIndex; // 	The sign encodes the parameterization (NoV / NoH)
 
diff --git a/Engine/Shaders/Private/SpecularProfileCommon.ush b/Engine/Shaders/Private/SpecularProfileCommon.ush
--- a/Engine/Shaders/Private/SpecularProfileCommon.ush
+++ b/Engine/Shaders/Private/SpecularProfileCommon.ush
@@ -23,7 +23,7 @@ float3 GetSpecularProfileCoord(float InSpecularProfileId, float NoV, float NoL,
 {
 	const uint ProfileID = GetSpecularProfileId(InSpecularProfileId);
 	const uint SliceIndex = ProfileID * SUBSTRATE_SPECULAR_PROFILE_ENTRY_COUNT + SUBSTRATE_SPECULAR_PROFILE_ENTRY_LIGHT;
-	return float3(InternalGetSpecularProfileCoord(NoV, NoL, VoH, NoH, InSpecularProfileId), SliceIndex);
+	return float3(InternalGetSpecularProfileCoord(InSpecularProfileId, NoV, NoL, VoH, NoH), SliceIndex);
 }
 
 float3 GetSpecularProfileEnvCoord(float InSpecularProfileId, float NoV, float InRoughness)
diff --git a/Engine/Source/Runtime/Engine/Private/Rendering/SpecularProfile.cpp b/Engine/Source/Runtime/Engine/Private/Rendering/SpecularProfile.cpp
--- a/Engine/Source/Runtime/Engine/Private/Rendering/SpecularProfile.cpp
+++ b/Engine/Source/Runtime/Engine/Private/Rendering/SpecularProfile.cpp
@@ -367,10 +367,16 @@ IPooledRenderTarget* FSpecularProfileTextureManager::GetAtlasTexture(FRDGBuilder
 
 		FPooledRenderTargetDesc Desc(FPooledRenderTargetDesc::Create2DArrayDesc(FIntPoint(Resolution, Resolution), Format, FClearValueBinding::None, TexCreate_None, TexCreate_UAV | TexCreate_ShaderResource, false, LayerCount * SUBSTRATE_SPECULAR_PROFILE_ENTRY_COUNT));
 		GRenderTargetPool.FindFreeElement(RHICmdList, Desc, GSpecularProfileTextureAtlas, TEXT("SpecularProfileTexture"));
-
-		// 2. Fill in profiles
 		const auto GlobalShaderMap = GetGlobalShaderMap(ShaderPlatform);
 		FRDGTextureRef SpecularProfileTexture = GraphBuilder.RegisterExternalTexture(GSpecularProfileTextureAtlas, TEXT("SpecularProfileTexture"));
+		
+		// 2. Optionally clear profiles texture
+		if (bForceUpdate)
+		{
+			AddClearUAVPass(GraphBuilder, GraphBuilder.CreateUAV(SpecularProfileTexture), 0.f);
+		}
+
+		// 3. Fill in profiles
 		FRDGTextureUAVRef SpecularProfileUAVSkipBarrier = GraphBuilder.CreateUAV(SpecularProfileTexture, ERDGUnorderedAccessViewFlags::SkipBarrier);
 		for (uint32 LayerIt = 0; LayerIt < LayerCount; ++LayerIt)
 		{
