- Deprecated and replaced any AnimSequence/Montage/Composite APIs which performed animation data evaluation without an FAnimExtractContext available.
(prep for future changes)
- Fixed up all use-cases across codebase

diff --git a/Engine/Plugins/Animation/ACLPlugin/Source/ACLPluginEditor/Private/ACLStatsDumpCommandlet.cpp b/Engine/Plugins/Animation/ACLPlugin/Source/ACLPluginEditor/Private/ACLStatsDumpCommandlet.cpp
--- a/Engine/Plugins/Animation/ACLPlugin/Source/ACLPluginEditor/Private/ACLStatsDumpCommandlet.cpp
+++ b/Engine/Plugins/Animation/ACLPlugin/Source/ACLPluginEditor/Private/ACLStatsDumpCommandlet.cpp
@@ -255,6 +255,7 @@ static void SampleUEClip(const acl::track_array_qvvf& Tracks, USkeleton* UESkele
 	const FReferenceSkeleton& RefSkeleton = UESkeleton->GetReferenceSkeleton();
 	const TArray<FTransform>& RefSkeletonPose = UESkeleton->GetRefLocalPoses();
 
+	const FAnimExtractContext Context(static_cast<double>(SampleTime));
 	const uint32 NumBones = Tracks.get_num_tracks();
 	for (uint32 BoneIndex = 0; BoneIndex < NumBones; ++BoneIndex)
 	{
@@ -269,7 +270,7 @@ static void SampleUEClip(const acl::track_array_qvvf& Tracks, USkeleton* UESkele
 
 			if (UEClip->GetDataModel()->IsValidBoneTrackName(BoneName))
 			{
-				UEClip->GetBoneTransform(BoneTransform, FSkeletonPoseBoneIndex(BoneTreeIndex), double(SampleTime), false, TOptional<EAnimInterpolationType>());
+				UEClip->GetBoneTransform(BoneTransform, FSkeletonPoseBoneIndex(BoneTreeIndex), Context, false);
 			}
 		}
 
diff --git a/Engine/Plugins/Animation/AnimationLocomotionLibrary/Source/Editor/Private/DistanceCurveModifier.cpp b/Engine/Plugins/Animation/AnimationLocomotionLibrary/Source/Editor/Private/DistanceCurveModifier.cpp
--- a/Engine/Plugins/Animation/AnimationLocomotionLibrary/Source/Editor/Private/DistanceCurveModifier.cpp
+++ b/Engine/Plugins/Animation/AnimationLocomotionLibrary/Source/Editor/Private/DistanceCurveModifier.cpp
@@ -49,7 +49,8 @@ void UDistanceCurveModifier::OnApply_Implementation(UAnimSequence* Animation)
 			const float Time = Step * SampleInterval;
 
 			const bool bAllowLooping = false;
-			const FVector RootMotionTranslation = Animation->ExtractRootMotion(Time, SampleInterval, bAllowLooping).GetTranslation();
+			const FAnimExtractContext Context(static_cast<double>(Time), true, FDeltaTimeRecord(SampleInterval), bAllowLooping);
+			const FVector RootMotionTranslation = Animation->ExtractRootMotion(Context).GetTranslation();
 			const float RootMotionSpeedSq = CalculateMagnitudeSq(RootMotionTranslation, Axis) / SampleInterval;
 
 			if (RootMotionSpeedSq < MinSpeedSq)
@@ -71,7 +72,7 @@ void UDistanceCurveModifier::OnApply_Implementation(UAnimSequence* Animation)
 		// TODO: This works for clips that are broken into starts/stops/pivots, but needs to be rethought for more complex clips.
 		const float ValueSign = (Time < TimeOfMinSpeed) ? -1.0f : 1.0f;
 
-		const FVector RootMotionTranslation = Animation->ExtractRootMotionFromRange(TimeOfMinSpeed, Time).GetTranslation();
+		const FVector RootMotionTranslation = Animation->ExtractRootMotionFromRange(TimeOfMinSpeed, Time, FAnimExtractContext()).GetTranslation();
 		UAnimationBlueprintLibrary::AddFloatCurveKey(Animation, CurveName, Time, ValueSign * CalculateMagnitude(RootMotionTranslation, Axis));
 	}
 }
diff --git a/Engine/Plugins/Animation/AnimationLocomotionLibrary/Source/Runtime/Private/AnimDistanceMatchingLibrary.cpp b/Engine/Plugins/Animation/AnimationLocomotionLibrary/Source/Runtime/Private/AnimDistanceMatchingLibrary.cpp
--- a/Engine/Plugins/Animation/AnimationLocomotionLibrary/Source/Runtime/Private/AnimDistanceMatchingLibrary.cpp
+++ b/Engine/Plugins/Animation/AnimationLocomotionLibrary/Source/Runtime/Private/AnimDistanceMatchingLibrary.cpp
@@ -93,11 +93,14 @@ namespace UE::Anim::DistanceMatchingUtility
 				const float SequenceLength = AnimSequence->GetPlayLength();
 				static const float StepTime = 1.f / 30.f;				
 
+				FAnimExtractContext Context;
 				// Traverse the distance curve, accumulating animated distance until the desired distance is reached.
 				while ((AccumulatedDistance < DistanceTraveled) && (bAllowLooping || (NewTime + StepTime < SequenceLength)))
 				{
-					const float CurrentDistance = AnimSequence->EvaluateCurveData(CurveName, NewTime);
-					const float DistanceAfterStep = AnimSequence->EvaluateCurveData(CurveName, NewTime + StepTime);
+					Context.CurrentTime = NewTime;
+					const float CurrentDistance = AnimSequence->EvaluateCurveData(CurveName, Context);
+					Context.CurrentTime = NewTime + StepTime;
+					const float DistanceAfterStep = AnimSequence->EvaluateCurveData(CurveName, Context);
 					const float AnimationDistanceThisStep = DistanceAfterStep - CurrentDistance;
 
 					if (!FMath::IsNearlyZero(AnimationDistanceThisStep))
@@ -244,7 +247,7 @@ FSequencePlayerReference UAnimDistanceMatchingLibrary::SetPlayrateToMatchSpeed(c
 				if (!FMath::IsNearlyZero(AnimLength))
 				{
 					// Calculate the speed as: (distance traveled by the animation) / (length of the animation)
-					const FVector RootMotionTranslation = AnimSequence->ExtractRootMotionFromRange(0.0f, AnimLength).GetTranslation();
+					const FVector RootMotionTranslation = AnimSequence->ExtractRootMotionFromRange(0.0f, AnimLength, FAnimExtractContext()).GetTranslation();
 					const float RootMotionDistance = RootMotionTranslation.Size2D();
 					if (!FMath::IsNearlyZero(RootMotionDistance))
 					{
diff --git a/Engine/Plugins/Animation/AnimationModifierLibrary/Source/AnimationModifierLibrary/Private/MotionExtractorUtilities.cpp b/Engine/Plugins/Animation/AnimationModifierLibrary/Source/AnimationModifierLibrary/Private/MotionExtractorUtilities.cpp
--- a/Engine/Plugins/Animation/AnimationModifierLibrary/Source/AnimationModifierLibrary/Private/MotionExtractorUtilities.cpp
+++ b/Engine/Plugins/Animation/AnimationModifierLibrary/Source/AnimationModifierLibrary/Private/MotionExtractorUtilities.cpp
@@ -28,9 +28,10 @@ namespace UE::Anim::MotionExtractorUtility
 		bool bMeetsCriterion = false;
 		for (int32 Step = 0; Step < NumSteps; ++Step)
 		{
-			const float Time = static_cast<float>(Step) * SampleInterval;
-			const bool bAllowLooping = false;
-			const FTransform RootMotion = Animation->ExtractRootMotion(Time, SampleInterval, bAllowLooping);
+			const double Time = static_cast<double>(Step) * SampleInterval;
+			constexpr bool bAllowLooping = false;
+			const FAnimExtractContext Context(Time, true, FDeltaTimeRecord(SampleInterval), bAllowLooping);
+			const FTransform RootMotion = Animation->ExtractRootMotion(Context);
 			const bool bMetCriterion = bMeetsCriterion;
 			bMeetsCriterion = ::Invoke(Pred, RootMotion, SampleInterval);
 
diff --git a/Engine/Plugins/Animation/AnimationWarping/Source/Editor/Private/AnimationModifiers/OrientationWarpingModifier.cpp b/Engine/Plugins/Animation/AnimationWarping/Source/Editor/Private/AnimationModifiers/OrientationWarpingModifier.cpp
--- a/Engine/Plugins/Animation/AnimationWarping/Source/Editor/Private/AnimationModifiers/OrientationWarpingModifier.cpp
+++ b/Engine/Plugins/Animation/AnimationWarping/Source/Editor/Private/AnimationModifiers/OrientationWarpingModifier.cpp
@@ -50,8 +50,9 @@ void UOrientationWarpingModifier::OnApply_Implementation(UAnimSequence* Animatio
 	{
 		const float Time = Step * SampleInterval;
 
-		const bool bAllowLooping = false;
-		const FTransform RootMotion = Animation->ExtractRootMotion(Time, SampleInterval, bAllowLooping);
+		constexpr bool bAllowLooping = false;
+		const FAnimExtractContext Context(static_cast<double>(Time), true, FDeltaTimeRecord(SampleInterval), bAllowLooping);
+		const FTransform RootMotion = Animation->ExtractRootMotion(Context);
 		const FVector RootMotionTranslation = RootMotion.GetTranslation();
 		const FQuat RootMotionRotation = RootMotion.GetRotation();
 
diff --git a/Engine/Plugins/Animation/AnimationWarping/Source/Runtime/Private/AnimationWarpingLibrary.cpp b/Engine/Plugins/Animation/AnimationWarping/Source/Runtime/Private/AnimationWarpingLibrary.cpp
--- a/Engine/Plugins/Animation/AnimationWarping/Source/Runtime/Private/AnimationWarpingLibrary.cpp
+++ b/Engine/Plugins/Animation/AnimationWarping/Source/Runtime/Private/AnimationWarpingLibrary.cpp
@@ -37,7 +37,8 @@ bool UAnimationWarpingLibrary::GetCurveValueFromAnimation(const UAnimSequenceBas
 					float ActualTime = Segment->ConvertTrackPosToAnimPos(Time);
 					ActualTime = FMath::Clamp(ActualTime, Segment->AnimStartTime, Segment->AnimEndTime);
 
-					OutValue = Segment->GetAnimReference()->EvaluateCurveData(CurveName, ActualTime);
+					const FAnimExtractContext Context(static_cast<double>(ActualTime));
+					OutValue = Segment->GetAnimReference()->EvaluateCurveData(CurveName, Context);
 					return true;
 				}
 			}
@@ -45,7 +46,8 @@ bool UAnimationWarpingLibrary::GetCurveValueFromAnimation(const UAnimSequenceBas
 	}
 	else if (Animation && Animation->HasCurveData(CurveName))
 	{
-		OutValue = Animation->EvaluateCurveData(CurveName, Time);
+		const FAnimExtractContext Context(static_cast<double>(Time));
+		OutValue = Animation->EvaluateCurveData(CurveName, Context);
 		return true;
 	}
 
diff --git a/Engine/Plugins/Animation/AnimationWarping/Source/Runtime/Private/AnimationWarpingRuntimeModule.cpp b/Engine/Plugins/Animation/AnimationWarping/Source/Runtime/Private/AnimationWarpingRuntimeModule.cpp
--- a/Engine/Plugins/Animation/AnimationWarping/Source/Runtime/Private/AnimationWarpingRuntimeModule.cpp
+++ b/Engine/Plugins/Animation/AnimationWarping/Source/Runtime/Private/AnimationWarpingRuntimeModule.cpp
@@ -51,7 +51,8 @@ void FModule::ShutdownModule()
 
 void FModule::SampleRootMotion(const FDeltaTimeRecord& SampleRange, const UAnimSequence& Sequence, bool bLoopingSequence, UE::Anim::FStackAttributeContainer& OutAttributes) const
 {
-	const FTransform RootMotionTransform = Sequence.ExtractRootMotion(SampleRange.GetPrevious(), SampleRange.Delta, bLoopingSequence);
+	const FAnimExtractContext Context(static_cast<double>(SampleRange.GetPrevious()), true, FDeltaTimeRecord(SampleRange.Delta), bLoopingSequence);
+	const FTransform RootMotionTransform = Sequence.ExtractRootMotion(Context);
 	FTransformAnimationAttribute* RootMotionAttribute = OutAttributes.FindOrAdd<FTransformAnimationAttribute>(RootMotionAttributeId);
 	RootMotionAttribute->Value = RootMotionTransform;
 }
diff --git a/Engine/Plugins/Animation/AnimationWarping/Source/Runtime/Private/BoneControllers/AnimNode_OrientationWarping.cpp b/Engine/Plugins/Animation/AnimationWarping/Source/Runtime/Private/BoneControllers/AnimNode_OrientationWarping.cpp
--- a/Engine/Plugins/Animation/AnimationWarping/Source/Runtime/Private/BoneControllers/AnimNode_OrientationWarping.cpp
+++ b/Engine/Plugins/Animation/AnimationWarping/Source/Runtime/Private/BoneControllers/AnimNode_OrientationWarping.cpp
@@ -262,7 +262,8 @@ void FAnimNode_OrientationWarping::EvaluateSkeletalControl_AnyThread(FComponentS
 			{
 				if (UAnimSequenceBase* AnimSequence = Cast<UAnimSequenceBase>(CurrentAnimAsset))
 				{
-					FTransform PredictedRootMotionTransform = AnimSequence->ExtractRootMotion(CurrentAnimAssetTime, TargetTime, false);
+					const FAnimExtractContext Context(static_cast<double>(CurrentAnimAssetTime), true, FDeltaTimeRecord(TargetTime), false);
+					FTransform PredictedRootMotionTransform = AnimSequence->ExtractRootMotion(Context);
 					PredictedRootMotionDeltaTranslation = PredictedRootMotionTransform.GetTranslation();
 					PredictedRootMotionDeltaTranslation = PredictedRootMotionDeltaTranslation - RotationAxisVector.Dot(PredictedRootMotionDeltaTranslation) * RotationAxisVector;
 				}
diff --git a/Engine/Plugins/Animation/AnimationWarping/Source/Runtime/Private/BoneControllers/AnimNode_Steering.cpp b/Engine/Plugins/Animation/AnimationWarping/Source/Runtime/Private/BoneControllers/AnimNode_Steering.cpp
--- a/Engine/Plugins/Animation/AnimationWarping/Source/Runtime/Private/BoneControllers/AnimNode_Steering.cpp
+++ b/Engine/Plugins/Animation/AnimationWarping/Source/Runtime/Private/BoneControllers/AnimNode_Steering.cpp
@@ -106,7 +106,8 @@ void FAnimNode_Steering::EvaluateSkeletalControl_AnyThread(FComponentSpacePoseCo
 								}
 								else
 								{
-									PredictedRootMotionTransform = AnimSequence->ExtractRootMotion(CurrentAnimAssetTime, AnimatedTargetTime, true);
+									const FAnimExtractContext Context(static_cast<double>(CurrentAnimAssetTime), true, FDeltaTimeRecord(AnimatedTargetTime), true);
+									PredictedRootMotionTransform = AnimSequence->ExtractRootMotion(Context);
 								}
 
 								FQuat PredictedRootMotionQuat = PredictedRootMotionTransform.GetRotation();
diff --git a/Engine/Plugins/Animation/IKRig/Source/IKRigEditor/Private/RetargetEditor/IKRetargetBatchOperation.cpp b/Engine/Plugins/Animation/IKRig/Source/IKRigEditor/Private/RetargetEditor/IKRetargetBatchOperation.cpp
--- a/Engine/Plugins/Animation/IKRig/Source/IKRigEditor/Private/RetargetEditor/IKRetargetBatchOperation.cpp
+++ b/Engine/Plugins/Animation/IKRig/Source/IKRigEditor/Private/RetargetEditor/IKRetargetBatchOperation.cpp
@@ -465,7 +465,7 @@ void UIKRetargetBatchOperation::ConvertAnimation(
 					continue;
 				}
 
-				SpeedCurveValues.Add(SpeedCurveName, SourceSequence->EvaluateCurveData(SpeedCurveName, TimeAtCurrentFrame));
+				SpeedCurveValues.Add(SpeedCurveName, SourceSequence->EvaluateCurveData(SpeedCurveName, FAnimExtractContext(static_cast<double>(TimeAtCurrentFrame))));
 			}
 
 			// get the settings profile
diff --git a/Engine/Plugins/Animation/MLDeformer/ChaosClothGenerator/Source/ChaosClothGenerator/Private/ChaosClothGenerator.cpp b/Engine/Plugins/Animation/MLDeformer/ChaosClothGenerator/Source/ChaosClothGenerator/Private/ChaosClothGenerator.cpp
--- a/Engine/Plugins/Animation/MLDeformer/ChaosClothGenerator/Source/ChaosClothGenerator/Private/ChaosClothGenerator.cpp
+++ b/Engine/Plugins/Animation/MLDeformer/ChaosClothGenerator/Source/ChaosClothGenerator/Private/ChaosClothGenerator.cpp
@@ -963,7 +963,7 @@ namespace UE::Chaos::ClothGenerator
 		FAnimationPoseData AnimationPoseData(OutPose, OutCurve, TempAttributes);
 		AnimationSequence->GetAnimationPose(AnimationPoseData, ExtractionContext);
 
-		const FTransform RootTransform = AnimationSequence->ExtractRootTrackTransform(Time, nullptr);
+		const FTransform RootTransform = AnimationSequence->ExtractRootTrackTransform(FAnimExtractContext(Time), nullptr);
 
 		TArray<FTransform> BoneTransforms;
 		BoneTransforms.SetNumZeroed(NumBones);
diff --git a/Engine/Plugins/Animation/MotionWarping/Source/MotionWarping/Private/MotionWarpingComponent.cpp b/Engine/Plugins/Animation/MotionWarping/Source/MotionWarping/Private/MotionWarpingComponent.cpp
--- a/Engine/Plugins/Animation/MotionWarping/Source/MotionWarping/Private/MotionWarpingComponent.cpp
+++ b/Engine/Plugins/Animation/MotionWarping/Source/MotionWarping/Private/MotionWarpingComponent.cpp
@@ -82,7 +82,7 @@ FTransform UMotionWarpingUtilities::ExtractRootMotionFromAnimation(const UAnimSe
 			{
 				if (CurStep.AnimSequence)
 				{
-					AccumulatedRootMotionParams.Accumulate(CurStep.AnimSequence->ExtractRootMotionFromRange(CurStep.StartPosition, CurStep.EndPosition));
+					AccumulatedRootMotionParams.Accumulate(CurStep.AnimSequence->ExtractRootMotionFromRange(CurStep.StartPosition, CurStep.EndPosition, FAnimExtractContext()));
 				}
 			}
 		}
@@ -92,7 +92,7 @@ FTransform UMotionWarpingUtilities::ExtractRootMotionFromAnimation(const UAnimSe
 
 	if (const UAnimSequence* Anim = Cast<UAnimSequence>(Animation))
 	{
-		return Anim->ExtractRootMotionFromRange(StartTime, EndTime);
+		return Anim->ExtractRootMotionFromRange(StartTime, EndTime, FAnimExtractContext());
 	}
 
 	return FTransform::Identity;
@@ -107,13 +107,13 @@ FTransform UMotionWarpingUtilities::ExtractRootTransformFromAnimation(const UAni
 			if (const UAnimSequence* AnimSequence = Cast<UAnimSequence>(Segment->GetAnimReference()))
 			{
 				const float AnimSequenceTime = Segment->ConvertTrackPosToAnimPos(Time);
-				return AnimSequence->ExtractRootTrackTransform(AnimSequenceTime, nullptr);
+				return AnimSequence->ExtractRootTrackTransform(FAnimExtractContext(static_cast<double>(AnimSequenceTime)), nullptr);
 			}	
 		}
 	}
 	else if (const UAnimSequence* AnimSequence = Cast<UAnimSequence>(Animation))
 	{
-		return AnimSequence->ExtractRootTrackTransform(Time, nullptr);
+		return AnimSequence->ExtractRootTrackTransform(FAnimExtractContext(static_cast<double>(Time)), nullptr);
 	}
 
 	return FTransform::Identity;
diff --git a/Engine/Plugins/Animation/PoseSearch/Source/Runtime/Private/PoseSearchAssetSampler.cpp b/Engine/Plugins/Animation/PoseSearch/Source/Runtime/Private/PoseSearchAssetSampler.cpp
--- a/Engine/Plugins/Animation/PoseSearch/Source/Runtime/Private/PoseSearchAssetSampler.cpp
+++ b/Engine/Plugins/Animation/PoseSearch/Source/Runtime/Private/PoseSearchAssetSampler.cpp
@@ -65,7 +65,7 @@ static FTransform ExtractRootTransformInternal(const UAnimMontage* AnimMontage,
 	{
 		if (CurStep.AnimSequence)
 		{
-			AccumulatedRootMotionParams.Accumulate(CurStep.AnimSequence->ExtractRootMotionFromRange(CurStep.StartPosition, CurStep.EndPosition));
+			AccumulatedRootMotionParams.Accumulate(CurStep.AnimSequence->ExtractRootMotionFromRange(CurStep.StartPosition, CurStep.EndPosition, FAnimExtractContext()));
 		}
 	}
 	return AccumulatedRootMotionParams.GetRootMotionTransform();
@@ -177,6 +177,7 @@ static void ProcessRootTransform(const UBlendSpace* BlendSpace, const FVector& B
 
 			// Accumulate root motion after samples have been updated.
 			FRootMotionMovementParams RootMotionMovementParams;
+			FAnimExtractContext ExtractionCtx(0.0, true, {}, bIsLoopable);	
 			for (int32 BlendSampleIndex = 0; BlendSampleIndex < BlendSamplesData.Num(); BlendSampleIndex++)
 			{
 				FBlendSampleData& BlendSample = BlendSamplesData[BlendSampleIndex];
@@ -191,7 +192,9 @@ static void ProcessRootTransform(const UBlendSpace* BlendSpace, const FVector& B
 						DeltaTime += BlendSample.Animation->GetPlayLength();
 					}
 					
-					const FTransform BlendSampleRootMotion = BlendSample.Animation->ExtractRootMotion(PrevSampleTimes[BlendSampleIndex], DeltaTime, bIsLoopable);
+					ExtractionCtx.CurrentTime = PrevSampleTimes[BlendSampleIndex];
+					ExtractionCtx.DeltaTimeRecord.Delta = DeltaTime;
+					const FTransform BlendSampleRootMotion = BlendSample.Animation->ExtractRootMotion(ExtractionCtx);
 					RootMotionMovementParams.AccumulateWithBlend(BlendSampleRootMotion, BlendSample.GetClampedWeight());
 				}
 			}
@@ -524,9 +527,10 @@ FTransform FAnimationAssetSampler::ExtractRootTransform(float Time) const
 	}
 	else if (const UAnimSequenceBase* SequenceBase = Cast<UAnimSequenceBase>(AnimationAssetPtr.Get()))
 	{
+		FAnimExtractContext ExtractionCtx(0.0, true, FDeltaTimeRecord(Time), true);
 		if (IsLoopable())
 		{
-			RootTransform = SequenceBase->ExtractRootMotion(0.0f, Time, true);
+			RootTransform = SequenceBase->ExtractRootMotion(ExtractionCtx);
 		}
 		else
 		{
@@ -538,7 +542,7 @@ FTransform FAnimationAssetSampler::ExtractRootTransform(float Time) const
 			// animation to estimate where the root would be at Time
 			if (ExtrapolationTime < -SMALL_NUMBER)
 			{
-				FTransform SampleToExtrapolate = SequenceBase->ExtractRootMotionFromRange(0.0f, ExtrapolationSampleTime);
+				FTransform SampleToExtrapolate = SequenceBase->ExtractRootMotionFromRange(0.0f, ExtrapolationSampleTime, ExtractionCtx);
 
 				const FTransform ExtrapolatedRootMotion = UE::PoseSearch::ExtrapolateRootMotion(
 					SampleToExtrapolate,
@@ -548,13 +552,13 @@ FTransform FAnimationAssetSampler::ExtractRootTransform(float Time) const
 			}
 			else
 			{
-				RootTransform = SequenceBase->ExtractRootMotionFromRange(0.0f, ClampedTime);
+				RootTransform = SequenceBase->ExtractRootMotionFromRange(0.0f, ClampedTime, ExtractionCtx);
 
 				// If Time is greater than PlayLength, ExtrapolationTime will be a positive number. In this case, we extrapolate
 				// the end of the animation to estimate where the root would be at Time
 				if (ExtrapolationTime > SMALL_NUMBER)
 				{
-					FTransform SampleToExtrapolate = SequenceBase->ExtractRootMotionFromRange(PlayLength - ExtrapolationSampleTime, PlayLength);
+					FTransform SampleToExtrapolate = SequenceBase->ExtractRootMotionFromRange(PlayLength - ExtrapolationSampleTime, PlayLength, ExtractionCtx);
 
 					const FTransform ExtrapolatedRootMotion = UE::PoseSearch::ExtrapolateRootMotion(
 						SampleToExtrapolate,
diff --git a/Engine/Plugins/Experimental/AnimNext/Source/AnimNext/Private/DecompressionTools.cpp b/Engine/Plugins/Experimental/AnimNext/Source/AnimNext/Private/DecompressionTools.cpp
--- a/Engine/Plugins/Experimental/AnimNext/Source/AnimNext/Private/DecompressionTools.cpp
+++ b/Engine/Plugins/Experimental/AnimNext/Source/AnimNext/Private/DecompressionTools.cpp
@@ -186,7 +186,7 @@ void FDecompressionTools::GetBonePose(const UAnimSequence* AnimSequence, const F
 	const int32 NumTracks = AnimSequence->CompressedData.CompressedTrackToSkeletonMapTable.Num();
 #endif 
 	const bool bTreatAnimAsAdditive = (AnimSequence->IsValidAdditive() && !bForceUseRawData); // Raw data is never additive
-	const FRootMotionReset RootMotionReset(AnimSequence->bEnableRootMotion, AnimSequence->RootMotionRootLock, AnimSequence->bForceRootLock, AnimSequence->ExtractRootTrackTransform(0.0f, /*&RequiredBones*/nullptr), bTreatAnimAsAdditive);
+	const FRootMotionReset RootMotionReset(AnimSequence->bEnableRootMotion, AnimSequence->RootMotionRootLock, AnimSequence->bForceRootLock, AnimSequence->ExtractRootTrackTransform(FAnimExtractContext(0.0), /*&RequiredBones*/nullptr), bTreatAnimAsAdditive);
 
 #if WITH_EDITOR
 	// Evaluate raw (source) curve and bone data
diff --git a/Engine/Plugins/Experimental/AnimNextAnimGraph/Source/AnimNextAnimGraph/Private/Traits/BlendStackTrait.cpp b/Engine/Plugins/Experimental/AnimNextAnimGraph/Source/AnimNextAnimGraph/Private/Traits/BlendStackTrait.cpp
--- a/Engine/Plugins/Experimental/AnimNextAnimGraph/Source/AnimNextAnimGraph/Private/Traits/BlendStackTrait.cpp
+++ b/Engine/Plugins/Experimental/AnimNextAnimGraph/Source/AnimNextAnimGraph/Private/Traits/BlendStackTrait.cpp
@@ -242,9 +242,9 @@ namespace UE::AnimNext
 					{
 						auto ExtractRootMotionAttribute = [AnimSequence = PlayAnimPayload->AnimSequence](float StartTime, float DeltaTime, bool bAllowLooping)
 						{
-							// We do not check for lifetimes, assume the sequence is alive during pose list exectuion.
+							// We do not check for lifetimes, assume the sequence is alive during pose list execution.
 							check(AnimSequence->IsValidLowLevel());
-							return AnimSequence->ExtractRootMotion(StartTime, DeltaTime, bAllowLooping);
+							return AnimSequence->ExtractRootMotion(FAnimExtractContext(static_cast<double>(StartTime), true, FDeltaTimeRecord(DeltaTime), bAllowLooping)); 
 						};
 
 						return FOnExtractRootMotionAttribute::CreateLambda(ExtractRootMotionAttribute);
diff --git a/Engine/Plugins/Experimental/AnimNextAnimGraph/Source/AnimNextAnimGraph/Private/Traits/SequencePlayer.cpp b/Engine/Plugins/Experimental/AnimNextAnimGraph/Source/AnimNextAnimGraph/Private/Traits/SequencePlayer.cpp
--- a/Engine/Plugins/Experimental/AnimNextAnimGraph/Source/AnimNextAnimGraph/Private/Traits/SequencePlayer.cpp
+++ b/Engine/Plugins/Experimental/AnimNextAnimGraph/Source/AnimNextAnimGraph/Private/Traits/SequencePlayer.cpp
@@ -65,7 +65,7 @@ namespace UE::AnimNext
 			{
 				// We do not check for lifetimes, assume the sequence is alive during pose list exectuion.
 				check(AnimSequence->IsValidLowLevel());
-				return AnimSequence->ExtractRootMotion(StartTime, DeltaTime, bAllowLooping);
+				return AnimSequence->ExtractRootMotion(FAnimExtractContext(static_cast<double>(StartTime), true, FDeltaTimeRecord(DeltaTime), bAllowLooping));
 			};
 
 			return FOnExtractRootMotionAttribute::CreateLambda(ExtractRootMotionAttribute);
diff --git a/Engine/Plugins/Experimental/Animation/ContextualAnimation/Source/ContextualAnimation/Private/ContextualAnimUtilities.cpp b/Engine/Plugins/Experimental/Animation/ContextualAnimation/Source/ContextualAnimation/Private/ContextualAnimUtilities.cpp
--- a/Engine/Plugins/Experimental/Animation/ContextualAnimation/Source/ContextualAnimation/Private/ContextualAnimUtilities.cpp
+++ b/Engine/Plugins/Experimental/Animation/ContextualAnimation/Source/ContextualAnimation/Private/ContextualAnimUtilities.cpp
@@ -53,12 +53,12 @@ FTransform UContextualAnimUtilities::ExtractRootMotionFromAnimation(const UAnimS
 {
 	if (const UAnimMontage* Anim = Cast<UAnimMontage>(Animation))
 	{
-		return Anim->ExtractRootMotionFromTrackRange(StartTime, EndTime);
+		return Anim->ExtractRootMotionFromTrackRange(StartTime, EndTime, FAnimExtractContext());
 	}
 
 	if (const UAnimSequence* Anim = Cast<UAnimSequence>(Animation))
 	{
-		return Anim->ExtractRootMotionFromRange(StartTime, EndTime);
+		return Anim->ExtractRootMotionFromRange(StartTime, EndTime, FAnimExtractContext());
 	}
 
 	return FTransform::Identity;
@@ -73,13 +73,13 @@ FTransform UContextualAnimUtilities::ExtractRootTransformFromAnimation(const UAn
 			if (const UAnimSequence* AnimSequence = Cast<UAnimSequence>(Segment->GetAnimReference()))
 			{
 				const float AnimSequenceTime = Segment->ConvertTrackPosToAnimPos(Time);
-				return AnimSequence->ExtractRootTrackTransform(AnimSequenceTime, nullptr);
+				return AnimSequence->ExtractRootTrackTransform(FAnimExtractContext(static_cast<double>(AnimSequenceTime)), nullptr);
 			}
 		}
 	}
 	else if (const UAnimSequence* AnimSequence = Cast<UAnimSequence>(Animation))
 	{
-		return AnimSequence->ExtractRootTrackTransform(Time, nullptr);
+		return AnimSequence->ExtractRootTrackTransform(FAnimExtractContext(static_cast<double>(Time)), nullptr);
 	}
 
 	return FTransform::Identity;
diff --git a/Engine/Plugins/Experimental/EvaluationNotifies/Source/Runtime/Private/EvaluationNotifies/AnimNotifyState_Alignment.cpp b/Engine/Plugins/Experimental/EvaluationNotifies/Source/Runtime/Private/EvaluationNotifies/AnimNotifyState_Alignment.cpp
--- a/Engine/Plugins/Experimental/EvaluationNotifies/Source/Runtime/Private/EvaluationNotifies/AnimNotifyState_Alignment.cpp
+++ b/Engine/Plugins/Experimental/EvaluationNotifies/Source/Runtime/Private/EvaluationNotifies/AnimNotifyState_Alignment.cpp
@@ -102,7 +102,7 @@ void FAlignmentNotifyInstance::Update(const UAnimSequenceBase* AnimationAsset, f
     					while(PredictionTime < TargetTime)
     					{
     						FVector PrevPosition = PredictedTransform.GetLocation();
-    						PredictedTransform = AnimationAsset->ExtractRootMotion(PredictionTime, TimeStep, false) * PredictedTransform;
+    						PredictedTransform = AnimationAsset->ExtractRootMotion(FAnimExtractContext(static_cast<double>(PredictionTime), true, FDeltaTimeRecord(TimeStep), false)) * PredictedTransform;
     						UE_VLOG_SEGMENT(Output.AnimInstanceProxy->GetAnimInstanceObject(), "Alignment", Display, PrevPosition, PredictedTransform.GetLocation(), AltColor ? FColor::Yellow : FColor::Red, TEXT(""));
     						AltColor = !AltColor;
     						PredictionTime += TimeStep;
@@ -111,7 +111,7 @@ void FAlignmentNotifyInstance::Update(const UAnimSequenceBase* AnimationAsset, f
 #endif
     				
 					const float PredictionDelta = EndTime + PostNotifyTargetTime - CurrentTime ;
-					PredictedRootMotion = AnimationAsset->ExtractRootMotion(CurrentTime, PredictionDelta, false);
+					PredictedRootMotion = AnimationAsset->ExtractRootMotion(FAnimExtractContext(static_cast<double>(CurrentTime), true, FDeltaTimeRecord(PredictionDelta), false));
 					AlignmentTransform = PredictedRootMotion * AlignmentTransform;
 
     				// this predicted transform can be reused for the weight curve, but only if there was no time offset
@@ -132,7 +132,7 @@ void FAlignmentNotifyInstance::Update(const UAnimSequenceBase* AnimationAsset, f
 						{
 							// if we didn't already predict root motion above, or we predicted to a time offset, then compute the root motion to the end of the notify here
 							const float PredictionDelta = EndTime - CurrentTime;
-							PredictedRootMotion = AnimationAsset->ExtractRootMotion(CurrentTime, PredictionDelta, false);
+							PredictedRootMotion = AnimationAsset->ExtractRootMotion(FAnimExtractContext(static_cast<double>(CurrentTime), true, FDeltaTimeRecord(PredictionDelta), false));
 						}
 						
 						const float RootMotionDistanceRemaining = PredictedRootMotion.GetLocation().Length();
diff --git a/Engine/Plugins/Runtime/NetworkPredictionExtras/Source/NetworkPredictionExtras/Private/MockRootMotionSourceObject.cpp b/Engine/Plugins/Runtime/NetworkPredictionExtras/Source/NetworkPredictionExtras/Private/MockRootMotionSourceObject.cpp
--- a/Engine/Plugins/Runtime/NetworkPredictionExtras/Source/NetworkPredictionExtras/Private/MockRootMotionSourceObject.cpp
+++ b/Engine/Plugins/Runtime/NetworkPredictionExtras/Source/NetworkPredictionExtras/Private/MockRootMotionSourceObject.cpp
@@ -158,7 +158,7 @@ FMockRootMotionReturnValue UMockRootMotionSource_Montage::Step(const FMockRootMo
 	}
 
 	// Extract root motion from animation sequence 
-	ReturnValue.DeltaTransform = Montage->ExtractRootMotionFromTrackRange(StartPlayPosition, EndPlayPosition);
+	ReturnValue.DeltaTransform = Montage->ExtractRootMotionFromTrackRange(StartPlayPosition, EndPlayPosition, FAnimExtractContext());
 	ReturnValue.TransformType = FMockRootMotionReturnValue::ETransformType::AnimationRelative;
 
 	//UE_LOG(LogTemp, Warning, TEXT("%.4f: Delta: %s"), EndPosition, *NewTransform.GetTranslation().ToString());
diff --git a/Engine/Plugins/Runtime/USDCore/Source/USDUtilities/Private/USDSkeletalDataConversion.cpp b/Engine/Plugins/Runtime/USDCore/Source/USDUtilities/Private/USDSkeletalDataConversion.cpp
--- a/Engine/Plugins/Runtime/USDCore/Source/USDUtilities/Private/USDSkeletalDataConversion.cpp
+++ b/Engine/Plugins/Runtime/USDCore/Source/USDUtilities/Private/USDSkeletalDataConversion.cpp
@@ -3870,11 +3870,11 @@ bool UnrealToUsd::ConvertAnimSequence(UAnimSequence* AnimSequence, pxr::UsdPrim&
 		// Blend shape weights
 		for (int32 TimeCode = 0; TimeCode < NumTimeCodes; ++TimeCode)
 		{
-			const float AnimTime = TimeCode / TimeCodesPerSecond;
+			const double AnimTime = TimeCode / TimeCodesPerSecond;
 
 			FBlendedCurve BlendedCurve;
-			const bool bForceUseRawData = true;
-			AnimSequence->EvaluateCurveData(BlendedCurve, AnimTime, bForceUseRawData);
+			constexpr bool bForceUseRawData = true;
+			AnimSequence->EvaluateCurveData(BlendedCurve,  FAnimExtractContext(static_cast<double>(AnimTime)), bForceUseRawData);
 
 			BlendShapeNames.clear();
 			BlendShapeNames.reserve(BlendedCurve.Num());
diff --git a/Engine/Source/Editor/AnimationBlueprintLibrary/Private/AnimationBlueprintLibrary.cpp b/Engine/Source/Editor/AnimationBlueprintLibrary/Private/AnimationBlueprintLibrary.cpp
--- a/Engine/Source/Editor/AnimationBlueprintLibrary/Private/AnimationBlueprintLibrary.cpp
+++ b/Engine/Source/Editor/AnimationBlueprintLibrary/Private/AnimationBlueprintLibrary.cpp
@@ -221,7 +221,7 @@ FTransform UAnimationBlueprintLibrary::ExtractRootTrackTransform(const UAnimSequ
 		return FTransform::Identity;
 	}
 	
-	return AnimationSequenceBase->ExtractRootTrackTransform(Time, nullptr);
+	return AnimationSequenceBase->ExtractRootTrackTransform(FAnimExtractContext(static_cast<double>(Time)), nullptr);
 }
 
 const FRawAnimSequenceTrack& UAnimationBlueprintLibrary::GetRawAnimationTrackByName(const UAnimSequenceBase* AnimationSequenceBase, const FName TrackName)
diff --git a/Engine/Source/Editor/UnrealEd/Private/Fbx/FbxAnimationExport.cpp b/Engine/Source/Editor/UnrealEd/Private/Fbx/FbxAnimationExport.cpp
--- a/Engine/Source/Editor/UnrealEd/Private/Fbx/FbxAnimationExport.cpp
+++ b/Engine/Source/Editor/UnrealEd/Private/Fbx/FbxAnimationExport.cpp
@@ -293,7 +293,7 @@ void FFbxExporter::ExportCustomAnimCurvesToFbx(const TMap<FName, FbxAnimCurve*>&
 	
 	auto ExportLambda = [&](float AnimTime, FbxTime ExportTime, bool bLastKey) {
 		FBlendedCurve BlendedCurve;
-		AnimSeq->EvaluateCurveData(BlendedCurve, AnimTime, true);
+		AnimSeq->EvaluateCurveData(BlendedCurve, FAnimExtractContext(static_cast<double>(AnimTime)), true);
 		
 		//Loop over the custom curves and add the actual keys
 		for (auto CustomCurve : CustomCurves)
diff --git a/Engine/Source/Runtime/AnimGraphRuntime/Private/AnimNodes/AnimNode_BlendSpaceGraphBase.cpp b/Engine/Source/Runtime/AnimGraphRuntime/Private/AnimNodes/AnimNode_BlendSpaceGraphBase.cpp
--- a/Engine/Source/Runtime/AnimGraphRuntime/Private/AnimNodes/AnimNode_BlendSpaceGraphBase.cpp
+++ b/Engine/Source/Runtime/AnimGraphRuntime/Private/AnimNodes/AnimNode_BlendSpaceGraphBase.cpp
@@ -94,7 +94,7 @@ void FAnimNode_BlendSpaceGraphBase::Evaluate_AnyThread(FPoseContext& Output)
 
 	check(BlendSpace != nullptr);
 
-	BlendSpace->GetAnimationPose(BlendSampleDataCache, SamplePoseLinks, {}/** unused parameter **/, Output);
+	BlendSpace->GetAnimationPose(BlendSampleDataCache, SamplePoseLinks, FAnimExtractContext(), Output);
 }
 
 void FAnimNode_BlendSpaceGraphBase::GatherDebugData(FNodeDebugData& DebugData)
diff --git a/Engine/Source/Runtime/Engine/Classes/Animation/AnimComposite.h b/Engine/Source/Runtime/Engine/Classes/Animation/AnimComposite.h
--- a/Engine/Source/Runtime/Engine/Classes/Animation/AnimComposite.h
+++ b/Engine/Source/Runtime/Engine/Classes/Animation/AnimComposite.h
@@ -51,9 +51,9 @@ public:
 	virtual bool IsValidAdditive() const override { return GetAdditiveAnimType() != AAT_None; }
 	virtual void EnableRootMotionSettingFromMontage(bool bInEnableRootMotion, const ERootMotionRootLock::Type InRootMotionRootLock) override;
 	virtual bool HasRootMotion() const override;
-	FTransform ExtractRootMotion(float StartTime, float DeltaTime, bool bAllowLooping) const override final;
-	FTransform ExtractRootMotionFromRange(float StartTrackPosition, float EndTrackPosition) const override final;
-	FTransform ExtractRootTrackTransform(float Time, const FBoneContainer* RequiredBones) const override final;
+	virtual FTransform ExtractRootMotion(const FAnimExtractContext& ExtractionContext) const override final;
+	virtual FTransform ExtractRootMotionFromRange(double StartTime, double EndTime, const FAnimExtractContext& ExtractionContext) const override final;
+	virtual FTransform ExtractRootTrackTransform(const FAnimExtractContext& ExtractionContext, const FBoneContainer* RequiredBones) const override final;
 	virtual void GetAnimNotifiesFromDeltaPositions(const float& PreviousPosition, const float & CurrentPosition, FAnimNotifyContext& NotifyContext) const override;
 	virtual bool IsNotifyAvailable() const override;
 	//~ End UAnimSequenceBase Interface
diff --git a/Engine/Source/Runtime/Engine/Classes/Animation/AnimCompositeBase.h b/Engine/Source/Runtime/Engine/Classes/Animation/AnimCompositeBase.h
--- a/Engine/Source/Runtime/Engine/Classes/Animation/AnimCompositeBase.h
+++ b/Engine/Source/Runtime/Engine/Classes/Animation/AnimCompositeBase.h
@@ -365,7 +365,10 @@ class UAnimCompositeBase : public UAnimSequenceBase
 	//~ End UAnimSequenceBase Interface
 
 	// Extracts root motion from the supplied FAnimTrack between the Start End range specified
-	ENGINE_API void ExtractRootMotionFromTrack(const FAnimTrack &SlotAnimTrack, float StartTrackPosition, float EndTrackPosition, FRootMotionMovementParams &RootMotion) const;
+	UE_DEPRECATED(5.6, "Use static ExtractRootMotionFromTrack with FAnimExtractionContext")
+	void ExtractRootMotionFromTrack(const FAnimTrack &SlotAnimTrack, float StartTrackPosition, float EndTrackPosition, FRootMotionMovementParams &RootMotion) const { FAnimExtractContext Context; ExtractRootMotionFromTrack(SlotAnimTrack, StartTrackPosition, EndTrackPosition, Context, RootMotion); }
+	
+	ENGINE_API void ExtractRootMotionFromTrack(const FAnimTrack &SlotAnimTrack, float StartTrackPosition, float EndTrackPosition, const FAnimExtractContext& Context, FRootMotionMovementParams& RootMotion) const;
 
 	// this is to prevent anybody adding recursive asset to anim composite
 	// as a result of anim composite being a part of anim sequence base
diff --git a/Engine/Source/Runtime/Engine/Classes/Animation/AnimMontage.h b/Engine/Source/Runtime/Engine/Classes/Animation/AnimMontage.h
--- a/Engine/Source/Runtime/Engine/Classes/Animation/AnimMontage.h
+++ b/Engine/Source/Runtime/Engine/Classes/Animation/AnimMontage.h
@@ -853,7 +853,10 @@ public:
 	 * It will break down the range into steps if needed to handle looping animations, or different animations.
 	 * These steps will be processed sequentially, and output the RootMotion transform in component space.
 	 */
-	ENGINE_API FTransform ExtractRootMotionFromTrackRange(float StartTrackPosition, float EndTrackPosition) const;
+	ENGINE_API FTransform ExtractRootMotionFromTrackRange(float StartTrackPosition, float EndTrackPosition, const FAnimExtractContext& Context) const;
+
+	UE_DEPRECATED(5.6, "Use ExtractRootMotionFromTrackRange with FAnimExtractContext")
+	FTransform ExtractRootMotionFromTrackRange(float StartTrackPosition, float EndTrackPosition) const { const FAnimExtractContext Context; return ExtractRootMotionFromTrackRange(StartTrackPosition, EndTrackPosition, Context); }
 
 	/** Get the Montage's Group Name. This is the group from the first slot.  */
 	UFUNCTION(BlueprintPure, Category = "Montage")
diff --git a/Engine/Source/Runtime/Engine/Classes/Animation/AnimSequence.h b/Engine/Source/Runtime/Engine/Classes/Animation/AnimSequence.h
--- a/Engine/Source/Runtime/Engine/Classes/Animation/AnimSequence.h
+++ b/Engine/Source/Runtime/Engine/Classes/Animation/AnimSequence.h
@@ -431,11 +431,11 @@ public:
 	virtual EAdditiveAnimationType GetAdditiveAnimType() const override { return AdditiveAnimType; }
 	ENGINE_API virtual int32 GetNumberOfSampledKeys() const override;
 	virtual FFrameRate GetSamplingFrameRate() const override { return PlatformTargetFrameRate.Default; }
-	ENGINE_API virtual void EvaluateCurveData(FBlendedCurve& OutCurve, float CurrentTime, bool bForceUseRawData = false) const override;
+	ENGINE_API virtual void EvaluateCurveData(FBlendedCurve& OutCurve, const FAnimExtractContext& AnimExtractContext, bool bForceUseRawData = false) const override;
 
 	UE_DEPRECATED(5.3, "Please use EvaluateCurveData that takes an FName.")
 	virtual float EvaluateCurveData(SmartName::UID_Type CurveUID, float CurrentTime, bool bForceUseRawData = false) const override { return 0.0f; }
-	ENGINE_API virtual float EvaluateCurveData(FName CurveName, float CurrentTime, bool bForceUseRawData = false) const override;
+	ENGINE_API virtual float EvaluateCurveData(FName CurveName, const FAnimExtractContext& AnimExtractContext, bool bForceUseRawData = false) const override;
 
 	UE_DEPRECATED(5.3, "Please use HasCurveData that takes an FName.")
 	virtual bool HasCurveData(SmartName::UID_Type CurveUID, bool bForceUseRawData) const override { return false; }
@@ -443,14 +443,19 @@ public:
 
 	//~ End UAnimSequenceBase Interface
 
+	UE_DEPRECATED(5.6, "Please use ExtractRootMotion with FAnimExtractContext")
+	virtual FTransform ExtractRootMotion(float StartTime, float DeltaTime, bool bAllowLooping) const override final { PRAGMA_DISABLE_DEPRECATION_WARNINGS return UAnimSequenceBase::ExtractRootMotion(StartTime, DeltaTime, bAllowLooping); PRAGMA_ENABLE_DEPRECATION_WARNINGS }
+	UE_DEPRECATED(5.6, "Please use ExtractRootMotionFromRange with FAnimExtractContext")
+	virtual FTransform ExtractRootMotionFromRange(float StartTrackPosition, float EndTrackPosition) const override final { PRAGMA_DISABLE_DEPRECATION_WARNINGS return UAnimSequenceBase::ExtractRootMotionFromRange(StartTrackPosition, EndTrackPosition); PRAGMA_ENABLE_DEPRECATION_WARNINGS }
+	UE_DEPRECATED(5.6, "Please use ExtractRootTrackTransform with FAnimExtractContext")
+	virtual FTransform ExtractRootTrackTransform(float Time, const FBoneContainer* RequiredBones) const override final { PRAGMA_DISABLE_DEPRECATION_WARNINGS return UAnimSequenceBase::ExtractRootTrackTransform(Time, RequiredBones); PRAGMA_ENABLE_DEPRECATION_WARNINGS }
+	
 	// Extract Root Motion transform from the animation
-	ENGINE_API FTransform ExtractRootMotion(float StartTime, float DeltaTime, bool bAllowLooping) const override final;
-
+	ENGINE_API virtual FTransform ExtractRootMotion(const FAnimExtractContext& ExtractionContext) const override final;
 	// Extract Root Motion transform from a contiguous position range (no looping)
-	ENGINE_API FTransform ExtractRootMotionFromRange(float StartTrackPosition, float EndTrackPosition) const override final;
-
+	ENGINE_API virtual FTransform ExtractRootMotionFromRange(double StartTime, double EndTime, const FAnimExtractContext& ExtractionContext) const override final;
 	// Extract the transform from the root track for the given animation position
-	ENGINE_API FTransform ExtractRootTrackTransform(float Time, const FBoneContainer* RequiredBones) const override final;
+	ENGINE_API virtual FTransform ExtractRootTrackTransform(const FAnimExtractContext& ExtractionContext, const FBoneContainer* RequiredBones) const override final;
 
 	// Begin Transform related functions 
 	ENGINE_API virtual void GetAnimationPose(FAnimationPoseData& OutAnimationPoseData, const FAnimExtractContext& ExtractionContext) const override;
@@ -553,16 +558,18 @@ public:
 	ENGINE_API void GetAdditiveBasePose(FCompactPose& OutPose, FBlendedCurve& OutCurve, const FAnimExtractContext& ExtractionContext) const;
 	ENGINE_API void GetAdditiveBasePose(FAnimationPoseData& OutAnimationPoseData, const FAnimExtractContext& ExtractionContext) const;
 
+	UE_DEPRECATED(5.6, "Please use GetBoneTransform with FAnimExtractContext")
+	void GetBoneTransform(FTransform& OutAtom, FSkeletonPoseBoneIndex BoneIndex, double Time, bool bUseRawData, TOptional<EAnimInterpolationType> InterpolationOverride=TOptional<EAnimInterpolationType>()) const { FAnimExtractContext ExtractionContext(Time); ExtractionContext.InterpolationOverride = InterpolationOverride; return GetBoneTransform(OutAtom, BoneIndex, ExtractionContext, bUseRawData); }
+
 	/**
 	 * Get Bone Transform of the Time given, relative to Parent for the Track Given
 	 *
-	 * @param	OutAtom			[out] Output bone transform.
-	 * @param	TrackIndex		Index of track to interpolate.
-	 * @param	Time			Time on track to interpolate to.
+	 * @param	OutAtom					[out] Output bone transform.
+	 * @param	BoneIndex				Index of bone to evaluate.
+	 * @param	FAnimExtractContext		Extraction context containing time and interpolation info.
 	 * @param	bUseRawData		If true, use raw animation data instead of compressed data.
-	 * @param   InterpolationOverride The optional animation interpolation type to use. If not set, it will use the AnimSequence interpolation type.
 	 */
-	ENGINE_API void GetBoneTransform(FTransform& OutAtom, FSkeletonPoseBoneIndex BoneIndex, double Time, bool bUseRawData, TOptional<EAnimInterpolationType> InterpolationOverride=TOptional<EAnimInterpolationType>()) const;
+	ENGINE_API void GetBoneTransform(FTransform& OutAtom, FSkeletonPoseBoneIndex BoneIndex, const FAnimExtractContext& ExtractionContext, bool bUseRawData) const;
 
 	UE_DEPRECATED(5.1, "Use other GetBoneTransform signature using double and skeleton index")
 	void GetBoneTransform(FTransform& OutAtom, int32 TrackIndex, float Time, bool bUseRawData) const {}
@@ -571,7 +578,7 @@ public:
 	 * Get Bone Transform of the Time given, relative to Parent for the Track Given
 	 *
 	 * @param	OutAtom			[out] Output bone transform.
-	 * @param	TrackIndex		Index of track to interpolate.
+	 * @param	BoneIndex		Index of bone to evaluate.
 	 * @param	DecompContext	Decompression context to use.
 	 * @param	bUseRawData		If true, use raw animation data instead of compressed data.
 	 */
diff --git a/Engine/Source/Runtime/Engine/Classes/Animation/AnimSequenceBase.h b/Engine/Source/Runtime/Engine/Classes/Animation/AnimSequenceBase.h
--- a/Engine/Source/Runtime/Engine/Classes/Animation/AnimSequenceBase.h
+++ b/Engine/Source/Runtime/Engine/Classes/Animation/AnimSequenceBase.h
@@ -139,12 +139,16 @@ public:
 	ENGINE_API virtual void GetAnimNotifiesFromDeltaPositions(const float& PreviousPosition, const float & CurrentPosition, FAnimNotifyContext& NotifyContext) const;
 
 	/** Evaluate curve data to Instance at the time of CurrentTime **/
-	ENGINE_API virtual void EvaluateCurveData(FBlendedCurve& OutCurve, float CurrentTime, bool bForceUseRawData = false) const;
+	UE_DEPRECATED(5.6, "Please use EvaluateCurveData with FAnimExtractContext")
+	virtual void EvaluateCurveData(FBlendedCurve& OutCurve, float CurrentTime, bool bForceUseRawData = false) const { const FAnimExtractContext Context(static_cast<double>(CurrentTime)); EvaluateCurveData(OutCurve, Context, bForceUseRawData); }
+	ENGINE_API virtual void EvaluateCurveData(FBlendedCurve& OutCurve, const FAnimExtractContext& AnimExtractContext, bool bForceUseRawData = false) const;
 
 	UE_DEPRECATED(5.3, "Please use EvaluateCurveData that takes a FName.")
 	virtual float EvaluateCurveData(SmartName::UID_Type CurveUID, float CurrentTime, bool bForceUseRawData = false) const { return 0.0f; }
 	
-	ENGINE_API virtual float EvaluateCurveData(FName CurveName, float CurrentTime, bool bForceUseRawData = false) const;
+	UE_DEPRECATED(5.6, "Please use EvaluateCurveData with FAnimExtractContext")
+	virtual float EvaluateCurveData(FName CurveName, float CurrentTime, bool bForceUseRawData = false) const { const FAnimExtractContext Context(static_cast<double>(CurrentTime)); return EvaluateCurveData(CurveName, Context, bForceUseRawData); }
+	ENGINE_API virtual float EvaluateCurveData(FName CurveName, const FAnimExtractContext& AnimExtractContext, bool bForceUseRawData = false) const;
 	
 	ENGINE_API virtual const FRawCurveTracks& GetCurveData() const;
 	
@@ -221,14 +225,23 @@ public:
 
 	virtual bool HasRootMotion() const { return false; }
 
+	UE_DEPRECATED(5.6, "Please use ExtractRootMotion with FAnimExtractContext")
+	virtual FTransform ExtractRootMotion(float StartTime, float DeltaTime, bool bAllowLooping) const { const FAnimExtractContext Context(static_cast<double>(StartTime), true, FDeltaTimeRecord(DeltaTime), bAllowLooping); return ExtractRootMotion(Context); }
+
 	// Extract Root Motion transform from the animation
-	virtual FTransform ExtractRootMotion(float StartTime, float DeltaTime, bool bAllowLooping) const { return {}; }
+	virtual FTransform ExtractRootMotion(const FAnimExtractContext& ExtractionContext) const { return {}; }
+
+	UE_DEPRECATED(5.6, "Please use ExtractRootMotionFromRange with FAnimExtractContext")
+	virtual FTransform ExtractRootMotionFromRange(float StartTrackPosition, float EndTrackPosition) const { const FAnimExtractContext Context; return ExtractRootMotionFromRange(StartTrackPosition, EndTrackPosition, Context); }
 
 	// Extract Root Motion transform from a contiguous position range (no looping)
-	virtual FTransform ExtractRootMotionFromRange(float StartTrackPosition, float EndTrackPosition) const { return {}; }
+	virtual FTransform ExtractRootMotionFromRange(double StartTime, double EndTime, const FAnimExtractContext& ExtractionContext) const { return {}; }
+
+	UE_DEPRECATED(5.6, "Please use ExtractRootTrackTransform with FAnimExtractContext")
+	virtual FTransform ExtractRootTrackTransform(float Time, const FBoneContainer* RequiredBones) const { const FAnimExtractContext Context(static_cast<double>(Time)); return ExtractRootTrackTransform(Context, RequiredBones); }
 
 	// Extract the transform from the root track for the given animation position
-	virtual FTransform ExtractRootTrackTransform(float Time, const FBoneContainer* RequiredBones) const { return {}; }
+	virtual FTransform ExtractRootTrackTransform(const FAnimExtractContext& ExtractionContext, const FBoneContainer* RequiredBones) const { return {}; }
 
 	ENGINE_API virtual void Serialize(FArchive& Ar) override;
 
diff --git a/Engine/Source/Runtime/Engine/Classes/Animation/AnimationAsset.h b/Engine/Source/Runtime/Engine/Classes/Animation/AnimationAsset.h
--- a/Engine/Source/Runtime/Engine/Classes/Animation/AnimationAsset.h
+++ b/Engine/Source/Runtime/Engine/Classes/Animation/AnimationAsset.h
@@ -103,6 +103,9 @@ public:
 	bool IsPreviousValid() const { return bPreviousIsValid; }
 
 	float Delta = 0.f;
+	
+	FDeltaTimeRecord() = default;
+	explicit FDeltaTimeRecord(float InDeltaTime) : Delta(InDeltaTime), Previous(0.f) {}
 private:
 	float Previous = 0.f;
 	bool  bPreviousIsValid = false; // Will be set to true when Previous has been set
diff --git a/Engine/Source/Runtime/Engine/Private/Animation/AnimComposite.cpp b/Engine/Source/Runtime/Engine/Private/Animation/AnimComposite.cpp
--- a/Engine/Source/Runtime/Engine/Private/Animation/AnimComposite.cpp
+++ b/Engine/Source/Runtime/Engine/Private/Animation/AnimComposite.cpp
@@ -59,7 +59,7 @@ void UAnimComposite::HandleAssetPlayerTickedInternal(FAnimAssetTickContext &Cont
 {
 	Super::HandleAssetPlayerTickedInternal(Context, PreviousTime, MoveDelta, Instance, NotifyQueue);
 
-	ExtractRootMotionFromTrack(AnimationTrack, PreviousTime, PreviousTime + MoveDelta, Context.RootMotionMovementParams);
+	ExtractRootMotionFromTrack(AnimationTrack, PreviousTime, PreviousTime + MoveDelta, FAnimExtractContext(), Context.RootMotionMovementParams);
 }
 
 void UAnimComposite::GetAnimationPose(FAnimationPoseData& OutAnimationPoseData, const FAnimExtractContext& ExtractionContext) const
@@ -70,7 +70,7 @@ void UAnimComposite::GetAnimationPose(FAnimationPoseData& OutAnimationPoseData,
 
 	FBlendedCurve CompositeCurve;
 	CompositeCurve.InitFrom(OutCurve);
-	EvaluateCurveData(CompositeCurve, ExtractionContext.CurrentTime);
+	EvaluateCurveData(CompositeCurve, ExtractionContext);
 
 	// combine both curve
 	OutCurve.Combine(CompositeCurve);
@@ -98,26 +98,28 @@ bool UAnimComposite::HasRootMotion() const
 	return AnimationTrack.HasRootMotion();
 }
 
-FTransform UAnimComposite::ExtractRootMotion(float StartTime, float DeltaTime, bool bAllowLooping) const
+FTransform UAnimComposite::ExtractRootMotion(const FAnimExtractContext& ExtractionContext) const
 {
-	return ExtractRootMotionFromRange(StartTime, DeltaTime);
+	return ExtractRootMotionFromRange(ExtractionContext.CurrentTime, static_cast<double>(ExtractionContext.DeltaTimeRecord.Delta), ExtractionContext);
 }
 
-FTransform UAnimComposite::ExtractRootMotionFromRange(float StartTrackPosition, float EndTrackPosition) const
+FTransform UAnimComposite::ExtractRootMotionFromRange(double StartTime, double EndTime, const FAnimExtractContext& ExtractionContext) const
 {
 	FRootMotionMovementParams RootMotion;
-	ExtractRootMotionFromTrack(AnimationTrack, StartTrackPosition, EndTrackPosition, RootMotion);
+	ExtractRootMotionFromTrack(AnimationTrack, StartTime, EndTime, ExtractionContext, RootMotion);
 	return RootMotion.GetRootMotionTransform();
 }
 
-FTransform UAnimComposite::ExtractRootTrackTransform(float Time, const FBoneContainer* RequiredBones) const
+FTransform UAnimComposite::ExtractRootTrackTransform(const FAnimExtractContext& ExtractionContext, const FBoneContainer* RequiredBones) const
 {
-	if (const FAnimSegment* AnimSegment = AnimationTrack.GetSegmentAtTime(Time))
+	if (const FAnimSegment* AnimSegment = AnimationTrack.GetSegmentAtTime(ExtractionContext.CurrentTime))
 	{
 		float SegmentTime = 0.0f;
-		if (const UAnimSequenceBase* SequenceBase = AnimSegment->GetAnimationData(Time, SegmentTime))
+		if (const UAnimSequenceBase* SequenceBase = AnimSegment->GetAnimationData(ExtractionContext.CurrentTime, SegmentTime))
 		{
-			return SequenceBase->ExtractRootTrackTransform(SegmentTime, RequiredBones);
+			FAnimExtractContext Context = ExtractionContext;
+			Context.CurrentTime = static_cast<double>(SegmentTime);
+			return SequenceBase->ExtractRootTrackTransform(Context, RequiredBones);
 		}
 	}
 
@@ -125,10 +127,13 @@ FTransform UAnimComposite::ExtractRootTrackTransform(float Time, const FBoneCont
 	if (!AnimationTrack.AnimSegments.IsEmpty())
 	{
 		const int32 LastSegmentIndex = AnimationTrack.AnimSegments.Num() - 1;
-		if (Time > AnimationTrack.AnimSegments[LastSegmentIndex].AnimEndTime)
+		if (ExtractionContext.CurrentTime > AnimationTrack.AnimSegments[LastSegmentIndex].AnimEndTime)
 		{
 			const UAnimSequenceBase* SequenceBase = AnimationTrack.AnimSegments[LastSegmentIndex].GetAnimReference().Get();
-			return SequenceBase->ExtractRootTrackTransform(SequenceBase->GetPlayLength(), RequiredBones);
+
+			FAnimExtractContext Context = ExtractionContext;
+			Context.CurrentTime = static_cast<double>(SequenceBase->GetPlayLength());
+			return SequenceBase->ExtractRootTrackTransform(Context, RequiredBones);
 		}
 	}
 
diff --git a/Engine/Source/Runtime/Engine/Private/Animation/AnimCompositeBase.cpp b/Engine/Source/Runtime/Engine/Private/Animation/AnimCompositeBase.cpp
--- a/Engine/Source/Runtime/Engine/Private/Animation/AnimCompositeBase.cpp
+++ b/Engine/Source/Runtime/Engine/Private/Animation/AnimCompositeBase.cpp
@@ -782,7 +782,7 @@ UAnimCompositeBase::UAnimCompositeBase(const FObjectInitializer& ObjectInitializ
 {
 }
 
-void UAnimCompositeBase::ExtractRootMotionFromTrack(const FAnimTrack &SlotAnimTrack, float StartTrackPosition, float EndTrackPosition, FRootMotionMovementParams &RootMotion) const
+void UAnimCompositeBase::ExtractRootMotionFromTrack(const FAnimTrack& SlotAnimTrack, float StartTrackPosition, float EndTrackPosition, const FAnimExtractContext& Context, FRootMotionMovementParams& RootMotion) const
 {
 	TArray<FRootMotionExtractionStep> RootMotionExtractionSteps;
 	SlotAnimTrack.GetRootMotionExtractionStepsForTrackRange(RootMotionExtractionSteps, StartTrackPosition, EndTrackPosition);
@@ -797,7 +797,7 @@ void UAnimCompositeBase::ExtractRootMotionFromTrack(const FAnimTrack &SlotAnimTr
 		const FRootMotionExtractionStep & CurrentStep = RootMotionExtractionSteps[StepIndex];
 		if (CurrentStep.AnimSequence->bEnableRootMotion)
 		{
-			FTransform DeltaTransform = CurrentStep.AnimSequence->ExtractRootMotionFromRange(CurrentStep.StartPosition, CurrentStep.EndPosition);
+			FTransform DeltaTransform = CurrentStep.AnimSequence->ExtractRootMotionFromRange(CurrentStep.StartPosition, CurrentStep.EndPosition, Context);
 			RootMotion.Accumulate(DeltaTransform);
 		
 			UE_LOG(LogRootMotion, Log, TEXT("\t\tCurrentStep: %d, StartPos: %.3f, EndPos: %.3f, Anim: %s DeltaTransform Translation: %s, Rotation: %s"),
diff --git a/Engine/Source/Runtime/Engine/Private/Animation/AnimCurveTypes.cpp b/Engine/Source/Runtime/Engine/Private/Animation/AnimCurveTypes.cpp
--- a/Engine/Source/Runtime/Engine/Private/Animation/AnimCurveTypes.cpp
+++ b/Engine/Source/Runtime/Engine/Private/Animation/AnimCurveTypes.cpp
@@ -403,7 +403,7 @@ bool FCachedFloatCurve::IsValid(const UAnimSequenceBase* InAnimSequence) const
 
 float FCachedFloatCurve::GetValueAtPosition(const UAnimSequenceBase* InAnimSequence, const float& InPosition) const
 {
-	return InAnimSequence->EvaluateCurveData(CurveName, InPosition);
+	return InAnimSequence->EvaluateCurveData(CurveName, FAnimExtractContext(static_cast<const double>(InPosition)));
 }
 
 const FFloatCurve* FCachedFloatCurve::GetFloatCurve(const UAnimSequenceBase* InAnimSequence) const
diff --git a/Engine/Source/Runtime/Engine/Private/Animation/AnimInstanceProxy.cpp b/Engine/Source/Runtime/Engine/Private/Animation/AnimInstanceProxy.cpp
--- a/Engine/Source/Runtime/Engine/Private/Animation/AnimInstanceProxy.cpp
+++ b/Engine/Source/Runtime/Engine/Private/Animation/AnimInstanceProxy.cpp
@@ -1618,7 +1618,8 @@ void FAnimInstanceProxy::SlotEvaluatePoseWithBlendProfiles(const FName& SlotNode
 			// Add montage curves.
 			FBlendedCurve MontageCurve;
 			MontageCurve.InitFrom(*RequiredBones);
-			Montage->EvaluateCurveData(MontageCurve, EvalState.MontagePosition);
+			const FAnimExtractContext Context(static_cast<double>(EvalState.MontagePosition));
+			Montage->EvaluateCurveData(MontageCurve, Context);
 			NewPose.Curve.Combine(MontageCurve);
 
 			// Capture non-additive poses only in the non-additive pass.
@@ -1959,7 +1960,7 @@ void FAnimInstanceProxy::SlotEvaluatePose(const FName& SlotNodeName, const FAnim
 			// add montage curves 
 			FBlendedCurve MontageCurve;
 			MontageCurve.InitFrom(*RequiredBones);
-			Montage->EvaluateCurveData(MontageCurve, EvalState.MontagePosition);
+			Montage->EvaluateCurveData(MontageCurve, ExtractionContext);
 			NewPose.Curve.Combine(MontageCurve);
 
 #if DEBUG_MONTAGEINSTANCE_WEIGHT
diff --git a/Engine/Source/Runtime/Engine/Private/Animation/AnimMontage.cpp b/Engine/Source/Runtime/Engine/Private/Animation/AnimMontage.cpp
--- a/Engine/Source/Runtime/Engine/Private/Animation/AnimMontage.cpp
+++ b/Engine/Source/Runtime/Engine/Private/Animation/AnimMontage.cpp
@@ -915,7 +915,7 @@ bool UAnimMontage::HasRootMotion() const
  * It will break down the range into steps if needed to handle looping animations, or different animations.
  * These steps will be processed sequentially, and output the RootMotion transform in component space.
  */
-FTransform UAnimMontage::ExtractRootMotionFromTrackRange(float StartTrackPosition, float EndTrackPosition) const
+FTransform UAnimMontage::ExtractRootMotionFromTrackRange(float StartTrackPosition, float EndTrackPosition, const FAnimExtractContext& Context) const
 {
 	FRootMotionMovementParams RootMotion;
 
@@ -927,7 +927,7 @@ FTransform UAnimMontage::ExtractRootMotionFromTrackRange(float StartTrackPositio
 		// Get RootMotion pieces from this track.
 		// We can deal with looping animations, or multiple animations. So we break those up into sequential operations.
 		// (Animation, StartFrame, EndFrame) so we can then extract root motion sequentially.
-		ExtractRootMotionFromTrack(SlotAnimTrack, StartTrackPosition, EndTrackPosition, RootMotion);
+		ExtractRootMotionFromTrack(SlotAnimTrack, StartTrackPosition, EndTrackPosition, Context, RootMotion);
 
 	}
 
@@ -2045,7 +2045,7 @@ bool FAnimMontageInstance::SimulateAdvance(float DeltaTime, float& InOutPosition
 		// Extract Root Motion for this time slice, and accumulate it.
 		if (bExtractRootMotion)
 		{
-			OutRootMotionParams.Accumulate(Montage->ExtractRootMotionFromTrackRange(PreviousSubStepPosition, InOutPosition));
+			OutRootMotionParams.Accumulate(Montage->ExtractRootMotionFromTrackRange(PreviousSubStepPosition, InOutPosition, FAnimExtractContext()));
 		}
 
 		// if we reached end of section, and we were not processing a branching point, and no events has messed with out current position..
@@ -2567,7 +2567,7 @@ void FAnimMontageInstance::Advance(float DeltaTime, struct FRootMotionMovementPa
 					// IsRootMotionDisabled() can be changed by AnimNotifyState BranchingPoints while advancing, so it needs to be checked here.
 					if (bExtractRootMotion && AnimInstance.IsValid() && !IsRootMotionDisabled())
 					{
-						const FTransform RootMotion = Montage->ExtractRootMotionFromTrackRange(PreviousSubStepPosition, Position);
+						const FTransform RootMotion = Montage->ExtractRootMotionFromTrackRange(PreviousSubStepPosition, Position, FAnimExtractContext());
 						if (bBlendRootMotion)
 						{
 							// Defer blending in our root motion until after we get our slot weight updated
diff --git a/Engine/Source/Runtime/Engine/Private/Animation/AnimSequence.cpp b/Engine/Source/Runtime/Engine/Private/Animation/AnimSequence.cpp
--- a/Engine/Source/Runtime/Engine/Private/Animation/AnimSequence.cpp
+++ b/Engine/Source/Runtime/Engine/Private/Animation/AnimSequence.cpp
@@ -1319,8 +1319,7 @@ static FAutoConsoleVariableRef CVarForceRawData(
 	ECVF_Cheat);
 #endif // WITH_EDITOR
 
-
-void UAnimSequence::GetBoneTransform(FTransform& OutAtom, FSkeletonPoseBoneIndex BoneIndex, double Time, bool bUseRawData, TOptional<EAnimInterpolationType> InterpolationOverride) const
+void UAnimSequence::GetBoneTransform(FTransform& OutAtom, FSkeletonPoseBoneIndex BoneIndex, const FAnimExtractContext& ExtractionContext, bool bUseRawData) const
 {
 	// If the caller didn't request that raw animation data be used . . .
 	const bool bEvaluateCompressedData =
@@ -1329,11 +1328,11 @@ void UAnimSequence::GetBoneTransform(FTransform& OutAtom, FSkeletonPoseBoneIndex
 #endif 
 		(!bUseRawData || !CanEvaluateRawAnimationData());
 
-	const EAnimInterpolationType InterpolationType = InterpolationOverride.Get(Interpolation);
+	const EAnimInterpolationType InterpolationType = ExtractionContext.InterpolationOverride.Get(Interpolation);
 	if ( bEvaluateCompressedData && IsCompressedDataValid())
 	{
 		FAnimSequenceDecompressionContext DecompContext(PlatformTargetFrameRate.Default, PlatformTargetFrameRate.Default.AsFrameTime(GetPlayLength()).RoundToFrame().Value, InterpolationType, GetRetargetTransformsSourceName(), *CompressedData.CompressedDataStructure, GetSkeleton()->GetRefLocalPoses(), CompressedData.CompressedTrackToSkeletonMapTable, GetSkeleton(), IsValidAdditive(), AdditiveAnimType);
-		DecompContext.Seek(Time);
+		DecompContext.Seek(ExtractionContext.CurrentTime);
 		if (CompressedData.BoneCompressionCodec)
 		{
 			const int32 TrackIndex = CompressedData.GetTrackIndexFromSkeletonIndex(BoneIndex.GetInt());
@@ -1348,12 +1347,12 @@ void UAnimSequence::GetBoneTransform(FTransform& OutAtom, FSkeletonPoseBoneIndex
 #if WITH_EDITOR
 		ValidateModel();
 		const FName BoneName = GetSkeleton()->GetReferenceSkeleton().GetBoneName(BoneIndex.GetInt());
-		OutAtom = DataModelInterface->EvaluateBoneTrackTransform(BoneName, DataModelInterface->GetFrameRate().AsFrameTime(Time), InterpolationType);
+		OutAtom = DataModelInterface->EvaluateBoneTrackTransform(BoneName, DataModelInterface->GetFrameRate().AsFrameTime(ExtractionContext.CurrentTime), InterpolationType);
 
 		const FAnimationCurveIdentifier TransformCurveId(BoneName, ERawCurveTrackTypes::RCT_Transform);
 		if (const FTransformCurve* TransformCurvePtr = DataModelInterface->FindTransformCurve(TransformCurveId))
 		{
-			const FTransform AdditiveTransform = TransformCurvePtr->Evaluate(Time, 1.f);
+			const FTransform AdditiveTransform = TransformCurvePtr->Evaluate(ExtractionContext.CurrentTime, 1.f);
 			const FTransform LocalTransform = OutAtom;
 			OutAtom.SetRotation(LocalTransform.GetRotation() * AdditiveTransform.GetRotation());
 			OutAtom.SetTranslation(LocalTransform.TransformPosition(AdditiveTransform.GetTranslation()));
@@ -1401,11 +1400,12 @@ void UAnimSequence::HandleAssetPlayerTickedInternal(FAnimAssetTickContext &Conte
 
 	if (bEnableRootMotion)
 	{
-		Context.RootMotionMovementParams.Accumulate(ExtractRootMotion(PreviousTime, MoveDelta, Instance.bLooping));
+		const FAnimExtractContext ExtractionContext(static_cast<double>(PreviousTime), true, FDeltaTimeRecord(MoveDelta), Instance.bLooping);
+		Context.RootMotionMovementParams.Accumulate(ExtractRootMotion(ExtractionContext));
 	}
 }
 
-FTransform UAnimSequence::ExtractRootTrackTransform(float Time, const FBoneContainer * RequiredBones) const
+FTransform UAnimSequence::ExtractRootTrackTransform(const FAnimExtractContext& ExtractionContext, const FBoneContainer* RequiredBones) const
 {
 	const int32 RootBoneIndex = 0;
 	const bool bContainsRootBoneTrack = [this, RootBoneIndex]() 
@@ -1443,7 +1443,7 @@ FTransform UAnimSequence::ExtractRootTrackTransform(float Time, const FBoneConta
 	{
 		// if we do have root data, then return root data
 		FTransform RootTransform;
-		GetBoneTransform(RootTransform, FSkeletonPoseBoneIndex(RootBoneIndex), static_cast<double>(Time), 
+		GetBoneTransform(RootTransform, FSkeletonPoseBoneIndex(RootBoneIndex), ExtractionContext, 
 			PRAGMA_DISABLE_DEPRECATION_WARNINGS
 			bUseRawDataOnly
 			PRAGMA_ENABLE_DEPRECATION_WARNINGS);
@@ -1475,17 +1475,17 @@ FTransform UAnimSequence::ExtractRootTrackTransform(float Time, const FBoneConta
 	return FTransform::Identity;
 }
 
-FTransform UAnimSequence::ExtractRootMotion(float StartTime, float DeltaTime, bool bAllowLooping) const
+FTransform UAnimSequence::ExtractRootMotion(const FAnimExtractContext& ExtractionContext) const
 {
 	FRootMotionMovementParams RootMotionParams;
 
-	if (DeltaTime != 0.f)
+	if (ExtractionContext.DeltaTimeRecord.Delta != 0.f)
 	{
-		bool const bPlayingBackwards = (DeltaTime < 0.f);
+		const bool bPlayingBackwards = (ExtractionContext.DeltaTimeRecord.Delta < 0.f);
 
-		float PreviousPosition = StartTime;
-		float CurrentPosition = StartTime;
-		float DesiredDeltaMove = DeltaTime;
+		float PreviousPosition = ExtractionContext.CurrentTime;
+		float CurrentPosition = ExtractionContext.CurrentTime;
+		float DesiredDeltaMove = ExtractionContext.DeltaTimeRecord.Delta;
 
 		do
 		{
@@ -1496,10 +1496,10 @@ FTransform UAnimSequence::ExtractRootMotion(float StartTime, float DeltaTime, bo
 			ensureMsgf(bPlayingBackwards ? (CurrentPosition <= PreviousPosition) : (CurrentPosition >= PreviousPosition), TEXT("in Animation %s(Skeleton %s) : bPlayingBackwards(%d), PreviousPosition(%0.2f), Current Position(%0.2f)"),
 				*GetName(), *GetNameSafe(GetSkeleton()), bPlayingBackwards, PreviousPosition, CurrentPosition);
 
-			RootMotionParams.Accumulate(ExtractRootMotionFromRange(PreviousPosition, CurrentPosition));
+			RootMotionParams.Accumulate(ExtractRootMotionFromRange(PreviousPosition, CurrentPosition, ExtractionContext));
 
 			// If we've hit the end of the animation, and we're allowed to loop, keep going.
-			if ((AdvanceType == ETAA_Finished) && bAllowLooping)
+			if ((AdvanceType == ETAA_Finished) && ExtractionContext.bLooping)
 			{
 				const double ActualDeltaMove = (CurrentPosition - PreviousPosition);
 				DesiredDeltaMove -= ActualDeltaMove;
@@ -1517,7 +1517,7 @@ FTransform UAnimSequence::ExtractRootMotion(float StartTime, float DeltaTime, bo
 	return RootMotionParams.GetRootMotionTransform();
 }
 
-FTransform UAnimSequence::ExtractRootMotionFromRange(float StartTrackPosition, float EndTrackPosition) const
+FTransform UAnimSequence::ExtractRootMotionFromRange(double StartTime, double EndTime, const FAnimExtractContext& ExtractionContext) const
 {
 	const FVector DefaultScale(1.f);
 
@@ -1531,8 +1531,11 @@ FTransform UAnimSequence::ExtractRootMotionFromRange(float StartTrackPosition, f
 		}
 	}
 
-	FTransform StartTransform = ExtractRootTrackTransform(StartTrackPosition, nullptr);
-	FTransform EndTransform = ExtractRootTrackTransform(EndTrackPosition, nullptr);
+	FAnimExtractContext Context = ExtractionContext;
+	Context.CurrentTime = StartTime;
+	FTransform StartTransform = ExtractRootTrackTransform(Context, nullptr);
+	Context.CurrentTime = EndTime;
+	FTransform EndTransform = ExtractRootTrackTransform(Context, nullptr);
 
 	// Use old calculation if needed.
 	if (bUseNormalizedRootMotionScale)
@@ -1732,12 +1735,14 @@ void UAnimSequence::GetBonePose(FAnimationPoseData& OutAnimationPoseData, const
 	const int32 NumTracks = CompressedData.CompressedTrackToSkeletonMapTable.Num();
 #endif 
 	const bool bTreatAnimAsAdditive = (IsValidAdditive() && !bUseRawDataForPoseExtraction); // Raw data is never additive
+	FAnimExtractContext RootMotionContext = ExtractionContext;
+	RootMotionContext.CurrentTime = 0.0;
 	const FRootMotionReset RootMotionReset(bEnableRootMotion, RootMotionRootLock,
 #if WITH_EDITOR
 		!ExtractionContext.bIgnoreRootLock &&
 #endif // WITH_EDITOR
 		bForceRootLock,
-		ExtractRootTrackTransform(0.0f, &RequiredBones), bTreatAnimAsAdditive);
+		ExtractRootTrackTransform(RootMotionContext, &RequiredBones), bTreatAnimAsAdditive);
 
 #if WITH_EDITOR
 	// Evaluate raw (source) curve and bone data
@@ -1774,7 +1779,7 @@ void UAnimSequence::GetBonePose(FAnimationPoseData& OutAnimationPoseData, const
 		if (!bUseRawDataForPoseExtraction)
 #endif // WITH_EDITOR
 		{
-			EvaluateCurveData(OutAnimationPoseData.GetCurve(), ExtractionContext.CurrentTime, bUseRawDataForPoseExtraction);
+			EvaluateCurveData(OutAnimationPoseData.GetCurve(), ExtractionContext, bUseRawDataForPoseExtraction);
 		}
 	}
 
@@ -2849,7 +2854,7 @@ int32 UAnimSequence::GetNumberOfSampledKeys() const
 #endif
 }
 
-void UAnimSequence::EvaluateCurveData(FBlendedCurve& OutCurve, float CurrentTime, bool bForceUseRawData) const
+void UAnimSequence::EvaluateCurveData(FBlendedCurve& OutCurve, const FAnimExtractContext& AnimExtractContext, bool bForceUseRawData) const
 {
 	SCOPE_CYCLE_COUNTER(STAT_AnimSeq_EvalCurveData);
 	
@@ -2862,19 +2867,19 @@ void UAnimSequence::EvaluateCurveData(FBlendedCurve& OutCurve, float CurrentTime
 	if (CanEvaluateRawAnimationData() && bEvaluateRawData)
 	{
 #if WITH_EDITOR
-		UE::Anim::EvaluateFloatCurvesFromModel(DataModelInterface.GetInterface(), OutCurve, CurrentTime);
+		UE::Anim::EvaluateFloatCurvesFromModel(DataModelInterface.GetInterface(), OutCurve, AnimExtractContext.CurrentTime);
 #else
-		Super::EvaluateCurveData(OutCurve, CurrentTime, bForceUseRawData);
+		Super::EvaluateCurveData(OutCurve, AnimExtractContext, bForceUseRawData);
 #endif
 	}
 	else if(IsCurveCompressedDataValid() && CompressedData.CurveCompressionCodec)
 	{
 		CSV_SCOPED_TIMING_STAT(Animation, EvaluateCurveData);
-		CompressedData.CurveCompressionCodec->DecompressCurves(CompressedData, OutCurve, CurrentTime);
+		CompressedData.CurveCompressionCodec->DecompressCurves(CompressedData, OutCurve, AnimExtractContext.CurrentTime);
 	}
 }
 
-float UAnimSequence::EvaluateCurveData(FName CurveName, float CurrentTime, bool bForceUseRawData) const
+float UAnimSequence::EvaluateCurveData(FName CurveName, const FAnimExtractContext& AnimExtractContext, bool bForceUseRawData) const
 {
 	QUICK_SCOPE_CYCLE_COUNTER(EvaluateCurveDataByName);
 
@@ -2886,11 +2891,11 @@ float UAnimSequence::EvaluateCurveData(FName CurveName, float CurrentTime, bool
 	check(!bForceUseRawData || CanEvaluateRawAnimationData());
 	if (CanEvaluateRawAnimationData() && bEvaluateRawData)
 	{
-		return Super::EvaluateCurveData(CurveName, CurrentTime, bForceUseRawData);
+		return Super::EvaluateCurveData(CurveName, AnimExtractContext, bForceUseRawData);
 	}
 	else if(IsCurveCompressedDataValid() && CompressedData.CurveCompressionCodec)
 	{
-		return CompressedData.CurveCompressionCodec->DecompressCurve(CompressedData, CurveName, CurrentTime);
+		return CompressedData.CurveCompressionCodec->DecompressCurve(CompressedData, CurveName, AnimExtractContext.CurrentTime);
 	}
 
 	return 0.f;
diff --git a/Engine/Source/Runtime/Engine/Private/Animation/AnimSequenceBase.cpp b/Engine/Source/Runtime/Engine/Private/Animation/AnimSequenceBase.cpp
--- a/Engine/Source/Runtime/Engine/Private/Animation/AnimSequenceBase.cpp
+++ b/Engine/Source/Runtime/Engine/Private/Animation/AnimSequenceBase.cpp
@@ -1069,23 +1069,22 @@ void UAnimSequenceBase::RefreshParentAssetData()
 #endif	//WITH_EDITOR
 
 
-/** Add curve data to Instance at the time of CurrentTime **/
-void UAnimSequenceBase::EvaluateCurveData(FBlendedCurve& OutCurve, float CurrentTime, bool bForceUseRawData) const
+void UAnimSequenceBase::EvaluateCurveData(FBlendedCurve& OutCurve, const FAnimExtractContext& AnimExtractContext, bool bForceUseRawData) const
 {
 	CSV_SCOPED_TIMING_STAT(Animation, EvaluateCurveData);
 
 	PRAGMA_DISABLE_DEPRECATION_WARNINGS
-	RawCurveData.EvaluateCurveData(OutCurve, CurrentTime);
+	RawCurveData.EvaluateCurveData(OutCurve, AnimExtractContext.CurrentTime);
 	PRAGMA_ENABLE_DEPRECATION_WARNINGS
 }
 
-float UAnimSequenceBase::EvaluateCurveData(FName CurveName, float CurrentTime, bool bForceUseRawData) const
+float UAnimSequenceBase::EvaluateCurveData(FName CurveName, const FAnimExtractContext& AnimExtractContext, bool bForceUseRawData) const
 {
 	PRAGMA_DISABLE_DEPRECATION_WARNINGS
 	const FFloatCurve* Curve = (const FFloatCurve*)RawCurveData.GetCurveData(CurveName, ERawCurveTrackTypes::RCT_Float);
 	PRAGMA_ENABLE_DEPRECATION_WARNINGS
 
-	return Curve != nullptr ? Curve->Evaluate(CurrentTime) : 0.0f;
+	return Curve != nullptr ? Curve->Evaluate(AnimExtractContext.CurrentTime) : 0.0f;
 }
 
 const FRawCurveTracks& UAnimSequenceBase::GetCurveData() const
diff --git a/Engine/Source/Runtime/Engine/Private/Animation/AnimSequenceHelpers.cpp b/Engine/Source/Runtime/Engine/Private/Animation/AnimSequenceHelpers.cpp
--- a/Engine/Source/Runtime/Engine/Private/Animation/AnimSequenceHelpers.cpp
+++ b/Engine/Source/Runtime/Engine/Private/Animation/AnimSequenceHelpers.cpp
@@ -261,17 +261,17 @@ FTransform ExtractRootMotionFromAnimationAsset(const UAnimationAsset* Animation,
 	
 	if (const UAnimMontage* Montage = Cast<UAnimMontage>(Animation))
 	{
-		Result = Montage->ExtractRootMotionFromTrackRange(StartPosition, EndPosition);
+		Result = Montage->ExtractRootMotionFromTrackRange(StartPosition, EndPosition, FAnimExtractContext());
 	}
 	else if (const UAnimComposite* AnimComposite = Cast<UAnimComposite>(Animation))
 	{
 		FRootMotionMovementParams RootMotion;
-		AnimComposite->ExtractRootMotionFromTrack(AnimComposite->AnimationTrack, StartPosition, EndPosition, RootMotion);
+		AnimComposite->ExtractRootMotionFromTrack(AnimComposite->AnimationTrack, StartPosition, EndPosition, FAnimExtractContext(), RootMotion);
 		Result = RootMotion.GetRootMotionTransform();
 	}
 	else if (const UAnimSequence* AnimSequence = Cast<UAnimSequence>(Animation))
 	{
-		Result = AnimSequence->ExtractRootMotionFromRange(StartPosition, EndPosition);
+		Result = AnimSequence->ExtractRootMotionFromRange(StartPosition, EndPosition, FAnimExtractContext());
 	}
 
 	if (MirrorDataTable)
@@ -308,18 +308,18 @@ FTransform ExtractRootTransformFromAnimationAsset(const UAnimationAsset* Animati
 			if (const UAnimSequence* AnimSequence = Cast<UAnimSequence>(Segment->GetAnimReference()))
 			{
 				const float AnimSequenceTime = Segment->ConvertTrackPosToAnimPos(Time);
-				Result = AnimSequence->ExtractRootTrackTransform(AnimSequenceTime, nullptr);
+				Result = AnimSequence->ExtractRootTrackTransform(FAnimExtractContext(static_cast<double>(AnimSequenceTime)), nullptr);
 				bUseNormalizedRootMotionScale = AnimSequence->bUseNormalizedRootMotionScale;
 			}	
 		}
 	}
 	else if (const UAnimComposite* AnimComposite = Cast<UAnimComposite>(Animation))
 	{
-		Result = AnimComposite->ExtractRootTrackTransform(Time, nullptr);
+		Result = AnimComposite->ExtractRootTrackTransform(FAnimExtractContext(static_cast<double>(Time)), nullptr);
 	}
 	else if (const UAnimSequence* AnimSequence = Cast<UAnimSequence>(Animation))
 	{
-		Result = AnimSequence->ExtractRootTrackTransform(Time, nullptr);
+		Result = AnimSequence->ExtractRootTrackTransform(FAnimExtractContext(static_cast<double>(Time)), nullptr);
 		bUseNormalizedRootMotionScale = AnimSequence->bUseNormalizedRootMotionScale;
 	}
 
diff --git a/Engine/Source/Runtime/Engine/Private/Animation/BlendSpace.cpp b/Engine/Source/Runtime/Engine/Private/Animation/BlendSpace.cpp
--- a/Engine/Source/Runtime/Engine/Private/Animation/BlendSpace.cpp
+++ b/Engine/Source/Runtime/Engine/Private/Animation/BlendSpace.cpp
@@ -774,7 +774,8 @@ void UBlendSpace::TickAssetPlayer(FAnimTickRecord& Instance, struct FAnimNotifyQ
 						{
 							if (Context.RootMotionMode == ERootMotionMode::RootMotionFromEverything && Sample.Animation->bEnableRootMotion)
 							{
-								Context.RootMotionMovementParams.AccumulateWithBlend(Sample.Animation->ExtractRootMotion(PrevSampleDataTime, DeltaTimePosition, Instance.bLooping), SampleEntry.GetClampedWeight());
+								const FAnimExtractContext ExtractionContext(static_cast<double>(PrevSampleDataTime), true, FDeltaTimeRecord(DeltaTimePosition), Instance.bLooping);
+								Context.RootMotionMovementParams.AccumulateWithBlend(Sample.Animation->ExtractRootMotion(ExtractionContext), SampleEntry.GetClampedWeight());
 							}
 						}
 
diff --git a/Engine/Source/Runtime/Engine/Private/Animation/PoseAsset.cpp b/Engine/Source/Runtime/Engine/Private/Animation/PoseAsset.cpp
--- a/Engine/Source/Runtime/Engine/Private/Animation/PoseAsset.cpp
+++ b/Engine/Source/Runtime/Engine/Private/Animation/PoseAsset.cpp
@@ -1501,7 +1501,7 @@ void UPoseAsset::CreatePoseFromAnimation(class UAnimSequence* AnimSequence, cons
 				NewPose.Reset(NumTracks);
 				NewPose.AddUninitialized(NumTracks);
 				
-				float IntervalBetweenKeys = (NumPoses > 1)? AnimSequence->GetPlayLength() / (NumPoses -1 ) : 0.f;
+				const double IntervalBetweenKeys = (NumPoses > 1)? AnimSequence->GetPlayLength() / (NumPoses -1 ) : 0.f;
 
 				// add curves - only float curves
 				const FAnimationCurveData& AnimationCurveData = DataModel->GetCurveData();
@@ -1534,7 +1534,7 @@ void UPoseAsset::CreatePoseFromAnimation(class UAnimSequence* AnimSequence, cons
 						// have to do iterate over time
 						// support curve
 						FBlendedCurve SourceCurve;
-						AnimSequence->EvaluateCurveData(SourceCurve, PoseIndex*IntervalBetweenKeys, true);
+						AnimSequence->EvaluateCurveData(SourceCurve, FAnimExtractContext(PoseIndex*IntervalBetweenKeys), true);
 						
 						SourceCurve.ForEachElement([&CurveData](const UE::Anim::FCurveElement& InElement)
 						{
diff --git a/Engine/Source/Runtime/Engine/Private/Character.cpp b/Engine/Source/Runtime/Engine/Private/Character.cpp
--- a/Engine/Source/Runtime/Engine/Private/Character.cpp
+++ b/Engine/Source/Runtime/Engine/Private/Character.cpp
@@ -1425,7 +1425,7 @@ void ACharacter::SimulatedRootMotionPositionFixup(float DeltaSeconds)
 				if( FMath::Abs(DeltaPosition) > UE_KINDA_SMALL_NUMBER )
 				{
 					// Find Root Motion delta move to get back to where we were on the client.
-					const FTransform LocalRootMotionTransform = ClientMontageInstance->Montage->ExtractRootMotionFromTrackRange(ServerPosition, ClientPosition);
+					const FTransform LocalRootMotionTransform = ClientMontageInstance->Montage->ExtractRootMotionFromTrackRange(ServerPosition, ClientPosition, FAnimExtractContext());
 
 					// Simulate Root Motion for delta move.
 					if( CharacterMovement )
diff --git a/Engine/Source/Runtime/MovieSceneTracks/Private/Sections/MovieSceneSkeletalAnimationSection.cpp b/Engine/Source/Runtime/MovieSceneTracks/Private/Sections/MovieSceneSkeletalAnimationSection.cpp
--- a/Engine/Source/Runtime/MovieSceneTracks/Private/Sections/MovieSceneSkeletalAnimationSection.cpp
+++ b/Engine/Source/Runtime/MovieSceneTracks/Private/Sections/MovieSceneSkeletalAnimationSection.cpp
@@ -612,7 +612,7 @@ bool UMovieSceneSkeletalAnimationSection::GetRootMotionVelocity(FFrameTime Previ
 		//mz todo need to get the starting value.
 		float PreviousTimeSeconds = static_cast<float>(MapTimeToAnimation(PreviousTime, FrameRate));
 		float CurrentTimeSeconds  = static_cast<float>(MapTimeToAnimation(CurrentTime, FrameRate));
-		OutVelocity = AnimSequence->ExtractRootMotionFromRange(PreviousTimeSeconds, CurrentTimeSeconds);
+		OutVelocity = AnimSequence->ExtractRootMotionFromRange(PreviousTimeSeconds, CurrentTimeSeconds, FAnimExtractContext());
 		return true;
 	}
 	return false;
@@ -695,9 +695,7 @@ int32 UMovieSceneSkeletalAnimationSection::SetBoneIndexForRootMotionCalculations
 						for (int32 Index = 0; Index < NumFrames; ++Index)
 						{
 							const double Pos = FMath::Clamp(AnimSequence->GetSamplingFrameRate().AsSeconds(Index), 0.0, AnimSequence->GetPlayLength());
-PRAGMA_DISABLE_DEPRECATION_WARNINGS
-							AnimSequence->GetBoneTransform(Transform, TrackIndex, Pos, false);
-PRAGMA_ENABLE_DEPRECATION_WARNINGS
+							AnimSequence->GetBoneTransform(Transform, FSkeletonPoseBoneIndex(BoneTreeIndex), FAnimExtractContext(Pos), false);
 							if (Transform.Equals(FTransform::Identity) == false)
 							{
 								TempRootBoneIndex = BoneTreeIndex;
@@ -724,7 +722,7 @@ FTransform UMovieSceneSkeletalAnimationSection::GetRootMotionStartOffset() const
 	{
 		UMovieScene* MovieScene = GetTypedOuter<UMovieScene>();
 		const double StartSeconds = MovieScene ? (MapTimeToAnimation(FFrameNumber(0), MovieScene->GetTickResolution())): 0.0;
-		return AnimSequence->ExtractRootTrackTransform(StartSeconds, nullptr);
+		return AnimSequence->ExtractRootTrackTransform(FAnimExtractContext(StartSeconds), nullptr);
 	}
 	return FTransform::Identity;;
 }
@@ -762,8 +760,8 @@ bool UMovieSceneSkeletalAnimationSection::GetRootMotionTransform(FAnimationPoseD
 		}
 		else //not set then just use root.
 		{
-			StartBoneTransform = AnimSequence->ExtractRootTrackTransform(StartSeconds, nullptr);
-			InOutParams.OutPoseTransform = AnimSequence->ExtractRootTrackTransform(CurrentTimeSeconds, nullptr);
+			StartBoneTransform = AnimSequence->ExtractRootTrackTransform(FAnimExtractContext(StartSeconds), nullptr);
+			InOutParams.OutPoseTransform = AnimSequence->ExtractRootTrackTransform(FAnimExtractContext(CurrentTimeSeconds), nullptr);
 		}
 		//note though we don't support mesh space addtive just local additive it will still work the same here for the root so 
 		if (!InOutParams.bOutIsAdditive)
