Mark constraints as dirty if bodies are changed from kinematic to dynamic etc - so that the constraint sorting can get updated.

diff --git a/Engine/Source/Runtime/Engine/Private/PhysicsEngine/ImmediatePhysics/ImmediatePhysicsChaos/ImmediatePhysicsSimulation_Chaos.cpp b/Engine/Source/Runtime/Engine/Private/PhysicsEngine/ImmediatePhysics/ImmediatePhysicsChaos/ImmediatePhysicsSimulation_Chaos.cpp
--- a/Engine/Source/Runtime/Engine/Private/PhysicsEngine/ImmediatePhysics/ImmediatePhysicsChaos/ImmediatePhysicsSimulation_Chaos.cpp
+++ b/Engine/Source/Runtime/Engine/Private/PhysicsEngine/ImmediatePhysics/ImmediatePhysicsChaos/ImmediatePhysicsSimulation_Chaos.cpp
@@ -522,6 +522,9 @@ namespace ImmediatePhysics_Chaos
 		{
 			ActorHandle->SetIsKinematic(bKinematic);
 			Implementation->bActorsDirty = true;
+			// Force constraint sorting to be updated
+			Implementation->Joints.LinearConstraints.MarkJointsDirty();
+			Implementation->Joints.NonLinearConstraints.MarkJointsDirty();
 		}
 	}
 
@@ -532,6 +535,9 @@ namespace ImmediatePhysics_Chaos
 		{
 			ActorHandle->SetEnabled(bEnable);
 			Implementation->bActorsDirty = true;
+			// Force constraint sorting to be updated
+			Implementation->Joints.LinearConstraints.MarkJointsDirty();
+			Implementation->Joints.NonLinearConstraints.MarkJointsDirty();
 		}
 	}
 
diff --git a/Engine/Source/Runtime/Experimental/Chaos/Public/Chaos/PBDJointConstraints.h b/Engine/Source/Runtime/Experimental/Chaos/Public/Chaos/PBDJointConstraints.h
--- a/Engine/Source/Runtime/Experimental/Chaos/Public/Chaos/PBDJointConstraints.h
+++ b/Engine/Source/Runtime/Experimental/Chaos/Public/Chaos/PBDJointConstraints.h
@@ -397,6 +397,10 @@ namespace Chaos
 		CHAOS_API int32 GetConstraintLevel(int32 ConstraintIndex) const;
 		CHAOS_API int32 GetConstraintColor(int32 ConstraintIndex) const;
 
+		// This will force a re-evaluation of the sorting of constraints - needed if objects switch
+		// between kinematic/dynamic, for example.
+		void MarkJointsDirty() { bJointsDirty = true; }
+
 		// @todo(chaos): functionality no longer supported here (see TPBDJointContainerSolver)
 		UE_DEPRECATED(5.2, "Joint Solver API moved to TPBDJointContainerSolver")
 		void AddBodies(FSolverBodyContainer& SolverBodyContainer) {}
