Implement memory usage estimate during animated curve compression

diff --git a/Engine/Plugins/Animation/ACLPlugin/Source/ACLPlugin/Classes/AnimCurveCompressionCodec_ACL.h b/Engine/Plugins/Animation/ACLPlugin/Source/ACLPlugin/Classes/AnimCurveCompressionCodec_ACL.h
--- a/Engine/Plugins/Animation/ACLPlugin/Source/ACLPlugin/Classes/AnimCurveCompressionCodec_ACL.h
+++ b/Engine/Plugins/Animation/ACLPlugin/Source/ACLPlugin/Classes/AnimCurveCompressionCodec_ACL.h
@@ -33,6 +33,7 @@ class UAnimCurveCompressionCodec_ACL : public UAnimCurveCompressionCodec
 	// UAnimCurveCompressionCodec implementation
 	virtual void PopulateDDCKey(FArchive& Ar) override;
 	virtual bool Compress(const FCompressibleAnimData& AnimSeq, FAnimCurveCompressionResult& OutResult) override;
+	virtual int64 EstimateCompressionMemoryUsage(const UAnimSequence& AnimSequence) const override;
 #endif
 
 	// UAnimCurveCompressionCodec implementation
diff --git a/Engine/Plugins/Animation/ACLPlugin/Source/ACLPlugin/Private/AnimBoneCompressionCodec_ACLBase.cpp b/Engine/Plugins/Animation/ACLPlugin/Source/ACLPlugin/Private/AnimBoneCompressionCodec_ACLBase.cpp
--- a/Engine/Plugins/Animation/ACLPlugin/Source/ACLPlugin/Private/AnimBoneCompressionCodec_ACLBase.cpp
+++ b/Engine/Plugins/Animation/ACLPlugin/Source/ACLPlugin/Private/AnimBoneCompressionCodec_ACLBase.cpp
@@ -490,7 +490,7 @@ void UAnimBoneCompressionCodec_ACLBase::PopulateDDCKey(const UE::Anim::Compressi
 
 int64 UAnimBoneCompressionCodec_ACLBase::EstimateCompressionMemoryUsage(const UAnimSequence& AnimSequence) const
 {
-	const int64 AnimSeqRawSize = AnimSequence.GetApproxRawSize();
+	const int64 AnimSeqRawSize = AnimSequence.GetApproxBoneRawSize();
 
 	int64 EstimatedMemoryUsage = 0;
 	EstimatedMemoryUsage += AnimSeqRawSize;	// We copy the raw data into the ACL format
@@ -500,7 +500,7 @@ int64 UAnimBoneCompressionCodec_ACLBase::EstimateCompressionMemoryUsage(const UA
 		if (AnimSequence.RefPoseSeq)
 		{
 			// We copy the additive base into the ACL format
-			EstimatedMemoryUsage += AnimSequence.RefPoseSeq->GetApproxRawSize();
+			EstimatedMemoryUsage += AnimSequence.RefPoseSeq->GetApproxBoneRawSize();
 		}
 		else
 		{
diff --git a/Engine/Plugins/Animation/ACLPlugin/Source/ACLPlugin/Private/AnimCurveCompressionCodec_ACL.cpp b/Engine/Plugins/Animation/ACLPlugin/Source/ACLPlugin/Private/AnimCurveCompressionCodec_ACL.cpp
--- a/Engine/Plugins/Animation/ACLPlugin/Source/ACLPlugin/Private/AnimCurveCompressionCodec_ACL.cpp
+++ b/Engine/Plugins/Animation/ACLPlugin/Source/ACLPlugin/Private/AnimCurveCompressionCodec_ACL.cpp
@@ -12,6 +12,7 @@
 #include "AnimationCompression.h"
 
 #if WITH_EDITORONLY_DATA
+#include "Animation/AnimSequence.h"
 #include "Animation/MorphTarget.h"
 #include "Engine/SkeletalMesh.h"
 #include "Rendering/SkeletalMeshModel.h"
@@ -271,6 +272,23 @@ bool UAnimCurveCompressionCodec_ACL::Compress(const FCompressibleAnimData& AnimS
 	ACLAllocatorImpl.deallocate(CompressedTracks, CompressedDataSize);
 	return true;
 }
+
+int64 UAnimCurveCompressionCodec_ACL::EstimateCompressionMemoryUsage(const UAnimSequence& AnimSequence) const
+{
+	const int64 NumCurves = AnimSequence.GetDataModelInterface()->GetFloatCurves().Num();
+	const int64 NumSamples = AnimSequence.GetNumberOfSampledKeys();
+
+	const int64 RawDataSize = NumCurves * NumSamples;
+
+	int64 EstimatedMemoryUsage = 0;
+	EstimatedMemoryUsage += RawDataSize;	// We copy the raw data into the ACL format
+	EstimatedMemoryUsage *= 2;				// Internally, ACL copies the raw data into a different format than the input because the input is not modified
+	EstimatedMemoryUsage += RawDataSize;	// ACL keeps a mutable copy of the lossy data that it modifies during compression
+	EstimatedMemoryUsage += RawDataSize;	// ACL will allocate the output buffer, assume that it's as large as the raw data
+	EstimatedMemoryUsage += 100 * 1024;		// Reserve 100 KB for internal bookkeeping and other required metadata
+
+	return EstimatedMemoryUsage;
+}
 #endif // WITH_EDITORONLY_DATA
 
 bool UAnimCurveCompressionCodec_ACL::ValidateCompressedData(UObject* DataOwner, const FCompressedAnimSequence& AnimSeq) const
diff --git a/Engine/Source/Runtime/Engine/Classes/Animation/AnimCurveCompressionCodec.h b/Engine/Source/Runtime/Engine/Classes/Animation/AnimCurveCompressionCodec.h
--- a/Engine/Source/Runtime/Engine/Classes/Animation/AnimCurveCompressionCodec.h
+++ b/Engine/Source/Runtime/Engine/Classes/Animation/AnimCurveCompressionCodec.h
@@ -57,6 +57,19 @@ class UAnimCurveCompressionCodec : public UObject
 	/** Compresses the curve data from an animation sequence. */
 	ENGINE_API virtual bool Compress(const FCompressibleAnimData& AnimSeq, FAnimCurveCompressionResult& OutResult) PURE_VIRTUAL(UAnimCurveCompressionCodec::Compress, return false;);
 
+	/**
+	 * Estimates the peak memory usage in bytes needed to compress the given data with this codec.
+	 * This is used to make informed scheduling decisions during asset cooking. Estimates that are too low may cause
+	 * out-of-memory conditions. Estimates that are too high can unnecessarily limit the number of concurrent cook
+	 * processes.
+	 *
+	 * Values below zero indicate that no estimate has been given and high memory usage is hence assumed.
+	 *
+	 * @param AnimSequence	The Animation Sequence to estimate the memory usage for.
+	 * @returns Estimated peak memory usage in bytes. Values < 0 indicate maximum memory usage.
+	 */
+	ENGINE_API virtual int64 EstimateCompressionMemoryUsage(const UAnimSequence& AnimSequence) const;
+
 	/*
 	 * Called to generate a unique DDC key for this codec instance.
 	 * A suitable key should be generated from: the InstanceGuid, a codec version, and all relevant properties that drive the behavior.
diff --git a/Engine/Source/Runtime/Engine/Classes/Animation/AnimCurveCompressionCodec_UniformIndexable.h b/Engine/Source/Runtime/Engine/Classes/Animation/AnimCurveCompressionCodec_UniformIndexable.h
--- a/Engine/Source/Runtime/Engine/Classes/Animation/AnimCurveCompressionCodec_UniformIndexable.h
+++ b/Engine/Source/Runtime/Engine/Classes/Animation/AnimCurveCompressionCodec_UniformIndexable.h
@@ -50,6 +50,7 @@ class UAnimCurveCompressionCodec_UniformIndexable : public UAnimCurveCompression
 #if WITH_EDITORONLY_DATA
 	// UAnimCurveCompressionCodec overrides
 	ENGINE_API virtual bool Compress(const FCompressibleAnimData& AnimSeq, FAnimCurveCompressionResult& OutResult) override;
+	ENGINE_API virtual int64 EstimateCompressionMemoryUsage(const UAnimSequence& AnimSequence) const override;
 	ENGINE_API virtual void PopulateDDCKey(FArchive& Ar) override;
 #endif
 	
diff --git a/Engine/Source/Runtime/Engine/Classes/Animation/AnimCurveCompressionCodec_UniformlySampled.h b/Engine/Source/Runtime/Engine/Classes/Animation/AnimCurveCompressionCodec_UniformlySampled.h
--- a/Engine/Source/Runtime/Engine/Classes/Animation/AnimCurveCompressionCodec_UniformlySampled.h
+++ b/Engine/Source/Runtime/Engine/Classes/Animation/AnimCurveCompressionCodec_UniformlySampled.h
@@ -30,6 +30,7 @@ class UAnimCurveCompressionCodec_UniformlySampled : public UAnimCurveCompression
 #if WITH_EDITORONLY_DATA
 	// UAnimCurveCompressionCodec overrides
 	ENGINE_API virtual bool Compress(const FCompressibleAnimData& AnimSeq, FAnimCurveCompressionResult& OutResult) override;
+	ENGINE_API virtual int64 EstimateCompressionMemoryUsage(const UAnimSequence& AnimSequence) const override;
 	ENGINE_API virtual void PopulateDDCKey(FArchive& Ar) override;
 #endif
 	
diff --git a/Engine/Source/Runtime/Engine/Classes/Animation/AnimSequence.h b/Engine/Source/Runtime/Engine/Classes/Animation/AnimSequence.h
--- a/Engine/Source/Runtime/Engine/Classes/Animation/AnimSequence.h
+++ b/Engine/Source/Runtime/Engine/Classes/Animation/AnimSequence.h
@@ -596,6 +596,16 @@ public:
 	 * @return		The approximate size of raw animation data.
 	 */
 	ENGINE_API int64 GetApproxRawSize() const;
+
+	/**
+	 * @return		The approximate size of raw bone animation data.
+	 */
+	ENGINE_API int64 GetApproxBoneRawSize() const;
+
+	/**
+	 * @return		The approximate size of raw curve animation data.
+	 */
+	ENGINE_API int64 GetApproxCurveRawSize() const;
 #endif // WITH_EDITOR
 
 	/**
diff --git a/Engine/Source/Runtime/Engine/Private/Animation/AnimBoneCompressionCodec.cpp b/Engine/Source/Runtime/Engine/Private/Animation/AnimBoneCompressionCodec.cpp
--- a/Engine/Source/Runtime/Engine/Private/Animation/AnimBoneCompressionCodec.cpp
+++ b/Engine/Source/Runtime/Engine/Private/Animation/AnimBoneCompressionCodec.cpp
@@ -19,7 +19,7 @@ int64 UAnimBoneCompressionCodec::EstimateCompressionMemoryUsage(const UAnimSeque
 {
 #if WITH_EDITOR
 	// This is a conservative estimate that gives a codec enough space to create two raw copies of the input data.
-	return 2 * AnimSequence.GetApproxRawSize();
+	return 2 * int64(AnimSequence.GetApproxBoneRawSize());
 #else
 	return -1;
 #endif // WITH_EDITOR
diff --git a/Engine/Source/Runtime/Engine/Private/Animation/AnimCurveCompressionCodec.cpp b/Engine/Source/Runtime/Engine/Private/Animation/AnimCurveCompressionCodec.cpp
--- a/Engine/Source/Runtime/Engine/Private/Animation/AnimCurveCompressionCodec.cpp
+++ b/Engine/Source/Runtime/Engine/Private/Animation/AnimCurveCompressionCodec.cpp
@@ -1,6 +1,7 @@
 // Copyright Epic Games, Inc. All Rights Reserved.
 
 #include "Animation/AnimCurveCompressionCodec.h"
+#include "Animation/AnimSequence.h"
 #include "UObject/FortniteMainBranchObjectVersion.h"
 #include "UObject/FortniteReleaseBranchCustomObjectVersion.h"
 #include "UObject/Package.h"
@@ -35,6 +36,16 @@ void UAnimCurveCompressionCodec::Serialize(FArchive& Ar)
 }
 
 #if WITH_EDITORONLY_DATA
+int64 UAnimCurveCompressionCodec::EstimateCompressionMemoryUsage(const UAnimSequence& AnimSequence) const
+{
+#if WITH_EDITOR
+	// This is a conservative estimate that gives a codec enough space to create two raw copies of the input data.
+	return 2 * int64(AnimSequence.GetApproxCurveRawSize());
+#else
+	return -1;
+#endif // WITH_EDITOR
+}
+
 void UAnimCurveCompressionCodec::PopulateDDCKey(FArchive& Ar)
 {
 	// We use the UClass name to compute the DDC key to avoid two codecs with equivalent properties (e.g. none)
diff --git a/Engine/Source/Runtime/Engine/Private/Animation/AnimCurveCompressionCodec_UniformIndexable.cpp b/Engine/Source/Runtime/Engine/Private/Animation/AnimCurveCompressionCodec_UniformIndexable.cpp
--- a/Engine/Source/Runtime/Engine/Private/Animation/AnimCurveCompressionCodec_UniformIndexable.cpp
+++ b/Engine/Source/Runtime/Engine/Private/Animation/AnimCurveCompressionCodec_UniformIndexable.cpp
@@ -156,6 +156,19 @@ bool UAnimCurveCompressionCodec_UniformIndexable::Compress(const FCompressibleAn
 	return true;
 }
 
+int64 UAnimCurveCompressionCodec_UniformIndexable::EstimateCompressionMemoryUsage(const UAnimSequence& AnimSequence) const
+{
+	const int64 NumCurves = AnimSequence.GetDataModelInterface()->GetFloatCurves().Num();
+	const int64 NumSamples = AnimSequence.GetNumberOfSampledKeys();
+
+	int64 BufferSize = 0;
+	BufferSize += sizeof(int32);	// NumSamples
+	BufferSize += sizeof(float);	// SampleRate
+	BufferSize += sizeof(float) * NumCurves * NumSamples;	// Animated curve samples
+
+	return BufferSize;
+}
+
 void UAnimCurveCompressionCodec_UniformIndexable::PopulateDDCKey(FArchive& Ar)
 {
 	Super::PopulateDDCKey(Ar);
diff --git a/Engine/Source/Runtime/Engine/Private/Animation/AnimCurveCompressionCodec_UniformlySampled.cpp b/Engine/Source/Runtime/Engine/Private/Animation/AnimCurveCompressionCodec_UniformlySampled.cpp
--- a/Engine/Source/Runtime/Engine/Private/Animation/AnimCurveCompressionCodec_UniformlySampled.cpp
+++ b/Engine/Source/Runtime/Engine/Private/Animation/AnimCurveCompressionCodec_UniformlySampled.cpp
@@ -3,6 +3,7 @@
 #include "Animation/AnimCurveCompressionCodec_UniformlySampled.h"
 #include "Animation/AnimCompressionTypes.h"
 #include "Animation/AnimCurveUtils.h"
+#include "Animation/AnimSequence.h"
 
 #include UE_INLINE_GENERATED_CPP_BY_NAME(AnimCurveCompressionCodec_UniformlySampled)
 
@@ -149,6 +150,38 @@ bool UAnimCurveCompressionCodec_UniformlySampled::Compress(const FCompressibleAn
 	return true;
 }
 
+int64 UAnimCurveCompressionCodec_UniformlySampled::EstimateCompressionMemoryUsage(const UAnimSequence& AnimSequence) const
+{
+	// Same calculation as above
+	int64 NumSamples;
+	if (UseAnimSequenceSampleRate)
+	{
+		NumSamples = AnimSequence.GetNumberOfSampledKeys();
+	}
+	else
+	{
+		NumSamples = FMath::RoundToInt(AnimSequence.GetPlayLength() * SampleRate) + 1;
+	}
+
+	const int64 NumCurves = AnimSequence.GetDataModelInterface()->GetFloatCurves().Num();
+
+	const int64 ConstantBitsetSize = sizeof(uint32) * ((NumCurves + 31) / 32);
+
+	// Assume every curve is animated since they use more memory
+	const int64 NumConstantCurves = 0;
+	const int64 NumAnimatedCurves = NumCurves;
+
+	int64 BufferSize = 0;
+	BufferSize += sizeof(int32);	// NumConstantCurves
+	BufferSize += sizeof(int32);	// NumSamples
+	BufferSize += sizeof(float);	// SampleRate
+	BufferSize += ConstantBitsetSize;								// Constant curve bitset
+	BufferSize += sizeof(float) * NumConstantCurves;				// Constant curve samples
+	BufferSize += sizeof(float) * NumAnimatedCurves * NumSamples;	// Animated curve samples
+
+	return BufferSize;
+}
+
 void UAnimCurveCompressionCodec_UniformlySampled::PopulateDDCKey(FArchive& Ar)
 {
 	Super::PopulateDDCKey(Ar);
diff --git a/Engine/Source/Runtime/Engine/Private/Animation/AnimSequence.cpp b/Engine/Source/Runtime/Engine/Private/Animation/AnimSequence.cpp
--- a/Engine/Source/Runtime/Engine/Private/Animation/AnimSequence.cpp
+++ b/Engine/Source/Runtime/Engine/Private/Animation/AnimSequence.cpp
@@ -536,26 +536,44 @@ int64 UAnimSequence::GetUncompressedRawSize() const
 
 int64 UAnimSequence::GetApproxRawSize() const
 {
+	return GetApproxBoneRawSize() + GetApproxCurveRawSize();
+}
+
+int64 UAnimSequence::GetApproxBoneRawSize() const
+{
+	int64 Total = 0;
+
 	if (ShouldDataModelBeValid())
 	{
 		ValidateModel();
 
 		const int64 NumberOfTracks = DataModelInterface->GetNumBoneTracks();
 
-		int64 Total = sizeof(FRawAnimSequenceTrack) * NumberOfTracks;
+		Total = sizeof(FRawAnimSequenceTrack) * NumberOfTracks;
 		Total += int64(DataModelInterface->GetNumberOfKeys()) * sizeof(FVector3f) * NumberOfTracks;
 		Total += int64(DataModelInterface->GetNumberOfKeys()) * sizeof(FQuat4f) * NumberOfTracks;
 		Total += int64(DataModelInterface->GetNumberOfKeys()) * sizeof(FVector3f) * NumberOfTracks;
+	}
+
+	return Total;
+}
+
+int64 UAnimSequence::GetApproxCurveRawSize() const
+{
+	int64 Total = 0;
+
+	if (ShouldDataModelBeValid())
+	{
+		ValidateModel();
 
 		for (const FFloatCurve& Curve : DataModelInterface->GetCurveData().FloatCurves)
 		{
 			Total += sizeof(FFloatCurve);
 			Total += sizeof(FRichCurveKey) * int64(Curve.FloatCurve.Keys.Num());
 		}
-		return Total;
 	}
 
-	return 0;
+	return Total;
 }
 #endif // WITH_EDITOR
 
diff --git a/Engine/Source/Runtime/Engine/Private/Animation/AnimationCompressionDerivedData.cpp b/Engine/Source/Runtime/Engine/Private/Animation/AnimationCompressionDerivedData.cpp
--- a/Engine/Source/Runtime/Engine/Private/Animation/AnimationCompressionDerivedData.cpp
+++ b/Engine/Source/Runtime/Engine/Private/Animation/AnimationCompressionDerivedData.cpp
@@ -4,6 +4,8 @@
 
 #include "AnimationUtils.h"
 #include "Animation/AnimBoneCompressionSettings.h"
+#include "Animation/AnimCurveCompressionSettings.h"
+#include "Animation/AnimCurveCompressionCodec.h"
 #include "Experimental/Misc/ExecutionResource.h"
 #include "Animation/AnimationSequenceCompiler.h"
 #include "Animation/AnimCompress.h"
@@ -288,6 +290,7 @@ int64 FAnimationSequenceAsyncCacheTask::GetRequiredMemoryEstimate() const
 
 	if (const UAnimSequence* AnimSequence = WeakAnimSequence.Get())
 	{
+		// Includes bone/curves
 		const int64 AdditiveAnimSize = !AnimSequence->IsValidAdditive() ? 0 : [AnimSequence]()
 			{
 				if (const UAnimSequence* RefPoseSeq = AnimSequence->RefPoseSeq)
@@ -298,6 +301,7 @@ int64 FAnimationSequenceAsyncCacheTask::GetRequiredMemoryEstimate() const
 				return AnimSequence->GetApproxRawSize();
 			}();
 
+		// Includes bone/curves
 		RequiredMemory = AnimSequence->GetApproxRawSize() + AdditiveAnimSize;
 		if (Compression::FAnimationCompressionMemorySummaryScope::ShouldStoreCompressionResults())
 		{
@@ -319,10 +323,28 @@ int64 FAnimationSequenceAsyncCacheTask::GetRequiredMemoryEstimate() const
 						RequiredMemory = -1;
 						break;
 					}
+
 					RequiredMemory += PeakMemoryEstimate;
 				}
 			}
 		}
+
+		// Curve compression occurs after bone compression and thus executes in its shadow (memory peak wise)
+		// but it can still be the case that a sequence has no bone data while having curve data
+		if (const UAnimCurveCompressionSettings* CurveCompressionSettings = AnimSequence->CurveCompressionSettings.Get())
+		{
+			const int64 PeakMemoryEstimate = CurveCompressionSettings->Codec->EstimateCompressionMemoryUsage(*AnimSequence);
+			if (PeakMemoryEstimate < 0)
+			{
+				// Assume the worst and default to the default behavior when no estimate is given.
+				UE_LOG(LogAnimationCompression, Warning, TEXT("Got invalid memory usage estimate from codec %s for %s. This can negatively affect the time compression takes."), *CurveCompressionSettings->Codec->GetFullName(), *AnimSequence->GetFullName());
+				RequiredMemory = -1;
+			}
+			else
+			{
+				RequiredMemory += PeakMemoryEstimate;
+			}
+		}
 	}
 
 	return RequiredMemory;
