Fix const-correctness of AActors in FActorInstanceGuid

diff --git a/Engine/Source/Runtime/Engine/Private/WorldPartition/ActorInstanceGuids.cpp b/Engine/Source/Runtime/Engine/Private/WorldPartition/ActorInstanceGuids.cpp
--- a/Engine/Source/Runtime/Engine/Private/WorldPartition/ActorInstanceGuids.cpp
+++ b/Engine/Source/Runtime/Engine/Private/WorldPartition/ActorInstanceGuids.cpp
@@ -112,7 +112,7 @@ FGuid FActorInstanceGuid::GetLevelInstanceGuid(ULevel* Level)
 }
 
 #if WITH_EDITOR
-void FActorInstanceGuid::InitializeFrom(AActor& InActor)
+void FActorInstanceGuid::InitializeFrom(const AActor& InActor)
 {
 	ActorGuid = InActor.GetActorGuid();
 	ActorInstanceGuid = InActor.GetActorInstanceGuid();	
@@ -130,7 +130,7 @@ bool FActorInstanceGuid::IsDefault()
 	return !ActorGuid.IsValid() && !ActorInstanceGuid.IsValid();
 }
 
-FActorInstanceGuid FActorInstanceGuid::GetActorGuids(AActor& InActor)
+FActorInstanceGuid FActorInstanceGuid::GetActorGuids(const AActor& InActor)
 {
 	FActorInstanceGuid Guids;
 
@@ -165,14 +165,14 @@ FActorInstanceGuid FActorInstanceGuid::GetActorGuids(AActor& InActor)
 	return Guids;
 }
 
-FGuid FActorInstanceGuid::GetActorInstanceGuid(AActor& InActor)
+FGuid FActorInstanceGuid::GetActorInstanceGuid(const AActor& InActor)
 {
 	FActorInstanceGuid Guids = GetActorGuids(InActor);
 
 	return Guids.ActorInstanceGuid;
 }
 
-void FActorInstanceGuid::ReleaseActorInstanceGuid(AActor& InActor)
+void FActorInstanceGuid::ReleaseActorInstanceGuid(const AActor& InActor)
 {
 	// Only loaded objects will have those annotations
 	if (InActor.HasAnyFlags(RF_WasLoaded))
diff --git a/Engine/Source/Runtime/Engine/Public/WorldPartition/ActorInstanceGuids.h b/Engine/Source/Runtime/Engine/Public/WorldPartition/ActorInstanceGuids.h
--- a/Engine/Source/Runtime/Engine/Public/WorldPartition/ActorInstanceGuids.h
+++ b/Engine/Source/Runtime/Engine/Public/WorldPartition/ActorInstanceGuids.h
@@ -15,16 +15,16 @@ struct FActorInstanceGuid
 
 	bool IsDefault();
 #if WITH_EDITOR
-	void InitializeFrom(AActor& InActor);
+	void InitializeFrom(const AActor& InActor);
 #endif
 
 	ENGINE_API static void ReleaseLevelInstanceGuid(ULevel* Level);
 	ENGINE_API static void SetLevelInstanceGuid(ULevel* Level, ULevel* OwnerLevel, const FGuid& Guid, const FGuid& ResolvedGuid = FGuid());
 	ENGINE_API static FGuid GetLevelInstanceGuid(ULevel* Level);
 
-	ENGINE_API static FActorInstanceGuid GetActorGuids(AActor& InActor);
-	ENGINE_API static FGuid GetActorInstanceGuid(AActor& InActor);
-	ENGINE_API static void ReleaseActorInstanceGuid(AActor& InActor);
+	ENGINE_API static FActorInstanceGuid GetActorGuids(const AActor& InActor);
+	ENGINE_API static FGuid GetActorInstanceGuid(const AActor& InActor);
+	ENGINE_API static void ReleaseActorInstanceGuid(const AActor& InActor);
 	ENGINE_API static void Serialize(FArchive& Ar, AActor& InActor);
 };
 
