World::Rename will now route PreRenameFromRoot and PostRenameFromRoot to its persistent level, which in turn will route it to its actors.
Level entity handle these calls to rename its owned entities in case of world renaming, which happens during save as.
This is required for saving a newly created world and placing entities saved in their own file to work.
This is also required for scene graph unit tests which creates new maps, place entities, saves and reload.

diff --git a/Engine/Source/Runtime/Engine/Classes/Engine/Level.h b/Engine/Source/Runtime/Engine/Classes/Engine/Level.h
--- a/Engine/Source/Runtime/Engine/Classes/Engine/Level.h
+++ b/Engine/Source/Runtime/Engine/Classes/Engine/Level.h
@@ -43,6 +43,7 @@ class SNotificationItem;
 class UActorFolder;
 class IWorldPartitionCell;
 class UWorldPartitionRuntimeCell;
+struct FWorldRenameFromRootContext;
 struct FFolder;
 
 #if WITH_EDITOR
@@ -925,7 +926,7 @@ public:
 	virtual void FinishDestroy() override;
 	ENGINE_API virtual UWorld* GetWorld() const override final;
 
-#if	WITH_EDITOR
+#if WITH_EDITOR
 	virtual void PreEditUndo() override;
 	virtual void PostEditUndo() override;	
 	virtual void PostEditChangeProperty(FPropertyChangedEvent& PropertyChangedEvent) override;
@@ -942,8 +943,11 @@ public:
 
 	void PreSaveFromRoot(FObjectPreSaveRootContext ObjectSaveContext);
 	void PostSaveFromRoot(FObjectPostSaveRootContext ObjectSaveContext);
-	
-#if	WITH_EDITOR
+
+#if WITH_EDITOR
+	void PreRenameFromRoot(const FWorldRenameFromRootContext& RenameContext);
+	void PostRenameFromRoot(const FWorldRenameFromRootContext& RenameContext);
+
 	//~Begin ITextureStreamingContainer Interface.
 	virtual void InitializeTextureStreamingContainer(uint32 InPackedTextureStreamingQualityLevelFeatureLevel) override;
 	virtual uint16 RegisterStreamableTexture(UTexture* InTexture) override;
diff --git a/Engine/Source/Runtime/Engine/Classes/Engine/World.h b/Engine/Source/Runtime/Engine/Classes/Engine/World.h
--- a/Engine/Source/Runtime/Engine/Classes/Engine/World.h
+++ b/Engine/Source/Runtime/Engine/Classes/Engine/World.h
@@ -578,6 +578,18 @@ struct ENGINE_API FWorldAsyncTraceState
 };
 
 #if WITH_EDITOR
+struct ENGINE_API FWorldRenameFromRootContext
+{
+	/** New name for PreRenameFromRoot, old name for PostRenameFromRoot. */
+	FName Name;
+
+	/** New outer for PreRenameFromRoot, old outer for PostRenameFromRoot. */
+	UObject* Outer = nullptr;
+
+	/** Rename flags. */
+	ERenameFlags Flags = REN_None;
+};
+
 /* FAsyncPreRegisterDDCRequest - info about an async DDC request that we're going to wait on before registering components */
 class ENGINE_API FAsyncPreRegisterDDCRequest
 {
diff --git a/Engine/Source/Runtime/Engine/Classes/GameFramework/Actor.h b/Engine/Source/Runtime/Engine/Classes/GameFramework/Actor.h
--- a/Engine/Source/Runtime/Engine/Classes/GameFramework/Actor.h
+++ b/Engine/Source/Runtime/Engine/Classes/GameFramework/Actor.h
@@ -40,6 +40,7 @@ struct FNetViewer;
 struct FNetworkObjectInfo;
 class FActorTransactionAnnotation;
 class FComponentInstanceDataCache;
+struct FWorldRenameFromRootContext;
 class UDEPRECATED_DataLayer;
 class UDataLayerAsset;
 class UExternalDataLayerAsset;
@@ -2411,6 +2412,10 @@ public:
 #endif // WITH_EDITOR
 	//~ End UObject Interface
 
+#if WITH_EDITOR
+	ENGINE_API virtual void PreRenameFromRoot(const FWorldRenameFromRootContext& RenameContext) {}
+	ENGINE_API virtual void PostRenameFromRoot(const FWorldRenameFromRootContext& RenameContext) {}
+#endif
 	ENGINE_API virtual void PreDuplicateFromRoot(FObjectDuplicationParameters& DupParams) {}
 	ENGINE_API virtual void PreSaveFromRoot(FObjectPreSaveRootContext ObjectSaveContext);
 	ENGINE_API virtual void PostSaveFromRoot(FObjectPostSaveRootContext ObjectSaveContext);
diff --git a/Engine/Source/Runtime/Engine/Private/Level.cpp b/Engine/Source/Runtime/Engine/Private/Level.cpp
--- a/Engine/Source/Runtime/Engine/Private/Level.cpp
+++ b/Engine/Source/Runtime/Engine/Private/Level.cpp
@@ -496,6 +496,30 @@ void ULevel::AddReferencedObjects(UObject* InThis, FReferenceCollector& Collecto
 	Super::AddReferencedObjects( This, Collector );
 }
 
+#if	WITH_EDITOR
+void ULevel::PreRenameFromRoot(const FWorldRenameFromRootContext& RenameContext)
+{
+	for (AActor* Actor : Actors)
+	{
+		if (Actor)
+		{
+			Actor->PreRenameFromRoot(RenameContext);
+		}
+	}
+}
+
+void ULevel::PostRenameFromRoot(const FWorldRenameFromRootContext& RenameContext)
+{
+	for (AActor* Actor : Actors)
+	{
+		if (Actor)
+		{
+			Actor->PostRenameFromRoot(RenameContext);
+		}
+	}
+}
+#endif
+
 void ULevel::PreSaveFromRoot(FObjectPreSaveRootContext ObjectSaveContext)
 {
 	for (AActor* Actor : Actors)
diff --git a/Engine/Source/Runtime/Engine/Private/World.cpp b/Engine/Source/Runtime/Engine/Private/World.cpp
--- a/Engine/Source/Runtime/Engine/Private/World.cpp
+++ b/Engine/Source/Runtime/Engine/Private/World.cpp
@@ -916,12 +916,25 @@ bool UWorld::Rename(const TCHAR* InName, UObject* NewOuter, ERenameFlags Flags)
 		return false;
 	}
 
+	const FWorldRenameFromRootContext PreRenameFromRootContext({ InName, NewOuter, Flags });
+	const FWorldRenameFromRootContext PostRenameFromRootContext({ GetFName(), GetOuter(), Flags });
+
+	if (!bTestRename)
+	{
+		PersistentLevel->PreRenameFromRoot(PreRenameFromRootContext);
+	}
+
 	// Rename the world itself
 	if ( !Super::Rename(InName, NewOuter, Flags) )
 	{
 		return false;
 	}
 
+	if (!bTestRename)
+	{
+		PersistentLevel->PostRenameFromRoot(PostRenameFromRootContext);
+	}
+
 	// We're moving the world to a new package, rename UObjects which are map data but don't have the UWorld in their Outer chain.  There are two cases:
 	// 1) legacy lightmap textures and MapBuildData object will be in the same package as the UWorld.  We need to move these to the new world package.
 	// 2) MapBuildData will be in a separate package with lightmap textures underneath it.  We need to move these to an appropriate build data package.
