Make UInputDelegateBinding::SupportsInputDelegate always return true, instead of only for blueprint generated classes.

There may be times where you have a Native C++ class with non-native (blueprint added) subobjects or components. In this case, you would need to bind their dynamic input delegates. This makes the behavior more deterministic and less error prone.

This is behind a cvar "Input.bAlwaysAllowInputDelegateBindings" , which is defaulted to ON. This CVar will be removed in a future release, but is there as an intermediate step for upgrading projects just in case a fallback is needed.

diff --git a/Engine/Plugins/AI/MLAdapter/Source/MLAdapter/Private/Sensors/MLAdapterSensor_EnhancedInput.cpp b/Engine/Plugins/AI/MLAdapter/Source/MLAdapter/Private/Sensors/MLAdapterSensor_EnhancedInput.cpp
--- a/Engine/Plugins/AI/MLAdapter/Source/MLAdapter/Private/Sensors/MLAdapterSensor_EnhancedInput.cpp
+++ b/Engine/Plugins/AI/MLAdapter/Source/MLAdapter/Private/Sensors/MLAdapterSensor_EnhancedInput.cpp
@@ -97,7 +97,9 @@ void UMLAdapterSensor_EnhancedInput::OnAvatarSet(AActor* Avatar)
 	if (Pawn != nullptr)
 	{
 		UClass* PawnClass = Pawn->GetClass();
+	PRAGMA_DISABLE_DEPRECATION_WARNINGS
 		if (UInputDelegateBinding::SupportsInputDelegate(PawnClass))
+	PRAGMA_ENABLE_DEPRECATION_WARNINGS
 		{
 			InputComponent = NewObject<UEnhancedInputComponent>(Pawn, UInputSettings::GetDefaultInputComponentClass());
 			InputComponent->RegisterComponent();
diff --git a/Engine/Source/Runtime/Engine/Classes/Engine/InputDelegateBinding.h b/Engine/Source/Runtime/Engine/Classes/Engine/InputDelegateBinding.h
--- a/Engine/Source/Runtime/Engine/Classes/Engine/InputDelegateBinding.h
+++ b/Engine/Source/Runtime/Engine/Classes/Engine/InputDelegateBinding.h
@@ -45,6 +45,7 @@ class UInputDelegateBinding : public UDynamicBlueprintBinding
 	virtual void BindToInputComponent(UInputComponent* InputComponent, UObject* ObjectToBindTo) const { };
 
 	/** Returns true if the given class supports input binding delegates (i.e. it is a BP generated class) */
+	UE_DEPRECATED(5.6, "SupportsInputDelegate is deprecated. All UClasses will be considered valid for input delegate support.")
 	static ENGINE_API bool SupportsInputDelegate(const UClass* InClass);
 
 	/**
diff --git a/Engine/Source/Runtime/Engine/Private/InputDelegateBinding.cpp b/Engine/Source/Runtime/Engine/Private/InputDelegateBinding.cpp
--- a/Engine/Source/Runtime/Engine/Private/InputDelegateBinding.cpp
+++ b/Engine/Source/Runtime/Engine/Private/InputDelegateBinding.cpp
@@ -4,6 +4,12 @@
 #include "Engine/BlueprintGeneratedClass.h"
 #include "GameFramework/Actor.h"
 #include "Components/InputComponent.h"
+#include "HAL/IConsoleManager.h"
+
+static bool bAlwaysAllowInputDelegateBindings = true;
+static FAutoConsoleVariableRef CVarAlwaysAllowInputDelegateBindings(TEXT("Input.bAlwaysAllowInputDelegateBindings"),
+	bAlwaysAllowInputDelegateBindings,
+	TEXT("If true then UInputDelegateBinding::SupportsInputDelegate will always return true. Otherwise, only blueprint generated class will support dynamic input delegates"));
 
 #include UE_INLINE_GENERATED_CPP_BY_NAME(InputDelegateBinding)
 
@@ -19,14 +25,21 @@ UInputDelegateBinding::UInputDelegateBinding(const FObjectInitializer& ObjectIni
 	}
 }
 
+PRAGMA_DISABLE_DEPRECATION_WARNINGS
 bool UInputDelegateBinding::SupportsInputDelegate(const UClass* InClass)
 {
-	return !!Cast<UBlueprintGeneratedClass>(InClass);
+	// We want to treat every class as supporting input delegate binding no matter what,
+	// because even with a native UClass it can still have blueprint generated subobjects/components
+	// who need to be dynamically bound.
+	return bAlwaysAllowInputDelegateBindings || Cast<UBlueprintGeneratedClass>(InClass);
 }
+PRAGMA_ENABLE_DEPRECATION_WARNINGS
 
 void UInputDelegateBinding::BindInputDelegates(const UClass* InClass, UInputComponent* InputComponent, UObject* InObjectToBindTo /* = nullptr */)
 {
+PRAGMA_DISABLE_DEPRECATION_WARNINGS
 	if (InClass && InputComponent && SupportsInputDelegate(InClass))
+PRAGMA_ENABLE_DEPRECATION_WARNINGS
 	{
 		ensureMsgf(InputComponent, TEXT("Attempting to bind input delegates to an invalid Input Component!"));
 		
@@ -57,7 +70,9 @@ void UInputDelegateBinding::BindInputDelegatesWithSubojects(AActor* InActor, UIn
 
 	const UClass* ActorClass = InActor ? InActor->GetClass() : nullptr;
 
+PRAGMA_DISABLE_DEPRECATION_WARNINGS
 	if (ActorClass && InputComponent && SupportsInputDelegate(ActorClass))
+PRAGMA_ENABLE_DEPRECATION_WARNINGS
 	{
 		// Bind any input delegates on the base actor class
 		UInputDelegateBinding::BindInputDelegates(ActorClass, InputComponent, InputComponent->GetOwner());
diff --git a/Engine/Source/Runtime/Engine/Private/LevelScriptActor.cpp b/Engine/Source/Runtime/Engine/Private/LevelScriptActor.cpp
--- a/Engine/Source/Runtime/Engine/Private/LevelScriptActor.cpp
+++ b/Engine/Source/Runtime/Engine/Private/LevelScriptActor.cpp
@@ -48,7 +48,9 @@ ALevelScriptActor::ALevelScriptActor(const FObjectInitializer& ObjectInitializer
 
 void ALevelScriptActor::PreInitializeComponents()
 {
+PRAGMA_DISABLE_DEPRECATION_WARNINGS
 	if (UInputDelegateBinding::SupportsInputDelegate(GetClass()) && !InputComponent)
+PRAGMA_ENABLE_DEPRECATION_WARNINGS
 	{
 		// create an InputComponent object so that the level script actor can bind key events
 		InputComponent = NewObject<UInputComponent>(this, UInputSettings::GetDefaultInputComponentClass());
diff --git a/Engine/Source/Runtime/Engine/Private/Pawn.cpp b/Engine/Source/Runtime/Engine/Private/Pawn.cpp
--- a/Engine/Source/Runtime/Engine/Private/Pawn.cpp
+++ b/Engine/Source/Runtime/Engine/Private/Pawn.cpp
@@ -489,7 +489,9 @@ void APawn::PawnClientRestart()
 			{
 				SetupPlayerInputComponent(InputComponent);
 				InputComponent->RegisterComponent();
+			PRAGMA_DISABLE_DEPRECATION_WARNINGS
 				if (UInputDelegateBinding::SupportsInputDelegate(GetClass()))
+			PRAGMA_ENABLE_DEPRECATION_WARNINGS	
 				{
 					InputComponent->bBlockInput = bBlockInput;
 					UInputDelegateBinding::BindInputDelegatesWithSubojects(this, InputComponent);
diff --git a/Engine/Source/Runtime/Engine/Private/PlayerController.cpp b/Engine/Source/Runtime/Engine/Private/PlayerController.cpp
--- a/Engine/Source/Runtime/Engine/Private/PlayerController.cpp
+++ b/Engine/Source/Runtime/Engine/Private/PlayerController.cpp
@@ -2699,7 +2699,9 @@ void APlayerController::SetupInputComponent()
 		InputComponent->RegisterComponent();
 	}
 
+PRAGMA_DISABLE_DEPRECATION_WARNINGS
 	if (UInputDelegateBinding::SupportsInputDelegate(GetClass()))
+PRAGMA_ENABLE_DEPRECATION_WARNINGS		
 	{
 		InputComponent->bBlockInput = bBlockInput;		
 		UInputDelegateBinding::BindInputDelegatesWithSubojects(this, InputComponent);
diff --git a/Engine/Source/Runtime/GameplayDebugger/Private/GameplayDebuggerPlayerManager.cpp b/Engine/Source/Runtime/GameplayDebugger/Private/GameplayDebuggerPlayerManager.cpp
--- a/Engine/Source/Runtime/GameplayDebugger/Private/GameplayDebuggerPlayerManager.cpp
+++ b/Engine/Source/Runtime/GameplayDebugger/Private/GameplayDebuggerPlayerManager.cpp
@@ -52,7 +52,9 @@ void AGameplayDebuggerPlayerManager::PostInitProperties()
 			InputComponent->RegisterComponent();
 		}
 
+	PRAGMA_DISABLE_DEPRECATION_WARNINGS
 		if (UInputDelegateBinding::SupportsInputDelegate(GetClass()))
+	PRAGMA_ENABLE_DEPRECATION_WARNINGS
 		{
 			UInputDelegateBinding::BindInputDelegates(GetClass(), InputComponent);
 		}
