Clamp texture mip batch count as a speculative fix for a crash.

#rnx

diff --git a/Engine/Source/Runtime/Engine/Private/Streaming/Texture2DStreamIn_IO.cpp b/Engine/Source/Runtime/Engine/Private/Streaming/Texture2DStreamIn_IO.cpp
--- a/Engine/Source/Runtime/Engine/Private/Streaming/Texture2DStreamIn_IO.cpp
+++ b/Engine/Source/Runtime/Engine/Private/Streaming/Texture2DStreamIn_IO.cpp
@@ -53,7 +53,7 @@ static void ValidateMipBulkDataSize(const UTexture2D& Texture, int32 MipSizeX, i
 void FTexture2DStreamIn_IO::SetIORequests(const FContext& Context)
 {
 	int32 InFlightBulkSize = 0;
-	const int32 BatchCount = CurrentFirstLODIdx - PendingFirstLODIdx;
+	const int32 BatchCount = FMath::Min(CurrentFirstLODIdx - PendingFirstLODIdx, static_cast<int32>(ResourceState.MaxNumLODs));
 	FBulkDataBatchRequest::FBatchBuilder Batch = FBulkDataBatchRequest::NewBatch(BatchCount);
 	for (int32 MipIndex = PendingFirstLODIdx; MipIndex < CurrentFirstLODIdx && !IsCancelled(); ++MipIndex)
 	{
