Expanded allowed UI limits to CameraShutterSpeed [1/s] to allow for realistic values used in photography.

#rnx

diff --git a/Engine/Source/Runtime/Engine/Classes/Engine/Scene.h b/Engine/Source/Runtime/Engine/Classes/Engine/Scene.h
--- a/Engine/Source/Runtime/Engine/Classes/Engine/Scene.h
+++ b/Engine/Source/Runtime/Engine/Classes/Engine/Scene.h
@@ -1741,8 +1741,8 @@ struct FPostProcessSettings
 	UPROPERTY(EditAnywhere, BlueprintReadWrite, Category="Rendering Features|Ambient Cubemap", meta=(DisplayName = "Cubemap Texture"))
 	TObjectPtr<class UTextureCube> AmbientCubemap;
 
-	/** The camera shutter in seconds.*/
-	UPROPERTY(EditAnywhere, BlueprintReadWrite, Category="Lens|Camera", meta=(ClampMin = "1.0", ClampMax = "2000.0", editcondition = "bOverride_CameraShutterSpeed", DisplayName = "Shutter Speed (1/s)"))
+	/** The camera shutter in 1/seconds.*/
+	UPROPERTY(EditAnywhere, BlueprintReadWrite, Category="Lens|Camera", meta=(ClampMin = "0.01", ClampMax = "8000.0", editcondition = "bOverride_CameraShutterSpeed", DisplayName = "Shutter Speed (1/s)"))
     float CameraShutterSpeed;
 
 	/** The camera sensor sensitivity in ISO.*/
