Substrate - added option to disable roughness tracking on materials.

#tests ps dx12 editor
#rnx

diff --git a/Engine/Shaders/Private/Definitions.usf b/Engine/Shaders/Private/Definitions.usf
--- a/Engine/Shaders/Private/Definitions.usf
+++ b/Engine/Shaders/Private/Definitions.usf
@@ -84,6 +84,10 @@
 #define SUBSTRATE_OPTIMIZED_UNLIT							0
 #endif
 
+#ifndef SUBSTRATE_MATERIAL_ROUGHNESS_TRACKING
+#define SUBSTRATE_MATERIAL_ROUGHNESS_TRACKING				1
+#endif
+
 #ifndef SUBSTRATE_NORMAL_QUALITY
 #define SUBSTRATE_NORMAL_QUALITY							0
 #define SUBSTRATE_TOP_LAYER_TYPE                            uint
diff --git a/Engine/Shaders/Private/Substrate/SubstrateTree.ush b/Engine/Shaders/Private/Substrate/SubstrateTree.ush
--- a/Engine/Shaders/Private/Substrate/SubstrateTree.ush
+++ b/Engine/Shaders/Private/Substrate/SubstrateTree.ush
@@ -396,6 +396,7 @@ void FSubstrateTree::UpdateAllBSDFWithBottomUpOperatorVisit_Vertical(
 	// If the BSDF is from the bottom, we reduce its contribution according to the top layer coverage only.
 	CurrentBSDF.TopLayerDataWeight *= bBSDFComesFromTopLayer ? 1.0 : saturate(1.0f - Op.VerticalTop_Coverage);
 
+#if SUBSTRATE_MATERIAL_ROUGHNESS_TRACKING
 	if (!bBSDFComesFromTopLayer)	// Only apply to BSDF from the bottom layer
 	{
 		// We propagate the bottom layer lob up after the top layer top layer simulate the fact that light rays are potentially heavily scattered/refracted by the top layer.
@@ -418,6 +419,7 @@ void FSubstrateTree::UpdateAllBSDFWithBottomUpOperatorVisit_Vertical(
 
 		CurrentBSDF.SubstrateSetBSDFRoughness(SubstrateLobeVarianceToRoughness(RefractedLobe.Sigma));
 	}
+#endif // SUBSTRATE_MATERIAL_ROUGHNESS_TRACKING
 
 #endif // SUBSTRATE_INLINE_SHADING
 }
diff --git a/Engine/Shaders/Public/ShaderVersion.ush b/Engine/Shaders/Public/ShaderVersion.ush
--- a/Engine/Shaders/Public/ShaderVersion.ush
+++ b/Engine/Shaders/Public/ShaderVersion.ush
@@ -3,5 +3,5 @@
 // Each time the console command is executed it generates a new GUID. As this file is included
 // in Platform.ush (which should be included in any shader) it allows to invalidate the shader DDC.
 // 
-#pragma message("UESHADERMETADATA_VERSION 58C33AF1-D8A8-455A-BA24-22BF5CAF6AFC")
+#pragma message("UESHADERMETADATA_VERSION 75B6BD02-19BE-4FB4-B0B6-671BEBE9A2A1")
 
diff --git a/Engine/Source/Runtime/Core/Private/UObject/FortniteMainBranchObjectVersion.cpp b/Engine/Source/Runtime/Core/Private/UObject/FortniteMainBranchObjectVersion.cpp
--- a/Engine/Source/Runtime/Core/Private/UObject/FortniteMainBranchObjectVersion.cpp
+++ b/Engine/Source/Runtime/Core/Private/UObject/FortniteMainBranchObjectVersion.cpp
@@ -14,6 +14,6 @@ TMap<FGuid, FGuid> FFortniteMainBranchObjectVersion::GetSystemGuids()
 	SystemGuids.Add(DevGuids.Niagara_LatestScriptCompileVersion, FGuid("393081597FA7234C92F00869DE262D73"));
 	SystemGuids.Add(DevGuids.SkeletalMeshDerivedDataVersion, FGuid("0829C9BAB43B4B25A63747FDA68E7197"));
 	SystemGuids.Add(DevGuids.STATICMESH_DERIVEDDATA_VER, FGuid("656E338C78C34641A5D3BFAE303997AF"));
-	SystemGuids.Add(DevGuids.MaterialTranslationDDCVersion, FGuid("8CF065BF76124292A1A0546D532BEBDF"));
+	SystemGuids.Add(DevGuids.MaterialTranslationDDCVersion, FGuid("5DFD770145674B3AAFA544176CF7372B"));
 	return SystemGuids;
 }
diff --git a/Engine/Source/Runtime/Engine/Private/Materials/HLSLMaterialTranslator.cpp b/Engine/Source/Runtime/Engine/Private/Materials/HLSLMaterialTranslator.cpp
--- a/Engine/Source/Runtime/Engine/Private/Materials/HLSLMaterialTranslator.cpp
+++ b/Engine/Source/Runtime/Engine/Private/Materials/HLSLMaterialTranslator.cpp
@@ -335,6 +335,7 @@ struct FHLSLMaterialTranslator::FEnvironmentDefines
 	bool bSubstrateLegacyIrisTangent;
 	uint8 NumCustomizedUVs;
 	bool bTextureSampleDebug;
+	bool bSubstrateRoughnessTracking;
 	TArray<TPair<FString, int32>> SubstrateDefines;
 	TArray<TObjectPtr<UMaterialParameterCollection>> ParameterCollections;
 
@@ -436,6 +437,7 @@ struct FHLSLMaterialTranslator::FEnvironmentDefines
 		Ar << NumCustomizedUVs;
 		Ar << bTextureSampleDebug;
 		Ar << SubstrateDefines;
+		Ar << bSubstrateRoughnessTracking;
 
 		TArray<UObject*> ParameterCollectionObjects;
 		if (!Ar.IsLoading())
@@ -2916,11 +2918,13 @@ void FHLSLMaterialTranslator::GetMaterialEnvironment(EShaderPlatform InPlatform,
 	if (EnvironmentDefines->bMaterialIsSubstrate)
 	{
 		OutEnvironment.SetDefine(TEXT("SUBSTRATE_USES_CONVERSION_FROM_LEGACY"), EnvironmentDefines->bSubstrateUsesConversionFromLegacy);
+		OutEnvironment.SetDefine(TEXT("SUBSTRATE_OPTIMIZED_UNLIT"), EnvironmentDefines->bSubstrateOptimizedUnlit);
+
 		OutEnvironment.SetDefine(TEXT("SUBSTRATE_MATERIAL_OUTPUT_OPAQUE_ROUGH_REFRACTIONS"), EnvironmentDefines->bSubstrateMaterialOutputOpaqueRoughRefractions);
 		OutEnvironment.SetDefine(TEXT("SUBSTRATE_MATERIAL_EXPORT_TYPE"), EnvironmentDefines->SubstrateMaterialExportType);
 		OutEnvironment.SetDefine(TEXT("SUBSTRATE_MATERIAL_EXPORT_CONTEXT"), EnvironmentDefines->SubstrateMaterialExportContext);
 		OutEnvironment.SetDefine(TEXT("SUBSTRATE_MATERIAL_EXPORT_LEGACY_BLEND_MODE"), EnvironmentDefines->SubstrateMaterialExportLegacyBlendMode);
-		OutEnvironment.SetDefine(TEXT("SUBSTRATE_OPTIMIZED_UNLIT"), EnvironmentDefines->bSubstrateOptimizedUnlit);
+		OutEnvironment.SetDefine(TEXT("SUBSTRATE_MATERIAL_ROUGHNESS_TRACKING"), EnvironmentDefines->bSubstrateRoughnessTracking);
 
 		OutEnvironment.SetDefine(TEXT("SUBSTRATE_LEGACY_IRIS_NORMAL"), EnvironmentDefines->bSubstrateLegacyIrisNormal);
 		OutEnvironment.SetDefine(TEXT("SUBSTRATE_LEGACY_IRIS_TANGENT"), EnvironmentDefines->bSubstrateLegacyIrisTangent);
@@ -16026,6 +16030,8 @@ void FHLSLMaterialTranslator::PrepareEnvironmentDefines()
 		// Unlit cannot be combined with other BSDF so we can simply pick the default strata context
 		EnvironmentDefines->bSubstrateOptimizedUnlit = SubstrateCompilationContext[ESubstrateCompilationContext::SCC_Default].bSubstrateMaterialIsUnlitNode;
 
+		EnvironmentDefines->bSubstrateRoughnessTracking = Material->HasSubstrateRoughnessTracking();
+
 		{
 			// For now, the fully simplified mode is used for Lumen or anything else supported inlined evaluation. The export is only valid for the default case.
 			// STRATA_TODO: generate an export for the different context (need to generate two export functions: the default one and the FullSimplification one)
diff --git a/Engine/Source/Runtime/Engine/Private/Materials/Material.cpp b/Engine/Source/Runtime/Engine/Private/Materials/Material.cpp
--- a/Engine/Source/Runtime/Engine/Private/Materials/Material.cpp
+++ b/Engine/Source/Runtime/Engine/Private/Materials/Material.cpp
@@ -1096,6 +1096,7 @@ void UMaterial::SetInitialValues()
 	MaterialDecalResponse = MDR_ColorNormalRoughness;
 
 	SubstrateCompilationConfig = FSubstrateCompilationConfig();
+	bSubstrateRoughnessTracking = true;
 
 	bAllowDevelopmentShaderCompile = true;
 	bIsMaterialEditorStatsMaterial = false;
@@ -7388,6 +7389,11 @@ bool UMaterial::HasPixelAnimation() const
 	return bHasPixelAnimation;
 }
 
+bool UMaterial::HasSubstrateRoughnessTracking() const
+{
+	return bSubstrateRoughnessTracking;
+}
+
 void UMaterial::SetShadingModel(EMaterialShadingModel NewModel)
 {
 	ensure(ShadingModel < MSM_NUM);
diff --git a/Engine/Source/Runtime/Engine/Private/Materials/MaterialInstance.cpp b/Engine/Source/Runtime/Engine/Private/Materials/MaterialInstance.cpp
--- a/Engine/Source/Runtime/Engine/Private/Materials/MaterialInstance.cpp
+++ b/Engine/Source/Runtime/Engine/Private/Materials/MaterialInstance.cpp
@@ -4979,6 +4979,11 @@ bool UMaterialInstance::IsTessellationEnabled() const
 	return bEnableTessellation;
 }
 
+bool UMaterialInstance::HasSubstrateRoughnessTracking() const
+{
+	return Parent ? Parent->HasSubstrateRoughnessTracking() : true;
+}
+
 /** Checks to see if an input property should be active, based on the state of the material */
 bool UMaterialInstance::IsPropertyActive(EMaterialProperty InProperty) const
 {
diff --git a/Engine/Source/Runtime/Engine/Private/Materials/MaterialInterface.cpp b/Engine/Source/Runtime/Engine/Private/Materials/MaterialInterface.cpp
--- a/Engine/Source/Runtime/Engine/Private/Materials/MaterialInterface.cpp
+++ b/Engine/Source/Runtime/Engine/Private/Materials/MaterialInterface.cpp
@@ -1533,6 +1533,11 @@ bool UMaterialInterface::IsTessellationEnabled() const
 	return false;
 }
 
+bool UMaterialInterface::HasSubstrateRoughnessTracking() const
+{
+	return true;
+}
+
 void UMaterialInterface::SetFeatureLevelToCompile(ERHIFeatureLevel::Type FeatureLevel, bool bShouldCompile)
 {
 	uint32 FeatureLevelBit = (1 << FeatureLevel);
diff --git a/Engine/Source/Runtime/Engine/Private/Materials/MaterialShared.cpp b/Engine/Source/Runtime/Engine/Private/Materials/MaterialShared.cpp
--- a/Engine/Source/Runtime/Engine/Private/Materials/MaterialShared.cpp
+++ b/Engine/Source/Runtime/Engine/Private/Materials/MaterialShared.cpp
@@ -2307,6 +2307,11 @@ int32 FMaterialResource::GetNeuralProfileId() const
 	return GetMaterialDomain() == MD_PostProcess ? Material->NeuralProfileId : INDEX_NONE;
 }
 
+bool FMaterialResource::HasSubstrateRoughnessTracking() const
+{ 
+	return Material->HasSubstrateRoughnessTracking();
+}
+
 uint32 FMaterialResource::GetStencilCompare() const
 {
 	return GetMaterialDomain() == MD_PostProcess ? uint32(Material->StencilCompare.GetValue()) : 0;
diff --git a/Engine/Source/Runtime/Engine/Public/MaterialShared.h b/Engine/Source/Runtime/Engine/Public/MaterialShared.h
--- a/Engine/Source/Runtime/Engine/Public/MaterialShared.h
+++ b/Engine/Source/Runtime/Engine/Public/MaterialShared.h
@@ -2241,6 +2241,7 @@ public:
 	virtual bool ShouldAlwaysEvaluateWorldPositionOffset() const { return false; }
 	virtual bool ComputeFogPerPixel() const { return false; }
 	virtual bool IsSky() const { return false; }
+	virtual bool HasSubstrateRoughnessTracking() const { return true; }
 	virtual FString GetFriendlyName() const = 0;
 	/** Similar to GetFriendlyName, but but avoids historical behavior of the former, returning the exact asset name for material instances instead of just the material. */
 	virtual FString GetAssetName() const { return GetFriendlyName(); }
@@ -3065,6 +3066,7 @@ public:
 	ENGINE_API virtual uint16 GetPreshaderGap() const override;
 	ENGINE_API virtual UMaterialInterface* GetMaterialInterface() const override;
 	ENGINE_API virtual int32 GetNeuralProfileId() const override;
+	ENGINE_API virtual bool HasSubstrateRoughnessTracking() const override;
 
 	/**
 	 * Should shaders compiled for this material be saved to disk?
diff --git a/Engine/Source/Runtime/Engine/Public/Materials/Material.h b/Engine/Source/Runtime/Engine/Public/Materials/Material.h
--- a/Engine/Source/Runtime/Engine/Public/Materials/Material.h
+++ b/Engine/Source/Runtime/Engine/Public/Materials/Material.h
@@ -664,6 +664,12 @@ public:
 	UPROPERTY(duplicatetransient)
 	uint8 bUsedAsSpecialEngineMaterial : 1;
 
+	/**
+	 * This can be used to disable roughness tracking. When this is done, top layers roughness will no longer impact bottom layer roughness.
+	 */
+	UPROPERTY(EditAnywhere, BlueprintReadOnly, meta = (DisplayName = "Roughness Tracking"), Category = Substrate)
+	uint8 bSubstrateRoughnessTracking : 1;
+
 	/** 
 	 * Indicates that the material and its instances can be used with skeletal meshes.  
 	 * This will result in the shaders required to support skeletal meshes being compiled which will increase shader compile time and memory usage.
@@ -1253,6 +1259,7 @@ public:
 	ENGINE_API virtual bool WritesToRuntimeVirtualTexture() const override;
 	ENGINE_API virtual bool HasMeshPaintTexture() const override;
 	ENGINE_API virtual bool HasCustomPrimitiveData() const override;
+	ENGINE_API virtual bool HasSubstrateRoughnessTracking() const override;
 
 	ENGINE_API virtual FGraphEventArray PrecachePSOs(const FPSOPrecacheVertexFactoryDataList& VertexFactoryDataList, const FPSOPrecacheParams& PreCacheParams, EPSOPrecachePriority Priority, TArray<FMaterialPSOPrecacheRequestID>& OutMaterialPSORequestIDs) override;
 
diff --git a/Engine/Source/Runtime/Engine/Public/Materials/MaterialInstance.h b/Engine/Source/Runtime/Engine/Public/Materials/MaterialInstance.h
--- a/Engine/Source/Runtime/Engine/Public/Materials/MaterialInstance.h
+++ b/Engine/Source/Runtime/Engine/Public/Materials/MaterialInstance.h
@@ -916,6 +916,7 @@ public:
 	ENGINE_API virtual USubsurfaceProfile* GetSubsurfaceProfileOverride_Internal() const override;
 	ENGINE_API virtual bool CastsRayTracedShadows() const override;
 	ENGINE_API virtual bool IsTessellationEnabled() const override;
+	ENGINE_API virtual bool HasSubstrateRoughnessTracking() const override;
 
 	/** Checks to see if an input property should be active, based on the state of the material */
 	ENGINE_API virtual bool IsPropertyActive(EMaterialProperty InProperty) const override;
diff --git a/Engine/Source/Runtime/Engine/Public/Materials/MaterialInterface.h b/Engine/Source/Runtime/Engine/Public/Materials/MaterialInterface.h
--- a/Engine/Source/Runtime/Engine/Public/Materials/MaterialInterface.h
+++ b/Engine/Source/Runtime/Engine/Public/Materials/MaterialInterface.h
@@ -1013,6 +1013,7 @@ public:
 	ENGINE_API virtual UNeuralProfile* GetNeuralProfile_Internal() const;
 	ENGINE_API virtual bool CastsRayTracedShadows() const;
 	ENGINE_API virtual bool IsTessellationEnabled() const;
+	ENGINE_API virtual bool HasSubstrateRoughnessTracking() const;
 
 	/**
 	 * Force the streaming system to disregard the normal logic for the specified duration and
