Add some nice BP metadata to the input device library. Better return value display names and categories

#rnx

diff --git a/Engine/Source/Runtime/Engine/Classes/GameFramework/InputDeviceLibrary.h b/Engine/Source/Runtime/Engine/Classes/GameFramework/InputDeviceLibrary.h
--- a/Engine/Source/Runtime/Engine/Classes/GameFramework/InputDeviceLibrary.h
+++ b/Engine/Source/Runtime/Engine/Classes/GameFramework/InputDeviceLibrary.h
@@ -45,7 +45,7 @@ public:
 	 * @param OutInputDevices		Array of input device ID's that will be populated with the mapped devices.
 	 * @return						The number of mapped devices, INDEX_NONE if the user was not found.
 	 */
-	UFUNCTION(BlueprintCallable, Category = "PlatformInputDevice")
+	UFUNCTION(BlueprintCallable, Category = "Input Devices", meta=(ReturnDisplayName="Number of Devices"))
 	static ENGINE_API int32 GetAllInputDevicesForUser(const FPlatformUserId UserId, TArray<FInputDeviceId>& OutInputDevices);
 
 	/**
@@ -54,7 +54,7 @@ public:
 	 * @param OutInputDevices	Array of input devices to populate
 	 * @return					The number of connected input devices
 	 */
-	UFUNCTION(BlueprintCallable, Category = "PlatformInputDevice")
+	UFUNCTION(BlueprintCallable, Category = "Input Devices", meta=(ReturnDisplayName="Number of Devices"))
 	static ENGINE_API int32 GetAllInputDevices(TArray<FInputDeviceId>& OutInputDevices);
 
 	/**
@@ -63,7 +63,7 @@ public:
 	 * @param OutInputDevices	Array of input devices to populate
 	 * @return					The number of connected input devices
 	 */
-	UFUNCTION(BlueprintCallable, Category = "PlatformInputDevice")
+	UFUNCTION(BlueprintCallable, Category = "Input Devices", meta=(ReturnDisplayName="Number of Devices"))
 	static ENGINE_API int32 GetAllConnectedInputDevices(TArray<FInputDeviceId>& OutInputDevices);
 
 	/**
@@ -72,14 +72,14 @@ public:
 	 * @param OutUsers		Array that will be populated with the platform users.
 	 * @return				The number of active platform users
 	 */
-	UFUNCTION(BlueprintCallable, Category = "PlatformInputDevice")
+	UFUNCTION(BlueprintCallable, Category = "Input Devices", meta=(ReturnDisplayName="Number of Users"))
 	static ENGINE_API int32 GetAllActiveUsers(TArray<FPlatformUserId>& OutUsers);
 
 	/**
 	 * Returns the platform user id that is being used for unmapped input devices.
 	 * Will be PLATFORMUSERID_NONE if platform does not support this (this is the default behavior)
 	 */
-	UFUNCTION(BlueprintPure, Category = "PlatformInputDevice")
+	UFUNCTION(BlueprintPure, Category = "Input Devices")
 	static ENGINE_API FPlatformUserId GetUserForUnpairedInputDevices();
 
 	/**
@@ -88,27 +88,27 @@ public:
 	 * map devices such as the keyboard and mouse that don't get assigned unique ID's from their
 	 * owning platform code.
 	 */
-	UFUNCTION(BlueprintPure, Category = "PlatformInputDevice")
+	UFUNCTION(BlueprintPure, Category = "Input Devices")
 	static ENGINE_API FPlatformUserId GetPrimaryPlatformUser();
 
 	/** Returns true if the given Platform User Id is the user for unpaired input devices on this platform. */
-	UFUNCTION(BlueprintPure, Category = "PlatformInputDevice")
+	UFUNCTION(BlueprintPure, Category = "Input Devices")
 	static ENGINE_API bool IsUnpairedUserId(const FPlatformUserId PlatformId);
 
 	/** Returns true if the given input device is mapped to the unpaired platform user id. */
-	UFUNCTION(BlueprintPure, Category = "PlatformInputDevice")
+	UFUNCTION(BlueprintPure, Category = "Input Devices")
 	static ENGINE_API bool IsInputDeviceMappedToUnpairedUser(const FInputDeviceId InputDevice);
 
 	/** Returns the default device id used for things like keyboard/mouse input */
-	UFUNCTION(BlueprintPure, Category = "PlatformInputDevice")
+	UFUNCTION(BlueprintPure, Category = "Input Devices")
 	static ENGINE_API FInputDeviceId GetDefaultInputDevice();
 
 	/** Returns the platform user attached to this input device, or PLATFORMUSERID_NONE if invalid */
-	UFUNCTION(BlueprintCallable, Category = "PlatformInputDevice")
+	UFUNCTION(BlueprintCallable, Category = "Input Devices")
 	static ENGINE_API FPlatformUserId GetUserForInputDevice(FInputDeviceId DeviceId);
 
 	/** Returns the primary input device used by a specific player, or INPUTDEVICEID_NONE if invalid */
-	UFUNCTION(BlueprintCallable, Category = "PlatformInputDevice")
+	UFUNCTION(BlueprintCallable, Category = "Input Devices")
 	static ENGINE_API FInputDeviceId GetPrimaryInputDeviceForUser(FPlatformUserId UserId);
 
 	/**
@@ -117,7 +117,7 @@ public:
 	 * @param DeviceId		The device to get the connection state of
 	 * @return				The connection state of the given device. EInputDeviceConnectionState::Unknown if the device is not mapped
 	 */
-	UFUNCTION(BlueprintCallable, Category = "PlatformInputDevice")
+	UFUNCTION(BlueprintCallable, Category = "Input Devices")
 	static ENGINE_API EInputDeviceConnectionState GetInputDeviceConnectionState(const FInputDeviceId DeviceId);
 
 	/**
@@ -125,7 +125,7 @@ public:
 	 * 
 	 * @return	True if the given input device is valid (it has been assigned an ID by the PlatformInputDeviceMapper)
 	 */
-	UFUNCTION(BlueprintPure, Category = "PlatformInputDevice")
+	UFUNCTION(BlueprintPure, Category = "Input Devices", meta=(ReturnDisplayName="Is Valid"))
 	static ENGINE_API bool IsValidInputDevice(FInputDeviceId DeviceId);
 
 	/**
@@ -133,30 +133,30 @@ public:
 	 * 
 	 * @return	True if the platform ID is valid (it has been assigned by the PlatformInputDeviceMapper)
 	 */
-	UFUNCTION(BlueprintPure, Category = "PlatformInputDevice")
+	UFUNCTION(BlueprintPure, Category = "Input Devices", meta=(ReturnDisplayName="Is Valid"))
 	static ENGINE_API bool IsValidPlatformId(FPlatformUserId UserId);
 
 	/** Static invalid platform user */
-	UFUNCTION(BlueprintPure, meta = (ScriptConstant = "None", ScriptConstantHost = "/Script/CoreUObject.PlatformUserId"), Category = "PlatformInputDevice")
+	UFUNCTION(BlueprintPure, meta = (ScriptConstant = "None", ScriptConstantHost = "/Script/CoreUObject.PlatformUserId"), Category = "Input Devices")
 	static ENGINE_API FPlatformUserId PlatformUserId_None();
 
 	/** Static invalid input device */
-	UFUNCTION(BlueprintPure, meta = (ScriptConstant = "None", ScriptConstantHost = "/Script/CoreUObject.InputDeviceId"), Category = "PlatformInputDevice")
+	UFUNCTION(BlueprintPure, meta = (ScriptConstant = "None", ScriptConstantHost = "/Script/CoreUObject.InputDeviceId"), Category = "Input Devices")
 	static ENGINE_API FInputDeviceId InputDeviceId_None();
 
 	/** Returns true if PlatformUserId A is equal to PlatformUserId B (A == B) */
-	UFUNCTION(BlueprintPure, meta = (DisplayName = "Equal (PlatformUserId)", CompactNodeTitle = "==", ScriptMethod = "Equals", ScriptOperator = "==", Keywords = "== equal"), Category = "PlatformInputDevice")
+	UFUNCTION(BlueprintPure, meta = (DisplayName = "Equal (PlatformUserId)", CompactNodeTitle = "==", ScriptMethod = "Equals", ScriptOperator = "==", Keywords = "== equal"), Category = "Input Devices")
 	static ENGINE_API bool EqualEqual_PlatformUserId(FPlatformUserId A, FPlatformUserId B);
 
 	/** Returns true if PlatformUserId A is not equal to PlatformUserId B (A != B) */
-	UFUNCTION(BlueprintPure, meta = (DisplayName = "Not Equal (PlatformUserId)", CompactNodeTitle = "!=", ScriptMethod = "NotEqual", ScriptOperator = "!=", Keywords = "!= not equal"), Category = "PlatformInputDevice")
+	UFUNCTION(BlueprintPure, meta = (DisplayName = "Not Equal (PlatformUserId)", CompactNodeTitle = "!=", ScriptMethod = "NotEqual", ScriptOperator = "!=", Keywords = "!= not equal"), Category = "Input Devices")
 	static ENGINE_API bool NotEqual_PlatformUserId(FPlatformUserId A, FPlatformUserId B);
 	
 	/** Returns true if InputDeviceId A is equal to InputDeviceId B (A == B) */
-	UFUNCTION(BlueprintPure, meta = (DisplayName = "Equal (InputDeviceId)", CompactNodeTitle = "==", ScriptMethod = "Equals", ScriptOperator = "==", Keywords = "== equal"), Category = "PlatformInputDevice")
+	UFUNCTION(BlueprintPure, meta = (DisplayName = "Equal (InputDeviceId)", CompactNodeTitle = "==", ScriptMethod = "Equals", ScriptOperator = "==", Keywords = "== equal"), Category = "Input Devices")
 	static ENGINE_API bool EqualEqual_InputDeviceId(FInputDeviceId A, FInputDeviceId B);
 	
 	/** Returns true if InputDeviceId A is not equal to InputDeviceId B (A != B) */
-	UFUNCTION(BlueprintPure, meta = (DisplayName = "Not Equal (InputDeviceId)", CompactNodeTitle = "!=", ScriptMethod = "NotEqual", ScriptOperator = "!=", Keywords = "!= not equal"), Category = "PlatformInputDevice")
+	UFUNCTION(BlueprintPure, meta = (DisplayName = "Not Equal (InputDeviceId)", CompactNodeTitle = "!=", ScriptMethod = "NotEqual", ScriptOperator = "!=", Keywords = "!= not equal"), Category = "Input Devices")
 	static ENGINE_API bool NotEqual_InputDeviceId(FInputDeviceId A, FInputDeviceId B);
 };
