Ensure that instanced static mesh components don't produce non-deterministic output during cooks by generating an InstancingRandomSeed using the Rand function during cook.

diff --git a/Engine/Source/Runtime/Engine/Classes/Components/HierarchicalInstancedStaticMeshComponent.h b/Engine/Source/Runtime/Engine/Classes/Components/HierarchicalInstancedStaticMeshComponent.h
--- a/Engine/Source/Runtime/Engine/Classes/Components/HierarchicalInstancedStaticMeshComponent.h
+++ b/Engine/Source/Runtime/Engine/Classes/Components/HierarchicalInstancedStaticMeshComponent.h
@@ -362,7 +362,6 @@ protected:
 
 	ENGINE_API FVector CalcTranslatedInstanceSpaceOrigin() const;
 	ENGINE_API void GetInstanceTransforms(TArray<FMatrix>& InstanceTransforms, FVector const& Offset) const;
-	ENGINE_API void InitializeInstancingRandomSeed();
 
 	/** Removes specified instances */ 
 	ENGINE_API void RemoveInstancesInternal(TConstArrayView<int32> InstanceIndices);
diff --git a/Engine/Source/Runtime/Engine/Classes/Components/InstancedStaticMeshComponent.h b/Engine/Source/Runtime/Engine/Classes/Components/InstancedStaticMeshComponent.h
--- a/Engine/Source/Runtime/Engine/Classes/Components/InstancedStaticMeshComponent.h
+++ b/Engine/Source/Runtime/Engine/Classes/Components/InstancedStaticMeshComponent.h
@@ -694,6 +694,8 @@ protected:
 	ENGINE_API virtual bool DeleteSMInstances(TArrayView<const FSMInstanceId> InstanceIds) override;
 	ENGINE_API virtual bool DuplicateSMInstances(TArrayView<const FSMInstanceId> InstanceIds, TArray<FSMInstanceId>& OutNewInstanceIds) override;
 
+	void InitializeInstancingRandomSeed();
+
 	friend FStaticLightingTextureMapping_InstancedStaticMesh;
 	friend FInstancedLightMap2D;
 	friend FInstancedShadowMap2D;
diff --git a/Engine/Source/Runtime/Engine/Private/HierarchicalInstancedStaticMesh.cpp b/Engine/Source/Runtime/Engine/Private/HierarchicalInstancedStaticMesh.cpp
--- a/Engine/Source/Runtime/Engine/Private/HierarchicalInstancedStaticMesh.cpp
+++ b/Engine/Source/Runtime/Engine/Private/HierarchicalInstancedStaticMesh.cpp
@@ -2098,13 +2098,7 @@ void UHierarchicalInstancedStaticMeshComponent::PostEditChangeChainProperty(FPro
 {
 	if ((PropertyChangedEvent.Property != NULL && PropertyChangedEvent.Property->GetFName() == "InstancingRandomSeed"))
 	{
-		// If we don't have a random seed for this instanced static mesh component yet, then go ahead and
-		// generate one now.  This will be saved with the static mesh component and used for future generation
-		// of random numbers for this component's instances. (Used by the PerInstanceRandom material expression)
-		while (InstancingRandomSeed == 0)
-		{
-			InstancingRandomSeed = FMath::Rand();
-		}
+		InitializeInstancingRandomSeed();
 	}
 
 	Super::PostEditChangeChainProperty(PropertyChangedEvent);
@@ -2841,17 +2835,6 @@ void UHierarchicalInstancedStaticMeshComponent::GetInstanceTransforms(TArray<FMa
 	UE_LOG(LogStaticMesh, Verbose, TEXT("Copied %d transforms in %.3fs."), Num, float(FPlatformTime::Seconds() - StartTime));
 }
 
-void UHierarchicalInstancedStaticMeshComponent::InitializeInstancingRandomSeed()
-{
-	// If we don't have a random seed for this instanced static mesh component yet, then go ahead and
-		// generate one now.  This will be saved with the static mesh component and used for future generation
-		// of random numbers for this component's instances. (Used by the PerInstanceRandom material expression)
-	while (InstancingRandomSeed == 0)
-	{
-		InstancingRandomSeed = FMath::Rand();
-	}
-}
-
 void UHierarchicalInstancedStaticMeshComponent::BuildTreeAsync()
 {
 	TRACE_CPUPROFILER_EVENT_SCOPE(UHierarchicalInstancedStaticMeshComponent::BuildTreeAsync);
diff --git a/Engine/Source/Runtime/Engine/Private/InstancedStaticMesh.cpp b/Engine/Source/Runtime/Engine/Private/InstancedStaticMesh.cpp
--- a/Engine/Source/Runtime/Engine/Private/InstancedStaticMesh.cpp
+++ b/Engine/Source/Runtime/Engine/Private/InstancedStaticMesh.cpp
@@ -4627,14 +4627,7 @@ void UInstancedStaticMeshComponent::OnRegister()
 
 	if (FApp::CanEverRender() && !HasAnyFlags(RF_ClassDefaultObject | RF_ArchetypeObject))
 	{
-		// If we don't have a random seed for this instanced static mesh component yet, then go ahead and
-		// generate one now.  This will be saved with the static mesh component and used for future generation
-		// of random numbers for this component's instances. (Used by the PerInstanceRandom material expression)
-		while (InstancingRandomSeed == 0)
-		{
-			InstancingRandomSeed = FMath::Rand();
-		}
-
+		InitializeInstancingRandomSeed();
 		// Propagate the current number of instances to the manager, because of the multifarious ways the engine shoves data into the properties it is possible for the count to get out of sync.
 		// At this point we need to let the manager reset if needed, also at this point we may assume that we have no idea of the state of individual members.
 		PrimitiveInstanceDataManager.OnRegister(PerInstanceSMData.Num());
@@ -4999,13 +4992,7 @@ void UInstancedStaticMeshComponent::OnPostLoadPerInstanceData()
 
 	if (!HasAnyFlags(RF_ClassDefaultObject|RF_ArchetypeObject))
 	{
-		// If we don't have a random seed for this instanced static mesh component yet, then go ahead and
-		// generate one now.  This will be saved with the static mesh component and used for future generation
-		// of random numbers for this component's instances. (Used by the PerInstanceRandom material expression)
-		while (InstancingRandomSeed == 0)
-		{
-			InstancingRandomSeed = FMath::Rand();
-		}		
+		InitializeInstancingRandomSeed();
 	}
 
 	if (AActor* Owner = GetOwner())
@@ -5621,6 +5608,35 @@ bool UInstancedStaticMeshComponent::DuplicateSMInstances(TArrayView<const FSMIns
 	return true;
 }
 
+void UInstancedStaticMeshComponent::InitializeInstancingRandomSeed()
+{
+	// If we don't have a random seed for this instanced static mesh component yet, then go ahead and
+	// generate one now.  This will be saved with the static mesh component and used for future generation
+	// of random numbers for this component's instances. (Used by the PerInstanceRandom material expression)
+	if (IsRunningCookCommandlet())
+	{
+		// During cook, we must generate a deterministic value. Subsequent save of the asset in editor can save
+		// a new unique value that persists.
+		if (InstancingRandomSeed == 0)
+		{
+			TStringBuilder<256> ResultBuilder;
+			GetPathName(nullptr, ResultBuilder);
+			InstancingRandomSeed = (int32)CityHash32((const char*)ResultBuilder.GetData(), ResultBuilder.Len()*sizeof(decltype(ResultBuilder)::ElementType));
+			if (InstancingRandomSeed == 0)
+			{
+				InstancingRandomSeed = 42;
+			}
+		}
+	}
+	else
+	{
+		while (InstancingRandomSeed == 0)
+		{
+			InstancingRandomSeed = FMath::Rand();
+		}
+	}
+}
+
 void FInstancedStaticMeshVertexFactoryShaderParameters::GetElementShaderBindings(
 	const class FSceneInterface* Scene,
 	const FSceneView* View,
