Add developer-specifiable ordering to ULevelExporterT3D::ExportText(). This solves a problem where the transform gizmo could change during alt-drag depending on where the actors were relative to each other in the level.

diff --git a/Engine/Source/Editor/UnrealEd/Private/EditorExporters.cpp b/Engine/Source/Editor/UnrealEd/Private/EditorExporters.cpp
--- a/Engine/Source/Editor/UnrealEd/Private/EditorExporters.cpp
+++ b/Engine/Source/Editor/UnrealEd/Private/EditorExporters.cpp
@@ -654,8 +654,28 @@ bool ULevelExporterT3D::ExportText( const FExportObjectInnerContext* Context, UO
 	const bool bShowCancelButton = true;
 	SlowTask.MakeDialogDelayed(1.f, bShowCancelButton);
 
+	// Allow for a pre-specified (and pre-sorted) list of level inners
+	TArray<TObjectPtr<AActor>> SpecifiedActors;
+	const TArray<UObject*>* LevelInners = Context->GetObjectInners(Level);
+	if (LevelInners)
+	{
+		SpecifiedActors.Reserve(LevelInners->Num());
+		for (UObject* Inner : *LevelInners)
+		{
+			if (AActor* Actor = Cast<AActor>(Inner))
+			{
+				if (ensureMsgf(Actor->GetLevel() == Level, TEXT("Specified actors must actually be inners of the level!")))
+				{
+					SpecifiedActors.Add(Actor);
+				}
+			}
+		}
+	}
+
 	// loop through all of the actors just in this level
-	for (AActor* Actor : Level->Actors)
+	const TArray<TObjectPtr<AActor>>& Actors = LevelInners ? SpecifiedActors : Level->Actors;
+	
+	for (AActor* Actor : Actors)
 	{
 		if (SlowTask.ShouldCancel())
 		{
diff --git a/Engine/Source/Runtime/Engine/Private/UnrealExporter.cpp b/Engine/Source/Runtime/Engine/Private/UnrealExporter.cpp
--- a/Engine/Source/Runtime/Engine/Private/UnrealExporter.cpp
+++ b/Engine/Source/Runtime/Engine/Private/UnrealExporter.cpp
@@ -526,22 +526,7 @@ FExportObjectInnerContext::FExportObjectInnerContext()
 	// For each object . . .
 	for (UObject* InnerObj : TObjectRange<UObject>(RF_ClassDefaultObject, true, EInternalObjectFlags::Garbage))
 	{
-		UObject* OuterObj = InnerObj->GetOuter();
-		if ( OuterObj )
-		{
-			InnerList* Inners = ObjectToInnerMap.Find( OuterObj );
-			if ( Inners )
-			{
-				// Add object to existing inner list.
-				Inners->Add( InnerObj );
-			}
-			else
-			{
-				// Create a new inner list for the outer object.
-				InnerList& InnersForOuterObject = ObjectToInnerMap.Add( OuterObj, InnerList() );
-				InnersForOuterObject.Add( InnerObj );
-			}
-		}
+		AddObjectToInnerMap(InnerObj);
 	}
 }
 
@@ -553,26 +538,30 @@ FExportObjectInnerContext::FExportObjectInnerContext(const TArray<UObject*>& Obj
 	{
 		if (!ObjsToIgnore.Contains(InnerObj))
 		{
-			UObject* OuterObj = InnerObj->GetOuter();
-			if (IsValid(OuterObj))
-			{
-				InnerList* Inners = ObjectToInnerMap.Find(OuterObj);
-				if (Inners)
-				{
-					// Add object to existing inner list.
-					Inners->Add(InnerObj);
-				}
-				else
-				{
-					// Create a new inner list for the outer object.
-					InnerList& InnersForOuterObject = ObjectToInnerMap.Add(OuterObj, InnerList());
-					InnersForOuterObject.Add(InnerObj);
-				}
-			}
+			AddObjectToInnerMap(InnerObj);
 		}
 	}
 }
 
+void FExportObjectInnerContext::AddObjectToInnerMap(UObject* InObject)
+{
+	UObject* OuterObj = InObject->GetOuter();
+	if (IsValid(OuterObj))
+	{
+		InnerList* Inners = ObjectToInnerMap.Find(OuterObj);
+		if (Inners)
+		{
+			// Add object to existing inner list.
+			Inners->Add(InObject);
+		}
+		else
+		{
+			// Create a new inner list for the outer object.
+			InnerList& InnersForOuterObject = ObjectToInnerMap.Add(OuterObj, InnerList());
+			InnersForOuterObject.Add(InObject);
+		}
+	}
+}
 
 bool FExportObjectInnerContext::IsObjectSelected(const UObject* InObj) const
 {
@@ -1389,8 +1378,12 @@ FSelectedActorExportObjectInnerContext::FSelectedActorExportObjectInnerContext(c
 	: FExportObjectInnerContext(false) //call the empty version of the base class
 {
 	// For each selected actor...
-	for (const AActor* Actor : InSelectedActors)
+	for (AActor* Actor : InSelectedActors)
 	{
+		// Explicitly mark each actor as an inner of its outer.
+		// This can be used (e.g. in ULevelExporterT3D`::ExportText()`)
+		// to specify the precise order in which objects will be exported.
+		AddObjectToInnerMap(Actor);
 		AddSelectedActor(Actor);
 	}
 }
@@ -1407,19 +1400,8 @@ void FSelectedActorExportObjectInnerContext::AddSelectedActor(const AActor* InAc
 
 	ForEachObjectWithOuter(InActor, [this](UObject* InnerObj)
 	{
-		UObject* OuterObj = InnerObj->GetOuter();
-		InnerList* Inners = ObjectToInnerMap.Find(OuterObj);
-		if (Inners)
-		{
-			// Add object to existing inner list.
-			Inners->Add(InnerObj);
-		}
-		else
-		{
-			// Create a new inner list for the outer object.
-			InnerList& InnersForOuterObject = ObjectToInnerMap.Add(OuterObj, InnerList());
-			InnersForOuterObject.Add(InnerObj);
-		}
+		AddObjectToInnerMap(InnerObj);
 	}, /** bIncludeNestedObjects */ true, RF_NoFlags, EInternalObjectFlags::Garbage);
 }
+
 #endif
diff --git a/Engine/Source/Runtime/Engine/Public/UnrealExporter.h b/Engine/Source/Runtime/Engine/Public/UnrealExporter.h
--- a/Engine/Source/Runtime/Engine/Public/UnrealExporter.h
+++ b/Engine/Source/Runtime/Engine/Public/UnrealExporter.h
@@ -49,6 +49,8 @@ protected:
 
 	ObjectToInnerMapType				ObjectToInnerMap;
 
+	ENGINE_API void AddObjectToInnerMap(UObject* InObject);
+
 public:
 	/** 
 	 * Get the inner list of an object
