New material translator push. Better support for texture sampling and texture parameters.

diff --git a/Engine/Source/Runtime/Engine/Private/Materials/MaterialExpressionsIR.cpp b/Engine/Source/Runtime/Engine/Private/Materials/MaterialExpressionsIR.cpp
--- a/Engine/Source/Runtime/Engine/Private/Materials/MaterialExpressionsIR.cpp
+++ b/Engine/Source/Runtime/Engine/Private/Materials/MaterialExpressionsIR.cpp
@@ -397,7 +397,7 @@ void UMaterialExpressionAppendVector::Build(MIR::FEmitter& Emitter)
 	check(Dimensions >= 2 && Dimensions <= 4);
 
 	// Construct the output vector type.
-	MIR::FPrimitiveTypePtr ResultType = MIR::FPrimitiveType::GetVector(MIR::SK_Float, Dimensions);
+	MIR::FPrimitiveTypePtr ResultType = MIR::FPrimitiveType::GetVector(MIR::ScalarKind_Float, Dimensions);
 
 	// Set up each output vector component 
 	MIR::FValue* Components[4] = { nullptr, nullptr, nullptr, nullptr };
@@ -619,42 +619,15 @@ void UMaterialExpressionTextureObject::Build(UE::MIR::FEmitter& Emitter)
 	Emitter.Output(0, Emitter.EmitTextureObject(Texture, SamplerType));
 }
 
-static void BuildTextureSampleExpression(UTexture* Texture, UMaterialExpressionTextureSample* Expr, UE::MIR::FEmitter& Emitter)
+static void BuildTextureSampleExpression(MIR::FValue* Texture, UMaterialExpressionTextureSample* Expr, UE::MIR::FEmitter& Emitter)
 {
-	if (MIR::FValue* InputTextureValue = Emitter.TryInput(&Expr->TextureObject))
+	MIR::FValue* TexCoords = Emitter.TryInput(&Expr->Coordinates);
+	if (!TexCoords)
 	{
-		auto InputTextureSampler = InputTextureValue->GetTexture();
-		if (!InputTextureSampler.IsSet())
-		{
-			Emitter.Error(TEXT("Value flowing into texture input is not a texture."));
-			return;
-		}
-
-		if (InputTextureSampler->Get<EMaterialSamplerType>() != SAMPLERTYPE_Color)
-		{
-			Emitter.Error(TEXT("Input texture sampler type must be color."));
-			return;
-		}
-
-		Texture = InputTextureSampler->Get<UTexture*>();
-	}
-	
-	if (!Texture)
-	{
-		Emitter.Error("Unspecified texture.");
-		return;
+		TexCoords = Emitter.EmitExternalInput(MIR::TexCoordIndexToExternalInput(Expr->ConstCoordinate));
 	}
 
-	MIR::FValue* TexCoordsValue = Emitter.TryInput(&Expr->Coordinates);
-	if (!TexCoordsValue)
-	{
-		TexCoordsValue = Emitter.EmitExternalInput(MIR::TexCoordIndexToExternalInput(Expr->ConstCoordinate));
-	}
-
-	UE_MIR_CHECKPOINT();
-
-	MIR::FValue* OutputValue = Emitter.EmitTextureSample(Texture, TexCoordsValue, Expr->SamplerSource, Expr->MipValueMode, Expr->SamplerType);
-	
+	MIR::FValue* OutputValue = Emitter.EmitTextureSample(Texture, TexCoords, Expr->SamplerSource, Expr->MipValueMode, Expr->SamplerType);
 	Emitter.Output(0, Emitter.EmitSwizzle(OutputValue, MIR::FSwizzleMask::XYZ()));
 	Emitter.Output(1, Emitter.EmitSubscript(OutputValue, 0)); 
 	Emitter.Output(2, Emitter.EmitSubscript(OutputValue, 1)); 
@@ -665,26 +638,60 @@ static void BuildTextureSampleExpression(UTexture* Texture, UMaterialExpressionT
 
 void UMaterialExpressionTextureSample::Build(UE::MIR::FEmitter& Emitter)
 {
-	BuildTextureSampleExpression(Texture, this, Emitter);
+	MIR::FValue* TextureValue = Emitter.TryInput(&TextureObject) ? Emitter.InputTexture(&TextureObject) : nullptr;
+	UE_MIR_CHECKPOINT();
+
+	if (!TextureValue)
+	{
+		if (!Texture)
+		{
+			Emitter.Error(TEXT("No texture specified for this expression."));
+			return;
+		}
+
+		TextureValue = Emitter.EmitTextureObject(Texture.Get(), SamplerType);
+	}
+
+	if (TextureValue && TextureValue->GetTexture()->Get<EMaterialSamplerType>() != SAMPLERTYPE_Color)
+	{
+		Emitter.Error(TEXT("Input texture sampler type must be color."));
+		return;
+	}
+
+	BuildTextureSampleExpression(TextureValue, this, Emitter);
 }
 
-void UMaterialExpressionTextureSampleParameter::Build(UE::MIR::FEmitter& Emitter)
+static MIR::FValue* BuildTextureObjectParameter(UE::MIR::FEmitter& Emitter, UMaterialExpressionTextureSampleParameter* Expr)
 {
 	FMaterialParameterMetadata Param;
-	if (!GetParameterValue(Param))
+	if (!Expr->GetParameterValue(Param))
 	{
 		Emitter.Error(TEXT("Failed to get parameter value"));
-		return;
+		return nullptr;
 	}
 
-	UTexture* TheTexture = Cast<UTexture>(Param.Value.AsTextureObject());
-	if (!TheTexture)
+	MIR::FValue* ParameterValue = Emitter.EmitParameter(Expr->GetParameterName(), Param, Expr->SamplerType);
+	if (!ParameterValue->Type->IsTexture())
 	{
-		Emitter.Error(TEXT("Parameter does not have a texture object"));
-		return;
+		Emitter.Error(TEXT("Parameter is not a texture"));
+		return nullptr;
 	}
 
-	BuildTextureSampleExpression(TheTexture, this, Emitter);
+	return ParameterValue;
+}
+
+void UMaterialExpressionTextureSampleParameter::Build(UE::MIR::FEmitter& Emitter)
+{
+	MIR::FValue* ParameterValue = BuildTextureObjectParameter(Emitter, this);
+	UE_MIR_CHECKPOINT();
+	BuildTextureSampleExpression(ParameterValue, this, Emitter);
+}
+
+void UMaterialExpressionTextureObjectParameter::Build(UE::MIR::FEmitter& Emitter)
+{
+	MIR::FValue* ParameterValue = BuildTextureObjectParameter(Emitter, this);
+	UE_MIR_CHECKPOINT();
+	Emitter.Output(0, ParameterValue);
 }
 
 void UMaterialExpressionTextureCoordinate::Build(UE::MIR::FEmitter& Emitter)
diff --git a/Engine/Source/Runtime/Engine/Private/Materials/MaterialIR.cpp b/Engine/Source/Runtime/Engine/Private/Materials/MaterialIR.cpp
--- a/Engine/Source/Runtime/Engine/Private/Materials/MaterialIR.cpp
+++ b/Engine/Source/Runtime/Engine/Private/Materials/MaterialIR.cpp
@@ -16,7 +16,7 @@ const TCHAR* ValueKindToString(EValueKind Kind)
 		case VK_Constant: return TEXT("Constant");
 		case VK_ExternalInput: return TEXT("ExternalInput");
 		case VK_TextureObject: return TEXT("TextureObject");
-		case VK_UniformParameter: return TEXT("Parameter");
+		case VK_UniformParameter: return TEXT("UniformParameter");
 		case VK_Dimensional: return TEXT("Dimensional");
 		case VK_SetMaterialOutput: return TEXT("SetMaterialOutput");
 		case VK_Operator: return TEXT("Operator");
@@ -139,7 +139,7 @@ TArrayView<FValue*> FValue::GetUses()
 		case VK_TextureSample:
 		{
 			auto This = static_cast<FTextureSample*>(this);
-			return { &This->TexCoordArg, 3 };
+			return { &This->TextureObject, 4 };
 		}
 
 		default: UE_MIR_UNREACHABLE();
@@ -215,6 +215,11 @@ TOptional<TPair<UTexture*, EMaterialSamplerType>> FValue::GetTexture()
 	{
 		return FPair{ TextureObject->Texture, TextureObject->SamplerType };
 	}
+	else if (auto UniformParameter = As<MIR::FUniformParameter>())
+	{
+		check(UniformParameter->Metadata.Value.Type == EMaterialParameterType::Texture);
+		return FPair{ UniformParameter->Metadata.Value.Texture, UniformParameter->SamplerType };
+	}
 	return {};
 }
 
diff --git a/Engine/Source/Runtime/Engine/Private/Materials/MaterialIRDebug.cpp b/Engine/Source/Runtime/Engine/Private/Materials/MaterialIRDebug.cpp
--- a/Engine/Source/Runtime/Engine/Private/Materials/MaterialIRDebug.cpp
+++ b/Engine/Source/Runtime/Engine/Private/Materials/MaterialIRDebug.cpp
@@ -33,9 +33,9 @@ static void DumpValueInfo(const MIR::FValue* Value, FString& Out)
 	{
 		switch (Constant->Type->AsPrimitive()->ScalarKind)
 		{
-			case MIR::SK_Bool:	Out.Append(Constant->Boolean ? TEXT("true") : TEXT("false")); break;
-			case MIR::SK_Int:	Out.Appendf(TEXT("%lld"), Constant->Integer); break;
-			case MIR::SK_Float: Out.Appendf(TEXT("%f"), Constant->Float); break;
+			case MIR::ScalarKind_Bool:	Out.Append(Constant->Boolean ? TEXT("true") : TEXT("false")); break;
+			case MIR::ScalarKind_Int:	Out.Appendf(TEXT("%lld"), Constant->Integer); break;
+			case MIR::ScalarKind_Float: Out.Appendf(TEXT("%f"), Constant->Float); break;
 			default: UE_MIR_UNREACHABLE();
 		}
 	}
diff --git a/Engine/Source/Runtime/Engine/Private/Materials/MaterialIREmitter.cpp b/Engine/Source/Runtime/Engine/Private/Materials/MaterialIREmitter.cpp
--- a/Engine/Source/Runtime/Engine/Private/Materials/MaterialIREmitter.cpp
+++ b/Engine/Source/Runtime/Engine/Private/Materials/MaterialIREmitter.cpp
@@ -196,7 +196,7 @@ FEmitter& FEmitter::DefaultInputToFloatScalar(const FExpressionInput* InInput, T
 	return *this;
 }
 
-FValue*  FEmitter::TryInput(const FExpressionInput* InInput)
+FValue* FEmitter::TryInput(const FExpressionInput* InInput)
 {
 	FValue* Value = MIR::Internal::GetInputValue(BuilderImpl, InInput);
 	return Value ? Value : nullptr;
@@ -213,6 +213,17 @@ FValue* FEmitter::Input(const FExpressionInput* InInput)
 	return Value;
 }
 
+FValue* FEmitter::InputOfType(const FExpressionInput* InInput, ETypeKind Kind)
+{
+	FValue* Value = Input(InInput);
+	if (Value->IsPoison())
+	{
+		return Value;
+	}
+	CheckInputTypeIs(InInput, Value, Kind);
+	return Value;
+}
+
 bool FEmitter::InputConstantBoolScalar(const FExpressionInput* InInput)
 {
 	FValue* Value = InputBoolScalar(InInput);
@@ -234,14 +245,14 @@ bool FEmitter::InputConstantBoolScalar(const FExpressionInput* InInput)
 FValue* FEmitter::InputBoolScalar(const FExpressionInput* InInput)
 {
 	FValue* Value = Input(InInput);
-	CheckInputIsScalar(InInput, Value, SK_Bool);
+	CheckInputIsScalar(InInput, Value, ScalarKind_Bool);
 	return Value;
 }
 
 FValue* FEmitter::InputFloatScalar(const FExpressionInput* InInput)
 {
 	FValue* Value = Input(InInput);
-	CheckInputIsScalar(InInput, Value, SK_Float);
+	CheckInputIsScalar(InInput, Value, ScalarKind_Float);
 	return Value;
 }
 
@@ -264,18 +275,18 @@ FValue* FEmitter::InputPrimitive(const FExpressionInput* InInput)
 	{
 		return Value;
 	}
-	CheckInputTypeIs(InInput, Value, TK_Primitive);
+	CheckInputTypeIs(InInput, Value, TypeKind_Primitive);
 	return Value;
 }
 
-FValue* FEmitter::InputOfType(const FExpressionInput* InInput, ETypeKind Kind)
+FValue* FEmitter::InputTexture(const FExpressionInput* InInput)
 {
 	FValue* Value = Input(InInput);
 	if (Value->IsPoison())
 	{
 		return Value;
 	}
-	CheckInputTypeIs(InInput, Value, Kind);
+	CheckValueIsTexture(Value);
 	return Value;
 }
 
@@ -293,7 +304,7 @@ FEmitter& FEmitter::Output(const FExpressionOutput* Output, FValue* Value)
 
 bool FEmitter::ToConstantBoolScalar(MIR::FValue* Value)
 {
-	CheckValueIsScalar(Value, SK_Bool);
+	CheckValueIsScalar(Value, ScalarKind_Bool);
 
 	if (!bExpressionHasError)
 	{
@@ -350,6 +361,20 @@ void FEmitter::CheckValueIsScalar(const FValue* Value, EScalarKind Kind)
 	}
 }
 
+void FEmitter::CheckValueIsTexture(const FValue* Value)
+{
+	if (Value->IsPoison())
+	{
+		return;
+	}
+
+	if (!Value->Type->IsTexture())
+	{
+		Errorf(TEXT("Value expected to be a texture. It is %s instead."), Value->Type->GetSpelling().GetData());
+		return;
+	}
+}
+
 FValue* FEmitter::EmitConstantFromShaderValue(const UE::Shader::FValue& InValue)
 {
 	using namespace UE::Shader;
@@ -374,9 +399,9 @@ FValue* FEmitter::EmitConstantScalarZero(EScalarKind Kind)
 {
 	switch (Kind)
 	{
-		case SK_Bool: return EmitConstantFalse();
-		case SK_Int: return EmitConstantIntScalar(0);
-		case SK_Float: return EmitConstantFloatScalar(0.0f);
+		case ScalarKind_Bool: return EmitConstantFalse();
+		case ScalarKind_Int: return EmitConstantIntScalar(0);
+		case ScalarKind_Float: return EmitConstantFloatScalar(0.0f);
 		default: UE_MIR_UNREACHABLE();
 	}
 }
@@ -398,7 +423,7 @@ FValue* FEmitter::EmitConstantBoolScalar(bool InX)
 
 FValue* FEmitter::EmitConstantFloatScalar(TFloat InX)
 {
-	FConstant Scalar = MakePrototype<FConstant>(FPrimitiveType::GetScalar(SK_Float));
+	FConstant Scalar = MakePrototype<FConstant>(FPrimitiveType::GetScalar(ScalarKind_Float));
 	Scalar.Float = InX;
 	return EmitPrototype(this, Scalar);
 }
@@ -429,7 +454,7 @@ FValue* FEmitter::EmitConstantFloat4(const FVector4f& InValue)
 
 FValue* FEmitter::EmitConstantIntScalar(TInteger InX)
 {
-	FConstant Scalar = MakePrototype<FConstant>(FPrimitiveType::GetScalar(SK_Int));
+	FConstant Scalar = MakePrototype<FConstant>(FPrimitiveType::GetScalar(ScalarKind_Int));
 	Scalar.Integer = InX;
 	return EmitPrototype(this, Scalar);
 }
@@ -621,7 +646,7 @@ FValue* FEmitter::EmitTextureObject(UTexture* Texture, EMaterialSamplerType Samp
 	return EmitPrototype(this, Proto);
 }
 
-FValue* FEmitter::EmitParameter(FName Name, FMaterialParameterMetadata& Metadata)
+FValue* FEmitter::EmitParameter(FName Name, FMaterialParameterMetadata& Metadata, EMaterialSamplerType SamplerType)
 {
 	switch (Metadata.Value.Type)
 	{
@@ -657,7 +682,19 @@ FValue* FEmitter::EmitParameter(FName Name, FMaterialParameterMetadata& Metadata
 
 		case EMaterialParameterType::Texture:
 		{
-			return EmitTextureObject(Metadata.Value.Texture, SAMPLERTYPE_Color);
+			if (Metadata.Value.Texture->GetMaterialType() != MCT_Texture2D)
+			{
+				Errorf(TEXT("Unsupported texture type"));
+				return MIR::FPoison::Get();
+			}
+
+			MIR::FUniformParameter Prototype = MakePrototype<MIR::FUniformParameter>(FObjectType::GetTexture2D());
+			Prototype.Info.Name = Name;
+			Prototype.Info.Index = -1;
+			Prototype.Info.Association = GlobalParameter;
+			Prototype.Metadata = Metadata;
+			Prototype.SamplerType = SamplerType;
+			return EmitPrototype(this, Prototype);
 		}
 
 		case EMaterialParameterType::StaticSwitch:
@@ -845,11 +882,11 @@ static FValue* TryFoldOperatorScalar(FEmitter* Emitter, EOperator Op, FValue* A,
 		bool Result;
 		switch (PrimitiveType->ScalarKind)
 		{
-		case SK_Int:
+		case ScalarKind_Int:
 			Result = FoldComparisonOperatorScalar<TInteger>(Op, AConstant->Integer, BConstant->Integer);
 			break;
 
-		case SK_Float:
+		case ScalarKind_Float:
 			Result = FoldComparisonOperatorScalar<TFloat>(Op, AConstant->Float, BConstant->Float);
 			break;
 
@@ -862,13 +899,13 @@ static FValue* TryFoldOperatorScalar(FEmitter* Emitter, EOperator Op, FValue* A,
 	{
 		switch (PrimitiveType->ScalarKind)
 		{
-			case SK_Bool:
+			case ScalarKind_Bool:
 			{
 				bool Result = (bool)FoldScalarBitOperator(Op, AConstant->Boolean, BConstant ? BConstant->Boolean : 0);
 				return Emitter->EmitConstantBoolScalar(Result);
 			}
 
-			case SK_Int:
+			case ScalarKind_Int:
 			{
 				TInteger Result = FoldScalarBitOperator(Op, AConstant->Integer, BConstant ? BConstant->Integer : 0);
 				return Emitter->EmitConstantIntScalar(Result);
@@ -882,13 +919,13 @@ static FValue* TryFoldOperatorScalar(FEmitter* Emitter, EOperator Op, FValue* A,
 	{
 		switch (PrimitiveType->ScalarKind)
 		{
-			case SK_Int:
+			case ScalarKind_Int:
 			{
 				int Result = FoldScalarArithmeticOperator<TInteger>(Op, AConstant->Integer, BConstant ? BConstant->Integer : 0, CConstant ? CConstant->Integer : 0);
 				return Emitter->EmitConstantIntScalar(Result);
 			}
 
-			case SK_Float:
+			case ScalarKind_Float:
 			{
 				TFloat Result = FoldScalarArithmeticOperator<TFloat>(Op, AConstant->Float, BConstant ? BConstant->Float : 0, CConstant ? CConstant->Float : 0);
 				return Emitter->EmitConstantFloatScalar(Result);
@@ -937,7 +974,7 @@ FValue* FEmitter::EmitOperator(EOperator Op, FValue* A, FValue* B, FValue* C)
 	// Determine the result type. If the operator is arithmetic, the result type will be the same
 	// as the arguments type. Otherwise it will have the same number of components but boolean.
 	FPrimitiveTypePtr ResultType = IsComparisonOperator(Op) || IsBitOperator(Op)
-		? FPrimitiveType::Get(SK_Bool, PrimitiveType->NumRows, PrimitiveType->NumColumns)
+		? FPrimitiveType::Get(ScalarKind_Bool, PrimitiveType->NumRows, PrimitiveType->NumColumns)
 		: PrimitiveType;
 
 
@@ -1079,24 +1116,24 @@ static FValue* CastConstant(FEmitter* Emitter, FConstant* Constant, EScalarKind
 
 	switch (ConstantScalarKind)
 	{
-		case SK_Bool:
-		case SK_Int:
+		case ScalarKind_Bool:
+		case ScalarKind_Int:
 		{
 			switch (TargetKind)
 			{
-				case SK_Bool: UE_MIR_TODO();
-				case SK_Int: return Emitter->EmitConstantIntScalar(Constant->Integer);
-				case SK_Float: return Emitter->EmitConstantFloatScalar((TFloat)Constant->Integer);
+				case ScalarKind_Bool: UE_MIR_TODO();
+				case ScalarKind_Int: return Emitter->EmitConstantIntScalar(Constant->Integer);
+				case ScalarKind_Float: return Emitter->EmitConstantFloatScalar((TFloat)Constant->Integer);
 				default: UE_MIR_UNREACHABLE();
 			}
 		}
 
-		case SK_Float:
+		case ScalarKind_Float:
 		{
 			switch (TargetKind)
 			{
-				case SK_Bool: UE_MIR_TODO();
-				case SK_Int: return Emitter->EmitConstantIntScalar((int)Constant->Float);
+				case ScalarKind_Bool: UE_MIR_TODO();
+				case ScalarKind_Int: return Emitter->EmitConstantIntScalar((int)Constant->Float);
 				default: UE_MIR_UNREACHABLE();
 			}
 		}
@@ -1253,10 +1290,10 @@ FValue* FEmitter::EmitConstruct(FTypePtr Type, FValue* Initializer)
 	return Result;
 }
 
-FValue* FEmitter::EmitTextureSample(UTexture* Texture, FValue* TexCoord, ESamplerSourceMode SamplerSourceMode, ETextureMipValueMode MipValueMode, EMaterialSamplerType SamplerType)
+FValue* FEmitter::EmitTextureSample(FValue* Texture, FValue* TexCoord, ESamplerSourceMode SamplerSourceMode, ETextureMipValueMode MipValueMode, EMaterialSamplerType SamplerType)
 {
 	MIR::FTextureSample Prototype = MakePrototype<MIR::FTextureSample>(FPrimitiveType::GetFloat4());
-	Prototype.Texture = Texture;
+	Prototype.TextureObject = Texture;
 	Prototype.TexCoordArg = TexCoord;
 	Prototype.MipValueArg = nullptr;
 	Prototype.AutomaticMipBiasArg = nullptr;
diff --git a/Engine/Source/Runtime/Engine/Private/Materials/MaterialIRModuleBuilderAnalysis.cpp b/Engine/Source/Runtime/Engine/Private/Materials/MaterialIRModuleBuilderAnalysis.cpp
--- a/Engine/Source/Runtime/Engine/Private/Materials/MaterialIRModuleBuilderAnalysis.cpp
+++ b/Engine/Source/Runtime/Engine/Private/Materials/MaterialIRModuleBuilderAnalysis.cpp
@@ -23,22 +23,36 @@ void AnalyzeExternalInput(FAnalysisData& Data, MIR::FExternalInput* ExternalInpu
 	 Data.Module->GetStatistics().ExternalInputUsedMask[SF_Pixel][(int)ExternalInput->Id] = true;
 }
 
-void AnalyzeTextureSample(FAnalysisData& Data, MIR::FTextureSample* TextureSample)
+void AnalyzeTextureObject(FAnalysisData& Data, MIR::FTextureObject* TextureObject)
 {
-	EMaterialTextureParameterType ParamType = MIR::Internal::TextureMaterialValueTypeToParameterType(TextureSample->Texture->GetMaterialType());
+	EMaterialTextureParameterType ParamType = MIR::Internal::TextureMaterialValueTypeToParameterType(TextureObject->Texture->GetMaterialType());
 
 	FMaterialTextureParameterInfo ParamInfo {
 		.ParameterInfo = { "", EMaterialParameterAssociation::GlobalParameter, INDEX_NONE },
-		.TextureIndex = Data.Material->GetReferencedTextures().Find(TextureSample->Texture),
+		.TextureIndex = Data.Material->GetReferencedTextures().Find(TextureObject->Texture),
 		.SamplerSource = SSM_FromTextureAsset, // TODO - Is this needed?
 	};
 			
 	check(ParamInfo.TextureIndex != INDEX_NONE);
 
-	TextureSample->TextureParameterIndex = Data.CompilationOutput->UniformExpressionSet.FindOrAddTextureParameter(ParamType, ParamInfo);
+	TextureObject->Analysis_UniformParameterIndex = Data.CompilationOutput->UniformExpressionSet.FindOrAddTextureParameter(ParamType, ParamInfo);
 }
 
-void AnalyzeUniformParameter(FAnalysisData& Data, MIR::FUniformParameter* Parameter)
+static void AnalyzeTextureUniformParameter(FAnalysisData& Data, MIR::FUniformParameter* Parameter)
+{
+	EMaterialTextureParameterType ParamType = MIR::Internal::TextureMaterialValueTypeToParameterType(Parameter->Metadata.Value.Texture->GetMaterialType());
+	
+	FMaterialTextureParameterInfo ParamInfo {
+		.ParameterInfo = Parameter->Info,
+		.TextureIndex = Data.Material->GetReferencedTextures().Find(Parameter->Metadata.Value.Texture),
+		.SamplerSource = SSM_FromTextureAsset, // TODO - Is this needed?
+		.VirtualTextureLayerIndex = 0xff,
+	};
+	
+	Parameter->Analysis_UniformParameterIndex = Data.CompilationOutput->UniformExpressionSet.FindOrAddTextureParameter(ParamType, ParamInfo);
+}
+
+static void AnalyzePrimitiveUniformParameter(FAnalysisData& Data, MIR::FUniformParameter* Parameter)
 {
 	FUniformExpressionSet& UniformExpressionSet = Data.CompilationOutput->UniformExpressionSet;
 
@@ -59,13 +73,13 @@ void AnalyzeUniformParameter(FAnalysisData& Data, MIR::FUniformParameter* Parame
 		Data.UniformDefaultValueOffsets.Add(DefaultValue, DefaultValueOffset);
 	}
 
-	Parameter->Index = UniformExpressionSet.FindOrAddNumericParameter(Parameter->Metadata.Value.Type, Parameter->Info, DefaultValueOffset);
+	Parameter->Analysis_UniformParameterIndex = UniformExpressionSet.FindOrAddNumericParameter(Parameter->Metadata.Value.Type, Parameter->Info, DefaultValueOffset);
 
 	// Get the parameter type as primitive
 	MIR::FPrimitiveTypePtr ParameterType = Parameter->Type->AsPrimitive();
 
 	// Only int and float parameters supported for now (TODO)
-	check(ParameterType->ScalarKind == MIR::SK_Int || ParameterType->ScalarKind == MIR::SK_Float);
+	check(ParameterType->ScalarKind == MIR::ScalarKind_Int || ParameterType->ScalarKind == MIR::ScalarKind_Float);
 
 	// Depending on the number of compoents this parameter has, try to find a previous allocation that may fit it.
 	uint32 ParameterUniformBufferOffset = UINT32_MAX;
@@ -107,7 +121,7 @@ void AnalyzeUniformParameter(FAnalysisData& Data, MIR::FUniformParameter* Parame
 	}
 
 	// Add the parameter evaluation to the uniform data
-	UniformExpressionSet.AddNumericParameterEvaluation(Parameter->Index, ParameterUniformBufferOffset);
+	UniformExpressionSet.AddNumericParameterEvaluation(Parameter->Analysis_UniformParameterIndex, ParameterUniformBufferOffset);
 
 	if (Data.Insights)
 	{
@@ -132,8 +146,8 @@ void AnalyzeUniformParameter(FAnalysisData& Data, MIR::FUniformParameter* Parame
 
 		switch (ParameterType->ScalarKind)
 		{
-			case MIR::SK_Int: ParamInsight.ComponentType = FMaterialInsights::FUniformBufferSlotComponentType::CT_Int; break;
-			case MIR::SK_Float: ParamInsight.ComponentType = FMaterialInsights::FUniformBufferSlotComponentType::CT_Float; break;
+			case MIR::ScalarKind_Int: ParamInsight.ComponentType = FMaterialInsights::FUniformBufferSlotComponentType::CT_Int; break;
+			case MIR::ScalarKind_Float: ParamInsight.ComponentType = FMaterialInsights::FUniformBufferSlotComponentType::CT_Float; break;
 			default: UE_MIR_UNREACHABLE();
 		}
 
@@ -141,6 +155,19 @@ void AnalyzeUniformParameter(FAnalysisData& Data, MIR::FUniformParameter* Parame
 	}
 }
 
+void AnalyzeUniformParameter(FAnalysisData& Data, MIR::FUniformParameter* Parameter)
+{
+	if (Parameter->Type->IsObjectOfKind(MIR::ObjectKind_Texture2D))
+	{
+		AnalyzeTextureUniformParameter(Data, Parameter);
+	}
+	else
+	{
+		check(Parameter->Type->AsPrimitive());
+		AnalyzePrimitiveUniformParameter(Data, Parameter);
+	}
+}
+
 void AnalyzeValue(FAnalysisData& Data, MIR::FValue* Value)
 {
 	#define EXPAND_CASE(ValueType) \
@@ -153,7 +180,7 @@ void AnalyzeValue(FAnalysisData& Data, MIR::FValue* Value)
 	switch (Value->Kind)
 	{
 		EXPAND_CASE(ExternalInput);
-		EXPAND_CASE(TextureSample);
+		EXPAND_CASE(TextureObject);
 		EXPAND_CASE(UniformParameter);
 		default: break;
 	}
diff --git a/Engine/Source/Runtime/Engine/Private/Materials/MaterialIRToHLSLTranslator.cpp b/Engine/Source/Runtime/Engine/Private/Materials/MaterialIRToHLSLTranslator.cpp
--- a/Engine/Source/Runtime/Engine/Private/Materials/MaterialIRToHLSLTranslator.cpp
+++ b/Engine/Source/Runtime/Engine/Private/Materials/MaterialIRToHLSLTranslator.cpp
@@ -302,7 +302,8 @@ struct FTranslator : FMaterialIRToHLSLTranslation
 		{
 			case MIR::VK_Constant: LowerConstant(static_cast<const MIR::FConstant*>(InValue)); break;
 			case MIR::VK_ExternalInput: LowerExternalInput(static_cast<const MIR::FExternalInput*>(InValue)); break;
-			case MIR::VK_UniformParameter: LowerParameter(static_cast<const MIR::FUniformParameter*>(InValue)); break;
+			case MIR::VK_TextureObject: LowerTextureObject(static_cast<const MIR::FTextureObject*>(InValue)); break;
+			case MIR::VK_UniformParameter: LowerUniformParameter(static_cast<const MIR::FUniformParameter*>(InValue)); break;
 			default: UE_MIR_UNREACHABLE();
 		}
 
@@ -334,9 +335,9 @@ struct FTranslator : FMaterialIRToHLSLTranslation
 
 		switch (checkNoEntry->ScalarKind)
 		{
-			case MIR::SK_Bool:  Printer.Buffer.Append(Constant->Boolean ? TEXT("true") : TEXT("false")); break;
-			case MIR::SK_Int:   Printer.Buffer.Appendf(TEXT("%") PRId64, Constant->Integer); break;
-			case MIR::SK_Float: Printer.Buffer.Appendf(TEXT("%.5ff"), Constant->Float); break;
+			case MIR::ScalarKind_Bool:  Printer.Buffer.Append(Constant->Boolean ? TEXT("true") : TEXT("false")); break;
+			case MIR::ScalarKind_Int:   Printer.Buffer.Appendf(TEXT("%") PRId64, Constant->Integer); break;
+			case MIR::ScalarKind_Float: Printer.Buffer.Appendf(TEXT("%.5ff"), Constant->Float); break;
 		}
 	}
 
@@ -365,14 +366,43 @@ struct FTranslator : FMaterialIRToHLSLTranslation
 		}
 	}
 
-	void LowerParameter(const MIR::FUniformParameter* Parameter)
+	void LowerTextureObject(const MIR::FTextureObject* TextureObject)
 	{
-		Printer << TEXT("Material.PreshaderBuffer[") << Parameter->Index << TEXT("].");
-		MIR::FPrimitiveTypePtr PrimitiveType = Parameter->Type->AsPrimitive();
+		LowerTextureReference(TextureObject->Texture->GetMaterialType(), TextureObject->Analysis_UniformParameterIndex);
+	}
+
+	void LowerUniformParameter(const MIR::FUniformParameter* UniformParameter)
+	{
+		if (UniformParameter->Type->IsObjectOfKind(MIR::ObjectKind_Texture2D))
+		{
+			LowerTextureReference(UniformParameter->Metadata.Value.Texture->GetMaterialType(), UniformParameter->Analysis_UniformParameterIndex);
+		}
+		else
+		{
+			LowerPrimitiveUniformParameter(UniformParameter);
+		}
+	}
+
+	void LowerPrimitiveUniformParameter(const MIR::FUniformParameter* UniformParameter)
+	{
+		MIR::FPrimitiveTypePtr PrimitiveType = UniformParameter->Type->AsPrimitive();
+		check(PrimitiveType);
+
+		if (PrimitiveType->ScalarKind == MIR::EScalarKind::ScalarKind_Int)
+		{
+			Printer << "asint(";
+		}
+
+		Printer << TEXT("Material.PreshaderBuffer[") << UniformParameter->Analysis_UniformParameterIndex << TEXT("].");
 		check(PrimitiveType->IsScalar() || PrimitiveType->IsVector()); // no matrices yet
 		
 		static const TCHAR* Suffixes[] = { TEXT("x"), TEXT("xy"), TEXT("xyz"), TEXT("xyzw") };
 		Printer << Suffixes[PrimitiveType->NumRows - 1];
+
+		if (PrimitiveType->ScalarKind != MIR::EScalarKind::ScalarKind_Float)
+		{
+			Printer << ")";
+		}
 	}
 
 	void LowerDimensional(const MIR::FDimensional* Dimensional) 
@@ -488,9 +518,9 @@ struct FTranslator : FMaterialIRToHLSLTranslation
 			Printer << TEXT("(");
 		}
 
-		switch (TextureSample->Texture->GetMaterialType())
+		switch (TextureSample->TextureObject->Type->AsObject()->ObjectKind)
 		{
-			case MCT_Texture2D:
+			case MIR::ObjectKind_Texture2D:
 				Printer << TEXT("Texture2DSample");
 				break;
 
@@ -499,8 +529,8 @@ struct FTranslator : FMaterialIRToHLSLTranslation
 		}
 
 		Printer << BeginArgs
-			<< ListSeparator << LowerTextureReference(TextureSample->Texture->GetMaterialType(), TextureSample->TextureParameterIndex)
-			<< ListSeparator << LowerTextureSamplerReference(TextureSample->SamplerSourceMode, TextureSample->Texture->GetMaterialType(), TextureSample->TextureParameterIndex)
+			<< ListSeparator << LowerValue(TextureSample->TextureObject)
+			<< ListSeparator << LowerTextureSamplerReference(TextureSample->TextureObject, TextureSample->SamplerSourceMode)
 			<< ListSeparator << LowerValue(TextureSample->TexCoordArg)
 			<< EndArgs;
 
@@ -565,12 +595,12 @@ struct FTranslator : FMaterialIRToHLSLTranslation
 		return true;
 	}
 
-	ENoOp LowerTextureSamplerReference(ESamplerSourceMode SamplerSource, EMaterialValueType TextureType, int TextureParameterIndex)
+	ENoOp LowerTextureSamplerReference(MIR::FValue* TextureValue, ESamplerSourceMode SamplerSource)
 	{
 		switch (SamplerSource)
 		{
 		case SSM_FromTextureAsset:
-			LowerTextureReference(TextureType, TextureParameterIndex);
+			LowerValue(TextureValue);
 			Printer << TEXT("Sampler");
 			break;
 
@@ -678,9 +708,9 @@ struct FTranslator : FMaterialIRToHLSLTranslation
 		{
 			switch (ArithmeticType->ScalarKind)
 			{
-				case MIR::SK_Bool:	Printer << TEXT("bool"); break;
-				case MIR::SK_Int: 	Printer << TEXT("int"); break;
-				case MIR::SK_Float:	Printer << TEXT("float"); break;
+				case MIR::ScalarKind_Bool:	Printer << TEXT("bool"); break;
+				case MIR::ScalarKind_Int: 	Printer << TEXT("int"); break;
+				case MIR::ScalarKind_Float:	Printer << TEXT("float"); break;
 			}
 
 			if (ArithmeticType->NumRows > 1)
diff --git a/Engine/Source/Runtime/Engine/Private/Materials/MaterialIRTypes.cpp b/Engine/Source/Runtime/Engine/Private/Materials/MaterialIRTypes.cpp
--- a/Engine/Source/Runtime/Engine/Private/Materials/MaterialIRTypes.cpp
+++ b/Engine/Source/Runtime/Engine/Private/Materials/MaterialIRTypes.cpp
@@ -11,8 +11,8 @@ const TCHAR* TypeKindToString(ETypeKind Kind)
 {
 	switch (Kind)
 	{
-		case TK_Void: return TEXT("void");
-		case TK_Primitive: return TEXT("primitive");
+		case TypeKind_Void: return TEXT("void");
+		case TypeKind_Primitive: return TEXT("primitive");
 		default: UE_MIR_UNREACHABLE();
 	}
 }
@@ -31,19 +31,19 @@ FTypePtr FType::FromShaderType(const UE::Shader::FType& InShaderType)
 		case UE::Shader::EValueType::Float2:
 		case UE::Shader::EValueType::Float3:
 		case UE::Shader::EValueType::Float4:
-			return FPrimitiveType::GetVector(SK_Float, (int)InShaderType.ValueType - (int)UE::Shader::EValueType::Float1 + 1);
+			return FPrimitiveType::GetVector(ScalarKind_Float, (int)InShaderType.ValueType - (int)UE::Shader::EValueType::Float1 + 1);
 
 		case UE::Shader::EValueType::Int1:
 		case UE::Shader::EValueType::Int2:
 		case UE::Shader::EValueType::Int3:
 		case UE::Shader::EValueType::Int4:
-			return FPrimitiveType::GetVector(SK_Int, (int)InShaderType.ValueType - (int)UE::Shader::EValueType::Int1 + 1);
+			return FPrimitiveType::GetVector(ScalarKind_Int, (int)InShaderType.ValueType - (int)UE::Shader::EValueType::Int1 + 1);
 
 		case UE::Shader::EValueType::Bool1:
 		case UE::Shader::EValueType::Bool2:
 		case UE::Shader::EValueType::Bool3:
 		case UE::Shader::EValueType::Bool4:
-			return FPrimitiveType::GetVector(SK_Bool, (int)InShaderType.ValueType - (int)UE::Shader::EValueType::Int1 + 1);
+			return FPrimitiveType::GetVector(ScalarKind_Bool, (int)InShaderType.ValueType - (int)UE::Shader::EValueType::Int1 + 1);
 
 		default:
 			UE_MIR_UNREACHABLE();
@@ -61,19 +61,19 @@ FTypePtr FType::FromMaterialValueType(EMaterialValueType Type)
 		case MCT_Float2:
 		case MCT_Float3:
 		case MCT_Float4:
-			return FPrimitiveType::GetVector(SK_Float, (int)Type - (int)MCT_Float1 + 1);
+			return FPrimitiveType::GetVector(ScalarKind_Float, (int)Type - (int)MCT_Float1 + 1);
 
 		case MCT_Float:
-			return FPrimitiveType::GetVector(SK_Float, 4);
+			return FPrimitiveType::GetVector(ScalarKind_Float, 4);
 
 		case MCT_UInt1:
 		case MCT_UInt2:
 		case MCT_UInt3:
 		case MCT_UInt4:
-			return FPrimitiveType::GetVector(SK_Int, (int)Type - (int)MCT_UInt1 + 1);
+			return FPrimitiveType::GetVector(ScalarKind_Int, (int)Type - (int)MCT_UInt1 + 1);
 
 		case MCT_Bool:
-			return FPrimitiveType::GetVector(SK_Bool, (int)Type - (int)UE::Shader::EValueType::Int1 + 1);
+			return FPrimitiveType::GetVector(ScalarKind_Bool, (int)Type - (int)UE::Shader::EValueType::Int1 + 1);
 
 		default:
 			UE_MIR_UNREACHABLE();
@@ -99,7 +99,7 @@ FTypePtr FType::FromMaterialParameterType(EMaterialParameterType Type)
 
 FTypePtr FType::GetVoid()
 {
-	static FType Type{ TK_Void };
+	static FType Type{ TypeKind_Void };
 	return &Type;
 }
 
@@ -113,6 +113,13 @@ FStringView FType::GetSpelling() const
 	{
 		return PrimitiveType->Spelling;
 	}
+	else if (auto ObjectType = AsObject())
+	{
+		switch (ObjectType->ObjectKind)
+		{
+			case ObjectKind_Texture2D: return TEXT("Texture2D");
+		}
+	}
 	UE_MIR_UNREACHABLE();
 }
 
@@ -126,7 +133,7 @@ UE::Shader::EValueType FType::ToValueType() const
 		{
 			if (PrimitiveType->NumRows == 4 && PrimitiveType->NumColumns == 4)
 			{
-				if (PrimitiveType->ScalarKind == SK_Float)
+				if (PrimitiveType->ScalarKind == ScalarKind_Float)
 				{
 					return EValueType::Float4x4;
 				}
@@ -143,9 +150,9 @@ UE::Shader::EValueType FType::ToValueType() const
 
 		switch (PrimitiveType->ScalarKind)
 		{
-			case SK_Bool: 	return (EValueType)((int)EValueType::Bool1 + PrimitiveType->NumRows - 1); break;
-			case SK_Int: 	return (EValueType)((int)EValueType::Int1 + PrimitiveType->NumRows - 1); break;
-			case SK_Float: 	return (EValueType)((int)EValueType::Float1 + PrimitiveType->NumRows - 1); break;
+			case ScalarKind_Bool: 	return (EValueType)((int)EValueType::Bool1 + PrimitiveType->NumRows - 1); break;
+			case ScalarKind_Int: 	return (EValueType)((int)EValueType::Int1 + PrimitiveType->NumRows - 1); break;
+			case ScalarKind_Float: 	return (EValueType)((int)EValueType::Float1 + PrimitiveType->NumRows - 1); break;
 			default: UE_MIR_UNREACHABLE();
 		}
 	}
@@ -159,25 +166,25 @@ UE::Shader::EValueType FType::ToValueType() const
 
 FTypePtr FType::GetPoison()
 {
-	static FPoisonType Type{ TK_Poison };
+	static FType Type{ TypeKind_Poison };
 	return &Type;
 }
 
 bool FType::IsPoison() const
 {
-	return Kind == TK_Poison;
+	return Kind == TypeKind_Poison;
 }
 
 bool FType::IsBoolean() const
 {
 	FPrimitiveTypePtr Ptr = AsPrimitive();
-	return Ptr && Ptr->ScalarKind == SK_Bool;
+	return Ptr && Ptr->ScalarKind == ScalarKind_Bool;
 }
 
 bool FType::IsInteger() const
 {
 	FPrimitiveTypePtr Ptr = AsPrimitive();
-	return Ptr && Ptr->ScalarKind == SK_Int;
+	return Ptr && Ptr->ScalarKind == ScalarKind_Int;
 }
 
 bool FType::IsBoolScalar() const
@@ -185,14 +192,32 @@ bool FType::IsBoolScalar() const
 	return this == FPrimitiveType::GetBool1();
 }
 
+bool FType::IsObjectOfKind(EObjectKind ObjectKind) const
+{
+	if (FObjectTypePtr ObjectPtr = AsObject())
+	{
+		return ObjectPtr->ObjectKind == ObjectKind;
+	}
+	return false;
+}
+
+bool FType::IsTexture() const
+{
+	if (FObjectTypePtr ObjectPtr = AsObject())
+	{
+		return true;
+	}
+	return false; 
+}
+
 FPrimitiveTypePtr FType::AsPrimitive() const
 {
-	return Kind == TK_Primitive ? static_cast<FPrimitiveTypePtr>(this) : nullptr; 
+	return Kind == TypeKind_Primitive ? static_cast<FPrimitiveTypePtr>(this) : nullptr; 
 }
 
 FObjectTypePtr FType::AsObject() const
 {
-	return Kind == TK_Object ? static_cast<FObjectTypePtr>(this) : nullptr; 
+	return Kind == TypeKind_Object ? static_cast<FObjectTypePtr>(this) : nullptr; 
 }
 
 FPrimitiveTypePtr FType::AsScalar() const
@@ -217,41 +242,41 @@ const TCHAR* ScalarKindToString(EScalarKind Kind)
 {
 	switch (Kind)
 	{
-		case SK_Bool: return TEXT("bool"); break;
-		case SK_Int: return TEXT("int"); break;
-		case SK_Float: return TEXT("float"); break;
+		case ScalarKind_Bool: return TEXT("bool"); break;
+		case ScalarKind_Int: return TEXT("int"); break;
+		case ScalarKind_Float: return TEXT("float"); break;
 		default: UE_MIR_UNREACHABLE();
 	}
 }
 
 FPrimitiveTypePtr FPrimitiveType::GetBool1()
 {
-	return GetScalar(SK_Bool);
+	return GetScalar(ScalarKind_Bool);
 }
 
 FPrimitiveTypePtr FPrimitiveType::GetInt1()
 {
-	return GetScalar(SK_Int);
+	return GetScalar(ScalarKind_Int);
 }
 
 FPrimitiveTypePtr FPrimitiveType::GetFloat1()
 {
-	return GetScalar(SK_Float);
+	return GetScalar(ScalarKind_Float);
 }
 
 FPrimitiveTypePtr FPrimitiveType::GetFloat2()
 {
-	return GetVector(SK_Float, 2);
+	return GetVector(ScalarKind_Float, 2);
 }
 
 FPrimitiveTypePtr FPrimitiveType::GetFloat3()
 {
-	return GetVector(SK_Float, 3);
+	return GetVector(ScalarKind_Float, 3);
 }
 
 FPrimitiveTypePtr FPrimitiveType::GetFloat4()
 {
-	return GetVector(SK_Float, 4);
+	return GetVector(ScalarKind_Float, 4);
 }
 
 const FPrimitiveType* FPrimitiveType::GetScalar(EScalarKind InScalarKind)
@@ -274,59 +299,59 @@ const FPrimitiveType* FPrimitiveType::GetMatrix(EScalarKind InScalarKind, int Nu
 
 FPrimitiveTypePtr FPrimitiveType::Get(EScalarKind InScalarKind, int NumRows, int NumColumns)
 {
-	check(InScalarKind >= 0 && InScalarKind <= SK_Float);
+	check(InScalarKind >= 0 && InScalarKind <= ScalarKind_Float);
 	
 	static const FStringView Invalid = TEXT("invalid");
 
 	static const FPrimitiveType Types[] {
-		{ { TK_Primitive }, { TEXT("bool") }, 		SK_Bool, 1, 1 },
-		{ { TK_Primitive }, Invalid, 				SK_Bool, 1, 2 }, 
-		{ { TK_Primitive }, Invalid, 				SK_Bool, 1, 3 },
-		{ { TK_Primitive }, Invalid, 				SK_Bool, 1, 4 },
-		{ { TK_Primitive }, { TEXT("bool2") },   	SK_Bool, 2, 1 },
-		{ { TK_Primitive }, { TEXT("bool2x2") }, 	SK_Bool, 2, 2 },
-		{ { TK_Primitive }, { TEXT("bool2x3") }, 	SK_Bool, 2, 3 },
-		{ { TK_Primitive }, { TEXT("bool2x4") }, 	SK_Bool, 2, 4 },
-		{ { TK_Primitive }, { TEXT("bool3") },   	SK_Bool, 3, 1 },
-		{ { TK_Primitive }, { TEXT("bool3x2") }, 	SK_Bool, 3, 2 },
-		{ { TK_Primitive }, { TEXT("bool3x3") }, 	SK_Bool, 3, 3 },
-		{ { TK_Primitive }, { TEXT("bool3x4") }, 	SK_Bool, 3, 4 },
-		{ { TK_Primitive }, { TEXT("bool4") },   	SK_Bool, 4, 1 },
-		{ { TK_Primitive }, { TEXT("bool4x2") }, 	SK_Bool, 4, 2 },
-		{ { TK_Primitive }, { TEXT("bool4x3") }, 	SK_Bool, 4, 3 },
-		{ { TK_Primitive }, { TEXT("bool4x4") }, 	SK_Bool, 4, 4 },
-		{ { TK_Primitive }, { TEXT("int") }, 		SK_Int, 1, 1 },
-		{ { TK_Primitive }, Invalid, 				SK_Int, 1, 2 },
-		{ { TK_Primitive }, Invalid, 				SK_Int, 1, 3 },
-		{ { TK_Primitive }, Invalid, 				SK_Int, 1, 4 },
-		{ { TK_Primitive }, { TEXT("int2") },   	SK_Int, 2, 1 },
-		{ { TK_Primitive }, { TEXT("int2x2") }, 	SK_Int, 2, 2 },
-		{ { TK_Primitive }, { TEXT("int2x3") }, 	SK_Int, 2, 3 },
-		{ { TK_Primitive }, { TEXT("int2x4") }, 	SK_Int, 2, 4 },
-		{ { TK_Primitive }, { TEXT("int3") },   	SK_Int, 3, 1 },
-		{ { TK_Primitive }, { TEXT("int3x2") }, 	SK_Int, 3, 2 },
-		{ { TK_Primitive }, { TEXT("int3x3") }, 	SK_Int, 3, 3 },
-		{ { TK_Primitive }, { TEXT("int3x4") }, 	SK_Int, 3, 4 },
-		{ { TK_Primitive }, { TEXT("int4") },   	SK_Int, 4, 1 },
-		{ { TK_Primitive }, { TEXT("int4x2") }, 	SK_Int, 4, 2 },
-		{ { TK_Primitive }, { TEXT("int4x3") }, 	SK_Int, 4, 3 },
-		{ { TK_Primitive }, { TEXT("int4x4") }, 	SK_Int, 4, 4 },
-		{ { TK_Primitive }, { TEXT("float") }, 		SK_Float, 1, 1 },
-		{ { TK_Primitive }, Invalid, 				SK_Float, 1, 2 },
-		{ { TK_Primitive }, Invalid, 				SK_Float, 1, 3 },
-		{ { TK_Primitive }, Invalid, 				SK_Float, 1, 4 },
-		{ { TK_Primitive }, { TEXT("float2") },   	SK_Float, 2, 1 },
-		{ { TK_Primitive }, { TEXT("float2x2") }, 	SK_Float, 2, 2 },
-		{ { TK_Primitive }, { TEXT("float2x3") }, 	SK_Float, 2, 3 },
-		{ { TK_Primitive }, { TEXT("float2x4") }, 	SK_Float, 2, 4 },
-		{ { TK_Primitive }, { TEXT("float3") },   	SK_Float, 3, 1 },
-		{ { TK_Primitive }, { TEXT("float3x2") }, 	SK_Float, 3, 2 },
-		{ { TK_Primitive }, { TEXT("float3x3") }, 	SK_Float, 3, 3 },
-		{ { TK_Primitive }, { TEXT("float3x4") }, 	SK_Float, 3, 4 },
-		{ { TK_Primitive }, { TEXT("float4") },   	SK_Float, 4, 1 },
-		{ { TK_Primitive }, { TEXT("float4x2") }, 	SK_Float, 4, 2 },
-		{ { TK_Primitive }, { TEXT("float4x3") }, 	SK_Float, 4, 3 },
-		{ { TK_Primitive }, { TEXT("float4x4") }, 	SK_Float, 4, 4 },
+		{ { TypeKind_Primitive }, { TEXT("bool") }, 		ScalarKind_Bool, 1, 1 },
+		{ { TypeKind_Primitive }, Invalid, 				ScalarKind_Bool, 1, 2 }, 
+		{ { TypeKind_Primitive }, Invalid, 				ScalarKind_Bool, 1, 3 },
+		{ { TypeKind_Primitive }, Invalid, 				ScalarKind_Bool, 1, 4 },
+		{ { TypeKind_Primitive }, { TEXT("bool2") },   	ScalarKind_Bool, 2, 1 },
+		{ { TypeKind_Primitive }, { TEXT("bool2x2") }, 	ScalarKind_Bool, 2, 2 },
+		{ { TypeKind_Primitive }, { TEXT("bool2x3") }, 	ScalarKind_Bool, 2, 3 },
+		{ { TypeKind_Primitive }, { TEXT("bool2x4") }, 	ScalarKind_Bool, 2, 4 },
+		{ { TypeKind_Primitive }, { TEXT("bool3") },   	ScalarKind_Bool, 3, 1 },
+		{ { TypeKind_Primitive }, { TEXT("bool3x2") }, 	ScalarKind_Bool, 3, 2 },
+		{ { TypeKind_Primitive }, { TEXT("bool3x3") }, 	ScalarKind_Bool, 3, 3 },
+		{ { TypeKind_Primitive }, { TEXT("bool3x4") }, 	ScalarKind_Bool, 3, 4 },
+		{ { TypeKind_Primitive }, { TEXT("bool4") },   	ScalarKind_Bool, 4, 1 },
+		{ { TypeKind_Primitive }, { TEXT("bool4x2") }, 	ScalarKind_Bool, 4, 2 },
+		{ { TypeKind_Primitive }, { TEXT("bool4x3") }, 	ScalarKind_Bool, 4, 3 },
+		{ { TypeKind_Primitive }, { TEXT("bool4x4") }, 	ScalarKind_Bool, 4, 4 },
+		{ { TypeKind_Primitive }, { TEXT("int") }, 		ScalarKind_Int, 1, 1 },
+		{ { TypeKind_Primitive }, Invalid, 				ScalarKind_Int, 1, 2 },
+		{ { TypeKind_Primitive }, Invalid, 				ScalarKind_Int, 1, 3 },
+		{ { TypeKind_Primitive }, Invalid, 				ScalarKind_Int, 1, 4 },
+		{ { TypeKind_Primitive }, { TEXT("int2") },   	ScalarKind_Int, 2, 1 },
+		{ { TypeKind_Primitive }, { TEXT("int2x2") }, 	ScalarKind_Int, 2, 2 },
+		{ { TypeKind_Primitive }, { TEXT("int2x3") }, 	ScalarKind_Int, 2, 3 },
+		{ { TypeKind_Primitive }, { TEXT("int2x4") }, 	ScalarKind_Int, 2, 4 },
+		{ { TypeKind_Primitive }, { TEXT("int3") },   	ScalarKind_Int, 3, 1 },
+		{ { TypeKind_Primitive }, { TEXT("int3x2") }, 	ScalarKind_Int, 3, 2 },
+		{ { TypeKind_Primitive }, { TEXT("int3x3") }, 	ScalarKind_Int, 3, 3 },
+		{ { TypeKind_Primitive }, { TEXT("int3x4") }, 	ScalarKind_Int, 3, 4 },
+		{ { TypeKind_Primitive }, { TEXT("int4") },   	ScalarKind_Int, 4, 1 },
+		{ { TypeKind_Primitive }, { TEXT("int4x2") }, 	ScalarKind_Int, 4, 2 },
+		{ { TypeKind_Primitive }, { TEXT("int4x3") }, 	ScalarKind_Int, 4, 3 },
+		{ { TypeKind_Primitive }, { TEXT("int4x4") }, 	ScalarKind_Int, 4, 4 },
+		{ { TypeKind_Primitive }, { TEXT("float") }, 		ScalarKind_Float, 1, 1 },
+		{ { TypeKind_Primitive }, Invalid, 				ScalarKind_Float, 1, 2 },
+		{ { TypeKind_Primitive }, Invalid, 				ScalarKind_Float, 1, 3 },
+		{ { TypeKind_Primitive }, Invalid, 				ScalarKind_Float, 1, 4 },
+		{ { TypeKind_Primitive }, { TEXT("float2") },   	ScalarKind_Float, 2, 1 },
+		{ { TypeKind_Primitive }, { TEXT("float2x2") }, 	ScalarKind_Float, 2, 2 },
+		{ { TypeKind_Primitive }, { TEXT("float2x3") }, 	ScalarKind_Float, 2, 3 },
+		{ { TypeKind_Primitive }, { TEXT("float2x4") }, 	ScalarKind_Float, 2, 4 },
+		{ { TypeKind_Primitive }, { TEXT("float3") },   	ScalarKind_Float, 3, 1 },
+		{ { TypeKind_Primitive }, { TEXT("float3x2") }, 	ScalarKind_Float, 3, 2 },
+		{ { TypeKind_Primitive }, { TEXT("float3x3") }, 	ScalarKind_Float, 3, 3 },
+		{ { TypeKind_Primitive }, { TEXT("float3x4") }, 	ScalarKind_Float, 3, 4 },
+		{ { TypeKind_Primitive }, { TEXT("float4") },   	ScalarKind_Float, 4, 1 },
+		{ { TypeKind_Primitive }, { TEXT("float4x2") }, 	ScalarKind_Float, 4, 2 },
+		{ { TypeKind_Primitive }, { TEXT("float4x3") }, 	ScalarKind_Float, 4, 3 },
+		{ { TypeKind_Primitive }, { TEXT("float4x4") }, 	ScalarKind_Float, 4, 4 },
 	};
 
 	int Index = InScalarKind * 4 * 4 + (NumRows - 1) * 4 + (NumColumns - 1);
@@ -341,7 +366,7 @@ FPrimitiveTypePtr FPrimitiveType::ToScalar() const
 
 FTypePtr FObjectType::GetTexture2D()
 {
-	static FObjectType Instance { TK_Object };
+	static const FObjectType Instance { { TypeKind_Object }, ObjectKind_Texture2D };
 	return &Instance;
 }
 
diff --git a/Engine/Source/Runtime/Engine/Public/Materials/MaterialExpressionTextureObjectParameter.h b/Engine/Source/Runtime/Engine/Public/Materials/MaterialExpressionTextureObjectParameter.h
--- a/Engine/Source/Runtime/Engine/Public/Materials/MaterialExpressionTextureObjectParameter.h
+++ b/Engine/Source/Runtime/Engine/Public/Materials/MaterialExpressionTextureObjectParameter.h
@@ -26,6 +26,7 @@ class UMaterialExpressionTextureObjectParameter : public UMaterialExpressionText
 
 	//~ Begin UMaterialExpression Interface
 #if WITH_EDITOR
+	virtual void Build(UE::MIR::FEmitter& Emitter) override;
 	virtual void GetCaption(TArray<FString>& OutCaptions) const override;	
 	virtual int32 Compile(class FMaterialCompiler* Compiler, int32 OutputIndex) override;
 	virtual int32 CompilePreview(class FMaterialCompiler* Compiler, int32 OutputIndex) override;
diff --git a/Engine/Source/Runtime/Engine/Public/Materials/MaterialIR.h b/Engine/Source/Runtime/Engine/Public/Materials/MaterialIR.h
--- a/Engine/Source/Runtime/Engine/Public/Materials/MaterialIR.h
+++ b/Engine/Source/Runtime/Engine/Public/Materials/MaterialIR.h
@@ -226,6 +226,9 @@ struct FTextureObject : TValue<VK_TextureObject>
 {
 	UTexture* Texture;
 	EMaterialSamplerType SamplerType;
+
+	/// Index of this parameter in the uniform expression set
+	uint32 Analysis_UniformParameterIndex;
 };
 
 //
@@ -233,9 +236,13 @@ struct FUniformParameter: TValue<VK_UniformParameter>
 {
 	FMaterialParameterInfo Info;
 	FMaterialParameterMetadata Metadata;
-	int Index;
+	EMaterialSamplerType SamplerType;
+
+	/// Index of this parameter in the uniform expression set
+	uint32 Analysis_UniformParameterIndex;
 };
 
+
 /* Instructions */
 
 //
@@ -270,13 +277,13 @@ struct FDimensional : TInstruction<VK_Dimensional>
 {
 	static constexpr int MaxNumComponents = 16;
 
-	// Returns the constant array of component values. 
+	/// Returns the constant array of component values. 
 	TArrayView<FValue* const> GetComponents() const;
 
-	// Returns the mutable array of component values. 
+	/// Returns the mutable array of component values. 
 	TArrayView<FValue*> GetComponents();
 
-	// Returns whether all components are constant.
+	/// Returns whether all components are constant.
 	bool AreComponentsConstant() const;
 };
 
@@ -380,17 +387,14 @@ struct FCast : TInstruction<VK_Cast>
 //
 struct FTextureSample : TInstruction<VK_TextureSample>
 {
+	FValue* TextureObject;
 	FValue* TexCoordArg;
-	// TexCoordDerivatives;
 	FValue* MipValueArg;
 	FValue* AutomaticMipBiasArg;
-	UTexture* Texture;
 	ESamplerSourceMode SamplerSourceMode;
 	ETextureMipValueMode MipValueMode;
 	EMaterialSamplerType SamplerType;
-	
-	/* Analysis Values */
-	int TextureParameterIndex;
+	// TODO: TexCoordDerivatives;
 };
 
 } // namespace UE::MIR
diff --git a/Engine/Source/Runtime/Engine/Public/Materials/MaterialIREmitter.h b/Engine/Source/Runtime/Engine/Public/Materials/MaterialIREmitter.h
--- a/Engine/Source/Runtime/Engine/Public/Materials/MaterialIREmitter.h
+++ b/Engine/Source/Runtime/Engine/Public/Materials/MaterialIREmitter.h
@@ -66,6 +66,10 @@ public:
 	// an error and returns null.
 	FValue* Input(const FExpressionInput* Input);
 	
+	// Gets and returns the value flowing into input with specified `Input`.
+	// Before returning the value, it checks that its type is of specified `Kind`.
+	FValue* InputOfType(const FExpressionInput* Input, ETypeKind Kind);
+
 	// Gets and returns the value flowing into input with specified `Input`.
 	// Before returning the value, it checks that the value is a constant boolean scalar.
 	bool InputConstantBoolScalar(const FExpressionInput* Input);
@@ -86,9 +90,8 @@ public:
 	// Before returning the value, it checks that it has a primitive type.
 	FValue* InputPrimitive(const FExpressionInput* Input);
 
-	// Gets and returns the value flowing into input with specified `Input`.
-	// Before returning the value, it checks that its type is of specified `Kind`.
-	FValue* InputOfType(const FExpressionInput* Input, ETypeKind Kind);
+	//
+	FValue* InputTexture(const FExpressionInput* Input);
 
 	/* Output */
 	
@@ -120,6 +123,9 @@ public:
 	//
 	void CheckValueIsScalar(const FValue* Value, EScalarKind Kind);
 
+	//
+	void CheckValueIsTexture(const FValue* Value);
+
 	/* Constants */
 
 	FValue* EmitConstantFromShaderValue(const UE::Shader::FValue& InValue);
@@ -145,7 +151,7 @@ public:
 	FValue* EmitSubscript(FValue* Value, int ComponentIndex);
 	FValue* EmitSwizzle(FValue* Value, FSwizzleMask Mask);
 	FValue* EmitTextureObject(UTexture* Texture, EMaterialSamplerType SamplerType);
-	FValue* EmitParameter(FName Name, FMaterialParameterMetadata& Metadata);
+	FValue* EmitParameter(FName Name, FMaterialParameterMetadata& Metadata, EMaterialSamplerType SamplerType = SAMPLERTYPE_Color);
 	FValue* EmitCustomPrimitiveData(uint32 PrimitiveDataIndex);
 
 	/* Instructions */
@@ -154,7 +160,7 @@ public:
 	FValue* EmitOperator(EOperator Operator, FValue* A, FValue* B = nullptr, FValue* C = nullptr);
 	FValue* EmitBranch(FValue* Condition, FValue* True, FValue* False);
 	FValue* EmitConstruct(FTypePtr Type, FValue* Initializer);
-	FValue* EmitTextureSample(UTexture* Texture, FValue* TexCoord, ESamplerSourceMode SamplerSourceMode, ETextureMipValueMode MipValueMode, EMaterialSamplerType SamplerType);
+	FValue* EmitTextureSample(FValue* Texture, FValue* TexCoord, ESamplerSourceMode SamplerSourceMode, ETextureMipValueMode MipValueMode, EMaterialSamplerType SamplerType);
 
 	/* Types */
 
diff --git a/Engine/Source/Runtime/Engine/Public/Materials/MaterialIRTypes.h b/Engine/Source/Runtime/Engine/Public/Materials/MaterialIRTypes.h
--- a/Engine/Source/Runtime/Engine/Public/Materials/MaterialIRTypes.h
+++ b/Engine/Source/Runtime/Engine/Public/Materials/MaterialIRTypes.h
@@ -13,10 +13,15 @@ namespace UE::MIR
 
 enum ETypeKind
 {
-	TK_Poison,
-	TK_Void,
-	TK_Primitive,
-	TK_Object,
+	TypeKind_Poison,
+	TypeKind_Void,
+	TypeKind_Primitive,
+	TypeKind_Object,
+};
+
+enum EObjectKind 
+{
+	ObjectKind_Texture2D,
 };
 
 const TCHAR* TypeKindToString(ETypeKind Kind);
@@ -53,12 +58,15 @@ struct FType
 	// Returns whether this type is a `bool` scalar.
 	bool IsBoolScalar() const;
 
+	// Returns whether this type is of a specific object kind.
+	bool IsObjectOfKind(EObjectKind ObjectKind) const;
+
+	// Returns whether this type is a texture.
+	bool IsTexture() const;
+
 	// Returns this type upcast this type to ArithmeticType if it is one. Otherwise it returns nullptr.
 	FPrimitiveTypePtr AsPrimitive() const;
 
-	//
-	FObjectTypePtr AsObject() const;
-
 	// Returns this type upcast this type to ArithmeticType if it's a scalar. Otherwise it returns nullptr.
 	FPrimitiveTypePtr AsScalar() const;
 
@@ -70,20 +78,21 @@ struct FType
 
 	// Returns the this type name spelling (e.g. float4x4).
 	FStringView GetSpelling() const;
+	
+	//
+	FObjectTypePtr AsObject() const;
 
 	// Converts this type to a UE::Shader::EValueType.
 	UE::Shader::EValueType ToValueType() const;
 };
 
-struct FPoisonType : FType
-{
-};
-
 // Primitive types of a single scalar.
 // Note: These are listed in precision order. Converting one to the other is then simply performed taking the max EScalarKind.
 enum EScalarKind
 {
-	SK_Bool, SK_Int, SK_Float,
+	ScalarKind_Bool,
+	ScalarKind_Int,
+	ScalarKind_Float,
 };
 
 const TCHAR* ScalarKindToString(EScalarKind Kind);
@@ -116,6 +125,8 @@ struct FPrimitiveType : FType
 
 struct FObjectType : FType
 {
+	EObjectKind ObjectKind;
+
 	static FTypePtr GetTexture2D();
 };
 
