PR #12566: Gpu skinning crash when precaching PSOs

diff --git a/Engine/Source/Runtime/Engine/Private/SkeletalRenderGPUSkin.cpp b/Engine/Source/Runtime/Engine/Private/SkeletalRenderGPUSkin.cpp
--- a/Engine/Source/Runtime/Engine/Private/SkeletalRenderGPUSkin.cpp
+++ b/Engine/Source/Runtime/Engine/Private/SkeletalRenderGPUSkin.cpp
@@ -86,22 +86,25 @@ inline ESkeletalMeshGPUSkinTechnique GetGPUSkinTechnique(USkinnedMeshComponent*
 {
 	ESkeletalMeshGPUSkinTechnique GPUSkinTechnique = ESkeletalMeshGPUSkinTechnique::Inline;
 
-	if (MeshComponent->GetMeshDeformerInstanceForLOD(LODIndex) != nullptr)
+	if (MeshComponent)
 	{
-		GPUSkinTechnique = ESkeletalMeshGPUSkinTechnique::MeshDeformer;
-	}
-	else if (GEnableGPUSkinCache && MeshComponent->IsSkinCacheAllowed(LODIndex))
-	{
-		GPUSkinTechnique = ESkeletalMeshGPUSkinTechnique::GPUSkinCache;
-
-		if (FeatureLevel == ERHIFeatureLevel::ES3_1)
+		if (MeshComponent->GetMeshDeformerInstanceForLOD(LODIndex) != nullptr)
 		{
-			// Some mobile GPUs (MALI) has a 64K elements limitation on texel buffers
-			// SkinCache fetches mesh position through R32F texel buffer, thus any mesh that has more than 64K/3 vertices will not work correctly on such GPUs
-			// We force this limitation for all mobile, to have an uniform behaviour across all mobile platforms
-			if (RenderData->LODRenderData[LODIndex].GetNumVertices() * 3 >= (64 * 1024))
+			GPUSkinTechnique = ESkeletalMeshGPUSkinTechnique::MeshDeformer;
+		}
+		else if (GEnableGPUSkinCache && MeshComponent->IsSkinCacheAllowed(LODIndex))
+		{
+			GPUSkinTechnique = ESkeletalMeshGPUSkinTechnique::GPUSkinCache;
+
+			if (FeatureLevel == ERHIFeatureLevel::ES3_1)
 			{
-				GPUSkinTechnique = ESkeletalMeshGPUSkinTechnique::Inline;
+				// Some mobile GPUs (MALI) has a 64K elements limitation on texel buffers
+				// SkinCache fetches mesh position through R32F texel buffer, thus any mesh that has more than 64K/3 vertices will not work correctly on such GPUs
+				// We force this limitation for all mobile, to have an uniform behaviour across all mobile platforms
+				if (RenderData->LODRenderData[LODIndex].GetNumVertices() * 3 >= (64 * 1024))
+				{
+					GPUSkinTechnique = ESkeletalMeshGPUSkinTechnique::Inline;
+				}
 			}
 		}
 	}
@@ -2033,7 +2036,6 @@ void FSkeletalMeshObjectGPUSkin::GetUsedVertexFactoryData(
 		CompLODInfo = &SkinnedMeshComponent->LODInfo[LODIndex];
 	}
 
-	check(SkinnedMeshComponent);
 	const ESkeletalMeshGPUSkinTechnique GPUSkinTechnique = ::GetGPUSkinTechnique(SkinnedMeshComponent, SkelMeshRenderData, LODIndex, InFeatureLevel);
 
 	// Setup tmp MeshObjectLOD object to extract the vertex factory buffers
