[HWRT] Simplify some logic in FInstancedStaticMeshSceneProxy.
- Don't need to override HasRayTracingRepresentation() or bSupportRayTracing since base class (FStaticMeshSceneProxy) has the same logic.
- Don't need to early out when !bSupportRayTracing in GetDynamicRayTracingInstances(...) since proxy is expected to be excluded from ray tracing when bSupportRayTracing is false.

diff --git a/Engine/Source/Runtime/Engine/Classes/Engine/InstancedStaticMesh.h b/Engine/Source/Runtime/Engine/Classes/Engine/InstancedStaticMesh.h
--- a/Engine/Source/Runtime/Engine/Classes/Engine/InstancedStaticMesh.h
+++ b/Engine/Source/Runtime/Engine/Classes/Engine/InstancedStaticMesh.h
@@ -468,15 +468,12 @@ public:
 #if RHI_RAYTRACING
 	virtual bool IsRayTracingStaticRelevant() const override
 	{
+		// TODO: support static ray tracing instances when not using WPO
 		return false;
 	}
 
-	ENGINE_API virtual bool HasRayTracingRepresentation() const override;
-
 	ENGINE_API virtual void GetDynamicRayTracingInstances(FRayTracingInstanceCollector& Collector) final override;
 
-	ENGINE_API void SetupRayTracingDynamicInstances(int32 NumDynamicInstances, int32 LODIndex);
-
 #endif
 
 	ENGINE_API virtual void GetLightRelevance(const FLightSceneProxy* LightSceneProxy, bool& bDynamic, bool& bRelevant, bool& bLightMapped, bool& bShadowMapped) const override;
@@ -542,6 +539,8 @@ protected:
 	float InstanceLODDistanceScale;
 
 #if RHI_RAYTRACING
+	ENGINE_API void SetupRayTracingDynamicInstances(int32 NumDynamicInstances, int32 LODIndex);
+
 	struct FRayTracingDynamicData
 	{
 		FRayTracingGeometry DynamicGeometry;
diff --git a/Engine/Source/Runtime/Engine/Private/InstancedStaticMesh.cpp b/Engine/Source/Runtime/Engine/Private/InstancedStaticMesh.cpp
--- a/Engine/Source/Runtime/Engine/Private/InstancedStaticMesh.cpp
+++ b/Engine/Source/Runtime/Engine/Private/InstancedStaticMesh.cpp
@@ -50,9 +50,6 @@
 #include "InstancedStaticMeshSceneProxyDesc.h"
 #include "InstancedStaticMesh/ISMInstanceDataManager.h"
 
-#if RHI_RAYTRACING
-#endif
-
 #if WITH_EDITOR
 #include "DeviceProfiles/DeviceProfile.h"
 #include "DeviceProfiles/DeviceProfileManager.h"
@@ -1406,10 +1403,6 @@ void FInstancedStaticMeshSceneProxy::SetupProxy(const FInstancedStaticMeshSceneP
 	UserData_DeselectedInstances = UserData_AllInstances;
 	UserData_DeselectedInstances.bRenderSelected = false;
 
-#if RHI_RAYTRACING
-	bSupportRayTracing = IsRayTracingAllowed() && InProxyDesc.GetStaticMesh()->bSupportRayTracing;
-#endif
-
 	const bool bUseGPUScene = UseGPUScene(GetScene().GetShaderPlatform(), GetScene().GetFeatureLevel());
 	
 	const bool bEnableGpuLodSelection = CVarGpuLodSelection.GetValueOnAnyThread() != 0;
@@ -1594,22 +1587,20 @@ FInstanceDataUpdateTaskInfo *FInstancedStaticMeshSceneProxy::GetInstanceDataUpda
 }
 
 #if RHI_RAYTRACING
-bool FInstancedStaticMeshSceneProxy::HasRayTracingRepresentation() const
-{
-	return bSupportRayTracing;
-}
-
 void FInstancedStaticMeshSceneProxy::GetDynamicRayTracingInstances(FRayTracingInstanceCollector& Collector)
 {
 	if (!CVarRayTracingRenderInstances.GetValueOnRenderThread())
 	{
+		// TODO: Exclude proxy during ray tracing primitive gather instead of doing this early out here.
 		return;
 	}
 
-	if (!bSupportRayTracing)
+#if DO_CHECK
+	if (!ensureMsgf(bSupportRayTracing, TEXT("If proxy doesn't support ray tracing it should've been excluded earlier.")))
 	{
 		return;
 	}
+#endif
 
 	const uint32 InstanceCount = GetInstanceDataHeader().NumInstances;
 
@@ -1618,7 +1609,6 @@ void FInstancedStaticMeshSceneProxy::GetDynamicRayTracingInstances(FRayTracingIn
 		return;
 	}
 
-
 	// TODO: Should only do this if any instance uses static geometry
 	Collector.AddReferencedGeometryGroup(RenderData->RayTracingGeometryGroupHandle);
 
