Adding a log dedicated to verbose playback diagnostics

diff --git a/Engine/Source/Runtime/AudioMixer/Private/AudioMixerBuffer.h b/Engine/Source/Runtime/AudioMixer/Private/AudioMixerBuffer.h
--- a/Engine/Source/Runtime/AudioMixer/Private/AudioMixerBuffer.h
+++ b/Engine/Source/Runtime/AudioMixer/Private/AudioMixerBuffer.h
@@ -97,3 +97,17 @@ namespace Audio
 	};
 }
 
+inline const TCHAR* LexToString(const Audio::EBufferType::Type In)
+{
+#define CASE_TO_STRING(X) case X: return TEXT(#X);
+	using namespace Audio::EBufferType;
+	switch(In)
+	{
+		CASE_TO_STRING(Invalid);
+		CASE_TO_STRING(PCMPreview)
+		CASE_TO_STRING(PCMRealTime)
+		CASE_TO_STRING(Streaming)
+		default: return TEXT("Unknown");
+	}
+#undef CASE_TO_STRING
+}
\ No newline at end of file
diff --git a/Engine/Source/Runtime/AudioMixer/Private/AudioMixerSource.cpp b/Engine/Source/Runtime/AudioMixer/Private/AudioMixerSource.cpp
--- a/Engine/Source/Runtime/AudioMixer/Private/AudioMixerSource.cpp
+++ b/Engine/Source/Runtime/AudioMixer/Private/AudioMixerSource.cpp
@@ -1,6 +1,7 @@
 // Copyright Epic Games, Inc. All Rights Reserved.
 #include "AudioMixerSource.h"
 
+#include "Audio/AudioTimingLog.h"
 #include "AudioDefines.h"
 #include "AudioMixerSourceBuffer.h"
 #include "ActiveSound.h"
@@ -842,6 +843,10 @@ namespace Audio
 			}
 #endif
 
+			UE_CLOG(Audio::MatchesLogFilter(*WaveInstance->GetName()), LogAudioTiming, Verbose,
+				TEXT("FMixerSource::Init Name=%s,BufferType=%d,CachedRealtimeFirstBuffer=0x%p"),*WaveInstance->GetName(),
+					(int32)MixerBuffer->GetType(),WaveData->CachedRealtimeFirstBuffer);
+
 			// Whether or not we're 3D
 			bIs3D = !UseObjectBasedSpatialization() && WaveInstance->GetUseSpatialization() && SoundBuffer->NumChannels < 3;
 
@@ -1046,6 +1051,11 @@ namespace Audio
 		}
 
 		++TickCount;
+		
+		UE_CLOG(Audio::MatchesLogFilter(*WaveInstance->GetName()), LogAudioTiming, Verbose,
+			TEXT("FMixerSource::Update, Name=%s, StartTime=%.2f bProcedural=%d, bIsSourceBus=%d, CachedRealTimeFirstBuffer=0x%p, MixerSourceBufferValid=%s, TickCount=%d"),
+			*WaveInstance->GetName(), WaveInstance->StartTime, WaveInstance->WaveData->bProcedural, WaveInstance->WaveData->bIsSourceBus,
+			WaveInstance->WaveData->CachedRealtimeFirstBuffer,ToCStr(LexToString(MixerSourceBuffer.IsValid())), TickCount);
 
 		// Allow plugins to override any data in a waveinstance
 		if (AudioDevice->SourceDataOverridePluginInterface && WaveInstance->bEnableSourceDataOverride)
@@ -1207,6 +1217,11 @@ namespace Audio
 			}
 		}
 		
+		UE_CLOG(Audio::MatchesLogFilter(*WaveInstance->GetName()), LogAudioTiming, Verbose,
+			TEXT("FMixerSource::PrepareForInitialization, Name=%s, StartTime=%.2f bProcedural=%d, bIsSourceBus=%d, CachedRealTimeFirstBuffer=0x%p, MixerSourceBufferValid=%s"),
+			*WaveInstance->GetName(), WaveInstance->StartTime, WaveInstance->WaveData->bProcedural, WaveInstance->WaveData->bIsSourceBus,
+			WaveInstance->WaveData->CachedRealtimeFirstBuffer,ToCStr(LexToString(MixerSourceBuffer.IsValid())));
+		
 		return MixerSourceBuffer.IsValid();
 	}
 
@@ -1222,8 +1237,13 @@ namespace Audio
 			// Check if we have a realtime audio task already (doing first decode)
 			if (MixerSourceBuffer->IsAsyncTaskInProgress())
 			{
+				const bool bAsyncTaskDone = MixerSourceBuffer->IsAsyncTaskDone();
+				UE_CLOG(Audio::MatchesLogFilter(*WaveInstance->GetName()), LogAudioTiming, Verbose,
+					TEXT("FMixerSource::IsPreparedToInit (not ready), Name=%s, StartTime=%.2f bProcedural=%d, bIsSourceBus=%d, CachedRealTimeFirstBuffer=0x%p, IsAyncTaskDone=%s, (IsAsyncTaskInProgress)"),
+					*WaveInstance->GetName(), WaveInstance->StartTime, WaveInstance->WaveData->bProcedural, WaveInstance->WaveData->bIsSourceBus, WaveInstance->WaveData->CachedRealtimeFirstBuffer,ToCStr(LexToString(bAsyncTaskDone)));
+				
 				// not ready
-				return MixerSourceBuffer->IsAsyncTaskDone();
+				return bAsyncTaskDone;
 			}
 			else if (WaveInstance)
 			{
@@ -1248,6 +1268,10 @@ namespace Audio
 							}
 
 							check(MixerSourceBuffer.IsValid());
+									
+							UE_CLOG(Audio::MatchesLogFilter(*WaveInstance->GetName()), LogAudioTiming, Verbose,
+								TEXT("FMixerSource::IsPreparedToInit (not ready), Name=%s, StartTime=%.2f bProcecural=%d, bIsSourceBus=%d, CachedRealTimeFirstBuffer=0x%p, (Kicking Off Initial Async Decode)"),
+								*WaveInstance->GetName(), WaveInstance->StartTime, WaveInstance->WaveData->bProcedural, WaveInstance->WaveData->bIsSourceBus, WaveInstance->WaveData->CachedRealtimeFirstBuffer);
 
 							ICompressedAudioInfo* Decoder = MixerBuffer->GetDecompressionState(false);
 							MixerSourceBuffer->ReadMoreRealtimeData(Decoder, 0, EBufferReadMode::Asynchronous);
@@ -1261,7 +1285,11 @@ namespace Audio
 
 			return true;
 		}
-
+		UE_CLOG(Audio::MatchesLogFilter(*WaveInstance->GetName()), LogAudioTiming, Verbose,
+			TEXT("FMixerSource::IsPreparedToInit (not ready), Name=%s, StartTime=%.2f bProcedural=%d, bIsSourceBus=%d, CachedRealTimeFirstBuffer=0x%p, MixerBuffer=0x%p, IsRealTimeSourceReady=%s (realtime source not ready, or no mixer buffer)"),
+			*WaveInstance->GetName(), WaveInstance->StartTime, WaveInstance->WaveData->bProcedural, WaveInstance->WaveData->bIsSourceBus, WaveInstance->WaveData->CachedRealtimeFirstBuffer,
+				MixerBuffer, ToCStr(LexToString(MixerBuffer->IsRealTimeSourceReady())));
+		
 		return false;
 	}
 
@@ -1292,10 +1320,21 @@ namespace Audio
 
 		AUDIO_MIXER_TRACE_CPUPROFILER_EVENT_SCOPE(AudioMixerSource::Play);
 
+
+		UE_CLOG(Audio::MatchesLogFilter(*WaveInstance->GetName()), LogAudioTiming, Verbose, 
+			TEXT("FMixerSource::Play, Name=%s, StartTime=%.2f bProcedural=%d, bIsSourceBus=%d, CachedRealTimeFirstBuffer=0x%p, MixerSourceBufferValid=%s, TickCount=%d, InitState=%s"),
+			*WaveInstance->GetName(), WaveInstance->StartTime, WaveInstance->WaveData->bProcedural, WaveInstance->WaveData->bIsSourceBus,
+			WaveInstance->WaveData->CachedRealtimeFirstBuffer,ToCStr(LexToString(MixerSourceBuffer.IsValid())), TickCount, ToCStr(LexToString(InitializationState)));
+
 		// It's possible if Pause and Play are called while a sound is async initializing. In this case
 		// we'll just not actually play the source here. Instead we'll call play when the sound finishes loading.
 		if (MixerSourceVoice && InitializationState == EMixerSourceInitializationState::Initialized)
 		{
+			UE_CLOG(Audio::MatchesLogFilter(*WaveInstance->GetName()), LogAudioTiming, Verbose,
+				TEXT("FMixerSourceVoice::Play, Name=%s, StartTime=%.2f bProcedural=%d, bIsSourceBus=%d, CachedRealTimeFirstBuffer=0x%p, MixerSourceBufferValid=%s, TickCount=%d, InitState=%s"),
+				*WaveInstance->GetName(), WaveInstance->StartTime, WaveInstance->WaveData->bProcedural, WaveInstance->WaveData->bIsSourceBus,
+				WaveInstance->WaveData->CachedRealtimeFirstBuffer,ToCStr(LexToString(MixerSourceBuffer.IsValid())), TickCount, ToCStr(LexToString(InitializationState)));
+
 			MixerSourceVoice->Play();
 
 #if UE_AUDIO_PROFILERTRACE_ENABLED
diff --git a/Engine/Source/Runtime/AudioMixer/Private/AudioMixerSource.h b/Engine/Source/Runtime/AudioMixer/Private/AudioMixerSource.h
--- a/Engine/Source/Runtime/AudioMixer/Private/AudioMixerSource.h
+++ b/Engine/Source/Runtime/AudioMixer/Private/AudioMixerSource.h
@@ -199,3 +199,15 @@ namespace Audio
 		uint32 bPrevAllowedSpatializationSetting : 1;
 	};
 }
+
+inline const TCHAR* LexToString(const Audio::EMixerSourceInitializationState State)
+{
+	switch(State)
+	{
+	case Audio::EMixerSourceInitializationState::Initialized: return TEXT("Initialized");
+	case Audio::EMixerSourceInitializationState::NotInitialized: return TEXT("NotInitialized");
+	case Audio::EMixerSourceInitializationState::Initializing: return TEXT("Initializing");
+	default: break;
+	}
+	return TEXT("Unknown");
+}
diff --git a/Engine/Source/Runtime/AudioMixer/Private/AudioMixerSourceManager.cpp b/Engine/Source/Runtime/AudioMixer/Private/AudioMixerSourceManager.cpp
--- a/Engine/Source/Runtime/AudioMixer/Private/AudioMixerSourceManager.cpp
+++ b/Engine/Source/Runtime/AudioMixer/Private/AudioMixerSourceManager.cpp
@@ -2,6 +2,7 @@
 
 #include "AudioMixerSourceManager.h"
 
+#include "Audio/AudioTimingLog.h"
 #include "AudioDefines.h"
 #include "AudioMixerSourceBuffer.h"
 #include "AudioMixerDevice.h"
@@ -161,7 +162,7 @@ static FAutoConsoleCommand GModulationSetMaxPitchRange(
 	if (SourceInfos[SourceId].bIsDebugMode)																													\
 	{																																			\
 		FString CustomMessage = FString::Printf(Format, ##__VA_ARGS__);																			\
-		FString LogMessage = FString::Printf(TEXT("<Debug Sound Log> [Id=%d][Name=%s]: %s"), SourceId, *SourceInfos[SourceId].DebugName, *CustomMessage);	\
+		FString LogMessage = FString::Printf(TEXT("<Debug Sound Log> [Id=%d][Name=%s]: %s"), SourceId, *SourceInfos[SourceId].GetDebugName(), *CustomMessage);	\
 		UE_LOG(LogAudioMixer, Log, TEXT("%s"), *LogMessage);																								\
 	}
 
@@ -1357,6 +1358,9 @@ namespace Audio
 			SourceInfo.bIsActive = true;
 
 			SourceInfo.StartTime = StartTime;
+					
+			UE_CLOG(Audio::MatchesLogFilter(*SourceInfo.DebugName), LogAudioTiming, Verbose,
+				TEXT("FMixerSourceManager::Play (render thread), Wave=%s, SourceId=%d"), *SourceInfo.GetDebugName(), SourceId);
 
 			AUDIO_MIXER_DEBUG_LOG(SourceId, TEXT("Is playing"));
 		}, AUDIO_MIXER_THREAD_COMMAND_STRING("Play()"));
@@ -2147,6 +2151,10 @@ namespace Audio
 			// Fill array with elements all at once to avoid sequential Add() operation overhead.
 			const int32 NumSamples = NumOutputFrames * SourceInfo.NumInputChannels;
 			
+			UE_CLOG(Audio::MatchesLogFilter(*SourceInfo.GetDebugName()), LogAudioTiming, Verbose,
+				TEXT("ComputeSourceBuffersForIdRange Name=%s, NumOutputFrames=%d, FramesPlayed=%lld, CurrentFrameIndex=%d "),
+					*SourceInfo.GetDebugName(), NumOutputFrames, SourceInfo.NumFramesPlayed, SourceInfo.CurrentFrameIndex);
+			
 			// Initialize both the pre-distance attenuation buffer and the source buffer
 			SourceInfo.PreDistanceAttenuationBuffer.Reset();
 			SourceInfo.PreDistanceAttenuationBuffer.AddZeroed(NumSamples);
diff --git a/Engine/Source/Runtime/AudioMixer/Private/AudioMixerSourceManager.h b/Engine/Source/Runtime/AudioMixer/Private/AudioMixerSourceManager.h
--- a/Engine/Source/Runtime/AudioMixer/Private/AudioMixerSourceManager.h
+++ b/Engine/Source/Runtime/AudioMixer/Private/AudioMixerSourceManager.h
@@ -593,6 +593,15 @@ namespace Audio
 			uint8 bIsDebugMode : 1;
 			FString DebugName;
 #endif // AUDIO_MIXER_ENABLE_DEBUG_MODE
+			FString GetDebugName() const
+			{
+#if AUDIO_MIXER_ENABLE_DEBUG_MODE
+				return DebugName;
+#else
+				return {};
+#endif //AUDIO_MIXER_ENABLE_DEBUG_MODE
+				
+			}
 		};
 
 		static void ApplyDistanceAttenuation(FSourceInfo& InSourceInfo, int32 NumSamples);
diff --git a/Engine/Source/Runtime/Engine/Classes/Sound/SoundWave.h b/Engine/Source/Runtime/Engine/Classes/Sound/SoundWave.h
--- a/Engine/Source/Runtime/Engine/Classes/Sound/SoundWave.h
+++ b/Engine/Source/Runtime/Engine/Classes/Sound/SoundWave.h
@@ -53,6 +53,18 @@ enum EDecompressionType : int
 	DTYPE_MAX,
 };
 
+inline const TCHAR* LexToString(const EDecompressionType InType)
+{
+#define CASE_TO_STRING(X) case X: return TEXT(#X);
+	switch(InType)
+	{
+		FOREACH_ENUM_EDECOMPRESSIONTYPE(CASE_TO_STRING);
+		default: break;
+	}
+	return TEXT("Unknown");
+#undef CASE_TO_STRING
+}
+
 /** Precache states */
 enum class ESoundWavePrecacheState
 {
diff --git a/Engine/Source/Runtime/Engine/Private/Audio/AudioTimingLog.cpp b/Engine/Source/Runtime/Engine/Private/Audio/AudioTimingLog.cpp
new file mode 100644
--- /dev/null
+++ b/Engine/Source/Runtime/Engine/Private/Audio/AudioTimingLog.cpp
@@ -0,0 +1,5 @@
+// Copyright Epic Games, Inc. All Rights Reserved.
+
+#include "Audio/AudioTimingLog.h"
+
+DEFINE_LOG_CATEGORY(LogAudioTiming);
\ No newline at end of file
diff --git a/Engine/Source/Runtime/Engine/Private/AudioDebug.cpp b/Engine/Source/Runtime/Engine/Private/AudioDebug.cpp
--- a/Engine/Source/Runtime/Engine/Private/AudioDebug.cpp
+++ b/Engine/Source/Runtime/Engine/Private/AudioDebug.cpp
@@ -27,6 +27,17 @@
 
 #if ENABLE_AUDIO_DEBUG
 
+namespace Audio
+{
+	ENGINE_API FString ExtraTracingSoundWaveFilterCVar;
+}
+
+static FAutoConsoleVariableRef CVarExtraTracingSoundWaveFilter(
+	TEXT("au.debug.extra_logging_filter"),
+	Audio::ExtraTracingSoundWaveFilterCVar,
+	TEXT("Enable extra logging for any sound wave that has this substring in its name."),
+	ECVF_Default);
+
 // Console variables
 static int32 ActiveSoundVisualizeModeCVar = 1;
 FAutoConsoleVariableRef CVarAudioVisualizeActiveSoundsMode(
diff --git a/Engine/Source/Runtime/Engine/Private/AudioDecompress.cpp b/Engine/Source/Runtime/Engine/Private/AudioDecompress.cpp
--- a/Engine/Source/Runtime/Engine/Private/AudioDecompress.cpp
+++ b/Engine/Source/Runtime/Engine/Private/AudioDecompress.cpp
@@ -12,11 +12,6 @@
 #include "Stats/StatsTrace.h"
 #include "Sound/StreamedAudioChunkSeekTable.h"
 
-#if !UE_BUILD_SHIPPING
-// We need to hit this from the streaming cache to trace.
-/* extern */ FString ExtraTracingSoundWaveFilterCVar;
-#endif
-
 namespace AudioDecompressPrivate
 {
 #ifndef WITH_AUDIO_DECODER_DIAGNOSTICS
@@ -31,13 +26,7 @@ namespace AudioDecompressPrivate
 		TEXT("Keep a reference to next to play chunk" ),
 		ECVF_Default);
 	
-#if !UE_BUILD_SHIPPING
-	static FAutoConsoleVariableRef CVarExtraTracingSoundWaveFilter(
-		TEXT("au.debug.extra_logging_filter"),
-		ExtraTracingSoundWaveFilterCVar,
-		TEXT("Enable extra logging for any sound wave that has this substring in its name."),
-		ECVF_Default);
-#endif
+
 
 #if WITH_AUDIO_DECODER_DIAGNOSTICS
 	
@@ -317,7 +306,7 @@ bool IStreamedCompressedInfo::StreamCompressedData(uint8* Destination, bool bLoo
 
 	SCOPE_CYCLE_COUNTER(STAT_AudioStreamedDecompressTime);
 	
-	UE_LOG(LogAudio, VeryVerbose, TEXT("Streaming compressed data from SoundWave'%s' - Chunk=%d\tCurrentSampleCount=%d\tTrueSampleCount=%d\tNumChunks=%d\tOffset=%d\tChunkSize=%d\tLooping=%s\tLastPCMOffset=%d\tContainsEOF=%s" ), 
+	UE_LOG(LogAudioStreamCaching, VeryVerbose, TEXT("Streaming compressed data from SoundWave'%s' - Chunk=%d\tCurrentSampleCount=%d\tTrueSampleCount=%d\tNumChunks=%d\tOffset=%d\tChunkSize=%d\tLooping=%s\tLastPCMOffset=%d\tContainsEOF=%s" ), 
 		*StreamingSoundWave->GetFName().ToString(), 	
 		CurrentChunkIndex, 
 		CurrentSampleCount, 
@@ -333,9 +322,7 @@ bool IStreamedCompressedInfo::StreamCompressedData(uint8* Destination, bool bLoo
 	// If we have a pending next chunk from seeking, move to it now.
 	if (StreamSeekBlockIndex != INDEX_NONE)
 	{
-#if !UE_BUILD_SHIPPING
-		UE_CLOG(bWantsExtraLogging, LogAudioStreamCaching, Display, TEXT("[%s.%p] - Seek Requesting Block = %d"), *WriteToString<128>(StreamingSoundWave->GetFName()), this, StreamSeekBlockIndex.load(std::memory_order_relaxed));
-#endif
+		UE_CLOG(WantsFilteredLogging(), LogAudioStreamCaching, Display, TEXT("[%s.%p] - Seek Requesting Block = %d"), *WriteToString<128>(StreamingSoundWave->GetFName()), this, StreamSeekBlockIndex.load(std::memory_order_relaxed));
 		uint32 ChunkSize = 0;
 		const uint8* NewlySeekedChunk = GetLoadedChunk(StreamingSoundWave, StreamSeekBlockIndex, ChunkSize);
 		UE_LOG(LogAudio, Log, TEXT("Seek request for (%s): (%s) Chunk=%d (%s), Offset=%d, OffsetInAudioFrames=%u"),
@@ -394,9 +381,8 @@ bool IStreamedCompressedInfo::StreamCompressedData(uint8* Destination, bool bLoo
 							LastPCMOffset = 0;
 							bStoringEndOfFile = false;
 
-#if !UE_BUILD_SHIPPING
-							UE_CLOG(bWantsExtraLogging, LogAudioStreamCaching, Display, TEXT("[%s.%p] - Committing Seek Chunk=%d Size=%u"), *WriteToString<128>(StreamingSoundWave->GetFName()), this, CurrentChunkIndex, ChunkSize);
-#endif
+
+							UE_CLOG(WantsFilteredLogging(), LogAudioStreamCaching, Display, TEXT("[%s.%p] - Committing Seek Chunk=%d Size=%u"), *WriteToString<128>(StreamingSoundWave->GetFName()), this, CurrentChunkIndex, ChunkSize);
 						}
 						else // Seek failed (off the end of the chunk).
 						{
@@ -417,9 +403,7 @@ bool IStreamedCompressedInfo::StreamCompressedData(uint8* Destination, bool bLoo
 	// If next chunk wasn't loaded when last one finished reading, try to get it again now
 	if (SrcBufferData == NULL)
 	{
-#if !UE_BUILD_SHIPPING
-		UE_CLOG(bWantsExtraLogging, LogAudioStreamCaching, Display, TEXT("[%s.%p] - Demand-loading Chunk=%d"), *WriteToString<128>(StreamingSoundWave->GetFName()), this, CurrentChunkIndex);
-#endif
+		UE_CLOG(WantsFilteredLogging(), LogAudioStreamCaching, Display, TEXT("[%s.%p] - Demand-loading Chunk=%d"), *WriteToString<128>(StreamingSoundWave->GetFName()), this, CurrentChunkIndex);
 		uint32 ChunkSize = 0;
 		SrcBufferData = GetLoadedChunk(StreamingSoundWave, CurrentChunkIndex, ChunkSize);
 		if (SrcBufferData)
@@ -513,9 +497,8 @@ bool IStreamedCompressedInfo::StreamCompressedData(uint8* Destination, bool bLoo
 
 		// Decompress the next compression frame of audio (many samples) into the PCM buffer
 		const int32 DecodedFrames = DecompressToPCMBuffer(/*Unused*/ 0);
-#if !UE_BUILD_SHIPPING
-		UE_CLOG(bWantsExtraLogging, LogAudioStreamCaching, Display, TEXT("[%s.%p] - Decoded frames = %d, Chunk = %d @ %u"), *WriteToString<128>(StreamingSoundWave->GetFName()), this, DecodedFrames, CurrentChunkIndex, SrcBufferOffset);
-#endif
+		UE_CLOG(WantsFilteredLogging(), LogAudioStreamCaching, Display, TEXT("[%s.%p] - Decoded frames = %d, Chunk = %d @ %u"), *WriteToString<128>(StreamingSoundWave->GetFName()), this, DecodedFrames, CurrentChunkIndex, SrcBufferOffset);
+
 		if (DecodedFrames < 0)
 		{
 			FXxHash64 Hash = FXxHash64::HashBuffer(SrcBufferData, SrcBufferDataSize);
@@ -542,15 +525,12 @@ bool IStreamedCompressedInfo::StreamCompressedData(uint8* Destination, bool bLoo
 			// Have we reached the end of buffer?
 			if (SrcBufferOffset >= SrcBufferDataSize - SrcBufferPadding)
 			{
-#if !UE_BUILD_SHIPPING
-				UE_CLOG(bWantsExtraLogging, LogAudioStreamCaching, Display, TEXT("[%s.%p] - End Of Buffer %d = %u/%u (pad = %u)"), *WriteToString<128>(StreamingSoundWave->GetFName()), this, CurrentChunkIndex, SrcBufferOffset, SrcBufferDataSize, SrcBufferPadding);
-#endif
+				UE_CLOG(WantsFilteredLogging(), LogAudioStreamCaching, Display, TEXT("[%s.%p] - End Of Buffer %d = %u/%u (pad = %u)"), *WriteToString<128>(StreamingSoundWave->GetFName()), this, CurrentChunkIndex, SrcBufferOffset, SrcBufferDataSize, SrcBufferPadding);
+
 				// Special case for the last chunk of audio
 				if (CurrentChunkIndex == StreamingSoundWave->GetNumChunks() - 1)
 				{
-#if !UE_BUILD_SHIPPING
-						UE_CLOG(bWantsExtraLogging, LogAudioStreamCaching, Display, TEXT("[%s.%p] - Last Chunk Handling LastPCMByteSize = %u"), *WriteToString<128>(StreamingSoundWave->GetFName()), this, LastPCMByteSize);
-#endif
+						UE_CLOG(WantsFilteredLogging(), LogAudioStreamCaching, Display, TEXT("[%s.%p] - Last Chunk Handling LastPCMByteSize = %u"), *WriteToString<128>(StreamingSoundWave->GetFName()), this, LastPCMByteSize);
 
 					// check whether all decoded PCM was written
 					if (LastPCMByteSize == 0)
@@ -741,9 +721,7 @@ const uint8* IStreamedCompressedInfo::GetLoadedChunk(const FSoundWaveProxyPtr& I
 	}
 	else if (ChunkIndex == 0)
 	{
-#if !UE_BUILD_SHIPPING
-		UE_CLOG(bWantsExtraLogging, LogAudioStreamCaching, Display, TEXT("[%s.%p] - GetLoadedChunk(0)"), *WriteToString<128>(InSoundWave->GetFName()), this);
-#endif
+		UE_CLOG(WantsFilteredLogging(), LogAudioStreamCaching, Display, TEXT("[%s.%p] - GetLoadedChunk(0)"), *WriteToString<128>(InSoundWave->GetFName()), this);
 		TArrayView<const uint8> ZerothChunk = FSoundWaveProxy::GetZerothChunk(InSoundWave, true);
 		OutChunkSize = ZerothChunk.Num();
 		return ZerothChunk.GetData();
@@ -753,9 +731,7 @@ const uint8* IStreamedCompressedInfo::GetLoadedChunk(const FSoundWaveProxyPtr& I
 		CurCompressedChunkHandle = IStreamingManager::Get().GetAudioStreamingManager().GetLoadedChunk(InSoundWave, ChunkIndex, false, true);
 		RetainNextChunkInSequence(InSoundWave,ChunkIndex);
 		OutChunkSize = CurCompressedChunkHandle.Num();
-#if !UE_BUILD_SHIPPING
-		UE_CLOG(bWantsExtraLogging, LogAudioStreamCaching, Display, TEXT("[%s.%p] - GetLoadedChunk(%d, %u)"), *WriteToString<128>(InSoundWave->GetFName()), this, ChunkIndex, OutChunkSize);
-#endif
+		UE_CLOG(WantsFilteredLogging(), LogAudioStreamCaching, Display, TEXT("[%s.%p] - GetLoadedChunk(%d, %u)"), *WriteToString<128>(InSoundWave->GetFName()), this, ChunkIndex, OutChunkSize);
 		return CurCompressedChunkHandle.GetData();
 	}
 }
@@ -897,12 +873,9 @@ bool ICompressedAudioInfo::StreamCompressedInfo(USoundWave* Wave, FSoundQualityI
 	}
 
 #if !UE_BUILD_SHIPPING
-	if (ExtraTracingSoundWaveFilterCVar.Len() &&
-		Wave->GetName().Contains(ExtraTracingSoundWaveFilterCVar))
-	{
-		bWantsExtraLogging = true;
-	}
-#endif
+	bWantsFilteredLogging = Audio::MatchesLogFilter(*GetNameSafe(Wave));
+#endif //UE_BUILD_SHIPPING
+	
 	return StreamCompressedInfoInternal(StreamingSoundWave, QualityInfo);
 }
 
@@ -912,12 +885,8 @@ bool ICompressedAudioInfo::StreamCompressedInfo(const FSoundWaveProxyPtr& Wave,
 	StreamingSoundWave = Wave;
 
 #if !UE_BUILD_SHIPPING
-	if (ExtraTracingSoundWaveFilterCVar.Len() &&
-		Wave->GetFName().ToString().Contains(ExtraTracingSoundWaveFilterCVar))
-	{
-		bWantsExtraLogging = true;
-	}
-#endif
+	bWantsFilteredLogging = Audio::MatchesLogFilter(Wave->GetFName());
+#endif //UE_BUILD_SHIPPING
 
 	return StreamCompressedInfoInternal(StreamingSoundWave, QualityInfo);
 }
diff --git a/Engine/Source/Runtime/Engine/Private/AudioStreamingCache.cpp b/Engine/Source/Runtime/Engine/Private/AudioStreamingCache.cpp
--- a/Engine/Source/Runtime/Engine/Private/AudioStreamingCache.cpp
+++ b/Engine/Source/Runtime/Engine/Private/AudioStreamingCache.cpp
@@ -8,6 +8,8 @@ AudioStreaming.cpp: Implementation of audio streaming classes.
 #include "AudioStreamCacheMemoryHandle.h"
 #include "Async/Async.h"
 #include "Audio.h"
+#include "Audio/AudioDebug.h"
+#include "Audio/AudioTimingLog.h"
 #include "AudioCompressionSettingsUtils.h"
 #include "CanvasTypes.h"
 #include "Engine/Engine.h"
@@ -23,9 +25,6 @@ TRACE_DECLARE_ATOMIC_INT_COUNTER(AudioStreamCache_Loads, TEXT("AudioStreamCache
 TRACE_DECLARE_ATOMIC_INT_COUNTER(AudioStreamCache_Loads_Request, TEXT("AudioStreamCache Loads Request"));
 TRACE_DECLARE_ATOMIC_INT_COUNTER(AudioStreamCache_Loads_Playback, TEXT("AudioStreamCache Loads Playback"));
 
-// Defined in AudioDecompress.cpp off au.debug.extra_logging_filter
-extern FString ExtraTracingSoundWaveFilterCVar;
-
 static int32 DebugMaxElementsDisplayCVar = 128;
 FAutoConsoleVariableRef CVarDebugDisplayCaches(
 	TEXT("au.streamcaching.MaxCachesToDisplay"),
@@ -398,9 +397,7 @@ void FAudioChunkCache::FCacheElement::ReleaseRetainedAudioOnSoundWave()
 {
 	if (FSoundWavePtr SoundWaveDataPtr = SoundWaveWeakPtr.Pin())
 	{
-#if !UE_BUILD_SHIPPING
-			UE_CLOG(bWantsExtraLogging, LogAudioStreamCaching, Display, TEXT("[%s] - Release Retained Audio, chunk %d, ObjectKey=%s"), *WriteToString<128>(Key.SoundWaveName), Key.ChunkIndex, *WriteToString<40>(Key.ObjectKey));
-#endif
+		UE_CLOG(WantsFilteredLogging(), LogAudioTiming, Verbose, TEXT("[%s] - Release Retained Audio, chunk %d, ObjectKey=%s"), *WriteToString<128>(Key.SoundWaveName), Key.ChunkIndex, *WriteToString<40>(Key.ObjectKey));
 		check(SoundWaveDataPtr.IsValid());
 		SoundWaveDataPtr->ReleaseCompressedAudio();
 	}
@@ -799,10 +796,7 @@ uint64 FAudioChunkCache::AddOrTouchChunk(const FChunkKey& InKey, const TSharedPt
 			ExecuteOnLoadCompleteCallback(EAudioChunkLoadResult::AlreadyLoaded, OnLoadCompleted, CallbackThread);
 		}
 
-#if !UE_BUILD_SHIPPING
-		UE_CLOG(FoundElement->bWantsExtraLogging, LogAudioStreamCaching, Display, TEXT("[%s] - AddOrTouchChunk(found), chunk %d, ObjectKey=%s"), *WriteToString<128>(FoundElement->Key.SoundWaveName), FoundElement->Key.ChunkIndex, *WriteToString<40>(FoundElement->Key.ObjectKey));
-#endif
-
+	UE_CLOG(FoundElement->WantsFilteredLogging(), LogAudioTiming, Verbose, TEXT("[%s] - AddOrTouchChunk(found), chunk %d, ObjectKey=%s"), *WriteToString<128>(FoundElement->Key.SoundWaveName), FoundElement->Key.ChunkIndex, *WriteToString<40>(FoundElement->Key.ObjectKey));
 
 #if DEBUG_STREAM_CACHE
 		FoundElement->DebugInfo.NumTimesTouched++;
@@ -868,9 +862,7 @@ TArrayView<uint8> FAudioChunkCache::GetChunk(const FChunkKey& InKey, const TShar
 	FCacheElement* FoundElement = FindElementForKey(InKey);
 	if (FoundElement)
 	{
-#if !UE_BUILD_SHIPPING
-		UE_CLOG(FoundElement->bWantsExtraLogging, LogAudioStreamCaching, Display, TEXT("[%s] - GetChunk(found), chunk %d, ObjectKey=%s"), *WriteToString<128>(FoundElement->Key.SoundWaveName), FoundElement->Key.ChunkIndex, *WriteToString<40>(FoundElement->Key.ObjectKey));
-#endif
+		UE_CLOG(FoundElement->WantsFilteredLogging(), LogAudioTiming, Verbose, TEXT("[%s] - GetChunk(found), chunk %d, ObjectKey=%s"), *WriteToString<128>(FoundElement->Key.SoundWaveName), FoundElement->Key.ChunkIndex, *WriteToString<40>(FoundElement->Key.ObjectKey));
 		OutCacheOffset = FoundElement->CacheLookupID;
 		TouchElement(FoundElement);
 		if (FoundElement->IsLoadInProgress())
@@ -916,9 +908,7 @@ TArrayView<uint8> FAudioChunkCache::GetChunk(const FChunkKey& InKey, const TShar
 		{
 			if (FStreamedAudioChunk* Chunk = FoundElement->GetChunk(InKey.ChunkIndex))
 			{
-#if !UE_BUILD_SHIPPING
-				UE_CLOG(FoundElement->bWantsExtraLogging, LogAudioStreamCaching, Display, TEXT("[%s] - GetChunk(blocking, hit), chunk %d, ObjectKey=%s"), *WriteToString<128>(FoundElement->Key.SoundWaveName), FoundElement->Key.ChunkIndex, *WriteToString<40>(FoundElement->Key.ObjectKey));
-#endif
+				UE_CLOG(FoundElement->WantsFilteredLogging(), LogAudioTiming, Verbose, TEXT("[%s] - GetChunk(blocking, hit), chunk %d, ObjectKey=%s"), *WriteToString<128>(FoundElement->Key.SoundWaveName), FoundElement->Key.ChunkIndex, *WriteToString<40>(FoundElement->Key.ObjectKey));
 
 				int32 ChunkAudioDataSize = Chunk->AudioDataSize;
 #if DEBUG_STREAM_CACHE
@@ -1174,9 +1164,7 @@ uint64 FAudioChunkCache::TrimMemory(uint64 BytesToFree, bool bInAllowRetainedChu
 	{
 		if (CurrentElement->CanEvictChunk() && CurrentElement->ChunkDataSize != 0)
 		{
-#if !UE_BUILD_SHIPPING
-			UE_CLOG(CurrentElement->bWantsExtraLogging, LogAudioStreamCaching, Display, TEXT("[%s] - Evicting for trim, chunk %d, ObjectKey=%s"), *WriteToString<128>(CurrentElement->Key.SoundWaveName), CurrentElement->Key.ChunkIndex, *WriteToString<40>(CurrentElement->Key.ObjectKey));
-#endif
+			UE_CLOG(CurrentElement->WantsFilteredLogging(), LogAudioTiming, Verbose, TEXT("[%s] - Evicting for trim, chunk %d, ObjectKey=%s"), *WriteToString<128>(CurrentElement->Key.SoundWaveName), CurrentElement->Key.ChunkIndex, *WriteToString<40>(CurrentElement->Key.ObjectKey));
 
 			uint32 ChunkSize = CurrentElement->ChunkDataSize;
 			BytesFreed += ChunkSize;
@@ -1221,9 +1209,7 @@ uint64 FAudioChunkCache::TrimMemory(uint64 BytesToFree, bool bInAllowRetainedChu
 				CurrentElement->ReleaseRetainedAudioOnSoundWave();
 				if (CurrentElement->CanEvictChunk())
 				{
-#if !UE_BUILD_SHIPPING
-					UE_CLOG(CurrentElement->bWantsExtraLogging, LogAudioStreamCaching, Display, TEXT("[%s] - Evicting for trim, chunk %d, ObjectKey=%s"), *WriteToString<128>(CurrentElement->Key.SoundWaveName), CurrentElement->Key.ChunkIndex, *WriteToString<40>(CurrentElement->Key.ObjectKey));
-#endif
+					UE_CLOG(CurrentElement->WantsFilteredLogging(), LogAudioTiming, Verbose, TEXT("[%s] - Evicting for trim, chunk %d, ObjectKey=%s"), *WriteToString<128>(CurrentElement->Key.SoundWaveName), CurrentElement->Key.ChunkIndex, *WriteToString<40>(CurrentElement->Key.ObjectKey));
 					uint32 ChunkSize = CurrentElement->ChunkDataSize;
 					BytesFreed += ChunkSize;
 					RetainedBytesFreed += ChunkSize;
@@ -1592,19 +1578,9 @@ FAudioChunkCache::FCacheElement* FAudioChunkCache::InsertChunk(const FChunkKey&
 			}
 		}
 
-#if !UE_BUILD_SHIPPING
-		CacheElement->bWantsExtraLogging = false;
-		if (ExtraTracingSoundWaveFilterCVar.Len())
-		{
-			TStringBuilder<128> SoundName;
-			SoundName << InSoundWavePtr->GetFName();
-			if (FCString::Stristr(*SoundName, *ExtraTracingSoundWaveFilterCVar))
-			{
-				CacheElement->bWantsExtraLogging = true;
-				UE_LOG(LogAudioStreamCaching, Display, TEXT("[%s] - Inserting chunk index %d, ObjectKey=%s"), *SoundName, InKey.ChunkIndex, *WriteToString<40>(InKey.ObjectKey));
-			}
-		}
-#endif
+		UE_CLOG(CacheElement->WantsFilteredLogging(), LogAudioTiming, Verbose, TEXT("[%s] - Inserting chunk index %d, ObjectKey=%s"),
+			*InSoundWavePtr->GetFName().ToString(), InKey.ChunkIndex, *WriteToString<40>(InKey.ObjectKey));
+		
 		CacheElement->bIsLoaded = false;
 		CacheElement->Key = InKey;
 		CacheElement->SoundWaveWeakPtr = InSoundWavePtr;
@@ -1663,9 +1639,7 @@ FAudioChunkCache::FCacheElement* FAudioChunkCache::EvictLeastRecentChunk(bool bB
 
 	if (bIsChunkEvictable)
 	{
-#if !UE_BUILD_SHIPPING
-		UE_CLOG(CacheElement->bWantsExtraLogging, LogAudioStreamCaching, Display, TEXT("[%s] - Evicting ununsed LRU chunk index %d, ObjectKey=%s"), *WriteToString<128>(CacheElement->Key.SoundWaveName), CacheElement->Key.ChunkIndex, *WriteToString<40>(CacheElement->Key.ObjectKey));
-#endif
+		UE_CLOG(CacheElement->WantsFilteredLogging(), LogAudioTiming, Verbose, TEXT("[%s] - Evicting ununsed LRU chunk index %d, ObjectKey=%s"), *WriteToString<128>(CacheElement->Key.SoundWaveName), CacheElement->Key.ChunkIndex, *WriteToString<40>(CacheElement->Key.ObjectKey));
 
 		check(CacheElement->MoreRecentElement != nullptr);
 		check(CacheElement->LessRecentElement == nullptr);
@@ -1677,9 +1651,7 @@ FAudioChunkCache::FCacheElement* FAudioChunkCache::EvictLeastRecentChunk(bool bB
 	}
 	else if (bBlockForPendingLoads && bIsChunkLoadingButUnreferenced)
 	{
-#if !UE_BUILD_SHIPPING
-		UE_CLOG(CacheElement->bWantsExtraLogging, LogAudioStreamCaching, Display, TEXT("[%s] - Evicting ununsed loading chunk index %d, ObjectKey=%s"), *WriteToString<128>(CacheElement->Key.SoundWaveName), CacheElement->Key.ChunkIndex, *WriteToString<40>(CacheElement->Key.ObjectKey));
-#endif
+		UE_CLOG(CacheElement->WantsFilteredLogging(), LogAudioTiming, Verbose, TEXT("[%s] - Evicting ununsed loading chunk index %d, ObjectKey=%s"), *WriteToString<128>(CacheElement->Key.SoundWaveName), CacheElement->Key.ChunkIndex, *WriteToString<40>(CacheElement->Key.ObjectKey));
 		CacheElement->WaitForAsyncLoadCompletion(true);
 
 		FCacheElement* NewLeastRecentElement = CacheElement->MoreRecentElement;
@@ -1811,9 +1783,7 @@ void FAudioChunkCache::KickOffAsyncLoad(FCacheElement* CacheElement, const FChun
 	int32 ChunkDataSize = Chunk->AudioDataSize;
 
 	EAsyncIOPriorityAndFlags AsyncIOPriority = GetAsyncPriorityForChunk(InKey, bNeededForPlayback);
-#if !UE_BUILD_SHIPPING
-	UE_CLOG(CacheElement->bWantsExtraLogging, LogAudioStreamCaching, Display, TEXT("[%s] - KickOffAsyncLoad pri %d outstanding %d chunk index %d, ObjectKey=%s"), *WriteToString<128>(CacheElement->Key.SoundWaveName), (uint32)AsyncIOPriority, NumberOfLoadsInFlight.GetValue(), CacheElement->Key.ChunkIndex, *WriteToString<40>(CacheElement->Key.ObjectKey));
-#endif
+	UE_CLOG(CacheElement->WantsFilteredLogging(), LogAudioTiming, Verbose, TEXT("[%s] - KickOffAsyncLoad pri %d outstanding %d chunk index %d, ObjectKey=%s"), *WriteToString<128>(CacheElement->Key.SoundWaveName), (uint32)AsyncIOPriority, NumberOfLoadsInFlight.GetValue(), CacheElement->Key.ChunkIndex, *WriteToString<40>(CacheElement->Key.ObjectKey));
 
 	*MemoryCounterBytes -= CacheElement->ChunkDataSize;
 
@@ -1902,9 +1872,7 @@ void FAudioChunkCache::KickOffAsyncLoad(FCacheElement* CacheElement, const FChun
 	{
 		if (CacheElement->IsLoadInProgress())
 		{
-#if !UE_BUILD_SHIPPING
-			UE_CLOG(CacheElement->bWantsExtraLogging, LogAudioStreamCaching, Display, TEXT("[%s] - Waiting for async load while kicking async load, chunk %d, ObjectKey=%s"), *WriteToString<128>(CacheElement->Key.SoundWaveName), CacheElement->Key.ChunkIndex, *WriteToString<40>(CacheElement->Key.ObjectKey));
-#endif
+			UE_CLOG(CacheElement->WantsFilteredLogging(), LogAudioTiming, Verbose, TEXT("[%s] - Waiting for async load while kicking async load, chunk %d, ObjectKey=%s"), *WriteToString<128>(CacheElement->Key.SoundWaveName), CacheElement->Key.ChunkIndex, *WriteToString<40>(CacheElement->Key.ObjectKey));
 			CacheElement->WaitForAsyncLoadCompletion(true);
 		}
 
@@ -1930,18 +1898,12 @@ void FAudioChunkCache::KickOffAsyncLoad(FCacheElement* CacheElement, const FChun
 				TGraphTask<FClearAudioChunkCacheReadRequestTask>::CreateTask().ConstructAndDispatchWhenReady(LocalReadRequest);
 			}
 
-#if !UE_BUILD_SHIPPING
-			if (CacheElement->bWantsExtraLogging)
-			{
-				bool bMatched = CacheElement->Key == InKey;
-				UE_LOG(LogAudioStreamCaching, Display, TEXT("[%s] - AsyncReadComplete index %d, ObjectKey=%s"), *WriteToString<128>(CacheElement->Key.SoundWaveName), CacheElement->Key.ChunkIndex, *WriteToString<40>(CacheElement->Key.ObjectKey));
-				UE_LOG(LogAudioStreamCaching, Display, TEXT("[%s] - Cancelled=%d, Size=%u, InKey %s"), *WriteToString<128>(CacheElement->Key.SoundWaveName), bWasCancelled, ChunkDataSize, bMatched ? TEXT("matched") : TEXT("DID NOT MATCH!"));
-				if (!bMatched)
-				{
-					UE_LOG(LogAudioStreamCaching, Display, TEXT("[%s] - InKey %s / %d / %s"), *WriteToString<128>(CacheElement->Key.SoundWaveName), *WriteToString<128>(InKey.SoundWaveName), InKey.ChunkIndex, *WriteToString<40>(InKey.ObjectKey));
-				}
-			}
-#endif
+			// Diagnostic logging.
+			const bool bMatched = CacheElement->Key == InKey;
+			UE_CLOG(CacheElement->WantsFilteredLogging(), LogAudioTiming, Verbose, TEXT("[%s] - AsyncReadComplete index %d, ObjectKey=%s"), *WriteToString<128>(CacheElement->Key.SoundWaveName), CacheElement->Key.ChunkIndex, *WriteToString<40>(CacheElement->Key.ObjectKey));
+			UE_CLOG(CacheElement->WantsFilteredLogging(), LogAudioTiming, Verbose, TEXT("[%s] - Cancelled=%d, Size=%u, InKey %s"), *WriteToString<128>(CacheElement->Key.SoundWaveName), bWasCancelled, ChunkDataSize, bMatched ? TEXT("matched") : TEXT("DID NOT MATCH!"));
+			UE_CLOG(CacheElement->WantsFilteredLogging() && !bMatched, LogAudioTiming, Verbose, TEXT("[%s] - InKey %s / %d / %s"), *WriteToString<128>(CacheElement->Key.SoundWaveName), *WriteToString<128>(InKey.SoundWaveName), InKey.ChunkIndex, *WriteToString<40>(InKey.ObjectKey));
+
 			// Populate key and DataSize. The async read request was set up to write directly into CacheElement->ChunkData.
 			CacheElement->Key = InKey;
 			CacheElement->ChunkDataSize = ChunkDataSize;
@@ -1970,9 +1932,7 @@ void FAudioChunkCache::KickOffAsyncLoad(FCacheElement* CacheElement, const FChun
 		CacheElement->ReadRequest = nullptr;
 		if (Chunk->BulkData.IsBulkDataLoaded())
 		{
-#if !UE_BUILD_SHIPPING
-			UE_CLOG(CacheElement->bWantsExtraLogging, LogAudioStreamCaching, Display, TEXT("[%s] - Chunk %d has bulk data loaded, , ObjectKey=%s"), *WriteToString<128>(CacheElement->Key.SoundWaveName), CacheElement->Key.ChunkIndex, *WriteToString<40>(CacheElement->Key.ObjectKey));
-#endif
+			UE_CLOG(CacheElement->WantsFilteredLogging(), LogAudioTiming, Verbose, TEXT("[%s] - Chunk %d has bulk data loaded, , ObjectKey=%s"), *WriteToString<128>(CacheElement->Key.SoundWaveName), CacheElement->Key.ChunkIndex, *WriteToString<40>(CacheElement->Key.ObjectKey));
 			// If this chunk has been inlined and loaded, move out the data into our newly allocated block.
 			const FBulkDataBuffer<uint8> ChunkMemory = Chunk->MoveOutAsBuffer();
 			
diff --git a/Engine/Source/Runtime/Engine/Private/SoundWave.cpp b/Engine/Source/Runtime/Engine/Private/SoundWave.cpp
--- a/Engine/Source/Runtime/Engine/Private/SoundWave.cpp
+++ b/Engine/Source/Runtime/Engine/Private/SoundWave.cpp
@@ -2,6 +2,7 @@
 
 #include "Sound/SoundWave.h"
 
+#include "Audio/AudioTimingLog.h"
 #include "AudioCompressionSettingsUtils.h"
 #include "AudioDecompress.h"
 #include "AudioDerivedData.h"
@@ -1145,7 +1146,7 @@ void USoundWave::Serialize( FArchive& Ar )
 			if (OwnerBehavior.LoadingBehavior != ESoundWaveLoadingBehavior::Uninitialized)
 			{
 				UE_LOG(LogAudio, VeryVerbose, TEXT("Saving Owners Loading Behavior: %s, %s"), 
-				       EnumToString(OwnerBehavior.LoadingBehavior), *GetFullName());
+					Audio::ShortEnumString(EnumToString(OwnerBehavior.LoadingBehavior)), *GetFullName());
 				
 				Flags |= HasOwnerLoadingBehaviorFlag;
 				Flags |= (uint32)OwnerBehavior.LoadingBehavior << LoadingBehaviorShift;
@@ -2098,7 +2099,18 @@ void USoundWave::PostLoad()
 			// OR we've inlined the first chunk (we do when we're pre-buffering). We don't want these dumped by the stream cache if we can help it.
 			const bool bShouldRetainFirstChunk = bAllowRetaining &&
 				(ActualLoadingBehavior == ESoundWaveLoadingBehavior::RetainOnLoad || bShouldRetainInlinedFirstChunked);
-				
+
+			UE_CLOG(Audio::MatchesLogFilter(*GetName()), LogAudioTiming, Verbose, TEXT("PostLoad: Wave='%s'\tFormat='%s'\tLooping=%s\tChunk1(Retain=%s,Inlined=%s,Size=%lld),LB=%s, NumChunks=%d"), 
+				*GetFName().ToString(),
+				*GetRuntimeFormat().ToString(),
+				ToCStr(LexToString(bLooping)),
+				ToCStr(LexToString(bShouldRetainFirstChunk)),
+				ToCStr(LexToString(bIsFirstChunkInlined)),
+				Chunks.IsValidIndex(1) ? Chunks[1].BulkData.GetElementCount() : 0,	
+				Audio::ShortEnumString(EnumToString(ActualLoadingBehavior)),
+				Chunks.Num()
+			);
+			
 			if (bShouldRetainFirstChunk)
 			{
 				RetainCompressedAudio(false);
@@ -4379,6 +4391,13 @@ void USoundWave::RetainCompressedAudio(bool bForceSync /*= false*/)
 		}
 
 		SoundWaveDataPtr->FirstChunk = IStreamingManager::Get().GetAudioStreamingManager().GetLoadedChunk(Proxy, 1, true);
+		UE_CLOG (Audio::MatchesLogFilter(*GetName()), LogAudioTiming, Verbose, TEXT("RetainCompressedAudio (Sync) on SoundWave: Wave='%s'\tFormat='%s'\tLooping=%s\tChunkSize=%u"), 
+						*SoundWaveDataPtr->GetFName().ToString(),
+						*SoundWaveDataPtr->GetRuntimeFormat().ToString(),
+						ToCStr(LexToString(SoundWaveDataPtr->IsLooping())),
+						SoundWaveDataPtr->FirstChunk.IsValid() ? SoundWaveDataPtr->FirstChunk.Num() : 0
+					);
+
 		UE_CLOG(!(SoundWaveDataPtr->FirstChunk.IsValid()), LogAudio, Display, TEXT("First chunk was invalid after synchronous load in RetainCompressedAudio(). This was likely because the cache was blown. Sound: %s"), *GetFullName());
 	}
 	else
@@ -4393,6 +4412,12 @@ void USoundWave::RetainCompressedAudio(bool bForceSync /*= false*/)
 
 					if (WeakThis.IsValid())
 					{
+						UE_CLOG (Audio::MatchesLogFilter(WeakThis->SoundWaveDataPtr->GetFName()), LogAudioTiming, Verbose, TEXT("Async SetFirstChunk on SoundWave: Wave='%s'\tFormat='%s'\tLooping=%s\tChunkSize=%u"), 
+							*WeakThis->SoundWaveDataPtr->GetFName().ToString(),
+							*WeakThis->SoundWaveDataPtr->GetRuntimeFormat().ToString(),
+							ToCStr(LexToString(WeakThis->SoundWaveDataPtr->IsLooping())),
+							MovedHandle.Num()
+						);
 						WeakThis->SoundWaveDataPtr->FirstChunk = MoveTemp(MovedHandle);
 					}
 				});
diff --git a/Engine/Source/Runtime/Engine/Private/SoundWaveProxyReader.cpp b/Engine/Source/Runtime/Engine/Private/SoundWaveProxyReader.cpp
--- a/Engine/Source/Runtime/Engine/Private/SoundWaveProxyReader.cpp
+++ b/Engine/Source/Runtime/Engine/Private/SoundWaveProxyReader.cpp
@@ -3,6 +3,8 @@
 #include "Sound/SoundWaveProxyReader.h"
 
 #include "Audio.h"
+#include "Audio/AudioDebug.h"
+#include "Audio/AudioTimingLog.h"
 #include "Containers/Array.h"
 #include "Containers/ArrayView.h"
 #include "DSP/BufferVectorOperations.h"
@@ -58,6 +60,9 @@ FSoundWaveProxyReader::FSoundWaveProxyReader(FSoundWaveProxyRef InWaveProxy, con
 	DurationInSeconds = WaveProxy->GetDuration();
 	MaxLoopStartTimeInSeconds = FMath::Max(0, DurationInSeconds - MinLoopDurationInSeconds);
 
+	UE_CLOG(Audio::MatchesLogFilter(WaveProxy->GetFName()),LogAudioTiming, Verbose,
+		TEXT("FSoundWaveProxyReader ctor: Name=%s,NumChannels=%d, Rate=%.2f, NumFrames=%d"), *InWaveProxy->GetFName().ToString(), NumChannels, SampleRate, NumFramesInWave);
+
 	// Clamp times
 	Settings.StartTimeInSeconds = FMath::Clamp(Settings.StartTimeInSeconds, 0.f, DurationInSeconds);
 	Settings.LoopStartTimeInSeconds = ClampLoopStartTime(Settings.LoopStartTimeInSeconds);
@@ -160,6 +165,10 @@ bool FSoundWaveProxyReader::SeekToTime(float InSeconds)
 	{
 		return true;
 	}
+
+	UE_CLOG(Audio::MatchesLogFilter(WaveProxy->GetFName()), LogAudioTiming, Verbose,
+		TEXT("FSoundWaveProxyReader SeekToTime: SeekTime=%.2f, Name=%s, NumChannels=%d, Rate=%.2f, NumFrames=%d"),
+		InSeconds, *WaveProxy->GetFName().ToString(), NumChannels, SampleRate, NumFramesInWave);
 	
 	if (WaveProxy->IsSeekable() && CompressedAudioInfo)
 	{
@@ -198,6 +207,10 @@ bool FSoundWaveProxyReader::SeekToFrame(uint32 InFrameNum)
 		return true;
 	}
 
+	UE_CLOG(Audio::MatchesLogFilter(WaveProxy->GetFName()), LogAudioTiming, Verbose,
+		TEXT("FSoundWaveProxyReader SeekToFrame: SeekFrame=%u, Name=%s, NumChannels=%d, Rate=%.2f, NumFrames=%d"),
+		InFrameNum, *WaveProxy->GetFName().ToString(), NumChannels, SampleRate, NumFramesInWave);
+
 	if (WaveProxy->IsSeekable() && CompressedAudioInfo)
 	{
 		CompressedAudioInfo->SeekToFrame(InFrameNum);
@@ -251,6 +264,10 @@ int32 FSoundWaveProxyReader::PopAudio(Audio::FAlignedFloatBuffer& OutBuffer)
 		NumSamplesUnset -= NumSamplesCopiedThisLoop;
 		OutBufferView = OutBufferView.Slice(NumSamplesCopiedThisLoop, NumSamplesUnset);
 
+		UE_CLOG(Audio::MatchesLogFilter(WaveProxy->GetFName()), LogAudioTiming, VeryVerbose,
+				TEXT("FSoundWaveProxyReader PopAudio, NumSamplesCopiedThisLoop=%d, NumSamplesCopied=%d, NumSamplesUnset=%d, Name=%s, NumChannels=%d, Rate=%.2f, NumFrames=%d"),
+						NumSamplesCopiedThisLoop, NumSamplesCopied, NumSamplesUnset,  *WaveProxy->GetFName().ToString(), NumChannels, SampleRate, NumFramesInWave)
+		
 		if (Settings.bIsLooping)
 		{
 			// Seek to loop start if we have used up all our decodable samples or 
@@ -259,6 +276,10 @@ int32 FSoundWaveProxyReader::PopAudio(Audio::FAlignedFloatBuffer& OutBuffer)
 			{
 				if ((!bDecoderCanDecodeMoreData) || (CurrentFrameIndex >= LoopEndFrameIndex))
 				{
+					UE_CLOG(Audio::MatchesLogFilter(WaveProxy->GetFName()), LogAudioTiming, Verbose,
+								TEXT("FSoundWaveProxyReader PopAudio (looping-return to start) (((!bDecoderCanDecodeMoreData) || (CurrentFrameIndex >= LoopEndFrameIndex))), Name=%s, NumChannels=%d, Rate=%.2f, NumFrames=%d"),
+									*WaveProxy->GetFName().ToString(), NumChannels, SampleRate, NumFramesInWave)
+					
 					SeekToTime(Settings.LoopStartTimeInSeconds);
 					bDecoderCanDecodeMoreData = bIsDecoderValid && (EDecodeResult::Fail != DecodeResult);
 				}
@@ -277,6 +298,10 @@ int32 FSoundWaveProxyReader::PopAudio(Audio::FAlignedFloatBuffer& OutBuffer)
 
 		if (bCanProduceMoreAudio && DecoderOutput.Num() == 0)
 		{
+			UE_CLOG(Audio::MatchesLogFilter(WaveProxy->GetFName()),LogAudioTiming, Verbose,
+				TEXT("FSoundWaveProxyReader PopAudio (break-loop, (bCanProduceMoreAudio && DecoderOutput.Num() == 0), Name=%s, NumChannels=%d, Rate=%.2f, NumFrames=%d"),
+					*WaveProxy->GetFName().ToString(), NumChannels, SampleRate, NumFramesInWave);
+			
 			// we can produce more audio, but we were unable to
 			// this is likely due to the streaming data not being available yet
 			// let's early out to avoid a hitch and hope that it's ready on the next read
@@ -284,6 +309,10 @@ int32 FSoundWaveProxyReader::PopAudio(Audio::FAlignedFloatBuffer& OutBuffer)
 		}
 	}
 
+	UE_CLOG(Audio::MatchesLogFilter(WaveProxy->GetFName()), LogAudioTiming, Verbose,
+		TEXT("FSoundWaveProxyReader PopAudio (loop-end), OutBufferView.Num()=%d, NumSamplesCopied=%d, Name=%s, NumChannels=%d, Rate=%.2f, NumFrames=%d"),
+		OutBufferView.Num(), NumSamplesCopied,  *WaveProxy->GetFName().ToString(), NumChannels, SampleRate, NumFramesInWave);
+	
 	// Zero pad any unset samples. 
 	if (OutBufferView.Num() > 0)
 	{
diff --git a/Engine/Source/Runtime/Engine/Public/Audio/AudioDebug.h b/Engine/Source/Runtime/Engine/Public/Audio/AudioDebug.h
--- a/Engine/Source/Runtime/Engine/Public/Audio/AudioDebug.h
+++ b/Engine/Source/Runtime/Engine/Public/Audio/AudioDebug.h
@@ -17,6 +17,26 @@ namespace Audio
 		}
 		return InString;
 	}
+	
+	FORCEINLINE bool MatchesLogFilter(const TCHAR* InName)
+	{
+#if UE_BUILD_SHIPPING
+		return false;		// Off for shipping.
+#else
+		extern ENGINE_API FString ExtraTracingSoundWaveFilterCVar;
+		return !ExtraTracingSoundWaveFilterCVar.IsEmpty()	// Trivially reject if we have no filter.
+			&& FCString::Strifind(InName, *ExtraTracingSoundWaveFilterCVar,false);
+#endif // UE_BUILD_SHIPPING
+	}
+	
+	FORCEINLINE bool MatchesLogFilter(const FName& InName)
+	{
+#if UE_BUILD_SHIPPING
+		return false;		// Off for shipping.
+#else 		
+		return MatchesLogFilter(*InName.ToString());
+#endif //		
+	}
 }
 
 #if ENABLE_AUDIO_DEBUG
diff --git a/Engine/Source/Runtime/Engine/Public/Audio/AudioTimingLog.h b/Engine/Source/Runtime/Engine/Public/Audio/AudioTimingLog.h
new file mode 100644
--- /dev/null
+++ b/Engine/Source/Runtime/Engine/Public/Audio/AudioTimingLog.h
@@ -0,0 +1,15 @@
+// Copyright Epic Games, Inc. All Rights Reserved.
+#pragma once
+
+#include "CoreTypes.h"
+#include "Logging/LogMacros.h"
+
+// Log AudioTiming. For verbose detailed timings diagnostics and debugging.
+#if UE_BUILD_SHIPPING
+	// Dial back the compiled verbosity in shipping
+	#define AUDIO_TIMING_COMPILED_VERBOSITY Error
+#else // UE_BUILD_SHIPPING
+	#define AUDIO_TIMING_COMPILED_VERBOSITY All
+#endif // UE_BUILD_SHIPPING
+
+ENGINE_API DECLARE_LOG_CATEGORY_EXTERN(LogAudioTiming, Log, AUDIO_TIMING_COMPILED_VERBOSITY);
diff --git a/Engine/Source/Runtime/Engine/Public/AudioDecompress.h b/Engine/Source/Runtime/Engine/Public/AudioDecompress.h
--- a/Engine/Source/Runtime/Engine/Public/AudioDecompress.h
+++ b/Engine/Source/Runtime/Engine/Public/AudioDecompress.h
@@ -171,9 +171,17 @@ protected:
 	mutable bool bHasError = false;
 
 #if !UE_BUILD_SHIPPING
-	// If this sound matches a cvar filter for extra logging.
-	bool bWantsExtraLogging = false;
-#endif
+	bool bWantsFilteredLogging = false;
+#endif //UE_BUILD_SHIPPING
+	bool WantsFilteredLogging() const 
+	{
+#if UE_BUILD_SHIPPING
+		return false;
+#else
+		return bWantsFilteredLogging;
+#endif //UE_BUILD_SHIPPING		
+	}
+	
 	FSoundWaveProxyPtr StreamingSoundWave;
 };
 
diff --git a/Engine/Source/Runtime/Engine/Public/AudioStreamingCache.h b/Engine/Source/Runtime/Engine/Public/AudioStreamingCache.h
--- a/Engine/Source/Runtime/Engine/Public/AudioStreamingCache.h
+++ b/Engine/Source/Runtime/Engine/Public/AudioStreamingCache.h
@@ -305,6 +305,14 @@ private:
 		// with the other extra logging.
 		bool bWantsExtraLogging = false;
 #endif
+		bool WantsFilteredLogging() const
+		{
+#if !UE_BUILD_SHIPPING		
+			return bWantsExtraLogging;
+#else
+			return false;
+#endif //!UE_BUILD_SHIPPING
+		}
 
 		// How many disparate consumers have called GetLoadedChunk.
 		FThreadSafeCounter NumConsumers;
