EditTextSourceString scripting improvements

* Added a variant of UKismetTextLibrary::EditTextSourceString (UKismetTextLibrary::EditTextPropertySourceString) that takes a property name, so that Python can also use it (as it's not a CustomThunk function).
* The two variants share the same common function (UKismetTextLibrary::Generic_EditTextPropertySourceString) internally.
* Updated UKismetTextLibrary::Generic_EditTextPropertySourceString to respect the read-only state of the property (via PropertyAccessUtil).
* Updated UKismetTextLibrary::Generic_EditTextPropertySourceString to emit pre/post edit change events (via PropertyAccessUtil) when in the editor.
* Updated UKismetTextLibrary::Generic_EditTextPropertySourceString to perform archetype value propagation (via PropertyAccessUtil) when in the editor.

diff --git a/Engine/Source/Runtime/Engine/Classes/Kismet/KismetTextLibrary.h b/Engine/Source/Runtime/Engine/Classes/Kismet/KismetTextLibrary.h
--- a/Engine/Source/Runtime/Engine/Classes/Kismet/KismetTextLibrary.h
+++ b/Engine/Source/Runtime/Engine/Classes/Kismet/KismetTextLibrary.h
@@ -9,6 +9,8 @@
 #include "Internationalization/PolyglotTextData.h"
 #include "KismetTextLibrary.generated.h"
 
+class FTextProperty;
+
 #if !CPP
 /** This code is only used for meta-data extraction and the main declaration is with FText, be sure to update there as well */
 
@@ -373,6 +375,22 @@ class UKismetTextLibrary : public UBlueprintFunctionLibrary
 	static ENGINE_API bool EditTextSourceString(UObject* TextOwner, UPARAM(ref) FText& Text, const FString& SourceString);
 	DECLARE_FUNCTION(execEditTextSourceString);
 
+	/**
+	 * Edit the source string of the given text property, akin to what happens when editing a text property in a details panel.
+	 * This will attempt to preserve the existing ID of the text property being edited, or failing that will attempt to build a deterministic ID based on the object and property info.
+	 * 
+	 * @note This is an ADVANCED function that is ONLY safe to be used in environments where the modified text property will be gathered for localization (eg, in the editor, or a game mode that collects text properties to be localized).
+	 * 
+	 * @param TextOwner		The object that owns the given Text to be edited.
+	 * @param PropertyName	The name of the text property to edit. This must be a property that exists on TextOwner.
+	 * @param SourceString	The source string that the edited text property should use.
+	 * 
+	 * @return True if edit was possible, or false if not.
+	 */
+	UFUNCTION(BlueprintCallable, Category="Utilities|Text", meta=(DefaultToSelf="TextOwner"))
+	static ENGINE_API bool EditTextPropertySourceString(UObject* TextOwner, const FName PropertyName, const FString& SourceString);
+	static bool Generic_EditTextPropertySourceString(UObject* TextOwner, const FTextProperty* TextProperty, const FString& SourceString);
+
 	/** Converts a generic value to localized formatted text using the user culture. */
 	UFUNCTION(BlueprintPure, CustomThunk, DisplayName = "To Text", Category = "Utilities|Text", meta = (CustomStructureParam = "Value", BlueprintInternalUseOnly = "true"))
 	static ENGINE_API FText Conv_NumericPropertyToText(UPARAM(ref) const int32& Value);
diff --git a/Engine/Source/Runtime/Engine/Private/KismetTextLibrary.cpp b/Engine/Source/Runtime/Engine/Private/KismetTextLibrary.cpp
--- a/Engine/Source/Runtime/Engine/Private/KismetTextLibrary.cpp
+++ b/Engine/Source/Runtime/Engine/Private/KismetTextLibrary.cpp
@@ -9,6 +9,7 @@
 #include "Internationalization/TextPackageNamespaceUtil.h"
 #include "Misc/RuntimeErrors.h"
 #include "UObject/EnumProperty.h"
+#include "UObject/PropertyAccessUtil.h"
 
 #define LOCTEXT_NAMESPACE "Kismet"
 
@@ -379,7 +380,7 @@ DEFINE_FUNCTION(UKismetTextLibrary::execEditTextSourceString)
 	P_GET_OBJECT(UObject, TextOwner);
 
 	P_GET_PROPERTY_REF(FTextProperty, Text);
-	FTextProperty* TextProperty = CastField<FTextProperty>(Stack.MostRecentProperty);
+	const FTextProperty* TextProperty = CastField<FTextProperty>(Stack.MostRecentProperty);
 
 	P_GET_PROPERTY_REF(FStrProperty, SourceString);
 
@@ -401,43 +402,106 @@ DEFINE_FUNCTION(UKismetTextLibrary::execEditTextSourceString)
 			return;
 		}
 
-		if (!TextOwner->GetClass()->HasProperty(TextProperty))
+		if (Generic_EditTextPropertySourceString(TextOwner, TextProperty, SourceString))
 		{
-			LogRuntimeWarning(FText::Format(LOCTEXT("EditTextSourceString.Warning.InvalidTextProperty", "The given Text resolved to a TextProperty ({0}) that doesn't belong to the given TextOwner ({1})!"), FText::AsCultureInvariant(TextProperty->GetPathName()), FText::AsCultureInvariant(TextOwner->GetPathName())));
-			return;
+			*(bool*)RESULT_PARAM = true;
+			Text = TextProperty->GetPropertyValue_InContainer(TextOwner);
 		}
+	}
+	P_NATIVE_END;
+}
+
+bool UKismetTextLibrary::EditTextPropertySourceString(UObject* TextOwner, const FName PropertyName, const FString& SourceString)
+{
+	if (!TextOwner)
+	{
+		LogRuntimeWarning(LOCTEXT("EditTextPropertySourceString.Warning.NullTextOwner", "The given TextOwner was null!"));
+		return false;
+	}
+
+	const FTextProperty* TextProperty = CastField<FTextProperty>(TextOwner->GetClass()->FindPropertyByName(PropertyName));
+	if (!TextProperty)
+	{
+		LogRuntimeWarning(FText::Format(LOCTEXT("EditTextPropertySourceString.Warning.NullTextProperty", "The given PropertyName ({0}) did not resolve to a TextProperty!"), FText::FromName(PropertyName)));
+		return false;
+	}
+
+	return Generic_EditTextPropertySourceString(TextOwner, TextProperty, SourceString);
+}
+
+bool UKismetTextLibrary::Generic_EditTextPropertySourceString(UObject* TextOwner, const FTextProperty* TextProperty, const FString& SourceString)
+{
+	checkf(TextOwner, TEXT("TextOwner is null!"));
+	checkf(TextProperty, TEXT("TextProperty is null!"));
+
+	if (!TextOwner->GetClass()->HasProperty(TextProperty))
+	{
+		LogRuntimeWarning(FText::Format(LOCTEXT("EditTextPropertySourceString.Warning.InvalidTextProperty", "The resolved TextProperty ({0}) doesn't belong to the given TextOwner ({1})!"), FText::AsCultureInvariant(TextProperty->GetPathName()), FText::AsCultureInvariant(TextOwner->GetPathName())));
+		return false;
+	}
 		
-		auto TextOwnerIsInEditorWorld = [TextOwner]()
-		{
-			const UWorld* WorldContext = TextOwner->GetWorld();
-			return WorldContext && WorldContext->WorldType == EWorldType::Editor;
-		};
+	auto TextOwnerIsInEditorWorld = [TextOwner]()
+	{
+		const UWorld* WorldContext = TextOwner->GetWorld();
+		return WorldContext && WorldContext->WorldType == EWorldType::Editor;
+	};
 
-		auto TextOwnerIsInAsset = [TextOwner]()
+	auto TextOwnerIsInAsset = [TextOwner]()
+	{
+		for (const UObject* Obj = TextOwner; Obj; Obj = Obj->GetOuter())
 		{
-			for (const UObject* Obj = TextOwner; Obj; Obj = Obj->GetOuter())
+			if (Obj->IsAsset())
 			{
-				if (Obj->IsAsset())
-				{
-					return true;
-				}
+				return true;
 			}
-			return false;
-		};
+		}
+		return false;
+	};
+
+	const bool bIsEditingEditorObject = GIsEditor && (TextOwnerIsInEditorWorld() || TextOwnerIsInAsset());
+	const bool bIsEditingTemplateObject = PropertyAccessUtil::IsObjectTemplate(TextOwner);
+
+	if (PropertyAccessUtil::CanSetPropertyValue(TextProperty, bIsEditingEditorObject ? PropertyAccessUtil::EditorReadOnlyFlags : PropertyAccessUtil::RuntimeReadOnlyFlags, bIsEditingTemplateObject) != EPropertyAccessResultFlags::Success)
+	{
+		LogRuntimeWarning(FText::Format(LOCTEXT("EditTextPropertySourceString.Warning.InvalidAccessPermissions", "The resolved TextProperty ({0}) cannot be edited as it is read-only on the given TextOwner ({1})!"), FText::AsCultureInvariant(TextProperty->GetPathName()), FText::AsCultureInvariant(TextOwner->GetPathName())));
+		return false;
+	}
 
-		const bool bApplyPackageNamespace = GIsEditor && (TextOwnerIsInEditorWorld() || TextOwnerIsInAsset());
-		auto DeterministicTextKeyGenerator = [TextOwner, TextProperty, bApplyPackageNamespace]()
+	TArray<UObject*> InheritedInstances;
+	TUniquePtr<FPropertyAccessChangeNotify> ChangeNotify;
+	if (GIsEditor)
+	{
+		if (TextOwner->HasAnyFlags(RF_ClassDefaultObject | RF_ArchetypeObject))
 		{
-			return TextNamespaceUtil::GenerateDeterministicTextKey(TextOwner, TextProperty, bApplyPackageNamespace);
-		};
+			PropertyAccessUtil::GetArchetypeInstancesInheritingPropertyValue(TextProperty, TextOwner, InheritedInstances);
+		}
 
-		if (TextNamespaceUtil::EditTextProperty(TextOwner, TextProperty, TextNamespaceUtil::ETextEditAction::SourceString, SourceString, DeterministicTextKeyGenerator, bApplyPackageNamespace))
+		ChangeNotify = PropertyAccessUtil::BuildBasicChangeNotify(TextProperty, TextOwner, EPropertyAccessChangeNotifyMode::Always, EPropertyChangeType::ValueSet);
+		PropertyAccessUtil::EmitPreChangeNotify(ChangeNotify.Get(), false);
+	}
+
+	auto DeterministicTextKeyGenerator = [TextOwner, TextProperty, bIsEditingEditorObject]()
+	{
+		return TextNamespaceUtil::GenerateDeterministicTextKey(TextOwner, TextProperty, bIsEditingEditorObject);
+	};
+
+	const bool bEditedValue = TextNamespaceUtil::EditTextProperty(TextOwner, TextProperty, TextNamespaceUtil::ETextEditAction::SourceString, SourceString, DeterministicTextKeyGenerator, bIsEditingEditorObject);
+
+	if (GIsEditor)
+	{
+		if (bEditedValue && bIsEditingTemplateObject)
 		{
-			*(bool*)RESULT_PARAM = true;
-			Text = TextProperty->GetPropertyValue_InContainer(TextOwner);
+			// Propagate to archetype instances that had the same value prior to the edit
+			for (UObject* InheritedInstance : InheritedInstances)
+			{
+				PropertyAccessUtil::CopySinglePropertyValue(TextProperty, TextProperty->GetPropertyValuePtr_InContainer(TextOwner), TextProperty, TextProperty->GetPropertyValuePtr_InContainer(InheritedInstance));
+			}
 		}
+
+		PropertyAccessUtil::EmitPostChangeNotify(ChangeNotify.Get(), false);
 	}
-	P_NATIVE_END;
+
+	return bEditedValue;
 }
 
 FText UKismetTextLibrary::Conv_NumericPropertyToText(const int32& Value)
