Skeletal Mesh: Remove speculative fix for multi-threaded morph accesses, now that the real issue is fixed.

#rnx

diff --git a/Engine/Source/Runtime/Engine/Private/SkeletalMesh.cpp b/Engine/Source/Runtime/Engine/Private/SkeletalMesh.cpp
--- a/Engine/Source/Runtime/Engine/Private/SkeletalMesh.cpp
+++ b/Engine/Source/Runtime/Engine/Private/SkeletalMesh.cpp
@@ -4373,10 +4373,6 @@ void USkeletalMesh::UnregisterMorphTarget(UMorphTarget* MorphTarget, bool bInval
 
 void USkeletalMesh::InitMorphTargets(bool bInKeepEmptyMorphTargets)
 {
-	// FIXME: Big hammer fix to prevent reduction + morph target mapping crashing when running on multiple threads. 
-	static FCriticalSection InitMorphTargetMutex;
-	FScopeLock ScopeLock(&InitMorphTargetMutex);
-
 	TRACE_CPUPROFILER_EVENT_SCOPE(USkeletalMesh::InitMorphTargets);
 	GetMorphTargetIndexMap().Empty();
 
@@ -6676,4 +6672,4 @@ void GetSkeletalMeshRefTangentBasis(const USkeletalMesh* Mesh, const FSkeletalMe
 	GetRefTangentBasisTyped(Mesh, Section, LODData.StaticVertexBuffers.StaticMeshVertexBuffer, SkinWeightVertexBuffer, VertIndexInChunk, OutTangentX, OutTangentY, OutTangentZ);
 }
 
-#undef LOCTEXT_NAMESPACE
\ No newline at end of file
+#undef LOCTEXT_NAMESPACE
