Skeletal Mesh: Ensure we read the ground truth skin weight profile data when generating an FMeshDesciption object from the LOD model.

#rnx

diff --git a/Engine/Source/Runtime/Engine/Private/SkeletalMeshLODModel.cpp b/Engine/Source/Runtime/Engine/Private/SkeletalMeshLODModel.cpp
--- a/Engine/Source/Runtime/Engine/Private/SkeletalMeshLODModel.cpp
+++ b/Engine/Source/Runtime/Engine/Private/SkeletalMeshLODModel.cpp
@@ -1581,7 +1581,9 @@ void FSkeletalMeshLODModel::GetMeshDescription(const USkeletalMesh *InSkeletalMe
 		FSkinWeightsVertexAttributesRef SkinWeightAttribute = MeshAttributes.GetVertexSkinWeights(SkinWeightProfileInfo.Key);
 		const FImportedSkinWeightProfileData& SkinWeightProfileData = SkinWeightProfileInfo.Value;
 
-#if 1
+#if 0
+		// This version reads the copied import data influences. This is less precise than the ground truth version below
+		// which is what will actually be used at render time.
 		TMultiMap<int32, int32> InfluenceMap;
 		for (int32 Index = 0; Index < SkinWeightProfileData.SourceModelInfluences.Num(); Index++)
 		{
@@ -1614,15 +1616,35 @@ void FSkeletalMeshLODModel::GetMeshDescription(const USkeletalMesh *InSkeletalMe
 		}
 #else
 		check(SkinWeightProfileData.SkinWeights.Num() == NumVertices);
-		
-		for (int32 Index = 0; Index < SkinWeightProfileData.SkinWeights.Num(); Index++)
+
+		int32 Index = 0;
+		for (int32 SectionIndex = 0; SectionIndex < Sections.Num(); SectionIndex++)
 		{
-			const FRawSkinWeight& RawSkinWeights = SkinWeightProfileData.SkinWeights[Index];
-			const int32 TargetVertexIndex = SourceToTargetVertexMap[Index];
-			const FVertexID VertexID = VertexIDs[TargetVertexIndex];
+			const FSkelMeshSection& Section = Sections[SectionIndex];
+
+			// Convert positions and bone weights
+			const TArray<FBoneIndexType>& BoneMap = Section.BoneMap;
+			for (int32 VertexIndex = 0; VertexIndex < Section.SoftVertices.Num(); VertexIndex++)
+			{
+				FRawSkinWeight RawSkinWeights = SkinWeightProfileData.SkinWeights[Index];
+				const int32 TargetVertexIndex = SourceToTargetVertexMap[Index];
+				const FVertexID VertexID = VertexIDs[TargetVertexIndex];
+
+				for (int32 WeightIndex = 0; WeightIndex < MAX_TOTAL_INFLUENCES; WeightIndex++)
+				{
+					const uint16 BoneWeight = RawSkinWeights.InfluenceWeights[WeightIndex];
+					FBoneIndexType& BoneIndex = RawSkinWeights.InfluenceBones[WeightIndex];
+					if (BoneWeight != 0 && ensure(BoneMap.IsValidIndex(BoneIndex)))
+					{
+						BoneIndex = BoneMap[BoneIndex];
+					}
+				}
 			
-			FBoneWeights Weights = FBoneWeights::Create(RawSkinWeights.InfluenceBones, RawSkinWeights.InfluenceWeights);
-			SkinWeightAttribute.Set(VertexID, Weights);
+				FBoneWeights Weights = FBoneWeights::Create(RawSkinWeights.InfluenceBones, RawSkinWeights.InfluenceWeights);
+				SkinWeightAttribute.Set(VertexID, Weights);
+
+				Index++;
+			}
 		}
 #endif
 	}
diff --git a/Engine/Source/Runtime/Engine/Public/Rendering/SkeletalMeshLODModel.h b/Engine/Source/Runtime/Engine/Public/Rendering/SkeletalMeshLODModel.h
--- a/Engine/Source/Runtime/Engine/Public/Rendering/SkeletalMeshLODModel.h
+++ b/Engine/Source/Runtime/Engine/Public/Rendering/SkeletalMeshLODModel.h
@@ -363,8 +363,11 @@ public:
 	/** Set of vertex attributes identified by an FName */
 	TMap<FName, FSkeletalMeshModelVertexAttribute> VertexAttributes;
 
-	/** Mapping from final mesh vertex index to raw import vertex index. Needed for vertex animation, which only stores positions for import verts. */
+	/** For each point in this LOD model that maps to its source point in the FMeshDescription that produced it.
+	 *  I.e. a linear map mapping the result of GetVertices() to its matching FVertexId entry in the original FMeshDescription.
+	 */
 	TArray<int32>				MeshToImportVertexMap;
+	
 	/** The max index in MeshToImportVertexMap, ie. the number of imported (raw) verts. */
 	int32						MaxImportVertex;
 
@@ -472,7 +475,10 @@ private:
 
 	TArray<FSkelMeshImportedMeshInfo> ImportedMeshInfos;
 
-	/** Editor only data: array of the original point (wedge) indices for each of the vertices in a FSkeletalMeshLODModel */
+	/** For each point in this LOD model that maps to its source point in FSkeletalMeshImportData.
+	 *  I.e. a linear map mapping the result of GetVertices() to its matching FSkeletalMeshImportData::Points entry.
+	 *  This can be used to map back to alternative skin weight profiles and vertex attributes on the import data.
+	 */
 	TArray<uint32>				RawPointIndices2;
 
 	FIntBulkData				RawPointIndices_DEPRECATED;
