[NetPhysics] API to enable and disable compare state/input for resimulation directly on the NetworkPhysicsComponent.
Default value can be set via CVars still and override can be done via NetworkPhysicsSettingsComponent.

diff --git a/Engine/Source/Runtime/Engine/Private/PhysicsEngine/NetworkPhysicsComponent.cpp b/Engine/Source/Runtime/Engine/Private/PhysicsEngine/NetworkPhysicsComponent.cpp
--- a/Engine/Source/Runtime/Engine/Private/PhysicsEngine/NetworkPhysicsComponent.cpp
+++ b/Engine/Source/Runtime/Engine/Private/PhysicsEngine/NetworkPhysicsComponent.cpp
@@ -31,9 +31,11 @@ namespace PhysicsReplicationCVars
 		bool bAllowRewindToClosestState = true;
 		static FAutoConsoleVariableRef CVarResimAllowRewindToClosestState(TEXT("np2.Resim.AllowRewindToClosestState"), bAllowRewindToClosestState, TEXT("When rewinding to a specific frame, if the client doens't have state data for that frame, use closest data available. Only affects the first rewind frame, when FPBDRigidsEvolution is set to Reset."));
 		bool bCompareStateToTriggerRewind = false;
-		static FAutoConsoleVariableRef CVarResimCompareStateToTriggerRewind(TEXT("np2.Resim.CompareStateToTriggerRewind"), bCompareStateToTriggerRewind, TEXT("When true, cache local players custom state struct in rewind history and compare the predicted state with incoming server state to trigger resimulations if they differ, comparison done through FNetworkPhysicsData::CompareData"));
+		static FAutoConsoleVariableRef CVarResimCompareStateToTriggerRewind(TEXT("np2.Resim.CompareStateToTriggerRewind"), bCompareStateToTriggerRewind, TEXT("When true, cache local FNetworkPhysicsData state in rewind history and compare the predicted state with incoming server state to trigger resimulations if they differ, comparison done through FNetworkPhysicsData::CompareData. Only applies if IsLocallyControlled, to enable this for simulated proxies, where IsLocallyControlled is false, also enable np2.Resim.CompareStateToTriggerRewind.IncludeSimProxies)"));
+		bool bCompareStateToTriggerRewindIncludeSimProxies = false;
+		static FAutoConsoleVariableRef CVarResimCompareStateToTriggerRewindIncludeSimProxies(TEXT("np2.Resim.CompareStateToTriggerRewind.IncludeSimProxies"), bCompareStateToTriggerRewindIncludeSimProxies, TEXT("When true, include simulated proxies when np2.Resim.CompareStateToTriggerRewind is enabled."));
 		bool bCompareInputToTriggerRewind = false;
-		static FAutoConsoleVariableRef CVarResimCompareInputToTriggerRewind(TEXT("np2.Resim.CompareInputToTriggerRewind"), bCompareInputToTriggerRewind, TEXT("When true, compare local players predicted inputs with incoming server inputs to trigger resimulations if they differ, comparison done through FNetworkPhysicsData::CompareData."));
+		static FAutoConsoleVariableRef CVarResimCompareInputToTriggerRewind(TEXT("np2.Resim.CompareInputToTriggerRewind"), bCompareInputToTriggerRewind, TEXT("When true, compare local predicted FNetworkPhysicsData input with incoming server inputs to trigger resimulations if they differ, comparison done through FNetworkPhysicsData::CompareData."));
 		bool bEnableUnreliableFlow = true;
 		static FAutoConsoleVariableRef CVarResimEnableUnreliableFlow(TEXT("np2.Resim.EnableUnreliableFlow"), bEnableUnreliableFlow, TEXT("When true, allow data to be sent unreliably. Also sends FNetworkPhysicsData not marked with FNetworkPhysicsData::bimportant unreliably over the network."));
 		bool bEnableReliableFlow = false;
@@ -445,6 +447,18 @@ void UNetworkPhysicsComponent::InitPhysics()
 	{
 		SetNumberOfStatesToNetwork(CVarRedundantStates->GetInt() + 1);
 	}
+	if (const IConsoleVariable* CVarCompareStateToTriggerRewind = IConsoleManager::Get().FindConsoleVariable(TEXT("np2.Resim.CompareStateToTriggerRewind")))
+	{
+		bCompareStateToTriggerRewind = CVarCompareStateToTriggerRewind->GetBool();
+	}
+	if (const IConsoleVariable* CVarCompareStateToTriggerRewind = IConsoleManager::Get().FindConsoleVariable(TEXT("np2.Resim.CompareStateToTriggerRewind.IncludeSimProxies")))
+	{
+		bCompareStateToTriggerRewindIncludeSimProxies = CVarCompareStateToTriggerRewind->GetBool();
+	}
+	if (const IConsoleVariable* CVarCompareInputToTriggerRewind = IConsoleManager::Get().FindConsoleVariable(TEXT("np2.Resim.CompareInputToTriggerRewind")))
+	{
+		bCompareInputToTriggerRewind = CVarCompareInputToTriggerRewind->GetBool();
+	}
 
 	if (AActor* Owner = GetOwner())
 	{
@@ -520,7 +534,11 @@ void UNetworkPhysicsComponent::InitializeComponent()
 				NetworkPhysicsComponent_Internal->RootPhysicsObject = RootPhysicsObject;
 				NetworkPhysicsComponent_Internal->InputsToNetwork = InputsToNetwork;
 				NetworkPhysicsComponent_Internal->StatesToNetwork = StatesToNetwork;
+				NetworkPhysicsComponent_Internal->bCompareStateToTriggerRewind = bCompareStateToTriggerRewind;
+				NetworkPhysicsComponent_Internal->bCompareStateToTriggerRewindIncludeSimProxies = bCompareStateToTriggerRewindIncludeSimProxies;
+				NetworkPhysicsComponent_Internal->bCompareInputToTriggerRewind = bCompareInputToTriggerRewind;
 				CreateAsyncDataHistory();
+				UpdateAsyncComponent(true);
 			}
 		}
 	}
@@ -952,9 +970,30 @@ bool UNetworkPhysicsComponent::IsLocallyControlled() const
 	return false;
 }
 
-void UNetworkPhysicsComponent::SetIsRelayingLocalInputs(bool bInRelayingLocalInputs)
+void UNetworkPhysicsComponent::SetCompareStateToTriggerRewind(const bool bInCompareStateToTriggerRewind, const bool bInIncludeSimProxies)
 {
-	bIsRelayingLocalInputs = bInRelayingLocalInputs;
+	bCompareStateToTriggerRewind = bInCompareStateToTriggerRewind;
+	bCompareStateToTriggerRewindIncludeSimProxies = bInIncludeSimProxies;
+	if (NetworkPhysicsComponent_Internal)
+	{
+		if (FAsyncNetworkPhysicsComponentInput* AsyncInput = NetworkPhysicsComponent_Internal->GetProducerInputData_External())
+		{
+			AsyncInput->bCompareStateToTriggerRewind = bCompareStateToTriggerRewind;
+			AsyncInput->bCompareStateToTriggerRewindIncludeSimProxies = bInIncludeSimProxies;
+		}
+	}
+}
+
+void UNetworkPhysicsComponent::SetCompareInputToTriggerRewind(const bool bInCompareInputToTriggerRewind)
+{
+	bCompareInputToTriggerRewind = bInCompareInputToTriggerRewind;
+	if (NetworkPhysicsComponent_Internal)
+	{
+		if (FAsyncNetworkPhysicsComponentInput* AsyncInput = NetworkPhysicsComponent_Internal->GetProducerInputData_External())
+		{
+			AsyncInput->bCompareInputToTriggerRewind = bCompareInputToTriggerRewind;
+		}
+	}
 }
 
 APlayerController* UNetworkPhysicsComponent::GetPlayerController() const
@@ -999,7 +1038,7 @@ void UNetworkPhysicsComponent::UpdateAsyncComponent(const bool bFullUpdate)
 				}
 			}
 
-			// bIsLocallyControlled is marshaled outside of the bFullUpdate because it's not always set at BeginPlay when last bFullUpdate is called.
+			// bIsLocallyControlled is marshaled outside of the bFullUpdate because it's not always set when last bFullUpdate is called.
 			AsyncInput->bIsLocallyControlled = IsLocallyControlled();
 
 			if (bFullUpdate)
@@ -1110,6 +1149,9 @@ FAsyncNetworkPhysicsComponent::FAsyncNetworkPhysicsComponent() : TSimCallbackObj
 	, PhysicsReplicationMode(EPhysicsReplicationMode::Default)
 	, SettingsComponent(nullptr)
 	, RootPhysicsObject(nullptr)
+	, bCompareStateToTriggerRewind(false)
+	, bCompareStateToTriggerRewindIncludeSimProxies(false)
+	, bCompareInputToTriggerRewind(false)
 {
 }
 
@@ -1204,6 +1246,18 @@ void FAsyncNetworkPhysicsComponent::ConsumeAsyncInput(const int32 PhysicsStep)
 			{
 				UnregisterDataHistoryFromRewindData();
 			}
+			if (AsyncInput->bCompareStateToTriggerRewind.IsSet())
+			{
+				bCompareStateToTriggerRewind = *AsyncInput->bCompareStateToTriggerRewind;
+			}
+			if (AsyncInput->bCompareStateToTriggerRewind.IsSet())
+			{
+				bCompareStateToTriggerRewindIncludeSimProxies = *AsyncInput->bCompareStateToTriggerRewindIncludeSimProxies;
+			}
+			if (AsyncInput->bCompareInputToTriggerRewind.IsSet())
+			{
+				bCompareInputToTriggerRewind = *AsyncInput->bCompareInputToTriggerRewind;
+			}
 			if (AsyncInput->InputHelper.IsSet())
 			{
 				// Setup rewind data if not already done, and get history size
@@ -1232,7 +1286,7 @@ void FAsyncNetworkPhysicsComponent::ConsumeAsyncInput(const int32 PhysicsStep)
 			/** Receive data helper */
 			auto ReceiveHelper = [&](Chaos::FBaseRewindHistory* History, Chaos::FBaseRewindHistory* ReceiveData, const bool bImportant, const bool bCompareData)
 			{
-				const bool bCompareDataForRewind = bCompareData && IsLocallyControlled() && !IsServer();
+				const bool bCompareDataForRewind = bCompareData && !IsServer();
 				const int32 ResimFrame = History->ReceiveNewData(*ReceiveData, (IsServer() ? 0 : NetworkPhysicsTickOffset), bCompareDataForRewind, bImportant);
 				if (bCompareDataForRewind)
 				{
@@ -1250,6 +1304,9 @@ void FAsyncNetworkPhysicsComponent::ConsumeAsyncInput(const int32 PhysicsStep)
 				ReceiveData->ResetFast();
 			};
 
+			const bool bCompareInput = ComponentSettings.GetCompareInputToTriggerRewind(bCompareInputToTriggerRewind) && IsLocallyControlled();
+			const bool bCompareState = ComponentSettings.GetCompareStateToTriggerRewind(bCompareStateToTriggerRewind) && (IsLocallyControlled() || ComponentSettings.GetCompareStateToTriggerRewindIncludeSimProxies(bCompareStateToTriggerRewindIncludeSimProxies));
+
 			// Receive Inputs
 			if (AsyncInput->InputData && AsyncInput->InputData->HasDataInHistory())
 			{
@@ -1276,7 +1333,7 @@ void FAsyncNetworkPhysicsComponent::ConsumeAsyncInput(const int32 PhysicsStep)
 				{
 					AsyncInput->InputData->ValidateDataInHistory(ActorComponent.Get());
 				}
-				ReceiveHelper(InputHistory.Get(), AsyncInput->InputData.Get(), /*bImportant*/false, ComponentSettings.GetCompareInputToTriggerRewind());
+				ReceiveHelper(InputHistory.Get(), AsyncInput->InputData.Get(), /*bImportant*/false, bCompareInput);
 
 				/* If the server-side input history doesn't have any entries ahead of the current physics tick, the input buffer is empty, inject the latest received input as the input for the current tick.
 				* This happens during a desync where the client is far behind the server */
@@ -1298,7 +1355,7 @@ void FAsyncNetworkPhysicsComponent::ConsumeAsyncInput(const int32 PhysicsStep)
 			// Receive States
 			if (AsyncInput->StateData && AsyncInput->StateData->HasDataInHistory())
 			{
-				ReceiveHelper(StateHistory.Get(), AsyncInput->StateData.Get(), /*bImportant*/false, ComponentSettings.GetCompareStateToTriggerRewind());
+				ReceiveHelper(StateHistory.Get(), AsyncInput->StateData.Get(), /*bImportant*/false, bCompareState);
 			}
 
 			// Receive Important Inputs
@@ -1308,7 +1365,7 @@ void FAsyncNetworkPhysicsComponent::ConsumeAsyncInput(const int32 PhysicsStep)
 				{
 					continue;
 				}
-				ReceiveHelper(InputHistory.Get(), InputImportant.Get(), /*bImportant*/true, ComponentSettings.GetCompareInputToTriggerRewind());
+				ReceiveHelper(InputHistory.Get(), InputImportant.Get(), /*bImportant*/true, bCompareInput);
 			}
 
 			// Receive Important States
@@ -1318,7 +1375,7 @@ void FAsyncNetworkPhysicsComponent::ConsumeAsyncInput(const int32 PhysicsStep)
 				{
 					continue;
 				}
-				ReceiveHelper(StateHistory.Get(), StateImportant.Get(), /*bImportant*/true, ComponentSettings.GetCompareStateToTriggerRewind());
+				ReceiveHelper(StateHistory.Get(), StateImportant.Get(), /*bImportant*/true, bCompareState);
 			}
 		}
 	}
@@ -1543,8 +1600,10 @@ void FAsyncNetworkPhysicsComponent::OnPostProcessInputs_Internal(const int32 Phy
 #endif
 		}
 
-		// Cache current state if this is the server of we are comparing predicted states on autonomous proxy
-		const bool bShouldCacheStateHistory = bIsServer || (ComponentSettings.GetCompareStateToTriggerRewind() && bShouldCacheInputHistory);
+		// Cache current state if this is the server or we are comparing predicted states clients
+		const bool bShouldCacheStateHistory = bIsServer
+			|| (ComponentSettings.GetCompareStateToTriggerRewind(bCompareStateToTriggerRewind) && (IsLocallyControlled() || ComponentSettings.GetCompareStateToTriggerRewindIncludeSimProxies(bCompareStateToTriggerRewindIncludeSimProxies)));
+
 		if (StateHistory && StateData && bShouldCacheStateHistory)
 		{
 			// Compute of the local frame coming from the client that was used to generate this state
diff --git a/Engine/Source/Runtime/Engine/Public/Physics/NetworkPhysicsComponent.h b/Engine/Source/Runtime/Engine/Public/Physics/NetworkPhysicsComponent.h
--- a/Engine/Source/Runtime/Engine/Public/Physics/NetworkPhysicsComponent.h
+++ b/Engine/Source/Runtime/Engine/Public/Physics/NetworkPhysicsComponent.h
@@ -766,7 +766,10 @@ public:
 	/** Mark this as controlled through locally relayed inputs rather than controlled as a pawn through a player controller.
 	* Set if NetworkPhysicsComponent is implemented on an AActor instead of APawn and it's currently being fed inputs, or if this is controlled by the server. 
 	* NOTE: The actor for this NetworkPhysicsComponent also needs to be owned by the local client if this is used client-side. */
-	ENGINE_API void SetIsRelayingLocalInputs(bool bInRelayingLocalInputs);
+	ENGINE_API void SetIsRelayingLocalInputs(bool bInRelayingLocalInputs)
+	{
+		bIsRelayingLocalInputs = bInRelayingLocalInputs;
+	}
 
 	/** Stop relaying local inputs after next network send.
 	* Deferred version of SetIsRelayingLocalInputs(false) to ensure that the last replicated data gets sent.
@@ -782,6 +785,13 @@ public:
 	/** Check if this is controlled locally through relayed inputs from autonomous proxy. It's recommended to use IsLocallyControlled() when checking if this is locally controlled. */
 	ENGINE_API const bool GetIsRelayingLocalInputs() const { return bIsRelayingLocalInputs; }
 
+	/** Override the initial value set by CVar np2.Resim.CompareStateToTriggerRewind on initialize -- When true, cache the clients FNetworkPhysicsData state in rewind history for autonomous proxies and compare the predicted state with incoming server state to trigger resimulations if they differ, comparison done through FNetworkPhysicsData::CompareData. 
+	* @param bInIncludeSimProxies overrides value set by CVar np2.Resim.CompareStateToTriggerRewind.IncludeSimProxies on initialize -- When true CompareStateToTriggerRewind is also done for simulated proxies. */
+	ENGINE_API void SetCompareStateToTriggerRewind(const bool bInCompareStateToTriggerRewind, const bool bInIncludeSimProxies = false);
+
+	/** Override the initial value set by CVar np2.Resim.CompareInputToTriggerRewind on initialize -- When true, compare autonomous proxies predicted FNetworkPhysicsData inputs with incoming server inputs to trigger resimulations if they differ, comparison done through FNetworkPhysicsData::CompareData.*/
+	ENGINE_API void SetCompareInputToTriggerRewind(const bool bInCompareInputToTriggerRewind);
+
 	/** Get FAsyncNetworkPhysicsComponent on the Physics Thread */
 	FAsyncNetworkPhysicsComponent* GetNetworkPhysicsComponent_Internal() { return NetworkPhysicsComponent_Internal; }
 
@@ -921,6 +931,11 @@ private:
 	// If we are currently relaying inputs and will stop after next network send.
 	bool bStopRelayingLocalInputsDeferred = false;
 
+	// Compare state / input to trigger rewind via FNetworkPhysicsData::CompareData
+	bool bCompareStateToTriggerRewind = false;
+	bool bCompareStateToTriggerRewindIncludeSimProxies = false; // Include simulated proxies when bCompareStateToTriggerRewind is enabled
+	bool bCompareInputToTriggerRewind = false;
+
 	// ToDo, retrieve from NetworkPhysicsSettingsComponent so changes at runtime gets picked up
 	bool bEnableUnreliableFlow = true;
 	bool bEnableReliableFlow = false;
@@ -1082,6 +1097,9 @@ struct FAsyncNetworkPhysicsComponentInput : public Chaos::FSimCallbackInput
 	TOptional<TUniquePtr<FNetworkPhysicsDataHelper>> StateHelper;
 	TOptional<bool> bRegisterDataHistoryInRewindData;
 	TOptional<bool> bUnregisterDataHistoryFromRewindData;
+	TOptional<bool> bCompareStateToTriggerRewind;
+	TOptional<bool> bCompareStateToTriggerRewindIncludeSimProxies;
+	TOptional<bool> bCompareInputToTriggerRewind;
 	TOptional<FNetworkPhysicsSettingsComponentAsync*> SettingsComponent;
 
 	TUniquePtr<Chaos::FBaseRewindHistory> InputData;
@@ -1103,6 +1121,9 @@ struct FAsyncNetworkPhysicsComponentInput : public Chaos::FSimCallbackInput
 		StateHelper.Reset();
 		bRegisterDataHistoryInRewindData.Reset();
 		bUnregisterDataHistoryFromRewindData.Reset();
+		bCompareStateToTriggerRewind.Reset();
+		bCompareStateToTriggerRewindIncludeSimProxies.Reset();
+		bCompareInputToTriggerRewind.Reset();
 		SettingsComponent.Reset();
 
 		if (InputData)
@@ -1264,6 +1285,15 @@ private:
 	// Send last N number of states each replication call to patch up holes due to packet loss
 	uint16 StatesToNetwork = 1;
 
+	// Cache predicted states and then compare incoming states via FNetworkPhysicsData::CompareData to trigger a resim if they desync
+	bool bCompareStateToTriggerRewind;
+
+	// Include simulated proxies when bCompareStateToTriggerRewind is enabled
+	bool bCompareStateToTriggerRewindIncludeSimProxies;
+
+	// Cache compare incoming inputs with locally predicted inputs via FNetworkPhysicsData::CompareData to trigger a resim if they desync
+	bool bCompareInputToTriggerRewind;
+
 	FDelegateHandle DelegateOnPreProcessInputs_Internal;
 	FDelegateHandle DelegateOnPostProcessInputs_Internal;
 };
diff --git a/Engine/Source/Runtime/Engine/Public/Physics/NetworkPhysicsSettingsComponent.h b/Engine/Source/Runtime/Engine/Public/Physics/NetworkPhysicsSettingsComponent.h
--- a/Engine/Source/Runtime/Engine/Public/Physics/NetworkPhysicsSettingsComponent.h
+++ b/Engine/Source/Runtime/Engine/Public/Physics/NetworkPhysicsSettingsComponent.h
@@ -69,6 +69,7 @@ namespace PhysicsReplicationCVars
 		extern int32 RedundantRemoteInputs;
 		extern int32 RedundantStates;
 		extern bool bCompareStateToTriggerRewind;
+		extern bool bCompareStateToTriggerRewindIncludeSimProxies;
 		extern bool bCompareInputToTriggerRewind;
 		extern bool bApplySimProxyStateAtRuntime;
 		extern bool bApplySimProxyInputAtRuntime;
@@ -417,17 +418,28 @@ struct FNetworkPhysicsSettingsNetworkPhysicsComponent
 
 	UPROPERTY(config, EditDefaultsOnly, Category = "Overrides", Meta = (InlineEditConditionToggle))
 	uint32 bOverrideCompareStateToTriggerRewind : 1;
-	// Overrides CVar: np2.Resim.CompareStateToTriggerRewind -- When true, cache local players custom state struct in rewind history and compare the predicted state with incoming server state to trigger resimulations if they differ, comparison done through FNetworkPhysicsData::CompareData.
+	// Overrides CVar: np2.Resim.CompareStateToTriggerRewind -- When true, cache local FNetworkPhysicsData state in rewind history and compare the predicted state with incoming server state to trigger resimulations if they differ, comparison done through FNetworkPhysicsData::CompareData.
 	UPROPERTY(config, EditDefaultsOnly, Category = "Overrides", Meta = (EditCondition = "bOverrideCompareStateToTriggerRewind"))
 	bool bCompareStateToTriggerRewind = PhysicsReplicationCVars::ResimulationCVars::bCompareStateToTriggerRewind;
-	const bool GetCompareStateToTriggerRewind() const { return bOverrideCompareStateToTriggerRewind ? bCompareStateToTriggerRewind : PhysicsReplicationCVars::ResimulationCVars::bCompareStateToTriggerRewind; }
+	const bool GetCompareStateToTriggerRewind(const bool bDefaultValue) const { return bOverrideCompareStateToTriggerRewind ? bCompareStateToTriggerRewind : bDefaultValue; }
+	UE_DEPRECATED(5.5, "Deprecated, use GetCompareStateToTriggerRewind(const bool bDefaultValue) instead and pass in the value that should apply if not overridden.")
+		const bool GetCompareStateToTriggerRewind() const { return bOverrideCompareStateToTriggerRewind ? bCompareStateToTriggerRewind : PhysicsReplicationCVars::ResimulationCVars::bCompareStateToTriggerRewind; }
+
+	UPROPERTY(config, EditDefaultsOnly, Category = "Overrides", Meta = (InlineEditConditionToggle))
+	uint32 bOverridebCompareStateToTriggerRewindIncludeSimProxies : 1;
+	// Overrides CVar: np2.Resim.CompareStateToTriggerRewind.IncludeSimProxies -- When true, include simulated proxies when np2.Resim.CompareStateToTriggerRewind is enabled.
+	UPROPERTY(config, EditDefaultsOnly, Category = "Overrides", Meta = (EditCondition = "bOverridebCompareStateToTriggerRewindIncludeSimProxies"))
+	bool bCompareStateToTriggerRewindIncludeSimProxies = PhysicsReplicationCVars::ResimulationCVars::bCompareStateToTriggerRewindIncludeSimProxies;
+	const bool GetCompareStateToTriggerRewindIncludeSimProxies(const bool bDefaultValue) const { return bOverridebCompareStateToTriggerRewindIncludeSimProxies ? bCompareStateToTriggerRewindIncludeSimProxies : bDefaultValue; }
 
 	UPROPERTY(config, EditDefaultsOnly, Category = "Overrides", Meta = (InlineEditConditionToggle))
 	uint32 bOverrideCompareInputToTriggerRewind : 1;
-	// Overrides CVar: np2.Resim.CompareInputToTriggerRewind -- When true, compare local players predicted inputs with incoming server inputs to trigger resimulations if they differ, comparison done through FNetworkPhysicsData::CompareData.
+	// Overrides CVar: np2.Resim.CompareInputToTriggerRewind -- When true, compare local predicted FNetworkPhysicsData input with incoming server input to trigger resimulations if they differ, comparison done through FNetworkPhysicsData::CompareData.
 	UPROPERTY(config, EditDefaultsOnly, Category = "Overrides", Meta = (EditCondition = "bOverrideCompareInputToTriggerRewind"))
 	bool bCompareInputToTriggerRewind = PhysicsReplicationCVars::ResimulationCVars::bCompareInputToTriggerRewind;
-	const bool GetCompareInputToTriggerRewind() const { return bOverrideCompareInputToTriggerRewind ? bCompareInputToTriggerRewind : PhysicsReplicationCVars::ResimulationCVars::bCompareInputToTriggerRewind; }
+	const bool GetCompareInputToTriggerRewind(const bool bDefaultValue) const { return bOverrideCompareInputToTriggerRewind ? bCompareInputToTriggerRewind : bDefaultValue; }
+	UE_DEPRECATED(5.5, "Deprecated, use GetCompareInputToTriggerRewind(const bool bDefaultValue) instead and pass in the value that should apply if not overridden.")
+		const bool GetCompareInputToTriggerRewind() const { return bOverrideCompareInputToTriggerRewind ? bCompareInputToTriggerRewind : PhysicsReplicationCVars::ResimulationCVars::bCompareInputToTriggerRewind; }
 
 	UPROPERTY(config, EditDefaultsOnly, Category = "Overrides", Meta = (InlineEditConditionToggle))
 	uint32 bOverrideEnableUnreliableFlow : 1;
