Allow direct setting of the simulation space settings

diff --git a/Engine/Source/Runtime/Engine/Private/PhysicsEngine/ImmediatePhysics/ImmediatePhysicsChaos/ImmediatePhysicsSimulation_Chaos.cpp b/Engine/Source/Runtime/Engine/Private/PhysicsEngine/ImmediatePhysics/ImmediatePhysicsChaos/ImmediatePhysicsSimulation_Chaos.cpp
--- a/Engine/Source/Runtime/Engine/Private/PhysicsEngine/ImmediatePhysics/ImmediatePhysicsChaos/ImmediatePhysicsSimulation_Chaos.cpp
+++ b/Engine/Source/Runtime/Engine/Private/PhysicsEngine/ImmediatePhysics/ImmediatePhysicsChaos/ImmediatePhysicsSimulation_Chaos.cpp
@@ -817,6 +817,11 @@ namespace ImmediatePhysics_Chaos
 		SimSpaceSettings.EulerAlpha = ChaosImmediate_Evolution_SimSpaceEulerAlpha;
 	}
 
+	void FSimulation::SetSimulationSpaceSettings(const Chaos::FSimulationSpaceSettings& SimulationSpaceSettings)
+	{
+		Implementation->Evolution.SetSimulationSpaceSettings(SimulationSpaceSettings);
+	}
+
 	void FSimulation::SetSolverSettings(const FReal FixedDt, const FReal CullDistance, const FReal MaxDepenetrationVelocity, const int32 UseLinearJointSolver, const int32 PositionIts, const int32 VelocityIts, const int32 ProjectionIts, const int32 bUseManifoilds)
 	{
 		if (FixedDt >= FReal(0))
diff --git a/Engine/Source/Runtime/Engine/Public/Physics/ImmediatePhysics/ImmediatePhysicsChaos/ImmediatePhysicsSimulation_Chaos.h b/Engine/Source/Runtime/Engine/Public/Physics/ImmediatePhysics/ImmediatePhysicsChaos/ImmediatePhysicsSimulation_Chaos.h
--- a/Engine/Source/Runtime/Engine/Public/Physics/ImmediatePhysics/ImmediatePhysicsChaos/ImmediatePhysicsSimulation_Chaos.h
+++ b/Engine/Source/Runtime/Engine/Public/Physics/ImmediatePhysics/ImmediatePhysicsChaos/ImmediatePhysicsSimulation_Chaos.h
@@ -15,6 +15,7 @@ namespace Chaos
 	class FPBDJointSolverSettings;
 	class FPBDCollisionSolverSettings;
 	class FCollisionDetectorSettings;
+	class FSimulationSpaceSettings;
 }
 
 namespace ImmediatePhysics_Chaos
@@ -115,6 +116,9 @@ namespace ImmediatePhysics_Chaos
 			const FReal Alpha, 
 			const FVector& ExternalLinearEtherDrag);
 
+		ENGINE_API void SetSimulationSpaceSettings(
+			const Chaos::FSimulationSpaceSettings& SimulationSpaceSettings);
+
 		/** Set settings. Invalid (negative) values with leave that value unchanged from defaults */
 		ENGINE_API void SetSolverSettings(
 			const FReal FixedDt,
