[HWRT] Refactor r.RayTracing.Geometry.StaticMeshes.WPO logic
- Exclude proxy from ray tracing when using r.RayTracing.Geometry.StaticMeshes.WPO=0 instead of marking it as dynamic and not generating any instances.

diff --git a/Engine/Source/Runtime/Engine/Private/SplineMeshSceneProxy.cpp b/Engine/Source/Runtime/Engine/Private/SplineMeshSceneProxy.cpp
--- a/Engine/Source/Runtime/Engine/Private/SplineMeshSceneProxy.cpp
+++ b/Engine/Source/Runtime/Engine/Private/SplineMeshSceneProxy.cpp
@@ -64,8 +64,8 @@ FSplineMeshSceneProxy::FSplineMeshSceneProxy(const FStaticMeshSceneProxyDesc& In
 	bVisibleInLumenScene = false;
 
 #if RHI_RAYTRACING
+	bHasRayTracingRepresentation = true;
 	bDynamicRayTracingGeometry = true;
-	bNeedsDynamicRayTracingGeometries = true;
 #endif
 
 	// make sure all the materials are okay to be rendered as a spline mesh
@@ -174,6 +174,7 @@ void FSplineMeshSceneProxy::GetDynamicRayTracingInstances(FRayTracingInstanceCol
 {
 	if (CVarRayTracingSplineMeshes.GetValueOnRenderThread() == 0  || !bSupportRayTracing)
 	{
+		// TODO: Exclude proxy during ray tracing primitive gather instead of doing this early out here.
 		return;
 	}
 
diff --git a/Engine/Source/Runtime/Engine/Private/StaticMeshSceneProxy.cpp b/Engine/Source/Runtime/Engine/Private/StaticMeshSceneProxy.cpp
--- a/Engine/Source/Runtime/Engine/Private/StaticMeshSceneProxy.cpp
+++ b/Engine/Source/Runtime/Engine/Private/StaticMeshSceneProxy.cpp
@@ -333,25 +333,21 @@ FStaticMeshSceneProxy::FStaticMeshSceneProxy(const FStaticMeshSceneProxyDesc& In
 
 #if RHI_RAYTRACING
 	bSupportRayTracing = IsRayTracingAllowed() && InProxyDesc.GetStaticMesh()->bSupportRayTracing;
+	bHasRayTracingRepresentation = false;
 	bDynamicRayTracingGeometry = false;
-	bNeedsDynamicRayTracingGeometries = false;
 
 	if (bSupportRayTracing)
 	{
 		const bool bWantsRayTracingWPO = MaterialRelevance.bUsesWorldPositionOffset && InProxyDesc.bEvaluateWorldPositionOffsetInRayTracing;
 
-		// r.RayTracing.Geometry.StaticMeshes.WPO is handled in the following way:
-		// 0 - mark ray tracing geometry as dynamic but don't create any dynamic geometries since it won't be included in ray tracing scene
-		// 1 - mark ray tracing geometry as dynamic and create dynamic geometries
-		// 2 - don't mark ray tracing geometry as dynamic nor create any dynamic geometries since WPO evaluation is disabled
-
-		// if r.RayTracing.Geometry.StaticMeshes.WPO == 2, WPO evaluation is disabled so don't need to mark geometry as dynamic
-		if (bWantsRayTracingWPO && CVarRayTracingStaticMeshesWPO.GetValueOnAnyThread() != 2)
+		if (bWantsRayTracingWPO)
 		{
-			bDynamicRayTracingGeometry = true;
-
-			// only need dynamic geometries when r.RayTracing.Geometry.StaticMeshes.WPO == 1
-			bNeedsDynamicRayTracingGeometries = CVarRayTracingStaticMeshesWPO.GetValueOnAnyThread() == 1;
+			bHasRayTracingRepresentation = CVarRayTracingStaticMeshesWPO.GetValueOnAnyThread() != 0;
+			bDynamicRayTracingGeometry = CVarRayTracingStaticMeshesWPO.GetValueOnAnyThread() == 1;
+		}
+		else
+		{
+			bHasRayTracingRepresentation = true;
 		}
 
 		// when ray tracing proxy uses rendering LODs we can use the FLODInfo for ray tracing
@@ -439,39 +435,34 @@ void FStaticMeshSceneProxy::SetEvaluateWorldPositionOffsetInRayTracing(FRHIComma
 		return;
 	}
 
-	// if r.RayTracing.Geometry.StaticMeshes.WPO == 2, WPO evaluation is disabled so don't need to mark geometry as dynamic
-	NewValue &= MaterialRelevance.bUsesWorldPositionOffset && CVarRayTracingStaticMeshesWPO.GetValueOnAnyThread() != 2;
-
-	if (NewValue && !bDynamicRayTracingGeometry)
-	{
-		bDynamicRayTracingGeometry = true;
-
-		// only need dynamic geometries when r.RayTracing.Geometry.StaticMeshes.WPO == 1
-		bNeedsDynamicRayTracingGeometries = CVarRayTracingStaticMeshesWPO.GetValueOnAnyThread() == 1;
+	const bool bOldDynamicRayTracingGeometry = bDynamicRayTracingGeometry;
 
-		if (GetPrimitiveSceneInfo())
-		{
-			GetPrimitiveSceneInfo()->bIsRayTracingStaticRelevant = IsRayTracingStaticRelevant();
-		}
+	const bool bWantsRayTracingWPO = MaterialRelevance.bUsesWorldPositionOffset && NewValue;
 
-		if (bNeedsDynamicRayTracingGeometries)
-		{
-			CreateDynamicRayTracingGeometries(RHICmdList);
-		}
+	if (bWantsRayTracingWPO)
+	{
+		bHasRayTracingRepresentation = CVarRayTracingStaticMeshesWPO.GetValueOnAnyThread() != 0;
+		bDynamicRayTracingGeometry = CVarRayTracingStaticMeshesWPO.GetValueOnAnyThread() == 1;
 	}
-	else if (!NewValue && bDynamicRayTracingGeometry)
+	else
 	{
-		bDynamicRayTracingGeometry = false;
-		bNeedsDynamicRayTracingGeometries = false;
-
-		if (GetPrimitiveSceneInfo())
-		{
-			GetPrimitiveSceneInfo()->bIsRayTracingStaticRelevant = IsRayTracingStaticRelevant();
-		}
+		bHasRayTracingRepresentation = true;
+	}
 
+	if (bDynamicRayTracingGeometry && !bOldDynamicRayTracingGeometry)
+	{
+		CreateDynamicRayTracingGeometries(RHICmdList);
+	}
+	else if (!bDynamicRayTracingGeometry && bOldDynamicRayTracingGeometry)
+	{
 		ReleaseDynamicRayTracingGeometries();
 	}
 
+	if (GetPrimitiveSceneInfo())
+	{
+		GetPrimitiveSceneInfo()->bIsRayTracingStaticRelevant = IsRayTracingStaticRelevant();
+	}
+
 	GetScene().UpdateCachedRayTracingState(this);
 #endif
 }
@@ -789,7 +780,7 @@ bool FStaticMeshSceneProxy::GetRayTracingMeshElement(int32 LODIndex, int32 Batch
 #if RHI_RAYTRACING
 void FStaticMeshSceneProxy::CreateDynamicRayTracingGeometries(FRHICommandListBase& RHICmdList)
 {
-	check(bDynamicRayTracingGeometry && bNeedsDynamicRayTracingGeometries);
+	check(bDynamicRayTracingGeometry);
 	check(DynamicRayTracingGeometries.IsEmpty());
 
 	FStaticMeshRayTracingProxyLODArray& RayTracingLODs = RenderData->RayTracingProxy->LODs;
@@ -814,6 +805,8 @@ void FStaticMeshSceneProxy::CreateDynamicRayTracingGeometries(FRHICommandListBas
 
 void FStaticMeshSceneProxy::ReleaseDynamicRayTracingGeometries()
 {
+	checkf(DynamicRayTracingGeometries.IsEmpty() || bDynamicRayTracingGeometry, TEXT("Proxy shouldn't have DynamicRayTracingGeometries since bDynamicRayTracingGeometry is false."));
+
 	for (auto& Geometry : DynamicRayTracingGeometries)
 	{
 		Geometry.ReleaseResource();
@@ -832,7 +825,7 @@ void FStaticMeshSceneProxy::CreateRenderThreadResources(FRHICommandListBase& RHI
 		RayTracingGeometryGroupHandle = RenderData->RayTracingGeometryGroupHandle;
 	}
 
-	if(IsRayTracingEnabled() && bNeedsDynamicRayTracingGeometries)
+	if(IsRayTracingEnabled() && bDynamicRayTracingGeometry)
 	{
 		CreateDynamicRayTracingGeometries(RHICmdList);
 	}
@@ -1850,7 +1843,7 @@ const FCardRepresentationData* FStaticMeshSceneProxy::GetMeshCardRepresentation(
 #if RHI_RAYTRACING
 bool FStaticMeshSceneProxy::HasRayTracingRepresentation() const
 {
-	return bSupportRayTracing;
+	return bHasRayTracingRepresentation;
 }
 
 TArray<FRayTracingGeometry*> FStaticMeshSceneProxy::GetStaticRayTracingGeometries() const
@@ -1880,26 +1873,28 @@ RayTracing::GeometryGroupHandle FStaticMeshSceneProxy::GetRayTracingGeometryGrou
 
 void FStaticMeshSceneProxy::GetDynamicRayTracingInstances(FRayTracingInstanceCollector& Collector)
 {
-	if (DynamicRayTracingGeometries.IsEmpty() || CVarRayTracingStaticMeshes.GetValueOnRenderThread() == 0)
+#if DO_CHECK
+	if (!ensureMsgf(IsRayTracingRelevant() && bSupportRayTracing && bHasRayTracingRepresentation && bDynamicRayTracingGeometry,
+		TEXT("GetDynamicRayTracingInstances() is only expected to be called for scene proxies that are compatible with ray tracing. ")
+		TEXT("Ray tracing primitive gathering code may be wrong.")))
 	{
 		return;
 	}
+#endif
 
-	checkf(CVarRayTracingStaticMeshesWPO.GetValueOnRenderThread() == 1,
-		TEXT("Proxy should only have entries in DynamicRayTracingGeometries when WPO evaluation is enabled"));
-
-	if (!ensureMsgf(IsRayTracingRelevant(),
-		TEXT("GetDynamicRayTracingInstances() is only expected to be called for scene proxies that are compatible with ray tracing. ")
-		TEXT("RT-relevant primitive gathering code in FDeferredShadingSceneRenderer may be wrong.")))
+	if (CVarRayTracingStaticMeshes.GetValueOnRenderThread() == 0)
 	{
+		// TODO: Exclude proxy during ray tracing primitive gather instead of doing this early out here.
 		return;
 	}
 
+	checkf(!DynamicRayTracingGeometries.IsEmpty(), TEXT("Proxy should have entries in DynamicRayTracingGeometries when using the GetDynamicRayTracingInstances() code path."));
+
 	const ESceneDepthPriorityGroup PrimitiveDPG = GetStaticDepthPriorityGroup();
 
-	bool bEvaluateWPO = bDynamicRayTracingGeometry;
+	bool bEvaluateWPO = true;
 
-	if (bEvaluateWPO && CVarRayTracingStaticMeshesWPOCulling.GetValueOnRenderThread() > 0)
+	if (CVarRayTracingStaticMeshesWPOCulling.GetValueOnRenderThread() > 0)
 	{
 		const FVector ViewCenter = Collector.GetReferenceView()->ViewMatrices.GetViewOrigin();
 		const FVector MeshCenter = GetBounds().Origin;
diff --git a/Engine/Source/Runtime/Engine/Public/StaticMeshSceneProxy.h b/Engine/Source/Runtime/Engine/Public/StaticMeshSceneProxy.h
--- a/Engine/Source/Runtime/Engine/Public/StaticMeshSceneProxy.h
+++ b/Engine/Source/Runtime/Engine/Public/StaticMeshSceneProxy.h
@@ -267,8 +267,8 @@ protected:
 	ENGINE_API void ReleaseDynamicRayTracingGeometries();
 
 	bool bSupportRayTracing : 1;
+	bool bHasRayTracingRepresentation : 1;
 	bool bDynamicRayTracingGeometry : 1;
-	bool bNeedsDynamicRayTracingGeometries : 1;
 
 	TArray<FRayTracingGeometry, TInlineAllocator<MAX_MESH_LOD_COUNT>> DynamicRayTracingGeometries;
 	TArray<FMeshBatch> CachedRayTracingMaterials;
