Add TODO related to DynamicRayTracingGeometries to FInstancedStaticMeshSceneProxy.

diff --git a/Engine/Source/Runtime/Engine/Private/InstancedStaticMesh.cpp b/Engine/Source/Runtime/Engine/Private/InstancedStaticMesh.cpp
--- a/Engine/Source/Runtime/Engine/Private/InstancedStaticMesh.cpp
+++ b/Engine/Source/Runtime/Engine/Private/InstancedStaticMesh.cpp
@@ -1371,6 +1371,9 @@ void FInstancedStaticMeshSceneProxy::SetupProxy(const FInstancedStaticMeshSceneP
 			bAnySegmentUsesWorldPositionOffset |= Section.Material->IsUsingWorldPositionOffset_Concurrent(GMaxRHIFeatureLevel);
 		}
 	}
+	
+	// TODO: when using WPO, the base class (FStaticMeshSceneProxy) currently creates DynamicRayTracingGeometries which is not actually used
+	// because FInstancedStaticMeshSceneProxy creates it's own ray tracing geometries for each simulated instance... 
 
 	if (OverlayMaterial != nullptr)
 	{
