Adding ENGINE_API to various material expression functions.

- All the functions on UMaterialExpressionMaterialFunctionCall that have ENGINE_API are already exposed as ENGINE_API for the base class.

#rnx

diff --git a/Engine/Source/Runtime/Engine/Public/Materials/MaterialExpressionCustom.h b/Engine/Source/Runtime/Engine/Public/Materials/MaterialExpressionCustom.h
--- a/Engine/Source/Runtime/Engine/Public/Materials/MaterialExpressionCustom.h
+++ b/Engine/Source/Runtime/Engine/Public/Materials/MaterialExpressionCustom.h
@@ -91,7 +91,7 @@ class UMaterialExpressionCustom : public UMaterialExpression
 	//~ Begin UObject Interface.
 #if WITH_EDITOR
 	virtual void PostEditChangeProperty(FPropertyChangedEvent& PropertyChangedEvent) override;
-	void RebuildOutputs();
+	ENGINE_API void RebuildOutputs();
 #endif // WITH_EDITOR
 	virtual void Serialize(FStructuredArchive::FRecord Record) override;
 	//~ End UObject Interface.
diff --git a/Engine/Source/Runtime/Engine/Public/Materials/MaterialExpressionMaterialFunctionCall.h b/Engine/Source/Runtime/Engine/Public/Materials/MaterialExpressionMaterialFunctionCall.h
--- a/Engine/Source/Runtime/Engine/Public/Materials/MaterialExpressionMaterialFunctionCall.h
+++ b/Engine/Source/Runtime/Engine/Public/Materials/MaterialExpressionMaterialFunctionCall.h
@@ -99,10 +99,10 @@ class UMaterialExpressionMaterialFunctionCall : public UMaterialExpression
 
 	//~ Begin UObject Interface.
 #if WITH_EDITOR
-	virtual void PreEditChange(FProperty* PropertyAboutToChange) override;
-	virtual void PostEditChangeProperty(FPropertyChangedEvent& PropertyChangedEvent) override;
+	ENGINE_API virtual void PreEditChange(FProperty* PropertyAboutToChange) override;
+	ENGINE_API virtual void PostEditChangeProperty(FPropertyChangedEvent& PropertyChangedEvent) override;
 #endif // WITH_EDITOR
-	virtual void PostLoad() override;
+	ENGINE_API virtual void PostLoad() override;
 	//~ End UObject Interface.
 
 #if WITH_EDITORONLY_DATA
@@ -117,21 +117,21 @@ class UMaterialExpressionMaterialFunctionCall : public UMaterialExpression
 
 	//~ Begin UMaterialExpression Interface
 #if WITH_EDITOR
-	virtual int32 Compile(class FMaterialCompiler* Compiler, int32 OutputIndex) override;
-	virtual void GetCaption(TArray<FString>& OutCaptions) const override;
-	virtual TArrayView<FExpressionInput*> GetInputsView() override;
-	virtual FExpressionInput* GetInput(int32 InputIndex) override;
-	virtual FName GetInputName(int32 InputIndex) const override;
-	virtual bool IsInputConnectionRequired(int32 InputIndex) const override;
-	virtual FString GetDescription() const override;
-	virtual void GetConnectorToolTip(int32 InputIndex, int32 OutputIndex, TArray<FString>& OutToolTip) override;
-	virtual void GetExpressionToolTip(TArray<FString>& OutToolTip) override;
-	virtual bool MatchesSearchQuery( const TCHAR* SearchQuery ) override;
-	virtual bool IsResultMaterialAttributes(int32 OutputIndex) override;
-	virtual bool IsResultSubstrateMaterial(int32 OutputIndex) override;
-	virtual void GatherSubstrateMaterialInfo(FSubstrateMaterialInfo& SubstrateMaterialInfo, int32 OutputIndex) override;
-	virtual FSubstrateOperator* SubstrateGenerateMaterialTopologyTree(class FMaterialCompiler* Compiler, class UMaterialExpression* Parent, int32 OutputIndex) override;
-	virtual uint32 GetInputType(int32 InputIndex) override;
+	ENGINE_API virtual int32 Compile(class FMaterialCompiler* Compiler, int32 OutputIndex) override;
+	ENGINE_API virtual void GetCaption(TArray<FString>& OutCaptions) const override;
+	ENGINE_API virtual TArrayView<FExpressionInput*> GetInputsView() override;
+	ENGINE_API virtual FExpressionInput* GetInput(int32 InputIndex) override;
+	ENGINE_API virtual FName GetInputName(int32 InputIndex) const override;
+	ENGINE_API virtual bool IsInputConnectionRequired(int32 InputIndex) const override;
+	ENGINE_API virtual FString GetDescription() const override;
+	ENGINE_API virtual void GetConnectorToolTip(int32 InputIndex, int32 OutputIndex, TArray<FString>& OutToolTip) override;
+	ENGINE_API virtual void GetExpressionToolTip(TArray<FString>& OutToolTip) override;
+	ENGINE_API virtual bool MatchesSearchQuery(const TCHAR* SearchQuery) override;
+	ENGINE_API virtual bool IsResultMaterialAttributes(int32 OutputIndex) override;
+	ENGINE_API virtual bool IsResultSubstrateMaterial(int32 OutputIndex) override;
+	ENGINE_API virtual void GatherSubstrateMaterialInfo(FSubstrateMaterialInfo& SubstrateMaterialInfo, int32 OutputIndex) override;
+	ENGINE_API virtual FSubstrateOperator* SubstrateGenerateMaterialTopologyTree(class FMaterialCompiler* Compiler, class UMaterialExpression* Parent, int32 OutputIndex) override;
+	ENGINE_API virtual uint32 GetInputType(int32 InputIndex) override;
 	//~ End UMaterialExpression Interface
 
 	/** Util to get name of a particular type, optionally with type appended in parenthesis */
@@ -149,7 +149,7 @@ class UMaterialExpressionMaterialFunctionCall : public UMaterialExpression
 
 	/** */
 	UFUNCTION(BlueprintCallable, Category = "MaterialEditing")
-	ENGINE_API bool SetMaterialFunction(UMaterialFunctionInterface* NewMaterialFunction);
+	ENGINE_API virtual bool SetMaterialFunction(UMaterialFunctionInterface* NewMaterialFunction);
 
 	/** 
 	 * Update FunctionInputs and FunctionOutputs from the MaterialFunction.  
@@ -163,7 +163,7 @@ class UMaterialExpressionMaterialFunctionCall : public UMaterialExpression
 		SharedCompileState = SharedState;
 	}
 
-	virtual bool GenerateHLSLExpression(FMaterialHLSLGenerator& Generator, UE::HLSLTree::FScope& Scope, int32 OutputIndex, UE::HLSLTree::FExpression const*& OutExpression) const override;
+	ENGINE_API virtual bool GenerateHLSLExpression(FMaterialHLSLGenerator& Generator, UE::HLSLTree::FScope& Scope, int32 OutputIndex, UE::HLSLTree::FExpression const*& OutExpression) const override;
 
 private:	
 	/** Helper that fixes up expression links where possible. */
