Net control message template type-checks message arguments.

diff --git a/Engine/Plugins/Online/OnlineSubsystemUtils/Source/OnlineSubsystemUtils/Private/OnlineBeaconHost.cpp b/Engine/Plugins/Online/OnlineSubsystemUtils/Source/OnlineSubsystemUtils/Private/OnlineBeaconHost.cpp
--- a/Engine/Plugins/Online/OnlineSubsystemUtils/Source/OnlineSubsystemUtils/Private/OnlineBeaconHost.cpp
+++ b/Engine/Plugins/Online/OnlineSubsystemUtils/Source/OnlineSubsystemUtils/Private/OnlineBeaconHost.cpp
@@ -218,7 +218,7 @@ bool AOnlineBeaconHost::HandleControlMessage(UNetConnection* Connection, uint8 M
 					LocalNetworkVersion, RemoteNetworkVersion,
 					LocalNetFeaturesDescription.ToString(), RemoteNetFeaturesDescription.ToString()
 				);
-				FNetControlMessage<NMT_Upgrade>::Send(Connection, LocalNetworkVersion);
+				FNetControlMessage<NMT_Upgrade>::Send(Connection, LocalNetworkVersion, LocalNetworkFeatures);
 				return false;
 			}
 			
diff --git a/Engine/Source/Runtime/Engine/Private/Net/NetPing.cpp b/Engine/Source/Runtime/Engine/Private/Net/NetPing.cpp
--- a/Engine/Source/Runtime/Engine/Private/Net/NetPing.cpp
+++ b/Engine/Source/Runtime/Engine/Private/Net/NetPing.cpp
@@ -572,7 +572,7 @@ void FNetPing::UpdatePingTimeout(EPingType PingType)
 			if (FailureAddress.Len() > 0 && Owner)
 			{
 				const uint32 PingTypeVal = static_cast<uint32>(PingType);
-				const ENetPingControlMessage MessageType = ENetPingControlMessage::PingFailure;
+				ENetPingControlMessage MessageType = ENetPingControlMessage::PingFailure;
 				FString MessageStr = FString::Printf(TEXT("%i=%s"), PingTypeVal, ToCStr(FailureAddress));
 
 				FNetControlMessage<NMT_NetPing>::Send(Owner.Get(), MessageType, MessageStr);
diff --git a/Engine/Source/Runtime/Engine/Private/World.cpp b/Engine/Source/Runtime/Engine/Private/World.cpp
--- a/Engine/Source/Runtime/Engine/Private/World.cpp
+++ b/Engine/Source/Runtime/Engine/Private/World.cpp
@@ -6528,7 +6528,7 @@ void UWorld::WelcomePlayer(UNetConnection* Connection)
 	if (FParse::Value(FCommandLine::Get(), TEXT("ClientNetPingICMPAddress="), ClientNetPingICMPAddress) && ClientNetPingICMPAddress.Len() > 0)
 	{
 		const uint32 PingType = static_cast<uint32>(EPingType::ICMP);
-		const ENetPingControlMessage MessageType = ENetPingControlMessage::SetPingAddress;
+		ENetPingControlMessage MessageType = ENetPingControlMessage::SetPingAddress;
 		FString MessageStr = FString::Printf(TEXT("%i=%s"), PingType, ToCStr(ClientNetPingICMPAddress));
 
 		FNetControlMessage<NMT_NetPing>::Send(Connection, MessageType, MessageStr);
@@ -6537,7 +6537,7 @@ void UWorld::WelcomePlayer(UNetConnection* Connection)
 	if (FParse::Value(FCommandLine::Get(), TEXT("ClientNetPingUDPAddress="), ClientNetPingUDPAddress) && ClientNetPingUDPAddress.Len() > 0)
 	{
 		const uint32 PingType = static_cast<uint32>(EPingType::UDPQoS);
-		const ENetPingControlMessage MessageType = ENetPingControlMessage::SetPingAddress;
+		ENetPingControlMessage MessageType = ENetPingControlMessage::SetPingAddress;
 		FString MessageStr = FString::Printf(TEXT("%i=%s"), PingType, ToCStr(ClientNetPingUDPAddress));
 
 		FNetControlMessage<NMT_NetPing>::Send(Connection, MessageType, MessageStr);
diff --git a/Engine/Source/Runtime/Engine/Public/Net/DataChannel.h b/Engine/Source/Runtime/Engine/Public/Net/DataChannel.h
--- a/Engine/Source/Runtime/Engine/Public/Net/DataChannel.h
+++ b/Engine/Source/Runtime/Engine/Public/Net/DataChannel.h
@@ -9,13 +9,20 @@
 #include "CoreMinimal.h"
 #include "Misc/Guid.h"
 #include "Misc/NetworkGuid.h"
+#include "Misc/NetworkVersion.h"
 #include "GameFramework/OnlineReplStructs.h"
+#include "Net/Core/Connection/NetEnums.h"
 #include "Net/DataBunch.h"
 #include "Engine/ChildConnection.h"
 
 // Forward declarations
 enum class ENetPingControlMessage : uint8;
 
+namespace UE::Net
+{
+template<int32 Index, typename... ParamTypes>
+class TNetControlMessageImpl;
+}
 
 /*-----------------------------------------------------------------------------
 	UControlChannel base class.
@@ -47,7 +54,7 @@ public:
 		CheckInitialized();
 		return Names[MessageIndex][0] != 0;
 	}
-	template<uint8 MessageType> friend class FNetControlMessage;
+	template<int32 Index, typename... ParamTypes> friend class UE::Net::TNetControlMessageImpl;
 private:
 	static void CheckInitialized()
 	{
@@ -91,50 +98,62 @@ private:
 	static ENGINE_API bool bInitialized;
 };
 
+namespace UE::Net
+{
+
+template<int32 Index, typename... ParamTypes>
+class TNetControlMessageImpl
+{
+public:
+	static uint8 Initialize(const TCHAR* InName)
+	{
+		FNetControlMessageInfo::SetName(Index, InName);
+		return 0;
+	}
+
+	/**
+	 * sends a message of this type on the specified connection's control channel
+	 * @note: const not used only because of the FArchive interface; the parameters are not modified
+	 */
+	static void Send(UNetConnection* Conn, ParamTypes&... Params)
+	{
+		static_assert(Index < FNetControlMessageInfo::MaxNames, "Control channel message must be a byte.");
+		checkSlow(!Conn->IsA(UChildConnection::StaticClass())); /** control channel messages can only be sent on the parent connection */
+		if (Conn->Channels[0] != nullptr && !Conn->Channels[0]->Closing)
+		{
+			FControlChannelOutBunch Bunch(Conn->Channels[0], false);
+			uint8 MessageType = Index;
+			Bunch << MessageType;
+			FNetControlMessageInfo::SendParams(Bunch, Params...);
+			Conn->Channels[0]->SendBunch(&Bunch, true);
+		}
+	}
+
+	/** receives a message of this type from the passed in bunch */
+	[[nodiscard]] static bool Receive(FInBunch& Bunch, ParamTypes&... Params)
+	{
+		FNetControlMessageInfo::ReceiveParams(Bunch, Params...);
+		return !Bunch.IsError();
+	}
+
+	/** throws away a message of this type from the passed in bunch */
+	static void Discard(FInBunch& Bunch)
+	{
+		TTuple<ParamTypes...> Params;
+		VisitTupleElements([&Bunch](auto& Param)
+			{
+				Bunch << Param;
+			},
+			Params);
+	}
+};
+
+}
+
 #define DEFINE_CONTROL_CHANNEL_MESSAGE(Name, Index, ...) \
 enum { NMT_##Name = Index }; \
-template<> class FNetControlMessage<Index> \
+template<> class FNetControlMessage<Index> : public UE::Net::TNetControlMessageImpl<Index, ##__VA_ARGS__> \
 { \
-public: \
-	static uint8 Initialize() \
-	{ \
-		FNetControlMessageInfo::SetName(Index, TEXT(#Name)); \
-		return 0; \
-	} \
-	/** sends a message of this type on the specified connection's control channel \
-		* @note: const not used only because of the FArchive interface; the parameters are not modified \
-		*/ \
-	template<typename... ParamTypes> \
-	static void Send(UNetConnection* Conn, ParamTypes&... Params) \
-	{ \
-		static_assert(Index < FNetControlMessageInfo::MaxNames, "Control channel message must be a byte."); \
-		checkSlow(!Conn->IsA(UChildConnection::StaticClass())); /** control channel messages can only be sent on the parent connection */ \
-		if (Conn->Channels[0] != NULL && !Conn->Channels[0]->Closing) \
-		{ \
-			FControlChannelOutBunch Bunch(Conn->Channels[0], false); \
-			uint8 MessageType = Index; \
-			Bunch << MessageType; \
-			FNetControlMessageInfo::SendParams(Bunch, Params...); \
-			Conn->Channels[0]->SendBunch(&Bunch, true); \
-		} \
-	} \
-	/** receives a message of this type from the passed in bunch */ \
-	template<typename... ParamTypes> \
-	[[nodiscard]] static bool Receive(FInBunch& Bunch, ParamTypes&... Params) \
-	{ \
-		FNetControlMessageInfo::ReceiveParams(Bunch, Params...); \
-		return !Bunch.IsError(); \
-	} \
-	/** throws away a message of this type from the passed in bunch */ \
-	static void Discard(FInBunch& Bunch) \
-	{ \
-		TTuple<__VA_ARGS__> Params; \
-		VisitTupleElements([&Bunch](auto& Param) \
-		{ \
-			Bunch << Param; \
-		}, \
-		Params); \
-	} \
 };
 
 #define DEFINE_CONTROL_CHANNEL_MESSAGE_ZEROPARAM(Name, Index) DEFINE_CONTROL_CHANNEL_MESSAGE(Name, Index)
@@ -147,12 +166,12 @@ public: \
 #define DEFINE_CONTROL_CHANNEL_MESSAGE_SEVENPARAM(Name, Index, TypeA, TypeB, TypeC, TypeD, TypeE, TypeF, TypeG) DEFINE_CONTROL_CHANNEL_MESSAGE(Name, Index, TypeA, TypeB, TypeC, TypeD, TypeE, TypeF, TypeG)
 #define DEFINE_CONTROL_CHANNEL_MESSAGE_EIGHTPARAM(Name, Index, TypeA, TypeB, TypeC, TypeD, TypeE, TypeF, TypeG, TypeH) DEFINE_CONTROL_CHANNEL_MESSAGE(Name, Index, TypeA, TypeB, TypeC, TypeD, TypeE, TypeF, TypeG, TypeH)
 
-#define IMPLEMENT_CONTROL_CHANNEL_MESSAGE(Name) static uint8 Dummy##_FNetControlMessage_##Name = FNetControlMessage<NMT_##Name>::Initialize();
+#define IMPLEMENT_CONTROL_CHANNEL_MESSAGE(Name) static uint8 Dummy##_FNetControlMessage_##Name = FNetControlMessage<NMT_##Name>::Initialize(TEXT(#Name));
 
 // message type definitions
-DEFINE_CONTROL_CHANNEL_MESSAGE(Hello, 0, uint8, uint32, FString, uint16); // initial client connection message
+DEFINE_CONTROL_CHANNEL_MESSAGE(Hello, 0, uint8, uint32, FString, EEngineNetworkRuntimeFeatures); // initial client connection message
 DEFINE_CONTROL_CHANNEL_MESSAGE(Welcome, 1, FString, FString, FString); // server tells client they're ok'ed to load the server's level
-DEFINE_CONTROL_CHANNEL_MESSAGE(Upgrade, 2, uint32, uint16); // server tells client their version is incompatible
+DEFINE_CONTROL_CHANNEL_MESSAGE(Upgrade, 2, uint32, EEngineNetworkRuntimeFeatures); // server tells client their version is incompatible
 DEFINE_CONTROL_CHANNEL_MESSAGE(Challenge, 3, FString); // server sends client challenge string to verify integrity
 DEFINE_CONTROL_CHANNEL_MESSAGE(Netspeed, 4, int32); // client sends requested transfer rate
 DEFINE_CONTROL_CHANNEL_MESSAGE(Login, 5, FString, FString, FUniqueNetIdRepl, FString); // client requests to be admitted to the game
@@ -196,4 +215,4 @@ DEFINE_CONTROL_CHANNEL_MESSAGE(BeaconAssignGUID, 27, FNetworkGUID); // client as
 DEFINE_CONTROL_CHANNEL_MESSAGE(BeaconNetGUIDAck, 28, FString); // server received NetGUIDAck from client, connection established successfully
 
 DEFINE_CONTROL_CHANNEL_MESSAGE(IrisProtocolMismatch, 29, uint64); // client has a different protocol hash from the server
-DEFINE_CONTROL_CHANNEL_MESSAGE(IrisNetRefHandleError, 30, uint32, uint64); // a specific handle caused an error and we want the remote connection to log all information it has on it
+DEFINE_CONTROL_CHANNEL_MESSAGE(IrisNetRefHandleError, 30, UE::Net::ENetRefHandleError, uint64); // a specific handle caused an error and we want the remote connection to log all information it has on it
