Substrate - clarified a shader compilation error with an actual failure error instead of mysterious HLSL compilation errors.

#tests pc dx12 editor
#rnx

diff --git a/Engine/Source/Runtime/Engine/Private/Materials/HLSLMaterialTranslator.cpp b/Engine/Source/Runtime/Engine/Private/Materials/HLSLMaterialTranslator.cpp
--- a/Engine/Source/Runtime/Engine/Private/Materials/HLSLMaterialTranslator.cpp
+++ b/Engine/Source/Runtime/Engine/Private/Materials/HLSLMaterialTranslator.cpp
@@ -2339,6 +2339,15 @@ void FHLSLMaterialTranslator::TranslateMaterial()
 		{
 			// The material can be null when some entries are automatically generated, for instance in the material layer blending system
 			bSubstrateFrontMaterialIsValid &= SubstrateCtx.SubstrateMaterialEffectiveClosureCount > 0;
+
+			UMaterialInterface* MaterialInterface = Material->GetMaterialInterface();
+			UMaterial* BaseMaterial = MaterialInterface ? MaterialInterface->GetMaterial() : nullptr;
+			if (!bSubstrateFrontMaterialIsValid && BaseMaterial && BaseMaterial->bUseMaterialAttributes)
+			{
+				Errorf(TEXT(" %s [%s]: Substrate - Material has no BSDF: this can happen when the SubstrateConvertMaterialAttributes node is not used and when material attributes are enabled."), *Material->GetDebugName(), *Material->GetAssetPath().ToString());
+				TranslationResult = EHLSLMaterialTranslatorResult::Failure;
+				return;
+			}
 		}
 
 		if (bSubstrateFrontMaterialIsValid)
