PR #12529: Fix the typo by changing "chnaged" to "changed" to prevent misunderstandings.

diff --git a/Engine/Source/Runtime/Engine/Private/AnimationAsset.cpp b/Engine/Source/Runtime/Engine/Private/AnimationAsset.cpp
--- a/Engine/Source/Runtime/Engine/Private/AnimationAsset.cpp
+++ b/Engine/Source/Runtime/Engine/Private/AnimationAsset.cpp
@@ -621,14 +621,14 @@ void UAnimationAsset::ValidateParentAsset()
 	{
 		if (ParentAsset->GetSkeleton() != GetSkeleton())
 		{
-			// parent asset chnaged skeleton, so we'll have to discard parent asset
+			// parent asset changed skeleton, so we'll have to discard parent asset
 			UE_LOG(LogAnimation, Warning, TEXT("%s: ParentAsset %s linked to different skeleton. Removing the reference."), *GetName(), *GetNameSafe(ParentAsset));
 			ParentAsset = nullptr;
 			Modify();
 		}
 		else if (ParentAsset->StaticClass() != StaticClass())
 		{
-			// parent asset chnaged skeleton, so we'll have to discard parent asset
+			// parent asset changed skeleton, so we'll have to discard parent asset
 			UE_LOG(LogAnimation, Warning, TEXT("%s: ParentAsset %s class type doesn't match. Removing the reference."), *GetName(), *GetNameSafe(ParentAsset));
 			ParentAsset = nullptr;
 			Modify();
