Fixed CIS following CL 38540608

diff --git a/Engine/Source/Runtime/Engine/Public/Materials/MaterialInstanceDynamic.h b/Engine/Source/Runtime/Engine/Public/Materials/MaterialInstanceDynamic.h
--- a/Engine/Source/Runtime/Engine/Public/Materials/MaterialInstanceDynamic.h
+++ b/Engine/Source/Runtime/Engine/Public/Materials/MaterialInstanceDynamic.h
@@ -217,7 +217,7 @@ protected:
 	ENGINE_API void InitializeMID(class UMaterialInterface* ParentMaterial);
 
 #if !WITH_EDITOR
-	virtual bool CanCacheTexturesSamplingInfo() const override
+	ENGINE_API virtual bool CanCacheTexturesSamplingInfo() const override
 	{
 		return false;
 	}
diff --git a/Engine/Source/Runtime/Engine/Public/Materials/MaterialInterface.h b/Engine/Source/Runtime/Engine/Public/Materials/MaterialInterface.h
--- a/Engine/Source/Runtime/Engine/Public/Materials/MaterialInterface.h
+++ b/Engine/Source/Runtime/Engine/Public/Materials/MaterialInterface.h
@@ -357,8 +357,8 @@ protected:
 	TArray<TObjectPtr<UAssetUserData>> AssetUserData;
 
 #if !WITH_EDITOR
-	void CacheTexturesSamplingInfo();
-	virtual bool CanCacheTexturesSamplingInfo() const;
+	ENGINE_API void CacheTexturesSamplingInfo();
+	ENGINE_API virtual bool CanCacheTexturesSamplingInfo() const;
 #endif
 
 	/** Pre-cached texture sampling information used for texture streaming (calculated on load) **/
