MVVM: fix Set with Broadcast node causing an infinite loop when used on an empty object

diff --git a/Engine/Source/Runtime/Engine/Private/FieldNotification/FieldNotificationLibrary.cpp b/Engine/Source/Runtime/Engine/Private/FieldNotification/FieldNotificationLibrary.cpp
--- a/Engine/Source/Runtime/Engine/Private/FieldNotification/FieldNotificationLibrary.cpp
+++ b/Engine/Source/Runtime/Engine/Private/FieldNotification/FieldNotificationLibrary.cpp
@@ -154,7 +154,7 @@ DEFINE_FUNCTION(UFieldNotificationLibrary::execSetPropertyValueAndBroadcast)
 			LOCTEXT("MissingInputProperty", "Failed to resolve the input parameter for SetPropertyValueAndBroadcast.")
 		);
 		FBlueprintCoreDelegates::ThrowScriptException(P_THIS, Stack, ExceptionInfo);
-		return;
+		Stack.Code++;
 	}
 
 	P_GET_OBJECT(UObject, Self);
@@ -206,7 +206,7 @@ DEFINE_FUNCTION(UFieldNotificationLibrary::execSetPropertyValueAndBroadcastField
 			LOCTEXT("MissingInputProperty", "Failed to resolve the input parameter for SetPropertyValueAndBroadcast.")
 		);
 		FBlueprintCoreDelegates::ThrowScriptException(P_THIS, Stack, ExceptionInfo);
-		return;
+		Stack.Code++;
 	}
 
 	P_GET_OBJECT(UObject, Self);
