[WorldPartitionActorDesc] Fix an issue with invalid editor / runtime bounds when updating class descriptor

diff --git a/Engine/Source/Runtime/Engine/Private/WorldPartition/WorldPartitionActorDescUtils.cpp b/Engine/Source/Runtime/Engine/Private/WorldPartition/WorldPartitionActorDescUtils.cpp
--- a/Engine/Source/Runtime/Engine/Private/WorldPartition/WorldPartitionActorDescUtils.cpp
+++ b/Engine/Source/Runtime/Engine/Private/WorldPartition/WorldPartitionActorDescUtils.cpp
@@ -261,6 +261,7 @@ void FWorldPartitionActorDescUtils::UpdateActorDescriptorFromActor(const AActor*
 void FWorldPartitionActorDescUtils::UpdateActorDescriptorFromActorDescriptor(TUniquePtr<FWorldPartitionActorDesc>& InActorDesc, TUniquePtr<FWorldPartitionActorDesc>& OutActorDesc)
 {
 	InActorDesc->TransferFrom(OutActorDesc.Get());
+	InActorDesc->UpdateActorToWorld();
 	OutActorDesc = MoveTemp(InActorDesc);
 }
 
