Ensure input bone index is valid before attempting usage

diff --git a/Engine/Source/Runtime/Engine/Private/Components/SkinnedMeshComponent.cpp b/Engine/Source/Runtime/Engine/Private/Components/SkinnedMeshComponent.cpp
--- a/Engine/Source/Runtime/Engine/Private/Components/SkinnedMeshComponent.cpp
+++ b/Engine/Source/Runtime/Engine/Private/Components/SkinnedMeshComponent.cpp
@@ -2023,7 +2023,7 @@ FTransform USkinnedMeshComponent::GetBoneTransform(int32 BoneIdx, const FTransfo
 			// isn't valid so for now return identity
 			return FTransform::Identity;
 		}
-		if(BoneIdx < LeaderBoneMap.Num())
+		if(BoneIdx >= 0 && BoneIdx < LeaderBoneMap.Num())
 		{
 			const int32 LeaderBoneIndex = LeaderBoneMap[BoneIdx];
 			const int32 NumLeaderTransforms = LeaderPoseComponentInst->GetNumComponentSpaceTransforms();
