Add GRHIName to EngineAnalytics's SessionStart event.
Requested by EngineQA to understand penetration of different RHIs on the market.
#rnx
#tests Compiled

diff --git a/Engine/Source/Runtime/Engine/Private/EngineAnalytics.cpp b/Engine/Source/Runtime/Engine/Private/EngineAnalytics.cpp
--- a/Engine/Source/Runtime/Engine/Private/EngineAnalytics.cpp
+++ b/Engine/Source/Runtime/Engine/Private/EngineAnalytics.cpp
@@ -273,6 +273,7 @@ void FEngineAnalytics::AppendMachineStats(TArray<FAnalyticsEventAttribute>& Even
 	EventAttributes.Emplace(TEXT("GRHIDeviceRevision"), GRHIDeviceRevision);
 	EventAttributes.Emplace(TEXT("GRHIAdapterInternalDriverVersion"), GRHIAdapterInternalDriverVersion);
 	EventAttributes.Emplace(TEXT("GRHIAdapterUserDriverVersion"), GRHIAdapterUserDriverVersion);
+	EventAttributes.Emplace(TEXT("GRHIName"), FString(GDynamicRHI ? GDynamicRHI->GetName() : TEXT("")));
 	EventAttributes.Emplace(TEXT("TotalPhysicalRAM"), static_cast<uint64>(Stats.TotalPhysical));
 	EventAttributes.Emplace(TEXT("CPUPhysicalCores"), FPlatformMisc::NumberOfCores());
 	EventAttributes.Emplace(TEXT("CPULogicalCores"), FPlatformMisc::NumberOfCoresIncludingHyperthreads());
