Remove unused functions.

diff --git a/Engine/Source/Runtime/Engine/Public/MaterialShaderType.h b/Engine/Source/Runtime/Engine/Public/MaterialShaderType.h
--- a/Engine/Source/Runtime/Engine/Public/MaterialShaderType.h
+++ b/Engine/Source/Runtime/Engine/Public/MaterialShaderType.h
@@ -323,12 +323,10 @@ struct FMaterialShaders
 	template<typename ShaderType> inline bool TryGetGeometryShader(TShaderRef<ShaderType>& OutShader) const { return TryGetShader(SF_Geometry, OutShader); }
 	template<typename ShaderType> inline bool TryGetMeshShader(TShaderRef<ShaderType>& OutShader) const { return TryGetShader(SF_Mesh, OutShader); }
 	template<typename ShaderType> inline bool TryGetComputeShader(TShaderRef<ShaderType>& OutShader) const { return TryGetShader(SF_Compute, OutShader); }
-	template<typename ShaderType> inline bool TryGetWorkGraphShader(TShaderRef<ShaderType>& OutShader) const { return TryGetShader(SF_WorkGraphComputeNode, OutShader); }
 
 	template<typename ShaderType> inline bool TryGetVertexShader(TShaderRef<ShaderType>* OutShader) const { return TryGetShader(SF_Vertex, OutShader); }
 	template<typename ShaderType> inline bool TryGetPixelShader(TShaderRef<ShaderType>* OutShader) const { return TryGetShader(SF_Pixel, OutShader); }
 	template<typename ShaderType> inline bool TryGetGeometryShader(TShaderRef<ShaderType>* OutShader) const { return TryGetShader(SF_Geometry, OutShader); }
 	template<typename ShaderType> inline bool TryGetMeshShader(TShaderRef<ShaderType>* OutShader) const { return TryGetShader(SF_Mesh, OutShader); }
 	template<typename ShaderType> inline bool TryGetComputeShader(TShaderRef<ShaderType>* OutShader) const { return TryGetShader(SF_Compute, OutShader); }
-	template<typename ShaderType> inline bool TryGetWorkGraphShader(TShaderRef<ShaderType>* OutShader) const { return TryGetShader(SF_WorkGraphComputeNode, OutShader); }
 };
