Saved 4k of memory in FViewInfo by removing inlined memory arrays from batched elements.

diff --git a/Engine/Source/Runtime/Engine/Public/BatchedElements.h b/Engine/Source/Runtime/Engine/Public/BatchedElements.h
--- a/Engine/Source/Runtime/Engine/Public/BatchedElements.h
+++ b/Engine/Source/Runtime/Engine/Public/BatchedElements.h
@@ -325,7 +325,7 @@ private:
 		/** largest vertex index used by this batch */
 		uint32 MaxVertex;
 		/** index buffer for triangles */
-		TArray<uint16,TInlineAllocator<6> > Indices;
+		TArray<uint16> Indices;
 		/** all triangles in this batch draw with the same texture */
 		const FTexture* Texture;
 		/** Parameters for this batched element */
@@ -341,8 +341,8 @@ private:
 	/** Max number of mesh vertices that will fit in a DrawPriUP call */
 	int32 MaxMeshVerticesAllowed;
 
-	TArray<FBatchedMeshElement,TInlineAllocator<2> > MeshElements;
-	TArray<FSimpleElementVertex,TInlineAllocator<4> > MeshVertices;
+	TArray<FBatchedMeshElement> MeshElements;
+	TArray<FSimpleElementVertex> MeshVertices;
 
 	/**
 	 * Sets the appropriate vertex and pixel shader.
