[HWRT] Deprecated BuildRequestIndex/RayTracingGeometryHandle and use FBuildRequestIndex/FGeometryHandle instead.

diff --git a/Engine/Source/Runtime/Engine/Private/Rendering/RayTracingGeometryManager.cpp b/Engine/Source/Runtime/Engine/Private/Rendering/RayTracingGeometryManager.cpp
--- a/Engine/Source/Runtime/Engine/Private/Rendering/RayTracingGeometryManager.cpp
+++ b/Engine/Source/Runtime/Engine/Private/Rendering/RayTracingGeometryManager.cpp
@@ -167,7 +167,7 @@ static float GetInitialBuildPriority(ERTAccelerationStructureBuildPriority InBui
 	}
 }
 
-FRayTracingGeometryManager::BuildRequestIndex FRayTracingGeometryManager::RequestBuildAccelerationStructure(FRayTracingGeometry* InGeometry, ERTAccelerationStructureBuildPriority InPriority, EAccelerationStructureBuildMode InBuildMode)
+FRayTracingGeometryManager::FBuildRequestIndex FRayTracingGeometryManager::RequestBuildAccelerationStructure(FRayTracingGeometry* InGeometry, ERTAccelerationStructureBuildPriority InPriority, EAccelerationStructureBuildMode InBuildMode)
 {
 	check(InGeometry->RayTracingBuildRequestIndex == INDEX_NONE);
 
@@ -177,7 +177,7 @@ FRayTracingGeometryManager::BuildRequestIndex FRayTracingGeometryManager::Reques
 	Request.BuildMode = EAccelerationStructureBuildMode::Build;
 
 	FScopeLock ScopeLock(&RequestCS);
-	BuildRequestIndex RequestIndex = GeometryBuildRequests.Add(Request);
+	FBuildRequestIndex RequestIndex = GeometryBuildRequests.Add(Request);
 	GeometryBuildRequests[RequestIndex].RequestIndex = RequestIndex;
 
 	INC_DWORD_STAT(STAT_RayTracingPendingBuilds);
@@ -188,7 +188,7 @@ FRayTracingGeometryManager::BuildRequestIndex FRayTracingGeometryManager::Reques
 	return RequestIndex;
 }
 
-void FRayTracingGeometryManager::RemoveBuildRequest(BuildRequestIndex InRequestIndex)
+void FRayTracingGeometryManager::RemoveBuildRequest(FBuildRequestIndex InRequestIndex)
 {
 	FScopeLock ScopeLock(&RequestCS);
 
@@ -231,7 +231,7 @@ void FRayTracingGeometryManager::ReleaseRayTracingGeometryGroupReference(RayTrac
 
 	if (Group.NumReferences == 0)
 	{
-		for (RayTracingGeometryHandle GeometryHandle : Group.GeometryHandles)
+		for (FGeometryHandle GeometryHandle : Group.GeometryHandles)
 		{
 			checkf(GeometryHandle == INDEX_NONE, TEXT("All FRayTracingGeometry in a group must be unregistered before releasing the group."));
 		}
@@ -245,13 +245,13 @@ void FRayTracingGeometryManager::ReleaseRayTracingGeometryGroupReference(RayTrac
 	}
 }
 
-FRayTracingGeometryManager::RayTracingGeometryHandle FRayTracingGeometryManager::RegisterRayTracingGeometry(FRayTracingGeometry* InGeometry)
+FRayTracingGeometryManager::FGeometryHandle FRayTracingGeometryManager::RegisterRayTracingGeometry(FRayTracingGeometry* InGeometry)
 {
 	check(InGeometry);
 
 	FScopeLock ScopeLock(&MainCS);
 
-	RayTracingGeometryHandle Handle = RegisteredGeometries.Add({});
+	FGeometryHandle Handle = RegisteredGeometries.Add({});
 
 	FRegisteredGeometry& RegisteredGeometry = RegisteredGeometries[Handle];
 	RegisteredGeometry.Geometry = InGeometry;
@@ -290,7 +290,7 @@ FRayTracingGeometryManager::RayTracingGeometryHandle FRayTracingGeometryManager:
 	return Handle;
 }
 
-void FRayTracingGeometryManager::ReleaseRayTracingGeometryHandle(RayTracingGeometryHandle Handle)
+void FRayTracingGeometryManager::ReleaseRayTracingGeometryHandle(FGeometryHandle Handle)
 {
 	check(Handle != INDEX_NONE);
 
@@ -371,7 +371,7 @@ void FRayTracingGeometryManager::SetRayTracingGeometryGroupCurrentFirstLODIndex(
 		FRHIResourceReplaceBatcher Batcher(RHICmdList, NewCurrentFirstLODIdx - Group.CurrentFirstLODIdx);
 		for (int32 LODIdx = Group.CurrentFirstLODIdx; LODIdx < NewCurrentFirstLODIdx; ++LODIdx)
 		{
-			RayTracingGeometryHandle GeometryHandle = Group.GeometryHandles[LODIdx];
+			FGeometryHandle GeometryHandle = Group.GeometryHandles[LODIdx];
 
 			// some LODs might be stripped during cook
 			// skeletal meshes only create static LOD when rendering as static
@@ -407,7 +407,7 @@ void FRayTracingGeometryManager::SetRayTracingGeometryGroupCurrentFirstLODIndex(
 	Group.CurrentFirstLODIdx = NewCurrentFirstLODIdx;
 }
 
-void FRayTracingGeometryManager::RefreshRegisteredGeometry(RayTracingGeometryHandle Handle)
+void FRayTracingGeometryManager::RefreshRegisteredGeometry(FGeometryHandle Handle)
 {
 	FScopeLock ScopeLock(&MainCS);
 
@@ -618,7 +618,7 @@ void FRayTracingGeometryManager::Tick(FRHICommandList& RHICmdList)
 		{
 			// make always resident geometries actually resident
 
-			for (RayTracingGeometryHandle GeometryHandle : AlwaysResidentGeometries)
+			for (FGeometryHandle GeometryHandle : AlwaysResidentGeometries)
 			{
 				FRegisteredGeometry& RegisteredGeometry = RegisteredGeometries[GeometryHandle];
 
@@ -635,16 +635,16 @@ void FRayTracingGeometryManager::Tick(FRHICommandList& RHICmdList)
 			}
 		}
 
-		TSet<RayTracingGeometryHandle> NotReferencedResidentGeometries = EvictableGeometries;
+		TSet<FGeometryHandle> NotReferencedResidentGeometries = EvictableGeometries;
 
-		TArray<RayTracingGeometryHandle> ReferencedGeometries;
+		TArray<FGeometryHandle> ReferencedGeometries;
 
 		uint64 RequestedSize = 0;
 		uint64 RequestedButEvictedSize = 0;
 
 		// Step 1
 		// - update LastReferencedFrame of referenced geometries and calculate memory required to make evicted geometries resident
-		for (RayTracingGeometryHandle GeometryHandle : ReferencedGeometryHandles)
+		for (FGeometryHandle GeometryHandle : ReferencedGeometryHandles)
 		{
 			FRegisteredGeometry& RegisteredGeometry = RegisteredGeometries[GeometryHandle];
 			RegisteredGeometry.LastReferencedFrame = GFrameCounterRenderThread;
@@ -673,7 +673,7 @@ void FRayTracingGeometryManager::Tick(FRHICommandList& RHICmdList)
 
 			for (uint8 LODIndex = Group.CurrentFirstLODIdx; LODIndex < Group.GeometryHandles.Num(); ++LODIndex)
 			{
-				RayTracingGeometryHandle GeometryHandle = Group.GeometryHandles[LODIndex];
+				FGeometryHandle GeometryHandle = Group.GeometryHandles[LODIndex];
 
 				if (GeometryHandle != INDEX_NONE) // some LODs might be stripped during cook
 				{
@@ -713,13 +713,13 @@ void FRayTracingGeometryManager::Tick(FRHICommandList& RHICmdList)
 		// - if making requested geometries resident will put us over budget -> evict some geometry not referenced by TLAS
 		if (TotalResidentSize + RequestedButEvictedSize > ResidentGeometryMemoryPoolSize)
 		{
-			TArray<RayTracingGeometryHandle> NotReferencedResidentGeometriesArray = NotReferencedResidentGeometries.Array();
+			TArray<FGeometryHandle> NotReferencedResidentGeometriesArray = NotReferencedResidentGeometries.Array();
 
 			// Step 3.1
 			// - sort to evict geometries in the following order:
 			//		- least recently used
 			//		- largest geometries
-			Algo::Sort(NotReferencedResidentGeometriesArray, [this](RayTracingGeometryHandle& LHSHandle, RayTracingGeometryHandle& RHSHandle)
+			Algo::Sort(NotReferencedResidentGeometriesArray, [this](FGeometryHandle& LHSHandle, FGeometryHandle& RHSHandle)
 				{
 					FRegisteredGeometry& LHS = RegisteredGeometries[LHSHandle];
 					FRegisteredGeometry& RHS = RegisteredGeometries[RHSHandle];
@@ -742,7 +742,7 @@ void FRayTracingGeometryManager::Tick(FRHICommandList& RHICmdList)
 			int32 Index = 0;
 			while (TotalResidentSize + RequestedButEvictedSize > ResidentGeometryMemoryPoolSize && Index < NotReferencedResidentGeometriesArray.Num())
 			{
-				RayTracingGeometryHandle GeometryHandle = NotReferencedResidentGeometriesArray[Index];
+				FGeometryHandle GeometryHandle = NotReferencedResidentGeometriesArray[Index];
 				FRegisteredGeometry& RegisteredGeometry = RegisteredGeometries[GeometryHandle];
 
 				check(RegisteredGeometry.Geometry->IsValid() && !RegisteredGeometry.Geometry->IsEvicted());
@@ -759,7 +759,7 @@ void FRayTracingGeometryManager::Tick(FRHICommandList& RHICmdList)
 		{
 			// Step 4.1
 			//	- sort by size to prioritize smaller geometries
-			Algo::Sort(ReferencedGeometries, [this](RayTracingGeometryHandle& LHSHandle, RayTracingGeometryHandle& RHSHandle)
+			Algo::Sort(ReferencedGeometries, [this](FGeometryHandle& LHSHandle, FGeometryHandle& RHSHandle)
 				{
 					FRegisteredGeometry& LHS = RegisteredGeometries[LHSHandle];
 					FRegisteredGeometry& RHS = RegisteredGeometries[RHSHandle];
@@ -774,7 +774,7 @@ void FRayTracingGeometryManager::Tick(FRHICommandList& RHICmdList)
 			{
 				// if referenced this frame, mark for eviction and add to pending list
 
-				RayTracingGeometryHandle GeometryHandle = ReferencedGeometries[Index];
+				FGeometryHandle GeometryHandle = ReferencedGeometries[Index];
 				FRegisteredGeometry& RegisteredGeometry = RegisteredGeometries[GeometryHandle];
 
 				if (RegisteredGeometry.Geometry->IsEvicted())
@@ -826,11 +826,11 @@ void FRayTracingGeometryManager::Tick(FRHICommandList& RHICmdList)
 	}
 
 	{
-		TSet<RayTracingGeometryHandle> CurrentPendingStreamingRequests;
+		TSet<FGeometryHandle> CurrentPendingStreamingRequests;
 		Swap(CurrentPendingStreamingRequests, PendingStreamingRequests);
 		PendingStreamingRequests.Reserve(CurrentPendingStreamingRequests.Num());
 
-		for (RayTracingGeometryHandle GeometryHandle : CurrentPendingStreamingRequests)
+		for (FGeometryHandle GeometryHandle : CurrentPendingStreamingRequests)
 		{
 			if (!RequestRayTracingGeometryStreamIn(RHICmdList, GeometryHandle))
 			{
@@ -855,7 +855,7 @@ void FRayTracingGeometryManager::Tick(FRHICommandList& RHICmdList)
 	CSV_CUSTOM_STAT(RayTracingGeometry, TotalAlwaysResidentSizeMB, TotalAlwaysResidentSize / 1024.0f / 1024.0f, ECsvCustomStatOp::Set);
 }
 
-bool FRayTracingGeometryManager::RequestRayTracingGeometryStreamIn(FRHICommandList& RHICmdList, RayTracingGeometryHandle GeometryHandle)
+bool FRayTracingGeometryManager::RequestRayTracingGeometryStreamIn(FRHICommandList& RHICmdList, FGeometryHandle GeometryHandle)
 {
 	FRegisteredGeometry& RegisteredGeometry = RegisteredGeometries[GeometryHandle];
 	FRayTracingGeometry* Geometry = RegisteredGeometry.Geometry;
@@ -1046,7 +1046,7 @@ void FRayTracingGeometryManager::EvictGeometry(FRHICommandListBase& RHICmdList,
 	RegisteredGeometry.Geometry->Evict();
 }
 
-void FRayTracingGeometryManager::BoostPriority(BuildRequestIndex InRequestIndex, float InBoostValue)
+void FRayTracingGeometryManager::BoostPriority(FBuildRequestIndex InRequestIndex, float InBoostValue)
 {
 	FScopeLock ScopeLock(&RequestCS);
 	GeometryBuildRequests[InRequestIndex].BuildPriority += InBoostValue;
diff --git a/Engine/Source/Runtime/Engine/Public/Rendering/RayTracingGeometryManager.h b/Engine/Source/Runtime/Engine/Public/Rendering/RayTracingGeometryManager.h
--- a/Engine/Source/Runtime/Engine/Public/Rendering/RayTracingGeometryManager.h
+++ b/Engine/Source/Runtime/Engine/Public/Rendering/RayTracingGeometryManager.h
@@ -21,20 +21,20 @@ public:
 	ENGINE_API FRayTracingGeometryManager();
 	ENGINE_API virtual ~FRayTracingGeometryManager();
 
-	ENGINE_API virtual BuildRequestIndex RequestBuildAccelerationStructure(FRayTracingGeometry* InGeometry, ERTAccelerationStructureBuildPriority InPriority, EAccelerationStructureBuildMode InBuildMode) override;
+	ENGINE_API virtual FBuildRequestIndex RequestBuildAccelerationStructure(FRayTracingGeometry* InGeometry, ERTAccelerationStructureBuildPriority InPriority, EAccelerationStructureBuildMode InBuildMode) override;
 
-	ENGINE_API virtual void RemoveBuildRequest(BuildRequestIndex InRequestIndex) override;
-	ENGINE_API virtual void BoostPriority(BuildRequestIndex InRequestIndex, float InBoostValue) override;
+	ENGINE_API virtual void RemoveBuildRequest(FBuildRequestIndex InRequestIndex) override;
+	ENGINE_API virtual void BoostPriority(FBuildRequestIndex InRequestIndex, float InBoostValue) override;
 	ENGINE_API virtual void ForceBuildIfPending(FRHIComputeCommandList& InCmdList, const TArrayView<const FRayTracingGeometry*> InGeometries) override;
 	ENGINE_API virtual void ProcessBuildRequests(FRHIComputeCommandList& InCmdList, bool bInBuildAll = false) override;
 
-	ENGINE_API virtual RayTracingGeometryHandle RegisterRayTracingGeometry(FRayTracingGeometry* InGeometry) override;
-	ENGINE_API virtual void ReleaseRayTracingGeometryHandle(RayTracingGeometryHandle Handle) override;
+	ENGINE_API virtual FGeometryHandle RegisterRayTracingGeometry(FRayTracingGeometry* InGeometry) override;
+	ENGINE_API virtual void ReleaseRayTracingGeometryHandle(FGeometryHandle Handle) override;
 
 	ENGINE_API virtual RayTracing::FGeometryGroupHandle RegisterRayTracingGeometryGroup(uint32 NumLODs, uint32 CurrentFirstLODIdx = 0) override;
 	ENGINE_API virtual void ReleaseRayTracingGeometryGroup(RayTracing::FGeometryGroupHandle Handle) override;
 
-	ENGINE_API virtual void RefreshRegisteredGeometry(RayTracingGeometryHandle Handle) override;
+	ENGINE_API virtual void RefreshRegisteredGeometry(FGeometryHandle Handle) override;
 
 	void SetRayTracingGeometryStreamingData(const FRayTracingGeometry* Geometry, FRayTracingStreamableAsset& StreamableAsset);
 	void SetRayTracingGeometryGroupCurrentFirstLODIndex(FRHICommandListBase& RHICmdList, RayTracing::FGeometryGroupHandle Handle, uint8 CurrentFirstLODIdx);
@@ -59,20 +59,20 @@ private:
 
 	struct FBuildRequest
 	{
-		BuildRequestIndex RequestIndex = INDEX_NONE;
+		FBuildRequestIndex RequestIndex = INDEX_NONE;
 
 		float BuildPriority = 0.0f;
 		FRayTracingGeometry* Owner;
 		EAccelerationStructureBuildMode BuildMode;
 
-		// TODO: Implement use-after-free checks in BuildRequestIndex using some bits to identify generation
+		// TODO: Implement use-after-free checks in FBuildRequestIndex using some bits to identify generation
 	};
 
 	void SetupBuildParams(const FBuildRequest& InBuildRequest, TArray<FRayTracingGeometryBuildParams>& InBuildParams, bool bRemoveFromRequestArray = true);
 
 	void ReleaseRayTracingGeometryGroupReference(RayTracing::FGeometryGroupHandle Handle);
 
-	bool RequestRayTracingGeometryStreamIn(FRHICommandList& RHICmdList, RayTracingGeometryHandle GeometryHandle);
+	bool RequestRayTracingGeometryStreamIn(FRHICommandList& RHICmdList, FGeometryHandle GeometryHandle);
 	void ProcessCompletedStreamingRequests(FRHICommandList& RHICmdList);
 
 	FCriticalSection RequestCS;
@@ -89,7 +89,7 @@ private:
 
 	struct FRayTracingGeometryGroup
 	{
-		TArray<RayTracingGeometryHandle> GeometryHandles;
+		TArray<FGeometryHandle> GeometryHandles;
 
 		TSet<FPrimitiveSceneProxy*> ProxiesWithCachedRayTracingState;
 
@@ -129,7 +129,7 @@ private:
 		FIoBuffer RequestBuffer;
 		FBulkDataBatchRequest Request;
 
-		RayTracingGeometryHandle GeometryHandle = INDEX_NONE;
+		FGeometryHandle GeometryHandle = INDEX_NONE;
 
 		bool bCancelled = false;
 
@@ -172,18 +172,18 @@ private:
 	// Used for keeping track of geometries when ray tracing is dynamic
 	TSparseArray<FRegisteredGeometry> RegisteredGeometries;
 
-	TSet<RayTracingGeometryHandle> ResidentGeometries;
+	TSet<FGeometryHandle> ResidentGeometries;
 	uint64 TotalResidentSize = 0;
 
-	TSet<RayTracingGeometryHandle> AlwaysResidentGeometries;
+	TSet<FGeometryHandle> AlwaysResidentGeometries;
 	uint64 TotalAlwaysResidentSize = 0;
 
-	TSet<RayTracingGeometryHandle> EvictableGeometries;
+	TSet<FGeometryHandle> EvictableGeometries;
 
-	TSet<RayTracingGeometryHandle> ReferencedGeometryHandles;
+	TSet<FGeometryHandle> ReferencedGeometryHandles;
 	TSet<RayTracing::FGeometryGroupHandle> ReferencedGeometryGroups;
 
-	TSet<RayTracingGeometryHandle> PendingStreamingRequests;
+	TSet<FGeometryHandle> PendingStreamingRequests;
 
 	TArray<FStreamingRequest> StreamingRequests;
 	int32 NumStreamingRequests = 0;
diff --git a/Engine/Source/Runtime/RenderCore/Public/RayTracingGeometryManagerInterface.h b/Engine/Source/Runtime/RenderCore/Public/RayTracingGeometryManagerInterface.h
--- a/Engine/Source/Runtime/RenderCore/Public/RayTracingGeometryManagerInterface.h
+++ b/Engine/Source/Runtime/RenderCore/Public/RayTracingGeometryManagerInterface.h
@@ -24,25 +24,28 @@ class IRayTracingGeometryManager
 {
 public:
 
-	using BuildRequestIndex = int32;
-	using RayTracingGeometryHandle = int32;
+	using FBuildRequestIndex = int32;
+	using BuildRequestIndex UE_DEPRECATED(5.6, "Use FBuildRequestIndex instead.") = FBuildRequestIndex;
+
+	using FGeometryHandle = int32;
+	using RayTracingGeometryHandle UE_DEPRECATED(5.6, "Use FGeometryHandle instead.") = FGeometryHandle;
 
 	virtual ~IRayTracingGeometryManager() = default;
 
-	RENDERCORE_API virtual BuildRequestIndex RequestBuildAccelerationStructure(FRayTracingGeometry* InGeometry, ERTAccelerationStructureBuildPriority InPriority, EAccelerationStructureBuildMode InBuildMode) = 0;
+	RENDERCORE_API virtual FBuildRequestIndex RequestBuildAccelerationStructure(FRayTracingGeometry* InGeometry, ERTAccelerationStructureBuildPriority InPriority, EAccelerationStructureBuildMode InBuildMode) = 0;
 
-	BuildRequestIndex RequestBuildAccelerationStructure(FRayTracingGeometry* InGeometry, ERTAccelerationStructureBuildPriority InPriority)
+	FBuildRequestIndex RequestBuildAccelerationStructure(FRayTracingGeometry* InGeometry, ERTAccelerationStructureBuildPriority InPriority)
 	{
 		return RequestBuildAccelerationStructure(InGeometry, InPriority, EAccelerationStructureBuildMode::Build);
 	}
 
-	RENDERCORE_API virtual void RemoveBuildRequest(BuildRequestIndex InRequestIndex) = 0;
-	RENDERCORE_API virtual void BoostPriority(BuildRequestIndex InRequestIndex, float InBoostValue) = 0;
+	RENDERCORE_API virtual void RemoveBuildRequest(FBuildRequestIndex InRequestIndex) = 0;
+	RENDERCORE_API virtual void BoostPriority(FBuildRequestIndex InRequestIndex, float InBoostValue) = 0;
 	RENDERCORE_API virtual void ForceBuildIfPending(FRHIComputeCommandList& InCmdList, const TArrayView<const FRayTracingGeometry*> InGeometries) = 0;
 	RENDERCORE_API virtual void ProcessBuildRequests(FRHIComputeCommandList& InCmdList, bool bInBuildAll = false) = 0;
 
-	RENDERCORE_API virtual RayTracingGeometryHandle RegisterRayTracingGeometry(FRayTracingGeometry* InGeometry) = 0;
-	RENDERCORE_API virtual void ReleaseRayTracingGeometryHandle(RayTracingGeometryHandle Handle) = 0;
+	RENDERCORE_API virtual FGeometryHandle RegisterRayTracingGeometry(FRayTracingGeometry* InGeometry) = 0;
+	RENDERCORE_API virtual void ReleaseRayTracingGeometryHandle(FGeometryHandle Handle) = 0;
 
 	/*
 	* RayTracing::FGeometryGroupHandle is used to group multiple FRayTracingGeometry that are associated with the same asset.
@@ -54,7 +57,7 @@ public:
 
 	RENDERCORE_API virtual void RequestUpdateCachedRenderState(RayTracing::FGeometryGroupHandle InRayTracingGeometryGroupHandle) = 0;
 
-	RENDERCORE_API virtual void RefreshRegisteredGeometry(RayTracingGeometryHandle Handle) = 0;
+	RENDERCORE_API virtual void RefreshRegisteredGeometry(FGeometryHandle Handle) = 0;
 
 	RENDERCORE_API virtual void PreRender() = 0;
 	RENDERCORE_API virtual void Tick(FRHICommandList& RHICmdList) = 0;
