Fix to physical material mask visualisation for a mask that hasn't had a texture set for it yet. Previously this attempted to register a nullptr material proxy and hit an assertion.

diff --git a/Engine/Source/Runtime/Engine/Private/StaticMeshSceneProxy.cpp b/Engine/Source/Runtime/Engine/Private/StaticMeshSceneProxy.cpp
--- a/Engine/Source/Runtime/Engine/Private/StaticMeshSceneProxy.cpp
+++ b/Engine/Source/Runtime/Engine/Private/StaticMeshSceneProxy.cpp
@@ -1671,12 +1671,12 @@ void FStaticMeshSceneProxy::GetDynamicMeshElements(const TArray<const FSceneView
 														PhysMatMaskVisualizationMaterialInstance = new FColoredTexturedMaterialRenderProxy(
 															PhysMatMaskVisualizationMaterial->GetRenderProxy(),
 															FLinearColor::White, NAME_Color, PhysicalMaterialMask->MaskTexture, NAME_LinearColor);
-													}
 
-													Collector.RegisterOneFrameMaterialProxy(PhysMatMaskVisualizationMaterialInstance);
-													MeshElement.MaterialRenderProxy = PhysMatMaskVisualizationMaterialInstance;
+														Collector.RegisterOneFrameMaterialProxy(PhysMatMaskVisualizationMaterialInstance);
+														MeshElement.MaterialRenderProxy = PhysMatMaskVisualizationMaterialInstance;
 
-													bDebugMaterialRenderProxySet = true;
+														bDebugMaterialRenderProxySet = true;
+													}
 												}
 											}
 										}
