Substrate - glints coverage is accounted when glints are disabled to maintain coherent material look.

#tests PD dx12 editor
#rnx

diff --git a/Engine/Source/Runtime/Core/Private/UObject/FortniteMainBranchObjectVersion.cpp b/Engine/Source/Runtime/Core/Private/UObject/FortniteMainBranchObjectVersion.cpp
--- a/Engine/Source/Runtime/Core/Private/UObject/FortniteMainBranchObjectVersion.cpp
+++ b/Engine/Source/Runtime/Core/Private/UObject/FortniteMainBranchObjectVersion.cpp
@@ -14,6 +14,6 @@ TMap<FGuid, FGuid> FFortniteMainBranchObjectVersion::GetSystemGuids()
 	SystemGuids.Add(DevGuids.Niagara_LatestScriptCompileVersion, FGuid("393081597FA7234C92F00869DE262D73"));
 	SystemGuids.Add(DevGuids.SkeletalMeshDerivedDataVersion, FGuid("2F90180769024483BB7D90AACEB2798D"));
 	SystemGuids.Add(DevGuids.STATICMESH_DERIVEDDATA_VER, FGuid("656E338C78C34641A5D3BFAE303997AF"));
-	SystemGuids.Add(DevGuids.MaterialTranslationDDCVersion, FGuid("5DFD770145674B3AAFA544176CF7372B"));
+	SystemGuids.Add(DevGuids.MaterialTranslationDDCVersion, FGuid("33FAE17575F24CBEB88622EEB1450EF4"));
 	return SystemGuids;
 }
diff --git a/Engine/Source/Runtime/Engine/Private/Materials/HLSLMaterialTranslator.cpp b/Engine/Source/Runtime/Engine/Private/Materials/HLSLMaterialTranslator.cpp
--- a/Engine/Source/Runtime/Engine/Private/Materials/HLSLMaterialTranslator.cpp
+++ b/Engine/Source/Runtime/Engine/Private/Materials/HLSLMaterialTranslator.cpp
@@ -13315,9 +13315,16 @@ bool FHLSLMaterialTranslator::FSubstrateCompilationContext::SubstrateGenerateDer
 					{
 					case SUBSTRATE_BSDF_TYPE_SLAB:
 					{
-						SubstrateMaterialComplexity.bIsSimple = SubstrateMaterialComplexity.bIsSimple && !bMayHaveColoredWeight && !It.Has(ESubstrateBsdfFeature_ComplexMask) && !It.Has(ESubstrateBsdfFeature_ComplexSpecialMask) && !It.Has(ESubstrateBsdfFeature_SingleMask);
-						SubstrateMaterialComplexity.bIsSingle = SubstrateMaterialComplexity.bIsSingle && !bMayHaveColoredWeight && !It.Has(ESubstrateBsdfFeature_ComplexMask) && !It.Has(ESubstrateBsdfFeature_ComplexSpecialMask);
-						SubstrateMaterialComplexity.bIsComplexSpecial |= It.Has(ESubstrateBsdfFeature_ComplexSpecialMask);
+						bool bIsComplexSpecial = false;
+						if (It.Has(ESubstrateBsdfFeature_ComplexSpecialMask))
+						{
+							// We need to check each features if really enabled for a platofrms
+							bIsComplexSpecial = It.Has(ESubstrateBsdfFeature_Glint) && Substrate::IsGlintEnabled(ShaderPlatform);
+						}
+
+						SubstrateMaterialComplexity.bIsSimple = SubstrateMaterialComplexity.bIsSimple && !bMayHaveColoredWeight && !It.Has(ESubstrateBsdfFeature_ComplexMask) && !bIsComplexSpecial && !It.Has(ESubstrateBsdfFeature_SingleMask);
+						SubstrateMaterialComplexity.bIsSingle = SubstrateMaterialComplexity.bIsSingle && !bMayHaveColoredWeight && !It.Has(ESubstrateBsdfFeature_ComplexMask) && !bIsComplexSpecial;
+						SubstrateMaterialComplexity.bIsComplexSpecial |= bIsComplexSpecial;
 						break;
 					}
 					case SUBSTRATE_BSDF_TYPE_HAIR:
@@ -13468,7 +13475,7 @@ bool FHLSLMaterialTranslator::FSubstrateCompilationContext::SubstrateGenerateDer
 						{
 							SubstrateMaterialRequestedSizeByte += UintByteSize;
 						}
-						if (It.Has(ESubstrateBsdfFeature_Glint))
+						if (It.Has(ESubstrateBsdfFeature_Glint) && Substrate::IsGlintEnabled(ShaderPlatform))
 						{
 							SubstrateMaterialRequestedSizeByte += UintByteSize;
 							SubstrateMaterialRequestedSizeByte += UintByteSize;
diff --git a/Engine/Source/Runtime/Engine/Private/Materials/MaterialExpressionSubstrate.cpp b/Engine/Source/Runtime/Engine/Private/Materials/MaterialExpressionSubstrate.cpp
--- a/Engine/Source/Runtime/Engine/Private/Materials/MaterialExpressionSubstrate.cpp
+++ b/Engine/Source/Runtime/Engine/Private/Materials/MaterialExpressionSubstrate.cpp
@@ -955,8 +955,11 @@ int32 UMaterialExpressionSubstrateSlabBSDF::CompileDefaultSlab(FMaterialCompiler
 }
 #endif //ENABLE_MATERIAL_LAYER_PROTOTYPE
 
-FSubstrateMaterialComplexity UMaterialExpressionSubstrateSlabBSDF::GetComplexity() const
+FSubstrateMaterialComplexity UMaterialExpressionSubstrateSlabBSDF::GetHighestComplexity() const
 {
+	// This function returns the highest complexity of a materials.
+	// It will be lowered depending on features enabled per platform.
+
 	uint32 FeatureMask = ESubstrateBsdfFeature_None;
 	if (HasGlint()) 			{ FeatureMask |= ESubstrateBsdfFeature_Glint; } 
 	if (HasAnisotropy()) 		{ FeatureMask |= ESubstrateBsdfFeature_Anisotropy; } 
@@ -992,13 +995,13 @@ FSubstrateMaterialComplexity UMaterialExpressionSubstrateSlabBSDF::GetComplexity
 void UMaterialExpressionSubstrateSlabBSDF::GetCaption(TArray<FString>& OutCaptions) const
 {
 	// The node complexity is manually maintained to match FSubstrateCompilationContext::SubstrateGenerateDerivedMaterialOperatorData and shaders.
-	OutCaptions.Add(TEXT("Substrate Slab BSDF - ") + FSubstrateMaterialComplexity::ToString(GetComplexity().SubstrateMaterialType()));
+	OutCaptions.Add(TEXT("Substrate Slab BSDF - ") + FSubstrateMaterialComplexity::ToString(GetHighestComplexity().SubstrateMaterialType()));
 }
 
 void UMaterialExpressionSubstrateSlabBSDF::GetExpressionToolTip(TArray<FString>& OutToolTip)
 {
 	OutToolTip.Add(TEXT("Substrate Slab BSDF"));
-	OutToolTip.Add(TEXT("Complexity = ") + FSubstrateMaterialComplexity::ToString(GetComplexity().SubstrateMaterialType()));
+	OutToolTip.Add(TEXT("Complexity = ") + FSubstrateMaterialComplexity::ToString(GetHighestComplexity().SubstrateMaterialType()));
 	OutToolTip.Add(TEXT("The complexity represents the cost of the shading path (Lighting, Lumen, SSS) the material will follow:"));
 	OutToolTip.Add(TEXT(" - Simple means the Slab only relies on Diffuse, F0 and Roughness. It will follow a fast shading path."));
 	OutToolTip.Add(TEXT(" - Single means the Slab uses more features such as F90, Fuzz, Second Roughness, MFP or SSS. It will follow a more expenssive shading path."));
@@ -1283,8 +1286,9 @@ bool UMaterialExpressionSubstrateSlabBSDF::HasAnisotropy() const
 
 bool UMaterialExpressionSubstrateSlabBSDF::HasGlint() const
 {
-	// Use IsGlintEnabled to not promote to ComplexSpecial render path if glint rendering is not enabled and avoid registering such possibility at runtime.
-	return GlintValue.IsConnected() && Substrate::IsGlintEnabled();
+	// We do not check Substrate::IsGlintEnabled() here. Because we want the glint coverage to affect lower platforms, so the data must flow.
+	// The Translator will disable glint and reduce memory footprint if required for a plaform.
+	return GlintValue.IsConnected();
 }
 
 bool UMaterialExpressionSubstrateSlabBSDF::HasSpecularProfile() const
diff --git a/Engine/Source/Runtime/Engine/Public/Materials/MaterialExpressionSubstrate.h b/Engine/Source/Runtime/Engine/Public/Materials/MaterialExpressionSubstrate.h
--- a/Engine/Source/Runtime/Engine/Public/Materials/MaterialExpressionSubstrate.h
+++ b/Engine/Source/Runtime/Engine/Public/Materials/MaterialExpressionSubstrate.h
@@ -371,7 +371,7 @@ class UMaterialExpressionSubstrateSlabBSDF : public UMaterialExpressionSubstrate
 	bool HasGlint() const;
 	bool HasSpecularProfile() const;
 
-	FSubstrateMaterialComplexity GetComplexity() const;
+	FSubstrateMaterialComplexity GetHighestComplexity() const;
 #endif
 	//~ End UMaterialExpression Interface
 };
