MVVM: fix Set with Broadcast node causing an infinite loop when used on an empty object

diff --git a/Engine/Source/Runtime/Engine/Private/FieldNotification/FieldNotificationLibrary.cpp b/Engine/Source/Runtime/Engine/Private/FieldNotification/FieldNotificationLibrary.cpp
--- a/Engine/Source/Runtime/Engine/Private/FieldNotification/FieldNotificationLibrary.cpp
+++ b/Engine/Source/Runtime/Engine/Private/FieldNotification/FieldNotificationLibrary.cpp
@@ -130,31 +130,34 @@ DEFINE_FUNCTION(UFieldNotificationLibrary::execSetPropertyValueAndBroadcast)
 {
 	P_GET_UBOOL(NewValueByRef);
 
+	// Old Value
 	Stack.StepCompiledIn<FProperty>(nullptr);
 	FProperty* TargetProperty = Stack.MostRecentProperty;
 	void* TargetValuePtr = Stack.MostRecentPropertyAddress;
 
+	// New Value
 	void* SourceValuePtr = nullptr;
 	if (NewValueByRef)
 	{
+		// Is linked in the graph
 		Stack.StepCompiledIn<FProperty>(nullptr);
 		SourceValuePtr = Stack.MostRecentPropertyAddress;
 	}
 	else if (TargetProperty)
 	{
+		// Is a default in the graph.
+		//New Value type is the same as Old Value
+
 		// Allocate temporary memory for the variable
 		SourceValuePtr = FMemory_Alloca_Aligned(TargetProperty->GetSize(), TargetProperty->GetMinAlignment());
-		TargetProperty->InitializeValue(SourceValuePtr); // probably not needed since they are double/float
+		TargetProperty->InitializeValue(SourceValuePtr);
 		Stack.Step(Stack.Object, SourceValuePtr);
 	}
 	else
 	{
-		const FBlueprintExceptionInfo ExceptionInfo(
-			EBlueprintExceptionType::AccessViolation,
-			LOCTEXT("MissingInputProperty", "Failed to resolve the input parameter for SetPropertyValueAndBroadcast.")
-		);
-		FBlueprintCoreDelegates::ThrowScriptException(P_THIS, Stack, ExceptionInfo);
-		Stack.Code++;
+		// Can't read the previous value, it won't be able to set it. Fails and code skip the context op code.
+		Stack.bArrayContextFailed = true;
+		return;
 	}
 
 	P_GET_OBJECT(UObject, Self);
@@ -165,12 +168,24 @@ DEFINE_FUNCTION(UFieldNotificationLibrary::execSetPropertyValueAndBroadcast)
 
 	P_FINISH;
 
-	P_NATIVE_BEGIN
-	FFieldNotificationId FieldNotify = FFieldNotificationId(TargetProperty->GetFName());
-	bool bResult = UE::FieldNotification::Private::SetPropertyValueAndBroadcast(SourceValuePtr, TargetValuePtr, TargetProperty, Self
-		, NetOwner, bHasLocalRepNotify, bShouldFlushDormancyOnSet, bIsNetProperty, TArrayView<FFieldNotificationId>(&FieldNotify, 1), P_THIS, Stack);
-	*(bool*)RESULT_PARAM = bResult;
-	P_NATIVE_END
+	*(bool*)RESULT_PARAM = false;
+	if (Self == nullptr || SourceValuePtr == nullptr)
+	{
+		const FBlueprintExceptionInfo ExceptionInfo(
+			EBlueprintExceptionType::AccessViolation,
+			LOCTEXT("MissingInputProperty", "Failed to resolve the input parameter for SetPropertyValueAndBroadcast.")
+		);
+		FBlueprintCoreDelegates::ThrowScriptException(P_THIS, Stack, ExceptionInfo);
+	}
+	else if (TargetValuePtr)
+	{
+		P_NATIVE_BEGIN
+		FFieldNotificationId FieldNotify = FFieldNotificationId(TargetProperty->GetFName());
+		bool bResult = UE::FieldNotification::Private::SetPropertyValueAndBroadcast(SourceValuePtr, TargetValuePtr, TargetProperty, Self
+			, NetOwner, bHasLocalRepNotify, bShouldFlushDormancyOnSet, bIsNetProperty, TArrayView<FFieldNotificationId>(&FieldNotify, 1), P_THIS, Stack);
+		*(bool*)RESULT_PARAM = bResult;
+		P_NATIVE_END
+	}
 
 	if (!NewValueByRef && SourceValuePtr)
 	{
@@ -201,12 +216,9 @@ DEFINE_FUNCTION(UFieldNotificationLibrary::execSetPropertyValueAndBroadcastField
 	}
 	else
 	{
-		const FBlueprintExceptionInfo ExceptionInfo(
-			EBlueprintExceptionType::AccessViolation,
-			LOCTEXT("MissingInputProperty", "Failed to resolve the input parameter for SetPropertyValueAndBroadcast.")
-		);
-		FBlueprintCoreDelegates::ThrowScriptException(P_THIS, Stack, ExceptionInfo);
-		Stack.Code++;
+		// Can't read the previous value, it won't be able to set it. Fails and code skip the context op code.
+		Stack.bArrayContextFailed = true;
+		return;
 	}
 
 	P_GET_OBJECT(UObject, Self);
@@ -218,12 +230,25 @@ DEFINE_FUNCTION(UFieldNotificationLibrary::execSetPropertyValueAndBroadcastField
 
 	P_FINISH;
 
-	P_NATIVE_BEGIN
-	FieldNotifies.AddUnique(FFieldNotificationId(TargetProperty->GetFName()));
-	bool bResult = UE::FieldNotification::Private::SetPropertyValueAndBroadcast(SourceValuePtr, TargetValuePtr, TargetProperty, Self
-		, NetOwner, bHasLocalRepNotify, bShouldFlushDormancyOnSet, bIsNetProperty, FieldNotifies, P_THIS, Stack);
-	*(bool*)RESULT_PARAM = bResult;
-	P_NATIVE_END
+	*(bool*)RESULT_PARAM = false;
+	if (Self == nullptr || SourceValuePtr == nullptr)
+	{
+		const FBlueprintExceptionInfo ExceptionInfo(
+			EBlueprintExceptionType::AccessViolation,
+			LOCTEXT("MissingInputProperty", "Failed to resolve the input parameter for SetPropertyValueAndBroadcast.")
+		);
+		FBlueprintCoreDelegates::ThrowScriptException(P_THIS, Stack, ExceptionInfo);
+		
+	}
+	else if (SourceValuePtr && TargetValuePtr)
+	{
+		P_NATIVE_BEGIN
+		FieldNotifies.AddUnique(FFieldNotificationId(TargetProperty->GetFName()));
+		bool bResult = UE::FieldNotification::Private::SetPropertyValueAndBroadcast(SourceValuePtr, TargetValuePtr, TargetProperty, Self
+			, NetOwner, bHasLocalRepNotify, bShouldFlushDormancyOnSet, bIsNetProperty, FieldNotifies, P_THIS, Stack);
+		*(bool*)RESULT_PARAM = bResult;
+		P_NATIVE_END
+	}
 
 	if (!NewValueByRef && SourceValuePtr)
 	{
