- Renamed CVars to better represent what the optimization is affecting.
- Renamed CVar r.Streaming.AllowParallelRenderAssetStreamingManagerProcessComponents to r.Streaming.AllowParallelRenderAssetStreamingManagerIncrementalUpdate
- Renamed CVar r.Streaming.WorkerCountForRenderAssetParallelProcessComponents to r.Streaming.WorkerCountForParallelRenderAssetStreamingManagerIncrementalUpdate
- Optimized FLevelRenderAssetManager::IncrementalBuild's EStaticBuildStep::ProcessActors and EStaticBuildStep::CompileElements steps to respect r.Streaming.AllowParallelRenderAssetStreamingIncrementalUpdate and run in parallel (reduces peaks)

diff --git a/Engine/Source/Runtime/Engine/Private/Streaming/DynamicTextureInstanceManager.cpp b/Engine/Source/Runtime/Engine/Private/Streaming/DynamicTextureInstanceManager.cpp
--- a/Engine/Source/Runtime/Engine/Private/Streaming/DynamicTextureInstanceManager.cpp
+++ b/Engine/Source/Runtime/Engine/Private/Streaming/DynamicTextureInstanceManager.cpp
@@ -111,7 +111,7 @@ void FDynamicRenderAssetInstanceManager::IncrementalUpdate(FRemovedRenderAssetAr
 	{
 		const int32 ToProcessComponentCount = PendingComponents.Num();
 		int32 MinBatchSize = 1;
-		if (FRenderAssetStreamingManager::IsParallelProcessComponentsAllowed(ToProcessComponentCount, MinBatchSize))
+		if (FRenderAssetStreamingManager::IsParallelForAllowedDuringIncrementalUpdate(ToProcessComponentCount, MinBatchSize))
 		{
 			TArray<FRenderAssetInstanceState::FPreAddComponentIgnoreBoundsPayload> Payloads;
 			Payloads.SetNum(ToProcessComponentCount);
diff --git a/Engine/Source/Runtime/Engine/Private/Streaming/LevelTextureManager.cpp b/Engine/Source/Runtime/Engine/Private/Streaming/LevelTextureManager.cpp
--- a/Engine/Source/Runtime/Engine/Private/Streaming/LevelTextureManager.cpp
+++ b/Engine/Source/Runtime/Engine/Private/Streaming/LevelTextureManager.cpp
@@ -88,25 +88,43 @@ float FLevelRenderAssetManager::GetWorldTime() const
 	return 0;
 }
 
-void FLevelRenderAssetManager::SetAsStatic(FDynamicRenderAssetInstanceManager& DynamicComponentManager, const UPrimitiveComponent* Primitive)
+void FLevelRenderAssetManager::SetAsStatic(FDynamicRenderAssetInstanceManager& DynamicComponentManager, const UPrimitiveComponent* Primitive, bool bIsConcurrent)
 {
 	check(Primitive);
 	Primitive->bAttachedToStreamingManagerAsStatic = true;
 	if (Primitive->bHandledByStreamingManagerAsDynamic)
 	{
-		DynamicComponentManager.Remove(Primitive, nullptr);
+		if (bIsConcurrent)
+		{
+			SetAsLock.Lock();
+			DynamicComponentManager.Remove(Primitive, nullptr);
+			SetAsLock.Unlock();
+		}
+		else
+		{
+			DynamicComponentManager.Remove(Primitive, nullptr);
+		}
 		Primitive->bHandledByStreamingManagerAsDynamic = false;
 	}
 }
 
 
-void FLevelRenderAssetManager::SetAsDynamic(FDynamicRenderAssetInstanceManager& DynamicComponentManager, FStreamingTextureLevelContext& LevelContext, const UPrimitiveComponent* Primitive)
+void FLevelRenderAssetManager::SetAsDynamic(FDynamicRenderAssetInstanceManager& DynamicComponentManager, FStreamingTextureLevelContext& LevelContext, const UPrimitiveComponent* Primitive, bool bIsConcurrent)
 {
 	check(Primitive);
 	Primitive->bAttachedToStreamingManagerAsStatic = false;
 	if (!Primitive->bHandledByStreamingManagerAsDynamic)
 	{
-		DynamicComponentManager.Add(Primitive, LevelContext);
+		if (bIsConcurrent)
+		{
+			SetAsLock.Lock();
+			DynamicComponentManager.Add(Primitive, LevelContext);
+			SetAsLock.Unlock();
+		}
+		else
+		{
+			DynamicComponentManager.Add(Primitive, LevelContext);
+		}
 	}
 }
 
@@ -141,29 +159,43 @@ void FLevelRenderAssetManager::IncrementalBuild(FDynamicRenderAssetInstanceManag
 		// All actors needs to be processed at once here because of the logic around bStaticComponentsRegisteredInStreamingManager.
 		// All components must have either bHandledByStreamingManagerAsDynamic or bAttachedToStreamingManagerAsStatic set once bStaticComponentsRegisteredInStreamingManager gets set.
 		// If any component gets created after, the logic in UPrimitiveComponent::CreateRenderState_Concurrent() will detect it as a new component and put it through the dynamic path.
-		for (const AActor* Actor : Level->Actors)
-		{
-			if (!Actor) continue;
-
-			const bool bIsStaticActor = Actor->IsRootComponentStatic();
 
-			TInlineComponentArray<UPrimitiveComponent*> Primitives;
-			Actor->GetComponents(Primitives);
-			for (const UPrimitiveComponent* Primitive : Primitives)
+		struct FProcessActorContext
+		{
+			TArray<const UPrimitiveComponent*> UnprocessedComponents{};
+			int32 NumSteps = 0;
+		};
+		
+		int32 MinBatchSize = 1;
+		bool bIsParallelForAllowed = FRenderAssetStreamingManager::IsParallelForAllowedDuringIncrementalUpdate(Level->Actors.Num(), MinBatchSize);
+		TArray<FProcessActorContext> Contexts;
+		ParallelForWithTaskContext(TEXT("ProcessActors"), Contexts, Level->Actors.Num(), MinBatchSize, [this, &DynamicComponentManager, &LevelContext, bIsParallelForAllowed](FProcessActorContext& Context, int32 Index)
+		{
+			if (const AActor* Actor = Level->Actors[Index])
 			{
-				check(Primitive);
-				if (bIsStaticActor && Primitive->Mobility == EComponentMobility::Static)
+				int32 NumSteps = 0;
+				const bool bIsStaticActor = Actor->IsRootComponentStatic();
+				Actor->ForEachComponent<UPrimitiveComponent>(false, [this, &DynamicComponentManager, &Context, &LevelContext, &NumSteps, &bIsStaticActor, bIsParallelForAllowed](UPrimitiveComponent* Primitive)
 				{
-					SetAsStatic(DynamicComponentManager, Primitive);
-					UnprocessedComponents.Push(Primitive);
-				}
-				else
-				{
-					SetAsDynamic(DynamicComponentManager, LevelContext, Primitive);
-				}
+					if (bIsStaticActor && Primitive->Mobility == EComponentMobility::Static)
+					{
+						SetAsStatic(DynamicComponentManager, Primitive, bIsParallelForAllowed);
+						Context.UnprocessedComponents.Push(Primitive);
+					}
+					else
+					{
+						SetAsDynamic(DynamicComponentManager, LevelContext, Primitive, bIsParallelForAllowed);
+					}
+					++NumSteps;
+				});
+				Context.NumSteps += FMath::Max<int32>(NumSteps, 1);
 			}
+		}, bIsParallelForAllowed ? EParallelForFlags::None : EParallelForFlags::ForceSingleThread);
 			
-			NumStepsLeft -= (int64)FMath::Max<int32>(Primitives.Num(), 1);
+		for (FProcessActorContext& Context : Contexts)
+		{
+			UnprocessedComponents.Append(Context.UnprocessedComponents);
+			NumStepsLeft -= Context.NumSteps;
 		}
 
 		NumStepsLeft -= (int64)FMath::Max<int32>(Level->Actors.Num(), 1);
@@ -182,7 +214,7 @@ void FLevelRenderAssetManager::IncrementalBuild(FDynamicRenderAssetInstanceManag
 		const int32 UnprocessedComponentCount = UnprocessedComponents.Num();
 		const int32 ToProcessComponentCount = bForceCompletion ? UnprocessedComponentCount : FMath::Min<int32>(NumStepsLeft, UnprocessedComponentCount);
 		int32 MinBatchSize = 1;
-		if (FRenderAssetStreamingManager::IsParallelProcessComponentsAllowed(ToProcessComponentCount, MinBatchSize) && StaticInstances.CanAddComponent())
+		if (FRenderAssetStreamingManager::IsParallelForAllowedDuringIncrementalUpdate(ToProcessComponentCount, MinBatchSize) && StaticInstances.CanAddComponent())
 		{
 			TArray<FRenderAssetInstanceState::FPreAddComponentPayload> Payloads;
 			Payloads.SetNum(ToProcessComponentCount);
@@ -279,7 +311,7 @@ void FLevelRenderAssetManager::IncrementalBuild(FDynamicRenderAssetInstanceManag
 			// Reprocess the components that didn't have valid data.
 			const int32 ToProcessComponentCount = PendingComponents.Num();
 			int32 MinBatchSize = 1;
-			if (FRenderAssetStreamingManager::IsParallelProcessComponentsAllowed(ToProcessComponentCount, MinBatchSize) && StaticInstances.CanAddComponent())
+			if (FRenderAssetStreamingManager::IsParallelForAllowedDuringIncrementalUpdate(ToProcessComponentCount, MinBatchSize) && StaticInstances.CanAddComponent())
 			{
 				TArray<FRenderAssetInstanceState::FPreAddComponentPayload> Payloads;
 				Payloads.SetNum(ToProcessComponentCount);
diff --git a/Engine/Source/Runtime/Engine/Private/Streaming/LevelTextureManager.h b/Engine/Source/Runtime/Engine/Private/Streaming/LevelTextureManager.h
--- a/Engine/Source/Runtime/Engine/Private/Streaming/LevelTextureManager.h
+++ b/Engine/Source/Runtime/Engine/Private/Streaming/LevelTextureManager.h
@@ -14,6 +14,7 @@ TextureInstanceManager.h: Definitions of classes used for texture streaming.
 #include "Engine/TextureStreamingTypes.h"
 #include "Streaming/TextureStreamingHelpers.h"
 #include "UObject/UObjectHash.h"
+#include "Async/Mutex.h"
 
 // The streaming data of a level.
 class FLevelRenderAssetManager
@@ -88,10 +89,12 @@ private:
 	TArray<const UPrimitiveComponent*> PendingComponents;
 	// Reversed lookup for ULevel::StreamingTextureGuids.
 	TMap<FGuid, int32> TextureGuidToLevelIndex;
+	// Used when SetAsStatic and/or SetAsDynamic are called concurrently
+	UE::FMutex SetAsLock;
 
 	bool NeedsIncrementalBuild(int32 NumStepsLeftForIncrementalBuild) const;
 	void IncrementalBuild(FDynamicRenderAssetInstanceManager& DynamicComponentManager, FStreamingTextureLevelContext& LevelContext, bool bForceCompletion, int64& NumStepsLeft);
 
-	FORCEINLINE_DEBUGGABLE void SetAsStatic(FDynamicRenderAssetInstanceManager& DynamicComponentManager, const UPrimitiveComponent* Primitive);
-	FORCEINLINE_DEBUGGABLE void SetAsDynamic(FDynamicRenderAssetInstanceManager& DynamicComponentManager, FStreamingTextureLevelContext& LevelContext, const UPrimitiveComponent* Primitive);
+	FORCEINLINE_DEBUGGABLE void SetAsStatic(FDynamicRenderAssetInstanceManager& DynamicComponentManager, const UPrimitiveComponent* Primitive, bool bIsConcurrent = false);
+	FORCEINLINE_DEBUGGABLE void SetAsDynamic(FDynamicRenderAssetInstanceManager& DynamicComponentManager, FStreamingTextureLevelContext& LevelContext, const UPrimitiveComponent* Primitive, bool bIsConcurrent = false);
 };
diff --git a/Engine/Source/Runtime/Engine/Private/Streaming/StreamingManagerTexture.cpp b/Engine/Source/Runtime/Engine/Private/Streaming/StreamingManagerTexture.cpp
--- a/Engine/Source/Runtime/Engine/Private/Streaming/StreamingManagerTexture.cpp
+++ b/Engine/Source/Runtime/Engine/Private/Streaming/StreamingManagerTexture.cpp
@@ -99,20 +99,20 @@ static TAutoConsoleVariable<int32> CVarProcessAddedRenderAssetsAfterAsyncWork(
 
 namespace UE::Private::RenderAssetStreamingManager
 {
-	static bool bAllowParallelProcessComponents = false;
-	static FAutoConsoleVariableRef CVarAllowParallelProcessComponents(
-		TEXT("r.Streaming.AllowParallelRenderAssetStreamingManagerProcessComponents"),
-		bAllowParallelProcessComponents,
-		TEXT("Whether it is allowed for RenderAssetStreamingManager to process components to add with a ParallelFor to use more cores. ")
+	static bool bAllowParallelForDuringIncrementalUpdate = false;
+	static FAutoConsoleVariableRef CVarAllowParallelForDuringIncrementalUpdate(
+		TEXT("r.Streaming.AllowParallelRenderAssetStreamingManagerIncrementalUpdate"),
+		bAllowParallelForDuringIncrementalUpdate,
+		TEXT("Whether it is allowed for RenderAssetStreamingManager to use more cores using ParallelFor when doing the incremental update. ")
 		TEXT("When enabled, texture streaming built data is not used to retrieve the primitive streaming render asset information."),
 		ECVF_Default);
 
-	static int32 GWorkerCountForRenderAssetParallelProcessComponents = 0;
-	static FAutoConsoleVariableRef CVarWorkerCountForRenderAssetParallelProcessComponents(
-		TEXT("r.Streaming.WorkerCountForRenderAssetParallelProcessComponents"),
-		GWorkerCountForRenderAssetParallelProcessComponents,
-		TEXT("Set the max number of workers when r.Streaming.AllowParallelRenderAssetStreamingManagerProcessComponents is enabled. ")
-		TEXT("Only taked into account if value > 0."),
+	static int32 GWorkerCountForParallelForDuringIncrementalUpdate = 0;
+	static FAutoConsoleVariableRef CVarWorkerCountForParallelForDuringIncrementalUpdate(
+		TEXT("r.Streaming.WorkerCountForParallelRenderAssetStreamingManagerIncrementalUpdate"),
+		GWorkerCountForParallelForDuringIncrementalUpdate,
+		TEXT("Set the max number of workers when r.Streaming.AllowParallelRenderAssetStreamingManagerIncrementalUpdate is enabled. ")
+		TEXT("Only taked into account if value > 1 (GameThread counts as 1 worker). "),
 		ECVF_Default);
 }
 
@@ -433,12 +433,18 @@ int64 FRenderAssetStreamingManager::GetPoolSize() const
 	return GTexturePoolSize;
 }
 
-bool FRenderAssetStreamingManager::IsParallelProcessComponentsAllowed(int32 NumComponentsToProcess, int32& OutParallelForMinBatchSize)
+bool FRenderAssetStreamingManager::IsParallelForAllowedDuringIncrementalUpdate(int32 NumElementsToProcess, int32& OutParallelForMinBatchSize)
 {
-	static int32 MinComponentsToProcessForParallelProcessing = 8;
-	bool bIsAllowed = UE::Private::RenderAssetStreamingManager::bAllowParallelProcessComponents && FApp::ShouldUseThreadingForPerformance() && (NumComponentsToProcess >= MinComponentsToProcessForParallelProcessing);
-	OutParallelForMinBatchSize = (bIsAllowed && (UE::Private::RenderAssetStreamingManager::GWorkerCountForRenderAssetParallelProcessComponents > 0) && (NumComponentsToProcess > UE::Private::RenderAssetStreamingManager::GWorkerCountForRenderAssetParallelProcessComponents)) ? FMath::DivideAndRoundUp(NumComponentsToProcess, UE::Private::RenderAssetStreamingManager::GWorkerCountForRenderAssetParallelProcessComponents) : 1;
+#if !WITH_EDITOR
+	static int32 MinNumElementsToProcess = 8;
+	bool bIsAllowed = UE::Private::RenderAssetStreamingManager::bAllowParallelForDuringIncrementalUpdate && FApp::ShouldUseThreadingForPerformance() && (NumElementsToProcess >= MinNumElementsToProcess);
+	OutParallelForMinBatchSize = (bIsAllowed && (UE::Private::RenderAssetStreamingManager::GWorkerCountForParallelForDuringIncrementalUpdate > 0) && (NumElementsToProcess > UE::Private::RenderAssetStreamingManager::GWorkerCountForParallelForDuringIncrementalUpdate)) ? FMath::DivideAndRoundUp(NumElementsToProcess, UE::Private::RenderAssetStreamingManager::GWorkerCountForParallelForDuringIncrementalUpdate) : 1;
 	return bIsAllowed;
+#else
+	// Some functions like UMaterialInterface::FindTextureStreamingDataIndexRange are not thread safe in editor
+	OutParallelForMinBatchSize = 1;
+	return false;
+#endif
 }
 
 void FRenderAssetStreamingManager::IncrementalUpdate(float Percentage, bool bUpdateDynamicComponents)
diff --git a/Engine/Source/Runtime/Engine/Private/Streaming/StreamingManagerTexture.h b/Engine/Source/Runtime/Engine/Private/Streaming/StreamingManagerTexture.h
--- a/Engine/Source/Runtime/Engine/Private/Streaming/StreamingManagerTexture.h
+++ b/Engine/Source/Runtime/Engine/Private/Streaming/StreamingManagerTexture.h
@@ -192,8 +192,8 @@ struct FRenderAssetStreamingManager final : public IRenderAssetStreamingManager
 	 */
 	void SetRenderAssetsRemovedTimestamp(const FRemovedRenderAssetArray& RemovedRenderAssets);
 
-	/** Whether it is allowed for this manager to process components to add with a ParallelFor to use more cores. */
-	static bool IsParallelProcessComponentsAllowed(int32 NumComponentsToProcess, int32& OutParallelForMinBatchSize);
+	/** Whether it is allowed for this manager to use more cores using ParallelFor when doing the incremental update. */
+	static bool IsParallelForAllowedDuringIncrementalUpdate(int32 NumElementsToProcess, int32& OutParallelForMinBatchSize);
 
 private:
 //BEGIN: Thread-safe functions and data
diff --git a/Engine/Source/Runtime/Engine/Private/Streaming/TextureInstanceState.cpp b/Engine/Source/Runtime/Engine/Private/Streaming/TextureInstanceState.cpp
--- a/Engine/Source/Runtime/Engine/Private/Streaming/TextureInstanceState.cpp
+++ b/Engine/Source/Runtime/Engine/Private/Streaming/TextureInstanceState.cpp
@@ -10,6 +10,7 @@
 #include "Streaming/TextureInstanceView.inl"
 #include "Components/PrimitiveComponent.h"
 #include "Streaming/TextureInstanceView.h"
+#include "Streaming/StreamingManagerTexture.h"
 #include "Async/ParallelFor.h"
 
 FRenderAssetInstanceState::FRenderAssetInstanceState(bool bForDynamicInstances)
@@ -804,19 +805,32 @@ int32 FRenderAssetInstanceState::CompileElements()
 	CompiledNumForcedLODCompMap.Empty();
 	MaxTexelFactor = 0;
 
+	TArray<const UStreamableRenderAsset*> RenderAssets;
+	RenderAssets.Reserve(RenderAssetMap.Num());
+
 	// First create an entry for all elements, so that there are no reallocs when inserting each compiled elements.
 	for (TMap<const UStreamableRenderAsset*, FRenderAssetDesc>::TConstIterator AssetIt(RenderAssetMap); AssetIt; ++AssetIt)
 	{
 		const UStreamableRenderAsset* Asset = AssetIt.Key();
 		CompiledRenderAssetMap.Add(Asset);
+		RenderAssets.Add(Asset);
 	}
 
+	struct FCompileElementsContext
+	{
+		float MaxTexelFactor = 0.0f;
+		TArray<TPair<const UStreamableRenderAsset*, int32>> CompiledNumForcedLODCompArray;
+	};
+	TArray<FCompileElementsContext> Contexts;
+
 	// Then fill in each array.
-	for (TMap<const UStreamableRenderAsset*, TArray<FCompiledElement> >::TIterator It(CompiledRenderAssetMap); It; ++It)
+	int32 MinBatchSize = 1;
+	bool bIsParallelForAllowed = FRenderAssetStreamingManager::IsParallelForAllowedDuringIncrementalUpdate(RenderAssets.Num(), MinBatchSize);
+	ParallelForWithTaskContext(TEXT("CompileElements"), Contexts, RenderAssets.Num(), MinBatchSize, [this, &RenderAssets](FCompileElementsContext& Context, int32 Index)
 	{
-		const UStreamableRenderAsset* Asset = It.Key();
+		const UStreamableRenderAsset* Asset = RenderAssets[Index];
 		const bool bIsNonTexture = Asset && !Asset->IsA<UTexture>();
-		TArray<FCompiledElement>& CompiledElemements = It.Value();
+		TArray<FCompiledElement>& CompiledElemements = CompiledRenderAssetMap.FindChecked(Asset);
 
 		int32 CompiledElementCount = 0;
 		for (auto ElementIt = GetElementIterator(Asset); ElementIt; ++ElementIt)
@@ -825,6 +839,7 @@ int32 FRenderAssetInstanceState::CompileElements()
 		}
 		CompiledElemements.AddUninitialized(CompiledElementCount);
 
+		int32 CompiledNumForcedLODCompMapCount = 0;
 		CompiledElementCount = 0;
 		for (auto ElementIt = GetElementIterator(Asset); ElementIt; ++ElementIt)
 		{
@@ -832,19 +847,11 @@ int32 FRenderAssetInstanceState::CompileElements()
 
 			if (bIsNonTexture && TexelFactor < 0.f)
 			{
-				int32* CountPtr = CompiledNumForcedLODCompMap.Find(Asset);
-				if (!CountPtr)
-				{
-					CompiledNumForcedLODCompMap.Add(Asset, 1);
-				}
-				else
-				{
-					++*CountPtr;
-				}
+				++CompiledNumForcedLODCompMapCount;
 			}
 
 			// No need to care about force load as MaxTexelFactor is used to ignore far away levels.
-			MaxTexelFactor = FMath::Max(TexelFactor, MaxTexelFactor);
+			Context.MaxTexelFactor = FMath::Max(TexelFactor, Context.MaxTexelFactor);
 
 			CompiledElemements[CompiledElementCount].BoundsIndex = ElementIt.GetBoundsIndex();
 			CompiledElemements[CompiledElementCount].TexelFactor = TexelFactor;
@@ -852,6 +859,20 @@ int32 FRenderAssetInstanceState::CompileElements()
 
 			++CompiledElementCount;
 		}
+
+		if (CompiledNumForcedLODCompMapCount > 0)
+		{
+			Context.CompiledNumForcedLODCompArray.Emplace(Asset, CompiledNumForcedLODCompMapCount);
+		}
+	}, bIsParallelForAllowed ? EParallelForFlags::None : EParallelForFlags::ForceSingleThread);
+
+	for (const FCompileElementsContext& Context : Contexts)
+	{
+		MaxTexelFactor = FMath::Max(MaxTexelFactor, Context.MaxTexelFactor);
+		for (const TPair<const UStreamableRenderAsset*, int32>& Pair : Context.CompiledNumForcedLODCompArray)
+		{
+			CompiledNumForcedLODCompMap.Add(Pair.Key, Pair.Value);
+		}
 	}
 	return CompiledRenderAssetMap.Num();
 }
