Cut 512 bytes from FSceneView when stereo rendering is not used.

diff --git a/Engine/Plugins/Animation/GameplayInsights/Source/RewindDebuggerVLog/Private/VisualLogEntryRenderer.cpp b/Engine/Plugins/Animation/GameplayInsights/Source/RewindDebuggerVLog/Private/VisualLogEntryRenderer.cpp
--- a/Engine/Plugins/Animation/GameplayInsights/Source/RewindDebuggerVLog/Private/VisualLogEntryRenderer.cpp
+++ b/Engine/Plugins/Animation/GameplayInsights/Source/RewindDebuggerVLog/Private/VisualLogEntryRenderer.cpp
@@ -24,7 +24,7 @@ namespace
 			if (!ElementToDraw->Description.IsEmpty())
 			{
 				const FSceneView* View = Canvas->SceneView;
-				if (View->CullingFrustum.IntersectPoint(Position))
+				if (View->GetCullingFrustum().IntersectPoint(Position))
 				{
 					const FPlane Proj = View->Project(Position);
 					if (Proj.W > 0.f)
diff --git a/Engine/Plugins/FX/Niagara/Source/Niagara/Private/NiagaraRendererMeshes.cpp b/Engine/Plugins/FX/Niagara/Source/Niagara/Private/NiagaraRendererMeshes.cpp
--- a/Engine/Plugins/FX/Niagara/Source/Niagara/Private/NiagaraRendererMeshes.cpp
+++ b/Engine/Plugins/FX/Niagara/Source/Niagara/Private/NiagaraRendererMeshes.cpp
@@ -675,7 +675,7 @@ void FNiagaraRendererMeshes::InitializeSortInfo(const FParticleMeshRenderData& P
 			if (bIsInstancedStereo)
 			{
 				// For Instanced Stereo, cull using an extended frustum that encompasses both eyes
-				OutSortInfo.CullPlanes = View.CullingFrustum.Planes;
+				OutSortInfo.CullPlanes = View.GetCullingFrustum().Planes;
 			}
 			else
 			{
diff --git a/Engine/Source/Runtime/Engine/Private/HierarchicalInstancedStaticMesh.cpp b/Engine/Source/Runtime/Engine/Private/HierarchicalInstancedStaticMesh.cpp
--- a/Engine/Source/Runtime/Engine/Private/HierarchicalInstancedStaticMesh.cpp
+++ b/Engine/Source/Runtime/Engine/Private/HierarchicalInstancedStaticMesh.cpp
@@ -1572,7 +1572,7 @@ void FHierarchicalStaticMeshSceneProxy::GetDynamicMeshElements(const TArray<cons
 					if ((View->IsInstancedStereoPass() || View->bIsMobileMultiViewEnabled) && IStereoRendering::IsStereoEyeView(*View) && GEngine->StereoRenderingDevice.IsValid())
 					{
 						// TODO: Stereo culling frustum needs to use the culling origin instead of the view origin.
-						InstanceParams.ViewFrustumLocal = View->CullingFrustum;
+						InstanceParams.ViewFrustumLocal = View->GetCullingFrustum();
 						for (FPlane& Plane : InstanceParams.ViewFrustumLocal.Planes)
 						{
 							Plane = Plane.TransformBy(WorldToLocal);
diff --git a/Engine/Source/Runtime/Engine/Private/SceneView.cpp b/Engine/Source/Runtime/Engine/Private/SceneView.cpp
--- a/Engine/Source/Runtime/Engine/Private/SceneView.cpp
+++ b/Engine/Source/Runtime/Engine/Private/SceneView.cpp
@@ -779,26 +779,28 @@ void FSceneViewOwner::operator=(const AActor* InActor)
 }
 #endif
 
-static void SetupViewFrustum(FSceneView& View)
+void FSceneView::SetupViewFrustum()
 {
-	if (View.SceneViewInitOptions.OverrideFarClippingPlaneDistance > 0.0f)
+	if (SceneViewInitOptions.OverrideFarClippingPlaneDistance > 0.0f)
 	{
-		const FPlane FarPlane(View.ViewMatrices.GetViewOrigin() + View.GetViewDirection() * View.SceneViewInitOptions.OverrideFarClippingPlaneDistance, View.GetViewDirection());
+		const FPlane FarPlane(ViewMatrices.GetViewOrigin() + GetViewDirection() * SceneViewInitOptions.OverrideFarClippingPlaneDistance, GetViewDirection());
 		// Derive the view frustum from the view projection matrix, overriding the far plane
-		GetViewFrustumBounds(View.ViewFrustum, View.ViewMatrices.GetViewProjectionMatrix(), FarPlane, true, false);
+		GetViewFrustumBounds(ViewFrustum, ViewMatrices.GetViewProjectionMatrix(), FarPlane, true, false);
 	}
 	else
 	{
 		// Derive the view frustum from the view projection matrix.
-		GetViewFrustumBounds(View.ViewFrustum, View.ViewMatrices.GetViewProjectionMatrix(), false);
+		GetViewFrustumBounds(ViewFrustum, ViewMatrices.GetViewProjectionMatrix(), false);
 	}
 
 	// Use a unified frustum for culling an instanced stereo pass.
-	if (IStereoRendering::IsStereoEyeView(View) && GEngine->StereoRenderingDevice.IsValid())
+	if (IStereoRendering::IsStereoEyeView(*this) && GEngine->StereoRenderingDevice.IsValid())
 	{
-		FVector MonoLocation = View.ViewLocation;
-		FRotator MonoRotation = View.ViewRotation;
-		GEngine->StereoRenderingDevice->CalculateStereoViewOffset(eSSE_MONOSCOPIC, MonoRotation, View.WorldToMetersScale, MonoLocation);
+		StereoCullingFrustum = MakeShared<FConvexVolume>();
+
+		FVector MonoLocation = ViewLocation;
+		FRotator MonoRotation = ViewRotation;
+		GEngine->StereoRenderingDevice->CalculateStereoViewOffset(eSSE_MONOSCOPIC, MonoRotation, WorldToMetersScale, MonoLocation);
 		const FMatrix ViewRotationMatrix = FInverseRotationMatrix(MonoRotation) * FMatrix(
 			FPlane(0, 0, 1, 0),
 			FPlane(1, 0, 0, 0),
@@ -807,32 +809,30 @@ static void SetupViewFrustum(FSceneView& View)
 		const FMatrix ViewMatrixForCulling = FTranslationMatrix(-MonoLocation) * ViewRotationMatrix;
 		const FMatrix ViewProjForCulling = ViewMatrixForCulling * GEngine->StereoRenderingDevice->GetStereoProjectionMatrix(eSSE_MONOSCOPIC);
 
-		if (View.SceneViewInitOptions.OverrideFarClippingPlaneDistance > 0.0f)
+		if (SceneViewInitOptions.OverrideFarClippingPlaneDistance > 0.0f)
 		{
-			const FPlane FarPlane(View.ViewMatrices.GetViewOrigin() + View.GetViewDirection() * View.SceneViewInitOptions.OverrideFarClippingPlaneDistance, View.GetViewDirection());
+			const FPlane FarPlane(ViewMatrices.GetViewOrigin() + GetViewDirection() * SceneViewInitOptions.OverrideFarClippingPlaneDistance, GetViewDirection());
 			// Derive the frustum from the view projection matrix, overriding the far plane
-			GetViewFrustumBounds(View.CullingFrustum, ViewProjForCulling, FarPlane, true, false);
+			GetViewFrustumBounds(*StereoCullingFrustum, ViewProjForCulling, FarPlane, true, false);
 		}
 		else
 		{
 			// Derive the frustum from the view projection matrix.
-			GetViewFrustumBounds(View.CullingFrustum, ViewProjForCulling, false);
+			GetViewFrustumBounds(*StereoCullingFrustum, ViewProjForCulling, false);
 		}
 
-		View.CullingOrigin = MonoLocation;
+		CullingOrigin = MonoLocation;
 	}
 	else
 	{
-		View.CullingFrustum = View.ViewFrustum;
-		View.CullingOrigin = View.ViewMatrices.GetViewOrigin();
+		StereoCullingFrustum = nullptr;
+		CullingOrigin = ViewMatrices.GetViewOrigin();
 	}
 
 	// Derive the view's near clipping distance and plane.
 	static_assert((int32)ERHIZBuffer::IsInverted != 0, "Fix Near Clip distance!");
-	FPlane NearClippingPlane;
-	View.bHasNearClippingPlane = View.ViewMatrices.GetViewProjectionMatrix().GetFrustumNearPlane(NearClippingPlane);
-	View.NearClippingPlane = NearClippingPlane;
-	View.NearClippingDistance = View.ViewMatrices.ComputeNearPlane();
+	bHasNearClippingPlane = ViewMatrices.GetViewProjectionMatrix().GetFrustumNearPlane(NearClippingPlane);
+	NearClippingDistance = ViewMatrices.ComputeNearPlane();
 }
 
 #if !(UE_BUILD_SHIPPING || UE_BUILD_TEST)
@@ -938,7 +938,7 @@ PRAGMA_ENABLE_DEPRECATION_WARNINGS
 	ShadowFreezeCamera::Init(ViewMatrices);
 #endif // !(UE_BUILD_SHIPPING || UE_BUILD_TEST)
 
-	SetupViewFrustum(*this);
+	SetupViewFrustum();
 
 	// Determine whether the view should reverse the cull mode due to a negative determinant.  Only do this for a valid scene
 	bReverseCulling = (Family && Family->Scene) ? ViewMatrices.GetViewMatrix().Determinant() < 0.0f : false;
@@ -1202,7 +1202,7 @@ void FSceneView::UpdateProjectionMatrix(const FMatrix& NewProjectionMatrix)
 	FViewMatrices NewViewMatrices = FViewMatrices(SceneViewInitOptions);
 	ViewMatrices = NewViewMatrices;
 
-	SetupViewFrustum(*this);
+	SetupViewFrustum();
 }
 
 void FViewMatrices::UpdateViewMatrix(const FVector& ViewLocation, const FRotator& ViewRotation)
diff --git a/Engine/Source/Runtime/Engine/Public/SceneView.h b/Engine/Source/Runtime/Engine/Public/SceneView.h
--- a/Engine/Source/Runtime/Engine/Public/SceneView.h
+++ b/Engine/Source/Runtime/Engine/Public/SceneView.h
@@ -1370,6 +1370,9 @@ private:
 	/** If the above is non-null, a translation that is applied to world-space before transforming by one of the shadow matrices. */
 	FVector		PreShadowTranslation;
 
+	TSharedPtr<FConvexVolume> StereoCullingFrustum;
+	void SetupViewFrustum();
+
 public:
 	FSceneViewInitOptions SceneViewInitOptions;
 
@@ -1510,15 +1513,20 @@ public:
 
 	FConvexVolume ViewFrustum;
 
+	UE_DEPRECATED(5.6, "Use GetCullingFrustum() instead.")
+	uint8 CullingFrustum;
+
+	const FConvexVolume& GetCullingFrustum() const
+	{
+		return StereoCullingFrustum ? *StereoCullingFrustum : ViewFrustum;
+	}
+
 	bool bHasNearClippingPlane;
 
 	FPlane NearClippingPlane;
 
 	float NearClippingDistance;
 
-	/** Monoscopic culling frustum, same as ViewFrustum in case of non-stereo */
-	FConvexVolume CullingFrustum;
-
 	/** Monoscopic culling origin, same as the view matrix origin in case of non-stereo */
 	FVector CullingOrigin;
 
