Cut several hundred bytes of memory from FinalPostProcessSettings by removing inline arrays.

diff --git a/Engine/Source/Runtime/Engine/Public/FinalPostProcessSettings.h b/Engine/Source/Runtime/Engine/Public/FinalPostProcessSettings.h
--- a/Engine/Source/Runtime/Engine/Public/FinalPostProcessSettings.h
+++ b/Engine/Source/Runtime/Engine/Public/FinalPostProcessSettings.h
@@ -165,8 +165,8 @@ public:
 	}
 
 	// was not merged during blending unlike e.g. BloomThreshold 
-	TArray<FCubemapEntry, TInlineAllocator<8>> ContributingCubemaps;
-	TArray<FLUTBlenderEntry, TInlineAllocator<8>> ContributingLUTs;
+	TArray<FCubemapEntry> ContributingCubemaps;
+	TArray<FLUTBlenderEntry> ContributingLUTs;
 
 	// List of materials to use in the buffer visualization overview
 	TArray<UMaterialInterface*> BufferVisualizationOverviewMaterials;
