Improved error reporting around duplicate static parameters in materials / material functions / etc. (when r.Material.PedanticErrorChecksEnabled is used) :
* Added asset path of the duplicate value and dumped the values that differ
* Handled the static component mask case as well as the static switch case

diff --git a/Engine/Source/Runtime/Engine/Private/Materials/Material.cpp b/Engine/Source/Runtime/Engine/Private/Materials/Material.cpp
--- a/Engine/Source/Runtime/Engine/Private/Materials/Material.cpp
+++ b/Engine/Source/Runtime/Engine/Private/Materials/Material.cpp
@@ -7056,10 +7056,9 @@ bool UMaterial::CheckInValidStateForCompilation(class FMaterialCompiler* Compile
 	// Check and report errors due to duplicate parameters set to distinct values.
 	if (!CachedExpressionData->DuplicateParameterErrors.IsEmpty())
 	{
-		for (const TPair<TObjectPtr<UMaterialExpression>, FName>& Error : CachedExpressionData->DuplicateParameterErrors)
+		for (const TPair<TObjectPtr<UMaterialExpression>, FString>& Error : CachedExpressionData->DuplicateParameterErrors)
 		{
-			FString ErrorMsg = FString::Format(TEXT("Parameter '{0}' is set multiple times to different values. Make sure each parameter is set once or always to the same value (e.g. same texture)."), { *Error.Get<1>().ToString() });
-			Compiler->AppendExpressionError(Error.Get<0>(), *ErrorMsg);
+			Compiler->AppendExpressionError(Error.Get<0>(), *Error.Get<1>());
 		}
 
 		bSuccess = false;
diff --git a/Engine/Source/Runtime/Engine/Private/Materials/MaterialCachedData.cpp b/Engine/Source/Runtime/Engine/Private/Materials/MaterialCachedData.cpp
--- a/Engine/Source/Runtime/Engine/Private/Materials/MaterialCachedData.cpp
+++ b/Engine/Source/Runtime/Engine/Private/Materials/MaterialCachedData.cpp
@@ -46,6 +46,8 @@
 
 #include UE_INLINE_GENERATED_CPP_BY_NAME(MaterialCachedData)
 
+#define LOCTEXT_NAMESPACE "Material"
+
 const FMaterialCachedParameterEntry FMaterialCachedParameterEntry::EmptyData{};
 const FMaterialCachedExpressionData FMaterialCachedExpressionData::EmptyData{};
 const FMaterialCachedExpressionEditorOnlyData FMaterialCachedExpressionEditorOnlyData::EmptyData{};
@@ -109,7 +111,8 @@ static bool TryAddParameter(FMaterialCachedExpressionData& CachedData,
 	EMaterialParameterType Type,
 	const FMaterialParameterInfo& ParameterInfo,
 	const FMaterialCachedParameterEditorInfo& InEditorInfo,
-	int32& OutIndex)
+	int32& OutIndex, 
+	TOptional<FMaterialCachedParameterEditorInfo>& OutPreviousEditorInfo)
 {
 	check(CachedData.EditorOnlyData);
 	FMaterialCachedParameterEntry& Entry = CachedData.GetParameterTypeEntry(Type);
@@ -130,7 +133,11 @@ static bool TryAddParameter(FMaterialCachedExpressionData& CachedData,
 	// Update any editor values that haven't been set yet
 	// TODO still need to do this??
 	OutIndex = ElementId.AsInteger();
+
 	FMaterialCachedParameterEditorInfo& EditorInfo = EditorEntry.EditorInfo[OutIndex];
+	// Copy the previous parameter's original info before eventually replacing it, for error reporting purposes :
+	OutPreviousEditorInfo.Emplace(EditorInfo);
+
 	if (!EditorInfo.ExpressionGuid.IsValid())
 	{
 		EditorInfo.ExpressionGuid = InEditorInfo.ExpressionGuid;
@@ -149,7 +156,7 @@ static bool TryAddParameter(FMaterialCachedExpressionData& CachedData,
 	return false;
 }
 
-bool FMaterialCachedExpressionData::AddParameter(const FMaterialParameterInfo& ParameterInfo, const FMaterialParameterMetadata& ParameterMeta, UObject*& OutReferencedTexture, UTextureCollection*& OutReferencedTextureCollection)
+bool FMaterialCachedExpressionData::AddParameter(const FMaterialParameterInfo& ParameterInfo, const FMaterialParameterMetadata& ParameterMeta, UObject*& OutReferencedTexture, UTextureCollection*& OutReferencedTextureCollection, FText& OutErrorMessage)
 {
 	check(EditorOnlyData);
 	int32 AssetIndex = INDEX_NONE;
@@ -160,7 +167,8 @@ bool FMaterialCachedExpressionData::AddParameter(const FMaterialParameterInfo& P
 
 	const FMaterialCachedParameterEditorInfo EditorInfo(ParameterMeta.ExpressionGuid, ParameterMeta.Description, ParameterMeta.Group, ParameterMeta.SortPriority, AssetIndex);
 	int32 Index = INDEX_NONE;
-	if (TryAddParameter(*this, ParameterMeta.Value.Type, ParameterInfo, EditorInfo, Index))
+	TOptional<FMaterialCachedParameterEditorInfo> PreviousEditorInfo;
+	if (TryAddParameter(*this, ParameterMeta.Value.Type, ParameterInfo, EditorInfo, Index, PreviousEditorInfo))
 	{
 		switch (ParameterMeta.Value.Type)
 		{
@@ -238,39 +246,120 @@ bool FMaterialCachedExpressionData::AddParameter(const FMaterialParameterInfo& P
 	}
 	else
 	{
-		bool bSameValue;
+		auto GetEditorInfoAsText = [EditorOnlyData = EditorOnlyData](const FMaterialCachedParameterEditorInfo& InEditorInfo) -> FText
+			{
+				FString AssetPath = EditorOnlyData->AssetPaths[InEditorInfo.AssetIndex];
+				return FText::Format(LOCTEXT("ReportParameterMetaData", "group:'{0}', asset:'{1}'"), FText::FromName(InEditorInfo.Group), FText::FromString(AssetPath));
+			};
+
+		auto ReportError = [&GetEditorInfoAsText](const FString& InCachedValue, const FString& InValue, const FMaterialCachedParameterEditorInfo& InCachedEditorInfo, const FMaterialCachedParameterEditorInfo& InEditorInfo) -> FText
+			{
+				return FText::Format(LOCTEXT("ReportDifferentParameterValueError", "{0} ({1}) vs. cached: {2} ({3})"), 
+					FText::FromString(InValue), GetEditorInfoAsText(InEditorInfo), FText::FromString(InCachedValue), GetEditorInfoAsText(InCachedEditorInfo));
+			};
+
+		bool bSameValue = false;
 		switch (ParameterMeta.Value.Type)
 		{
 		case EMaterialParameterType::Scalar:
 			bSameValue = ScalarValues[Index] == ParameterMeta.Value.AsScalar();
+			if (!bSameValue)
+			{
+				OutErrorMessage = ReportError(FString::Printf(TEXT("%f"), ScalarValues[Index]), FString::Printf(TEXT("%f"), ParameterMeta.Value.AsScalar()), *PreviousEditorInfo, EditorInfo);
+			}
 			break;
 
 		case EMaterialParameterType::Vector:
 			bSameValue = VectorValues[Index] == ParameterMeta.Value.AsLinearColor();
+			if (!bSameValue)
+			{
+				OutErrorMessage = ReportError(VectorValues[Index].ToString(), ParameterMeta.Value.AsLinearColor().ToString(), *PreviousEditorInfo, EditorInfo);
+			}
 			break;
 
 		case EMaterialParameterType::DoubleVector:
 			bSameValue = DoubleVectorValues[Index] == ParameterMeta.Value.AsVector4d();
+			if (!bSameValue)
+			{
+				OutErrorMessage = ReportError(DoubleVectorValues[Index].ToString(), ParameterMeta.Value.AsVector4d().ToString(), *PreviousEditorInfo, EditorInfo);
+			}
 			break;
 
 		case EMaterialParameterType::Texture:
 			bSameValue = TextureValues[Index] == ParameterMeta.Value.Texture;
+			if (!bSameValue)
+			{
+				OutErrorMessage = ReportError(TextureValues[Index].ToString(), TSoftObjectPtr<UTexture>(ParameterMeta.Value.Texture).ToString(), *PreviousEditorInfo, EditorInfo);
+			}
 			break;
 
 		case EMaterialParameterType::TextureCollection:
 			bSameValue = TextureCollectionValues[Index] == ParameterMeta.Value.TextureCollection;
+			if (!bSameValue)
+			{
+				OutErrorMessage = ReportError(TextureCollectionValues[Index].ToString(), TSoftObjectPtr<UTextureCollection>(ParameterMeta.Value.TextureCollection).ToString(), *PreviousEditorInfo, EditorInfo);
+			}
 			break;
 
 		case EMaterialParameterType::Font:
 			bSameValue = FontValues[Index] == ParameterMeta.Value.Font.Value && FontPageValues[Index] == ParameterMeta.Value.Font.Page;
+			if (!bSameValue)
+			{
+				OutErrorMessage = ReportError(
+					FString::Printf(TEXT("%s(%i)"), *FontValues[Index].ToString(), FontPageValues[Index]), 
+					FString::Printf(TEXT("%s(%i)"), *TSoftObjectPtr<UFont>(ParameterMeta.Value.Font.Value).ToString(), ParameterMeta.Value.Font.Page), 
+					*PreviousEditorInfo, EditorInfo);
+			}
 			break;
 
 		case EMaterialParameterType::RuntimeVirtualTexture:
 			bSameValue = RuntimeVirtualTextureValues[Index] == ParameterMeta.Value.RuntimeVirtualTexture;
+			if (!bSameValue)
+			{
+				OutErrorMessage = ReportError(RuntimeVirtualTextureValues[Index].ToString(), TSoftObjectPtr<URuntimeVirtualTexture>(ParameterMeta.Value.RuntimeVirtualTexture).ToString(), *PreviousEditorInfo, EditorInfo);
+			}
 			break;
 
 		case EMaterialParameterType::SparseVolumeTexture:
 			bSameValue = SparseVolumeTextureValues[Index] == ParameterMeta.Value.SparseVolumeTexture;
+			if (!bSameValue)
+			{
+				OutErrorMessage = ReportError(SparseVolumeTextureValues[Index].ToString(), TSoftObjectPtr<USparseVolumeTexture>(ParameterMeta.Value.SparseVolumeTexture).ToString(), *PreviousEditorInfo, EditorInfo);
+			}
+			break;
+
+		case EMaterialParameterType::StaticSwitch:
+			bSameValue = (StaticSwitchValues[Index] == ParameterMeta.Value.AsStaticSwitch()) && (DynamicSwitchValues[Index] == ParameterMeta.bDynamicSwitchParameter);
+			if (!bSameValue)
+			{
+				auto GetStaticSwitchString = [](bool bInValue, bool bIsDynamicSwitch) -> FString
+					{
+						return FText::Format(LOCTEXT("StaticSwitchValue", "{0}{1}"), bInValue ? LOCTEXT("true", "true") : LOCTEXT("false", "false"), bIsDynamicSwitch ? LOCTEXT("dynamic", "(dynamic)") : FText()).ToString();
+					};
+				OutErrorMessage = ReportError(
+					GetStaticSwitchString(StaticSwitchValues[Index], DynamicSwitchValues[Index]), 
+					GetStaticSwitchString(ParameterMeta.Value.AsStaticSwitch(), ParameterMeta.bDynamicSwitchParameter), 
+					*PreviousEditorInfo, EditorInfo);
+			}
+			break;
+
+		case EMaterialParameterType::StaticComponentMask:
+			bSameValue = EditorOnlyData->StaticComponentMaskValues[Index] == ParameterMeta.Value.AsStaticComponentMask();
+			if (!bSameValue)
+			{
+				auto GetStaticComponentMaskString = [](const FStaticComponentMaskValue& InValue) -> FString
+					{
+						return FText::Format(LOCTEXT("StaticComponentMaskValue", "R={0},G={1},B={2},A={3}")
+							, InValue.R ? LOCTEXT("true", "true") : LOCTEXT("false", "false")
+							, InValue.G ? LOCTEXT("true", "true") : LOCTEXT("false", "false")
+							, InValue.B ? LOCTEXT("true", "true") : LOCTEXT("false", "false")
+							, InValue.A ? LOCTEXT("true", "true") : LOCTEXT("false", "false")).ToString();
+					};
+				OutErrorMessage = ReportError(
+					GetStaticComponentMaskString(EditorOnlyData->StaticComponentMaskValues[Index]), 
+					GetStaticComponentMaskString(ParameterMeta.Value.AsStaticComponentMask()), 
+					*PreviousEditorInfo, EditorInfo);
+			}
 			break;
 
 		default:
@@ -349,6 +438,7 @@ void FMaterialCachedExpressionData::UpdateForExpressions(const FMaterialCachedEx
 		EditorOnlyData->ShaderTags.Append(ShaderTags);
 
 		FMaterialParameterMetadata ParameterMeta;
+		FText ErrorContext;
 		if (Expression->GetParameterValue(ParameterMeta))
 		{
 			const FName ParameterName = Expression->GetParameterName();
@@ -370,9 +460,11 @@ void FMaterialCachedExpressionData::UpdateForExpressions(const FMaterialCachedEx
 			const FMaterialParameterInfo ParameterInfo(ParameterName, Association, ParameterIndex);
 
 			// Try add the parameter. If this fails, the parameter is being added twice with different values. Report it as error.
-			if (!AddParameter(ParameterInfo, ParameterMeta, ReferencedTexture, ReferencedTextureCollection))
+			if (!AddParameter(ParameterInfo, ParameterMeta, ReferencedTexture, ReferencedTextureCollection, ErrorContext))
 			{
-				DuplicateParameterErrors.AddUnique({ Expression, ParameterName });
+				FText ErrorMessage = FText::Format(LOCTEXT("DuplicateParameterError", "Parameter '{0}' is set multiple times to different values : {1}. Make sure each parameter is set once or always to the same value."),
+					 FText::FromName(ParameterName), ErrorContext);
+				DuplicateParameterErrors.AddUnique({ Expression, ErrorMessage.ToString() });
 			}
 		}
 
@@ -955,3 +1047,4 @@ void FMaterialCachedExpressionData::GetAllGlobalParameterInfoOfType(EMaterialPar
 }
 
 
+#undef LOCTEXT_NAMESPACE
diff --git a/Engine/Source/Runtime/Engine/Private/Materials/MaterialHLSLTree.cpp b/Engine/Source/Runtime/Engine/Private/Materials/MaterialHLSLTree.cpp
--- a/Engine/Source/Runtime/Engine/Private/Materials/MaterialHLSLTree.cpp
+++ b/Engine/Source/Runtime/Engine/Private/Materials/MaterialHLSLTree.cpp
@@ -690,7 +690,8 @@ bool FExpressionParameter::PrepareValue(FEmitContext& Context, FEmitScope& Scope
 		if (!ParameterInfo.Name.IsNone())
 		{
 			UObject* UnusedReferencedTexture;
-			EmitData.CachedExpressionData->AddParameter(ParameterInfo, ParameterMeta, UnusedReferencedTexture);
+			FText UnusedErrorContext;
+			EmitData.CachedExpressionData->AddParameter(ParameterInfo, ParameterMeta, UnusedReferencedTexture, UnusedErrorContext);
 		}
 
 		UObject* ReferencedTexture = ParameterMeta.Value.AsTextureObject();
@@ -1831,7 +1832,8 @@ bool FExpressionStaticTerrainLayerWeight::PrepareValue(FEmitContext& Context, FE
 			WeightmapParameterMeta.Value = bTextureArrayEnabled ? GEngine->WeightMapArrayPlaceholderTexture : GEngine->WeightMapPlaceholderTexture;
 
 			UObject* UnusedReferencedTexture;
-			EmitData.CachedExpressionData->AddParameter(WeightmapParameterInfo, WeightmapParameterMeta, UnusedReferencedTexture);
+			FText UnusedErrorContext;
+			EmitData.CachedExpressionData->AddParameter(WeightmapParameterInfo, WeightmapParameterMeta, UnusedReferencedTexture, UnusedErrorContext);
 
 			FMaterialParameterInfo LayerMaskParameterInfo = BaseParameterInfo;
 			LayerMaskParameterInfo.Name = *FString::Printf(TEXT("LayerMask_%s"), *Parameter.LayerName.ToString());
@@ -1839,7 +1841,7 @@ bool FExpressionStaticTerrainLayerWeight::PrepareValue(FEmitContext& Context, FE
 			FMaterialParameterMetadata LayerMaskParameterMeta;
 			LayerMaskParameterMeta.Value = FLinearColor(1.f, 0.f, 0.f, 0.f);
 
-			EmitData.CachedExpressionData->AddParameter(LayerMaskParameterInfo, LayerMaskParameterMeta, UnusedReferencedTexture);
+			EmitData.CachedExpressionData->AddParameter(LayerMaskParameterInfo, LayerMaskParameterMeta, UnusedReferencedTexture, UnusedErrorContext);
 		}
 	}
 
diff --git a/Engine/Source/Runtime/Engine/Public/MaterialCachedData.h b/Engine/Source/Runtime/Engine/Public/MaterialCachedData.h
--- a/Engine/Source/Runtime/Engine/Public/MaterialCachedData.h
+++ b/Engine/Source/Runtime/Engine/Public/MaterialCachedData.h
@@ -204,12 +204,12 @@ struct FMaterialCachedExpressionData
 	void Validate(const UMaterialInterface& Material);
 
 	/** Adds a parameter. If this returns false, a parameter with identical name has already been added but it was set to a different value. */
-	bool AddParameter(const FMaterialParameterInfo& ParameterInfo, const FMaterialParameterMetadata& ParameterMeta, UObject*& OutReferencedTexture, UTextureCollection*& OutReferencedTextureCollection);
+	bool AddParameter(const FMaterialParameterInfo& ParameterInfo, const FMaterialParameterMetadata& ParameterMeta, UObject*& OutReferencedTexture, UTextureCollection*& OutReferencedTextureCollection, FText& OutErrorMessage);
 
-	bool AddParameter(const FMaterialParameterInfo& ParameterInfo, const FMaterialParameterMetadata& ParameterMeta, UObject*& OutReferencedTexture)
+	bool AddParameter(const FMaterialParameterInfo& ParameterInfo, const FMaterialParameterMetadata& ParameterMeta, UObject*& OutReferencedTexture, FText& OutErrorMessage)
 	{
 		UTextureCollection* UnusedReferencedTextureCollection = nullptr;
-		return AddParameter(ParameterInfo, ParameterMeta, OutReferencedTexture, UnusedReferencedTextureCollection);
+		return AddParameter(ParameterInfo, ParameterMeta, OutReferencedTexture, UnusedReferencedTextureCollection, OutErrorMessage);
 	}
 #endif // WITH_EDITOR
 
@@ -387,7 +387,7 @@ struct FMaterialCachedExpressionData
 
 #if WITH_EDITOR
 	/** Array of errors reporting a parameter being set multiple times to distinct values. */
-	TArray<TPair<TObjectPtr<class UMaterialExpression>, FName>> DuplicateParameterErrors;
+	TArray<TPair<TObjectPtr<class UMaterialExpression>, FString>> DuplicateParameterErrors;
 #endif
 };
 
