Early checking for higher priority inputs bound to the same gamepad buttons

#tests PIE with all combinations of each rollback on/off and a higher-priority input component

diff --git a/Engine/Source/Runtime/Engine/Classes/GameFramework/PlayerInput.h b/Engine/Source/Runtime/Engine/Classes/GameFramework/PlayerInput.h
--- a/Engine/Source/Runtime/Engine/Classes/GameFramework/PlayerInput.h
+++ b/Engine/Source/Runtime/Engine/Classes/GameFramework/PlayerInput.h
@@ -798,6 +798,9 @@ protected:
 	/** @return True if a key is handled by an action binding */
 	ENGINE_API virtual bool IsKeyHandledByAction(FKey Key) const;
 
+	/** Retrieves the actions bound to the given input component which are triggered by a given key. */
+	ENGINE_API void GetActionsBoundToKey(UInputComponent* InputComponent, FKey Key, TArray<TSharedPtr<FInputActionBinding>>& Actions);
+
 private:
 
 	/** Gesture recognizer object */
diff --git a/Engine/Source/Runtime/Engine/Private/UserInterface/PlayerInput.cpp b/Engine/Source/Runtime/Engine/Private/UserInterface/PlayerInput.cpp
--- a/Engine/Source/Runtime/Engine/Private/UserInterface/PlayerInput.cpp
+++ b/Engine/Source/Runtime/Engine/Private/UserInterface/PlayerInput.cpp
@@ -2193,6 +2193,12 @@ bool UPlayerInput::IsKeyHandledByAction( FKey Key ) const
 	return false;
 }
 
+void UPlayerInput::GetActionsBoundToKey(UInputComponent* InputComponent, FKey Key, TArray<TSharedPtr<FInputActionBinding>>& Actions)
+{
+	InputComponent->ConditionalBuildKeyMap(this);
+	FGetActionsBoundToKey::Get(InputComponent, this, Key, Actions);
+}
+
 const TArray<FInputActionKeyMapping>& UPlayerInput::GetKeysForAction(const FName ActionName) const
 {
 	ConditionalBuildKeyMappings();
