Minor memory optimization when serializing FVirtualTextureDataChunk

#rnx

- FVirtualTextureBuiltData::Cunks is quite often only 1 or 2 elements, but when we call SetNum to set it's initialize size when serializing from disk FArray will allocation additional slack, usually room for 4 elements in the array.
- By calling reserve first we only allocate the memory we know we will need, the call to SetNum will not result in additional slack being allocated.
- On CitySample this reduces the memory used by this array from 3.45MiB->1.87MiB

diff --git a/Engine/Source/Runtime/Engine/Private/VT/VirtualTextureBuiltData.cpp b/Engine/Source/Runtime/Engine/Private/VT/VirtualTextureBuiltData.cpp
--- a/Engine/Source/Runtime/Engine/Private/VT/VirtualTextureBuiltData.cpp
+++ b/Engine/Source/Runtime/Engine/Private/VT/VirtualTextureBuiltData.cpp
@@ -357,6 +357,8 @@ void FVirtualTextureBuiltData::Serialize(FArchive& Ar, UObject* Owner, int32 Fir
 
 	if (Ar.IsLoading())
 	{
+		Chunks.Reserve(NumChunksToSerialize);	// Allocates the exact memory we need as relying on SetNum would
+												// result in allocating additional slack.
 		Chunks.SetNum(NumChunksToSerialize);
 	}
 
