Substrate - clarified a shader compilation error with an actual failure error instead of mysterious HLSL compilation errors.
Re-added from back out 38637364, with a new bSubstrateEnabled test for operations that should only be done with substrate.

#tests pc dx12 editor
#rnx

diff --git a/Engine/Source/Runtime/Engine/Private/Materials/HLSLMaterialTranslator.cpp b/Engine/Source/Runtime/Engine/Private/Materials/HLSLMaterialTranslator.cpp
--- a/Engine/Source/Runtime/Engine/Private/Materials/HLSLMaterialTranslator.cpp
+++ b/Engine/Source/Runtime/Engine/Private/Materials/HLSLMaterialTranslator.cpp
@@ -2315,240 +2315,252 @@ void FHLSLMaterialTranslator::TranslateMaterial()
 		TextureExpression->GetTextureCollectionParameterInfo(MaterialCompilationOutput.UniformExpressionSet.UniformTextureCollectionParameters.AddDefaulted_GetRef());
 	}
 
-	for (uint32 SubstrateCompilationContextIndex = 0; SubstrateCompilationContextIndex < ESubstrateCompilationContext::SCC_MAX; ++SubstrateCompilationContextIndex)
+	if (bSubstrateEnabled)
 	{
-		CHECK_DDC_QUERY_FINISHED_ELSE_RETURN();
-		FSubstrateCompilationContext& SubstrateCtx = SubstrateCompilationContext[SubstrateCompilationContextIndex];
-
-		FString TreeFunctionPostFix = TEXT("ERROR");
-		switch (SubstrateCompilationContextIndex)
+		for (uint32 SubstrateCompilationContextIndex = 0; SubstrateCompilationContextIndex < ESubstrateCompilationContext::SCC_MAX; ++SubstrateCompilationContextIndex)
 		{
-		case ESubstrateCompilationContext::SCC_Default:
-			TreeFunctionPostFix = TEXT("");
-			break;
-		case ESubstrateCompilationContext::SCC_FullySimplified:
-			TreeFunctionPostFix = TEXT("_FullySimplified");
-			break;
-		default:
-			check(false);
-		}
+			CHECK_DDC_QUERY_FINISHED_ELSE_RETURN();
+			FSubstrateCompilationContext& SubstrateCtx = SubstrateCompilationContext[SubstrateCompilationContextIndex];
 
-		const bool bSubstrateFrontMaterialProvided = Chunk[MP_FrontMaterial] != INDEX_NONE;
-		bool bSubstrateFrontMaterialIsValid = bSubstrateFrontMaterialProvided;
-		if (bSubstrateFrontMaterialIsValid)
-		{
-			// The material can be null when some entries are automatically generated, for instance in the material layer blending system
-			bSubstrateFrontMaterialIsValid &= SubstrateCtx.SubstrateMaterialEffectiveClosureCount > 0;
-		}
+			FString TreeFunctionPostFix = TEXT("ERROR");
+			switch (SubstrateCompilationContextIndex)
+			{
+			case ESubstrateCompilationContext::SCC_Default:
+				TreeFunctionPostFix = TEXT("");
+				break;
+			case ESubstrateCompilationContext::SCC_FullySimplified:
+				TreeFunctionPostFix = TEXT("_FullySimplified");
+				break;
+			default:
+				check(false);
+			}
 
-		if (bSubstrateFrontMaterialIsValid)
-		{
-			bMaterialIsSubstrate = true;
+			const bool bSubstrateFrontMaterialProvided = Chunk[MP_FrontMaterial] != INDEX_NONE;
+			bool bSubstrateFrontMaterialIsValid = bSubstrateFrontMaterialProvided;
+			if (bSubstrateFrontMaterialIsValid)
+			{
+				// The material can be null when some entries are automatically generated, for instance in the material layer blending system
+				bSubstrateFrontMaterialIsValid &= SubstrateCtx.SubstrateMaterialEffectiveClosureCount > 0;
+
+				UMaterialInterface* MaterialInterface = Material->GetMaterialInterface();
+				UMaterial* BaseMaterial = MaterialInterface ? MaterialInterface->GetMaterial() : nullptr;
+				if (!bSubstrateFrontMaterialIsValid && BaseMaterial && BaseMaterial->bUseMaterialAttributes)
+				{
+					Errorf(TEXT(" %s [%s]: Substrate - Material has no BSDF: this can happen when the SubstrateConvertMaterialAttributes node is not used and when material attributes are enabled."), *Material->GetDebugName(), *Material->GetAssetPath().ToString());
+					TranslationResult = EHLSLMaterialTranslatorResult::Failure;
+					return;
+				}
+			}
 
-			if (SubstrateCtx.SubstrateMaterialRootOperator)
+			if (bSubstrateFrontMaterialIsValid)
 			{
-				// Now implement the functions needed to process the material topology
+				bMaterialIsSubstrate = true;
 
-				// Pre-Update the slab BSDF with operators (like thin film coating, which can alter F0/F90)
+				if (SubstrateCtx.SubstrateMaterialRootOperator)
 				{
-					// Update the coverage/transmittance of each node in the graph
-					check(SubstrateCtx.SubstrateMaterialRootOperator);
-					int32 RootMaximumDistanceToLeaves = SubstrateCtx.SubstrateMaterialRootOperator->MaxDistanceFromLeaves;
+					// Now implement the functions needed to process the material topology
 
-					ResourcesString += FString::Printf(TEXT("void  FSubstratePixelHeader::PreUpdateAllBSDFWithBottomUpOperatorVisit%s(float3 V)\n"), *TreeFunctionPostFix);
-					ResourcesString += "{\n";
-					for (uint32 ClosureIndex = 0; ClosureIndex < SubstrateCtx.SubstrateMaterialEffectiveClosureCount; ++ClosureIndex)
+					// Pre-Update the slab BSDF with operators (like thin film coating, which can alter F0/F90)
 					{
-						ResourcesString += "\t{\n";
-						for (auto& It : SubstrateCtx.SubstrateMaterialExpressionRegisteredOperators)
+						// Update the coverage/transmittance of each node in the graph
+						check(SubstrateCtx.SubstrateMaterialRootOperator);
+						int32 RootMaximumDistanceToLeaves = SubstrateCtx.SubstrateMaterialRootOperator->MaxDistanceFromLeaves;
+
+						ResourcesString += FString::Printf(TEXT("void  FSubstratePixelHeader::PreUpdateAllBSDFWithBottomUpOperatorVisit%s(float3 V)\n"), *TreeFunctionPostFix);
+						ResourcesString += "{\n";
+						for (uint32 ClosureIndex = 0; ClosureIndex < SubstrateCtx.SubstrateMaterialEffectiveClosureCount; ++ClosureIndex)
 						{
-							if (!It.IsDiscarded() && It.BSDFIndex == ClosureIndex)
+							ResourcesString += "\t{\n";
+							for (auto& It : SubstrateCtx.SubstrateMaterialExpressionRegisteredOperators)
 							{
-								// Walk up the graph to the root node and apply weight factors
-								std::function<void(const FSubstrateOperator&, int32)> WalkOperatorsUp = [&](const FSubstrateOperator& CurrentOperator, int32 PreviousOperatorIndex) -> void
+								if (!It.IsDiscarded() && It.BSDFIndex == ClosureIndex)
 								{
-									switch (CurrentOperator.OperatorType)
-									{
-									case SUBSTRATE_OPERATOR_WEIGHT:
-									{
-										break; // NOP
-									}
-									case SUBSTRATE_OPERATOR_HORIZONTAL:
-									{
-										break; // NOP
-									}
-									case SUBSTRATE_OPERATOR_VERTICAL:
-									{
-										// example ResourcesString += FString::Printf(TEXT("\t PreUpdateAllBSDFWithBottomUpOperatorVisit_Vertical(this, SubstrateTree, SubstrateTree.BSDFs[%d], V, %d /*Op index*/, %d /*PreviousIsInputA*/);\n"), BSDFIndex, CurrentOperator.Index, CurrentOperator.LeftIndex == PreviousOperatorIndex ? 1 : 0);
-										break; // NOP
-									}
-									case SUBSTRATE_OPERATOR_ADD:
-									{
-										break; // NOP
-									}
-									case SUBSTRATE_OPERATOR_SELECT:
+									// Walk up the graph to the root node and apply weight factors
+									std::function<void(const FSubstrateOperator&, int32)> WalkOperatorsUp = [&](const FSubstrateOperator& CurrentOperator, int32 PreviousOperatorIndex) -> void
+										{
+											switch (CurrentOperator.OperatorType)
+											{
+											case SUBSTRATE_OPERATOR_WEIGHT:
+											{
+												break; // NOP
+											}
+											case SUBSTRATE_OPERATOR_HORIZONTAL:
+											{
+												break; // NOP
+											}
+											case SUBSTRATE_OPERATOR_VERTICAL:
+											{
+												// example ResourcesString += FString::Printf(TEXT("\t PreUpdateAllBSDFWithBottomUpOperatorVisit_Vertical(this, SubstrateTree, SubstrateTree.BSDFs[%d], V, %d /*Op index*/, %d /*PreviousIsInputA*/);\n"), BSDFIndex, CurrentOperator.Index, CurrentOperator.LeftIndex == PreviousOperatorIndex ? 1 : 0);
+												break; // NOP
+											}
+											case SUBSTRATE_OPERATOR_ADD:
+											{
+												break; // NOP
+											}
+											case SUBSTRATE_OPERATOR_SELECT:
+											{
+												break; // NOP
+											}
+											default:
+											case SUBSTRATE_OPERATOR_BSDF:
+											{
+												check(false);
+											}
+											}
+
+											if (CurrentOperator.ParentIndex != INDEX_NONE)
+											{
+												WalkOperatorsUp(SubstrateCtx.SubstrateMaterialExpressionRegisteredOperators[CurrentOperator.ParentIndex], CurrentOperator.Index);
+											}
+										};
+
+									const int32 BSDFOperatorIndex = It.Index;
+									const FSubstrateOperator& BSDFOperator = SubstrateCtx.SubstrateMaterialExpressionRegisteredOperators[BSDFOperatorIndex];
+
+									// Start visiting node up from the BSDF leaf only if it has a parent.
+									if (BSDFOperator.ParentIndex != INDEX_NONE)
 									{
-										break; // NOP
+										WalkOperatorsUp(SubstrateCtx.SubstrateMaterialExpressionRegisteredOperators[BSDFOperator.ParentIndex], BSDFOperator.Index);
 									}
-									default:
-									case SUBSTRATE_OPERATOR_BSDF:
-									{
-										check(false);
-									}
-									}
-
-									if (CurrentOperator.ParentIndex != INDEX_NONE)
-									{
-										WalkOperatorsUp(SubstrateCtx.SubstrateMaterialExpressionRegisteredOperators[CurrentOperator.ParentIndex], CurrentOperator.Index);
-									}
-								};
-
-								const int32 BSDFOperatorIndex = It.Index;
-								const FSubstrateOperator& BSDFOperator = SubstrateCtx.SubstrateMaterialExpressionRegisteredOperators[BSDFOperatorIndex];
-
-								// Start visiting node up from the BSDF leaf only if it has a parent.
-								if (BSDFOperator.ParentIndex != INDEX_NONE)
-								{
-									WalkOperatorsUp(SubstrateCtx.SubstrateMaterialExpressionRegisteredOperators[BSDFOperator.ParentIndex], BSDFOperator.Index);
 								}
 							}
+							ResourcesString += "\t}\n";
 						}
-						ResourcesString += "\t}\n";
+						ResourcesString += "}\n";
 					}
-					ResourcesString += "}\n";
-				}
 
-				// Update the coverage/transmittance of each leaves (==BSDFs) of the Substrate tree.
-				{
-					ResourcesString += FString::Printf(TEXT("void FSubstratePixelHeader::UpdateAllBSDFsOperatorCoverageTransmittance%s(FSubstrateIntegrationSettings Settings, float3 V)\n"), *TreeFunctionPostFix);
-					ResourcesString += "{\n";
-					for (uint32 ClosureIndex = 0; ClosureIndex < SubstrateCtx.SubstrateMaterialEffectiveClosureCount; ++ClosureIndex)
+					// Update the coverage/transmittance of each leaves (==BSDFs) of the Substrate tree.
 					{
-						ResourcesString += FString::Printf(TEXT("\t #if SUBSTRATE_COMPILER_SUPPORTS_STRUCT_FORWARD_DECLARATION\n"), ClosureIndex);
-						ResourcesString += FString::Printf(TEXT("\t SubstrateTree.UpdateSingleBSDFOperatorCoverageTransmittance(this, %d, Settings, V);\n"), ClosureIndex);
-						ResourcesString += FString::Printf(TEXT("\t #else\n"), ClosureIndex);
-						ResourcesString += FString::Printf(TEXT("\t UpdateSingleBSDFOperatorCoverageTransmittance(SubstrateTree, this, %d, Settings, V);\n"), ClosureIndex);
-						ResourcesString += FString::Printf(TEXT("\t #endif\n"), ClosureIndex);
+						ResourcesString += FString::Printf(TEXT("void FSubstratePixelHeader::UpdateAllBSDFsOperatorCoverageTransmittance%s(FSubstrateIntegrationSettings Settings, float3 V)\n"), *TreeFunctionPostFix);
+						ResourcesString += "{\n";
+						for (uint32 ClosureIndex = 0; ClosureIndex < SubstrateCtx.SubstrateMaterialEffectiveClosureCount; ++ClosureIndex)
+						{
+							ResourcesString += FString::Printf(TEXT("\t #if SUBSTRATE_COMPILER_SUPPORTS_STRUCT_FORWARD_DECLARATION\n"), ClosureIndex);
+							ResourcesString += FString::Printf(TEXT("\t SubstrateTree.UpdateSingleBSDFOperatorCoverageTransmittance(this, %d, Settings, V);\n"), ClosureIndex);
+							ResourcesString += FString::Printf(TEXT("\t #else\n"), ClosureIndex);
+							ResourcesString += FString::Printf(TEXT("\t UpdateSingleBSDFOperatorCoverageTransmittance(SubstrateTree, this, %d, Settings, V);\n"), ClosureIndex);
+							ResourcesString += FString::Printf(TEXT("\t #endif\n"), ClosureIndex);
+						}
+						ResourcesString += "}\n";
 					}
-					ResourcesString += "}\n";
-				}
-
-				// Propagate up the coverage/transmittance of each node in the Substrate tree.
-				// For that we visit all the operator according to their distance from the Substrate tree leaves, from small to large.
-				{
-					check(SubstrateCtx.SubstrateMaterialRootOperator);
-					int32 RootMaximumDistanceToLeaves = SubstrateCtx.SubstrateMaterialRootOperator->MaxDistanceFromLeaves;
 
-					ResourcesString += FString::Printf(TEXT("void FSubstratePixelHeader::UpdateAllOperatorsCoverageTransmittance%s()\n"), *TreeFunctionPostFix);
-					ResourcesString += "{\n";
-					for (int32 DistanceToLeaves = 1; DistanceToLeaves <= RootMaximumDistanceToLeaves; ++DistanceToLeaves)
+					// Propagate up the coverage/transmittance of each node in the Substrate tree.
+					// For that we visit all the operator according to their distance from the Substrate tree leaves, from small to large.
 					{
-						ResourcesString += FString::Printf(TEXT("\t// MaxDistanceFromLeaves = %d \n"), DistanceToLeaves);
-						for (auto& It : SubstrateCtx.SubstrateMaterialExpressionRegisteredOperators)
+						check(SubstrateCtx.SubstrateMaterialRootOperator);
+						int32 RootMaximumDistanceToLeaves = SubstrateCtx.SubstrateMaterialRootOperator->MaxDistanceFromLeaves;
+
+						ResourcesString += FString::Printf(TEXT("void FSubstratePixelHeader::UpdateAllOperatorsCoverageTransmittance%s()\n"), *TreeFunctionPostFix);
+						ResourcesString += "{\n";
+						for (int32 DistanceToLeaves = 1; DistanceToLeaves <= RootMaximumDistanceToLeaves; ++DistanceToLeaves)
 						{
-							if (!It.IsDiscarded() && It.MaxDistanceFromLeaves == DistanceToLeaves)
+							ResourcesString += FString::Printf(TEXT("\t// MaxDistanceFromLeaves = %d \n"), DistanceToLeaves);
+							for (auto& It : SubstrateCtx.SubstrateMaterialExpressionRegisteredOperators)
 							{
-								ResourcesString += FString::Printf(TEXT("\t SubstrateTree.UpdateSingleOperatorCoverageTransmittance(%d /*operator index*/);\n"), It.Index);
+								if (!It.IsDiscarded() && It.MaxDistanceFromLeaves == DistanceToLeaves)
+								{
+									ResourcesString += FString::Printf(TEXT("\t SubstrateTree.UpdateSingleOperatorCoverageTransmittance(%d /*operator index*/);\n"), It.Index);
+								}
 							}
 						}
+						ResourcesString += "}\n";
 					}
-					ResourcesString += "}\n";
-				}
-
-				// Update the luminance weight of each BSDF according to the operators it has to traverse bottom-up to the Substrate tree root node.
-				{
-					// Update the coverage/transmittance of each node in the graph
-					check(SubstrateCtx.SubstrateMaterialRootOperator);
-					int32 RootMaximumDistanceToLeaves = SubstrateCtx.SubstrateMaterialRootOperator->MaxDistanceFromLeaves;
 
-					ResourcesString += FString::Printf(TEXT("void FSubstratePixelHeader::UpdateAllBSDFWithBottomUpOperatorVisit%s()\n"), *TreeFunctionPostFix);
-					ResourcesString += "{\n";
-					for (uint32 ClosureIndex = 0; ClosureIndex < SubstrateCtx.SubstrateMaterialEffectiveClosureCount; ++ClosureIndex)
+					// Update the luminance weight of each BSDF according to the operators it has to traverse bottom-up to the Substrate tree root node.
 					{
-						for (auto& It : SubstrateCtx.SubstrateMaterialExpressionRegisteredOperators)
+						// Update the coverage/transmittance of each node in the graph
+						check(SubstrateCtx.SubstrateMaterialRootOperator);
+						int32 RootMaximumDistanceToLeaves = SubstrateCtx.SubstrateMaterialRootOperator->MaxDistanceFromLeaves;
+
+						ResourcesString += FString::Printf(TEXT("void FSubstratePixelHeader::UpdateAllBSDFWithBottomUpOperatorVisit%s()\n"), *TreeFunctionPostFix);
+						ResourcesString += "{\n";
+						for (uint32 ClosureIndex = 0; ClosureIndex < SubstrateCtx.SubstrateMaterialEffectiveClosureCount; ++ClosureIndex)
 						{
-							if (!It.IsDiscarded() && It.BSDFIndex == ClosureIndex)
+							for (auto& It : SubstrateCtx.SubstrateMaterialExpressionRegisteredOperators)
 							{
-								// Walk up the graph to the root node and apply weight factors
-								std::function<void(const FSubstrateOperator&, int32)> WalkOperatorsUp = [&](const FSubstrateOperator& CurrentOperator, int32 PreviousOperatorIndex) -> void
+								if (!It.IsDiscarded() && It.BSDFIndex == ClosureIndex)
 								{
-									switch (CurrentOperator.OperatorType)
-									{
-									case SUBSTRATE_OPERATOR_WEIGHT:
+									// Walk up the graph to the root node and apply weight factors
+									std::function<void(const FSubstrateOperator&, int32)> WalkOperatorsUp = [&](const FSubstrateOperator& CurrentOperator, int32 PreviousOperatorIndex) -> void
+										{
+											switch (CurrentOperator.OperatorType)
+											{
+											case SUBSTRATE_OPERATOR_WEIGHT:
+											{
+												ResourcesString += FString::Printf(TEXT("\t SubstrateTree.UpdateAllBSDFWithBottomUpOperatorVisit_Weight(%d /*BSDFIndex*/, %d /*Op index*/, %d /*PreviousIsInputA*/);\n"), ClosureIndex, CurrentOperator.Index, 1);
+												break;
+											}
+											case SUBSTRATE_OPERATOR_HORIZONTAL:
+											{
+												ResourcesString += FString::Printf(TEXT("\t SubstrateTree.UpdateAllBSDFWithBottomUpOperatorVisit_Horizontal(%d /*BSDFIndex*/, %d /*Op index*/, %d /*PreviousIsInputA*/);\n"), ClosureIndex, CurrentOperator.Index, CurrentOperator.LeftIndex == PreviousOperatorIndex ? 1 : 0);
+												break;
+											}
+											case SUBSTRATE_OPERATOR_VERTICAL:
+											{
+												ResourcesString += FString::Printf(TEXT("\t SubstrateTree.UpdateAllBSDFWithBottomUpOperatorVisit_Vertical(%d /*BSDFIndex*/, %d /*Op index*/, %d /*PreviousIsInputA*/);\n"), ClosureIndex, CurrentOperator.Index, CurrentOperator.LeftIndex == PreviousOperatorIndex ? 1 : 0);
+												break;
+											}
+											case SUBSTRATE_OPERATOR_ADD:
+											case SUBSTRATE_OPERATOR_SELECT:
+											{
+												break; // NOP
+											}
+											default:
+											case SUBSTRATE_OPERATOR_BSDF:
+											{
+												check(false);
+											}
+											}
+
+											if (CurrentOperator.ParentIndex != INDEX_NONE)
+											{
+												WalkOperatorsUp(SubstrateCtx.SubstrateMaterialExpressionRegisteredOperators[CurrentOperator.ParentIndex], CurrentOperator.Index);
+											}
+										};
+
+									const int32 BSDFOperatorIndex = It.Index;
+									const FSubstrateOperator& BSDFOperator = SubstrateCtx.SubstrateMaterialExpressionRegisteredOperators[BSDFOperatorIndex];
+
+									// Start visiting node up from the BSDF leaf only if it has a parent.
+									if (BSDFOperator.ParentIndex != INDEX_NONE)
 									{
-										ResourcesString += FString::Printf(TEXT("\t SubstrateTree.UpdateAllBSDFWithBottomUpOperatorVisit_Weight(%d /*BSDFIndex*/, %d /*Op index*/, %d /*PreviousIsInputA*/);\n"), ClosureIndex, CurrentOperator.Index, 1);
-										break;
+										WalkOperatorsUp(SubstrateCtx.SubstrateMaterialExpressionRegisteredOperators[BSDFOperator.ParentIndex], BSDFOperator.Index);
 									}
-									case SUBSTRATE_OPERATOR_HORIZONTAL:
-									{
-										ResourcesString += FString::Printf(TEXT("\t SubstrateTree.UpdateAllBSDFWithBottomUpOperatorVisit_Horizontal(%d /*BSDFIndex*/, %d /*Op index*/, %d /*PreviousIsInputA*/);\n"), ClosureIndex, CurrentOperator.Index, CurrentOperator.LeftIndex == PreviousOperatorIndex ? 1 : 0);
-										break;
-									}
-									case SUBSTRATE_OPERATOR_VERTICAL:
-									{
-										ResourcesString += FString::Printf(TEXT("\t SubstrateTree.UpdateAllBSDFWithBottomUpOperatorVisit_Vertical(%d /*BSDFIndex*/, %d /*Op index*/, %d /*PreviousIsInputA*/);\n"), ClosureIndex, CurrentOperator.Index, CurrentOperator.LeftIndex == PreviousOperatorIndex ? 1 : 0);
-										break;
-									}
-									case SUBSTRATE_OPERATOR_ADD:
-									case SUBSTRATE_OPERATOR_SELECT:
-									{
-										break; // NOP
-									}
-									default:
-									case SUBSTRATE_OPERATOR_BSDF:
-									{
-										check(false);
-									}
-									}
-
-									if (CurrentOperator.ParentIndex != INDEX_NONE)
-									{
-										WalkOperatorsUp(SubstrateCtx.SubstrateMaterialExpressionRegisteredOperators[CurrentOperator.ParentIndex], CurrentOperator.Index);
-									}
-								};
-
-								const int32 BSDFOperatorIndex = It.Index;
-								const FSubstrateOperator& BSDFOperator = SubstrateCtx.SubstrateMaterialExpressionRegisteredOperators[BSDFOperatorIndex];
-
-								// Start visiting node up from the BSDF leaf only if it has a parent.
-								if (BSDFOperator.ParentIndex != INDEX_NONE)
-								{
-									WalkOperatorsUp(SubstrateCtx.SubstrateMaterialExpressionRegisteredOperators[BSDFOperator.ParentIndex], BSDFOperator.Index);
 								}
 							}
 						}
+						ResourcesString += "}\n";
 					}
-					ResourcesString += "}\n";
 				}
-			}
 
-			// Check if normal/tangent basis are valid
-			{
-				SubstrateCtx.FinalUsedSharedLocalBasesCount = 0;
-				uint8 RequestedSharedLocalBasesCount = 0;
-				SubstrateCtx.SubstrateEvaluateSharedLocalBases(this, RequestedSharedLocalBasesCount, nullptr);
-				if (RequestedSharedLocalBasesCount > SUBSTRATE_MAX_SHAREDLOCALBASES_REGISTERS)
+				// Check if normal/tangent basis are valid
 				{
-					Errorf(TEXT(" %s [%s]: Substrate - Material has more unique normal/tangent basis than the allowed limit %d/%d."), *Material->GetDebugName(), *Material->GetAssetPath().ToString(), RequestedSharedLocalBasesCount, SUBSTRATE_MAX_SHAREDLOCALBASES_REGISTERS);
+					SubstrateCtx.FinalUsedSharedLocalBasesCount = 0;
+					uint8 RequestedSharedLocalBasesCount = 0;
+					SubstrateCtx.SubstrateEvaluateSharedLocalBases(this, RequestedSharedLocalBasesCount, nullptr);
+					if (RequestedSharedLocalBasesCount > SUBSTRATE_MAX_SHAREDLOCALBASES_REGISTERS)
+					{
+						Errorf(TEXT(" %s [%s]: Substrate - Material has more unique normal/tangent basis than the allowed limit %d/%d."), *Material->GetDebugName(), *Material->GetAssetPath().ToString(), RequestedSharedLocalBasesCount, SUBSTRATE_MAX_SHAREDLOCALBASES_REGISTERS);
+					}
 				}
 			}
-		}
-		else
-		{
-			MaterialCompilationOutput.SubstrateMaterialCompilationOutput.SubstrateMaterialDescription = "";
-			ResourcesString += "// No Substrate material provided\n";
-
-			// Adde default Substrate functions
-			ResourcesString += "#if TEMPLATE_USES_SUBSTRATE\n";
-			ResourcesString += "void PreUpdateAllBSDFWithBottomUpOperatorVisit(float3 V) {}\n";
-			ResourcesString += "void UpdateAllBSDFsOperatorCoverageTransmittance(FSubstrateIntegrationSettings Settings, float3 V) {}\n";
-			ResourcesString += "void UpdateAllOperatorsCoverageTransmittance(inout FSubstrateTree SubstrateTree) {}\n";
-			ResourcesString += "void UpdateAllBSDFWithBottomUpOperatorVisit(inout FSubstrateTree SubstrateTree) {}\n";
-			ResourcesString += "void PreUpdateAllBSDFWithBottomUpOperatorVisit_FullySimplified(float3 V) {}\n";
-			ResourcesString += "void UpdateAllBSDFsOperatorCoverageTransmittance_FullySimplified(FSubstrateIntegrationSettings Settings, float3 V) {}\n";
-			ResourcesString += "void UpdateAllOperatorsCoverageTransmittance_FullySimplified(inout FSubstrateTree SubstrateTree) {}\n";
-			ResourcesString += "void UpdateAllBSDFWithBottomUpOperatorVisit_FullySimplified(inout FSubstrateTree SubstrateTree) {}\n";
-			ResourcesString += "#endif\n";
+			else
+			{
+				MaterialCompilationOutput.SubstrateMaterialCompilationOutput.SubstrateMaterialDescription = "";
+				ResourcesString += "// No Substrate material provided\n";
+
+				// Adde default Substrate functions
+				ResourcesString += "#if TEMPLATE_USES_SUBSTRATE\n";
+				ResourcesString += "void PreUpdateAllBSDFWithBottomUpOperatorVisit(float3 V) {}\n";
+				ResourcesString += "void UpdateAllBSDFsOperatorCoverageTransmittance(FSubstrateIntegrationSettings Settings, float3 V) {}\n";
+				ResourcesString += "void UpdateAllOperatorsCoverageTransmittance(inout FSubstrateTree SubstrateTree) {}\n";
+				ResourcesString += "void UpdateAllBSDFWithBottomUpOperatorVisit(inout FSubstrateTree SubstrateTree) {}\n";
+				ResourcesString += "void PreUpdateAllBSDFWithBottomUpOperatorVisit_FullySimplified(float3 V) {}\n";
+				ResourcesString += "void UpdateAllBSDFsOperatorCoverageTransmittance_FullySimplified(FSubstrateIntegrationSettings Settings, float3 V) {}\n";
+				ResourcesString += "void UpdateAllOperatorsCoverageTransmittance_FullySimplified(inout FSubstrateTree SubstrateTree) {}\n";
+				ResourcesString += "void UpdateAllBSDFWithBottomUpOperatorVisit_FullySimplified(inout FSubstrateTree SubstrateTree) {}\n";
+				ResourcesString += "#endif\n";
+			}
 		}
 	}
 
