[Shaders] Allow shader compilation exclusively through remote workers.

This adds INI option "[DevOptions.Shaders]:bExclusiveRemoteShaderCompiling" and command line option "-ExclusiveRemoteShaderCompiling" to the shader compiling manager,
which disables the local shader compiling thread as long as there is a distributed controller for shader compilation available.
When this option is enabled, the minimum batch size for compile jobs is always 1, otherwise a small number of jobs would always be left in the queue.

This option is disbaled by default and serves mostly for debugging the distributed controller but can also be utilized on low powered machines to reduce the local CPU resources.

diff --git a/Engine/Source/Runtime/Engine/Private/ShaderCompiler/ShaderCompiler.cpp b/Engine/Source/Runtime/Engine/Private/ShaderCompiler/ShaderCompiler.cpp
--- a/Engine/Source/Runtime/Engine/Private/ShaderCompiler/ShaderCompiler.cpp
+++ b/Engine/Source/Runtime/Engine/Private/ShaderCompiler/ShaderCompiler.cpp
@@ -784,6 +784,17 @@ FShaderCompilingManager::FShaderCompilingManager() :
 		bAllowAsynchronousShaderCompiling = false;
 	}
 
+	// If compiling through workers is not allowed, shaders cannot be compiled exclusively throgh remote workers either
+	GConfig->GetBool(TEXT("DevOptions.Shaders"), TEXT("bExclusiveRemoteShaderCompiling"), bExclusiveRemoteShaderCompiling, GEngineIni);
+	if (!bAllowCompilingThroughWorkers)
+	{
+		bExclusiveRemoteShaderCompiling = false;
+	}
+	else if (FParse::Param(FCommandLine::Get(), TEXT("ExclusiveRemoteShaderCompiling")))
+	{
+		bExclusiveRemoteShaderCompiling = true;
+	}
+
 	verify(GConfig->GetInt( TEXT("DevOptions.Shaders"), TEXT("MaxShaderJobBatchSize"), MaxShaderJobBatchSize, GEngineIni ));
 	verify(GConfig->GetBool( TEXT("DevOptions.Shaders"), TEXT("bPromptToRetryFailedShaderCompiles"), bPromptToRetryFailedShaderCompiles, GEngineIni ));
 	verify(GConfig->GetBool(TEXT("DevOptions.Shaders"), TEXT("bDebugBreakOnPromptToRetryShaderCompile"), bDebugBreakOnPromptToRetryShaderCompile, GEngineIni));
@@ -835,38 +846,52 @@ FShaderCompilingManager::FShaderCompilingManager() :
 	
 	if (BuildDistributionController)
 	{
-		UE_LOG(LogShaderCompilers, Display, TEXT("Using %s for Shader Compilation."), *BuildDistributionController->GetName());
+		UE_LOG(LogShaderCompilers, Display, TEXT("Using %s for shader compilation"), *BuildDistributionController->GetName());
 		RemoteCompileThread = MakeUnique<FShaderCompileDistributedThreadRunnable_Interface>(this, *BuildDistributionController);
 	}
+	else if (bExclusiveRemoteShaderCompiling)
+	{
+		UE_LOG(LogShaderCompilers, Display, TEXT("Exclusive remote shader compiling requested, but no distributed controller is available"));
+		bExclusiveRemoteShaderCompiling = false;
+	}
 
 	GConfig->SetBool(TEXT("/Script/UnrealEd.UnrealEdOptions"), TEXT("UsingXGE"), RemoteCompileThread.IsValid(), GEditorIni);
 
-	TUniquePtr<FShaderCompileThreadRunnableBase> LocalThread = MakeUnique<FShaderCompileThreadRunnable>(this);
-	if (RemoteCompileThread)
+	if (bExclusiveRemoteShaderCompiling)
 	{
-		checkf(ShaderCompiler::IsRemoteCompilingAllowed(), TEXT("We have a remote compiling thread without the remote compilation being allowed"));
-
-		// Only force-local jobs are guaranteed to stay on the local machine. Going wide with High priority jobs is important for the startup times,
-		// since special materials use High priority. Possibly the partition by priority is too rigid in general.
-		RemoteCompileThread->SetPriorityRange(EShaderCompileJobPriority::Low, EShaderCompileJobPriority::ExtraHigh);
-		LocalThread->SetPriorityRange(EShaderCompileJobPriority::Normal, EShaderCompileJobPriority::ForceLocal);
+		// Only add remote shader compiling thread
 		Threads.Add(MoveTemp(RemoteCompileThread));
 	}
 	else
 	{
-		UE_LOG(LogShaderCompilers, Display, TEXT("Using Local Shader Compiler with %d workers."), NumShaderCompilingThreads);
+		// Add local and remote shader compiling threads and distribute job priorities among local and remote workers
+		UE_LOG(LogShaderCompilers, Display, TEXT("Using %d local workers for shader compilation"), NumShaderCompilingThreads);
 
-		if (GIsBuildMachine)
+		TUniquePtr<FShaderCompileThreadRunnableBase> LocalThread = MakeUnique<FShaderCompileThreadRunnable>(this);
+		if (RemoteCompileThread)
 		{
-			int32 MinSCWsToSpawnBeforeWarning = 8; // optional, default to 8
-			GConfig->GetInt(TEXT("DevOptions.Shaders"), TEXT("MinSCWsToSpawnBeforeWarning"), MinSCWsToSpawnBeforeWarning, GEngineIni);
-			if (NumShaderCompilingThreads < static_cast<uint32>(MinSCWsToSpawnBeforeWarning))
+			checkf(ShaderCompiler::IsRemoteCompilingAllowed(), TEXT("We have a remote compiling thread without the remote compilation being allowed"));
+
+			// Only force-local jobs are guaranteed to stay on the local machine. Going wide with High priority jobs is important for the startup times,
+			// since special materials use High priority. Possibly the partition by priority is too rigid in general.
+			RemoteCompileThread->SetPriorityRange(EShaderCompileJobPriority::Low, EShaderCompileJobPriority::ExtraHigh);
+			LocalThread->SetPriorityRange(EShaderCompileJobPriority::Normal, EShaderCompileJobPriority::ForceLocal);
+			Threads.Add(MoveTemp(RemoteCompileThread));
+		}
+		else
+		{
+			if (GIsBuildMachine)
 			{
-				UE_LOG(LogShaderCompilers, Warning, TEXT("Only %d SCWs will be spawned, which will result in longer shader compile times."), NumShaderCompilingThreads);
+				int32 MinSCWsToSpawnBeforeWarning = 8; // optional, default to 8
+				GConfig->GetInt(TEXT("DevOptions.Shaders"), TEXT("MinSCWsToSpawnBeforeWarning"), MinSCWsToSpawnBeforeWarning, GEngineIni);
+				if (NumShaderCompilingThreads < static_cast<uint32>(MinSCWsToSpawnBeforeWarning))
+				{
+					UE_LOG(LogShaderCompilers, Warning, TEXT("Only %d SCWs will be spawned, which will result in longer shader compile times."), NumShaderCompilingThreads);
+				}
 			}
 		}
+		Threads.Add(MoveTemp(LocalThread));
 	}
-	Threads.Add(MoveTemp(LocalThread));
 
 	for (const auto& Thread : Threads)
 	{
diff --git a/Engine/Source/Runtime/Engine/Private/ShaderCompiler/ShaderCompilerDistributed.cpp b/Engine/Source/Runtime/Engine/Private/ShaderCompiler/ShaderCompilerDistributed.cpp
--- a/Engine/Source/Runtime/Engine/Private/ShaderCompiler/ShaderCompilerDistributed.cpp
+++ b/Engine/Source/Runtime/Engine/Private/ShaderCompiler/ShaderCompilerDistributed.cpp
@@ -320,9 +320,9 @@ int32 FShaderCompileDistributedThreadRunnable_Interface::CompilingLoop()
 	{
 		for (int32 PriorityIndex = MaxPriorityIndex; PriorityIndex >= MinPriorityIndex; --PriorityIndex)
 		{
-			// Grab as many jobs from the job queue as we can
+			// Grab as many jobs from the job queue as we can, unless there is no local shader compiling thread to pick up smaller batches
 			const EShaderCompileJobPriority Priority = (EShaderCompileJobPriority)PriorityIndex;
-			const int32 MinBatchSize = (Priority == EShaderCompileJobPriority::Low) ? 1 : DistributedShaderCompilerVariables::MinBatchSize;
+			const int32 MinBatchSize = (Priority == EShaderCompileJobPriority::Low || Manager->IsExclusiveRemoteCompilingEnabled()) ? 1 : DistributedShaderCompilerVariables::MinBatchSize;
 			const int32 NumJobs = Manager->AllJobs.GetPendingJobs(EShaderCompilerWorkerType::Distributed, Priority, MinBatchSize, INT32_MAX, PendingJobs);
 			if (NumJobs > 0)
 			{
diff --git a/Engine/Source/Runtime/Engine/Public/ShaderCompiler.h b/Engine/Source/Runtime/Engine/Public/ShaderCompiler.h
--- a/Engine/Source/Runtime/Engine/Public/ShaderCompiler.h
+++ b/Engine/Source/Runtime/Engine/Public/ShaderCompiler.h
@@ -900,6 +900,8 @@ private:
 	bool bAllowCompilingThroughWorkers;
 	/** Whether to allow shaders to compile in the background or to block after each material. */
 	bool bAllowAsynchronousShaderCompiling;
+	/** Whether shaders are compiled exclusively through remote workers and no local shader compiling thread is allocated. Ignored if 'bAllowCompilingThroughWorkers' is false. */
+	bool bExclusiveRemoteShaderCompiling = false;
 	/** Whether to ask to retry a failed shader compile error. */
 	bool bPromptToRetryFailedShaderCompiles;
 	/** If enabled when we enter the prompt to retry we will break in the debugger if one is attached rather than prompting. */
@@ -1038,6 +1040,15 @@ public:
 		return BuildDistributionController != nullptr;
 	}
 
+	/**
+	 * Returns whether shaders are exclusively compiled through remote workerd, i.e. no local shader compiling thread is allocated.
+	 * If this is true, the distributed controller must not refuse job batches below a minimum threshold, otherwise those jobs will never get picked up.
+	 */
+	bool IsExclusiveRemoteCompilingEnabled() const
+	{
+		return bExclusiveRemoteShaderCompiling;
+	}
+
 	/**
 	 * Returns whether normal throttling settings should be ignored because shader compilation is at the moment the only action blocking the critical path.
 	 * An example of such situation is startup compilation of global shaders and default materials, but there may be more cases like that.
