Path Tracing: Fix PreExposure not being set correctly

The logic for when to apply PreExposure needs to consider path tracing explicitly. Updated the logic to consider which debug modes in RT debug and PT respectively need to make use of the pre-exposure.

Refactor the logic for FRayTracingVisualizationData to avoid the explicit mutex (can be handled by function local static init instead) which simplifies the code a bit.

diff --git a/Engine/Source/Runtime/Engine/Private/RayTracingVisualizationData.cpp b/Engine/Source/Runtime/Engine/Private/RayTracingVisualizationData.cpp
--- a/Engine/Source/Runtime/Engine/Private/RayTracingVisualizationData.cpp
+++ b/Engine/Source/Runtime/Engine/Private/RayTracingVisualizationData.cpp
@@ -9,83 +9,75 @@
 
 #define LOCTEXT_NAMESPACE "FRayTracingVisualizationData"
 
-static FRayTracingVisualizationData GRayTracingVisualizationData;
-
-void FRayTracingVisualizationData::Initialize()
+FRayTracingVisualizationData::FRayTracingVisualizationData()
 {
-	UE::TScopeLock Lock(Mutex);
-	if (!bIsInitialized)
-	{
-		const bool bRayTracingShadersSupported = FDataDrivenShaderPlatformInfo::GetSupportsRayTracingShaders(GMaxRHIShaderPlatform);
+	const bool bRayTracingShadersSupported = FDataDrivenShaderPlatformInfo::GetSupportsRayTracingShaders(GMaxRHIShaderPlatform);
 
-		// always supported (as long as either inline RT or RT shaders work)
-		AddVisualizationMode(TEXT("Barycentrics"), LOCTEXT("Barycentrics", "Barycentrics"), FModeType::Standard, RAY_TRACING_DEBUG_VIZ_BARYCENTRICS, false);
+	// always supported (as long as either inline RT or RT shaders work)
+	AddVisualizationMode(TEXT("Barycentrics"), LOCTEXT("Barycentrics", "Barycentrics"), FModeType::Standard, RAY_TRACING_DEBUG_VIZ_BARYCENTRICS, false);
 
-		if (bRayTracingShadersSupported)
+	if (bRayTracingShadersSupported)
+	{
+		// runs basic lighting calculations on hits
+		AddVisualizationMode(TEXT("PrimaryRays"), LOCTEXT("PrimaryRays", "Primary Rays"), FModeType::Standard, RAY_TRACING_DEBUG_VIZ_PRIMARY_RAYS, true);
+
+		// all of these fields reflect entries in the payload which require running a CHS
+		AddVisualizationMode(TEXT("Radiance"), LOCTEXT("Radiance", "Radiance"), FModeType::Other, RAY_TRACING_DEBUG_VIZ_RADIANCE, true);
+		AddVisualizationMode(TEXT("WorldNormal"), LOCTEXT("WorldNormal", "World Normal"), FModeType::Standard, RAY_TRACING_DEBUG_VIZ_WORLD_NORMAL, false);
+		AddVisualizationMode(TEXT("BaseColor"), LOCTEXT("BaseColor", "Base Color"), FModeType::Standard, RAY_TRACING_DEBUG_VIZ_BASE_COLOR, false);
+		AddVisualizationMode(TEXT("DiffuseColor"), LOCTEXT("DiffuseColor", "Diffuse Color"), FModeType::Other, RAY_TRACING_DEBUG_VIZ_DIFFUSE_COLOR, false);
+		AddVisualizationMode(TEXT("SpecularColor"), LOCTEXT("SpecularColor", "Specular Color"), FModeType::Other, RAY_TRACING_DEBUG_VIZ_SPECULAR_COLOR, false);
+		AddVisualizationMode(TEXT("Opacity"), LOCTEXT("Opacity", "Opacity"), FModeType::Other, RAY_TRACING_DEBUG_VIZ_OPACITY, false);
+		AddVisualizationMode(TEXT("Metallic"), LOCTEXT("Metallic", "Metallic"), FModeType::Other, RAY_TRACING_DEBUG_VIZ_METALLIC, false);
+		AddVisualizationMode(TEXT("Specular"), LOCTEXT("Specular", "Specular"), FModeType::Other, RAY_TRACING_DEBUG_VIZ_SPECULAR, false);
+		AddVisualizationMode(TEXT("Roughness"), LOCTEXT("Roughness", "Roughness"), FModeType::Other, RAY_TRACING_DEBUG_VIZ_ROUGHNESS, false);
+		AddVisualizationMode(TEXT("Ior"), LOCTEXT("Ior", "Ior"), FModeType::Other, RAY_TRACING_DEBUG_VIZ_IOR, false);
+		AddVisualizationMode(TEXT("ShadingModelID"), LOCTEXT("ShadingModelID", "Shading Model ID"), FModeType::Other, RAY_TRACING_DEBUG_VIZ_SHADING_MODEL, false);
+		AddVisualizationMode(TEXT("BlendingMode"), LOCTEXT("BlendingMode", "Blending Mode"), FModeType::Other, RAY_TRACING_DEBUG_VIZ_BLENDING_MODE, false);
+		AddVisualizationMode(TEXT("PrimitiveLightingChannelMask"), LOCTEXT("PrimitiveLightingChannelMask", "Primitive Lighting Channel Mask"), FModeType::Other, RAY_TRACING_DEBUG_VIZ_LIGHTING_CHANNEL_MASK, false);
+		AddVisualizationMode(TEXT("CustomData"), LOCTEXT("CustomData", "Custom Data"), FModeType::Other, RAY_TRACING_DEBUG_VIZ_CUSTOM_DATA, false);
+		AddVisualizationMode(TEXT("GBufferAO"), LOCTEXT("GBufferAO", "GBuffer AO"), FModeType::Other, RAY_TRACING_DEBUG_VIZ_GBUFFER_AO, false);
+		AddVisualizationMode(TEXT("IndirectIrradiance"), LOCTEXT("IndirectIrradiance", "Indirect Irradiance"), FModeType::Other, RAY_TRACING_DEBUG_VIZ_INDIRECT_IRRADIANCE, true);
+		AddVisualizationMode(TEXT("WorldPosition"), LOCTEXT("WorldPosition", "World Position"), FModeType::Other, RAY_TRACING_DEBUG_VIZ_WORLD_POSITION, false);
+		AddVisualizationMode(TEXT("HitKind"), LOCTEXT("HitKind", "Hit Kind"), FModeType::Other, RAY_TRACING_DEBUG_VIZ_HITKIND, false);
+		AddVisualizationMode(TEXT("WorldTangent"), LOCTEXT("WorldTangent", "World Tangent"), FModeType::Other, RAY_TRACING_DEBUG_VIZ_WORLD_TANGENT, false);
+		AddVisualizationMode(TEXT("Anisotropy"), LOCTEXT("Anisotropy", "Anisotropy"), FModeType::Other, RAY_TRACING_DEBUG_VIZ_ANISOTROPY, false);
+
+		// debugging the geometry itself
+		AddVisualizationMode(TEXT("Instances"), LOCTEXT("Instances", "Instances"), FModeType::Standard, RAY_TRACING_DEBUG_VIZ_INSTANCES, false);
+		AddVisualizationMode(TEXT("Triangles"), LOCTEXT("Triangles", "Triangles"), FModeType::Standard, RAY_TRACING_DEBUG_VIZ_TRIANGLES, false);
+		AddVisualizationMode(TEXT("FarField"), LOCTEXT("FarField", "Far Field"), FModeType::Standard, RAY_TRACING_DEBUG_VIZ_FAR_FIELD, false);
+		AddVisualizationMode(TEXT("DynamicInstances"), LOCTEXT("DynamicInstances", "Dynamic Instances"), FModeType::Standard, RAY_TRACING_DEBUG_VIZ_DYNAMIC_INSTANCES, false);
+		AddVisualizationMode(TEXT("ProxyType"), LOCTEXT("ProxyType", "Proxy Type"), FModeType::Standard, RAY_TRACING_DEBUG_VIZ_PROXY_TYPE, false);
+		AddVisualizationMode(TEXT("Picker"), LOCTEXT("Picker", "Picker"), FModeType::Standard, RAY_TRACING_DEBUG_VIZ_PICKER, false);
+		AddVisualizationMode(TEXT("LightGridOccupancy"), LOCTEXT("LightGridOccupancy", "Light Grid Occupancy"), FModeType::Standard, RAY_TRACING_DEBUG_VIZ_LIGHT_GRID_COUNT, false);
+
+		// performance
+		AddVisualizationMode(TEXT("InstanceOverlap"), LOCTEXT("InstanceOverlap", "Instance Overlap"), FModeType::Performance, RAY_TRACING_DEBUG_VIZ_INSTANCE_OVERLAP, false);
+		AddVisualizationMode(TEXT("TriangleHitCount"), LOCTEXT("TriangleHitCount", "Triangle Hit Count"), FModeType::Performance, RAY_TRACING_DEBUG_VIZ_TRIANGLE_HITCOUNT, false);
+		AddVisualizationMode(TEXT("HitCountPerInstance"), LOCTEXT("HitCountPerInstance", "Hit Count Per Instance"), FModeType::Performance, RAY_TRACING_DEBUG_VIZ_HITCOUNT_PER_INSTANCE, false);
+
+		if (GRHISupportsShaderTimestamp)
 		{
-			// runs basic lighting calculations on hits
-			AddVisualizationMode(TEXT("PrimaryRays"), LOCTEXT("PrimaryRays", "Primary Rays"), FModeType::Standard, RAY_TRACING_DEBUG_VIZ_PRIMARY_RAYS, true);
-
-			// all of these fields reflect entries in the payload which require running a CHS
-			AddVisualizationMode(TEXT("Radiance"), LOCTEXT("Radiance", "Radiance"), FModeType::Other, RAY_TRACING_DEBUG_VIZ_RADIANCE, true);
-			AddVisualizationMode(TEXT("WorldNormal"), LOCTEXT("WorldNormal", "World Normal"), FModeType::Standard, RAY_TRACING_DEBUG_VIZ_WORLD_NORMAL, false);
-			AddVisualizationMode(TEXT("BaseColor"), LOCTEXT("BaseColor", "Base Color"), FModeType::Standard, RAY_TRACING_DEBUG_VIZ_BASE_COLOR, false);
-			AddVisualizationMode(TEXT("DiffuseColor"), LOCTEXT("DiffuseColor", "Diffuse Color"), FModeType::Other, RAY_TRACING_DEBUG_VIZ_DIFFUSE_COLOR, false);
-			AddVisualizationMode(TEXT("SpecularColor"), LOCTEXT("SpecularColor", "Specular Color"), FModeType::Other, RAY_TRACING_DEBUG_VIZ_SPECULAR_COLOR, false);
-			AddVisualizationMode(TEXT("Opacity"), LOCTEXT("Opacity", "Opacity"), FModeType::Other, RAY_TRACING_DEBUG_VIZ_OPACITY, false);
-			AddVisualizationMode(TEXT("Metallic"), LOCTEXT("Metallic", "Metallic"), FModeType::Other, RAY_TRACING_DEBUG_VIZ_METALLIC, false);
-			AddVisualizationMode(TEXT("Specular"), LOCTEXT("Specular", "Specular"), FModeType::Other, RAY_TRACING_DEBUG_VIZ_SPECULAR, false);
-			AddVisualizationMode(TEXT("Roughness"), LOCTEXT("Roughness", "Roughness"), FModeType::Other, RAY_TRACING_DEBUG_VIZ_ROUGHNESS, false);
-			AddVisualizationMode(TEXT("Ior"), LOCTEXT("Ior", "Ior"), FModeType::Other, RAY_TRACING_DEBUG_VIZ_IOR, false);
-			AddVisualizationMode(TEXT("ShadingModelID"), LOCTEXT("ShadingModelID", "Shading Model ID"), FModeType::Other, RAY_TRACING_DEBUG_VIZ_SHADING_MODEL, false);
-			AddVisualizationMode(TEXT("BlendingMode"), LOCTEXT("BlendingMode", "Blending Mode"), FModeType::Other, RAY_TRACING_DEBUG_VIZ_BLENDING_MODE, false);
-			AddVisualizationMode(TEXT("PrimitiveLightingChannelMask"), LOCTEXT("PrimitiveLightingChannelMask", "Primitive Lighting Channel Mask"), FModeType::Other, RAY_TRACING_DEBUG_VIZ_LIGHTING_CHANNEL_MASK, false);
-			AddVisualizationMode(TEXT("CustomData"), LOCTEXT("CustomData", "Custom Data"), FModeType::Other, RAY_TRACING_DEBUG_VIZ_CUSTOM_DATA, false);
-			AddVisualizationMode(TEXT("GBufferAO"), LOCTEXT("GBufferAO", "GBuffer AO"), FModeType::Other, RAY_TRACING_DEBUG_VIZ_GBUFFER_AO, false);
-			AddVisualizationMode(TEXT("IndirectIrradiance"), LOCTEXT("IndirectIrradiance", "Indirect Irradiance"), FModeType::Other, RAY_TRACING_DEBUG_VIZ_INDIRECT_IRRADIANCE, true);
-			AddVisualizationMode(TEXT("WorldPosition"), LOCTEXT("WorldPosition", "World Position"), FModeType::Other, RAY_TRACING_DEBUG_VIZ_WORLD_POSITION, false);
-			AddVisualizationMode(TEXT("HitKind"), LOCTEXT("HitKind", "Hit Kind"), FModeType::Other, RAY_TRACING_DEBUG_VIZ_HITKIND, false);
-			AddVisualizationMode(TEXT("WorldTangent"), LOCTEXT("WorldTangent", "World Tangent"), FModeType::Other, RAY_TRACING_DEBUG_VIZ_WORLD_TANGENT, false);
-			AddVisualizationMode(TEXT("Anisotropy"), LOCTEXT("Anisotropy", "Anisotropy"), FModeType::Other, RAY_TRACING_DEBUG_VIZ_ANISOTROPY, false);
-
-			// debugging the geometry itself
-			AddVisualizationMode(TEXT("Instances"), LOCTEXT("Instances", "Instances"), FModeType::Standard, RAY_TRACING_DEBUG_VIZ_INSTANCES, false);
-			AddVisualizationMode(TEXT("Triangles"), LOCTEXT("Triangles", "Triangles"), FModeType::Standard, RAY_TRACING_DEBUG_VIZ_TRIANGLES, false);
-			AddVisualizationMode(TEXT("FarField"), LOCTEXT("FarField", "Far Field"), FModeType::Standard, RAY_TRACING_DEBUG_VIZ_FAR_FIELD, false);
-			AddVisualizationMode(TEXT("DynamicInstances"), LOCTEXT("DynamicInstances", "Dynamic Instances"), FModeType::Standard, RAY_TRACING_DEBUG_VIZ_DYNAMIC_INSTANCES, false);
-			AddVisualizationMode(TEXT("ProxyType"), LOCTEXT("ProxyType", "Proxy Type"), FModeType::Standard, RAY_TRACING_DEBUG_VIZ_PROXY_TYPE, false);
-			AddVisualizationMode(TEXT("Picker"), LOCTEXT("Picker", "Picker"), FModeType::Standard, RAY_TRACING_DEBUG_VIZ_PICKER, false);
-			AddVisualizationMode(TEXT("LightGridOccupancy"), LOCTEXT("LightGridOccupancy", "Light Grid Occupancy"), FModeType::Standard, RAY_TRACING_DEBUG_VIZ_LIGHT_GRID_COUNT, false);
-
-			// performance
-			AddVisualizationMode(TEXT("InstanceOverlap"), LOCTEXT("InstanceOverlap", "Instance Overlap"), FModeType::Performance, RAY_TRACING_DEBUG_VIZ_INSTANCE_OVERLAP, false);
-			AddVisualizationMode(TEXT("TriangleHitCount"), LOCTEXT("TriangleHitCount", "Triangle Hit Count"), FModeType::Performance, RAY_TRACING_DEBUG_VIZ_TRIANGLE_HITCOUNT, false);
-			AddVisualizationMode(TEXT("HitCountPerInstance"), LOCTEXT("HitCountPerInstance", "Hit Count Per Instance"), FModeType::Performance, RAY_TRACING_DEBUG_VIZ_HITCOUNT_PER_INSTANCE, false);
-
-			if (GRHISupportsShaderTimestamp)
-			{
-				AddVisualizationMode(TEXT("Performance"), LOCTEXT("Performance", "Performance"), FModeType::Performance, RAY_TRACING_DEBUG_VIZ_PERFORMANCE, false);
-			}
-
-			AddVisualizationMode(TEXT("Traversal Primary Node"), LOCTEXT("TraversalPrimaryNode", "Traversal Primary Node"), FModeType::Performance, RAY_TRACING_DEBUG_VIZ_TRAVERSAL_PRIMARY_NODE, false, /*bHiddenInEditor*/ true);
-			AddVisualizationMode(TEXT("Traversal Primary Cluster"), LOCTEXT("TraversalPrimaryCluster", "Traversal Primary Cluster"), FModeType::Performance, RAY_TRACING_DEBUG_VIZ_TRAVERSAL_PRIMARY_CLUSTER, false, /*bHiddenInEditor*/ true);
-			AddVisualizationMode(TEXT("Traversal Primary Triangle"), LOCTEXT("TraversalPrimaryTriangle", "Traversal Primary Triangle"), FModeType::Performance, RAY_TRACING_DEBUG_VIZ_TRAVERSAL_PRIMARY_TRIANGLE, false, /*bHiddenInEditor*/ true);
-			AddVisualizationMode(TEXT("Traversal Primary All"), LOCTEXT("TraversalPrimaryAll", "Traversal Primary All"), FModeType::Performance, RAY_TRACING_DEBUG_VIZ_TRAVERSAL_PRIMARY_ALL, false, /*bHiddenInEditor*/ true);
-			AddVisualizationMode(TEXT("Traversal Primary Statistics"), LOCTEXT("TraversalPrimaryStatistics", "Traversal Primary Statistics"), FModeType::Performance, RAY_TRACING_DEBUG_VIZ_TRAVERSAL_PRIMARY_STATISTICS, false, /*bHiddenInEditor*/ true);
-
-			AddVisualizationMode(TEXT("Traversal Secondary Node"), LOCTEXT("TraversalSecondaryNode", "Traversal Secondary Node"), FModeType::Performance, RAY_TRACING_DEBUG_VIZ_TRAVERSAL_SECONDARY_NODE, false, /*bHiddenInEditor*/ true);
-			AddVisualizationMode(TEXT("Traversal Secondary Cluster"), LOCTEXT("TraversalSecondaryCluster", "Traversal Secondary Cluster"), FModeType::Performance, RAY_TRACING_DEBUG_VIZ_TRAVERSAL_SECONDARY_CLUSTER, false, /*bHiddenInEditor*/ true);
-			AddVisualizationMode(TEXT("Traversal Secondary Triangle"), LOCTEXT("TraversalSecondaryTriangle", "Traversal Secondary Triangle"), FModeType::Performance, RAY_TRACING_DEBUG_VIZ_TRAVERSAL_SECONDARY_TRIANGLE, false, /*bHiddenInEditor*/ true);
-			AddVisualizationMode(TEXT("Traversal Secondary All"), LOCTEXT("TraversalSecondaryAll", "Traversal Secondary All"), FModeType::Performance, RAY_TRACING_DEBUG_VIZ_TRAVERSAL_SECONDARY_ALL, false, /*bHiddenInEditor*/ true);
-			AddVisualizationMode(TEXT("Traversal Secondary Statistics"), LOCTEXT("TraversalSecondaryStatistics", "Traversal Secondary Statistics"), FModeType::Performance, RAY_TRACING_DEBUG_VIZ_TRAVERSAL_SECONDARY_STATISTICS, false, /*bHiddenInEditor*/ true);
-
-			if (Substrate::IsSubstrateEnabled())
-			{
-				AddVisualizationMode(TEXT("SubstrateMaterialProperties"), LOCTEXT("SubstrateMaterialProperties", "Substrate Material Properties"), FModeType::Material, RAY_TRACING_DEBUG_VIZ_SUBSTRATE_DATA, true);
-			}
+			AddVisualizationMode(TEXT("Performance"), LOCTEXT("Performance", "Performance"), FModeType::Performance, RAY_TRACING_DEBUG_VIZ_PERFORMANCE, false);
 		}
 
-		bIsInitialized = true;
+		AddVisualizationMode(TEXT("Traversal Primary Node"), LOCTEXT("TraversalPrimaryNode", "Traversal Primary Node"), FModeType::Performance, RAY_TRACING_DEBUG_VIZ_TRAVERSAL_PRIMARY_NODE, false, /*bHiddenInEditor*/ true);
+		AddVisualizationMode(TEXT("Traversal Primary Cluster"), LOCTEXT("TraversalPrimaryCluster", "Traversal Primary Cluster"), FModeType::Performance, RAY_TRACING_DEBUG_VIZ_TRAVERSAL_PRIMARY_CLUSTER, false, /*bHiddenInEditor*/ true);
+		AddVisualizationMode(TEXT("Traversal Primary Triangle"), LOCTEXT("TraversalPrimaryTriangle", "Traversal Primary Triangle"), FModeType::Performance, RAY_TRACING_DEBUG_VIZ_TRAVERSAL_PRIMARY_TRIANGLE, false, /*bHiddenInEditor*/ true);
+		AddVisualizationMode(TEXT("Traversal Primary All"), LOCTEXT("TraversalPrimaryAll", "Traversal Primary All"), FModeType::Performance, RAY_TRACING_DEBUG_VIZ_TRAVERSAL_PRIMARY_ALL, false, /*bHiddenInEditor*/ true);
+		AddVisualizationMode(TEXT("Traversal Primary Statistics"), LOCTEXT("TraversalPrimaryStatistics", "Traversal Primary Statistics"), FModeType::Performance, RAY_TRACING_DEBUG_VIZ_TRAVERSAL_PRIMARY_STATISTICS, false, /*bHiddenInEditor*/ true);
+
+		AddVisualizationMode(TEXT("Traversal Secondary Node"), LOCTEXT("TraversalSecondaryNode", "Traversal Secondary Node"), FModeType::Performance, RAY_TRACING_DEBUG_VIZ_TRAVERSAL_SECONDARY_NODE, false, /*bHiddenInEditor*/ true);
+		AddVisualizationMode(TEXT("Traversal Secondary Cluster"), LOCTEXT("TraversalSecondaryCluster", "Traversal Secondary Cluster"), FModeType::Performance, RAY_TRACING_DEBUG_VIZ_TRAVERSAL_SECONDARY_CLUSTER, false, /*bHiddenInEditor*/ true);
+		AddVisualizationMode(TEXT("Traversal Secondary Triangle"), LOCTEXT("TraversalSecondaryTriangle", "Traversal Secondary Triangle"), FModeType::Performance, RAY_TRACING_DEBUG_VIZ_TRAVERSAL_SECONDARY_TRIANGLE, false, /*bHiddenInEditor*/ true);
+		AddVisualizationMode(TEXT("Traversal Secondary All"), LOCTEXT("TraversalSecondaryAll", "Traversal Secondary All"), FModeType::Performance, RAY_TRACING_DEBUG_VIZ_TRAVERSAL_SECONDARY_ALL, false, /*bHiddenInEditor*/ true);
+		AddVisualizationMode(TEXT("Traversal Secondary Statistics"), LOCTEXT("TraversalSecondaryStatistics", "Traversal Secondary Statistics"), FModeType::Performance, RAY_TRACING_DEBUG_VIZ_TRAVERSAL_SECONDARY_STATISTICS, false, /*bHiddenInEditor*/ true);
+
+		if (Substrate::IsSubstrateEnabled())
+		{
+			AddVisualizationMode(TEXT("SubstrateMaterialProperties"), LOCTEXT("SubstrateMaterialProperties", "Substrate Material Properties"), FModeType::Material, RAY_TRACING_DEBUG_VIZ_SUBSTRATE_DATA, true);
+		}
 	}
 }
 
@@ -149,11 +141,7 @@ bool FRayTracingVisualizationData::GetModeTonemapped(const FName& InModeName) co
 
 FRayTracingVisualizationData& GetRayTracingVisualizationData()
 {
-	if (!GRayTracingVisualizationData.IsInitialized())
-	{
-		GRayTracingVisualizationData.Initialize();
-	}
-
+	static FRayTracingVisualizationData GRayTracingVisualizationData;
 	return GRayTracingVisualizationData;
 }
 
diff --git a/Engine/Source/Runtime/Engine/Public/RayTracingVisualizationData.h b/Engine/Source/Runtime/Engine/Public/RayTracingVisualizationData.h
--- a/Engine/Source/Runtime/Engine/Public/RayTracingVisualizationData.h
+++ b/Engine/Source/Runtime/Engine/Public/RayTracingVisualizationData.h
@@ -3,7 +3,6 @@
 #pragma once
 
 #include "CoreMinimal.h"
-#include "Async/Mutex.h"
 
 class FRayTracingVisualizationData
 {
@@ -35,16 +34,7 @@ public:
 	typedef TMultiMap<FName, FModeRecord> TModeMap;
 
 public:
-	FRayTracingVisualizationData()
-		: bIsInitialized(false)
-	{
-	}
-
-	/** Initialize the system. */
-	void Initialize();
-
-	/** Check if system was initialized. */
-	inline bool IsInitialized() const { return bIsInitialized; }
+	FRayTracingVisualizationData();
 
 	/** Get the display name of a named mode from the available mode map. **/
 	ENGINE_API FText GetModeDisplayName(const FName& InModeName) const;
@@ -75,12 +65,6 @@ private:
 	
 	/** Storage for console variable documentation strings. **/
 	FString ConsoleDocumentationVisualizationMode;
-
-	/** Flag indicating if system is initialized. **/
-	std::atomic_bool bIsInitialized;
-
-	/** Mutex for initialization. */
-	UE::FMutex Mutex;
 };
 
 ENGINE_API FRayTracingVisualizationData& GetRayTracingVisualizationData();
diff --git a/Engine/Source/Runtime/Renderer/Private/PathTracing.cpp b/Engine/Source/Runtime/Renderer/Private/PathTracing.cpp
--- a/Engine/Source/Runtime/Renderer/Private/PathTracing.cpp
+++ b/Engine/Source/Runtime/Renderer/Private/PathTracing.cpp
@@ -792,6 +792,14 @@ namespace PathTracing
 	{
 		return GetPathTracingDebugVisualizationMode() >= 0;
 	}
+
+	bool NeedsTonemapping()
+	{
+		const int DebugMode = GetPathTracingDebugVisualizationMode();
+		return DebugMode < 0 ||
+			   DebugMode == PATH_TRACER_DEBUG_VIZ_RADIANCE ||
+			   DebugMode == PATH_TRACER_DEBUG_VIZ_PRIMARY_RAYS;
+	}
 }
 
 // This function prepares the portion of shader arguments that may involve invalidating the path traced state
diff --git a/Engine/Source/Runtime/Renderer/Private/PathTracing.h b/Engine/Source/Runtime/Renderer/Private/PathTracing.h
--- a/Engine/Source/Runtime/Renderer/Private/PathTracing.h
+++ b/Engine/Source/Runtime/Renderer/Private/PathTracing.h
@@ -39,4 +39,5 @@ namespace PathTracing
 	bool UsesReferenceAtmosphere(const FViewInfo& View);
 	bool UsesReferenceDOF(const FViewInfo& View);
 	bool NeedsAntiAliasing(const FViewInfo& View);
+	bool NeedsTonemapping();
 }
diff --git a/Engine/Source/Runtime/Renderer/Private/PostProcess/PostProcessEyeAdaptation.cpp b/Engine/Source/Runtime/Renderer/Private/PostProcess/PostProcessEyeAdaptation.cpp
--- a/Engine/Source/Runtime/Renderer/Private/PostProcess/PostProcessEyeAdaptation.cpp
+++ b/Engine/Source/Runtime/Renderer/Private/PostProcess/PostProcessEyeAdaptation.cpp
@@ -1226,8 +1226,9 @@ void FViewInfo::UpdatePreExposure()
 		&& !ViewFamily.EngineShowFlags.LODColoration
 		&& !ViewFamily.EngineShowFlags.HLODColoration
 		&& !ViewFamily.EngineShowFlags.ActorColoration
-		&& ((!ViewFamily.EngineShowFlags.VisualizeBuffer) || CurrentBufferVisualizationMode != NAME_None) // disable pre-exposure for the buffer visualization modes
-		&& !HasRayTracedOverlay(ViewFamily);
+		&& ((!ViewFamily.EngineShowFlags.VisualizeBuffer) || CurrentBufferVisualizationMode != NAME_None)   // disable pre-exposure for the buffer visualization modes
+		&& (!ViewFamily.EngineShowFlags.PathTracing || PathTracing::NeedsTonemapping())                     // Either we aren't using path tracing at all, or we are using a path tracing mode that needs tone-mapping
+		&& (!ViewFamily.EngineShowFlags.RayTracingDebug || RaytracingDebugViewModeNeedsTonemapping(*this)); // Either we aren't using RT debug at all, or we are using a mode that needs tone-mapping
 
 	// Compute the PreExposure to use.
 	bool bUpdateLastExposure = false;
diff --git a/Engine/Source/Runtime/Renderer/Private/RayTracing/RayTracingDebug.cpp b/Engine/Source/Runtime/Renderer/Private/RayTracing/RayTracingDebug.cpp
--- a/Engine/Source/Runtime/Renderer/Private/RayTracing/RayTracingDebug.cpp
+++ b/Engine/Source/Runtime/Renderer/Private/RayTracing/RayTracingDebug.cpp
@@ -1378,26 +1378,37 @@ static uint32 GetRaytracingDebugViewMode(const FSceneView& View)
 {
 	const FRayTracingVisualizationData& VisualizationData = GetRayTracingVisualizationData();
 
-	uint32 OutMode;
-
-	FString ConsoleViewMode = CVarRayTracingVisualize.GetValueOnAnyThread();
+	FString ConsoleViewMode = CVarRayTracingVisualize.GetValueOnRenderThread();
 	if (!ConsoleViewMode.IsEmpty())
 	{
-		OutMode = VisualizationData.GetModeID(FName(*ConsoleViewMode));
+		return  VisualizationData.GetModeID(FName(*ConsoleViewMode));
 	}
-	else if(View.CurrentRayTracingDebugVisualizationMode != NAME_None)
+	else if (View.CurrentRayTracingDebugVisualizationMode != NAME_None)
 	{
-		OutMode = VisualizationData.GetModeID(View.CurrentRayTracingDebugVisualizationMode);
+		return VisualizationData.GetModeID(View.CurrentRayTracingDebugVisualizationMode);
 	}
-	else
+	// Set useful default value
+	return RAY_TRACING_DEBUG_VIZ_BARYCENTRICS;
+}
+
+bool RaytracingDebugViewModeNeedsTonemapping(const FSceneView& View)
+{
+	const FRayTracingVisualizationData& VisualizationData = GetRayTracingVisualizationData();
+
+	FString ConsoleViewMode = CVarRayTracingVisualize.GetValueOnRenderThread();
+	if (!ConsoleViewMode.IsEmpty())
 	{
-		// Set useful default value
-		OutMode = RAY_TRACING_DEBUG_VIZ_BARYCENTRICS;
+		return VisualizationData.GetModeTonemapped(FName(*ConsoleViewMode));
 	}
-
-	return OutMode;
+	else if (View.CurrentRayTracingDebugVisualizationMode != NAME_None)
+	{
+		return VisualizationData.GetModeTonemapped(View.CurrentRayTracingDebugVisualizationMode);
+	}
+	// unknown mode, assume no tone-mapping is needed
+	return false;
 }
 
+
 bool HasRaytracingDebugViewModeRaytracedOverlay(const FSceneViewFamily& ViewFamily)
 {
 	bool bAnySubstrate = false;
diff --git a/Engine/Source/Runtime/Renderer/Private/RayTracing/RaytracingOptions.h b/Engine/Source/Runtime/Renderer/Private/RayTracing/RaytracingOptions.h
--- a/Engine/Source/Runtime/Renderer/Private/RayTracing/RaytracingOptions.h
+++ b/Engine/Source/Runtime/Renderer/Private/RayTracing/RaytracingOptions.h
@@ -100,6 +100,7 @@ extern float GetRaytracingMaxNormalBias();
 extern float GetRayTracingCullingRadius();
 extern bool IsRayTracingInstanceDebugDataEnabled(const FViewInfo& View);
 extern bool IsRayTracingInstanceOverlapEnabled(const FViewInfo& View);
+extern bool RaytracingDebugViewModeNeedsTonemapping(const FSceneView& View);
 
 extern bool CanUseRayTracingAMDHitToken();
 
@@ -188,4 +189,9 @@ FORCEINLINE bool IsRayTracingInstanceOverlapEnabled(const FViewInfo& View)
 	return false;
 }
 
+FORCEINLINE bool RaytracingDebugViewModeNeedsTonemapping(const FViewInfo& View)
+{
+	return false;
+}
+
 #endif // RHI_RAYTRACING
