Fix teleport type not applying to deferred kinematic updates for skeletal meshes

Actor teleports in character blueprints for simulated skeletal meshes can have their teleports handled by deferred kinematic updates. Previously the teleport flag was ignored, leading the teleport flag to not work for character blueprints.

diff --git a/Engine/Source/Runtime/Engine/Private/PhysicsEngine/Experimental/PhysScene_Chaos.cpp b/Engine/Source/Runtime/Engine/Private/PhysicsEngine/Experimental/PhysScene_Chaos.cpp
--- a/Engine/Source/Runtime/Engine/Private/PhysicsEngine/Experimental/PhysScene_Chaos.cpp
+++ b/Engine/Source/Runtime/Engine/Private/PhysicsEngine/Experimental/PhysScene_Chaos.cpp
@@ -2029,7 +2029,7 @@ void FPhysScene_Chaos::UpdateKinematicsOnDeferredSkelMeshes()
 					}
 
 					Chaos::FRigidBodyHandle_External& Body_External = ActorHandle->GetGameThreadAPI();
-					if (!BodyInst->IsInstanceSimulatingPhysics())
+					if (!BodyInst->IsInstanceSimulatingPhysics() || Info.TeleportType != ETeleportType::None)
 					{
 						const int32 BoneIndex = BodyInst->InstanceBoneIndex;
 						if (BoneIndex != INDEX_NONE)
