In cooked builds, when loading level instances through the Blueprint LoadLevelInstance node (not through an actual level instance actor), make sure that the level streaming world asset is unique by appending the hash of the provided unique package name to it.

diff --git a/Engine/Source/Runtime/Engine/Private/WorldPartition/WorldPartitionRuntimeLevelStreamingCell.cpp b/Engine/Source/Runtime/Engine/Private/WorldPartition/WorldPartitionRuntimeLevelStreamingCell.cpp
--- a/Engine/Source/Runtime/Engine/Private/WorldPartition/WorldPartitionRuntimeLevelStreamingCell.cpp
+++ b/Engine/Source/Runtime/Engine/Private/WorldPartition/WorldPartitionRuntimeLevelStreamingCell.cpp
@@ -590,14 +590,13 @@ UWorldPartitionLevelStreamingDynamic* UWorldPartitionRuntimeLevelStreamingCell::
 		// Transfer WorldPartition's transform to LevelStreaming
 		LevelStreaming->SetLevelTransform(WorldPartition->GetInstanceTransform());
 
-		// LevelStreaming WorldAsset is a TSoftObjectPtr<UWorld>. If the WorldPartition's world is instanced then the TSoftObjectPtr<UWorld> will be remapped by FLinkerInstancingContext SoftObject remapping
-		// example if MainWorld is instanced it will have a package suffix and/or prefix like: /Temp+PATH+_LevelInstance1
-		// MainWorld: /Game/SomePath/MainWorld_LevelInstance1
-		// 
-		// and the LevelStreaming's WorldAsset TSoftObjectPtr<UWorld> should be remapped as: 
-		// /Temp/Game/SomePath/MainWorld/_Generated_/1AEHCD0PRR98XVM12PU4N1D8X_LevelInstance_1
-
-		check(!WorldPartition->GetTypedOuter<ULevel>()->IsInstancedLevel() || (LevelStreaming->PackageNameToLoad != LevelStreaming->GetWorldAssetPackageName()));
+		// Make sure we have a unique world assetif the world is instanced. Normally the world asset is remapped through the FLinkerInstancingContext soft object 
+		// remapping, but loading an instance of a partitioned world through the LoadLevelInstance blueprint node will not.
+		if (WorldPartition->GetTypedOuter<ULevel>()->IsInstancedLevel() && (LevelStreaming->PackageNameToLoad == LevelStreaming->GetWorldAssetPackageName()))
+		{
+			const TSoftObjectPtr<UWorld> NewWorldAsset(FSoftObjectPath(FString::Printf(TEXT("%s_LevelInstance_%08x"), *LevelStreaming->GetWorldAssetPackageFName().GetPlainNameString(), GetTypeHash(GetPackage()->GetName()))));
+			LevelStreaming->SetWorldAsset(NewWorldAsset);
+		}
 	}
 #endif
 
