- Changed how DPs handle setting/unsetting, making use of the CVar hierarchy/history system. We now store the parent chain as separate entries in thie history, so we can easily hotfix/etc a DP, and we only need to update just that DPs bits of cvar values.
- Additionally, it now allows to see more context/info when viewing the history (you can now visualize how the DP parent chain sets the same CVar
- Fixed a memory leak when unsetting cvars by Tag
- The old code is still there, and is the default, until more testing. It's controlled by -DPApplyMode=[0/1/2]

- Added FAutoRangedConsoleVariable that can be used to create a .min and .max versions of the cvar, and will clamp the value when Get*() functions are called on it. Plan is to use this for scalabulity groups and allow dynamic configs to clamp the values

diff --git a/Engine/Source/Runtime/Core/Private/HAL/ConsoleManager.cpp b/Engine/Source/Runtime/Core/Private/HAL/ConsoleManager.cpp
--- a/Engine/Source/Runtime/Core/Private/HAL/ConsoleManager.cpp
+++ b/Engine/Source/Runtime/Core/Private/HAL/ConsoleManager.cpp
@@ -118,6 +118,7 @@ public:
 	virtual void RegisterThreadPropagation(uint32 ThreadId, IConsoleThreadPropagation* InCallback) override;
 	virtual void UnregisterConsoleObject(IConsoleObject* Object, bool bKeepState) override;
 	virtual void UnsetAllConsoleVariablesWithTag(FName Tag, EConsoleVariableFlags Priority) override;
+	virtual void BatchUpdateTag(FName Tag, const TMap<FName, FString>& CVarsAndValues) override;
 #if ALLOW_OTHER_PLATFORM_CONFIG
 	virtual void LoadAllPlatformCVars(FName PlatformName, const FString& DeviceProfileName = FString()) override;
 	virtual void ClearAllPlatformCVars(FName PlatformName = NAME_None, const FString& DeviceProfileName = FString()) override;
@@ -876,6 +877,8 @@ constexpr bool IsArrayPriority(EConsoleVariableFlags Priority)
 		Priority == ECVF_SetByPluginLowPriority ||
 		Priority == ECVF_SetByPluginHighPriority ||
 		Priority == ECVF_SetByHotfix ||
+		Priority == ECVF_SetByDeviceProfile ||
+		Priority == ECVF_SetByScalability ||
 		Priority == ECVF_SetByPreview;
 }
 
@@ -889,36 +892,74 @@ public:
 	
 	bool bHasTaggedArrayData = false;
 	
-	void Track(const TCHAR* InValue, EConsoleVariableFlags SetBy, FName Tag)
+	/**
+	 * When we change a vlaue, we track the change at each priority level, so we can unset them later. 
+	 * @return false if this changed a value in place and it is not the final level (in array type priorities). 
+	 * This is used to not fire unnecessary callbacks in this case:
+	 *   PluginHighPriority -> Tag1=Foo, Tag2=Bar, and we are setting Tag1 to Baz
+	 * the value of the cvar cannot change, and when we update in place we are trying to minimize callbacks
+	 */
+	bool Track(const TCHAR* InValue, EConsoleVariableFlags Flags, FName Tag)
 	{
 		T LocalCopy;
 		UE::ConsoleManager::Private::GetValueFromString<T>(LocalCopy, UE::ConfigUtilities::ConvertValueFromHumanFriendlyValue(InValue));
 		
-		int Priority = (int)SetBy;
-		TArray<FTaggedHistoryData>& ValueArray = History.FindOrAdd(Priority);
-		if (!IsArrayPriority(SetBy))
+		EConsoleVariableFlags SetBy = (EConsoleVariableFlags)(Flags & ECVF_SetByMask);
+		TArray<FTaggedHistoryData>& ValueArray = History.FindOrAdd((int)SetBy);
+		
+		// if desired, set in place, looking for existing tag, and replacing it if it's there
+		// otherwise always add it to the end further down
+		if (Flags & ECVF_Set_ReplaceExistingTag)
 		{
-			if (ValueArray.Num() > 0)
+			for (int32 Index = 0; Index < ValueArray.Num(); Index++)
 			{
-				// don't do anything if we already are set
-				if (ValueArray[0].Key == Tag && ValueArray[0].Value.GetValueOnAnyThread() == LocalCopy)
+				FTaggedHistoryData& HistoryItem = ValueArray[Index];
+				if (HistoryItem.Key == Tag)
 				{
-					return;
+					HistoryItem.Value = LocalCopy;
+					// if we aren't the last entry, then we don't want to allow any callbacks to fire
+					// (the CanChange check later on just checks against SetBy level, not by array index)
+					return Index == ValueArray.Num() - 1;
 				}
-				//remove the entry but leave space
-				ValueArray.Empty(1);
 			}
 		}
-		// add the value to the array (even if it's a non-array type, we store the value in the array)
-		ValueArray.Emplace(Tag, LocalCopy);
+
+		if (!IsArrayPriority(SetBy) && ValueArray.Num() > 0)
+		{
+			// replace the value if it already exists here
+			ValueArray[0].Value = LocalCopy;
+		}
+		else
+		{
+			// untagged array types must only have one entry with no tag, but we push it to the end
+			if (Tag == NAME_None)
+			{
+				
+				for (int32 Index = 0; Index < ValueArray.Num(); Index++)
+				{
+					if (ValueArray[Index].Key == Tag)
+					{
+						ValueArray.RemoveAt(Index);
+						break;
+					}
+				}
+			}
+
+			// add the value to the array
+			ValueArray.Emplace(Tag, LocalCopy);
+		}
+
+		return true;
 	}
 	
 	/**
 	 * Unset the value at the SetBy, and optionally Tag.
 	 * Return true if it had been set, or false if nothing happened
 	 */
-	bool Unset(EConsoleVariableFlags SetBy, FName Tag)
+	bool Unset(EConsoleVariableFlags Flags, FName Tag)
 	{
+		EConsoleVariableFlags SetBy = (EConsoleVariableFlags)(Flags & ECVF_SetByMask);
+
 		if (!History.Contains(SetBy))
 		{
 			return false;
@@ -1085,7 +1126,10 @@ protected:
 	virtual void SetOtherPlatformValue(const TCHAR* InValue, EConsoleVariableFlags SetBy, FName Tag) override
 	{
 		// always track it
-		TrackHistory(InValue, SetBy, Tag);
+		if (TrackHistory(InValue, SetBy, Tag) == false)
+		{
+			return;
+		}
 
 		// set it if we are equal to or higher than before
 		uint32 CurrentSetBy = GetFlags() & ECVF_SetByMask;
@@ -1109,7 +1153,7 @@ protected:
 	 * This is the key function that subclasses need to call in their Set() implementation. This will track the values of cvars at
 	 * different priorities, so priorities/plugins/etc can be unset later, and the CVar will updatae state correctly (in Unset())
 	 */
-	void TrackHistory(const TCHAR* InValue, EConsoleVariableFlags SetBy, FName Tag)
+	bool TrackHistory(const TCHAR* InValue, EConsoleVariableFlags SetBy, FName Tag)
 	{
 		// make a history if we want to
 		if (PriorityHistory == nullptr)
@@ -1117,7 +1161,11 @@ protected:
 			PriorityHistory = new FConsoleVariableHistory<T>();
 			PriorityHistory->Track(*GetString(), ECVF_SetByConstructor, NAME_None);
 		}
-		PriorityHistory->Track(InValue, SetBy, Tag);
+
+		if (PriorityHistory->Track(InValue, SetBy, Tag) == false)
+		{
+			return false;
+		}
 
 		// this cvar so we can remove ourself later
 		if (Tag != NAME_None)
@@ -1132,6 +1180,8 @@ protected:
 			// set a flag to remember we need to remove ourself from TaggedCVars in destructor
 			PriorityHistory->bHasTaggedArrayData = true;
 		}
+
+		return true;
 	}
 	
 	virtual SIZE_T GetHistorySize() override
@@ -1150,6 +1200,24 @@ protected:
 			PriorityHistory->Log(Ar);
 		}
 	}
+
+	bool PreprocessSet(const TCHAR* InValue, EConsoleVariableFlags SetBy, FName Tag)
+	{
+		// FILL IN THE NEW RENDERTHREAD STUFF HERE
+
+		if (TrackHistory(InValue, SetBy, Tag) == false)
+		{
+			return false;
+		}
+
+		if (CanChange(SetBy) == false)
+		{
+			return false;
+		}
+
+		return true;
+	}
+
 	
 	/**
 	 * Removes the value at the given SetBy (and potentially Tag for the Array type SetBy priorities). This will update the
@@ -1324,7 +1392,11 @@ public:
 	
 	virtual void Set(const TCHAR* InValue, EConsoleVariableFlags SetBy, FName Tag)
 	{
-		this->TrackHistory(InValue, SetBy, Tag);
+		// if it was already set in the history (this returns false if nothing needs to be done), then skip the Set
+		if (this->TrackHistory(InValue, SetBy, Tag) == false)
+		{
+			return;
+		}
 
 		if (CanChange(SetBy))
 		{
@@ -1427,9 +1499,7 @@ template<> TConsoleVariableData<float>* FConsoleVariable<float>::AsVariableFloat
 
 template<> void FConsoleVariable<FString>::Set(const TCHAR* InValue, EConsoleVariableFlags SetBy, FName Tag)
 {
-	TrackHistory(InValue, SetBy, Tag);
-	
-	if(CanChange(SetBy))
+	if (PreprocessSet(InValue, SetBy, Tag))
 	{
 		Data.ShadowedValue[0] = InValue;
 		OnChanged(SetBy);
@@ -1603,7 +1673,10 @@ public:
 
 	virtual void Set(const TCHAR* InValue, EConsoleVariableFlags SetBy, FName Tag)
 	{
-		this->TrackHistory(InValue, SetBy, Tag);
+		if (this->TrackHistory(InValue, SetBy, Tag) == false)
+		{
+			return;
+		}
 		
 		if(CanChange(SetBy))
 		{
@@ -1688,9 +1761,7 @@ public:
 	}
 	virtual void Set(const TCHAR* InValue, EConsoleVariableFlags SetBy, FName Tag)
 	{
-		TrackHistory(InValue, SetBy, Tag);
-		
-		if (CanChange(SetBy))
+		if (PreprocessSet(InValue, SetBy, Tag))
 		{
 			MainValue = InValue;
 			OnChanged(SetBy);
@@ -1766,9 +1837,7 @@ public:
 	} 
 	virtual void Set(const TCHAR* InValue, EConsoleVariableFlags SetBy, FName Tag)
 	{
-		TrackHistory(InValue, SetBy, Tag);
-		
-		if(CanChange(SetBy))
+		if (PreprocessSet(InValue, SetBy, Tag))
 		{
 			int32 Value = FCString::Atoi(InValue);
 
@@ -1828,6 +1897,66 @@ IConsoleVariable* FConsoleManager::RegisterConsoleVariableBitRef(const TCHAR* CV
 	return AddConsoleObject(CVarName, new FConsoleVariableBitRef(FlagName, BitNumber, Force0MaskPtr, Force1MaskPtr, Help, (EConsoleVariableFlags)Flags))->AsVariable();
 }
 
+
+template<typename T>
+class FRangedConsoleVariable : public FConsoleVariable<T>
+{
+	using Base = FConsoleVariable<T>;
+public:
+	IConsoleVariable* Min;
+	IConsoleVariable* Max;
+
+	FRangedConsoleVariable(const FString& Name, T DefaultValue, T MinValue, T MaxValue, const TCHAR* Help, EConsoleVariableFlags Flags, IConsoleVariable* Parent)
+		: FConsoleVariable<T>(DefaultValue, Help, Flags, Parent)
+	{
+		FConsoleManager& Manager = (FConsoleManager&)IConsoleManager::Get();
+		Min = Manager.RegisterConsoleVariable(*(Name + ".min"), MinValue, *FString::Printf(TEXT("The min value for CVar %s"), *Name), Flags);
+		Max = Manager.RegisterConsoleVariable(*(Name + ".max"), MaxValue, *FString::Printf(TEXT("The max value for CVar %s"), *Name), Flags);
+
+		// when min and max change, trigger an update of this cvar so any listeners can deal with _potential_ new value returned from Get
+		Min->SetOnChangedCallback(FConsoleVariableDelegate::CreateLambda([this](IConsoleVariable*) { this->OnChangedDelegate().Broadcast(this); }));
+		Max->SetOnChangedCallback(FConsoleVariableDelegate::CreateLambda([this](IConsoleVariable*) { this->OnChangedDelegate().Broadcast(this); }));
+	}
+
+	virtual bool GetBool() const override
+	{
+		unimplemented();
+		return false;
+	}
+	virtual int32 GetInt() const override
+	{
+		return FMath::Clamp(Base::GetInt(), Min->GetInt(), Max->GetInt());
+	}
+	virtual float GetFloat() const override
+	{
+		return FMath::Clamp(Base::GetInt(), Min->GetFloat(), Max->GetFloat());
+	}
+	virtual FString GetString() const override
+	{
+		T Value, MinValue, MaxValue;
+		this->GetValue(Value);
+		Min->GetValue(MinValue);
+		Max->GetValue(MaxValue);
+
+		return LexToString(FMath::Clamp(Value, MinValue, MaxValue));
+	}
+};
+
+template<>
+TAutoRangedConsoleVariable<int32>::TAutoRangedConsoleVariable(const TCHAR* Name, int32 DefaultValue, int32 MinValue, int32 MaxValue, const TCHAR* Help, uint32 Flags, IConsoleVariable* Parent)
+	: TAutoConsoleVariable(GetManager().AddConsoleObject(Name, new FRangedConsoleVariable(Name, DefaultValue, MinValue, MaxValue, Help, (EConsoleVariableFlags)Flags, Parent))->AsVariable())
+{
+}
+
+template<>
+TAutoRangedConsoleVariable<float>::TAutoRangedConsoleVariable(const TCHAR* Name, float DefaultValue, float MinValue, float MaxValue, const TCHAR* Help, uint32 Flags, IConsoleVariable* Parent)
+	: TAutoConsoleVariable(GetManager().AddConsoleObject(Name, new FRangedConsoleVariable(Name, DefaultValue, MinValue, MaxValue, Help, (EConsoleVariableFlags)Flags, Parent))->AsVariable())
+{
+}
+
+//FAutoRangedConsoleVariable GTestVar(TEXT("josh.range"), 2, 0, 4, TEXT("Help str"), ECVF_Default);
+
+
 #if !(UE_BUILD_SHIPPING || UE_BUILD_TEST)
 // part of the automated test for console variables
 static TAutoConsoleVariable<int32> CVarDebugEarlyDefault(
@@ -3262,7 +3391,64 @@ void FConsoleManager::UnsetAllConsoleVariablesWithTag(FName Tag, EConsoleVariabl
 		Var->Unset(Priority, Tag);
 	}
 	
+	delete UE::ConsoleManager::Private::TaggedCVars.FindRef(Tag);
+	UE::ConsoleManager::Private::TaggedCVars.Remove(Tag);
+}
+
+void FConsoleManager::BatchUpdateTag(FName Tag, const TMap<FName, FString>& CVarsAndValues)
+{
+	// cache the tagged cvars 
+	TSet<IConsoleVariable*>* TaggedCVarsPtr = UE::ConsoleManager::Private::TaggedCVars.FindRef(Tag);
+	TSet<IConsoleVariable*> TaggedCVars;
+	if (TaggedCVarsPtr != nullptr)
+	{
+		TaggedCVars = *TaggedCVarsPtr;
+	}
+
+	for (auto& Pair : CVarsAndValues)
+	{
+		FString CVarNameString = Pair.Key.ToString();
+		IConsoleObject* CObj = ConsoleObjects.FindRef(CVarNameString);
+		// @todo handle when not found to make a dummy one?
+		if (CObj)
+		{
+			if (IConsoleVariable* CVar = CObj->AsVariable())
+			{
+				// remove this from the set of tracked cvars, so we don't unset it at the end
+				TaggedCVars.Remove(CVar);
+			}
+		}
+	}
+
+	// now unset aren't going to be set - this has to be done before we delete the global TaggedCVars below
+	// @todo - deal with different SetBys?
+	for (IConsoleVariable* CVar : TaggedCVars)
+	{
+		CVar->Unset(ECVF_SetByMask, Tag);
+	}
+
+	// remove it from the set, so we can add back cleanly
+	delete UE::ConsoleManager::Private::TaggedCVars.FindRef(Tag);
 	UE::ConsoleManager::Private::TaggedCVars.Remove(Tag);
+
+	// now set everything, without reordering the array types
+	for (auto& Pair : CVarsAndValues)
+	{
+		FString CVarNameString = Pair.Key.ToString();
+		IConsoleObject* CObj = ConsoleObjects.FindRef(CVarNameString);
+		// @todo handle when not found to make a dummy one?
+		if (CObj)
+		{
+			if (IConsoleVariable* CVar = CObj->AsVariable())
+			{
+				bool bIsScalabilityBucket = CVarNameString.StartsWith(TEXT("sg."));
+				uint32 CVarPriority = bIsScalabilityBucket ? ECVF_SetByScalability : ECVF_SetByDeviceProfile;
+				// since we are updating, replace with existing key
+				CVar->Set(*Pair.Value, (EConsoleVariableFlags)(CVarPriority | ECVF_Set_ReplaceExistingTag), Tag);
+			}
+		}
+	}
+
 }
 
 #if ALLOW_OTHER_PLATFORM_CONFIG
diff --git a/Engine/Source/Runtime/Core/Private/Misc/ConfigUtilities.cpp b/Engine/Source/Runtime/Core/Private/Misc/ConfigUtilities.cpp
--- a/Engine/Source/Runtime/Core/Private/Misc/ConfigUtilities.cpp
+++ b/Engine/Source/Runtime/Core/Private/Misc/ConfigUtilities.cpp
@@ -20,8 +20,10 @@
 namespace UE::DynamicConfig
 {
 TMulticastDelegate<void(const FName&, const FName&, class FConfigModificationTracker*)> HotfixPluginForBranch;
-TMulticastDelegate<void(const class FConfigModificationTracker* ChangeTracker)> ReloadObjects;
-TMulticastDelegate<void(const TSet<FString>&)> UpdateDeviceProfiles;
+TMulticastDelegate<void(const FConfigModificationTracker* ChangeTracker)> ReloadObjects;
+TMulticastDelegate<void(const FConfigModificationTracker* ChangeTracker)> UpdateCVarsAndDeviceProfiles;
+// deprecated:
+TMulticastDelegate<void(const TSet<FString>& ModifiedSections)> UpdateDeviceProfiles;
 
 void PerformDynamicConfig(FName Tag, TFunction<void(class FConfigModificationTracker*)> PerformModification,
 	TFunction<void(class FConfigModificationTracker*)> PostModification)
@@ -35,57 +37,6 @@ void PerformDynamicConfig(FName Tag, TFunction<void(class FConfigModificationTra
 	// now update everything if anything was read in!
 	if (ChangeTracker.ModifiedSectionsPerBranch.Num() > 0)
 	{
-		// reload objects that had their configs changed
-		for (const auto& SectionPair : ChangeTracker.CVars)
-		{
-			const FConfigModificationTracker::FCVarTracker& CVars = SectionPair.Value;
-			EConsoleVariableFlags Priority = (EConsoleVariableFlags)CVars.CVarPriority;
-
-			// now walk over the updated cvars and set them based on the priority
-			for (const auto& TagPair : CVars.CVarEntriesPerBranchPerTag)
-			{
-				for (const auto& BranchPair : TagPair.Value)
-				{
-					const FConfigSection& Section = BranchPair.Value;
-					
-					for (const auto& CVarPair : Section)
-					{
-						FString KeyString = CVarPair.Key.ToString();
-						if (KeyString.StartsWith(TEXT("sg.")))
-						{
-							UE::ConfigUtilities::ApplyCVarSettingsFromIni(*FString::Printf(TEXT("%s@%s"), *KeyString.Mid(3), *CVarPair.Value.GetValue()), *GScalabilityIni, Priority, false, TagPair.Key);
-						}
-						else
-						{
-							UE::ConfigUtilities::OnSetCVarFromIniEntry(TEXT("DynamicLayer"), *CVarPair.Key.ToString(), *CVarPair.Value.GetValue(), Priority, false, false, TagPair.Key);
-						}
-					}
-				}
-			}
-		}
-
-		// if we modify a scalability group section (ie EffectsQuality@2) that matches current value, re-set the group to apply the modified group
-		TSet<FString>* ScalabilitySections = ChangeTracker.ModifiedSectionsPerBranch.Find(*GScalabilityIni);
-		if (ScalabilitySections != nullptr)
-		{
-			for (const FString& SectionName : *ScalabilitySections)
-			{
-				FString Group, Level;
-				if (SectionName.Split(TEXT("@"), &Group, &Level))
-				{
-					FString SGName = FString::Printf(TEXT("sg.%s"), *Group);
-					IConsoleVariable* SGVar = IConsoleManager::Get().FindConsoleVariable(*SGName);
-					if (SGVar != nullptr)
-					{
-						if (SGVar->GetString() == Level)
-						{
-							SGVar->SetWithCurrentPriority(*Level);
-						}
-					}
-				}
-			}
-		}
-
 		if (ReloadObjects.IsBound())
 		{
 			// check for any object-reload control in each branch that was modified
@@ -102,8 +53,12 @@ void PerformDynamicConfig(FName Tag, TFunction<void(class FConfigModificationTra
 			// reload any config uobjects that were updated (assuming the uobject system is in use)
 			ReloadObjects.Broadcast(&ChangeTracker);
 		}
-		
-		// if we updated GDeviceProfilesIni, then fixup active DP if needed
+
+		// update any CVars and DPs that were updated
+		UpdateCVarsAndDeviceProfiles.Broadcast(&ChangeTracker);
+
+
+		// deprecated:
 		TSet<FString>* DPSections = ChangeTracker.ModifiedSectionsPerBranch.Find(*GDeviceProfilesIni);
 		if (DPSections != nullptr)
 		{
diff --git a/Engine/Source/Runtime/Core/Public/HAL/IConsoleManager.h b/Engine/Source/Runtime/Core/Public/HAL/IConsoleManager.h
--- a/Engine/Source/Runtime/Core/Public/HAL/IConsoleManager.h
+++ b/Engine/Source/Runtime/Core/Public/HAL/IConsoleManager.h
@@ -129,6 +129,9 @@ enum EConsoleVariableFlags
 	// Similar to ECVF_Set_NoSinkCall_Unsafe, but this will do nothing but set the cvar, and push updates to render thread, no changing flags, no sinks, etc
 	ECVF_Set_SetOnly_Unsafe = 		0x00020000,
 
+	// When this is set, this will stomp on existing tagged array type values in place, instead of adding to the end of the array
+	ECVF_Set_ReplaceExistingTag =	0x00040000,
+
 	// ------------------------------------------------
 
 	/* to get some history of where the last value was set by ( useful for track down why a cvar is in a specific state */
@@ -138,7 +141,7 @@ enum EConsoleVariableFlags
 
 	// Lowest priority (default after console variable creation)
 	ECVF_SetByConstructor =			0x00000000,
-	// Set by scalability groups from Scalability.ini (lower priority than game settings so it's easier to override partially)
+	// Set by scalability groups from Scalability.ini (lower priority than game settings so it's easier to override partially). THIS IS AN ARRAY TYPE
 	ECVF_SetByScalability =			0x01000000,
 	// Default priority for engine-level game user settings, platform-specific settings will override this
 	ECVF_SetByGameSetting =			0x02000000,
@@ -148,7 +151,7 @@ enum EConsoleVariableFlags
 	ECVF_SetBySystemSettingsIni =	0x04000000,
 	// Dyanmically loaded/unloaded plugins, used with the History concept to restore cvars on plugin unload. THIS IS AN ARRAY TYPE meaning multiple vales with this SetBy are stored in the history.
 	ECVF_SetByPluginLowPriority =	0x05000000,
-	// Per device settings using the DeviceProfiles.ini hierarchy (e.g. specific iOS device, higher priority than per project to do device specific settings)
+	// Per device settings using the DeviceProfiles.ini hierarchy (e.g. specific iOS device, higher priority than per project to do device specific settings). THIS IS AN ARRAY TYPE
 	ECVF_SetByDeviceProfile =		0x06000000,
 	// Dyanmically loaded/unloaded plugins, used with the History concept to restore cvars on plugin unload. THIS IS AN ARRAY TYPE meaning multiple vales with this SetBy are stored in the history.
 	ECVF_SetByPluginHighPriority =	0x07000000,
@@ -567,19 +570,19 @@ public:
 	virtual FString GetString() const = 0;
 
 	/** Generic versions for templated code */
-	void GetValue(int32& OutIntValue)
+	void GetValue(int32& OutIntValue) const
 	{
 		OutIntValue = GetInt();
 	}
-	void GetValue(bool& OutBoolValue)
+	void GetValue(bool& OutBoolValue) const
 	{
 		OutBoolValue = GetBool();
 	}
-	void GetValue(float& OutFloatValue)
+	void GetValue(float& OutFloatValue) const
 	{
 		OutFloatValue = GetFloat();
 	}
-	void GetValue(FString& OutStringValue)
+	void GetValue(FString& OutStringValue) const
 	{
 		OutStringValue = GetString();
 	}
@@ -1216,7 +1219,12 @@ struct IConsoleManager
 	  * @param Priority If set, then the internal search for CVars is restricted to this SetBy. This is an optimization, only to be used if you _know_ that all CVars set with this tag were set at this priority. Any other priorites cannot be unset
 	 */
 	virtual void UnsetAllConsoleVariablesWithTag(FName Tag, EConsoleVariableFlags Priority=ECVF_SetByMask) = 0;
-	
+
+	/**
+	* Add, replace, or remove cvars to match the passed in cvars
+	*/
+	virtual void BatchUpdateTag(FName Tag, const TMap<FName, FString>& CVarsAndValues) = 0;
+
 	virtual FConsoleVariableMulticastDelegate& OnCVarUnregistered() = 0;
 	virtual FConsoleObjectWithNameMulticastDelegate& OnConsoleObjectUnregistered() = 0;
 
@@ -1684,14 +1692,9 @@ template <class T>
 class TAutoConsoleVariable : public FAutoConsoleObject
 {
 public:
-	/**
-	 * Create a float, int or string console variable
-	 * @param Name must not be 0
-	 * @param Help must not be 0
-	 * @param Flags bitmask combined from EConsoleVariableFlags
-	 */
-	TAutoConsoleVariable(const TCHAR* Name, const T& DefaultValue, const TCHAR* Help, uint32 Flags = ECVF_Default)
-		: FAutoConsoleObject(IConsoleManager::Get().RegisterConsoleVariable(Name, DefaultValue, Help, Flags)) 
+
+	TAutoConsoleVariable(IConsoleVariable* ExistingObject)
+		: FAutoConsoleObject(ExistingObject)
 	{
 		if constexpr (std::is_same_v<T, bool>)
 		{
@@ -1715,6 +1718,17 @@ public:
 		}
 	}
 
+	/**
+	 * Create a float, int or string console variable
+	 * @param Name must not be 0
+	 * @param Help must not be 0
+	 * @param Flags bitmask combined from EConsoleVariableFlags
+	 */
+	TAutoConsoleVariable(const TCHAR* Name, const T& DefaultValue, const TCHAR* Help, uint32 Flags = ECVF_Default)
+		: TAutoConsoleVariable(IConsoleManager::Get().RegisterConsoleVariable(Name, DefaultValue, Help, Flags))
+	{
+	}
+
 	/**
 	 * Create a float, int or string console variable
 	 * @param Name must not be 0
@@ -2215,5 +2229,17 @@ public:
 	CORE_API FAutoConsoleCommandDeprecated(const TCHAR* Name, const TCHAR* CVarToShadow, const TCHAR* DeprecatedAtVersion, EShadowCVarBehavior UsageBehavior = EShadowCVarBehavior::Ensure, EShadowCVarBehavior LookupFailureBehavior = EShadowCVarBehavior::Ensure);
 };
 
+
+/**
+ * A class to create a ranged console variable. The min and max can be changed with <name>.min and <name.max>
+ */
+template<typename T>
+class TAutoRangedConsoleVariable : public TAutoConsoleVariable<T>
+{
+public:
+	CORE_API TAutoRangedConsoleVariable(const TCHAR* Name, T DefaultValue, T MinValue, T MaxValue, const TCHAR* Help, uint32 Flags = ECVF_Default, IConsoleVariable* Parent = nullptr);
+};
+
+
 CORE_API DECLARE_LOG_CATEGORY_EXTERN(LogConsoleResponse, Log, All);
 
diff --git a/Engine/Source/Runtime/Core/Public/Misc/ConfigUtilities.h b/Engine/Source/Runtime/Core/Public/Misc/ConfigUtilities.h
--- a/Engine/Source/Runtime/Core/Public/Misc/ConfigUtilities.h
+++ b/Engine/Source/Runtime/Core/Public/Misc/ConfigUtilities.h
@@ -27,6 +27,9 @@ namespace UE::DynamicConfig
 	// this calls the UObjectBaseUtility from code where object system is not linked (it also calls TSOnConfigSectionsChanged()!)
 	extern CORE_API TMulticastDelegate<void(const FConfigModificationTracker* ChangeTracker)> ReloadObjects;
 
+	extern CORE_API TMulticastDelegate<void(const FConfigModificationTracker* ChangeTracker)> UpdateCVarsAndDeviceProfiles;
+	
+	UE_DEPRECATED(5.6, "Use UpdateCVarsAndDeviceProfiles");
 	extern CORE_API TMulticastDelegate<void(const TSet<FString>& ModifiedSections)> UpdateDeviceProfiles;
 }
 
diff --git a/Engine/Source/Runtime/Engine/Classes/DeviceProfiles/DeviceProfileManager.h b/Engine/Source/Runtime/Engine/Classes/DeviceProfiles/DeviceProfileManager.h
--- a/Engine/Source/Runtime/Engine/Classes/DeviceProfiles/DeviceProfileManager.h
+++ b/Engine/Source/Runtime/Engine/Classes/DeviceProfiles/DeviceProfileManager.h
@@ -52,9 +52,9 @@ public:
 
 	/**
 	 * Reapplies the device profile. Useful when configs have changed (i.e. hotfix)
-	 * Applies base and then any overridden device profile.
+	 * Applies base and then any overridden device profile. If ModifiedCVars is specified, then only the given CVars will be updated
 	 */
-	ENGINE_API void ReapplyDeviceProfile(bool bForceReload = false);
+	ENGINE_API void ReapplyDeviceProfile(bool bForceReload = false, TSet<FString>* ModifiedCVars=nullptr);
 
 	/**
 	 * Examine the currently active or overridden profile for references to entries in DeviceProfilesToQuery
@@ -243,6 +243,8 @@ public:
 		DPM_SetCVars,
 		DPM_CacheValues,
 		DPM_CacheValuesIgnoreMatchingRules, // GatherDeviceProfileCVars will not return any matching rules cvars. Used when entire matching rules state is applied later (see GetAllReferencedDeviceProfileCVars).
+		DPM_SetCVarsForSingleDP,
+		DPM_SetCVarsForSingleFragment,
 	};
 
 	/**
@@ -255,6 +257,12 @@ public:
 	 */
 	static ENGINE_API TMap<FName, TSet<FString>> GetAllReferencedDeviceProfileCVars(UDeviceProfile* DeviceProfile);
 
+	/**
+	 * Gets the parent chain and fragment that make up the final DP settings
+	 * @return a list of pairs of DP/fragment name, along with the section name to look up in ConfigSystem
+	 */
+	TArray<TTuple<FName, FString>> GetDPChain(const FString& DeviceProfileName, FConfigCacheIni* ConfigSystem);
+
 private:
 	/**
 	 * Set the active device profile - set via the device profile blueprint.
diff --git a/Engine/Source/Runtime/Engine/Private/DeviceProfiles/DeviceProfileManager.cpp b/Engine/Source/Runtime/Engine/Private/DeviceProfiles/DeviceProfileManager.cpp
--- a/Engine/Source/Runtime/Engine/Private/DeviceProfiles/DeviceProfileManager.cpp
+++ b/Engine/Source/Runtime/Engine/Private/DeviceProfiles/DeviceProfileManager.cpp
@@ -55,6 +55,18 @@ UDeviceProfileManager* UDeviceProfileManager::DeviceProfileManagerSingleton = nu
 // didnt't create it along the way
 FDelayedAutoRegisterHelper GDPManagerSingletonHelper(EDelayedRegisterRunPhase::ObjectSystemReady, [] { UDeviceProfileManager::Get(); });
 
+static int32 CVarSetMode()
+{
+	static int32 Mode = -1;
+	if (Mode == -1)
+	{
+		Mode = 0;
+		// set the CVarSetMode from commandline
+		FParse::Value(FCommandLine::Get(), TEXT("DPApplyMode="), Mode);
+	}
+	return Mode;
+}
+
 UDeviceProfileManager& UDeviceProfileManager::Get(bool bFromPostCDOContruct)
 {
 	if (DeviceProfileManagerSingleton == nullptr)
@@ -101,22 +113,163 @@ UDeviceProfileManager& UDeviceProfileManager::Get(bool bFromPostCDOContruct)
 		FDelayedAutoRegisterHelper::RunAndClearDelayedAutoRegisterDelegates(EDelayedRegisterRunPhase::DeviceProfileManagerReady);
 		
 		// when we load/unload dynamic configs, we may need to propagate those changes to the active DP
-		UE::DynamicConfig::UpdateDeviceProfiles.AddLambda([](const TSet<FString>& ModifiedSections)
+		void UpdateCVarsAfterDynamicConfig(const FConfigModificationTracker* ChangeTracker);
+		UE::DynamicConfig::UpdateCVarsAndDeviceProfiles.AddStatic(UpdateCVarsAfterDynamicConfig);
+	}
+	return *DeviceProfileManagerSingleton;
+}
+
+TArray<TTuple<FName, FString>> UDeviceProfileManager::GetDPChain(const FString& DeviceProfileName, FConfigCacheIni* ConfigSystem)
+{
+	static FString DPSectionSuffix = *FString::Printf(TEXT(" %s"), *UDeviceProfile::StaticClass()->GetName());
+	static FString FragmentSectionSuffix = *FString::Printf(TEXT(" %s"), *UDeviceProfileFragment::StaticClass()->GetName());
+
+	// @todo fragments, before and after
+
+	// For each device profile, starting with the selected and working our way up the BaseProfileName tree,
+	FString BaseDeviceProfileName = DeviceProfileName;
+	bool bReachedEndOfTree = BaseDeviceProfileName.IsEmpty();
+	TArray<FString> ParentChain;
+	while (bReachedEndOfTree == false)
+	{
+		FString CurrentSectionName = BaseDeviceProfileName + DPSectionSuffix;
+
+		ParentChain.Add(BaseDeviceProfileName);
+
+		// check if there is a section named for the DeviceProfile
+		const FConfigSection* CurrentSection = ConfigSystem->GetSection(*CurrentSectionName, false, GDeviceProfilesIni);
+		if (CurrentSection != nullptr)
+		{
+			// Get the next device profile name, to look for CVars in, along the tree
+			FString NextBaseDeviceProfileName;
+			ConfigSystem->GetString(*CurrentSectionName, TEXT("BaseProfileName"), NextBaseDeviceProfileName, GDeviceProfilesIni);
+			BaseDeviceProfileName = NextBaseDeviceProfileName;
+		}
+
+		// Check if we have inevitably reached the end of the device profile tree.
+		bReachedEndOfTree = CurrentSection == nullptr || BaseDeviceProfileName.IsEmpty();
+	}
+
+	TArray<TTuple<FName, FString>> DPChain;
+
+	// now walk up the chain, and insert fragments for each level
+	for (int Index = ParentChain.Num() - 1; Index >= 0; Index--)
+	{
+		const FString& DPName = ParentChain[Index];
+		FString CurrentSectionName = DPName + DPSectionSuffix;
+
+		// includes of this DP
+		TArray<FString> PreFragmentIncludes, PostFragmentIncludes;
+		ConfigSystem->GetArray(*CurrentSectionName, TEXT("FragmentIncludes"), PreFragmentIncludes, GDeviceProfilesIni);
+		ConfigSystem->GetArray(*CurrentSectionName, TEXT("OverrideFragmentIncludes"), PostFragmentIncludes, GDeviceProfilesIni);
+
+		for (const FString& PreFragmentName : PreFragmentIncludes)
+		{
+			DPChain.Emplace(*PreFragmentName, PreFragmentName + FragmentSectionSuffix);
+		}
+		DPChain.Emplace(*DPName, CurrentSectionName);
+		for (const FString& PostFragmentName : PostFragmentIncludes)
+		{
+			DPChain.Emplace(*PostFragmentName, PostFragmentName + FragmentSectionSuffix);
+		}
+	}
+
+	// get fragments then append in reverse order
+	TArray<FSelectedFragmentProperties> DynamicFragments = FindMatchingFragments(DeviceProfileName, ConfigSystem);
+	for (int Index = DynamicFragments.Num() - 1; Index >= 0; Index--)
+	{
+		// @todo check benabled?
+		DPChain.Emplace(*DynamicFragments[Index].Fragment, DynamicFragments[Index].Fragment + FragmentSectionSuffix);
+	}
+
+	return DPChain;
+}
+
+void UpdateCVarsAfterDynamicConfig(const FConfigModificationTracker* ChangeTracker)
+{
+	// reload objects that had their configs changed
+	for (const auto& SectionPair : ChangeTracker->CVars)
+	{
+		const FConfigModificationTracker::FCVarTracker& CVars = SectionPair.Value;
+		EConsoleVariableFlags Priority = (EConsoleVariableFlags)CVars.CVarPriority;
+
+		// now walk over the updated cvars and set them based on the priority
+		for (const auto& TagPair : CVars.CVarEntriesPerBranchPerTag)
+		{
+			for (const auto& BranchPair : TagPair.Value)
 			{
-				TSet<FString> DeviceProfilesToQuery;
-				int32 DeviceProfilePatternLength = UDeviceProfile::StaticClass()->GetName().Len() + 1;
-				for (const FString& Section : ModifiedSections)
+				const FConfigSection& Section = BranchPair.Value;
+
+				for (const auto& CVarPair : Section)
 				{
-					DeviceProfilesToQuery.Add(Section.LeftChop(DeviceProfilePatternLength));
+					FString KeyString = CVarPair.Key.ToString();
+
+					if (KeyString.StartsWith(TEXT("sg.")))
+					{
+						UE::ConfigUtilities::ApplyCVarSettingsFromIni(*FString::Printf(TEXT("%s@%s"), *KeyString.Mid(3), *CVarPair.Value.GetValue()), *GScalabilityIni, Priority, false, TagPair.Key);
+					}
+					else
+					{
+						UE::ConfigUtilities::OnSetCVarFromIniEntry(TEXT("DynamicLayer"), *CVarPair.Key.ToString(), *CVarPair.Value.GetValue(), Priority, false, false, TagPair.Key);
+					}
+				}
+			}
+		}
+	}
+
+	if (CVarSetMode() > 0)
+	{
+		const TSet<FString>* DeviceProfileSections = ChangeTracker->ModifiedSectionsPerBranch.Find(*GDeviceProfilesIni);
+		if (DeviceProfileSections != nullptr)
+		{
+			UDeviceProfileManager& DPManager = UDeviceProfileManager::Get();
+			IConsoleManager& ConsoleManager = IConsoleManager::Get();
+
+			// get the set of sections the active DP cares about
+			TArray<TTuple<FName, FString>> DPChain = UDeviceProfileManager::Get().GetDPChain(DPManager.GetActiveDeviceProfileName(), GConfig);
+
+			for (const auto& DP : DPChain)
+			{
+				// mode 1 we wipe everything then start over instead of updating in place
+				if (CVarSetMode() == 1)
+				{
+					ConsoleManager.UnsetAllConsoleVariablesWithTag(DP.Get<0>());
 				}
 
-				if (UDeviceProfileManager::Get().DoActiveProfilesReference(DeviceProfilesToQuery))
+				// if we updated this DP, then refresh it
+				const FName& DPName = DP.Get<0>();
+				const FString& DPSectionName = DP.Get<1>();
+				if (DeviceProfileSections->Contains(DPSectionName))
 				{
-					UDeviceProfileManager::Get().ReapplyDeviceProfile();
+					// get just this DP's cvars, and then update in place
+					TMap<FName, FString> DeviceProfileCVars = DPManager.GatherDeviceProfileCVars(DPName.ToString(), UDeviceProfileManager::EDeviceProfileMode::DPM_SetCVarsForSingleDP);
+					ConsoleManager.BatchUpdateTag(DPName, DeviceProfileCVars);
 				}
-			});
+			}
+		}
+	}
+
+	// if we modify a scalability group section (ie EffectsQuality@2) that matches current value, re-set the group to apply the modified group
+	const TSet<FString>* ScalabilitySections = ChangeTracker->ModifiedSectionsPerBranch.Find(*GScalabilityIni);
+	if (ScalabilitySections != nullptr)
+	{
+		for (const FString& SectionName : *ScalabilitySections)
+		{
+			FString Group, Level;
+			if (SectionName.Split(TEXT("@"), &Group, &Level))
+			{
+				FString SGName = FString::Printf(TEXT("sg.%s"), *Group);
+				IConsoleVariable* SGVar = IConsoleManager::Get().FindConsoleVariable(*SGName);
+				if (SGVar != nullptr)
+				{
+					if (SGVar->GetString() == Level)
+					{
+						SGVar->SetWithCurrentPriority(*Level);
+					}
+				}
+			}
+		}
 	}
-	return *DeviceProfileManagerSingleton;
 }
 
 // Read the cvars from a [DeviceProfileFragment] section.
@@ -180,6 +333,7 @@ TMap<FName, FString> UDeviceProfileManager::GatherDeviceProfileCVars(const FStri
 	TArray<FSelectedFragmentProperties> FragmentsSelected;
 	TSet<FString> FragmentCVarKeys;
 	TArray<FString> SelectedFragmentCVars;
+	bool bSingleSectionMode = GatherMode == EDeviceProfileMode::DPM_SetCVarsForSingleDP || GatherMode == EDeviceProfileMode::DPM_SetCVarsForSingleFragment;
 
 
 	EPlatformMemorySizeBucket MemBucket = FPlatformMemory::GetMemorySizeBucket();
@@ -208,7 +362,8 @@ TMap<FName, FString> UDeviceProfileManager::GatherDeviceProfileCVars(const FStri
 		checkNoEntry();
 #endif
 	}
-	else
+	// no fragments when doing single at a time
+	else if (!bSingleSectionMode)
 	{
 		// the very first time, this is called to set cvars, we cache the current device's fragments
 		if (PlatformFragmentsSelected.Num() == 0)
@@ -232,32 +387,35 @@ TMap<FName, FString> UDeviceProfileManager::GatherDeviceProfileCVars(const FStri
 	AvailableProfiles.Remove(TEXT("DeviceProfiles"));
 
 
-	// Here we gather the cvars from selected fragments in reverse order
-	for (int i = FragmentsSelected.Num() - 1; i >= 0; i--)
+	if (!bSingleSectionMode)
 	{
-		const FSelectedFragmentProperties& SelectedFragment = FragmentsSelected[i];
-		if (SelectedFragment.bEnabled)
+		// Here we gather the cvars from selected fragments in reverse order
+		for (int i = FragmentsSelected.Num() - 1; i >= 0; i--)
 		{
-			TArray<FString> FragmentCVars;
-			GetFragmentCVars(*SelectedFragment.Fragment, TEXT("CVars"), FragmentCVars, ConfigSystem);
-			for (int j = FragmentCVars.Num() - 1; j >= 0; j--)
+			const FSelectedFragmentProperties& SelectedFragment = FragmentsSelected[i];
+			if (SelectedFragment.bEnabled)
 			{
-				const FString& FragCVar = FragmentCVars[j];
-
-				FString CVarKey, CVarValue;
-				if (FragCVar.Split(TEXT("="), &CVarKey, &CVarValue))
+				TArray<FString> FragmentCVars;
+				GetFragmentCVars(*SelectedFragment.Fragment, TEXT("CVars"), FragmentCVars, ConfigSystem);
+				for (int j = FragmentCVars.Num() - 1; j >= 0; j--)
 				{
-					if (!FragmentCVarKeys.Find(CVarKey))
+					const FString& FragCVar = FragmentCVars[j];
+
+					FString CVarKey, CVarValue;
+					if (FragCVar.Split(TEXT("="), &CVarKey, &CVarValue))
 					{
-						FragmentCVarKeys.Add(CVarKey);
-						SelectedFragmentCVars.Add(FragCVar);
+						if (!FragmentCVarKeys.Find(CVarKey))
+						{
+							FragmentCVarKeys.Add(CVarKey);
+							SelectedFragmentCVars.Add(FragCVar);
+						}
 					}
 				}
 			}
 		}
 	}
 
-	static FString SectionSuffix = *FString::Printf(TEXT(" %s"), *UDeviceProfile::StaticClass()->GetName());
+	const TCHAR* SectionSuffix = GatherMode == EDeviceProfileMode::DPM_SetCVarsForSingleFragment ? TEXT(" DeviceProfileFragment") : TEXT(" DeviceProfile");
 
 	// For each device profile, starting with the selected and working our way up the BaseProfileName tree,
 	FString BaseDeviceProfileName = DeviceProfileName;
@@ -290,8 +448,11 @@ TMap<FName, FString> UDeviceProfileManager::GatherDeviceProfileCVars(const FStri
 				}
 
 				TArray< FString > CurrentProfilesCVars, FragmentCVars, OverrideFragmentCVars;
-				GetCVarsFromDPFragmentIncludes(*CurrentSectionName, *ArrayName, FragmentCVars, ConfigSystem, false);
-				GetCVarsFromDPFragmentIncludes(*CurrentSectionName, *ArrayName, OverrideFragmentCVars, ConfigSystem, true);
+				if (!bSingleSectionMode)
+				{
+					GetCVarsFromDPFragmentIncludes(*CurrentSectionName, *ArrayName, FragmentCVars, ConfigSystem, false);
+					GetCVarsFromDPFragmentIncludes(*CurrentSectionName, *ArrayName, OverrideFragmentCVars, ConfigSystem, true);
+				}
 				CurrentSection->MultiFind(*ArrayName, CurrentProfilesCVars, true);
 
 				if (FragmentCVars.Num())
@@ -349,7 +510,7 @@ TMap<FName, FString> UDeviceProfileManager::GatherDeviceProfileCVars(const FStri
 
 			// Get the next device profile name, to look for CVars in, along the tree
 			FString NextBaseDeviceProfileName;
-			if (ConfigSystem->GetString(*CurrentSectionName, TEXT("BaseProfileName"), NextBaseDeviceProfileName, GDeviceProfilesIni))
+			if (!bSingleSectionMode && ConfigSystem->GetString(*CurrentSectionName, TEXT("BaseProfileName"), NextBaseDeviceProfileName, GDeviceProfilesIni))
 			{
 				BaseDeviceProfileName = NextBaseDeviceProfileName;
 				UE_LOG(LogDeviceProfileManager, Log, TEXT("Going up to parent DeviceProfile [%s]"), *BaseDeviceProfileName);
@@ -361,7 +522,7 @@ TMap<FName, FString> UDeviceProfileManager::GatherDeviceProfileCVars(const FStri
 		}
 
 		// Check if we have inevitably reached the end of the device profile tree.
-		bReachedEndOfTree = CurrentSection == nullptr || BaseDeviceProfileName.IsEmpty();
+		bReachedEndOfTree = bSingleSectionMode || CurrentSection == nullptr || BaseDeviceProfileName.IsEmpty();
 	}
 
 	return DeviceProfileCVars;
@@ -420,61 +581,100 @@ TMap<FName, TSet<FString>> UDeviceProfileManager::GetAllReferencedDeviceProfileC
 
 void UDeviceProfileManager::SetDeviceProfileCVars(const FString& DeviceProfileName)
 {
-	// walk over the parent chain, gathering the cvars this DP inherits and contains
-	TMap<FName, FString> DeviceProfileCVars = GatherDeviceProfileCVars(DeviceProfileName, EDeviceProfileMode::DPM_SetCVars);
+	if (CVarSetMode() > 0)
+	{
+		FConfigCacheIni* ConfigSystem = GConfig;
+		TArray<TTuple<FName, FString>> DPChain = UDeviceProfileManager::Get().GetDPChain(DeviceProfileName, ConfigSystem);
 
-	// reset some global state
-	DeviceProfileScalabilityCVars.Empty();
+		// now that we have the ordered chain, walk it and set cvars
+		for (const TTuple<FName, FString>& DP : DPChain)
+		{
+			FName DPName = DP.Get<0>();
+			const FString& SectionName = DP.Get<1>();
+			EDeviceProfileMode GatherMode = EDeviceProfileMode::DPM_SetCVarsForSingleDP;
+			if (SectionName.EndsWith(UDeviceProfileFragment::StaticClass()->GetName()))
+			{
+				GatherMode = EDeviceProfileMode::DPM_SetCVarsForSingleFragment;
+			}
 
-	// we should have always popped away old values by the time we get here
-	check(PushedSettings.Num() == 0);
+			// get the cvars for the DP or the Fragment
+			TMap<FName, FString> DeviceProfileCVars = GatherDeviceProfileCVars(DPName.ToString(), GatherMode);
 
-	// Preload a cvar we rely on in the loop below
-	if (const FConfigSection* Section = GConfig->GetSection(TEXT("ConsoleVariables"), false, *GEngineIni))
-	{
-		static FName AllowScalabilityAtRuntimeName = TEXT("dp.AllowScalabilityGroupsToChangeAtRuntime");
-		if (const FConfigValue* Value = Section->Find(AllowScalabilityAtRuntimeName))
-		{
-			const FString& KeyString = AllowScalabilityAtRuntimeName.ToString();
-			const FString& ValueString = Value->GetValue();
-			UE::ConfigUtilities::OnSetCVarFromIniEntry(*GEngineIni, *KeyString, *ValueString, ECVF_SetBySystemSettingsIni);
+			for (const TPair<FName, FString>& Pair : DeviceProfileCVars)
+			{
+				const FString CVarKey = Pair.Key.ToString();
+				const FString& CVarValue = Pair.Value;
+
+				// Cache any scalability related cvars so we can conveniently reapply them later as a way to reset the device defaults
+				bool bIsScalabilityBucket = CVarKey.StartsWith(TEXT("sg."));
+				//if (bIsScalabilityBucket && CVarAllowScalabilityGroupsToChangeAtRuntime.GetValueOnGameThread() > 0)
+				//{
+				//	DeviceProfileScalabilityCVars.Add(*CVarKey, *CVarValue);
+				//}
+
+				// Set by scalability or DP, depending
+				uint32 CVarPriority = bIsScalabilityBucket ? ECVF_SetByScalability : ECVF_SetByDeviceProfile;
+				UE::ConfigUtilities::OnSetCVarFromIniEntry(*GDeviceProfilesIni, *CVarKey, *CVarValue, CVarPriority, false, true, DPName);
+			}
 		}
 	}
-
-	// for each DP cvar - push the old value to restore, cache ScalabiliutyGroups, and then set the values
-	for (TPair<FName, FString>& Pair : DeviceProfileCVars)
+	else
 	{
-		FString CVarKey = Pair.Key.ToString();
-		const FString& CVarValue = Pair.Value;
+		// walk over the parent chain, gathering the cvars this DP inherits and contains
+		TMap<FName, FString> DeviceProfileCVars = GatherDeviceProfileCVars(DeviceProfileName, EDeviceProfileMode::DPM_SetCVars);
 
-		// get the actual cvar object
-		IConsoleVariable* CVar = IConsoleManager::Get().FindConsoleVariable(*CVarKey);
-		if (CVar)
-		{
-			// remember the previous value and priority
-			FPushedCVarSetting OldSetting = FPushedCVarSetting(CVar->GetString(), CVar->GetFlags());
-			PushedSettings.Add(CVarKey, OldSetting);
+		// reset some global state
+		DeviceProfileScalabilityCVars.Empty();
 
-			// indicate we are pushing, not setting
-			UE_LOG(LogDeviceProfileManager, Log, TEXT("Pushing Device Profile CVar: [[%s:%s -> %s]]"), *CVarKey, *OldSetting.Value, *CVarValue);
-		}
-		else
+		// we should have always popped away old values by the time we get here
+		check(PushedSettings.Num() == 0);
+
+		// Preload a cvar we rely on in the loop below
+		if (const FConfigSection* Section = GConfig->GetSection(TEXT("ConsoleVariables"), false, *GEngineIni))
 		{
-			UE_LOG(LogDeviceProfileManager, Log, TEXT("Creating unregistered Device Profile CVar: [[%s:%s]]"), *CVarKey, *Pair.Value);
+			static FName AllowScalabilityAtRuntimeName = TEXT("dp.AllowScalabilityGroupsToChangeAtRuntime");
+			if (const FConfigValue* Value = Section->Find(AllowScalabilityAtRuntimeName))
+			{
+				const FString& KeyString = AllowScalabilityAtRuntimeName.ToString();
+				const FString& ValueString = Value->GetValue();
+				UE::ConfigUtilities::OnSetCVarFromIniEntry(*GEngineIni, *KeyString, *ValueString, ECVF_SetBySystemSettingsIni);
+			}
 		}
 
-		// Cache any scalability related cvars so we can conveniently reapply them later as a way to reset the device defaults
-		bool bIsScalabilityBucket = CVarKey.StartsWith(TEXT("sg."));
-		if (bIsScalabilityBucket && CVarAllowScalabilityGroupsToChangeAtRuntime.GetValueOnGameThread() > 0)
+		// for each DP cvar - push the old value to restore, cache ScalabiliutyGroups, and then set the values
+		for (TPair<FName, FString>& Pair : DeviceProfileCVars)
 		{
-			DeviceProfileScalabilityCVars.Add(*CVarKey, *CVarValue);
-		}
+			const FString CVarKey = Pair.Key.ToString();
+			const FString& CVarValue = Pair.Value;
 
-		// Set by scalability or DP, depending
-		uint32 CVarPriority = bIsScalabilityBucket ? ECVF_SetByScalability : ECVF_SetByDeviceProfile;
-		UE::ConfigUtilities::OnSetCVarFromIniEntry(*GDeviceProfilesIni, *CVarKey, *CVarValue, CVarPriority, false, true);
-	}
+			// get the actual cvar object
+			IConsoleVariable* CVar = IConsoleManager::Get().FindConsoleVariable(*CVarKey);
+			if (CVar)
+			{
+				// remember the previous value and priority
+				FPushedCVarSetting OldSetting = FPushedCVarSetting(CVar->GetString(), CVar->GetFlags());
+				PushedSettings.Add(CVarKey, OldSetting);
 
+				// indicate we are pushing, not setting
+				UE_LOG(LogDeviceProfileManager, Log, TEXT("Pushing Device Profile CVar: [[%s:%s -> %s]]"), *CVarKey, *OldSetting.Value, *CVarValue);
+			}
+			else
+			{
+				UE_LOG(LogDeviceProfileManager, Log, TEXT("Creating unregistered Device Profile CVar: [[%s:%s]]"), *CVarKey, *Pair.Value);
+			}
+
+			// Cache any scalability related cvars so we can conveniently reapply them later as a way to reset the device defaults
+			bool bIsScalabilityBucket = CVarKey.StartsWith(TEXT("sg."));
+			if (bIsScalabilityBucket && CVarAllowScalabilityGroupsToChangeAtRuntime.GetValueOnGameThread() > 0)
+			{
+				DeviceProfileScalabilityCVars.Add(*CVarKey, *CVarValue);
+			}
+
+			// Set by scalability or DP, depending
+			uint32 CVarPriority = bIsScalabilityBucket ? ECVF_SetByScalability : ECVF_SetByDeviceProfile;
+			UE::ConfigUtilities::OnSetCVarFromIniEntry(*GDeviceProfilesIni, *CVarKey, *CVarValue, CVarPriority, false, true);
+		}
+	}
 
 	// now allow for ovverrides in !SHIPPING
 #if !UE_BUILD_SHIPPING
@@ -541,31 +741,43 @@ void UDeviceProfileManager::SetDeviceProfileCVars(const FString& DeviceProfileNa
 */
 static void RestorePushedState(TMap<FString, FPushedCVarSetting>& PushedSettings)
 {
-	// restore pushed settings
-	for (TMap<FString, FPushedCVarSetting>::TIterator It(PushedSettings); It; ++It)
+	if (CVarSetMode() > 0)
 	{
-		IConsoleVariable* CVar = IConsoleManager::Get().FindConsoleVariable(*It.Key());
-		if (CVar)
+		IConsoleManager& ConsoleManager = IConsoleManager::Get();
+		TArray<TTuple<FName, FString>> DPChain = UDeviceProfileManager::Get().GetDPChain(UDeviceProfileManager::Get().GetActiveDeviceProfileName(), GConfig);
+		for (const auto& DP : DPChain)
 		{
-			const FPushedCVarSetting& Setting = It.Value();
-
-			// Check if the priority has increased since we pushed it and is now higher than SetByDeviceProfile
-			EConsoleVariableFlags CurrentPriority = EConsoleVariableFlags(int32(CVar->GetFlags()) & (int32)ECVF_SetByMask);
-			if (CurrentPriority > ECVF_SetByDeviceProfile && CurrentPriority > Setting.SetBy )
-			{
-				UE_LOG(LogDeviceProfileManager, Warning, TEXT("Popping Device Profile CVar skipped because priority has been overridden to higher than ECVF_SetByDeviceProfile since the last push [[%s:%s]]"), *It.Key(), *Setting.Value);
-			}
-			else
+			ConsoleManager.UnsetAllConsoleVariablesWithTag(DP.Get<0>());
+		}
+	}
+	else
+	{
+		// restore pushed settings
+		for (TMap<FString, FPushedCVarSetting>::TIterator It(PushedSettings); It; ++It)
+		{
+			IConsoleVariable* CVar = IConsoleManager::Get().FindConsoleVariable(*It.Key());
+			if (CVar)
 			{
-				// restore it!
-				CVar->SetWithCurrentPriority(*Setting.Value);
-				UE_LOG(LogDeviceProfileManager, Log, TEXT("Popping Device Profile CVar: [[%s:%s]]"), *It.Key(), *Setting.Value);
-			}
+				const FPushedCVarSetting& Setting = It.Value();
 
+				// Check if the priority has increased since we pushed it and is now higher than SetByDeviceProfile
+				EConsoleVariableFlags CurrentPriority = EConsoleVariableFlags(int32(CVar->GetFlags()) & (int32)ECVF_SetByMask);
+				if (CurrentPriority > ECVF_SetByDeviceProfile && CurrentPriority > Setting.SetBy )
+				{
+					UE_LOG(LogDeviceProfileManager, Warning, TEXT("Popping Device Profile CVar skipped because priority has been overridden to higher than ECVF_SetByDeviceProfile since the last push [[%s:%s]]"), *It.Key(), *Setting.Value);
+				}
+				else
+				{
+					// restore it!
+					CVar->SetWithCurrentPriority(*Setting.Value);
+					UE_LOG(LogDeviceProfileManager, Log, TEXT("Popping Device Profile CVar: [[%s:%s]]"), *It.Key(), *Setting.Value);
+				}
+
+			}
 		}
-	}
 
-	PushedSettings.Reset();
+		PushedSettings.Reset();
+	}
 }
 
 const FSelectedFragmentProperties* UDeviceProfileManager::GetActiveDeviceProfileFragmentByTag(FName& FragmentTag) const
@@ -670,14 +882,35 @@ bool UDeviceProfileManager::DoActiveProfilesReference(const TSet<FString>& Devic
 	return bDoActiveProfilesReferenceQuerySet;
 }
 
-void UDeviceProfileManager::ReapplyDeviceProfile(bool bForceReload)
+void UDeviceProfileManager::ReapplyDeviceProfile(bool bForceReload, TSet<FString>* ModifiedCVars)
 {	
 	UDeviceProfile* OverrideProfile = DeviceProfileManagerSingleton->BaseDeviceProfile ? DeviceProfileManagerSingleton->GetActiveProfile() : nullptr;
 	UDeviceProfile* BaseProfile = DeviceProfileManagerSingleton->BaseDeviceProfile ? DeviceProfileManagerSingleton->BaseDeviceProfile : DeviceProfileManagerSingleton->GetActiveProfile();
 
 	UE_LOG(LogDeviceProfileManager, Log, TEXT("ReapplyDeviceProfile applying profile: [%s]"), *BaseProfile->GetName(), OverrideProfile ? *OverrideProfile->GetName() : TEXT("not set."));
 
-	if (OverrideProfile)
+
+	if (ModifiedCVars != nullptr)
+	{
+		// make sure any modified cvar is in the pushed state because we are going to set cvars without Popping first
+		if (OverrideProfile)
+		{
+			for (const FString& CVarName : *ModifiedCVars)
+			{
+				if (!PushedSettings.Contains(CVarName))
+				{
+					IConsoleVariable* CVar = IConsoleManager::Get().FindConsoleVariable(*CVarName);
+					if (CVar)
+					{
+						FPushedCVarSetting OldSetting = FPushedCVarSetting(CVar->GetString(), CVar->GetFlags());
+						PushedSettings.Add(CVarName, OldSetting);
+					}
+				}
+			}
+			SetDeviceProfileCVars(DeviceProfileManagerSingleton->GetActiveProfile()->GetName());
+		}
+	}
+	else if (OverrideProfile)
 	{
 		UE_LOG(LogDeviceProfileManager, Log, TEXT("ReapplyDeviceProfile applying override profile: [%s]"), *OverrideProfile->GetName());
 		// reapply the override.
