PR #12364: Providing FVector2f overload of UStaticMesh::SetUVChannel

diff --git a/Engine/Source/Runtime/Engine/Classes/Engine/StaticMesh.h b/Engine/Source/Runtime/Engine/Classes/Engine/StaticMesh.h
--- a/Engine/Source/Runtime/Engine/Classes/Engine/StaticMesh.h
+++ b/Engine/Source/Runtime/Engine/Classes/Engine/StaticMesh.h
@@ -1560,6 +1560,15 @@ public:
 	 */
 	ENGINE_API bool SetUVChannel(int32 LODIndex, int32 UVChannelIndex, const TMap<FVertexInstanceID, FVector2D>& TexCoords);
 
+	/**
+	 * Sets the texture coordinates at the specified UV channel index on the given LOD of a StaticMesh.
+	 * @param	LODIndex			Index of the StaticMesh LOD.
+	 * @param	UVChannelIndex		Index where to remove the UV channel.
+	 * @param	TexCoords			The texture coordinates to set on the UV channel.
+	 * @return true if the UV channel could be set.
+	 */
+	ENGINE_API bool SetUVChannel(int32 LODIndex, int32 UVChannelIndex, const TMap<FVertexInstanceID, FVector2f>& TexCoords);
+
 #endif
 
 	/** Create an empty StaticMeshDescription object, to describe a static mesh at runtime */
diff --git a/Engine/Source/Runtime/Engine/Private/StaticMesh.cpp b/Engine/Source/Runtime/Engine/Private/StaticMesh.cpp
--- a/Engine/Source/Runtime/Engine/Private/StaticMesh.cpp
+++ b/Engine/Source/Runtime/Engine/Private/StaticMesh.cpp
@@ -6383,6 +6383,42 @@ bool UStaticMesh::SetUVChannel(int32 LODIndex, int32 UVChannelIndex, const TMap<
 	return true;
 }
 
+bool UStaticMesh::SetUVChannel(int32 LODIndex, int32 UVChannelIndex, const TMap<FVertexInstanceID, FVector2f>& TexCoords)
+{
+	FMeshDescription* MeshDescription = GetMeshDescription(LODIndex);
+	if (!MeshDescription)
+	{
+		return false;
+	}
+
+	if (TexCoords.Num() < MeshDescription->VertexInstances().Num())
+	{
+		return false;
+	}
+
+	Modify();
+
+	FStaticMeshAttributes Attributes(*MeshDescription);
+
+	TMeshAttributesRef<FVertexInstanceID, FVector2f> UVs = Attributes.GetVertexInstanceUVs();
+	for (const FVertexInstanceID VertexInstanceID : MeshDescription->VertexInstances().GetElementIDs())
+	{
+		if (const FVector2f* UVCoord = TexCoords.Find(VertexInstanceID))
+		{
+			UVs.Set(VertexInstanceID, UVChannelIndex, *UVCoord);
+		}
+		else
+		{
+			ensureMsgf(false, TEXT("Tried to apply UV data that did not match the StaticMesh MeshDescription."));
+		}
+	}
+
+	CommitMeshDescription(LODIndex);
+	PostEditChange();
+
+	return true;
+}
+
 #endif
 
 int32 UStaticMesh::GetNumUVChannels(int32 LODIndex) const
