[NetPhysics] PhysicsReplicationLOD improvements
- Use AABB of all interacting LOD particles and the Focal particles when checking distance.
- Cache the LOD data when particles interact and return that when the other interacting particles query their LOD data.

Also extend and optimized the AABB functionality:
- A performant Clear() function
- MoveByVector() function that moves the AABB along a vector without growing or shrinking.
- Optimize the Center function by multiplying by 0.5 instead of dividing by 2

diff --git a/Engine/Source/Runtime/Engine/Classes/PhysicsEngine/PhysicsSettings.h b/Engine/Source/Runtime/Engine/Classes/PhysicsEngine/PhysicsSettings.h
--- a/Engine/Source/Runtime/Engine/Classes/PhysicsEngine/PhysicsSettings.h
+++ b/Engine/Source/Runtime/Engine/Classes/PhysicsEngine/PhysicsSettings.h
@@ -145,30 +145,45 @@ struct FPhysicsReplicationLODSettings
 	UPROPERTY(config, EditDefaultsOnly, Category = "Replication")
 	bool bEnablePhysicsReplicationLOD;
 
-	/** Within this distance from Focal Points in LOD, physics replicated objects will use EPhysicsReplicationMode::Resimulation
-	* Resimulation is fully forward predicted physics replication, for a seamless LOD transition this should be equal to or lower than DistanceForFullPrediction, read that property for information. */
-	UPROPERTY(EditAnywhere, Category = "Replication")
-	float DistanceForResimulationMode;
-
-	/** Within this distance from Focal Points in LOD, physics replicated objects will be in the clients forward predicted timeline
-	* The local player is typically forward predicted, while physics replicated objects are replicated without forward prediction.
+	/** Minimum distance in the physics replication LODs Base Distance. */
+	UPROPERTY(EditAnywhere, Category = "Replication", meta = (ClampMin = "0.0", UIMin = "0.0"))
+	float MinimumBaseDistance;
+
+	/** Multiplier for adding focal particles bounding box radius to the physics replication LODs Base Distance. Value of 0.75 adds 75% of the radius on top of MinimumBaseDistance. */
+	UPROPERTY(EditAnywhere, Category = "Replication", meta = (ClampMin = "0.0", UIMin = "0.0"))
+	float BaseDistanceRadiusMultiplier;
+
+	/** 'BaseDistance' is a result of MinimumBaseDistance + (the focal particles radius * BaseDistanceRadiusMultiplier)
+	* Set how many BaseDistance's from the focal particle where within physics replicated objects will use EPhysicsReplicationMode::Resimulation.
+	* Resimulation is a fully forward predicted physics replication, for a seamless LOD transition this should be equal to or lower than BaseDistancesForFullPrediction, read that property for information.
+	* NOTE: The resulting distance will be added on top of the focal particles radius. */
+	UPROPERTY(EditAnywhere, Category = "Replication", meta = (ClampMin = "0.0", UIMin = "0.0"))
+	float BaseDistancesForResimulationMode;
+
+	/** 'BaseDistance' is a result of MinimumBaseDistance + (the focal particles radius * BaseDistanceRadiusMultiplier)
+	* Set how many BaseDistance's from the focal particle where within physics replicated objects will be in the clients forward predicted timeline.
 	* Physics replication LOD gradually forward predicts the physics replication when objects get closer to focal points in LOD.
-	* Read the TimeOverDistance property for details on how forward prediction is gradually applied up until the distance of DistanceForFullPrediction. */
-	UPROPERTY(EditAnywhere, Category = "Replication")
-	float DistanceForFullPrediction;
+	* Read the TimeOverDistance property for details on how forward prediction is gradually applied up until the distance of DistanceForFullPrediction. 
+	* NOTE: The local player is typically forward predicted (and should then be a focal particle in LOD), while other physics replicated objects are typically replicated without forward prediction.
+	* NOTE: The resulting distance will be added on top of the focal particles radius. */
+	UPROPERTY(EditAnywhere, Category = "Replication", meta = (ClampMin = "0.0", UIMin = "0.0"))
+	float BaseDistancesForFullPrediction;
+
 
 	/** Time/Distance alignment value, used as Distance * TimeOverDistance = Time
-	* Time is how far behind the current timeline replication should run
-	* Starts outside of DistanceForFullPrediction and gets clamped by the received states timeline
+	* Time is how far behind the current (forward predicted) timeline replication should run
+	* Starts outside of DistanceForFullPrediction (calculated from BaseDistancesForFullPrediction) and gets clamped by the received states timeline.
 	* Lower value = Slower timeline transition which produces a larger LOD radius.
-	* EXAMPLE: At 500cm away an object should be 37,5ms behind the clients timeline if DistanceFullPrediction is 250cm and TimeOverDistance is 0.15ms/cm. (500cm - 250cm) * 0.15ms/cm = 37,5ms */
-	UPROPERTY(EditAnywhere, Category = "Replication")
+	* EXAMPLE: At 500cm away an object should be 37,5ms behind the clients timeline if DistanceForFullPrediction is 250cm and TimeOverDistance is 0.15ms/cm. (500cm - 250cm) * 0.15ms/cm = 37,5ms */
+	UPROPERTY(EditAnywhere, Category = "Replication", meta = (ClampMin = "0.0", UIMin = "0.0"))
 	float TimeOverDistance;
 
 	FPhysicsReplicationLODSettings()
 		: bEnablePhysicsReplicationLOD(false)
-		, DistanceForResimulationMode(200.0f)
-		, DistanceForFullPrediction(250.0f)
+		, MinimumBaseDistance(200.0f)
+		, BaseDistanceRadiusMultiplier(0.75f)
+		, BaseDistancesForResimulationMode(0.25f)
+		, BaseDistancesForFullPrediction(0.8f)
 		, TimeOverDistance(0.15f)
 	{ }
 };
diff --git a/Engine/Source/Runtime/Engine/Private/PhysicsEngine/PhysicsReplicationLOD.cpp b/Engine/Source/Runtime/Engine/Private/PhysicsEngine/PhysicsReplicationLOD.cpp
--- a/Engine/Source/Runtime/Engine/Private/PhysicsEngine/PhysicsReplicationLOD.cpp
+++ b/Engine/Source/Runtime/Engine/Private/PhysicsEngine/PhysicsReplicationLOD.cpp
@@ -14,14 +14,20 @@ namespace PhysicsReplicationLODCVars
 	bool bEnabled = false;
 	static FAutoConsoleVariableRef CVarEnabled(TEXT("p.ReplicationLOD.Enabled"), bEnabled, TEXT("Override a disabled PhysicsReplicationLOD to become enabled, can't disable an LOD that is set enabled via other means like project settings or API call."));
 
-	float DistanceForResimulationMode = -1.0f;
-	static FAutoConsoleVariableRef CVarDistanceForResimulationMode(TEXT("p.ReplicationLOD.DistanceForResimulationMode"), DistanceForResimulationMode, TEXT("Override distance where within resimulation will be recommended as the replication mode by LOD."));
+	float MinimumBaseDistance = -1.0f;
+	static FAutoConsoleVariableRef CVarMinimumBaseDistance(TEXT("p.ReplicationLOD.MinimumBaseDistance"), MinimumBaseDistance, TEXT("Override the minimum distance in the physics replication LODs Base Distance. Negative value disables the override and project settings will apply."));
 
-	float DistanceForFullPrediction = -1.0f;
-	static FAutoConsoleVariableRef CVarDistanceForFullPrediction(TEXT("p.ReplicationLOD.DistanceForFullPrediction"), DistanceForFullPrediction, TEXT("Override the distance where within full forward prediction to clients current timeline will be recommended."));
+	float BaseDistanceRadiusMultiplier = -1.0f;
+	static FAutoConsoleVariableRef CVarBaseDistanceRadiusMultiplier(TEXT("p.ReplicationLOD.BaseDistanceRadiusMultiplier"), BaseDistanceRadiusMultiplier, TEXT("Override the distance radius multiplier which adds focal particles bounding box radius to the physics replication LODs Base Distance. Value of 0.75f adds 75% of the radius on top of the LOD distances. Negative value disables the override and project settings will apply."));
+
+	float BaseDistancesForResimulationMode = -1.0f;
+	static FAutoConsoleVariableRef CVarBaseDistancesForResimulationMode(TEXT("p.ReplicationLOD.BaseDistancesForResimulationMode"), BaseDistancesForResimulationMode, TEXT("Override the amount of Base Distance to set the distance for where within resimulation will be recommended as the replication mode by LOD. Negative value disables the override and project settings will apply."));
+
+	float BaseDistancesForFullPrediction = -1.0f;
+	static FAutoConsoleVariableRef CVarBaseDistancesForFullPrediction(TEXT("p.ReplicationLOD.BaseDistancesForFullPrediction"), BaseDistancesForFullPrediction, TEXT("Override the amount of Base Distance to set the distance where within full forward prediction to clients current timeline will be recommended. Negative value disables the override and project settings will apply."));
 
 	float TimeOverDistance = -1.0f;
-	static FAutoConsoleVariableRef CVaTimeOverDistance(TEXT("p.ReplicationLOD.TimeOverDistance"), TimeOverDistance, TEXT("Time/Distance LOD alignment value in milliseconds/centimeter. Used as Distance * TimeOverDistance = Time, where Time is how far behind the current timeline replication should run. Start outside of DistanceFullPrediction and clamped by the received states timeline so we don't extrapolate backwards."));
+	static FAutoConsoleVariableRef CVaTimeOverDistance(TEXT("p.ReplicationLOD.TimeOverDistance"), TimeOverDistance, TEXT("Time/Distance LOD alignment value in milliseconds/centimeter. Used as Distance * TimeOverDistance = Time, where Time is how far behind the current timeline replication should run. Start outside of DistanceFullPrediction and clamped by the received states timeline so we don't extrapolate backwards. Negative value disables the override and project settings will apply."));
 
 	bool bDrawDebugEnabled = false;
 	static FAutoConsoleVariableRef CVarDrawDebug(TEXT("p.ReplicationLOD.DrawDebug.Enabled"), bDrawDebugEnabled, TEXT(""));
@@ -150,6 +156,9 @@ void FPhysicsReplicationLODAsync::ProcessInputs_Internal(int32 PhysicsStep)
 
 	// Called here instead of in OnPreSimulate_Internal for execution order, since Physics Replication might query the LOD system inside its OnPreSimulate_Internal 
 	CacheParticlesInFocalIslands();
+
+	// Clear the cached LOD data for interacting particles without reallocating
+	CachedIslandLodData.Reset();
 }
 
 void FPhysicsReplicationLODAsync::CacheParticlesInFocalIslands()
@@ -175,12 +184,16 @@ void FPhysicsReplicationLODAsync::CacheParticlesInFocalIslands()
 		}
 
 		// Get a list of particles from the same island as a focal particle is in, i.e. particles interacting with a resim particle
-		for (const Chaos::FGeometryParticleHandle* InteractParticle : IslandManager.FindParticlesInIslands(IslandManager.FindParticleIslands(FocalParticleHandle)))
+		ParticleIslands.Reset();
+		IslandParticles.Reset();
+		IslandManager.FindParticleIslands(FocalParticleHandle, OUT ParticleIslands);
+		IslandManager.FindParticlesInIslands(ParticleIslands, OUT IslandParticles);
+		for (const Chaos::FGeometryParticleHandle* IslandParticle : IslandParticles)
 		{
-			if (InteractParticle->GetParticleType() != Chaos::EParticleType::Static)
+			if (IslandParticle->GetParticleType() != Chaos::EParticleType::Static)
 			{
 				// Add all particles that are in the same island as a focal particle to an array
-				ParticlesInFocalIslands.Add(InteractParticle->GetHandleIdx());
+				ParticlesInFocalIslands.Add(IslandParticle->GetHandleIdx());
 			}
 		}
 	}
@@ -218,10 +231,22 @@ void FPhysicsReplicationLODAsync::OnPreSimulate_Internal()
 		{
 			if (Chaos::FGeometryParticleHandle* FocalParticleHandle = Interface.GetParticle(FocalParticle))
 			{
-				FColor DebugColor = IsEnabled() ? FLinearColor::MakeFromHSV8(0, 150, 255).ToFColor(false) : FColor::Black;
-				DebugColor.A = 255;
-				const float DistanceForFullPrediction = PhysicsReplicationLODCVars::DistanceForFullPrediction > 0.0f ? PhysicsReplicationLODCVars::DistanceForFullPrediction : DefaultSettings.DistanceForFullPrediction;
-				Chaos::FDebugDrawQueue::GetInstance().DrawDebugSphere(FocalParticleHandle->GetX(), DistanceForFullPrediction, 8, DebugColor, false, -1.0f, 0, 1.0f);
+				// LOD Settings
+				const float MinimumBaseDistance = PhysicsReplicationLODCVars::MinimumBaseDistance >= 0.0f ? PhysicsReplicationLODCVars::MinimumBaseDistance : DefaultSettings.MinimumBaseDistance;
+				const float BaseDistanceRadiusMultiplier = PhysicsReplicationLODCVars::BaseDistanceRadiusMultiplier >= 0.0f ? PhysicsReplicationLODCVars::BaseDistanceRadiusMultiplier : DefaultSettings.BaseDistanceRadiusMultiplier;
+				const float BaseDistancesForFullPrediction = PhysicsReplicationLODCVars::BaseDistancesForFullPrediction >= 0.0f ? PhysicsReplicationLODCVars::BaseDistancesForFullPrediction : DefaultSettings.BaseDistancesForFullPrediction;
+				const float BaseDistancesForResimulationMode = PhysicsReplicationLODCVars::BaseDistancesForResimulationMode >= 0.0f ? PhysicsReplicationLODCVars::BaseDistancesForResimulationMode : DefaultSettings.BaseDistancesForResimulationMode;
+
+				// Focal particle specific LOD settings
+				const float BaseDistance = MinimumBaseDistance + (FocalParticleHandle->LocalBounds().CenterRadius() * BaseDistanceRadiusMultiplier);
+				const float DistanceForFullPrediction = BaseDistance * BaseDistancesForFullPrediction;
+				const float DistanceForResimulationMode = BaseDistance * BaseDistancesForResimulationMode;
+
+				const FColor DebugColorFullPrediction = IsEnabled() ? FColor::Blue : FColor::Black;
+				const FColor DebugColorResim = IsEnabled() ? FColor::Red : FColor::Black;
+
+				Chaos::FDebugDrawQueue::GetInstance().DrawDebugSphere(FocalParticleHandle->WorldSpaceInflatedBounds().Center(), FocalParticleHandle->LocalBounds().CenterRadius() + DistanceForFullPrediction, 8, DebugColorFullPrediction, false, -1.0f, 0, 1.0f);
+				Chaos::FDebugDrawQueue::GetInstance().DrawDebugSphere(FocalParticleHandle->WorldSpaceInflatedBounds().Center(), FocalParticleHandle->LocalBounds().CenterRadius() + DistanceForResimulationMode, 8, DebugColorResim, false, -1.0f, 0, 1.0f);
 			}
 		}
 	}
@@ -263,8 +288,6 @@ void FPhysicsReplicationLODAsync::UnregisterFocalPoint_Internal(Chaos::FConstPhy
 
 FPhysicsRepLodData* FPhysicsReplicationLODAsync::GetLODData_Internal(Chaos::FConstPhysicsObjectHandle PhysicsObject, const uint32 LODFlags)
 {
-	LodData.Reset();
-
 	if (FocalParticles.Num() == 0)
 	{
 		return nullptr;
@@ -285,7 +308,9 @@ FPhysicsRepLodData* FPhysicsReplicationLODAsync::GetLODData_Internal(Chaos::FCon
 			return nullptr;
 		}
 	}
-	
+
+	LodData.Reset();
+
 	if (PerformIslandLOD(*ParticleHandle, LODFlags))
 	{
 		return &LodData;
@@ -327,10 +352,69 @@ bool FPhysicsReplicationLODAsync::PerformDistanceLOD(const Chaos::FGeometryParti
 	}
 
 	Chaos::FPBDRigidsSolver& RigidsSolver = GetSolver()->CastChecked();
+	Chaos::Private::FPBDIslandManager& IslandManager = RigidsSolver.GetEvolution()->GetIslandManager();
+
+	// Check if island already has cached LOD data, if so, early out and return that.
+	ParticleIslands.Reset();
+	IslandManager.FindParticleIslands(&ParticleHandle, OUT ParticleIslands);
+	for (const Chaos::Private::FPBDIsland* Island : ParticleIslands)
+	{
+		if (FPhysicsRepLodData* CachedLodData = CachedIslandLodData.Find(Island->GetArrayIndex()))
+		{
+			// Early Out
+			LodData = *CachedLodData;
+			return LodData.DataAssigned;
+		}
+	}
+
+	// Helper function to grow AABB for with all interacting particles
+	auto IslandAABBHelper = [this](const Chaos::FGeometryParticleHandle* Particle)
+		{
+			/* NOTE: Not using Particle->WorldSpaceInflatedBounds() due to it being inflated and also variable in size between physics frames (depending on velocity and rotation?) even for uniform shapes.
+			 * Non-stable or oscillating bounds runs the risk of making replication less smooth due to the inconsistencies in the resulting LOD data. */
+
+			ParticleAABB = Particle->LocalBounds();
+			ParticleAABB.MoveByVector(Particle->GetX()); // TODO: Might need to use Particle->WorldSpaceInflatedBounds().Center() since X might not be the center of the LocalBounds
+			IslandAABB.GrowToInclude(ParticleAABB);
+		};
+
+	// Populate a bounding box of interacting particles
+	IslandAABB.Clear();
+	IslandAABBHelper(&ParticleHandle); // We do this for the particle asking for LOD Data outside of the for-loop because there are scenarios where the particle is not part of the IslandParticles array (if it's kinematic and not interacting with anything)
+	IslandParticles.Reset();
+	IslandManager.FindParticlesInIslands(ParticleIslands, OUT IslandParticles);
+	for (const Chaos::FGeometryParticleHandle* IslandParticle : IslandParticles)
+	{
+		if (IslandParticle->GetParticleType() == Chaos::EParticleType::Static)
+		{
+			continue;
+		}
+
+		if (IslandParticle->GetHandleIdx() == ParticleHandle.GetHandleIdx())
+		{
+			continue;
+		}
+
+		if (IslandParticle->GetParticleType() == Chaos::EParticleType::Kinematic || IslandParticle->ObjectState() == Chaos::EObjectStateType::Kinematic)
+		{
+			// If there are movable kinematic particles in the calling particle's island(s), also cache other islands that the kinematic particle is part of so that LOD data gets cached for all relevant islands
+			IslandManager.FindParticleIslands(IslandParticle, OUT ParticleIslands);
+		}
+
+		IslandAABBHelper(IslandParticle);
+	}
+
+	// Get the islands sphere radius and center point
+	const float IslandRadius = IslandAABB.CenterRadius();
+	const Chaos::FVec3 IslandCenter = IslandAABB.Center();
+
 	Chaos::FReadPhysicsObjectInterface_Internal Interface = Chaos::FPhysicsObjectInternalInterface::GetRead();
 
-	const float DistanceForResimulationMode = PhysicsReplicationLODCVars::DistanceForResimulationMode > 0.0f ? PhysicsReplicationLODCVars::DistanceForResimulationMode : DefaultSettings.DistanceForResimulationMode;
-	const float DistanceForFullPrediction = PhysicsReplicationLODCVars::DistanceForFullPrediction > 0.0f ? PhysicsReplicationLODCVars::DistanceForFullPrediction : DefaultSettings.DistanceForFullPrediction;
+	// LOD Settings
+	const float MinimumBaseDistance = PhysicsReplicationLODCVars::MinimumBaseDistance >= 0.0f ? PhysicsReplicationLODCVars::MinimumBaseDistance : DefaultSettings.MinimumBaseDistance;
+	const float BaseDistanceRadiusMultiplier = PhysicsReplicationLODCVars::BaseDistanceRadiusMultiplier >= 0.0f ? PhysicsReplicationLODCVars::BaseDistanceRadiusMultiplier : DefaultSettings.BaseDistanceRadiusMultiplier;
+	const float BaseDistancesForFullPrediction = PhysicsReplicationLODCVars::BaseDistancesForFullPrediction >= 0.0f ? PhysicsReplicationLODCVars::BaseDistancesForFullPrediction : DefaultSettings.BaseDistancesForFullPrediction;
+	const float BaseDistancesForResimulationMode = PhysicsReplicationLODCVars::BaseDistancesForResimulationMode >= 0.0f ? PhysicsReplicationLODCVars::BaseDistancesForResimulationMode : DefaultSettings.BaseDistancesForResimulationMode;
 	const float TimeOverDistance = PhysicsReplicationLODCVars::TimeOverDistance > 0.0f ? PhysicsReplicationLODCVars::TimeOverDistance : DefaultSettings.TimeOverDistance;
 
 	for (Chaos::FConstPhysicsObjectHandle FocalParticle : FocalParticles)
@@ -341,14 +425,29 @@ bool FPhysicsReplicationLODAsync::PerformDistanceLOD(const Chaos::FGeometryParti
 			continue;
 		}
 
-		// Check distance towards focal particles
-		const Chaos::FVec3 PosOffset = FocalParticleHandle->GetX() - ParticleHandle.GetX();
-		const float Distance = PosOffset.Size();
+		const float FocalParticleRadius = FocalParticleHandle->LocalBounds().CenterRadius();
+		const Chaos::FVec3 FocalParticleCenter = FocalParticleHandle->WorldSpaceInflatedBounds().Center();
+
+		// Focal particle specific LOD settings
+		const float BaseDistance = MinimumBaseDistance + (FocalParticleRadius * BaseDistanceRadiusMultiplier);
+		const float DistanceForFullPrediction = BaseDistance * BaseDistancesForFullPrediction;
+		const float DistanceForResimulationMode = BaseDistance * BaseDistancesForResimulationMode;
+
+		// Check distance towards focal particles taking the radius of both AABBs into account
+		const Chaos::FVec3 PosOffset = FocalParticleCenter - IslandCenter;
+		const float Distance = FMath::Max(PosOffset.Size() - FocalParticleRadius - IslandRadius, 0.0f);
+
+		// Calculate recommended time and tick alignment from the current timeline
 		float TimeAlignment = FMath::Max((Distance - DistanceForFullPrediction), 0.0f) * TimeOverDistance;
 		TimeAlignment = (TimeAlignment * 0.001f); // Convert from ms to s
 		const float TickAlignment = TimeAlignment / RigidsSolver.GetAsyncDeltaTime();
 		const int32 TickAlignmentDiscrete = FMath::Floor(TickAlignment);
 
+		// Update LodData and keep the highest LOD data
+		LodData.AlignedTime = FMath::Max(LodData.AlignedTime, TimeAlignment);
+		LodData.AlignedFrame = FMath::Max(LodData.AlignedFrame, RigidsSolver.GetCurrentFrame() - TickAlignmentDiscrete);
+		LodData.DataAssigned = true;
+
 #if !(UE_BUILD_SHIPPING || UE_BUILD_TEST)
 		if (PhysicsReplicationLODCVars::bDrawDebugEnabled)
 		{
@@ -358,15 +457,11 @@ bool FPhysicsReplicationLODAsync::PerformDistanceLOD(const Chaos::FGeometryParti
 
 			FColor DebugColor = FLinearColor::MakeFromHSV8((uint8)FMath::CeilToInt32((255.0f * 0.8f) * LodAlpha), 150, 255).ToFColor(false);
 			DebugColor.A = 255;
-			Chaos::FDebugDrawQueue::GetInstance().DrawDebugDirectionalArrow(ParticleHandle.GetX(), ParticleHandle.GetX() + PosOffset, 1.0f, DebugColor, false, 5, 0, 1.0f);
+			Chaos::FDebugDrawQueue::GetInstance().DrawDebugDirectionalArrow(IslandCenter, FocalParticleCenter, 1.0f, DebugColor, false, -1.0f, 0, 1.0f);
+			Chaos::FDebugDrawQueue::GetInstance().DrawDebugSphere(IslandCenter, IslandRadius, 8, DebugColor, false, -1.0f, 0, 1.0f);
 		}
 #endif
 
-		// Update LodData and keep the highest LOD data
-		LodData.AlignedTime = FMath::Max(LodData.AlignedTime, TimeAlignment);
-		LodData.AlignedFrame = FMath::Max(LodData.AlignedFrame, RigidsSolver.GetCurrentFrame() - TickAlignmentDiscrete);
-		LodData.DataAssigned = true;
-
 		if (Distance <= DistanceForResimulationMode)
 		{
 			LodData.ReplicationMode = EPhysicsReplicationMode::Resimulation;
@@ -374,5 +469,11 @@ bool FPhysicsReplicationLODAsync::PerformDistanceLOD(const Chaos::FGeometryParti
 		}
 	}
 
+	// Cache the LOD data for each island that should use the same LOD
+	for (const Chaos::Private::FPBDIsland* Island : ParticleIslands)
+	{
+		CachedIslandLodData.Add(Island->GetArrayIndex(), LodData);
+	}
+
 	return LodData.DataAssigned;
 }
diff --git a/Engine/Source/Runtime/Engine/Public/PhysicsReplicationLOD.h b/Engine/Source/Runtime/Engine/Public/PhysicsReplicationLOD.h
--- a/Engine/Source/Runtime/Engine/Public/PhysicsReplicationLOD.h
+++ b/Engine/Source/Runtime/Engine/Public/PhysicsReplicationLOD.h
@@ -8,6 +8,7 @@
 #include "Physics/PhysicsInterfaceDeclares.h"
 #include "Physics/PhysicsInterfaceUtils.h"
 #include "PhysicsEngine/PhysicsSettings.h"
+#include "Chaos/Island/IslandManager.h"
 
 
 enum EPhysicsReplicationLODFlags : uint32
@@ -87,9 +88,14 @@ public:
 	: bEnabled(false)
 	, DefaultSettings()
 	, LodData()
+	, FocalParticles()
+	, ParticlesInFocalIslands()
+	, CachedIslandLodData()
+	, ParticleIslands()
+	, IslandParticles()
+	, IslandAABB()
+	, ParticleAABB()
 	{ 
-		FocalParticles.Reset(1);
-		ParticlesInFocalIslands.Reset(1);
 	}
 
 	virtual ~FPhysicsReplicationLODAsync() { }
@@ -124,7 +130,24 @@ private:
 private:
 	bool bEnabled;
 	FPhysicsReplicationLODSettings DefaultSettings;
+	FPhysicsRepLodData LodData;
 	TArray<Chaos::FConstPhysicsObjectHandle> FocalParticles;
+	
+	// Particles in the same islands as focal particles
 	TArray<int32> ParticlesInFocalIslands;
-	FPhysicsRepLodData LodData;
+
+	// LOD Data cached per island, valid for one physics frame
+	TMap<int32, FPhysicsRepLodData> CachedIslandLodData;
+
+	// Transient array to hold islands affected by LOD
+	TArray<const Chaos::Private::FPBDIsland*> ParticleIslands;
+	
+	// Transient array to hold particles that are part of islands
+	TArray<const Chaos::FGeometryParticleHandle*> IslandParticles;
+	
+	// Transient BoundingBox for island
+	Chaos::FAABB3 IslandAABB;
+	
+	// Transient BoundingBox for particle, used to cache LocalBounds and apply offset from world location and then add to IslandAABB
+	Chaos::FAABB3 ParticleAABB;
 };
diff --git a/Engine/Source/Runtime/Experimental/Chaos/Private/Chaos/Island/IslandManager.cpp b/Engine/Source/Runtime/Experimental/Chaos/Private/Chaos/Island/IslandManager.cpp
--- a/Engine/Source/Runtime/Experimental/Chaos/Private/Chaos/Island/IslandManager.cpp
+++ b/Engine/Source/Runtime/Experimental/Chaos/Private/Chaos/Island/IslandManager.cpp
@@ -678,7 +678,12 @@ namespace Chaos::Private
 	TArray<const FPBDIsland*> FPBDIslandManager::FindParticleIslands(const FGeometryParticleHandle* Particle) const
 	{
 		TArray<const FPBDIsland*> ParticleIslands;
+		FindParticleIslands(Particle, ParticleIslands);
+		return ParticleIslands;
+	}
 
+	void FPBDIslandManager::FindParticleIslands(const FGeometryParticleHandle* Particle, TArray<const FPBDIsland*>& OutParticleIslands) const
+	{
 		if (const FPBDIslandParticle* Node = GetGraphNode(Particle))
 		{
 			if (Node->Flags.bIsDynamic)
@@ -686,7 +691,7 @@ namespace Chaos::Private
 				// Dynamic particles are in one island
 				if (Node->Island != nullptr)
 				{
-					ParticleIslands.Add(Node->Island);
+					OutParticleIslands.Add(Node->Island);
 				}
 			}
 			else
@@ -696,19 +701,22 @@ namespace Chaos::Private
 				{
 					if (Edge->Island != nullptr)
 					{
-						ParticleIslands.AddUnique(Edge->Island);
+						OutParticleIslands.AddUnique(Edge->Island);
 					}
 				}
 			}
 		}
-
-		return ParticleIslands;
 	}
 
 	TArray<const FGeometryParticleHandle*> FPBDIslandManager::FindParticlesInIslands(const TArray<const FPBDIsland*> InIslands) const
 	{
 		TArray<const FGeometryParticleHandle*> IslandParticles;
+		FindParticlesInIslands(InIslands, IslandParticles);
+		return IslandParticles;
+	}
 
+	void FPBDIslandManager::FindParticlesInIslands(const TArray<const FPBDIsland*>& InIslands, TArray<const FGeometryParticleHandle*>& OutIslandParticles) const
+	{
 		for (const FPBDIsland* Island : InIslands)
 		{
 			for (const TArray<FPBDIslandConstraint*>& ContainerEdges : Island->ContainerEdges)
@@ -717,17 +725,15 @@ namespace Chaos::Private
 				{
 					if (Edge->Nodes[0] != nullptr)
 					{
-						IslandParticles.AddUnique(Edge->Nodes[0]->GetParticle());
+						OutIslandParticles.AddUnique(Edge->Nodes[0]->GetParticle());
 					}
 					if (Edge->Nodes[1] != nullptr)
 					{
-						IslandParticles.AddUnique(Edge->Nodes[1]->GetParticle());
+						OutIslandParticles.AddUnique(Edge->Nodes[1]->GetParticle());
 					}
 				}
 			}
 		}
-
-		return IslandParticles;
 	}
 
 	TArray<const FConstraintHandle*> FPBDIslandManager::FindConstraintsInIslands(const TArray<const FPBDIsland*> InIslands, int32 ContainerId) const
diff --git a/Engine/Source/Runtime/Experimental/Chaos/Public/Chaos/AABB.h b/Engine/Source/Runtime/Experimental/Chaos/Public/Chaos/AABB.h
--- a/Engine/Source/Runtime/Experimental/Chaos/Public/Chaos/AABB.h
+++ b/Engine/Source/Runtime/Experimental/Chaos/Public/Chaos/AABB.h
@@ -476,7 +476,14 @@ namespace Chaos
 			return *this;
 		}
 
-		FORCEINLINE TVector<T, d> Center() const { return (MMax - MMin) / (T)2 + MMin; }
+		// Move the bounding box along a vector
+		FORCEINLINE void MoveByVector(TVector<T, d> InVector)
+		{
+			MMin += InVector;
+			MMax += InVector;
+		}
+
+		FORCEINLINE TVector<T, d> Center() const { return (MMax - MMin) * T(0.5) + MMin; }
 		FORCEINLINE TVector<T, d> GetCenter() const { return Center(); }
 		FORCEINLINE TVector<T, d> GetCenterOfMass() const { return GetCenter(); }
 		FORCEINLINE TVector<T, d> Extents() const { return MMax - MMin; }
@@ -675,6 +682,17 @@ namespace Chaos
 			return (MMin.X == TNumericLimits<T>::Max());
 		}
 
+		// Shrink this AABB to being empty
+		FORCEINLINE void Clear()
+		{
+			MMin.X = TNumericLimits<T>::Max();
+			MMin.Y = TNumericLimits<T>::Max();
+			MMin.Z = TNumericLimits<T>::Max();
+			MMax.X = -TNumericLimits<T>::Max();
+			MMax.Y = -TNumericLimits<T>::Max();
+			MMax.Z = -TNumericLimits<T>::Max();
+		}
+
 		FORCEINLINE void Serialize(FChaosArchive& Ar)
 		{
 			// LWC : this method exists in order to properly call the << operator using FChaosArchive on Chaos::TVector 
diff --git a/Engine/Source/Runtime/Experimental/Chaos/Public/Chaos/Island/IslandManager.h b/Engine/Source/Runtime/Experimental/Chaos/Public/Chaos/Island/IslandManager.h
--- a/Engine/Source/Runtime/Experimental/Chaos/Public/Chaos/Island/IslandManager.h
+++ b/Engine/Source/Runtime/Experimental/Chaos/Public/Chaos/Island/IslandManager.h
@@ -503,7 +503,9 @@ namespace Chaos::Private
 		CHAOS_API int32 GetIslandIndex(const FPBDIsland* Island) const;
 		CHAOS_API int32 GetIslandArrayIndex(const FPBDIslandConstraint* Edge) const;
 		CHAOS_API TArray<const FPBDIsland*> FindParticleIslands(const FGeometryParticleHandle* Particle) const;
+		CHAOS_API void FindParticleIslands(const FGeometryParticleHandle* Particle, TArray<const FPBDIsland*>& OutParticleIslands) const;
 		CHAOS_API TArray<const FGeometryParticleHandle*> FindParticlesInIslands(const TArray<const FPBDIsland*> Islands) const;
+		CHAOS_API void FindParticlesInIslands(const TArray<const FPBDIsland*>& InIslands, TArray<const FGeometryParticleHandle*>& OutIslandParticles) const;
 		CHAOS_API TArray<const FConstraintHandle*> FindConstraintsInIslands(const TArray<const FPBDIsland*> Islands, int32 ContainerId) const;
 
 		// Debug draw
