Fix assert in the dynamic resolution scaling code.

These functions are now called from rendering tasks, so the checks for rendering thread were incorrect.

#rnx

diff --git a/Engine/Source/Runtime/Engine/Private/DynamicResolution.cpp b/Engine/Source/Runtime/Engine/Private/DynamicResolution.cpp
--- a/Engine/Source/Runtime/Engine/Private/DynamicResolution.cpp
+++ b/Engine/Source/Runtime/Engine/Private/DynamicResolution.cpp
@@ -1043,7 +1043,7 @@ public:
 
 	virtual DynamicRenderScaling::TMap<float> GetResolutionFractions_RenderThread() const override
 	{
-		check(IsInRenderingThread());
+		check(IsInParallelRenderingThread());
 
 		DynamicRenderScaling::TMap<float> ResolutionFractions = Proxy->Heuristic.QueryCurrentFrameResolutionFractions();
 		if (!ViewFamily.EngineShowFlags.ScreenPercentage)
diff --git a/Engine/Source/Runtime/Renderer/Private/SceneRendering.cpp b/Engine/Source/Runtime/Renderer/Private/SceneRendering.cpp
--- a/Engine/Source/Runtime/Renderer/Private/SceneRendering.cpp
+++ b/Engine/Source/Runtime/Renderer/Private/SceneRendering.cpp
@@ -644,7 +644,7 @@ public:
 
 	virtual DynamicRenderScaling::TMap<float> GetResolutionFractions_RenderThread() const override
 	{
-		check(IsInRenderingThread());
+		check(IsInParallelRenderingThread());
 
 		uint32 FrameId = 0;
 
