Expose option to smooth timesteps in immediate physics

diff --git a/Engine/Plugins/Experimental/PhysicsControl/Source/ControlRigPhysics/Private/RigPhysicsSolver.cpp b/Engine/Plugins/Experimental/PhysicsControl/Source/ControlRigPhysics/Private/RigPhysicsSolver.cpp
--- a/Engine/Plugins/Experimental/PhysicsControl/Source/ControlRigPhysics/Private/RigPhysicsSolver.cpp
+++ b/Engine/Plugins/Experimental/PhysicsControl/Source/ControlRigPhysics/Private/RigPhysicsSolver.cpp
@@ -651,23 +651,17 @@ void FRigPhysicsSolver::StepSimulation(FRigVMExecuteContext& ExecuteContext, dou
 		SolverSettings.VelocityIterations,
 		SolverSettings.ProjectionIterations,
 		SolverSettings.bUseManifolds);
+	Simulation->SetMaxNumRollingAverageStepTimes(1);
 
 	// Other settings - would normally be static (so Danny TODO move this)
 	ChaosJointSolverSettings.bSolvePositionLast = SolverSettings.bSolveJointPositionsLast;
 	ChaosJointSolverSettings.bSortEnabled = true;
 
 	// Simulation space
-#if 0
-	// Danny TODO this requires a tweak to the API to allow the full settings to be passed in and
-	// controlled by us. That needs reviewing by physics.
 	Chaos::FSimulationSpaceSettings ChaosSimulationSpaceSettings;
 	ChaosSimulationSpaceSettings.Alpha = SimulationSpaceSettings.SpaceMovementAmount;
 	ChaosSimulationSpaceSettings.ExternalLinearEtherDrag = SimulationSpaceSettings.ExternalLinearDragV;
 	Simulation->SetSimulationSpaceSettings(ChaosSimulationSpaceSettings);
-#else
-	Simulation->SetSimulationSpaceSettings(
-		SimulationSpaceSettings.SpaceMovementAmount, SimulationSpaceSettings.ExternalLinearDragV);
-#endif
 	Simulation->UpdateSimulationSpace(
 		SimulationSpaceState.SimulationSpaceTM,
 		SimulationSpaceData.LinearVelocity,
diff --git a/Engine/Source/Runtime/Engine/Private/PhysicsEngine/ImmediatePhysics/ImmediatePhysicsChaos/ImmediatePhysicsSimulation_Chaos.cpp b/Engine/Source/Runtime/Engine/Private/PhysicsEngine/ImmediatePhysics/ImmediatePhysicsChaos/ImmediatePhysicsSimulation_Chaos.cpp
--- a/Engine/Source/Runtime/Engine/Private/PhysicsEngine/ImmediatePhysics/ImmediatePhysicsChaos/ImmediatePhysicsSimulation_Chaos.cpp
+++ b/Engine/Source/Runtime/Engine/Private/PhysicsEngine/ImmediatePhysics/ImmediatePhysicsChaos/ImmediatePhysicsSimulation_Chaos.cpp
@@ -354,7 +354,7 @@ namespace ImmediatePhysics_Chaos
 		Chaos::FReal FixedStepTime;
 		Chaos::FReal RollingAverageStepTime;
 		int32 NumRollingAverageStepTimes;
-		int32 MaxNumRollingAverageStepTimes;
+		int32 MaxNumRollingAverageStepTimes; // Set to 1 to disable smoothing
 
 		bool bActorsDirty;
 	};
@@ -879,6 +879,14 @@ namespace ImmediatePhysics_Chaos
 		}
 	}
 
+	void FSimulation::SetMaxNumRollingAverageStepTimes(const int32 MaxNumRollingAverageStepTimes)
+	{
+		if (MaxNumRollingAverageStepTimes >= 0)
+		{
+			Implementation->MaxNumRollingAverageStepTimes = MaxNumRollingAverageStepTimes;
+		}
+	}
+
 	void FSimulation::DebugDraw()
 	{
 #if CHAOS_DEBUG_DRAW
@@ -894,9 +902,18 @@ namespace ImmediatePhysics_Chaos
 
 	FReal FSimulation::UpdateStepTime(const FReal DeltaTime, const FReal MaxStepTime)
 	{
-		Implementation->RollingAverageStepTime = Implementation->RollingAverageStepTime + (DeltaTime - Implementation->RollingAverageStepTime) / Implementation->NumRollingAverageStepTimes;
-		Implementation->RollingAverageStepTime = FMath::Min(Implementation->RollingAverageStepTime, MaxStepTime);
-		Implementation->NumRollingAverageStepTimes = FMath::Min(Implementation->NumRollingAverageStepTimes + 1, Implementation->MaxNumRollingAverageStepTimes);
+		if (Implementation->NumRollingAverageStepTimes > 1)
+		{
+			Implementation->RollingAverageStepTime = Implementation->RollingAverageStepTime + (DeltaTime - Implementation->RollingAverageStepTime) / Implementation->NumRollingAverageStepTimes;
+			Implementation->RollingAverageStepTime = FMath::Min(Implementation->RollingAverageStepTime, MaxStepTime);
+			Implementation->NumRollingAverageStepTimes = FMath::Min(Implementation->NumRollingAverageStepTimes + 1, Implementation->MaxNumRollingAverageStepTimes);
+		}
+		else
+		{
+			// Adjust things so that if we go back to smoothing the step time, we do so gracefully
+			Implementation->NumRollingAverageStepTimes = Implementation->MaxNumRollingAverageStepTimes;
+			Implementation->RollingAverageStepTime = DeltaTime;
+		}
 		return Implementation->RollingAverageStepTime;
 	}
 
@@ -937,6 +954,7 @@ namespace ImmediatePhysics_Chaos
 			ChaosImmediate_Evolution_VelocityIterations,
 			ChaosImmediate_Evolution_ProjectionIterations,
 			ChaosImmediate_Collision_UseManifolds);
+		SetMaxNumRollingAverageStepTimes(ChaosImmediate_Evolution_DeltaTimeCount);
 
 		if (JointSolverSettings)
 		{
diff --git a/Engine/Source/Runtime/Engine/Public/Physics/ImmediatePhysics/ImmediatePhysicsChaos/ImmediatePhysicsSimulation_Chaos.h b/Engine/Source/Runtime/Engine/Public/Physics/ImmediatePhysics/ImmediatePhysicsChaos/ImmediatePhysicsSimulation_Chaos.h
--- a/Engine/Source/Runtime/Engine/Public/Physics/ImmediatePhysics/ImmediatePhysicsChaos/ImmediatePhysicsSimulation_Chaos.h
+++ b/Engine/Source/Runtime/Engine/Public/Physics/ImmediatePhysics/ImmediatePhysicsChaos/ImmediatePhysicsSimulation_Chaos.h
@@ -119,7 +119,7 @@ namespace ImmediatePhysics_Chaos
 		ENGINE_API void SetSimulationSpaceSettings(
 			const Chaos::FSimulationSpaceSettings& SimulationSpaceSettings);
 
-		/** Set settings. Invalid (negative) values with leave that value unchanged from defaults */
+		/** Set settings. Invalid (negative) values will leave that value unchanged from defaults */
 		ENGINE_API void SetSolverSettings(
 			const FReal FixedDt,
 			const FReal CullDistance,
@@ -130,6 +130,11 @@ namespace ImmediatePhysics_Chaos
 			const int32 ProjectionIts,
 			const int32 bUseManifolds);
 
+		// Additional settings for less commonly adjusted parameters. As in SetSolverSettings,
+		// invalid (negative) values will leave that value unchanged from the default.
+		ENGINE_API void SetMaxNumRollingAverageStepTimes(
+			const int32 MaxNumRollingAverageStepTimes);
+
 		/** Explicit debug draw path if the use case needs it to happen at a point outside of the simulation **/
 		ENGINE_API void DebugDraw();
 
