FPhysScene_AsyncPhysicsStateJobQueue::Tick() now calls OnPostExecute_GameThread on completed jobs even if not all jobs are completed.
This avoids any artificial delay before receiving all OnPostExecute_GameThread for a level (that would be caused by other pending jobs of another level).

diff --git a/Engine/Source/Runtime/Engine/Private/PhysicsEngine/Experimental/PhysScene_AsyncPhysicsStateJobQueue.cpp b/Engine/Source/Runtime/Engine/Private/PhysicsEngine/Experimental/PhysScene_AsyncPhysicsStateJobQueue.cpp
--- a/Engine/Source/Runtime/Engine/Private/PhysicsEngine/Experimental/PhysScene_AsyncPhysicsStateJobQueue.cpp
+++ b/Engine/Source/Runtime/Engine/Private/PhysicsEngine/Experimental/PhysScene_AsyncPhysicsStateJobQueue.cpp
@@ -145,10 +145,30 @@ void FPhysScene_AsyncPhysicsStateJobQueue::Tick()
 		WaitForAsyncTasksExecution();
 	}
 
+	// Process completed tasks
+	check(PipeTasks.Num() == PipedJobs.Num());
+	for (int i = 0; i < PipeTasks.Num();)
+	{
+		if (PipeTasks[i].IsCompleted())
+		{
+			PipedJobs[i]->OnPostExecute_GameThread();
+			PipeTasks.RemoveAt(i, EAllowShrinking::No);
+			PipedJobs.RemoveAt(i, EAllowShrinking::No);
+		}
+		else
+		{
+			++i;
+		}
+	}
+
+#if DO_CHECK
 	if (AreAsyncTasksExecuted())
 	{
-		OnAsyncTasksExecuted();
+		check(PipeTasks.IsEmpty());
+		check(PipedJobs.IsEmpty());
+		check(!IsReadyToRunAsyncTasksEvent.IsSet());
 	}
+#endif
 }
 
 void FPhysScene_AsyncPhysicsStateJobQueue::FJobSet::Execute() const
