- Cover more cases allowed by UStaticMeshComponent to support asynchronous physics state creation (attach parent of other type than UPrimitiveComponent).
- Fixed UStaticMeshComponent not calling UpdateCollisionFromStaticMesh() before testing whether incremental pre-register was supported (ShouldIncrementalPreRegister).
This could lead to cases where preregister could have applied but didn't (i.e. not using the optimized code path).

diff --git a/Engine/Source/Runtime/Engine/Classes/Components/StaticMeshComponent.h b/Engine/Source/Runtime/Engine/Classes/Components/StaticMeshComponent.h
--- a/Engine/Source/Runtime/Engine/Classes/Components/StaticMeshComponent.h
+++ b/Engine/Source/Runtime/Engine/Classes/Components/StaticMeshComponent.h
@@ -31,6 +31,7 @@ class ULightComponent;
 class UStaticMesh;
 class UStaticMeshComponent;
 class UNavCollisionBase;
+class UWorld;
 struct FConvexVolume;
 struct FEngineShowFlags;
 struct FNavigableGeometryExport;
@@ -544,7 +545,6 @@ public:
 
 	//~ Begin UActorComponent Interface.
 protected: 
-	ENGINE_API virtual void OnPreRegister() override;
 	ENGINE_API virtual void OnRegister() override;
 	ENGINE_API virtual void OnUnregister() override;
 	ENGINE_API virtual void BeginPlay() override;
@@ -555,6 +555,7 @@ protected:
 	ENGINE_API virtual void OnAsyncDestroyPhysicsStateEnd_GameThread() override;
 	ENGINE_API virtual bool ShouldCreatePhysicsState() const override;
 	ENGINE_API virtual bool ShouldCreateRenderState() const override;
+	ENGINE_API virtual bool ShouldIncrementalPreRegister(UWorld* WorldContext) const override;
 	ENGINE_API virtual bool AllowsAsyncPhysicsStateCreation() const override;
 	ENGINE_API virtual bool AllowsAsyncPhysicsStateDestruction() const override;
 
diff --git a/Engine/Source/Runtime/Engine/Private/Components/StaticMeshComponent.cpp b/Engine/Source/Runtime/Engine/Private/Components/StaticMeshComponent.cpp
--- a/Engine/Source/Runtime/Engine/Private/Components/StaticMeshComponent.cpp
+++ b/Engine/Source/Runtime/Engine/Private/Components/StaticMeshComponent.cpp
@@ -808,13 +808,6 @@ void UStaticMeshComponent::PostEditImport()
 
 #endif // #if WITH_EDITOR
 
-void UStaticMeshComponent::OnPreRegister()
-{
-	TRACE_CPUPROFILER_EVENT_SCOPE(UStaticMeshComponent::OnPreRegister)
-	UpdateCollisionFromStaticMesh();
-	Super::OnPreRegister();
-}
-
 void UStaticMeshComponent::OnRegister()
 {
 	TRACE_CPUPROFILER_EVENT_SCOPE(UStaticMeshComponent::OnRegister)
@@ -825,7 +818,7 @@ void UStaticMeshComponent::OnRegister()
 
 	NotifyIfStaticMeshChanged();
 
-	// No need to call UpdateCollisionFromStaticMesh if pre-registered as it has already been called in UStaticMeshComponent::OnPreRegister
+	// No need to call UpdateCollisionFromStaticMesh if pre-registered as it has already been called
 	if (!IsPreRegistered())
 	{
 		UpdateCollisionFromStaticMesh();
@@ -2315,10 +2308,18 @@ bool UStaticMeshComponent::ShouldCreatePhysicsState() const
 	return Super::ShouldCreatePhysicsState() && GetStaticMesh() && !GetStaticMesh()->IsCompiling();
 }
 
+bool UStaticMeshComponent::ShouldIncrementalPreRegister(UWorld* WorldContext) const
+{
+	// Make sure BodyInstance CollisionEnabled is updated before calling Super::ShouldIncrementalPreRegister as it calls
+	// UPrimitiveComponent::ShouldCreatePhysicsState which calls BodyInstance.GetCollisionEnabled().
+	const_cast<UStaticMeshComponent*>(this)->UpdateCollisionFromStaticMesh();
+	return Super::ShouldIncrementalPreRegister(WorldContext);
+}
+
 bool UStaticMeshComponent::AllowsAsyncPhysicsStateCreation() const
 {
 	// Async physics state creation is supported for static collision or primitive component not using welding
-	return !BodyInstance.bAutoWeld || !GetAttachParent() || IsStaticPhysics();
+	return !BodyInstance.bAutoWeld || !Cast<UPrimitiveComponent>(GetAttachParent()) || IsStaticPhysics();
 }
 
 bool UStaticMeshComponent::AllowsAsyncPhysicsStateDestruction() const
