Changed FDynamicResolutionHeuristicProxy check(IsInRenderingThread()) to IsInParallelRenderingThread() since it is now called in parallel rendering threads.

diff --git a/Engine/Source/Runtime/Engine/Public/DynamicResolutionProxy.h b/Engine/Source/Runtime/Engine/Public/DynamicResolutionProxy.h
--- a/Engine/Source/Runtime/Engine/Public/DynamicResolutionProxy.h
+++ b/Engine/Source/Runtime/Engine/Public/DynamicResolutionProxy.h
@@ -39,7 +39,7 @@ public:
 	/** Returns the view fraction that should be used for current frame. */
 	FORCEINLINE DynamicRenderScaling::TMap<float> QueryCurrentFrameResolutionFractions() const
 	{
-		check(IsInRenderingThread());
+		check(IsInParallelRenderingThread());
 		return QueryCurrentFrameResolutionFractions_Internal();
 	}
 
