Cleaning up deprecated API from AnimSequence(base) and AnimCompositeBase

diff --git a/Engine/Plugins/Experimental/AnimNext/Source/AnimNext/Private/DecompressionTools.cpp b/Engine/Plugins/Experimental/AnimNext/Source/AnimNext/Private/DecompressionTools.cpp
--- a/Engine/Plugins/Experimental/AnimNext/Source/AnimNext/Private/DecompressionTools.cpp
+++ b/Engine/Plugins/Experimental/AnimNext/Source/AnimNext/Private/DecompressionTools.cpp
@@ -59,13 +59,11 @@ static bool CanEvaluateRawAnimationData(const UAnimSequence* AnimSequence)
 #endif
 }
 
-static bool UseRawDataForPoseExtraction(const UAnimSequence* AnimSequence, const FLODPose& AnimationPoseData)
+bool FDecompressionTools::UseRawDataForPoseExtraction(const UAnimSequence* AnimSequence, const FLODPose& AnimationPoseData)
 {
 	return CanEvaluateRawAnimationData(AnimSequence) && (
 #if WITH_EDITOR
-		PRAGMA_DISABLE_DEPRECATION_WARNINGS
-		AnimSequence->OnlyUseRawData() ||
-		PRAGMA_ENABLE_DEPRECATION_WARNINGS
+		AnimSequence->bUseRawDataOnly ||
 #endif
 		(AnimSequence->GetSkeletonVirtualBoneGuid() != AnimSequence->GetSkeleton()->GetVirtualBoneGuid()) || AnimationPoseData.GetDisableRetargeting() || AnimationPoseData.ShouldUseRawData() ||
 #if WITH_EDITOR
diff --git a/Engine/Plugins/Experimental/AnimNext/Source/AnimNext/Public/DecompressionTools.h b/Engine/Plugins/Experimental/AnimNext/Source/AnimNext/Public/DecompressionTools.h
--- a/Engine/Plugins/Experimental/AnimNext/Source/AnimNext/Public/DecompressionTools.h
+++ b/Engine/Plugins/Experimental/AnimNext/Source/AnimNext/Public/DecompressionTools.h
@@ -76,6 +76,8 @@ public:
 		, const int32 LODBoneIndex
 		, const bool bIsBakedAdditive
 		, const bool bDisableRetargeting);
+
+	static bool UseRawDataForPoseExtraction(const UAnimSequence* AnimSequence, const FLODPose& AnimationPoseData);
 };
 
 } // end namespace UE::AnimNext
diff --git a/Engine/Source/Editor/UnrealEd/Private/Animation/EditorAnimCompositeSegment.cpp b/Engine/Source/Editor/UnrealEd/Private/Animation/EditorAnimCompositeSegment.cpp
--- a/Engine/Source/Editor/UnrealEd/Private/Animation/EditorAnimCompositeSegment.cpp
+++ b/Engine/Source/Editor/UnrealEd/Private/Animation/EditorAnimCompositeSegment.cpp
@@ -30,12 +30,10 @@ void UEditorAnimCompositeSegment::InitAnimSegment(int32 AnimSegmentIndexIn)
 
 void UEditorAnimCompositeSegment::PostEditChangeProperty(FPropertyChangedEvent& PropertyChangedEvent)
 {
-	PRAGMA_DISABLE_DEPRECATION_WARNINGS
 	if (PropertyChangedEvent.GetPropertyName() == GET_MEMBER_NAME_CHECKED(FAnimSegment, AnimReference))
 	{
 		AnimSegment.UpdateCachedPlayLength();			
 	}
-	PRAGMA_ENABLE_DEPRECATION_WARNINGS
 	
 	Super::PostEditChangeProperty(PropertyChangedEvent);
 }
diff --git a/Engine/Source/Editor/UnrealEd/Private/Animation/EditorAnimSegment.cpp b/Engine/Source/Editor/UnrealEd/Private/Animation/EditorAnimSegment.cpp
--- a/Engine/Source/Editor/UnrealEd/Private/Animation/EditorAnimSegment.cpp
+++ b/Engine/Source/Editor/UnrealEd/Private/Animation/EditorAnimSegment.cpp
@@ -33,12 +33,10 @@ void UEditorAnimSegment::InitAnimSegment(int32 AnimSlotIndexIn, int32 AnimSegmen
 
 void UEditorAnimSegment::PostEditChangeProperty(FPropertyChangedEvent& PropertyChangedEvent)
 {
-	PRAGMA_DISABLE_DEPRECATION_WARNINGS
 	if (PropertyChangedEvent.GetPropertyName() == GET_MEMBER_NAME_CHECKED(FAnimSegment, AnimReference))
 	{
 		AnimSegment.UpdateCachedPlayLength();			
 	}
-	PRAGMA_ENABLE_DEPRECATION_WARNINGS
 	
 	Super::PostEditChangeProperty(PropertyChangedEvent);
 }
diff --git a/Engine/Source/Runtime/Engine/Classes/Animation/AnimCompositeBase.h b/Engine/Source/Runtime/Engine/Classes/Animation/AnimCompositeBase.h
--- a/Engine/Source/Runtime/Engine/Classes/Animation/AnimCompositeBase.h
+++ b/Engine/Source/Runtime/Engine/Classes/Animation/AnimCompositeBase.h
@@ -68,17 +68,10 @@ struct FAnimSegment
 {
 	GENERATED_USTRUCT_BODY()
 
-	PRAGMA_DISABLE_DEPRECATION_WARNINGS
 	FAnimSegment(const FAnimSegment&) = default;
 	FAnimSegment(FAnimSegment&&) = default;
 	FAnimSegment& operator=(const FAnimSegment&) = default;
 	FAnimSegment& operator=(FAnimSegment&&) = default;
-	PRAGMA_ENABLE_DEPRECATION_WARNINGS
-
-	UE_DEPRECATED(5.1, "Public access to AnimReference has been deprecated, use Set/Get-AnimReference instead")
-	/** Anim Reference to play - only allow AnimSequence or AnimComposite **/
-	UPROPERTY(EditAnywhere, Category=AnimSegment, meta=(DisplayName = "Animation Reference"))
-	TObjectPtr<UAnimSequenceBase> AnimReference;
 
 #if WITH_EDITORONLY_DATA
 	UPROPERTY(EditAnywhere, Category=AnimSegment, meta=(DisplayName = "Cached Animation Asset length"))
@@ -90,14 +83,17 @@ struct FAnimSegment
 	friend class UEditorAnimCompositeSegment;
 	ENGINE_API void UpdateCachedPlayLength();
 #endif // WITH_EDITOR
+
+protected:
+	/** Anim Reference to play - only allow AnimSequence or AnimComposite **/
+	UPROPERTY(EditAnywhere, Category=AnimSegment, meta=(DisplayName = "Animation Reference"))
+	TObjectPtr<UAnimSequenceBase> AnimReference;
 public:
 
 	ENGINE_API void SetAnimReference(UAnimSequenceBase* InAnimReference, bool bInitialize = false);
 	const TObjectPtr<UAnimSequenceBase>& GetAnimReference() const 
 	{
-		PRAGMA_DISABLE_DEPRECATION_WARNINGS
 		return AnimReference;
-		PRAGMA_ENABLE_DEPRECATION_WARNINGS
 	}
 
 #if WITH_EDITOR
@@ -124,9 +120,7 @@ public:
 	int32 LoopingCount;
 
 	FAnimSegment()
-		PRAGMA_DISABLE_DEPRECATION_WARNINGS
 		: AnimReference(nullptr)
-		PRAGMA_ENABLE_DEPRECATION_WARNINGS
 		, StartPos(0.f)
 		, AnimStartTime(0.f)
 		, AnimEndTime(0.f)
@@ -139,9 +133,7 @@ public:
 	/** Ensures PlayRate is non Zero */
 	float GetValidPlayRate() const
 	{
-		PRAGMA_DISABLE_DEPRECATION_WARNINGS
 		float SeqPlayRate = AnimReference ? AnimReference->RateScale : 1.0f;
-		PRAGMA_ENABLE_DEPRECATION_WARNINGS
 		float FinalPlayRate = SeqPlayRate * AnimPlayRate;
 		return (FMath::IsNearlyZero(FinalPlayRate) ? 1.f : FinalPlayRate);
 	}
@@ -194,15 +186,6 @@ public:
 	/** Converts 'Track Position' to position on AnimSequence.
 	 * Note: doesn't check that position is in valid range, must do that before calling this function! */
 	ENGINE_API float ConvertTrackPosToAnimPos(const float& TrackPosition) const;
-
-	/** 
-	 * Retrieves AnimNotifies between two Track time positions. ]PreviousTrackPosition, CurrentTrackPosition]
-	 * Between PreviousTrackPosition (exclusive) and CurrentTrackPosition (inclusive).
-	 * Supports playing backwards (CurrentTrackPosition<PreviousTrackPosition).
-	 * Only supports contiguous range, does NOT support looping and wrapping over.
-	 */
-	UE_DEPRECATED(4.19, "Use the GetAnimNotifiesFromTrackPositions that takes FAnimNotifyContext instead")
-	void GetAnimNotifiesFromTrackPositions(const float& PreviousTrackPosition, const float& CurrentTrackPosition, TArray<const FAnimNotifyEvent *> & OutActiveNotifies) const;
 	
 	/**
 	* Retrieves AnimNotifies between two Track time positions. ]PreviousTrackPosition, CurrentTrackPosition]
@@ -210,9 +193,6 @@ public:
 	* Supports playing backwards (CurrentTrackPosition<PreviousTrackPosition).
 	* Only supports contiguous range, does NOT support looping and wrapping over.
 	*/
-	UE_DEPRECATED(5.0, "Use the GetAnimNotifiesFromTrackPositions that takes FAnimNotifyContext instead")
-	ENGINE_API void GetAnimNotifiesFromTrackPositions(const float& PreviousTrackPosition, const float& CurrentTrackPosition, TArray<FAnimNotifyEventReference> & OutActiveNotifies) const;
-
 	ENGINE_API void GetAnimNotifiesFromTrackPositions(const float& PreviousTrackPosition, const float& CurrentTrackPosition, FAnimNotifyContext& NotifyContext) const;
 
 	/** 
@@ -230,9 +210,7 @@ public:
 	/** 
 	 * return true if anim notify is available 
 	 */
-	PRAGMA_DISABLE_DEPRECATION_WARNINGS
 	bool IsNotifyAvailable() const { return IsValid() && AnimReference && AnimReference->IsNotifyAvailable(); }
-	PRAGMA_ENABLE_DEPRECATION_WARNINGS
 private:
 
 	/**
@@ -285,10 +263,7 @@ struct FAnimTrack
 	ENGINE_API FAnimSegment* GetSegmentAtTime(float InTime);
 	ENGINE_API const FAnimSegment* GetSegmentAtTime(float InTime) const;
 
-	/** Get animation pose function */
-	UE_DEPRECATED(4.26, "Use other GetAnimationPose signature")
-	ENGINE_API void GetAnimationPose(/*out*/ FCompactPose& OutPose,/*out*/ FBlendedCurve& OutCurve, const FAnimExtractContext& ExtractionContext) const;
-	
+	/** Get animation pose function */	
 	ENGINE_API void GetAnimationPose(FAnimationPoseData& OutAnimationPoseData, const FAnimExtractContext& ExtractionContext) const;
 
 	/** Enable Root motion setting from montage */
@@ -319,24 +294,6 @@ struct FAnimTrack
 	// We're going to remove the top reference if found
 	bool ContainRecursive(const TArray<UAnimCompositeBase*>& CurrentAccumulatedList);
 
-	/**
-	* Retrieves AnimNotifies between two Track time positions. ]PreviousTrackPosition, CurrentTrackPosition]
-	* Between PreviousTrackPosition (exclusive) and CurrentTrackPosition (inclusive).
-	* Supports playing backwards (CurrentTrackPosition<PreviousTrackPosition).
-	* Only supports contiguous range, does NOT support looping and wrapping over.
-	*/
-	UE_DEPRECATED(4.19, "Use the GetAnimNotifiesFromTrackPositions that takes FAnimNotifyEventReferences instead")
-	void GetAnimNotifiesFromTrackPositions(const float& PreviousTrackPosition, const float& CurrentTrackPosition, TArray<const FAnimNotifyEvent *> & OutActiveNotifies) const;
-
-	/**
-	* Retrieves AnimNotifies between two Track time positions. ]PreviousTrackPosition, CurrentTrackPosition]
-	* Between PreviousTrackPosition (exclusive) and CurrentTrackPosition (inclusive).
-	* Supports playing backwards (CurrentTrackPosition<PreviousTrackPosition).
-	* Only supports contiguous range, does NOT support looping and wrapping over.
-	*/
-	UE_DEPRECATED(5.0, "Use the GetAnimNotifiesFromTrackPositions that takes FAnimNotifyContext instead")
-	ENGINE_API void GetAnimNotifiesFromTrackPositions(const float& PreviousTrackPosition, const float& CurrentTrackPosition, TArray<FAnimNotifyEventReference>& OutActiveNotifies) const;
-
 	/**
 	* Retrieves AnimNotifies between two Track time positions. ]PreviousTrackPosition, CurrentTrackPosition]
 	* Between PreviousTrackPosition (exclusive) and CurrentTrackPosition (inclusive).
diff --git a/Engine/Source/Runtime/Engine/Classes/Animation/AnimSequence.h b/Engine/Source/Runtime/Engine/Classes/Animation/AnimSequence.h
--- a/Engine/Source/Runtime/Engine/Classes/Animation/AnimSequence.h
+++ b/Engine/Source/Runtime/Engine/Classes/Animation/AnimSequence.h
@@ -184,33 +184,6 @@ struct FCompressedTrack
 
 };
 
-// Param structure for UAnimSequence::RequestAnimCompressionParams
-struct UE_DEPRECATED(5.2, "FRequestAnimCompressionParams has been deprecated") FRequestAnimCompressionParams;
-struct FRequestAnimCompressionParams
-{
-	// Is the compression to be performed Async
-	bool bAsyncCompression;
-
-	// Should we attempt to do framestripping (removing every other frame from raw animation tracks)
-	bool bPerformFrameStripping;
-
-	// If false we only perform frame stripping on even numbered frames (as a quality measure)
-	bool bPerformFrameStrippingOnOddNumberedFrames;
-
-	// Compression context
-	TSharedPtr<FAnimCompressContext> CompressContext;
-
-	// Constructors
-	FRequestAnimCompressionParams(bool bInAsyncCompression, bool bInAllowAlternateCompressor = false, bool bInOutput = false, const ITargetPlatform* = nullptr) {}
-	FRequestAnimCompressionParams(bool bInAsyncCompression, TSharedPtr<FAnimCompressContext> InCompressContext) {}
-
-	// Frame stripping initialization funcs (allow stripping per platform)
-	void InitFrameStrippingFromCVar() {}
-	void InitFrameStrippingFromPlatform() {}
-
-	const ITargetPlatform* TargetPlatform;
-};
-
 /** Enum used to decide whether we should strip animation data on dedicated server */
 UENUM()
 enum class EStripAnimDataOnDedicatedServerSettings : uint8
@@ -432,13 +405,7 @@ public:
 	ENGINE_API virtual int32 GetNumberOfSampledKeys() const override;
 	virtual FFrameRate GetSamplingFrameRate() const override { return PlatformTargetFrameRate.Default; }
 	ENGINE_API virtual void EvaluateCurveData(FBlendedCurve& OutCurve, const FAnimExtractContext& AnimExtractContext, bool bForceUseRawData = false) const override;
-
-	UE_DEPRECATED(5.3, "Please use EvaluateCurveData that takes an FName.")
-	virtual float EvaluateCurveData(SmartName::UID_Type CurveUID, float CurrentTime, bool bForceUseRawData = false) const override { return 0.0f; }
 	ENGINE_API virtual float EvaluateCurveData(FName CurveName, const FAnimExtractContext& AnimExtractContext, bool bForceUseRawData = false) const override;
-
-	UE_DEPRECATED(5.3, "Please use HasCurveData that takes an FName.")
-	virtual bool HasCurveData(SmartName::UID_Type CurveUID, bool bForceUseRawData) const override { return false; }
 	ENGINE_API virtual bool HasCurveData(FName CurveName, bool bForceUseRawData) const override;
 
 	//~ End UAnimSequenceBase Interface
@@ -459,17 +426,6 @@ public:
 
 	// Begin Transform related functions 
 	ENGINE_API virtual void GetAnimationPose(FAnimationPoseData& OutAnimationPoseData, const FAnimExtractContext& ExtractionContext) const override;
-
-	/**
-	* Get Bone Transform of the animation for the Time given, relative to Parent for all RequiredBones
-	*
-	* @param	OutPose				[out] Array of output bone transforms
-	* @param	OutCurve			[out] Curves to fill	
-	* @param	ExtractionContext	Extraction Context (position, looping, root motion, etc.)
-	* @param	bForceUseRawData	Override other settings and force raw data pose extraction
-	*/
-	UE_DEPRECATED(4.26, "Use other GetBonePose signature")
-	ENGINE_API void GetBonePose(FCompactPose& OutPose, FBlendedCurve& OutCurve, const FAnimExtractContext& ExtractionContext, bool bForceUseRawData=false) const;
 	
 	/**
 	* Get Bone Transform of the Time given, relative to Parent for all RequiredBones
@@ -481,18 +437,11 @@ public:
 	ENGINE_API void GetBonePose(struct FAnimationPoseData& OutAnimationPoseData, const FAnimExtractContext& ExtractionContext, bool bForceUseRawData = false) const;
 
 	const TArray<FTrackToSkeletonMap>& GetCompressedTrackToSkeletonMapTable() const { return CompressedData.CompressedTrackToSkeletonMapTable; }
-
-	UE_DEPRECATED(5.3, "Please use GetCompressedCurveIndexedNames")
-	ENGINE_API const TArray<struct FSmartName>& GetCompressedCurveNames() const;
-
 	const TArray<FAnimCompressedCurveIndexedName>& GetCompressedCurveIndexedNames() const { return CompressedData.IndexedCurveNames; }
 
 #if WITH_EDITORONLY_DATA
 protected:
 	ENGINE_API void UpdateCompressedCurveName(const FName& OldCurveName, const FName& NewCurveName);
-	
-	UE_DEPRECATED(5.3, "Please use UpdateCompressedCurveName that takes FNames.")
-	void UpdateCompressedCurveName(SmartName::UID_Type CurveUID, const struct FSmartName& NewCurveName) {}
 #endif
 
 public:
@@ -539,23 +488,17 @@ public:
 	/**
 	* Get Bone Transform of the additive animation for the Time given, relative to Parent for all RequiredBones
 	*
-	* @param	OutPose				[out] Output bone transforms
-	* @param	OutCurve			[out] Curves to fill	
-	* @param	ExtractionContext	Extraction Context (position, looping, root motion, etc.)
+	* @param	OutAnimationPoseData	[out] Output pose data
+	* @param	ExtractionContext		Extraction Context (position, looping, root motion, etc.)
 	*/
-	UE_DEPRECATED(4.26, "Use other GetBonePose_Additive signature")
-	ENGINE_API void GetBonePose_Additive(FCompactPose& OutPose, FBlendedCurve& OutCurve, const FAnimExtractContext& ExtractionContext) const;
 	ENGINE_API void GetBonePose_Additive(FAnimationPoseData& OutAnimationPoseData, const FAnimExtractContext& ExtractionContext) const;
 
 	/**
 	* Get Bone Transform of the base (reference) pose of the additive animation for the Time given, relative to Parent for all RequiredBones
 	*
-	* @param	OutPose				[out] Output bone transforms
-	* @param	OutCurve			[out] Curves to fill	
+	* @param	OutAnimationPoseData	[out] Output pose data
 	* @param	ExtractionContext	Extraction Context (position, looping, root motion, etc.)
 	*/
-	UE_DEPRECATED(4.26, "Use other GetAdditiveBasePose signature")
-	ENGINE_API void GetAdditiveBasePose(FCompactPose& OutPose, FBlendedCurve& OutCurve, const FAnimExtractContext& ExtractionContext) const;
 	ENGINE_API void GetAdditiveBasePose(FAnimationPoseData& OutAnimationPoseData, const FAnimExtractContext& ExtractionContext) const;
 
 	UE_DEPRECATED(5.6, "Please use GetBoneTransform with FAnimExtractContext")
@@ -571,9 +514,6 @@ public:
 	 */
 	ENGINE_API void GetBoneTransform(FTransform& OutAtom, FSkeletonPoseBoneIndex BoneIndex, const FAnimExtractContext& ExtractionContext, bool bUseRawData) const;
 
-	UE_DEPRECATED(5.1, "Use other GetBoneTransform signature using double and skeleton index")
-	void GetBoneTransform(FTransform& OutAtom, int32 TrackIndex, float Time, bool bUseRawData) const {}
-
 	/**
 	 * Get Bone Transform of the Time given, relative to Parent for the Track Given
 	 *
@@ -620,28 +560,9 @@ public:
 
 	// Get compressed data for this UAnimSequence. May be built directly or pulled from DDC
 #if WITH_EDITOR
-	UE_DEPRECATED(5.2, "Public access to ShouldPerformStripping has been deprecated")
-	ENGINE_API bool ShouldPerformStripping(const bool bPerformFrameStripping, const bool bPerformStrippingOnOddFramedAnims) const;
-
-	UE_DEPRECATED(5.2, "GetDDCCacheKeySuffix has been deprecated use CreateDerivedDataKeyHash instead")
-	ENGINE_API FString GetDDCCacheKeySuffix(const bool bPerformStripping, const ITargetPlatform* TargetPlatform) const;
-
-	UE_DEPRECATED(5.2, "ApplyCompressedData has been deprecated")
-	void ApplyCompressedData(const FString& DataCacheKeySuffix, const bool bPerformFrameStripping, const TArray<uint8>& Data) {}
-
 	ENGINE_API void WaitOnExistingCompression(const bool bWantResults=true);
-	
-	UE_DEPRECATED(5.2, "RequestAnimCompression has been deprecated use BeginCacheDerivedData or CacheDerivedData instead")
-	void RequestAnimCompression(
-	PRAGMA_DISABLE_DEPRECATION_WARNINGS
-		FRequestAnimCompressionParams Params) {}
-	PRAGMA_ENABLE_DEPRECATION_WARNINGS
-
-	UE_DEPRECATED(5.2, "RequestSyncAnimRecompression has been deprecated use CacheDerivedData instead")
-	ENGINE_API void RequestSyncAnimRecompression(bool bOutput = false);
-
-	UE_DEPRECATED(5.2, "RequestSyncAnimRecompression has been deprecated use BeginCacheDerivedData instead")
-	ENGINE_API void RequestAsyncAnimRecompression(bool bOutput = false);
+protected:
+	ENGINE_API bool ShouldPerformStripping(const bool bPerformFrameStripping, const bool bPerformStrippingOnOddFramedAnims) const;
 #endif
 
 protected:
@@ -656,27 +577,8 @@ public:
 	ENGINE_API bool IsCompressedDataValid() const;
 	ENGINE_API bool IsCurveCompressedDataValid() const;
 	// End Memory related functions
-
-	// Begin Utility functions
-
-	/**
-	* Get Skeleton Bone Index from Track Index for compressed data
-	*
-	* @param	TrackIndex		Track Index
-	*/
-	UE_DEPRECATED(5.2, "GetSkeletonIndexFromCompressedDataTrackIndex has been deprecated")
-	int32 GetSkeletonIndexFromCompressedDataTrackIndex(const int32 TrackIndex) const
-	{
-		PRAGMA_DISABLE_DEPRECATION_WARNINGS
-		return GetCompressedTrackToSkeletonMapTable()[TrackIndex].BoneTreeIndex;
-    	PRAGMA_ENABLE_DEPRECATION_WARNINGS
-	}
-	// End Utility functions
 	
 #if WITH_EDITOR
-	UE_DEPRECATED(5.2, "BakeTrackCurvesToRawAnimationTracks has been deprecated")
-	void BakeTrackCurvesToRawAnimationTracks(TArray<FRawAnimSequenceTrack>& NewRawTracks, TArray<FName>& NewTrackNames, TArray<FTrackToSkeletonMap>& NewTrackToSkeletonMapTable) {}
-
 	/**
 	 * Add Key to Transform Curves
 	 */
@@ -728,8 +630,6 @@ public:
 	ENGINE_API virtual void AdvanceMarkerPhaseAsFollower(const FMarkerTickContext& Context, float DeltaRemaining, bool bLooping, float& CurrentTime, FMarkerPair& PreviousMarker, FMarkerPair& NextMarker, const UMirrorDataTable* MirrorTable) const;
 	ENGINE_API virtual void GetMarkerIndicesForTime(float CurrentTime, bool bLooping, const TArray<FName>& ValidMarkerNames, FMarkerPair& OutPrevMarker, FMarkerPair& OutNextMarker) const;
 
-	UE_DEPRECATED(5.0, "Use other GetMarkerSyncPositionfromMarkerIndicies signature")
-	virtual FMarkerSyncAnimPosition GetMarkerSyncPositionfromMarkerIndicies(int32 PrevMarker, int32 NextMarker, float CurrentTime) const { return UAnimSequence::GetMarkerSyncPositionFromMarkerIndicies(PrevMarker, NextMarker, CurrentTime, nullptr); }
 	ENGINE_API virtual FMarkerSyncAnimPosition GetMarkerSyncPositionFromMarkerIndicies(int32 PrevMarker, int32 NextMarker, float CurrentTime, const UMirrorDataTable* MirrorTable) const;
 	ENGINE_API virtual void GetMarkerIndicesForPosition(const FMarkerSyncAnimPosition& SyncPosition, bool bLooping, FMarkerPair& OutPrevMarker, FMarkerPair& OutNextMarker, float& CurrentTime, const UMirrorDataTable* MirrorTable) const;
 	
@@ -755,18 +655,6 @@ public:
 	// Is this animation valid for baking into additive
 	ENGINE_API bool CanBakeAdditive() const;
 
-	// Bakes out track data for the skeletons virtual bones into the raw data
-	UE_DEPRECATED(5.2, "BakeOutVirtualBoneTracks has been deprecated")
-	void BakeOutVirtualBoneTracks(TArray<FRawAnimSequenceTrack>& NewRawTracks, TArray<FName>& NewAnimationTrackNames, TArray<FTrackToSkeletonMap>& NewTrackToSkeletonMapTable) {} 
-
-	// Performs multiple evaluations of the animation as a test of compressed data validatity
-	UE_DEPRECATED(5.2, "TestEvalauteAnimation has been deprecated")
-	void TestEvalauteAnimation() const {}
-
-	// Bakes out the additive version of this animation into the raw data.
-	UE_DEPRECATED(5.2, "BakeOutAdditiveIntoRawData has been deprecated, this behaviour has moved to FCompressibleAnimData")
-	void BakeOutAdditiveIntoRawData(TArray<FRawAnimSequenceTrack>& NewRawTracks, TArray<FName>& NewAnimationTrackNames, TArray<FTrackToSkeletonMap>& NewTrackToSkeletonMapTable, TArray<FFloatCurve>& NewCurveTracks, TArray<FRawAnimSequenceTrack>& AdditiveBaseAnimationData) const {}
-
 	// Test whether at any point we will scale a bone to 0 (needed for validating additive anims)
 	ENGINE_API bool DoesSequenceContainZeroScale() const;
 
@@ -776,35 +664,6 @@ public:
 	// Helper function to allow us to update streaming animations that depend on us with our data when we are updated
 	ENGINE_API void UpdateDependentStreamingAnimations() const;
 
-	// Generate a GUID from a hash of our own raw data
-	UE_DEPRECATED(5.1, "GenerateGuidFromRawData has been deprecated use IAnimationDataModel::GenerateGuid instead")
-	ENGINE_API FGuid GenerateGuidFromRawData() const;
-
-	// Should we be always using our raw data (i.e is our compressed data stale)
-	UE_DEPRECATED(5.2, "OnlyUseRawData has been deprecated")
-	bool OnlyUseRawData() const
-	{
-		PRAGMA_DISABLE_DEPRECATION_WARNINGS
-		return bUseRawDataOnly;
-		PRAGMA_ENABLE_DEPRECATION_WARNINGS
-	}
-	UE_DEPRECATED(5.2, "SetUseRawDataOnly has been deprecated")
-	void SetUseRawDataOnly(bool bInUseRawDataOnly)
-	{
-		PRAGMA_DISABLE_DEPRECATION_WARNINGS
-		bUseRawDataOnly = bInUseRawDataOnly;
-		PRAGMA_ENABLE_DEPRECATION_WARNINGS
-	}
-
-	// Return this animations guid for the raw data
-	UE_DEPRECATED(5.1, "GetRawDataGuid has been deprecated use GenerateGuidFromModel instead")
-	FGuid GetRawDataGuid() const
-	{ 
-		PRAGMA_DISABLE_DEPRECATION_WARNINGS
-		return RawDataGuid;
-		PRAGMA_ENABLE_DEPRECATION_WARNINGS
-	}
-
 	/** Resets Bone Animation, Curve data and Notify tracks **/
 	ENGINE_API void ResetAnimation();
 #endif
@@ -814,13 +673,9 @@ private:
 	* Get Bone Transform of the animation for the Time given, relative to Parent for all RequiredBones
 	* This return mesh rotation only additive pose
 	*
-	* @param	OutPose				[out] Output bone transforms
-	* @param	OutCurve			[out] Curves to fill	
-	* @param	ExtractionContext	Extraction Context (position, looping, root motion, etc.)
+	* @param	OutAnimationPoseData	[out] Output pose data
+	* @param	ExtractionContext		Extraction Context (position, looping, root motion, etc.)
 	*/
-	UE_DEPRECATED(4.26, "Use GetBonePose_AdditiveMeshRotationOnly with other signature")
-	ENGINE_API void GetBonePose_AdditiveMeshRotationOnly(FCompactPose& OutPose, FBlendedCurve& OutCurve, const FAnimExtractContext& ExtractionContext) const;
-
 	ENGINE_API void GetBonePose_AdditiveMeshRotationOnly(FAnimationPoseData& OutAnimationPoseData, const FAnimExtractContext& ExtractionContext) const;
 
 protected:
@@ -846,8 +701,8 @@ public:
 	ENGINE_API void ValidateCurrentPosition(const FMarkerSyncAnimPosition& Position, bool bPlayingForwards, bool bLooping, float&CurrentTime, FMarkerPair& PreviousMarker, FMarkerPair& NextMarker, const UMirrorDataTable* MirrorTable = nullptr) const;
 	ENGINE_API bool UseRawDataForPoseExtraction(const FBoneContainer& RequiredBones) const;
 
+protected:
 	// Should we be always using our raw data (i.e is our compressed data stale)
-	UE_DEPRECATED(5.2, "bUseRawDataOnly public access will be deprecated")
 	TAtomic<bool> bUseRawDataOnly;
 public:
 	/** Authored Sync markers */
@@ -858,32 +713,6 @@ public:
 	TArray<FName>				UniqueMarkerNames;
 
 public:
-#if WITH_EDITOR
-	UE_DEPRECATED(5.0, "AddBoneCustomAttribute has been deprecated see UAnimDataController::AddAttribute")
-	UFUNCTION(BlueprintCallable, Category=CustomAttributes, meta=(DeprecatedFunction, DeprecationMessage="AddBoneFloatCustomAttribute has been deprecated, use UAnimDataController::AddAttribute instead"))
-	ENGINE_API void AddBoneFloatCustomAttribute(const FName& BoneName, const FName& AttributeName, const TArray<float>& TimeKeys, const TArray<float>& ValueKeys);
-
-	UE_DEPRECATED(5.0, "AddBoneCustomAttribute has been deprecated see UAnimDataController::AddAttribute")
-	UFUNCTION(BlueprintCallable, Category = CustomAttributes, meta=(DeprecatedFunction, DeprecationMessage="AddBoneIntegerCustomAttribute has been deprecated, use UAnimDataController::AddAttribute instead"))
-	ENGINE_API void AddBoneIntegerCustomAttribute(const FName& BoneName, const FName& AttributeName, const TArray<float>& TimeKeys, const TArray<int32>& ValueKeys);
-
-	UE_DEPRECATED(5.0, "AddBoneStringCustomAttribute has been deprecated see UAnimDataController::AddAttribute")
-	UFUNCTION(BlueprintCallable, Category = CustomAttributes, meta=(DeprecatedFunction, DeprecationMessage="AddBoneStringCustomAttribute has been deprecated, use UAnimDataController::AddAttribute instead"))
-	ENGINE_API void AddBoneStringCustomAttribute(const FName& BoneName, const FName& AttributeName, const TArray<float>& TimeKeys, const TArray<FString>& ValueKeys);
-
-	UE_DEPRECATED(5.0, "RemoveCustomAttribute has been deprecated see UAnimDataController::RemoveAttribute")
-	UFUNCTION(BlueprintCallable, Category = CustomAttributes, meta=(DeprecatedFunction, DeprecationMessage="RemoveCustomAttribute has been deprecated, use UAnimDataController::RemoveAttribute instead"))
-	ENGINE_API void RemoveCustomAttribute(const FName& BoneName, const FName& AttributeName);
-
-	UE_DEPRECATED(5.0, "RemoveAllCustomAttributesForBone has been deprecated see UAnimDataController::RemoveAllAttributesForBone")
-	UFUNCTION(BlueprintCallable, Category = CustomAttributes, meta=(DeprecatedFunction, DeprecationMessage="RemoveAllCustomAttributesForBone has been deprecated, use UAnimDataController::RemoveAllAttributesForBone instead"))
-	ENGINE_API void RemoveAllCustomAttributesForBone(const FName& BoneName);
-
-	UE_DEPRECATED(5.0, "RemoveAllCustomAttributes has been deprecated see UAnimDataController::RemoveAllAttributes")
-	UFUNCTION(BlueprintCallable, Category = CustomAttributes, meta=(DeprecatedFunction, DeprecationMessage="RemoveAllCustomAttributes has been deprecated, use UAnimDataController::RemoveAllAttributes instead"))
-	ENGINE_API void RemoveAllCustomAttributes();
-#endif // WITH_EDITOR
-
 	ENGINE_API void EvaluateAttributes(FAnimationPoseData& OutAnimationPoseData, const FAnimExtractContext& ExtractionContext, bool bUseRawData) const;	
 protected:
 #if WITH_EDITOR
@@ -899,14 +728,7 @@ protected:
 	// End UAnimSequenceBase virtual overrides
 
 	ENGINE_API void EnsureValidRawDataGuid();
-	UE_DEPRECATED(5.2, "RecompressAnimationData has been deprecated")
-	ENGINE_API void RecompressAnimationData();
-	UE_DEPRECATED(5.2, "ResampleAnimationTrackData has been deprecated")
-	void ResampleAnimationTrackData() const {}
 	ENGINE_API void CalculateNumberOfSampledKeys();
-	UE_DEPRECATED(5.2, "ClearResampledAnimationTrackData has been deprecated")
-	void ClearResampledAnimationTrackData() {}
-
 	ENGINE_API void DeleteBoneAnimationData();
 	ENGINE_API void DeleteDeprecatedRawAnimationData();
 public:
diff --git a/Engine/Source/Runtime/Engine/Classes/Animation/AnimSequenceBase.h b/Engine/Source/Runtime/Engine/Classes/Animation/AnimSequenceBase.h
--- a/Engine/Source/Runtime/Engine/Classes/Animation/AnimSequenceBase.h
+++ b/Engine/Source/Runtime/Engine/Classes/Animation/AnimSequenceBase.h
@@ -99,18 +99,6 @@ public:
 	ENGINE_API void RenameNotifies(FName InOldName, FName InNewName);
 #endif
 	
-	/** 
-	 * Retrieves AnimNotifies given a StartTime and a DeltaTime.
-	 * Time will be advanced and support looping if bAllowLooping is true.
-	 * Supports playing backwards (DeltaTime<0).
-	 * Returns notifies between StartTime (exclusive) and StartTime+DeltaTime (inclusive)
-	 */
-	UE_DEPRECATED(4.19, "Use the GetAnimNotifies that takes FAnimNotifyEventReferences instead")
-	ENGINE_API void GetAnimNotifies(const float& StartTime, const float& DeltaTime, const bool bAllowLooping, TArray<const FAnimNotifyEvent *>& OutActiveNotifies) const;
-
-	UE_DEPRECATED(5.0, "Use the other GetAnimNotifies that takes FAnimNotifyContext instead")
-	ENGINE_API void GetAnimNotifies(const float& StartTime, const float& DeltaTime, const bool bAllowLooping, TArray<FAnimNotifyEventReference>& OutActiveNotifies) const;
-
 	/**
 	* Retrieves AnimNotifies given a StartTime and a DeltaTime.
 	* Time will be advanced and support looping if bAllowLooping is true.
@@ -119,17 +107,6 @@ public:
 	*/
 	ENGINE_API void GetAnimNotifies(const float& StartTime, const float& DeltaTime, FAnimNotifyContext& NotifyContext) const;
 
-	/** 
-	 * Retrieves AnimNotifies between two time positions. ]PreviousPosition, CurrentPosition]
-	 * Between PreviousPosition (exclusive) and CurrentPosition (inclusive).
-	 * Supports playing backwards (CurrentPosition<PreviousPosition).
-	 * Only supports contiguous range, does NOT support looping and wrapping over.
-	 */
-	UE_DEPRECATED(4.19, "Use the GetAnimNotifiesFromDeltaPositions that takes FAnimNotifyEventReferences instead")
-	ENGINE_API void GetAnimNotifiesFromDeltaPositions(const float& PreviousPosition, const float & CurrentPosition, TArray<const FAnimNotifyEvent *>& OutActiveNotifies) const;
-
-	UE_DEPRECATED(5.0, "Use the other GetAnimNotifiesFromDeltaPositions that takes FAnimNotifyContext instead")
-	ENGINE_API virtual void GetAnimNotifiesFromDeltaPositions(const float& PreviousPosition, const float & CurrentPosition, TArray<FAnimNotifyEventReference>& OutActiveNotifies) const;
 	/**
 	* Retrieves AnimNotifies between two time positions. ]PreviousPosition, CurrentPosition]
 	* Between PreviousPosition (exclusive) and CurrentPosition (inclusive).
@@ -142,25 +119,15 @@ public:
 	UE_DEPRECATED(5.6, "Please use EvaluateCurveData with FAnimExtractContext")
 	virtual void EvaluateCurveData(FBlendedCurve& OutCurve, float CurrentTime, bool bForceUseRawData = false) const { const FAnimExtractContext Context(static_cast<double>(CurrentTime)); EvaluateCurveData(OutCurve, Context, bForceUseRawData); }
 	ENGINE_API virtual void EvaluateCurveData(FBlendedCurve& OutCurve, const FAnimExtractContext& AnimExtractContext, bool bForceUseRawData = false) const;
-
-	UE_DEPRECATED(5.3, "Please use EvaluateCurveData that takes a FName.")
-	virtual float EvaluateCurveData(SmartName::UID_Type CurveUID, float CurrentTime, bool bForceUseRawData = false) const { return 0.0f; }
 	
 	UE_DEPRECATED(5.6, "Please use EvaluateCurveData with FAnimExtractContext")
 	virtual float EvaluateCurveData(FName CurveName, float CurrentTime, bool bForceUseRawData = false) const { const FAnimExtractContext Context(static_cast<double>(CurrentTime)); return EvaluateCurveData(CurveName, Context, bForceUseRawData); }
 	ENGINE_API virtual float EvaluateCurveData(FName CurveName, const FAnimExtractContext& AnimExtractContext, bool bForceUseRawData = false) const;
 	
 	ENGINE_API virtual const FRawCurveTracks& GetCurveData() const;
-	
-	UE_DEPRECATED(5.3, "Please use HasCurveData that takes a FName.")
-	virtual bool HasCurveData(SmartName::UID_Type CurveUID, bool bForceUseRawData = false) const { return false; }
-	
+		
 	ENGINE_API virtual bool HasCurveData(FName CurveName, bool bForceUseRawData = false) const;
 
-	/** Return Number of Keys **/
-	UE_DEPRECATED(4.19, "Use GetNumberOfSampledKeys instead")
-	ENGINE_API virtual int32 GetNumberOfFrames() const;
-
 	/** Return the total number of keys sampled for this animation, including the T0 key **/
 	ENGINE_API virtual int32 GetNumberOfSampledKeys() const;
 
@@ -214,10 +181,7 @@ public:
 	* @param	OutPose				Pose object to fill
 	* @param	OutCurve			Curves to fill
 	* @param	ExtractionContext	Extraction Context (position, looping, root motion, etc.)
-	*/
-	UE_DEPRECATED(4.26, "Use other GetAnimationPose signature")
-	ENGINE_API virtual void GetAnimationPose(struct FCompactPose& OutPose, FBlendedCurve& OutCurve, const FAnimExtractContext& ExtractionContext) const;
-	
+	*/	
 	ENGINE_API virtual void GetAnimationPose(FAnimationPoseData& OutPoseData, const FAnimExtractContext& ExtractionContext) const
 		PURE_VIRTUAL(UAnimSequenceBase::GetAnimationPose, );
 	
@@ -245,22 +209,14 @@ public:
 
 	ENGINE_API virtual void Serialize(FArchive& Ar) override;
 
-	UE_DEPRECATED(5.0, "Use other AdvanceMarkerPhaseAsLeader signature")
-	virtual void AdvanceMarkerPhaseAsLeader(bool bLooping, float MoveDelta, const TArray<FName>& ValidMarkerNames, float& CurrentTime, FMarkerPair& PrevMarker, FMarkerPair& NextMarker, TArray<FPassedMarker>& MarkersPassed) const {AdvanceMarkerPhaseAsLeader(bLooping, MoveDelta, ValidMarkerNames, CurrentTime, PrevMarker, NextMarker, MarkersPassed, nullptr); }
 	virtual void AdvanceMarkerPhaseAsLeader(bool bLooping, float MoveDelta, const TArray<FName>& ValidMarkerNames, float& CurrentTime, FMarkerPair& PrevMarker, FMarkerPair& NextMarker, TArray<FPassedMarker>& MarkersPassed, const UMirrorDataTable* MirrorTable) const { check(false); /*Should never call this (either missing override or calling on unsupported asset */ }
 
-	UE_DEPRECATED(5.0, "Use other AdvanceMarkerPhaseAsFollower signature")
-	virtual void AdvanceMarkerPhaseAsFollower(const FMarkerTickContext& Context, float DeltaRemaining, bool bLooping, float& CurrentTime, FMarkerPair& PreviousMarker, FMarkerPair& NextMarker) const { AdvanceMarkerPhaseAsFollower(Context, DeltaRemaining, bLooping, CurrentTime, PreviousMarker, NextMarker, nullptr);}
 	virtual void AdvanceMarkerPhaseAsFollower(const FMarkerTickContext& Context, float DeltaRemaining, bool bLooping, float& CurrentTime, FMarkerPair& PreviousMarker, FMarkerPair& NextMarker, const UMirrorDataTable* MirrorTable) const { check(false); /*Should never call this (either missing override or calling on unsupported asset */ }
 	
 	virtual void GetMarkerIndicesForTime(float CurrentTime, bool bLooping, const TArray<FName>& ValidMarkerNames, FMarkerPair& OutPrevMarker, FMarkerPair& OutNextMarker) const { check(false); /*Should never call this (either missing override or calling on unsupported asset */ }
 
-	UE_DEPRECATED(5.0, "Use other GetMarkerSyncPositionfromMarkerIndicies signature")
-	virtual FMarkerSyncAnimPosition GetMarkerSyncPositionfromMarkerIndicies(int32 PrevMarker, int32 NextMarker, float CurrentTime) const { return UAnimSequenceBase::GetMarkerSyncPositionFromMarkerIndicies(PrevMarker, NextMarker, CurrentTime, nullptr); }
 	virtual FMarkerSyncAnimPosition GetMarkerSyncPositionFromMarkerIndicies(int32 PrevMarker, int32 NextMarker, float CurrentTime, const UMirrorDataTable* MirrorTable) const { check(false); return FMarkerSyncAnimPosition(); /*Should never call this (either missing override or calling on unsupported asset */ }
 
-	UE_DEPRECATED(5.0, "Use other GetMarkerIndicesForPosition signature")
-	virtual void GetMarkerIndicesForPosition(const FMarkerSyncAnimPosition& SyncPosition, bool bLooping, FMarkerPair& OutPrevMarker, FMarkerPair& OutNextMarker, float& CurrentTime) const { GetMarkerIndicesForPosition(SyncPosition, bLooping, OutPrevMarker, OutNextMarker, CurrentTime, nullptr); }
 	virtual void GetMarkerIndicesForPosition(const FMarkerSyncAnimPosition& SyncPosition, bool bLooping, FMarkerPair& OutPrevMarker, FMarkerPair& OutNextMarker, float& CurrentTime, const UMirrorDataTable* MirrorTable ) const { check(false); /*Should never call this (either missing override or calling on unsupported asset */ }
 	
 	virtual float GetFirstMatchingPosFromMarkerSyncPos(const FMarkerSyncAnimPosition& InMarkerSyncGroupPosition) const { return 0.f; }
diff --git a/Engine/Source/Runtime/Engine/Private/Animation/AnimCompositeBase.cpp b/Engine/Source/Runtime/Engine/Private/Animation/AnimCompositeBase.cpp
--- a/Engine/Source/Runtime/Engine/Private/Animation/AnimCompositeBase.cpp
+++ b/Engine/Source/Runtime/Engine/Private/Animation/AnimCompositeBase.cpp
@@ -38,9 +38,7 @@ UAnimSequenceBase* FAnimSegment::GetAnimationData(float PositionInTrack, float&
 {
 	if( bValid && IsInRange(PositionInTrack) )
 	{
-		PRAGMA_DISABLE_DEPRECATION_WARNINGS
 		if(AnimReference)
-		PRAGMA_ENABLE_DEPRECATION_WARNINGS
 		{
 			const float ValidPlayRate = GetValidPlayRate();
 
@@ -64,9 +62,7 @@ UAnimSequenceBase* FAnimSegment::GetAnimationData(float PositionInTrack, float&
 				PositionInAnim = AnimEndTime + Delta * ValidPlayRate;
 			}
 
-			PRAGMA_DISABLE_DEPRECATION_WARNINGS
 			return AnimReference;
-			PRAGMA_ENABLE_DEPRECATION_WARNINGS
 		}
 	}
 
@@ -91,31 +87,6 @@ float FAnimSegment::ConvertTrackPosToAnimPos(const float& TrackPosition) const
 	return AnimPosition;
 }
 
-void FAnimSegment::GetAnimNotifiesFromTrackPositions(const float& PreviousTrackPosition, const float& CurrentTrackPosition, TArray<const FAnimNotifyEvent *> & OutActiveNotifies) const
-{
-	FAnimTickRecord TickRecord;
-	FAnimNotifyContext NotifyContext(TickRecord);
-	GetAnimNotifiesFromTrackPositions(PreviousTrackPosition, CurrentTrackPosition, NotifyContext);
-
-	OutActiveNotifies.Reset(NotifyContext.ActiveNotifies.Num());
-	for (FAnimNotifyEventReference& NotifyRef : NotifyContext.ActiveNotifies)
-	{
-		if (const FAnimNotifyEvent* Notify = NotifyRef.GetNotify())
-		{
-			OutActiveNotifies.Add(Notify);
-		}
-	}
-}
-
-void FAnimSegment::GetAnimNotifiesFromTrackPositions(const float& PreviousTrackPosition, const float& CurrentTrackPosition, TArray<FAnimNotifyEventReference> & OutActiveNotifies) const
-{
-	FAnimTickRecord TickRecord;
-	FAnimNotifyContext NotifyContext(TickRecord);
-	GetAnimNotifiesFromTrackPositions(PreviousTrackPosition, CurrentTrackPosition, NotifyContext);
-	// Slow copy due assumption of calling code that OutActiveNotifies is only extended
-	OutActiveNotifies.Append(NotifyContext.ActiveNotifies);
-}
-
 void FAnimSegment::GetAnimNotifiesFromTrackPositions(const float& PreviousTrackPosition, const float& CurrentTrackPosition, FAnimNotifyContext& NotifyContext) const
 {
 	const bool bTrackPlayingBackwards = (PreviousTrackPosition > CurrentTrackPosition);
@@ -130,9 +101,7 @@ void FAnimSegment::GetAnimNotifiesFromTrackPositions(const float& PreviousTrackP
 		)
 	{
 		// Only allow AnimSequences for now. Other types will need additional support.
-		PRAGMA_DISABLE_DEPRECATION_WARNINGS
 		UAnimSequenceBase* AnimSequenceBase = AnimReference;
-		PRAGMA_ENABLE_DEPRECATION_WARNINGS
 		if(AnimSequenceBase)
 		{
 			const float ValidPlayRate = GetValidPlayRate();
@@ -196,9 +165,7 @@ void FAnimSegment::GetRootMotionExtractionStepsForTrackRange(TArray<FRootMotionE
 		return;
 	}
 
-	PRAGMA_DISABLE_DEPRECATION_WARNINGS
 	if (!bValid || !AnimReference)
-	PRAGMA_ENABLE_DEPRECATION_WARNINGS
 	{
 		return;
 	}
@@ -235,10 +202,8 @@ void FAnimSegment::GetRootMotionExtractionStepsForTrackRange(TArray<FRootMotionE
 		const float AnimEndPoint = bAnimPlayingBackwards ? AnimStartTime : AnimEndTime;
 
 		// Only allow AnimSequences for now. Other types will need additional support.
-		PRAGMA_DISABLE_DEPRECATION_WARNINGS
 		UAnimSequence* AnimSequence = Cast<UAnimSequence>(AnimReference);
 		UAnimComposite* AnimComposite = Cast<UAnimComposite>(AnimReference);
-		PRAGMA_ENABLE_DEPRECATION_WARNINGS
 
 		if (AnimSequence || AnimComposite)
 		{
@@ -417,8 +382,7 @@ int32 FAnimTrack::GetTrackAdditiveType() const
 	if( AnimSegments.Num() > 0 )
 	{
 		const FAnimSegment& AnimSegment = AnimSegments[0];
-		const UAnimSequenceBase* AnimReference = AnimSegment.GetAnimReference();
-		if (AnimReference)
+		if (const UAnimSequenceBase* AnimReference = AnimSegment.GetAnimReference())
 		{
 			return AnimReference->GetAdditiveAnimType();
 		}
@@ -570,13 +534,6 @@ void FAnimTrack::SortAnimSegments()
 }
 #endif
 
-void FAnimTrack::GetAnimationPose(/*out*/ FCompactPose& OutPose, /*out*/ FBlendedCurve& OutCurve, const FAnimExtractContext& ExtractionContext) const
-{
-	UE::Anim::FStackAttributeContainer TempAttributes;
-	FAnimationPoseData OutAnimationPoseData(OutPose, OutCurve, TempAttributes);
-	GetAnimationPose(OutAnimationPoseData, ExtractionContext);
-}
-
 void FAnimTrack::GetAnimationPose(FAnimationPoseData& OutAnimationPoseData, const FAnimExtractContext& ExtractionContext) const
 {
 	bool bExtractedPose = false;
@@ -668,36 +625,6 @@ bool FAnimTrack::ContainRecursive(const TArray<UAnimCompositeBase*>& CurrentAccu
 	return false;
 }
 
-void FAnimTrack::GetAnimNotifiesFromTrackPositions(const float& PreviousTrackPosition, const float& CurrentTrackPosition, TArray<const FAnimNotifyEvent *> & OutActiveNotifies) const
-{
-	FAnimTickRecord TickRecord;
-	FAnimNotifyContext NotifyContext(TickRecord);
-	GetAnimNotifiesFromTrackPositions(PreviousTrackPosition, CurrentTrackPosition, NotifyContext);
-
-	OutActiveNotifies.Reset(NotifyContext.ActiveNotifies.Num());
-	for (FAnimNotifyEventReference& NotifyRef : NotifyContext.ActiveNotifies)
-	{
-		if (const FAnimNotifyEvent* Notify = NotifyRef.GetNotify())
-		{
-			OutActiveNotifies.Add(Notify);
-		}
-	}
-}
-
-void FAnimTrack::GetAnimNotifiesFromTrackPositions(const float& PreviousTrackPosition, const float& CurrentTrackPosition, TArray<FAnimNotifyEventReference> & OutActiveNotifies) const
-{
-	FAnimTickRecord TickRecord;
-	FAnimNotifyContext NotifyContext(TickRecord);
-	for (int32 SegmentIndex = 0; SegmentIndex<AnimSegments.Num(); ++SegmentIndex)
-	{
-		if (AnimSegments[SegmentIndex].IsValid())
-		{
-			AnimSegments[SegmentIndex].GetAnimNotifiesFromTrackPositions(PreviousTrackPosition, CurrentTrackPosition, NotifyContext);
-		}
-	}
-	Swap(OutActiveNotifies, NotifyContext.ActiveNotifies);
-}
-
 void FAnimTrack::GetAnimNotifiesFromTrackPositions(const float& PreviousTrackPosition, const float& CurrentTrackPosition, FAnimNotifyContext& NotifyContext) const
 {
 	for (int32 SegmentIndex = 0; SegmentIndex<AnimSegments.Num(); ++SegmentIndex)
@@ -827,22 +754,16 @@ void UAnimCompositeBase::PostLoad()
 
 void FAnimSegment::SetAnimReference(UAnimSequenceBase* InAnimReference, bool bInitialize /*= false*/)
 {
-	PRAGMA_DISABLE_DEPRECATION_WARNINGS
 	AnimReference = InAnimReference;
-	PRAGMA_ENABLE_DEPRECATION_WARNINGS
 
 #if WITH_EDITOR
 	UpdateCachedPlayLength();
 #endif // WITH_EDITOR
 
-	PRAGMA_DISABLE_DEPRECATION_WARNINGS
 	if (AnimReference && bInitialize)
-	PRAGMA_ENABLE_DEPRECATION_WARNINGS
 	{		
 		AnimStartTime = 0.f;
-		PRAGMA_DISABLE_DEPRECATION_WARNINGS
 		AnimEndTime = AnimReference->GetPlayLength();
-		PRAGMA_ENABLE_DEPRECATION_WARNINGS
 		AnimPlayRate = 1.f;
 		LoopingCount = 1;
 		StartPos = 0.f;
@@ -852,7 +773,6 @@ void FAnimSegment::SetAnimReference(UAnimSequenceBase* InAnimReference, bool bIn
 #if WITH_EDITOR
 bool FAnimSegment::IsPlayLengthOutOfDate() const
 {
-	PRAGMA_DISABLE_DEPRECATION_WARNINGS
 	if (AnimReference && !FMath::IsNearlyZero(CachedPlayLength))
 	{
 		// When the segment length is equal to _cached_ playlength and the current model playlength is different flag as out-of-date
@@ -860,7 +780,6 @@ bool FAnimSegment::IsPlayLengthOutOfDate() const
 		const float PlayableLength = (AnimEndTime - AnimStartTime);
 		return FMath::IsNearlyEqual(PlayableLength, CachedPlayLength, UE_KINDA_SMALL_NUMBER) && !FMath::IsNearlyEqual(AnimReference->GetPlayLength(), CachedPlayLength, UE_KINDA_SMALL_NUMBER);
 	}
-	PRAGMA_ENABLE_DEPRECATION_WARNINGS
 	
 	return false;
 }
@@ -868,9 +787,7 @@ bool FAnimSegment::IsPlayLengthOutOfDate() const
 void FAnimSegment::UpdateCachedPlayLength()
 {
 	CachedPlayLength = 0.f;
-	PRAGMA_DISABLE_DEPRECATION_WARNINGS
 	const IAnimationDataModel* DataModel = AnimReference ? AnimReference->GetDataModel() : nullptr;	
-	PRAGMA_ENABLE_DEPRECATION_WARNINGS
 	if(DataModel)
 	{
 		CachedPlayLength = DataModel->GetPlayLength();
diff --git a/Engine/Source/Runtime/Engine/Private/Animation/AnimSequence.cpp b/Engine/Source/Runtime/Engine/Private/Animation/AnimSequence.cpp
--- a/Engine/Source/Runtime/Engine/Private/Animation/AnimSequence.cpp
+++ b/Engine/Source/Runtime/Engine/Private/Animation/AnimSequence.cpp
@@ -487,9 +487,7 @@ void UAnimSequence::WillNeverCacheCookedPlatformDataAgain()
 {
 	if (GAllowClearingCompressedDataDuringCook)
 	{
-		PRAGMA_DISABLE_DEPRECATION_WARNINGS
 		bUseRawDataOnly = true;
-		PRAGMA_ENABLE_DEPRECATION_WARNINGS
 
 		UE::Anim::FAnimSequenceCompilingManager::Get().FinishCompilation({this});
 		// Clear out current platform, and any target platform data
@@ -504,9 +502,7 @@ void UAnimSequence::ClearAllCachedCookedPlatformData()
 {
 	if (GAllowClearingCompressedDataDuringCook)
 	{
-		PRAGMA_DISABLE_DEPRECATION_WARNINGS
 		bUseRawDataOnly = true;
-		PRAGMA_ENABLE_DEPRECATION_WARNINGS
 		
 		// Delete any cache tasks first because the destructor will cancel the cache and build tasks,
 		// and drop their pointers to the data.
@@ -745,7 +741,6 @@ void UAnimSequence::Serialize(FArchive& Ar)
 
 		if (bSerializeCompressedData)
 		{
-			PRAGMA_DISABLE_DEPRECATION_WARNINGS
 			SerializeCompressedData(Ar,false);
 			if (!bIsTransacting)
 			{
@@ -753,7 +748,6 @@ void UAnimSequence::Serialize(FArchive& Ar)
 				Ar << bTemp;
 				bUseRawDataOnly = bTemp;
 			}
-			PRAGMA_ENABLE_DEPRECATION_WARNINGS
 		}
 	}
 
@@ -1270,11 +1264,7 @@ bool UAnimSequence::IsCachedCookedPlatformDataLoaded(const ITargetPlatform* Targ
 
 	if (KeyHash == DataKeyHash)
 	{
-		if (!CacheTasksByKeyHash.Contains(KeyHash) && CompressedData.IsValid(this) &&
-			PRAGMA_DISABLE_DEPRECATION_WARNINGS
-			!bUseRawDataOnly
-			PRAGMA_ENABLE_DEPRECATION_WARNINGS
-			)
+		if (!CacheTasksByKeyHash.Contains(KeyHash) && CompressedData.IsValid(this) && !bUseRawDataOnly)
 		{
 			return true;
 		}
@@ -1429,11 +1419,7 @@ FTransform UAnimSequence::ExtractRootTrackTransform(const FAnimExtractContext& E
 	const bool bContainsRootBoneTrack = [this, RootBoneIndex]() 
 	{
 #if WITH_EDITOR
-		if (CanEvaluateRawAnimationData() &&
-			PRAGMA_DISABLE_DEPRECATION_WARNINGS
-			bUseRawDataOnly
-			PRAGMA_ENABLE_DEPRECATION_WARNINGS
-			)
+		if (CanEvaluateRawAnimationData() && bUseRawDataOnly)
 		{
 			ValidateModel();
 
@@ -1461,10 +1447,7 @@ FTransform UAnimSequence::ExtractRootTrackTransform(const FAnimExtractContext& E
 	{
 		// if we do have root data, then return root data
 		FTransform RootTransform;
-		GetBoneTransform(RootTransform, FSkeletonPoseBoneIndex(RootBoneIndex), ExtractionContext, 
-			PRAGMA_DISABLE_DEPRECATION_WARNINGS
-			bUseRawDataOnly
-			PRAGMA_ENABLE_DEPRECATION_WARNINGS);
+		GetBoneTransform(RootTransform, FSkeletonPoseBoneIndex(RootBoneIndex), ExtractionContext, bUseRawDataOnly);
 		return RootTransform;
 	}
 
@@ -1675,13 +1658,6 @@ void UAnimSequence::GetAnimationPose(FAnimationPoseData& OutAnimationPoseData, c
 #endif
 }
 
-void UAnimSequence::GetBonePose(FCompactPose& OutPose, FBlendedCurve& OutCurve, const FAnimExtractContext& ExtractionContext, bool bForceUseRawData) const
-{
-	UE::Anim::FStackAttributeContainer TempAttributes;
-	FAnimationPoseData OutAnimationPoseData(OutPose, OutCurve, TempAttributes);
-	GetBonePose(OutAnimationPoseData, ExtractionContext, bForceUseRawData);
-}
-
 void UAnimSequence::GetBonePose(FAnimationPoseData& OutAnimationPoseData, const FAnimExtractContext& ExtractionContext, bool bForceUseRawData /*= false*/) const
 {
 	SCOPE_CYCLE_COUNTER(STAT_AnimSeq_GetBonePose);
@@ -1805,20 +1781,6 @@ void UAnimSequence::GetBonePose(FAnimationPoseData& OutAnimationPoseData, const
 	EvaluateAttributes(OutAnimationPoseData, ExtractionContext, false);
 }
 
-const TArray<FSmartName>& UAnimSequence::GetCompressedCurveNames() const
-{
-	static TArray<FSmartName> Dummy;
-	return Dummy;
-}
-
-void UAnimSequence::GetBonePose_Additive(FCompactPose& OutPose, FBlendedCurve& OutCurve, const FAnimExtractContext& ExtractionContext) const
-{
-	UE::Anim::FStackAttributeContainer TempAttributes;
-	FAnimationPoseData OutAnimationPoseData(OutPose, OutCurve, TempAttributes);
-
-	GetBonePose_Additive(OutAnimationPoseData, ExtractionContext);
-}
-
 void UAnimSequence::GetBonePose_Additive(FAnimationPoseData& OutAnimationPoseData, const FAnimExtractContext& ExtractionContext) const
 {
 	FCompactPose& OutPose = OutAnimationPoseData.GetPose();
@@ -1853,14 +1815,6 @@ void UAnimSequence::GetBonePose_Additive(FAnimationPoseData& OutAnimationPoseDat
 	UE::Anim::Attributes::ConvertToAdditive(BaseAttributes, OutAttributes);
 }
 
-void UAnimSequence::GetAdditiveBasePose(FCompactPose& OutPose, FBlendedCurve& OutCurve, const FAnimExtractContext& ExtractionContext) const
-{
-	UE::Anim::FStackAttributeContainer TempAttributes;
-	FAnimationPoseData OutAnimationPoseData(OutPose, OutCurve, TempAttributes);
-
-	GetAdditiveBasePose(OutAnimationPoseData, ExtractionContext);
-}
-
 static void GetSequencePose(FAnimationPoseData& OutAnimationPoseData, const FAnimExtractContext& ExtractionContext, const UAnimSequence &Seq, int32 FrameIndex, bool bForceUseRawData)
 {
 	const double Fraction = (Seq.GetNumberOfSampledKeys() > 0) ? FMath::Clamp<double>((double)FrameIndex / (double)Seq.GetNumberOfSampledKeys(), 0.0, 1.0) : 0.0;
@@ -1908,14 +1862,6 @@ void UAnimSequence::GetAdditiveBasePose(FAnimationPoseData& OutAnimationPoseData
 	}
 }
 
-void UAnimSequence::GetBonePose_AdditiveMeshRotationOnly(FCompactPose& OutPose, FBlendedCurve& OutCurve, const FAnimExtractContext& ExtractionContext) const
-{
-	UE::Anim::FStackAttributeContainer TempAttributes;
-	FAnimationPoseData OutAnimationPoseData(OutPose, OutCurve, TempAttributes);
-
-	GetBonePose_AdditiveMeshRotationOnly(OutAnimationPoseData, ExtractionContext);
-}
-
 void UAnimSequence::GetBonePose_AdditiveMeshRotationOnly(FAnimationPoseData& OutAnimationPoseData, const FAnimExtractContext& ExtractionContext) const
 {
 	FCompactPose& OutPose = OutAnimationPoseData.GetPose();
@@ -2082,13 +2028,6 @@ bool UAnimSequence::ShouldPerformStripping(const bool bPerformFrameStripping, co
 	return bStripCandidate && bShouldPerformFrameStripping;
 }
 
-FString UAnimSequence::GetDDCCacheKeySuffix(const bool bPerformStripping, const ITargetPlatform* TargetPlatform) const
-{
-	return GetAnimSequenceSpecificCacheKeySuffix(*this, bPerformStripping, CompressionErrorThresholdScale, TargetPlatform);
-}
-#endif
-
-#if WITH_EDITOR
 void UAnimSequence::WaitOnExistingCompression(const bool bWantResults)
 {
 	check(IsInGameThread());
@@ -2212,13 +2151,6 @@ bool UAnimSequence::DoesSequenceContainZeroScale() const
 	return false;
 }
 
-FGuid UAnimSequence::GenerateGuidFromRawData() const
-{
-PRAGMA_DISABLE_DEPRECATION_WARNINGS
-	return GenerateGuidFromRawAnimData(RawAnimationData, RawCurveData);
-PRAGMA_ENABLE_DEPRECATION_WARNINGS	
-}
-
 void UAnimSequence::FlagDependentAnimationsAsRawDataOnly() const
 {
 	for (TObjectIterator<UAnimSequence> Iter; Iter; ++Iter)
@@ -2226,9 +2158,7 @@ void UAnimSequence::FlagDependentAnimationsAsRawDataOnly() const
 		UAnimSequence* Seq = *Iter;
 		if (Seq->RefPoseSeq == this)
 		{
-			PRAGMA_DISABLE_DEPRECATION_WARNINGS
 			Seq->bUseRawDataOnly = true;
-			PRAGMA_ENABLE_DEPRECATION_WARNINGS
 		}
 	}
 }
@@ -2737,11 +2667,7 @@ uint8* UAnimSequence::FindSyncMarkerPropertyData(int32 SyncMarkerIndex, FArrayPr
 
 bool UAnimSequence::DoesNeedRecompress() const
 {
-	return GetSkeleton() && (
-		PRAGMA_DISABLE_DEPRECATION_WARNINGS
-		bUseRawDataOnly
-		PRAGMA_ENABLE_DEPRECATION_WARNINGS
-		|| (GetSkeletonVirtualBoneGuid() != GetSkeleton()->GetVirtualBoneGuid()));
+	return GetSkeleton() && (bUseRawDataOnly || (GetSkeletonVirtualBoneGuid() != GetSkeleton()->GetVirtualBoneGuid()));
 }
 
 bool UAnimSequence::CreateAnimation(USkeletalMesh* Mesh)
@@ -2876,10 +2802,7 @@ void UAnimSequence::EvaluateCurveData(FBlendedCurve& OutCurve, const FAnimExtrac
 {
 	SCOPE_CYCLE_COUNTER(STAT_AnimSeq_EvalCurveData);
 	
-	const bool bEvaluateRawData =
-	    PRAGMA_DISABLE_DEPRECATION_WARNINGS
-	    bUseRawDataOnly
-        PRAGMA_ENABLE_DEPRECATION_WARNINGS
+	const bool bEvaluateRawData = bUseRawDataOnly
 	    || bForceUseRawData || !IsCurveCompressedDataValid();
 	check(!bForceUseRawData || CanEvaluateRawAnimationData());
 	if (CanEvaluateRawAnimationData() && bEvaluateRawData)
@@ -2901,10 +2824,7 @@ float UAnimSequence::EvaluateCurveData(FName CurveName, const FAnimExtractContex
 {
 	QUICK_SCOPE_CYCLE_COUNTER(EvaluateCurveDataByName);
 
-	const bool bEvaluateRawData =
-	 PRAGMA_DISABLE_DEPRECATION_WARNINGS
-     bUseRawDataOnly
-     PRAGMA_ENABLE_DEPRECATION_WARNINGS
+	const bool bEvaluateRawData = bUseRawDataOnly
 	  || bForceUseRawData || !IsCurveCompressedDataValid();
 	check(!bForceUseRawData || CanEvaluateRawAnimationData());
 	if (CanEvaluateRawAnimationData() && bEvaluateRawData)
@@ -2921,10 +2841,7 @@ float UAnimSequence::EvaluateCurveData(FName CurveName, const FAnimExtractContex
 
 bool UAnimSequence::HasCurveData(FName CurveName, bool bForceUseRawData) const
 {
-	const bool bEvaluateRawData =
-	PRAGMA_DISABLE_DEPRECATION_WARNINGS
-		bUseRawDataOnly
-	PRAGMA_ENABLE_DEPRECATION_WARNINGS
+	const bool bEvaluateRawData = bUseRawDataOnly
 		|| bForceUseRawData || !IsCurveCompressedDataValid();
 	check(!bForceUseRawData || CanEvaluateRawAnimationData());
 	if (CanEvaluateRawAnimationData() && bEvaluateRawData)
@@ -3308,34 +3225,13 @@ void UAnimSequence::ValidateCurrentPosition(const FMarkerSyncAnimPosition& Posit
 
 bool UAnimSequence::UseRawDataForPoseExtraction(const FBoneContainer& RequiredBones) const
 {
-	return CanEvaluateRawAnimationData() && (
-	PRAGMA_DISABLE_DEPRECATION_WARNINGS
-	bUseRawDataOnly
-	PRAGMA_ENABLE_DEPRECATION_WARNINGS
-	|| (GetSkeletonVirtualBoneGuid() != GetSkeleton()->GetVirtualBoneGuid()) || RequiredBones.GetDisableRetargeting() || RequiredBones.ShouldUseRawData() ||
+	return CanEvaluateRawAnimationData() && (bUseRawDataOnly || (GetSkeletonVirtualBoneGuid() != GetSkeleton()->GetVirtualBoneGuid()) || RequiredBones.GetDisableRetargeting() || RequiredBones.ShouldUseRawData() ||
  #if WITH_EDITOR
   	GForceRawData == 1 ||
  #endif // WITH_EDITOR
 	RequiredBones.ShouldUseSourceData());
 }
 
-#if WITH_EDITOR
-void UAnimSequence::AddBoneFloatCustomAttribute(const FName& BoneName, const FName& AttributeName, const TArray<float>& TimeKeys, const TArray<float>& ValueKeys)
-{
-	UE::Anim::AddTypedCustomAttribute<FFloatAnimationAttribute, float>(AttributeName, BoneName, this, MakeArrayView(TimeKeys), MakeArrayView(ValueKeys));
-}
-
-void UAnimSequence::AddBoneIntegerCustomAttribute(const FName& BoneName, const FName& AttributeName, const TArray<float>& TimeKeys, const TArray<int32>& ValueKeys)
-{
-	UE::Anim::AddTypedCustomAttribute<FIntegerAnimationAttribute, int32>(AttributeName, BoneName, this, MakeArrayView(TimeKeys), MakeArrayView(ValueKeys));
-}
-
-void UAnimSequence::AddBoneStringCustomAttribute(const FName& BoneName, const FName& AttributeName, const TArray<float>& TimeKeys, const TArray<FString>& ValueKeys)
-{
-	UE::Anim::AddTypedCustomAttribute<FStringAnimationAttribute, FString>(AttributeName, BoneName, this, MakeArrayView(TimeKeys), MakeArrayView(ValueKeys));
-}
-#endif // WITH_EDITOR
-
 void UAnimSequence::EvaluateAttributes(FAnimationPoseData& OutAnimationPoseData, const FAnimExtractContext& ExtractionContext, bool bUseRawData) const
 {
 	QUICK_SCOPE_CYCLE_COUNTER(STAT_EvaluateAttributes);
@@ -3380,45 +3276,6 @@ void UAnimSequence::OnSetSkeleton(USkeleton* NewSkeleton)
 	WaitOnExistingCompression(false);
 }
 
-void UAnimSequence::RemoveCustomAttribute(const FName& BoneName, const FName& AttributeName)
-{
-	ValidateModel();
-
-	TArray<FAnimationAttributeIdentifier> IdentifiersToRemove;
-
-	Algo::TransformIf(DataModelInterface->GetAttributes(), IdentifiersToRemove, 
-		[BoneName, AttributeName](const FAnimatedBoneAttribute& Attribute)
-		{
-			return Attribute.Identifier.GetBoneName() == BoneName && Attribute.Identifier.GetName() == AttributeName;
-		},
-		[] (const FAnimatedBoneAttribute& Attribute)
-		{
-			return Attribute.Identifier;
-		}
-	);
-
-	if (IdentifiersToRemove.Num())
-	{
-		IAnimationDataController::FScopedBracket Bracket(Controller, LOCTEXT("RemoveCustomAttributeDeprecated", "Removing Custom Attributes by Bone and Name"));
-		for (const FAnimationAttributeIdentifier& AttributeIdentifier : IdentifiersToRemove)
-		{
-			Controller->RemoveAttribute(AttributeIdentifier);
-		}
-	}
-}
-
-void UAnimSequence::RemoveAllCustomAttributesForBone(const FName& BoneName)
-{
-	ValidateModel();
-	Controller->RemoveAllAttributesForBone(BoneName);
-}
-
-void UAnimSequence::RemoveAllCustomAttributes()
-{
-	ValidateModel();
-	Controller->RemoveAllAttributes();
-}
-
 void UAnimSequence::SynchronousAnimatedBoneAttributesCompression()
 {
 	struct FByFramePoseEvalContext
@@ -4358,9 +4215,7 @@ void UAnimSequence::OnModelModified(const EAnimDataModelNotifyType& NotifyType,
 		PRAGMA_DISABLE_DEPRECATION_WARNINGS
 		RawDataGuid = bForceNewRawDataGuid ? FGuid::NewGuid() : Model->GenerateGuid();
 		PRAGMA_ENABLE_DEPRECATION_WARNINGS
-		PRAGMA_DISABLE_DEPRECATION_WARNINGS
         bUseRawDataOnly = true;
-        PRAGMA_ENABLE_DEPRECATION_WARNINGS
 
 		FlagDependentAnimationsAsRawDataOnly();
 		UpdateDependentStreamingAnimations();
@@ -4675,9 +4530,7 @@ FString UAnimSequence::CreateDerivedDataKeyString(const ITargetPlatform* TargetP
 			}
 		}
 
-		PRAGMA_DISABLE_DEPRECATION_WARNINGS
 		bPerformFrameStripping = ShouldPerformStripping(bPerformFrameStripping, bPerformFrameStrippingOnOddNumberedFrames);
-		PRAGMA_ENABLE_DEPRECATION_WARNINGS
 	}
 
 	FArcToHexString ArcToHexString;
@@ -4788,9 +4641,7 @@ FIoHash UAnimSequence::BeginCacheDerivedData(const ITargetPlatform* TargetPlatfo
 	// Mark sequence to only sample raw animation data when recompressing for running platform (or equivalent hash)
 	if (KeyHash == DataKeyHash)
 	{
-		PRAGMA_DISABLE_DEPRECATION_WARNINGS
 		bUseRawDataOnly = true;
-		PRAGMA_ENABLE_DEPRECATION_WARNINGS
 	}
 	
 	check(TargetData);
@@ -4814,9 +4665,7 @@ FIoHash UAnimSequence::BeginCacheDerivedData(const ITargetPlatform* TargetPlatfo
 			}
 		}
 
-		PRAGMA_DISABLE_DEPRECATION_WARNINGS
 		bPerformFrameStripping = ShouldPerformStripping(bPerformFrameStripping, bPerformFrameStrippingOnOddNumberedFrames);
-		PRAGMA_ENABLE_DEPRECATION_WARNINGS
 	}
 
 	// Always keep AnimatedBoneAttributes in lock-step when requesting compression
@@ -4980,10 +4829,8 @@ void UAnimSequence::FinishAsyncTasks()
 						FAssetRegistryModule& AssetRegistryModule = FModuleManager::LoadModuleChecked<FAssetRegistryModule>(TEXT("AssetRegistry"));
 						AssetRegistryModule.Get().AssetTagsFinalized(*this);
 					
-						PRAGMA_DISABLE_DEPRECATION_WARNINGS
 						check(bUseRawDataOnly);
 						bUseRawDataOnly = false;
-						PRAGMA_ENABLE_DEPRECATION_WARNINGS
 					}
 				
 					check(TaskData->IsValid(this, true));
@@ -5008,28 +4855,6 @@ void UAnimSequence::Reschedule(FQueuedThreadPool* InThreadPool, EQueuedWorkPrior
 		Pair.Value->Reschedule(InThreadPool, InPriority);
 	}
 }
-
-void UAnimSequence::RecompressAnimationData()
-{	
-	ClearAllCachedCookedPlatformData();
-	BeginCacheDerivedDataForCurrentPlatform();
-}
-
-void UAnimSequence::RequestAsyncAnimRecompression(bool bOutput)
-{ 
-	if (ITargetPlatform* RunningPlatform = GetTargetPlatformManagerRef().GetRunningTargetPlatform())
-	{
-		BeginCacheDerivedData(RunningPlatform);
-	}
-}
-
-void UAnimSequence::RequestSyncAnimRecompression(bool bOutput)
-{
-	if (ITargetPlatform* RunningPlatform = GetTargetPlatformManagerRef().GetRunningTargetPlatform())
-	{
-		CacheDerivedData(RunningPlatform);
-	}
-}
 #endif // WITH_EDITOR
 
 /*-----------------------------------------------------------------------------
diff --git a/Engine/Source/Runtime/Engine/Private/Animation/AnimSequenceBase.cpp b/Engine/Source/Runtime/Engine/Private/Animation/AnimSequenceBase.cpp
--- a/Engine/Source/Runtime/Engine/Private/Animation/AnimSequenceBase.cpp
+++ b/Engine/Source/Runtime/Engine/Private/Animation/AnimSequenceBase.cpp
@@ -381,32 +381,6 @@ bool UAnimSequenceBase::IsNotifyAvailable() const
 	return (Notifies.Num() != 0) && (GetPlayLength() > 0.f);
 }
 
-void UAnimSequenceBase::GetAnimNotifies(const float& StartTime, const float& DeltaTime, const bool bAllowLooping, TArray<const FAnimNotifyEvent *> & OutActiveNotifies) const
-{
-	FAnimTickRecord TickRecord;
-	TickRecord.bLooping = bAllowLooping; 
-	FAnimNotifyContext NotifyContext(TickRecord);
-	GetAnimNotifies(StartTime, DeltaTime, NotifyContext);
-
-	OutActiveNotifies.Reset(NotifyContext.ActiveNotifies.Num());
-	for (FAnimNotifyEventReference NotifyRef : NotifyContext.ActiveNotifies)
-	{
-		if (const FAnimNotifyEvent* Notify = NotifyRef.GetNotify())
-		{
-			OutActiveNotifies.Add(Notify);
-		}
-	}
-}
-
-void UAnimSequenceBase::GetAnimNotifies(const float& StartTime, const float& DeltaTime, const bool bAllowLooping, TArray<FAnimNotifyEventReference>& OutActiveNotifies) const
-{
-	FAnimTickRecord TickRecord;
-	TickRecord.bLooping = bAllowLooping;
-	FAnimNotifyContext NotifyContext(TickRecord);
-	GetAnimNotifies(StartTime, DeltaTime, NotifyContext);
-	Swap(NotifyContext.ActiveNotifies, OutActiveNotifies);
-}
-
 void UAnimSequenceBase::GetAnimNotifies(const float& StartTime, const float& DeltaTime, FAnimNotifyContext& NotifyContext) const
 {
 	// Early out if we have no notifies
@@ -454,30 +428,6 @@ void UAnimSequenceBase::GetAnimNotifies(const float& StartTime, const float& Del
 	}
 }
 
-void UAnimSequenceBase::GetAnimNotifiesFromDeltaPositions(const float& PreviousPosition, const float& CurrentPosition, TArray<const FAnimNotifyEvent *> & OutActiveNotifies) const
-{
-	FAnimTickRecord TickRecord;
-	FAnimNotifyContext NotifyContext(TickRecord);
-	GetAnimNotifiesFromDeltaPositions(PreviousPosition, CurrentPosition, NotifyContext);
-
-	OutActiveNotifies.Reset(NotifyContext.ActiveNotifies.Num());
-	for (FAnimNotifyEventReference NotifyRef : NotifyContext.ActiveNotifies)
-	{
-		if (const FAnimNotifyEvent* Notify = NotifyRef.GetNotify())
-		{
-			OutActiveNotifies.Add(Notify);
-		}
-	}
-}
-
-void UAnimSequenceBase::GetAnimNotifiesFromDeltaPositions(const float& PreviousPosition, const float & CurrentPosition, TArray<FAnimNotifyEventReference>& OutActiveNotifies) const
-{
-	FAnimTickRecord TickRecord;
-	FAnimNotifyContext NotifyContext(TickRecord);
-	GetAnimNotifiesFromDeltaPositions(PreviousPosition, CurrentPosition,NotifyContext);
-	Swap(NotifyContext.ActiveNotifies, OutActiveNotifies);
-}
-
 void UAnimSequenceBase::GetAnimNotifiesFromDeltaPositions(const float& PreviousPosition, const float& CurrentPosition,  FAnimNotifyContext& NotifyContext) const
 {
 	// Early out if we have no notifies
@@ -809,11 +759,6 @@ void UAnimSequenceBase::RefreshCacheData()
 #endif //WITH_EDITOR
 }
 
-int32 UAnimSequenceBase::GetNumberOfFrames() const
-{
-	return GetNumberOfSampledKeys();
-}
-
 int32 UAnimSequenceBase::GetNumberOfSampledKeys() const
 {
 	return GetSamplingFrameRate().AsFrameTime(GetPlayLength()).RoundToFrame().Value;
@@ -1128,13 +1073,6 @@ void UAnimSequenceBase::Serialize(FArchive& Ar)
 	PRAGMA_ENABLE_DEPRECATION_WARNINGS
 }
 
-void UAnimSequenceBase::GetAnimationPose(struct FCompactPose& OutPose, FBlendedCurve & OutCurve, const FAnimExtractContext& ExtractionContext) const
-{
-	UE::Anim::FStackAttributeContainer TempAttributes;
-	FAnimationPoseData OutAnimationPoseData(OutPose, OutCurve, TempAttributes);
-	GetAnimationPose(OutAnimationPoseData, ExtractionContext);
-}
-
 void UAnimSequenceBase::HandleAssetPlayerTickedInternal(FAnimAssetTickContext &Context, const float PreviousTime, const float MoveDelta, const FAnimTickRecord &Instance, struct FAnimNotifyQueue& NotifyQueue) const
 {
 	// Harvest and record notifies
