Updated the PSO Precaching for the Material Billboard Component to also trigger when any updates to the component are made via the AddElement, SetElements, SetMaterial, and PostLoad function calls. This is done to bring it more in line with how other components utilize the PSO Precaching mechanism.

diff --git a/Engine/Source/Runtime/Engine/Classes/Components/MaterialBillboardComponent.h b/Engine/Source/Runtime/Engine/Classes/Components/MaterialBillboardComponent.h
--- a/Engine/Source/Runtime/Engine/Classes/Components/MaterialBillboardComponent.h
+++ b/Engine/Source/Runtime/Engine/Classes/Components/MaterialBillboardComponent.h
@@ -93,6 +93,9 @@ class UMaterialBillboardComponent : public UPrimitiveComponent
 	ENGINE_API virtual void GetUsedMaterials(TArray<UMaterialInterface*>& OutMaterials, bool bGetDebugMaterials = false) const override;
 	//~ End UPrimitiveComponent Interface
 
+	/** Collect all the PSO precache data used by the material billboard component */
+	ENGINE_API virtual void CollectPSOPrecacheData(const FPSOPrecacheParams& BasePrecachePSOParams, FMaterialInterfacePSOPrecacheParamsList& OutParams) override;
+
 	//~ Begin USceneComponent Interface
 #if WITH_EDITOR
 	ENGINE_API virtual bool GetMaterialPropertyPath(int32 ElementIndex, UObject*& OutOwner, FString& OutPropertyPath, FProperty*& OutProperty) override;
diff --git a/Engine/Source/Runtime/Engine/Private/Components/MaterialBillboardComponent.cpp b/Engine/Source/Runtime/Engine/Private/Components/MaterialBillboardComponent.cpp
--- a/Engine/Source/Runtime/Engine/Private/Components/MaterialBillboardComponent.cpp
+++ b/Engine/Source/Runtime/Engine/Private/Components/MaterialBillboardComponent.cpp
@@ -331,6 +331,8 @@ void UMaterialBillboardComponent::AddElement(
 	Element->BaseSizeY = BaseSizeY;
 	Element->DistanceToSizeCurve = DistanceToSizeCurve;
 
+	PrecachePSOs();
+
 	MarkRenderStateDirty();
 }
 
@@ -339,6 +341,8 @@ void UMaterialBillboardComponent::SetElements(const TArray<FMaterialSpriteElemen
 	// Replace existing array
 	Elements = NewElements;
 
+	PrecachePSOs();
+
 	// Indicate scene proxy needs to be updated
 	MarkRenderStateDirty();
 }
@@ -358,27 +362,7 @@ void UMaterialBillboardComponent::PostLoad()
 {
 	Super::PostLoad();
 
-	if ((IsComponentPSOPrecachingEnabled() && RHISupportsManualVertexFetch(GMaxRHIShaderPlatform)))
-	{
-		FPSOPrecacheParams PrecachePSOParams;
-		SetupPrecachePSOParams(PrecachePSOParams);
-		PrecachePSOParams.PrimitiveType = PT_TriangleStrip;
-		PrecachePSOParams.bDisableBackFaceCulling = true;
-
-		const FVertexFactoryType* VFType = &FLocalVertexFactory::StaticType;
-		FPSOPrecacheVertexFactoryDataList VFDataList;
-		VFDataList.Add(FPSOPrecacheVertexFactoryData(VFType));
-
-		TArray<UMaterialInterface*> UsedMaterials;
-		GetUsedMaterials(UsedMaterials, false);
-		for (UMaterialInterface* MaterialInterface : UsedMaterials)
-		{
-			if (MaterialInterface)
-			{
-				MaterialInterface->PrecachePSOs(VFType, PrecachePSOParams);
-			}
-		}
-	}
+	PrecachePSOs();
 }
 
 void UMaterialBillboardComponent::SetMaterial(int32 ElementIndex, class UMaterialInterface* Material)
@@ -387,6 +371,8 @@ void UMaterialBillboardComponent::SetMaterial(int32 ElementIndex, class UMateria
 	{
 		Elements[ElementIndex].Material = Material;
 
+		PrecachePSOs();
+
 		MarkRenderStateDirty();
 	}
 }
@@ -398,3 +384,26 @@ void UMaterialBillboardComponent::GetUsedMaterials(TArray<UMaterialInterface*>&
 		OutMaterials.AddUnique(GetMaterial(ElementIndex));
 	}
 }
+
+void UMaterialBillboardComponent::CollectPSOPrecacheData(const FPSOPrecacheParams& BasePrecachePSOParams, FMaterialInterfacePSOPrecacheParamsList& OutParams)
+{
+	FPSOPrecacheParams PrecachePSOParams = BasePrecachePSOParams;
+	PrecachePSOParams.PrimitiveType = PT_TriangleStrip;
+	PrecachePSOParams.bDisableBackFaceCulling = true;
+ 
+	const FVertexFactoryType* VFType = &FLocalVertexFactory::StaticType;
+ 
+	TArray<UMaterialInterface*> UsedMaterials;
+	GetUsedMaterials(UsedMaterials, false);
+	for (UMaterialInterface* MaterialInterface : UsedMaterials)
+	{
+		if (MaterialInterface)
+		{
+			FMaterialInterfacePSOPrecacheParams& ComponentParams = OutParams[OutParams.AddDefaulted()];
+			ComponentParams.MaterialInterface = MaterialInterface;
+			ComponentParams.PSOPrecacheParams = PrecachePSOParams;
+			ComponentParams.VertexFactoryDataList.Add(FPSOPrecacheVertexFactoryData(VFType));
+		}
+	}
+}
+
