[LevelInstanceActorDesc] Don't attempt to set runtime bounds if GetLevelInstanceBoundsFromPackage call was not successful

diff --git a/Engine/Source/Runtime/Engine/Private/WorldPartition/LevelInstance/LevelInstanceActorDesc.cpp b/Engine/Source/Runtime/Engine/Private/WorldPartition/LevelInstance/LevelInstanceActorDesc.cpp
--- a/Engine/Source/Runtime/Engine/Private/WorldPartition/LevelInstance/LevelInstanceActorDesc.cpp
+++ b/Engine/Source/Runtime/Engine/Private/WorldPartition/LevelInstance/LevelInstanceActorDesc.cpp
@@ -402,9 +402,11 @@ void FLevelInstanceActorDesc::Serialize(FArchive& Ar)
 			{
 				if (!IsChildContainerInstance())
 				{
-					FBox NewRuntimeBounds = RuntimeBounds;
-					ULevelInstanceSubsystem::GetLevelInstanceBoundsFromPackage(ActorTransform, GetChildContainerPackage(), NewRuntimeBounds);
-					SetRuntimeBounds(NewRuntimeBounds);
+					FBox NewRuntimeBounds;
+					if (ULevelInstanceSubsystem::GetLevelInstanceBoundsFromPackage(ActorTransform, GetChildContainerPackage(), NewRuntimeBounds))
+					{
+						SetRuntimeBounds(NewRuntimeBounds);
+					}
 				}
 			}
 		}
