GameModeBase DispatchPostLogin has been deprecated and the behavior was moved to OnPostLogin, so subclasses can override it.

There are some cases where Network Control Messages (e.g. NMT_Join) happen outside of ActorTicks.  That causes issues if migrations need to happen during those calls (one of them being PostLogin).  This is a refactor to ensure we can wrap those in AutoRTFM Transactions if needed.

diff --git a/Engine/Source/Runtime/Engine/Classes/GameFramework/GameModeBase.h b/Engine/Source/Runtime/Engine/Classes/GameFramework/GameModeBase.h
--- a/Engine/Source/Runtime/Engine/Classes/GameFramework/GameModeBase.h
+++ b/Engine/Source/Runtime/Engine/Classes/GameFramework/GameModeBase.h
@@ -326,11 +326,12 @@ public:
 	ENGINE_API virtual void PostLogin(APlayerController* NewPlayer);
 
 	/** Called as part of the PostLogin process.  This is the last step before we spawn a new player. */
+	UE_DEPRECATED(5.6, "This functionality has been moved into OnPostLogin. You cannot call that function directly, it is executed at the end of PostLogin.")
 	ENGINE_API void DispatchPostLogin(AController* NewPlayer);
 
 protected:
-	/** Called as part of DispatchPostLogin */
-	virtual void OnPostLogin(AController* NewPlayer) { }
+	/** Called as part of the PostLogin process.  This is the last step before we spawn a new player. */
+	ENGINE_API virtual void OnPostLogin(AController* NewPlayer);
 
 public:
 	/** Notification that a player has successfully logged in, and has been given a player controller */
diff --git a/Engine/Source/Runtime/Engine/Private/GameModeBase.cpp b/Engine/Source/Runtime/Engine/Private/GameModeBase.cpp
--- a/Engine/Source/Runtime/Engine/Private/GameModeBase.cpp
+++ b/Engine/Source/Runtime/Engine/Private/GameModeBase.cpp
@@ -23,6 +23,7 @@
 #include "EngineUtils.h"
 #include "Kismet/GameplayStatics.h"
 #include "Engine/LevelStreaming.h"
+#include "UObject/RemoteObjectTransfer.h"
 
 #include UE_INLINE_GENERATED_CPP_BY_NAME(GameModeBase)
 
@@ -1034,21 +1035,23 @@ void AGameModeBase::PostLogin(APlayerController* NewPlayer)
 		GameSession->PostLogin(NewPlayer);
 	}
 
-	DispatchPostLogin(NewPlayer);
+	OnPostLogin(NewPlayer);
 
 	// Now that initialization is done, try to spawn the player's pawn and start match
 	HandleStartingNewPlayer(NewPlayer);
 }
 
 void AGameModeBase::DispatchPostLogin(AController* NewPlayer)
+{
+}
+
+void AGameModeBase::OnPostLogin(AController* NewPlayer)
 {
 	if (APlayerController* NewPC = Cast<APlayerController>(NewPlayer))
 	{
 		K2_PostLogin(NewPC);
 		FGameModeEvents::GameModePostLoginEvent.Broadcast(this, NewPC);
 	}
-
-	OnPostLogin(NewPlayer);
 }
 
 void AGameModeBase::Logout(AController* Exiting)
