Cluster Union Component:

Reset velocities on teleport with reset physics flag.

#tests PIE

diff --git a/Engine/Source/Runtime/Engine/Private/PhysicsEngine/ClusterUnionComponent.cpp b/Engine/Source/Runtime/Engine/Private/PhysicsEngine/ClusterUnionComponent.cpp
--- a/Engine/Source/Runtime/Engine/Private/PhysicsEngine/ClusterUnionComponent.cpp
+++ b/Engine/Source/Runtime/Engine/Private/PhysicsEngine/ClusterUnionComponent.cpp
@@ -818,6 +818,11 @@ void UClusterUnionComponent::OnUpdateTransform(EUpdateTransformFlags UpdateTrans
 		// This code path is generally used when setting the transform manually or when it's set via replication.
 		const FTransform Transform = GetComponentTransform();
 		PhysicsProxy->SetXR_External(Transform.GetLocation(), Transform.GetRotation());
+		if (Teleport == ETeleportType::ResetPhysics)
+		{
+			PhysicsProxy->SetLinearVelocity_External(FVector::ZeroVector);
+			PhysicsProxy->SetAngularVelocity_External(FVector::ZeroVector);
+		}
 	}
 }
 
