Disable r.RayTracing.DynamicInstances.Parallel to unblock people because of crashes in landscape.

#rnx

diff --git a/Engine/Source/Runtime/Engine/Private/PrimitiveSceneProxy.cpp b/Engine/Source/Runtime/Engine/Private/PrimitiveSceneProxy.cpp
--- a/Engine/Source/Runtime/Engine/Private/PrimitiveSceneProxy.cpp
+++ b/Engine/Source/Runtime/Engine/Private/PrimitiveSceneProxy.cpp
@@ -112,7 +112,7 @@ bool IsParallelGatherDynamicMeshElementsEnabled()
 		GSupportsParallelRenderingTasksWithSeparateRHIThread;
 }
 
-bool GParallelGatherDynamicRayTracingInstances = true;
+bool GParallelGatherDynamicRayTracingInstances = false;
 static FAutoConsoleVariableRef CVarParallelGatherDynamicRayTracingInstances(
 	TEXT("r.RayTracing.DynamicInstances.Parallel"),
 	GParallelGatherDynamicRayTracingInstances,
