Change to FProperty::ImportSingleProperty which means we can process Enums in Quotes within a Struct when loading from a config. This approach means that we identify Enums within Structs (and Structs in Arrays) and take away the first quote before passing it to ImportEnumFromBuffer later on. This means that ImportEnumFromBuffer can assume that the enum is unquoted and will allow it to process bitwise enums within quotes the same way whether it came from a Struct or not.

Additionally, this includes Unit Tests which prove that the parsing is well-behaved in various scenarios (including empty quotes or missing quotes).

#tests Tetsed editor config loading, and also passed new Unit Tests
[RN] When loading config files, Structs now support Quoted Bitwise Enums of the form "EMyEnum::A | EMyEnum::B"

diff --git a/Engine/Source/Runtime/CoreUObject/Private/UObject/Property.cpp b/Engine/Source/Runtime/CoreUObject/Private/UObject/Property.cpp
--- a/Engine/Source/Runtime/CoreUObject/Private/UObject/Property.cpp
+++ b/Engine/Source/Runtime/CoreUObject/Private/UObject/Property.cpp
@@ -2247,6 +2247,54 @@ const TCHAR* FProperty::ImportSingleProperty( const TCHAR* Str, void* DestData,
 						Str = Result;
 					}
 				}
+				else if (FEnumProperty* const EnumProperty = ExactCastField<FEnumProperty>(Property))
+				{
+					//We need to support the case where a bitmask enum is surrounded by Quotes. Non-struct properties
+					//remove quotes before passing to ImportText, so we do the same here by checking for a leading
+					//quote, and then consuming the closing quote after importing the enum property string
+
+					if (Index == INDEX_NONE)
+					{
+						Index = 0;
+					}
+
+					SkipWhitespace(Str);
+					
+					bool bQuoteToBeConsumed = false;
+					if (*Str == TCHAR('"'))
+					{
+						Str++;
+						bQuoteToBeConsumed = true;
+					}
+
+					const TCHAR* Result = ImportText(Str, Property, Property->ContainerPtrToValuePtr<void>(DestData, Index));
+
+#if WITH_EDITORONLY_DATA
+					UE::FInitializedPropertyValueState(ObjectStruct, DestData).Set(Property, Index);
+#endif
+
+					// in the failure case, don't return NULL so the caller can potentially skip less and get values further in the string
+					if (Result != nullptr)
+					{
+						if (bQuoteToBeConsumed)
+						{
+							SkipWhitespace(Result);
+							if (*Result == TCHAR('"'))
+							{
+								Str = Result + 1;
+							}
+							else
+							{
+								Warn->Logf(ELogVerbosity::Warning, TEXT("Property delineated by quotes is missing ending quote: %s"), Start);
+								Str = Result;
+							}
+						}
+						else
+						{
+							Str = Result;
+						}
+					}
+				}
 				else
 				{
 					if (Index == INDEX_NONE)
diff --git a/Engine/Source/Runtime/Engine/Private/Tests/TextImportTest.cpp b/Engine/Source/Runtime/Engine/Private/Tests/TextImportTest.cpp
new file mode 100644
--- /dev/null
+++ b/Engine/Source/Runtime/Engine/Private/Tests/TextImportTest.cpp
@@ -0,0 +1,70 @@
+// Copyright Epic Games, Inc. All Rights Reserved.
+
+#include "Tests/TextImportTest.h"
+#include "Tests/TestHarnessAdapter.h"
+#include "UObject/UnrealType.h"
+
+#if WITH_TESTS
+
+TEST_CASE_NAMED(FTextImportStructPropertyTest, "System::Engine::TextImport::StructProperty", "[Engine][TextImport][StructProperty]")
+{
+	struct FTextImportTestCase
+	{
+		bool bShouldError = false;
+		FTextImportTestStruct TargetOutput = FTextImportTestStruct();
+		FStringView TextInput = FStringView();
+	};
+
+	FTextImportTestCase TestCases[] =
+	{
+		{ false, { ETextImportTestFlags::FlagA, 2, "String,With,Commas" }, TEXTVIEW("(EmbeddedFlags=\"ETextImportTestFlags::FlagA\",TestInt=2, TestString=\"String,With,Commas\")")},
+		{ false, { ETextImportTestFlags::FlagA, 2, "String,With,Commas" }, TEXTVIEW("(EmbeddedFlags=\"FlagA \",TestInt=2, TestString=\"String,With,Commas\")")},
+		{ false, { ETextImportTestFlags::FlagA | ETextImportTestFlags::FlagB, 2, "String,With,Commas" }, TEXTVIEW("(EmbeddedFlags=\"ETextImportTestFlags::FlagA | ETextImportTestFlags::FlagB\",TestInt=2, TestString=\"String,With,Commas\")")},
+		{ false, { ETextImportTestFlags::FlagA | ETextImportTestFlags::FlagB, 2, "String,With,Commas" }, TEXTVIEW("(EmbeddedFlags=\"FlagA | FlagB\",TestInt=2, TestString=\"String,With,Commas\")")},
+		{ false, { ETextImportTestFlags::FlagA | ETextImportTestFlags::FlagB, 2, "String,With,Commas" }, TEXTVIEW("(EmbeddedFlags=ETextImportTestFlags::FlagA | ETextImportTestFlags::FlagB,TestInt=2, TestString=\"String,With,Commas\")")},
+		{ false, { ETextImportTestFlags::FlagA | ETextImportTestFlags::FlagB, 2, "String,With,Commas" }, TEXTVIEW("(EmbeddedFlags=FlagA | FlagB,TestInt=2, TestString=\"String,With,Commas\")")},
+		{ true, { ETextImportTestFlags::FlagA | ETextImportTestFlags::FlagB, 2, "String,With,Commas" }, TEXTVIEW("(EmbeddedFlags=\"FlagA | FlagB,TestInt=2, TestString=\"String,With,Commas\")")}, //Missing end quote
+		{ true, { ETextImportTestFlags::FlagA | ETextImportTestFlags::FlagB, 2, "String,With,Commas" }, TEXTVIEW("(EmbeddedFlags=FlagA | FlagB\",TestInt=2, TestString=\"String,With,Commas\")")}, //Missing start quote
+		{ true, { ETextImportTestFlags::TestStructDefault, 1, "" }, TEXTVIEW("(EmbeddedFlags=\"\",TestInt=\"\", TestString=\"\")")}, //Empty Quotes - Should give invalid enum error (same behaviour as for non-struct enums)
+		{ false, { ETextImportTestFlags::TestStructDefault, 1, "DefaultString" }, TEXTVIEW("(EmbeddedFlags=,TestInt=, TestString=)")}, //Empty No-Quotes - Returns default Struct 
+
+		{ false, { ETextImportTestFlags::TestStructDefault, 1, "DefaultString" }, TEXTVIEW("(WrongNameA=\"ETextImportTestFlags::FlagA\",WrongNameB=2, WrongNameC=\"String,With,Commas\")")},
+		{ true, { ETextImportTestFlags::TestStructDefault, 1, "" }, TEXTVIEW("(EmbeddedFlags=(,TestInt=, TestString=)")}, //Extra Bracket - Should give invalid enum error
+		{ true, { ETextImportTestFlags::TestStructDefault, 1, "" }, TEXTVIEW("(EmbeddedFlags=\",TestInt=, TestString=)")}, //Extra Quote - Should give Bad quoted string error
+	};
+	
+
+	UTextImportContainer* Container = NewObject<UTextImportContainer>();
+	check(Container);
+	const FStructProperty* ObjProp = FindFProperty<FStructProperty>(Container->GetClass(), TEXT("ResultStruct"));
+	check(ObjProp);
+	check(ObjProp->Struct);
+
+	int32 PortFlags = 0;
+
+	for (const FTextImportTestCase& TestCase : TestCases)
+	{
+		//Reset struct to defaults
+		Container->ResultStruct = FTextImportTestStruct();
+
+		//Parse the text into the ResultStruct
+		const TCHAR* InBuffer = TestCase.TextInput.GetData();
+		FStringOutputDevice ImportError;
+		ObjProp->Struct->ImportText(InBuffer, (void*)&Container->ResultStruct, Container, PortFlags, &ImportError, FString("FTextImportTestStruct"), true);
+
+		const bool bImportErrored = ImportError.Len() > 0;
+		if (!TestCase.bShouldError && bImportErrored)
+		{
+			FString ErrorString = FString::Printf(TEXT("Test case failed due to ImportError %s for Test Case: %s"), *ImportError, InBuffer);
+			AddError(ErrorString);
+		}
+		
+		CHECK_MESSAGE(FString::Printf(TEXT("Import Error Mismatch. Result: %s Expected: %s for TestCase: %s"), bImportErrored ? TEXT("Error") : TEXT("NoError"), TestCase.bShouldError ? TEXT("Error") : TEXT("NoError"), InBuffer), TestCase.bShouldError == bImportErrored);
+		if (!TestCase.bShouldError)
+		{
+			CHECK_MESSAGE(FString::Printf(TEXT("Import Results are Incorrect: %s Expected: %s for TestCase: %s"), *Container->ResultStruct.ToString(), *TestCase.TargetOutput.ToString(), InBuffer), TestCase.TargetOutput == Container->ResultStruct);
+		}
+	}
+}
+
+#endif //WITH_TESTS
\ No newline at end of file
diff --git a/Engine/Source/Runtime/Engine/Public/Tests/TextImportTest.h b/Engine/Source/Runtime/Engine/Public/Tests/TextImportTest.h
new file mode 100644
--- /dev/null
+++ b/Engine/Source/Runtime/Engine/Public/Tests/TextImportTest.h
@@ -0,0 +1,54 @@
+// Copyright Epic Games, Inc. All Rights Reserved.
+
+#include "Engine/Engine.h"
+#include "TextImportTest.generated.h"
+
+UENUM(Flags)
+enum class ETextImportTestFlags : uint32
+{
+	Default = 0,
+	FlagA = 1 << 0,
+	FlagB = 1 << 1,
+	FlagC = 1 << 2,
+	FlagD = 1 << 3,
+	FlagE = 1 << 4,
+	TestStructDefault = 1 << 5,
+};
+
+ENUM_CLASS_FLAGS(ETextImportTestFlags)
+
+USTRUCT()
+struct FTextImportTestStruct
+{
+	GENERATED_BODY()
+
+	UPROPERTY()
+	ETextImportTestFlags EmbeddedFlags = ETextImportTestFlags::TestStructDefault;
+
+	UPROPERTY()
+	int32 TestInt = 1;
+
+	UPROPERTY()
+	FString TestString = FString("DefaultString");
+
+	inline bool operator==(const FTextImportTestStruct& OtherStruct)
+	{
+		return EmbeddedFlags == OtherStruct.EmbeddedFlags && TestInt == OtherStruct.TestInt && TestString == OtherStruct.TestString;
+	}
+
+	inline FString ToString() const
+	{
+		return FString::Printf(TEXT("EmbeddedFlags: %d TestInt: %d TestString: %s"), static_cast<int32>(EmbeddedFlags), TestInt, *TestString);
+	}
+};
+
+
+UCLASS()
+class UTextImportContainer : public UObject
+{
+	GENERATED_BODY()
+
+public:
+	UPROPERTY()
+	FTextImportTestStruct ResultStruct;
+};
