Make UDataTable work correctly in Verse transactions.

diff --git a/Engine/Source/Runtime/Engine/Classes/Engine/DataTable.h b/Engine/Source/Runtime/Engine/Classes/Engine/DataTable.h
--- a/Engine/Source/Runtime/Engine/Classes/Engine/DataTable.h
+++ b/Engine/Source/Runtime/Engine/Classes/Engine/DataTable.h
@@ -10,6 +10,7 @@
 #include "UObject/Class.h"
 #include "UObject/UnrealType.h"
 #include "UObject/PropertyPortFlags.h"
+#include "Misc/TransactionallySafeCriticalSection.h"
 #include "DataTableUtils.h"
 #include "DataTable.generated.h"
 
@@ -428,7 +429,7 @@ protected:
 		UDataTable* Table;
 
 		static TMap<UDataTable*, int32> ScopeCount;
-		static FCriticalSection CriticalSection;
+		static FTransactionallySafeCriticalSection CriticalSection;
 	};
 };
 
diff --git a/Engine/Source/Runtime/Engine/Private/DataTable.cpp b/Engine/Source/Runtime/Engine/Private/DataTable.cpp
--- a/Engine/Source/Runtime/Engine/Private/DataTable.cpp
+++ b/Engine/Source/Runtime/Engine/Private/DataTable.cpp
@@ -48,14 +48,14 @@ namespace
 UDataTable::FScopedDataTableChange::FScopedDataTableChange(UDataTable* InTable)
 	: Table(InTable)
 {
-	FScopeLock Lock(&CriticalSection);
+	UE::TScopeLock Lock(CriticalSection);
 	int32& Count = ScopeCount.FindOrAdd(Table);
 	++Count;
 }
 
 UDataTable::FScopedDataTableChange::~FScopedDataTableChange()
 {
-	FScopeLock Lock(&CriticalSection);
+	UE::TScopeLock Lock(CriticalSection);
 	int32& Count = ScopeCount.FindChecked(Table);
 	--Count;
 	if (Count == 0)
@@ -66,7 +66,7 @@ UDataTable::FScopedDataTableChange::~FScopedDataTableChange()
 }
 
 TMap<UDataTable*, int32> UDataTable::FScopedDataTableChange::ScopeCount;
-FCriticalSection UDataTable::FScopedDataTableChange::CriticalSection;
+FTransactionallySafeCriticalSection UDataTable::FScopedDataTableChange::CriticalSection;
 
 #define DATATABLE_CHANGE_SCOPE()	UDataTable::FScopedDataTableChange ActiveScope(this);
 
diff --git a/Engine/Source/Runtime/Engine/Private/Tests/AutoRTFM/AutoRTFMDataTableTests.cpp b/Engine/Source/Runtime/Engine/Private/Tests/AutoRTFM/AutoRTFMDataTableTests.cpp
new file mode 100644
--- /dev/null
+++ b/Engine/Source/Runtime/Engine/Private/Tests/AutoRTFM/AutoRTFMDataTableTests.cpp
@@ -0,0 +1,38 @@
+// Copyright Epic Games, Inc. All Rights Reserved.
+
+#include "Misc/AutomationTest.h"
+#include "AutoRTFM/AutoRTFM.h"
+#include "Engine/DataTable.h"
+
+#if WITH_DEV_AUTOMATION_TESTS
+
+IMPLEMENT_SIMPLE_AUTOMATION_TEST(FAutoRTFMDataTableTests, "AutoRTFM + UDataTable", EAutomationTestFlags::EngineFilter | EAutomationTestFlags::ClientContext | EAutomationTestFlags::ServerContext | EAutomationTestFlags::CommandletContext)
+
+bool FAutoRTFMDataTableTests::RunTest(const FString & Parameters)
+{
+	if (!AutoRTFM::ForTheRuntime::IsAutoRTFMRuntimeEnabled())
+	{
+		ExecutionInfo.AddEvent(FAutomationEvent(EAutomationEventType::Info, TEXT("SKIPPED 'FAutoRTFMDataTableTests' test. AutoRTFM disabled.")));
+		return true;
+	}
+
+	UDataTable* Object = NewObject<UDataTable>();
+	Object->RowStruct = NewObject<UScriptStruct>();
+
+	AutoRTFM::ETransactionResult Result = AutoRTFM::Transact([&]
+	{
+		Object->EmptyTable();
+		AutoRTFM::AbortTransaction();
+	});
+	TestTrueExpr(AutoRTFM::ETransactionResult::AbortedByRequest == Result);
+
+	Result = AutoRTFM::Transact([&]
+	{
+		Object->EmptyTable();
+	});
+	TestTrueExpr(AutoRTFM::ETransactionResult::Committed == Result);
+
+	return true;
+}
+
+#endif //WITH_DEV_AUTOMATION_TESTS
