Adding pragma once to TextImportTest in order to fix a build compilation issue in UE Tests

#rnx

diff --git a/Engine/Source/Runtime/Engine/Public/Tests/TextImportTest.h b/Engine/Source/Runtime/Engine/Public/Tests/TextImportTest.h
--- a/Engine/Source/Runtime/Engine/Public/Tests/TextImportTest.h
+++ b/Engine/Source/Runtime/Engine/Public/Tests/TextImportTest.h
@@ -1,5 +1,7 @@
 // Copyright Epic Games, Inc. All Rights Reserved.
 
+#pragma once
+
 #include "Engine/Engine.h"
 #include "TextImportTest.generated.h"
 
