Fixing DP overriding user setting scalability.
Clamping user setting scalability to the experience-specific max quality levels.

#rnx

diff --git a/Engine/Source/Runtime/Engine/Classes/DeviceProfiles/DeviceProfileManager.h b/Engine/Source/Runtime/Engine/Classes/DeviceProfiles/DeviceProfileManager.h
--- a/Engine/Source/Runtime/Engine/Classes/DeviceProfiles/DeviceProfileManager.h
+++ b/Engine/Source/Runtime/Engine/Classes/DeviceProfiles/DeviceProfileManager.h
@@ -22,6 +22,11 @@ DECLARE_MULTICAST_DELEGATE( FOnDeviceProfileManagerUpdated );
 // Delegate used to notify systems when the active device profile changes
 DECLARE_MULTICAST_DELEGATE( FOnActiveDeviceProfileChanged );
 
+// Delegate used to notify systems pre reapply device profiles
+DECLARE_MULTICAST_DELEGATE( FWillReapplyDeviceProfiles );
+// Delegate used to notify systems post reapplied device profiles
+DECLARE_MULTICAST_DELEGATE( FOnReapplyDeviceProfiles );
+
 struct FPushedCVarSetting
 {
 	FPushedCVarSetting() : SetBy(ECVF_Default)	{}
@@ -153,6 +158,20 @@ public:
 	 */
 	ENGINE_API FOnActiveDeviceProfileChanged& OnActiveDeviceProfileChanged();
 
+	/**
+	 * Returns a delegate that is invoked pre reapply device profiles
+	 *
+	 * @return The delegate.
+	 */
+	ENGINE_API FWillReapplyDeviceProfiles& WillReapplyDeviceProfiles();
+
+	/**
+	 * Returns a delegate that is invoked post reapply device profiles
+	 *
+	 * @return The delegate.
+	 */
+	ENGINE_API FOnReapplyDeviceProfiles& OnReapplyDeviceProfiles();
+
 	/**
 	 * Returns the config files for all loaded device profiles, this will include platform-specific ones 
 	 */
@@ -333,6 +352,11 @@ private:
 	// Holds a delegate to be invoked when the active deviceprofile changes
 	FOnActiveDeviceProfileChanged ActiveDeviceProfileChangedDelegate;
 
+	// Holds a delegate to be invoked pre reapply device profiles
+	FWillReapplyDeviceProfiles PreReapplyDeviceProfilesDelegate;
+	// Holds a delegate to be invoked post reapply device profiles
+	FOnReapplyDeviceProfiles PostReapplyDeviceProfilesDelegate;
+
 	// Holds the selected device profile
 	UDeviceProfile* ActiveDeviceProfile;
 
diff --git a/Engine/Source/Runtime/Engine/Private/DeviceProfiles/DeviceProfileManager.cpp b/Engine/Source/Runtime/Engine/Private/DeviceProfiles/DeviceProfileManager.cpp
--- a/Engine/Source/Runtime/Engine/Private/DeviceProfiles/DeviceProfileManager.cpp
+++ b/Engine/Source/Runtime/Engine/Private/DeviceProfiles/DeviceProfileManager.cpp
@@ -889,6 +889,7 @@ void UDeviceProfileManager::ReapplyDeviceProfile(bool bForceReload, TSet<FString
 
 	UE_LOG(LogDeviceProfileManager, Log, TEXT("ReapplyDeviceProfile applying profile: [%s]"), *BaseProfile->GetName(), OverrideProfile ? *OverrideProfile->GetName() : TEXT("not set."));
 
+	PreReapplyDeviceProfilesDelegate.Broadcast();
 
 	if (ModifiedCVars != nullptr)
 	{
@@ -929,6 +930,8 @@ void UDeviceProfileManager::ReapplyDeviceProfile(bool bForceReload, TSet<FString
 		// broadcast cvar sinks now that we are done
 		IConsoleManager::Get().CallAllConsoleVariableSinks();
 	}
+
+	PostReapplyDeviceProfilesDelegate.Broadcast();
 }
 
 static void TestProfileForCircularReferences(const FString& ProfileName, const FString& ParentName, const FConfigFile &PlatformConfigFile)
@@ -1147,6 +1150,15 @@ FOnActiveDeviceProfileChanged& UDeviceProfileManager::OnActiveDeviceProfileChang
 	return ActiveDeviceProfileChangedDelegate;
 }
 
+FWillReapplyDeviceProfiles& UDeviceProfileManager::WillReapplyDeviceProfiles()
+{
+	return PreReapplyDeviceProfilesDelegate;
+}
+
+FOnReapplyDeviceProfiles& UDeviceProfileManager::OnReapplyDeviceProfiles()
+{
+	return PostReapplyDeviceProfilesDelegate;
+}
 
 void UDeviceProfileManager::GetProfileConfigFiles(OUT TArray<FString>& OutConfigFiles)
 {
diff --git a/Engine/Source/Runtime/Engine/Private/Scalability.cpp b/Engine/Source/Runtime/Engine/Private/Scalability.cpp
--- a/Engine/Source/Runtime/Engine/Private/Scalability.cpp
+++ b/Engine/Source/Runtime/Engine/Private/Scalability.cpp
@@ -1155,8 +1155,13 @@ void FQualityLevels::SetLandscapeQuality(int32 Value)
 {
 	LandscapeQuality = FMath::Clamp(Value, 0, CVarLandscapeQuality_NumLevels->GetInt() - 1);
 }
-	
+
 void LoadState(const FString& IniName)
+{
+	LoadState(IniName, TEXT("ScalabilityGroups"));
+}
+	
+void LoadState(const FString& IniName, const TCHAR* Section)
 {
 	check(!IniName.IsEmpty());
 
@@ -1166,8 +1171,6 @@ void LoadState(const FString& IniName)
 	// Use existing quality levels - Defaults with device profile customization
 	FQualityLevels State = GetQualityLevels();
 
-	const TCHAR* Section = TEXT("ScalabilityGroups");
-
 	// looks like cvars but here we just use the name for the ini
 	GConfig->GetFloat(Section, TEXT("sg.ResolutionQuality"), State.ResolutionQuality, IniName);
 	GConfig->GetInt(Section, TEXT("sg.ViewDistanceQuality"), State.ViewDistanceQuality, IniName);
@@ -1194,14 +1197,18 @@ void LoadState(const FString& IniName)
 }
 
 void SaveState(const FString& IniName)
+{
+	// Saves to the previous default
+	SaveState(IniName, TEXT("ScalabilityGroups"));
+}
+
+void SaveState(const FString& IniName, const TCHAR* Section)
 {
 	check(!IniName.IsEmpty());
 
 	// Save the "real" settings if in a temporary state
 	FQualityLevels State = GScalabilityUsingTemporaryQualityLevels ? GScalabilityBackupQualityLevels : GetQualityLevels();
 
-	const TCHAR* Section = TEXT("ScalabilityGroups");
-
 	// looks like cvars but here we just use the name for the ini
 	GConfig->SetFloat(Section, TEXT("sg.ResolutionQuality"), State.ResolutionQuality, IniName);
 	GConfig->SetInt(Section, TEXT("sg.ViewDistanceQuality"), State.ViewDistanceQuality, IniName);
diff --git a/Engine/Source/Runtime/Engine/Public/Scalability.h b/Engine/Source/Runtime/Engine/Public/Scalability.h
--- a/Engine/Source/Runtime/Engine/Public/Scalability.h
+++ b/Engine/Source/Runtime/Engine/Public/Scalability.h
@@ -201,9 +201,21 @@ namespace Scalability
 
 	/** @param IniName e.g. GEditorPerProjectIni or GGameUserSettingsIni */
 	ENGINE_API void LoadState(const FString& IniName);
+
+	/**
+	 * @param IniName e.g. GEditorPerProjectIni or GGameUserSettingsIni
+	 * @param Section the section to load from
+	 */
+	ENGINE_API void LoadState(const FString& IniName, const TCHAR* Section);
 	
 	/** @param IniName e.g. GEditorPerProjectIni or GGameUserSettingsIni */
 	ENGINE_API void SaveState(const FString& IniName);
+	
+	/** 
+	 * @param IniName e.g. GEditorPerProjectIni or GGameUserSettingsIni 
+	 * @param Section the section to save into
+	 */
+	ENGINE_API void SaveState(const FString& IniName, const TCHAR* Section);
 
 	/**
 	 * Sends an analytic event with all quality level data
