[WorldPartitionActorDesc] Fix a crash that can happen when attaching an actor to a newly-created (not saved yet) actor

The crash can happen during UpdateActorToWorld if the parent actor is set, but its actor descriptor has not been registered yet. In that case we can skip the update, since we don't have the parent transform. It's okay to skip the update, as this actor will be updated during PropagateActorToWorldUpdate, after the parent actor descriptor is registered.

diff --git a/Engine/Source/Runtime/Engine/Private/WorldPartition/WorldPartitionActorDesc.cpp b/Engine/Source/Runtime/Engine/Private/WorldPartition/WorldPartitionActorDesc.cpp
--- a/Engine/Source/Runtime/Engine/Private/WorldPartition/WorldPartitionActorDesc.cpp
+++ b/Engine/Source/Runtime/Engine/Private/WorldPartition/WorldPartitionActorDesc.cpp
@@ -995,8 +995,17 @@ void FWorldPartitionActorDesc::UpdateActorToWorld()
 	FTransform ParentTransform = FTransform::Identity;
 	if (ParentActor.IsValid())
 	{
-		FWorldPartitionActorDesc& ParentActorDesc = Container->GetActorDescChecked(ParentActor);
-		ParentTransform = ParentActorDesc.GetActorTransform();
+		if (FWorldPartitionActorDesc* ParentActorDesc = Container->GetActorDesc(ParentActor))
+		{
+			ParentTransform = ParentActorDesc->GetActorTransform();
+		}
+		else
+		{
+			// This can happen if the parent actor is set, but its actor descriptor has not been registered yet.
+			// In that case we can skip this update, since we don't have the parent transform. This actor will be updated
+			// during PropagateActorToWorldUpdate, after the parent actor descriptor is registered.
+			return;
+		}
 	}
 
 	if (bHasValidRelativeBounds == false)
