Scene Capture / Custom Render Pass lightweight shadow map support
Utility functionality to support lightweight PCF shadow maps.  Significantly lower overhead than running a full scene capture with proper lighting and shadows (over 10x faster on both CPU and GPU).  Intended for things like UI or HUD objects, or generating shadow masks for live compositing.  Key features:
* Support for scene captures or custom render passes writing view and projection matrices to a Material Parameter Collection.  Includes LWC support.  Note that without this, there's no way for a user to independently match the matrices used by a scene capture without looking at the source code (especially the specific math used to generate the projection matrices), or writing your own code.
* New material graph node which can apply the aforementioned Material Parameter Collection matrices to vectors.  The material graph otherwise doesn't expose LWC math except for specific built-in basis conversions (like world to local space of the current view, but not of the shadow view, which is necessary for implementing shadow mapping), and Niagara specific nodes.
* Added gather functionality to TextureSample material graph node.
* Implemented efficient 2x2, 4x4, 6x6, and 8x8 bilinear PCF material functions using gather.  Includes receiver slope bias support for decent quality with a small constant bias.
* Fixed bug where Scene Capture Custom Render Passes were rendering delayed by a frame, due to being added to the FScene after the main scene renderer constructor.  Logic added to interleave custom render passes and scene captures based on their sort order, and separate out custom render passes that should render with the main view, so they are added to the FScene before the main scene renderer constructor.

#virtualized

diff --git a/Engine/Plugins/Interchange/Runtime/Source/Import/Private/MaterialX/MaterialExpressions/MaterialExpressionTextureSampleParameterBlur.cpp b/Engine/Plugins/Interchange/Runtime/Source/Import/Private/MaterialX/MaterialExpressions/MaterialExpressionTextureSampleParameterBlur.cpp
--- a/Engine/Plugins/Interchange/Runtime/Source/Import/Private/MaterialX/MaterialExpressions/MaterialExpressionTextureSampleParameterBlur.cpp
+++ b/Engine/Plugins/Interchange/Runtime/Source/Import/Private/MaterialX/MaterialExpressions/MaterialExpressionTextureSampleParameterBlur.cpp
@@ -116,7 +116,8 @@ extern int32 CompileTextureSample(
 	const int32 MipValue1Index = INDEX_NONE,
 	const ETextureMipValueMode MipValueMode = TMVM_None,
 	const ESamplerSourceMode SamplerSource = SSM_FromTextureAsset,
-	const bool AutomaticViewMipBias = false
+	const bool AutomaticViewMipBias = false,
+	const ETextureGatherMode GatherMode = TGM_None
 );
 
 int32 UMaterialExpressionMaterialXTextureSampleParameterBlur::Compile(FMaterialCompiler* Compiler, int32 OutputIndex)
diff --git a/Engine/Shaders/Private/Common.ush b/Engine/Shaders/Private/Common.ush
--- a/Engine/Shaders/Private/Common.ush
+++ b/Engine/Shaders/Private/Common.ush
@@ -428,6 +428,58 @@ MaterialFloat4 TextureExternalSampleLevel(TextureExternal Tex, SamplerState Samp
 {
 	return Tex.SampleLevel(Sampler, UV, Mip);
 }
+MaterialFloat4 Texture2DGatherRed(Texture2D Tex, SamplerState Sampler, float2 UV)
+{
+	return Tex.GatherRed(Sampler, UV);
+}
+MaterialFloat4 Texture2DGatherGreen(Texture2D Tex, SamplerState Sampler, float2 UV)
+{
+	return Tex.GatherGreen(Sampler, UV);
+}
+MaterialFloat4 Texture2DGatherBlue(Texture2D Tex, SamplerState Sampler, float2 UV)
+{
+	return Tex.GatherBlue(Sampler, UV);
+}
+MaterialFloat4 Texture2DGatherAlpha(Texture2D Tex, SamplerState Sampler, float2 UV)
+{
+	return Tex.GatherAlpha(Sampler, UV);
+}
+MaterialFloat4 TextureCubeGatherRed(TextureCube Tex, SamplerState Sampler, float3 UV)
+{
+	uint StatusIgnored;
+	return Tex.GatherRed(Sampler, UV, StatusIgnored);
+}
+MaterialFloat4 TextureCubeGatherGreen(TextureCube Tex, SamplerState Sampler, float3 UV)
+{
+	uint StatusIgnored;
+	return Tex.GatherGreen(Sampler, UV, StatusIgnored);
+}
+MaterialFloat4 TextureCubeGatherBlue(TextureCube Tex, SamplerState Sampler, float3 UV)
+{
+	uint StatusIgnored;
+	return Tex.GatherBlue(Sampler, UV, StatusIgnored);
+}
+MaterialFloat4 TextureCubeGatherAlpha(TextureCube Tex, SamplerState Sampler, float3 UV)
+{
+	uint StatusIgnored;
+	return Tex.GatherAlpha(Sampler, UV, StatusIgnored);
+}
+MaterialFloat4 Texture2DArrayGatherRed(Texture2DArray Tex, SamplerState Sampler, float3 UV)
+{
+	return Tex.GatherRed(Sampler, UV);
+}
+MaterialFloat4 Texture2DArrayGatherGreen(Texture2DArray Tex, SamplerState Sampler, float3 UV)
+{
+	return Tex.GatherGreen(Sampler, UV);
+}
+MaterialFloat4 Texture2DArrayGatherBlue(Texture2DArray Tex, SamplerState Sampler, float3 UV)
+{
+	return Tex.GatherBlue(Sampler, UV);
+}
+MaterialFloat4 Texture2DArrayGatherAlpha(Texture2DArray Tex, SamplerState Sampler, float3 UV)
+{
+	return Tex.GatherAlpha(Sampler, UV);
+}
 
 // Re-routed texture sampling for decals
 // On iOS and Android MALI devices decal UVs has issues with deriavatives on polygon edges resulting in 'wireframe' rendering artifacts (UE-50971)
@@ -625,6 +677,27 @@ MaterialFloat4 TextureCubeArraySampleLevel(TextureCubeArray Tex, SamplerState Sa
 {
 	return TextureCubeArraySampleLevel(Tex, Sampler, float4(UV, ArrayIndex), Mip);
 }
+
+MaterialFloat4 TextureCubeArrayGatherRed(TextureCubeArray Tex, SamplerState Sampler, float4 UV)
+{
+	uint StatusIgnored;
+	return Tex.GatherRed(Sampler, UV, StatusIgnored);
+}
+MaterialFloat4 TextureCubeArrayGatherGreen(TextureCubeArray Tex, SamplerState Sampler, float4 UV)
+{
+	uint StatusIgnored;
+	return Tex.GatherGreen(Sampler, UV, StatusIgnored);
+}
+MaterialFloat4 TextureCubeArrayGatherBlue(TextureCubeArray Tex, SamplerState Sampler, float4 UV)
+{
+	uint StatusIgnored;
+	return Tex.GatherBlue(Sampler, UV, StatusIgnored);
+}
+MaterialFloat4 TextureCubeArrayGatherAlpha(TextureCubeArray Tex, SamplerState Sampler, float4 UV)
+{
+	uint StatusIgnored;
+	return Tex.GatherAlpha(Sampler, UV, StatusIgnored);
+}
 #endif // FEATURE_LEVEL >= FEATURE_LEVEL_SM5
 
 // TANGENTTOWORLD0 is the first row of the tangent to world matrix, w might be needed for padding and is not used yet.
diff --git a/Engine/Source/Editor/MaterialEditor/Private/MaterialEditingLibrary.cpp b/Engine/Source/Editor/MaterialEditor/Private/MaterialEditingLibrary.cpp
--- a/Engine/Source/Editor/MaterialEditor/Private/MaterialEditingLibrary.cpp
+++ b/Engine/Source/Editor/MaterialEditor/Private/MaterialEditingLibrary.cpp
@@ -15,6 +15,7 @@
 #include "Materials/MaterialFunctionInstance.h"
 #include "Materials/MaterialExpressionTextureBase.h"
 #include "Materials/MaterialExpressionCollectionParameter.h"
+#include "Materials/MaterialExpressionCollectionTransform.h"
 #include "Materials/MaterialExpressionFunctionInput.h"
 #include "Materials/MaterialExpressionFunctionOutput.h"
 #include "Materials/MaterialExpressionComponentMask.h"
@@ -553,6 +554,12 @@ UMaterialExpression* UMaterialEditingLibrary::CreateMaterialExpressionEx(UMateri
 			{
 				MECollectionParameter->Collection = Cast<UMaterialParameterCollection>(SelectedAsset);
 			}
+
+			UMaterialExpressionCollectionTransform* MECollectionTransform = Cast<UMaterialExpressionCollectionTransform>(NewExpression);
+			if (MECollectionTransform)
+			{
+				MECollectionTransform->Collection = Cast<UMaterialParameterCollection>(SelectedAsset);
+			}
 		}
 
 		UMaterialExpressionFunctionInput* FunctionInput = Cast<UMaterialExpressionFunctionInput>(NewExpression);
diff --git a/Engine/Source/Editor/MaterialEditor/Private/MaterialEditor.cpp b/Engine/Source/Editor/MaterialEditor/Private/MaterialEditor.cpp
--- a/Engine/Source/Editor/MaterialEditor/Private/MaterialEditor.cpp
+++ b/Engine/Source/Editor/MaterialEditor/Private/MaterialEditor.cpp
@@ -57,6 +57,7 @@
 
 #include "Materials/MaterialExpressionBreakMaterialAttributes.h"
 #include "Materials/MaterialExpressionCollectionParameter.h"
+#include "Materials/MaterialExpressionCollectionTransform.h"
 #include "Materials/MaterialExpressionComment.h"
 #include "Materials/MaterialExpressionComposite.h"
 #include "Materials/MaterialExpressionComponentMask.h"
@@ -1754,13 +1755,20 @@ void FMaterialEditor::CreateInternalWidgets()
 		LayoutLayerExpressionParameterDetails
 		);
 
-	FOnGetDetailCustomizationInstance LayoutCollectionParameterDetails = FOnGetDetailCustomizationInstance::CreateStatic(&FMaterialExpressionCollectionParameterDetails::MakeInstance);
+	FOnGetDetailCustomizationInstance LayoutCollectionParameterDetails = FOnGetDetailCustomizationInstance::CreateStatic(&FMaterialExpressionCollectionParameterDetails::MakeInstance, true);
 
 	MaterialDetailsView->RegisterInstancedCustomPropertyLayout( 
 		UMaterialExpressionCollectionParameter::StaticClass(), 
 		LayoutCollectionParameterDetails
 		);
 
+	FOnGetDetailCustomizationInstance LayoutCollectionTransformDetails = FOnGetDetailCustomizationInstance::CreateStatic(&FMaterialExpressionCollectionParameterDetails::MakeInstance, false);
+
+	MaterialDetailsView->RegisterInstancedCustomPropertyLayout(
+		UMaterialExpressionCollectionTransform::StaticClass(),
+		LayoutCollectionTransformDetails
+		);
+
 	FOnGetDetailCustomizationInstance LayoutCompositeDetails = FOnGetDetailCustomizationInstance::CreateStatic(&FMaterialExpressionCompositeDetails::MakeInstance);
 
 	MaterialDetailsView->RegisterInstancedCustomPropertyLayout( 
@@ -8265,6 +8273,7 @@ void FMaterialEditor::OnNodeDoubleClicked(class UEdGraphNode* Node)
 		UMaterialExpressionTextureSampleParameter* TextureParameterExpression = Cast<UMaterialExpressionTextureSampleParameter>(GraphNode->MaterialExpression);
 		UMaterialExpressionMaterialFunctionCall* FunctionExpression = Cast<UMaterialExpressionMaterialFunctionCall>(GraphNode->MaterialExpression);
 		UMaterialExpressionCollectionParameter* CollectionParameter = Cast<UMaterialExpressionCollectionParameter>(GraphNode->MaterialExpression);
+		UMaterialExpressionCollectionTransform* CollectionTransform = Cast<UMaterialExpressionCollectionTransform>(GraphNode->MaterialExpression);
 
 		TArray<UObject*> ObjectsToView;
 		UObject* ObjectToEdit = NULL;
@@ -8285,6 +8294,10 @@ void FMaterialEditor::OnNodeDoubleClicked(class UEdGraphNode* Node)
 		{
 			ObjectToEdit = CollectionParameter->Collection;
 		}
+		else if (CollectionTransform && CollectionTransform->Collection)
+		{
+			ObjectToEdit = CollectionTransform->Collection;
+		}
 
 		if (ObjectsToView.Num() > 0)
 		{
diff --git a/Engine/Source/Editor/MaterialEditor/Private/MaterialEditorDetailCustomization.cpp b/Engine/Source/Editor/MaterialEditor/Private/MaterialEditorDetailCustomization.cpp
--- a/Engine/Source/Editor/MaterialEditor/Private/MaterialEditorDetailCustomization.cpp
+++ b/Engine/Source/Editor/MaterialEditor/Private/MaterialEditorDetailCustomization.cpp
@@ -583,12 +583,13 @@ void FMaterialExpressionParameterDetails::OnSliderMinMaxEdited()
 	}
 }
 
-TSharedRef<IDetailCustomization> FMaterialExpressionCollectionParameterDetails::MakeInstance()
+TSharedRef<IDetailCustomization> FMaterialExpressionCollectionParameterDetails::MakeInstance(bool bInIncludeScalars)
 {
-	return MakeShareable( new FMaterialExpressionCollectionParameterDetails() );
+	return MakeShareable( new FMaterialExpressionCollectionParameterDetails(bInIncludeScalars) );
 }
 
-FMaterialExpressionCollectionParameterDetails::FMaterialExpressionCollectionParameterDetails()
+FMaterialExpressionCollectionParameterDetails::FMaterialExpressionCollectionParameterDetails(bool bInIncludeScalars)
+	: bIncludeScalars(bInIncludeScalars)
 {
 }
 
@@ -734,8 +735,11 @@ void FMaterialExpressionCollectionParameterDetails::PopulateParameters()
 	if (Collection)
 	{
 		TArray<FName> NameList;
-		Collection->GetParameterNames(NameList, false);
-		Collection->GetParameterNames(NameList, true);
+		if (bIncludeScalars)
+		{
+			Collection->GetParameterNames(NameList, /*bVectorParameters=*/ false);
+		}
+		Collection->GetParameterNames(NameList, /*bVectorParameters=*/ true);
 
 		NameList.Sort([](const FName& A, const FName& B) 
 		{
diff --git a/Engine/Source/Editor/MaterialEditor/Private/MaterialEditorDetailCustomization.h b/Engine/Source/Editor/MaterialEditor/Private/MaterialEditorDetailCustomization.h
--- a/Engine/Source/Editor/MaterialEditor/Private/MaterialEditorDetailCustomization.h
+++ b/Engine/Source/Editor/MaterialEditor/Private/MaterialEditorDetailCustomization.h
@@ -86,10 +86,10 @@ class FMaterialExpressionCollectionParameterDetails : public IDetailCustomizatio
 public:
 
 	/** Makes a new instance of this detail layout class for a specific detail view requesting it */
-	static TSharedRef<class IDetailCustomization> MakeInstance();
+	static TSharedRef<class IDetailCustomization> MakeInstance(bool bInIncludeScalars);
 
 	/** Constructor requires a delegate to populate group names with */
-	FMaterialExpressionCollectionParameterDetails();
+	FMaterialExpressionCollectionParameterDetails(bool bInIncludeScalars);
 
 	/** IDetailCustomization interface */
 	virtual void CustomizeDetails( IDetailLayoutBuilder& DetailLayout ) override;
@@ -119,6 +119,9 @@ private:
 
 	/** A list of all parameter names to choose from */
 	TArray<TSharedPtr<FString>> ParametersSource;
+
+	/** Whether to include scalars in this detail customization */
+	bool bIncludeScalars;
 };
 
 
diff --git a/Engine/Source/Runtime/Engine/Classes/Components/SceneCaptureComponent.h b/Engine/Source/Runtime/Engine/Classes/Components/SceneCaptureComponent.h
--- a/Engine/Source/Runtime/Engine/Classes/Components/SceneCaptureComponent.h
+++ b/Engine/Source/Runtime/Engine/Classes/Components/SceneCaptureComponent.h
@@ -14,6 +14,7 @@
 
 class AActor;
 class FSceneViewStateInterface;
+class UMaterialParameterCollection;
 class UPrimitiveComponent;
 
 /** 
@@ -140,6 +141,18 @@ class USceneCaptureComponent : public USceneComponent
 	UPROPERTY(EditAnywhere, BlueprintReadWrite, Category = SceneCapture)
 	bool bUseRayTracingIfEnabled;
 
+	/** Store WorldToLocal and/or Projection matrices (2D capture only) to a Material Parameter Collection on render. */
+	UPROPERTY(EditAnywhere, BlueprintReadWrite, Category=SceneCapture, AdvancedDisplay)
+	TObjectPtr<UMaterialParameterCollection> CollectionTransform;
+
+	/** Parameter name of the first element of the transform in the CollectionTransform Material Parameter Collection set above.  Requires space for 5 vectors (large world coordinate transform). */
+	UPROPERTY(EditAnywhere, BlueprintReadWrite, Category=SceneCapture, AdvancedDisplay)
+	FName CollectionTransformWorldToLocal;
+
+	/** Parameter name of the first element of the transform in the CollectionTransform Material Parameter Collection set above.  Requires space for 4 vectors. */
+	UPROPERTY(EditAnywhere, BlueprintReadWrite, Category=SceneCapture, AdvancedDisplay)
+	FName CollectionTransformProjection;
+
 	/** ShowFlags for the SceneCapture's ViewFamily, to control rendering settings for this view. Hidden but accessible through details customization */
 	UE_DEPRECATED_FORGAME(5.5, "Public access to this property is deprecated, and it will become private in a future release. Please use SetShowFlagSettings and GetShowFlagSettings instead.")
 	UPROPERTY(EditAnywhere, interp, Blueprintgetter = GetShowFlagSettings, BlueprintSetter = SetShowFlagSettings, Category=SceneCapture)
diff --git a/Engine/Source/Runtime/Engine/Internal/SceneCapture/SceneCaptureInternal.h b/Engine/Source/Runtime/Engine/Internal/SceneCapture/SceneCaptureInternal.h
--- a/Engine/Source/Runtime/Engine/Internal/SceneCapture/SceneCaptureInternal.h
+++ b/Engine/Source/Runtime/Engine/Internal/SceneCapture/SceneCaptureInternal.h
@@ -10,6 +10,7 @@ class ISceneRenderBuilder;
 
 /**
  * Internal function called from FRendererModule::BeginRenderingViewFamilies.  List of view families is required to handle captures where
- * USceneCaptureComponent2D::ShouldRenderWithMainViewFamily() is true.
+ * USceneCaptureComponent2D::ShouldRenderWithMainViewFamily() is true.  Optionally can process only custom render passes that should render
+ * with the main view, so they can be added to the FScene before the main view scene renderer is constructed.
  */
-ENGINE_API void SceneCaptureUpdateDeferredCapturesInternal(FSceneInterface* Scene, TConstArrayView<const FSceneViewFamily*> ViewFamilies, ISceneRenderBuilder& SceneRenderBuilder);
+ENGINE_API void SceneCaptureUpdateDeferredCapturesInternal(FSceneInterface* Scene, TConstArrayView<const FSceneViewFamily*> ViewFamilies, ISceneRenderBuilder& SceneRenderBuilder, bool bMainCustomRenderPassOnly);
diff --git a/Engine/Source/Runtime/Engine/Private/Components/SceneCaptureComponent.cpp b/Engine/Source/Runtime/Engine/Private/Components/SceneCaptureComponent.cpp
--- a/Engine/Source/Runtime/Engine/Private/Components/SceneCaptureComponent.cpp
+++ b/Engine/Source/Runtime/Engine/Private/Components/SceneCaptureComponent.cpp
@@ -509,11 +509,11 @@ void USceneCaptureComponent::Serialize(FArchive& Ar)
 void USceneCaptureComponent::UpdateDeferredCaptures(FSceneInterface* Scene)
 {
 	TUniquePtr<ISceneRenderBuilder> SceneRenderBuilder = ISceneRenderBuilder::Create(Scene);
-	SceneCaptureUpdateDeferredCapturesInternal(Scene, {}, *SceneRenderBuilder);
+	SceneCaptureUpdateDeferredCapturesInternal(Scene, {}, *SceneRenderBuilder, /*bMainCustomRenderPassOnly=*/ false);
 	SceneRenderBuilder->Execute();
 }
 
-void SceneCaptureUpdateDeferredCapturesInternal(FSceneInterface* Scene, TConstArrayView<const FSceneViewFamily*> ViewFamilies, ISceneRenderBuilder& SceneRenderBuilder)
+void SceneCaptureUpdateDeferredCapturesInternal(FSceneInterface* Scene, TConstArrayView<const FSceneViewFamily*> ViewFamilies, ISceneRenderBuilder& SceneRenderBuilder, bool bMainCustomRenderPassOnly)
 {
 	FScopeLock ScopeLock(&SceneCapturesToUpdateMapCS);
 	UWorld* World = Scene->GetWorld();
@@ -522,6 +522,16 @@ void SceneCaptureUpdateDeferredCapturesInternal(FSceneInterface* Scene, TConstAr
 		return;
 	}
 
+	if (bMainCustomRenderPassOnly)
+	{
+		// If custom render passes are globally disabled, just return since that's all we are trying to process
+		IConsoleVariable* AllowRenderInMainRenderer = IConsoleManager::Get().FindConsoleVariable(TEXT("r.SceneCapture.AllowRenderInMainRenderer"));
+		if (!AllowRenderInMainRenderer || !AllowRenderInMainRenderer->GetBool())
+		{
+			return;
+		}
+	}
+
 	// Only update the scene captures associated with the current scene.
 	// Updating others not associated with the scene would cause invalid data to be rendered into the target
 	TArray< TWeakObjectPtr<USceneCaptureComponent> > SceneCapturesToUpdate;
@@ -537,6 +547,42 @@ void SceneCaptureUpdateDeferredCapturesInternal(FSceneInterface* Scene, TConstAr
 		}
 	}
 
+	SceneCapturesToUpdate.Sort([](const TWeakObjectPtr<USceneCaptureComponent>& A, const TWeakObjectPtr<USceneCaptureComponent>& B)
+	{
+		if (!A.IsValid())
+		{
+			return false;
+		}
+		else if (!B.IsValid())
+		{
+			return true;
+		}
+		return A->CaptureSortPriority > B->CaptureSortPriority;
+	});
+
+	// To ensure CRPs run in the correct sorted order with scene captures, we need to interleave the CRPs.  This is accomplished by
+	// processing CRPs and scene captures in order downstream, but the main scene renderer is allocated before the deferred scene
+	// capture update loop, so we need to specifically separate out the CRPs that should render with the main scene renderer.  Any
+	// CRP that occurs after the last non-CRP Scene Capture goes with the main view, so we need to track the last such scene capture,
+	// and whether we're past it in the array in the loop below.
+	const USceneCaptureComponent* LastSceneCaptureNonCRP = nullptr;
+	if (bMainCustomRenderPassOnly)
+	{
+		for (int32 CaptureIndex = 0; CaptureIndex < SceneCapturesToUpdate.Num(); CaptureIndex++)
+		{
+			if (SceneCapturesToUpdate[CaptureIndex].IsValid())
+			{
+				USceneCaptureComponent2D* Component = Cast<USceneCaptureComponent2D>(SceneCapturesToUpdate[CaptureIndex].Get());
+
+				// If it's a cube capture, it never renders in the main renderer, otherwise check if it renders in the main renderer (meaning it's a CRP)
+				if (!Component || !Component->ShouldRenderInMainRenderer())
+				{
+					LastSceneCaptureNonCRP = SceneCapturesToUpdate[CaptureIndex].Get();
+				}
+			}
+		}
+	}
+
 	int32 IncrementIfNotRemoved;
 	for (int32 CaptureIndex = 0; CaptureIndex < SceneCapturesToUpdate.Num(); CaptureIndex += IncrementIfNotRemoved)
 	{
@@ -544,11 +590,23 @@ void SceneCaptureUpdateDeferredCapturesInternal(FSceneInterface* Scene, TConstAr
 
 		if (SceneCapturesToUpdate[CaptureIndex].IsValid())
 		{
+			// Track that we've passed the last non CRP scene capture by setting this to NULL when we found it
+			if (LastSceneCaptureNonCRP == SceneCapturesToUpdate[CaptureIndex].Get())
+			{
+				LastSceneCaptureNonCRP = nullptr;
+			}
+
 			if (USceneCaptureComponent2D* Component = Cast<USceneCaptureComponent2D>(SceneCapturesToUpdate[CaptureIndex].Get()))
 			{
-				// If this is intended to render with the main view family, and we don't have a main view family,
-				// remove it from the array so it can render with the main view family later.
-				if (Component->ShouldRenderWithMainViewFamily() && !MainViewFamily)
+				// If this is intended to render with the main view family, and we don't have a main view family, remove it from the array so
+				// it can render with the main view family later.  Also remove if we are only searching for main view Custom Render Passes, and:
+				//
+				//		1.  It's not a custom render pass:  !ShouldRenderInMainRenderer()    OR
+				//		2.  It's not a main view CRP, meaning there is another scene capture after it in the array (LastSceneCaptureNonCRP != nullptr),
+				//			and it's not explicitly flagged to render with the main view family.
+				//
+				if ((Component->ShouldRenderWithMainViewFamily() && !MainViewFamily) ||
+					(bMainCustomRenderPassOnly && (!Component->ShouldRenderInMainRenderer() || (LastSceneCaptureNonCRP && !Component->ShouldRenderWithMainViewFamily()))))
 				{
 					SceneCapturesToUpdate.RemoveAtSwap(CaptureIndex);
 					IncrementIfNotRemoved = 0;
@@ -557,19 +615,6 @@ void SceneCaptureUpdateDeferredCapturesInternal(FSceneInterface* Scene, TConstAr
 		}
 	}
 
-	SceneCapturesToUpdate.Sort([](const TWeakObjectPtr<USceneCaptureComponent>& A, const TWeakObjectPtr<USceneCaptureComponent>& B)
-	{
-		if (!A.IsValid())
-		{
-			return false;
-		}
-		else if (!B.IsValid())
-		{
-			return true;
-		}
-		return A->CaptureSortPriority > B->CaptureSortPriority;
-	});
-
 	for (TWeakObjectPtr<USceneCaptureComponent> Component : SceneCapturesToUpdate)
 	{
 		if (Component.IsValid())
@@ -587,10 +632,11 @@ void SceneCaptureUpdateDeferredCapturesInternal(FSceneInterface* Scene, TConstAr
 				((USceneCaptureComponent2D*)Component.Get())->MainViewFamily = nullptr;
 			}
 		}
-	}
 
-	// All scene captures for this world have been updated
-	SceneCapturesToUpdateMap.Remove(World);
+		// Remove this component, as we processed it (or it was invalid).  We remove specific processed items,
+		// as there are cases where not all items are processed, and need to be handled later.
+		SceneCapturesToUpdateMap.Remove(World, Component);
+	}
 }
 
 void USceneCaptureComponent::OnUnregister()
diff --git a/Engine/Source/Runtime/Engine/Private/Materials/HLSLMaterialTranslator.cpp b/Engine/Source/Runtime/Engine/Private/Materials/HLSLMaterialTranslator.cpp
--- a/Engine/Source/Runtime/Engine/Private/Materials/HLSLMaterialTranslator.cpp
+++ b/Engine/Source/Runtime/Engine/Private/Materials/HLSLMaterialTranslator.cpp
@@ -5189,6 +5189,76 @@ int32 FHLSLMaterialTranslator::AccessCollectionParameter(UMaterialParameterColle
 		ComponentIndex == -1 ? true : ComponentIndex % 4 == 3);
 }
 
+int32 FHLSLMaterialTranslator::CollectionTransform(int32 InputIndex, const TStaticArray<int32, 5>& CollectionParameters, EParameterCollectionTransformType TransformType)
+{
+	FString InputCode = GetParameterCode(InputIndex);
+	EMaterialValueType InputType = GetParameterType(InputIndex);
+
+	switch (TransformType)
+	{
+	case EParameterCollectionTransformType::Position:
+	default:
+		if (InputType == MCT_Float3 || InputType == MCT_LWCVector3)
+		{
+			// Treat input as a translation vector (w = 1)
+			return AddCodeChunkZeroDeriv(MakeNonLWCType(InputType), TEXT("(mul(float4(WSDemote(%s), 1.0), float4x4(%s, %s, %s, %s))).xyz"),
+				*InputCode,
+				*GetParameterCode(CollectionParameters[0]),
+				*GetParameterCode(CollectionParameters[1]),
+				*GetParameterCode(CollectionParameters[2]),
+				*GetParameterCode(CollectionParameters[3]));
+		}
+		else
+		{
+			// Caller should only pass in 3 or 4 element vector inputs
+			check(InputType == MCT_Float4 || InputType == MCT_LWCVector4);
+
+			// Treat input as a homogenous vector (w = user specified)
+			return AddCodeChunkZeroDeriv(MCT_Float4, TEXT("(mul(WSDemote(%s), float4x4(%s, %s, %s, %s)))"),
+				*InputCode,
+				*GetParameterCode(CollectionParameters[0]),
+				*GetParameterCode(CollectionParameters[1]),
+				*GetParameterCode(CollectionParameters[2]),
+				*GetParameterCode(CollectionParameters[3]));
+		}
+		break;
+	case EParameterCollectionTransformType::Vector:
+		// Treat input as a direction vector (w = 0)
+		return AddCodeChunkZeroDeriv(MCT_Float3, TEXT("(mul(float4(WSDemote(%s).xyz, 0.0), float4x4(%s, %s, %s, float4(0.0, 0.0, 0.0, 0.0)))).xyz"),
+			*InputCode,
+			*GetParameterCode(CollectionParameters[0]),
+			*GetParameterCode(CollectionParameters[1]),
+			*GetParameterCode(CollectionParameters[2]));
+	case EParameterCollectionTransformType::Projection:
+		// Optimized to save many ALU for a standard perspective or orthographic projection matrix, where most of the elements of the matrix are zero.
+		return AddCodeChunkZeroDeriv(MCT_Float4, TEXT("(WSDemote(%s).xyzz*float4(%s.x, %s.y, %s.zw) + float4(0.0, 0.0, %s.zw))"),
+			*InputCode,
+			*GetParameterCode(CollectionParameters[0]),
+			*GetParameterCode(CollectionParameters[1]),
+			*GetParameterCode(CollectionParameters[2]),
+			*GetParameterCode(CollectionParameters[3]));
+	case EParameterCollectionTransformType::LocalToWorld:
+		AddLWCFuncUsage(ELWCFunctionKind::MultiplyVectorMatrix);
+		return AddCodeChunkZeroDeriv(MakeLWCType(InputType), TEXT("WSMultiply(WSDemote(%s), MakeWSMatrix(%s.xyz, float4x4(%s, %s, %s, %s)))"),
+			*InputCode,
+			*GetParameterCode(CollectionParameters[4]),		// Tile offset
+			*GetParameterCode(CollectionParameters[0]),
+			*GetParameterCode(CollectionParameters[1]),
+			*GetParameterCode(CollectionParameters[2]),
+			*GetParameterCode(CollectionParameters[3]));
+	case EParameterCollectionTransformType::WorldToLocal:
+		AddLWCFuncUsage(ELWCFunctionKind::MultiplyVectorMatrix);
+		return AddCodeChunkZeroDeriv(MakeNonLWCType(InputType), TEXT("WSMultiply(WSPromote(%s), MakeWSInverseMatrix(%s.xyz, float4x4(%s, %s, %s, %s)))"),
+			*InputCode,
+			*GetParameterCode(CollectionParameters[4]),		// Tile offset
+			*GetParameterCode(CollectionParameters[0]),
+			*GetParameterCode(CollectionParameters[1]),
+			*GetParameterCode(CollectionParameters[2]),
+			*GetParameterCode(CollectionParameters[3]));
+		break;
+	}
+}
+
 bool FHLSLMaterialTranslator::GetParameterOverrideValueForCurrentFunction(EMaterialParameterType ParameterType, FName ParameterName, FMaterialParameterMetadata& OutResult) const
 {
 	const TArray<FMaterialFunctionCompileState*>& FunctionStack = FunctionStacks[ShaderFrequency];
@@ -7169,6 +7239,7 @@ int32 FHLSLMaterialTranslator::TextureSample(
 	int32 MipValue1Index,
 	ETextureMipValueMode MipValueMode,
 	ESamplerSourceMode SamplerSource,
+	ETextureGatherMode GatherMode,
 	int32 TextureReferenceIndex,
 	bool AutomaticViewMipBias,
 	bool AdaptiveVirtualTexture,
@@ -7234,6 +7305,26 @@ int32 FHLSLMaterialTranslator::TextureSample(
 		return Errorf(TEXT("Invalid mip map parameter"));
 	}
 
+	if (GatherMode != TGM_None)
+	{
+		if (bVirtualTexture)
+		{
+			return Errorf(TEXT("Gather does not support virtual textures."));
+		}
+
+		if (MipValueMode != TMVM_None)
+		{
+			return Errorf(TEXT("Gather does not support mip map overrides (implicitly accesses a specific mip)."));
+		}
+
+		if (!(TextureType & (MCT_Texture2D | MCT_TextureCube | MCT_Texture2DArray | MCT_TextureCubeArray)))
+		{
+			return Errorf(TEXT("Gather only supports 2D and Cube texture formats."));
+		}
+
+		AutomaticViewMipBias = false;
+	}
+
 	// if we are not in the PS we need a mip level
 	if (ShaderFrequency != SF_Pixel)
 	{
@@ -7569,33 +7660,66 @@ int32 FHLSLMaterialTranslator::TextureSample(
 	}
 
 	FString TextureTypeName;
-	if (TextureType == MCT_TextureCube)
-	{
-		TextureTypeName = TEXT("TextureCubeSample");
-	}
-	else if (TextureType == MCT_Texture2DArray)
-	{
-		TextureTypeName = TEXT("Texture2DArraySample");
-	}
-	else if (TextureType == MCT_TextureCubeArray)
-	{
-		TextureTypeName = TEXT("TextureCubeArraySample");
-	}
-	else if (TextureType == MCT_VolumeTexture)
-	{
-		TextureTypeName = TEXT("Texture3DSample");
-	}
-	else if (TextureType == MCT_TextureExternal)
+	if (GatherMode != TGM_None)
 	{
-		TextureTypeName = TEXT("TextureExternalSample");
-	}
-	else if (bVirtualTexture)
-	{
-		TextureTypeName = TEXT("TextureVirtualSample");
+		// Should have errored out above
+		check(TextureType& (MCT_Texture2D | MCT_TextureCube | MCT_Texture2DArray | MCT_TextureCubeArray));
+
+		if (TextureType == MCT_TextureCube)
+		{
+			TextureTypeName = TEXT("TextureCubeGather");
+		}
+		else if (TextureType == MCT_Texture2DArray)
+		{
+			TextureTypeName = TEXT("Texture2DArrayGather");
+		}
+		else if (TextureType == MCT_TextureCubeArray)
+		{
+			TextureTypeName = TEXT("TextureCubeArrayGather");
+		}
+		else // MCT_Texture2D
+		{
+			TextureTypeName = TEXT("Texture2DGather");
+		}
+
+		switch (GatherMode)
+		{
+		case TGM_Red:	TextureTypeName += TEXT("Red");  break;
+		case TGM_Green:	TextureTypeName += TEXT("Green");  break;
+		case TGM_Blue:	TextureTypeName += TEXT("Blue");  break;
+		case TGM_Alpha:	TextureTypeName += TEXT("Alpha");  break;
+		}
 	}
-	else // MCT_Texture2D
+	else
 	{
-		TextureTypeName = TEXT("Texture2DSample");
+		if (TextureType == MCT_TextureCube)
+		{
+			TextureTypeName = TEXT("TextureCubeSample");
+		}
+		else if (TextureType == MCT_Texture2DArray)
+		{
+			TextureTypeName = TEXT("Texture2DArraySample");
+		}
+		else if (TextureType == MCT_TextureCubeArray)
+		{
+			TextureTypeName = TEXT("TextureCubeArraySample");
+		}
+		else if (TextureType == MCT_VolumeTexture)
+		{
+			TextureTypeName = TEXT("Texture3DSample");
+		}
+		else if (TextureType == MCT_TextureExternal)
+		{
+			TextureTypeName = TEXT("TextureExternalSample");
+		}
+		else if (bVirtualTexture)
+		{
+			TextureTypeName = TEXT("TextureVirtualSample");
+		}
+		else // MCT_Texture2D
+		{
+			TextureTypeName = TEXT("Texture2DSample");
+		}
 	}
 
 	const bool bStoreTexCoordScales = ShaderFrequency == SF_Pixel && TextureReferenceIndex != INDEX_NONE;
@@ -7732,14 +7856,15 @@ int32 FHLSLMaterialTranslator::TextureSample(
 	{
 		// Non-VT MipValueMode logic
 		// 
-		// Re-route decal texture sampling so platforms may add specific workarounds there
-		if (bDecal)
+		// Re-route decal texture sampling so platforms may add specific workarounds there (assume no workarounds for Gather)
+		if (bDecal && GatherMode == TGM_None)
 		{
 			TextureTypeName += TEXT("_Decal");
 		}
 
+		// Note that gather implicitly samples from the first mip, and doesn't support explicit mip levels or derivatives
 		FString SampleCodeFinite;
-		if (MipValueMode == TMVM_None)
+		if (MipValueMode == TMVM_None || GatherMode != TGM_None)
 		{
 			SampleCodeFinite = FString::Printf(TEXT("%s(%s,%s,%s)"), *TextureTypeName, *TextureName, *SamplerStateCode, *UV_Value);
 		}
@@ -7762,7 +7887,7 @@ int32 FHLSLMaterialTranslator::TextureSample(
 
 		SampleCodeFinite = ApplySamplerType(SampleCodeFinite, SamplerType);
 
-		if (IsAnalyticDerivEnabled() && IsDerivativeValid(UvDerivativeStatus) && !bDecal)
+		if (IsAnalyticDerivEnabled() && IsDerivativeValid(UvDerivativeStatus) && !bDecal && GatherMode == TGM_None)
 		{
 			FString SampleCodeAnalytic = SampleCodeFinite;
 			if (MipValueMode == TMVM_None || MipValueMode == TMVM_MipBias)
diff --git a/Engine/Source/Runtime/Engine/Private/Materials/HLSLMaterialTranslator.h b/Engine/Source/Runtime/Engine/Private/Materials/HLSLMaterialTranslator.h
--- a/Engine/Source/Runtime/Engine/Private/Materials/HLSLMaterialTranslator.h
+++ b/Engine/Source/Runtime/Engine/Private/Materials/HLSLMaterialTranslator.h
@@ -925,6 +925,7 @@ protected:
 	virtual int32 GetCurrentFunctionStackDepth() override;
 
 	virtual int32 AccessCollectionParameter(UMaterialParameterCollection* ParameterCollection, int32 ParameterIndex, int32 ComponentIndex) override;
+	virtual int32 CollectionTransform(int32 InputIndex, const TStaticArray<int32, 5>& CollectionParameters, EParameterCollectionTransformType TransformType) override;
 
 	virtual int32 NumericParameter(EMaterialParameterType ParameterType, FName ParameterName, const UE::Shader::FValue& DefaultValue) override;
 	virtual int32 Constant(float X) override;
@@ -1045,6 +1046,7 @@ protected:
 		int32 MipValue1Index = INDEX_NONE,
 		ETextureMipValueMode MipValueMode = TMVM_None,
 		ESamplerSourceMode SamplerSource = SSM_FromTextureAsset,
+		ETextureGatherMode GatherMode = TGM_None,
 		int32 TextureReferenceIndex = INDEX_NONE,
 		bool AutomaticViewMipBias = false,
 		bool AdaptiveVirtualTexture = false,
diff --git a/Engine/Source/Runtime/Engine/Private/Materials/MaterialCachedData.cpp b/Engine/Source/Runtime/Engine/Private/Materials/MaterialCachedData.cpp
--- a/Engine/Source/Runtime/Engine/Private/Materials/MaterialCachedData.cpp
+++ b/Engine/Source/Runtime/Engine/Private/Materials/MaterialCachedData.cpp
@@ -7,6 +7,7 @@
 #include "Materials/Material.h"
 #include "Materials/MaterialExpression.h"
 #include "Materials/MaterialExpressionCollectionParameter.h"
+#include "Materials/MaterialExpressionCollectionTransform.h"
 #include "Materials/MaterialExpressionDynamicParameter.h"
 #include "Materials/MaterialExpressionQualitySwitch.h"
 #include "Materials/MaterialExpressionMakeMaterialAttributes.h"
@@ -514,6 +515,17 @@ void FMaterialCachedExpressionData::UpdateForExpressions(const FMaterialCachedEx
 				ParameterCollectionInfos.AddUnique(NewInfo);
 			}
 		}
+		else if (UMaterialExpressionCollectionTransform* ExpressionCollectionTransform = Cast<UMaterialExpressionCollectionTransform>(Expression))
+		{
+			UMaterialParameterCollection* Collection = ExpressionCollectionTransform->Collection;
+			if (Collection)
+			{
+				FMaterialParameterCollectionInfo NewInfo;
+				NewInfo.ParameterCollection = Collection;
+				NewInfo.StateId = Collection->StateId;
+				ParameterCollectionInfos.AddUnique(NewInfo);
+			}
+		}
 		else if (UMaterialExpressionDynamicParameter* ExpressionDynamicParameter = Cast< UMaterialExpressionDynamicParameter>(Expression))
 		{
 			DynamicParameterNames.Empty(ExpressionDynamicParameter->ParamNames.Num());
diff --git a/Engine/Source/Runtime/Engine/Private/Materials/MaterialExpressionHLSL.cpp b/Engine/Source/Runtime/Engine/Private/Materials/MaterialExpressionHLSL.cpp
--- a/Engine/Source/Runtime/Engine/Private/Materials/MaterialExpressionHLSL.cpp
+++ b/Engine/Source/Runtime/Engine/Private/Materials/MaterialExpressionHLSL.cpp
@@ -872,6 +872,36 @@ bool UMaterialExpressionCollectionParameter::GenerateHLSLExpression(FMaterialHLS
 	}
 }
 
+bool UMaterialExpressionCollectionTransform::GenerateHLSLExpression(FMaterialHLSLGenerator& Generator, UE::HLSLTree::FScope& Scope, int32 OutputIndex, UE::HLSLTree::FExpression const*& OutExpression) const
+{
+	int32 ParameterIndex = INDEX_NONE;
+	int32 ComponentIndex = INDEX_NONE;
+	if (Collection)
+	{
+		Collection->GetParameterIndex(ParameterId, ParameterIndex, ComponentIndex);
+	}
+	if (ParameterIndex != INDEX_NONE)
+	{
+		if (ComponentIndex != INDEX_NONE)
+		{
+			return Generator.Error(TEXT("CollectionTransform points to scalar, only vectors are supported!"));
+		}
+		// New translator is being deprecated in favor of even newer translator -- not going to bother implementing this for now.
+		return Generator.Error(TEXT("CollectionTransform unimplemented in new translator!"));
+	}
+	else
+	{
+		if (!Collection)
+		{
+			return Generator.Error(TEXT("CollectionTransform has invalid Collection!"));
+		}
+		else
+		{
+			return Generator.Errorf(TEXT("CollectionTransform has invalid parameter %s"), *ParameterName.ToString());
+		}
+	}
+}
+
 bool UMaterialExpressionDynamicParameter::GenerateHLSLExpression(FMaterialHLSLGenerator& Generator, UE::HLSLTree::FScope& Scope, int32 OutputIndex, UE::HLSLTree::FExpression const*& OutExpression) const
 {
 	using namespace UE::HLSLTree;
diff --git a/Engine/Source/Runtime/Engine/Private/Materials/MaterialExpressions.cpp b/Engine/Source/Runtime/Engine/Private/Materials/MaterialExpressions.cpp
--- a/Engine/Source/Runtime/Engine/Private/Materials/MaterialExpressions.cpp
+++ b/Engine/Source/Runtime/Engine/Private/Materials/MaterialExpressions.cpp
@@ -69,6 +69,7 @@
 #include "Materials/MaterialExpressionChannelMaskParameter.h"
 #include "Materials/MaterialExpressionClamp.h"
 #include "Materials/MaterialExpressionCollectionParameter.h"
+#include "Materials/MaterialExpressionCollectionTransform.h"
 #include "Materials/MaterialExpressionComment.h"
 #include "Materials/MaterialExpressionComposite.h"
 #include "Materials/MaterialExpressionComponentMask.h"
@@ -635,7 +636,8 @@ ENGINE_API int32 CompileTextureSample(
 	const int32 MipValue1Index=INDEX_NONE,
 	const ETextureMipValueMode MipValueMode=TMVM_None,
 	const ESamplerSourceMode SamplerSource=SSM_FromTextureAsset,
-	const bool AutomaticViewMipBias=false
+	const bool AutomaticViewMipBias=false,
+	const ETextureGatherMode GatherMode=TGM_None
 	)
 {
 	int32 TextureReferenceIndex = INDEX_NONE;
@@ -661,6 +663,7 @@ ENGINE_API int32 CompileTextureSample(
 					MipValue1Index,
 					MipValueMode,
 					SamplerSource,
+					GatherMode,
 					TextureReferenceIndex,
 					AutomaticViewMipBias);
 }
@@ -2877,6 +2880,7 @@ int32 UMaterialExpressionTextureSample::Compile(class FMaterialCompiler* Compile
 				CompileMipValue1(Compiler),
 				MipValueMode,
 				SamplerSource,
+				GatherMode,
 				TextureReferenceIndex,
 				bDoAutomaticViewMipBias);
 		}
@@ -2914,6 +2918,7 @@ int32 UMaterialExpressionTextureSample::Compile(class FMaterialCompiler* Compile
 				CompileMipValue1(Compiler),
 				MipValueMode,
 				SamplerSource,
+				GatherMode,
 				TextureReferenceIndex,
 				bDoAutomaticViewMipBias);
 		}
@@ -3684,7 +3689,7 @@ int32 UMaterialExpressionRuntimeVirtualTextureSample::Compile(class FMaterialCom
 			TextureCodeIndex[TexureLayerIndex],
 			CoordinateIndex, 
 			SAMPLERTYPE_VirtualMasks,
-			MipValue0Index, MipValue1Index, TextureMipLevelMode, SamplerSourceMode,
+			MipValue0Index, MipValue1Index, TextureMipLevelMode, SamplerSourceMode, TGM_None,
 			TextureReferenceIndex[TexureLayerIndex],
 			bAutomaticMipViewBias, bAdaptive, (bEnableFeedback || bForceEnableFeedback));
 	}
@@ -21424,6 +21429,243 @@ bool UMaterialExpressionCollectionParameter::MatchesSearchQuery(const TCHAR* Sea
 	return Super::MatchesSearchQuery(SearchQuery);
 }
 #endif // WITH_EDITOR
+
+
+//
+//	UMaterialExpressionCollectionTransform
+//
+UMaterialExpressionCollectionTransform::UMaterialExpressionCollectionTransform(const FObjectInitializer& ObjectInitializer)
+	: Super(ObjectInitializer)
+{
+#if WITH_EDITORONLY_DATA
+	// Structure to hold one-time initialization
+	struct FConstructorStatics
+	{
+		FText NAME_Parameters;
+		FConstructorStatics()
+			: NAME_Parameters(LOCTEXT("Parameters", "Parameters"))
+		{
+		}
+	};
+	static FConstructorStatics ConstructorStatics;
+
+	MenuCategories.Add(ConstructorStatics.NAME_Parameters);
+
+	bCollapsed = false;
+	bHidePreviewWindow = true;
+#endif
+}
+
+
+void UMaterialExpressionCollectionTransform::PostLoad()
+{
+	if (Collection)
+	{
+		Collection->ConditionalPostLoad();
+		ParameterName = Collection->GetParameterName(ParameterId);
+	}
+
+	Super::PostLoad();
+}
+
+#if WITH_EDITOR
+void UMaterialExpressionCollectionTransform::PostEditChangeProperty(FPropertyChangedEvent& PropertyChangedEvent)
+{
+	FProperty* PropertyThatChanged = PropertyChangedEvent.Property;
+
+	if (Collection)
+	{
+		ParameterId = Collection->GetParameterId(ParameterName);
+	}
+	else
+	{
+		ParameterId = FGuid();
+	}
+
+	Super::PostEditChangeProperty(PropertyChangedEvent);
+}
+
+int32 UMaterialExpressionCollectionTransform::Compile(class FMaterialCompiler* Compiler, int32 OutputIndex)
+{
+	if(!Input.GetTracedInput().Expression)
+	{
+		return Compiler->Errorf(TEXT("CollectionTransform missing input"));
+	}
+
+	int32 InputIndex = Input.Compile(Compiler);
+	EMaterialValueType InputType = Compiler->GetType(InputIndex);
+	if (InputType != MCT_Float3 && InputType != MCT_Float4 && InputType != MCT_LWCVector3 && InputType != MCT_LWCVector4)
+	{
+		return Compiler->Errorf(TEXT("CollectionTransform requires vector input"));
+	}
+
+	if (!Collection)
+	{
+		return Compiler->Errorf(TEXT("CollectionTransform has invalid Collection!"));
+	}
+
+	int32 ParameterIndex = -1;
+	int32 ComponentIndex = -1;
+	Collection->GetParameterIndex(ParameterId, ParameterIndex, ComponentIndex);
+
+	if (ParameterIndex == -1)
+	{
+		return Compiler->Errorf(TEXT("CollectionTransform has invalid parameter %s"), *ParameterName.ToString());
+	}
+
+	if (ComponentIndex != -1)
+	{
+		return Compiler->Errorf(TEXT("CollectionTransform parameter %s is scalar, vectors are required"), *ParameterName.ToString());
+	}
+
+	TStaticArray<int32, 5> CollectionParameters;
+
+	if (TransformType == EParameterCollectionTransformType::Position || TransformType == EParameterCollectionTransformType::Projection)
+	{
+		if (ParameterIndex + 4 > Collection->GetTotalVectorStorage())
+		{
+			return Compiler->Errorf(TEXT("CollectionTransform parameter %s requires 4 vectors for Position or Projection matrix"), *ParameterName.ToString());
+		}
+		else
+		{
+			for (int32 ParameterOffset = 0; ParameterOffset < 4; ParameterOffset++)
+			{
+				CollectionParameters[ParameterOffset] = Compiler->AccessCollectionParameter(Collection, ParameterIndex + ParameterOffset, -1);
+			}
+			CollectionParameters[4] = -1;
+		}
+	}
+	else if (TransformType == EParameterCollectionTransformType::Vector)
+	{
+		if (ParameterIndex + 3 > Collection->GetTotalVectorStorage())
+		{
+			return Compiler->Errorf(TEXT("CollectionTransform parameter %s requires 3 vectors for Vector matrix"), *ParameterName.ToString());
+		}
+		else
+		{
+			for (int32 ParameterOffset = 0; ParameterOffset < 3; ParameterOffset++)
+			{
+				CollectionParameters[ParameterOffset] = Compiler->AccessCollectionParameter(Collection, ParameterIndex + ParameterOffset, -1);
+			}
+			CollectionParameters[3] = -1;
+			CollectionParameters[4] = -1;
+		}
+	}
+	else
+	{
+		if (ParameterIndex + 5 > Collection->GetTotalVectorStorage())
+		{
+			return Compiler->Errorf(TEXT("CollectionTransform parameter %s requires 5 vectors for LWC Matrix"), *ParameterName.ToString());
+		}
+		else
+		{
+			for (int32 ParameterOffset = 0; ParameterOffset < 5; ParameterOffset++)
+			{
+				CollectionParameters[ParameterOffset] = Compiler->AccessCollectionParameter(Collection, ParameterIndex + ParameterOffset, -1);
+			}
+			check(TransformType == EParameterCollectionTransformType::LocalToWorld || TransformType == EParameterCollectionTransformType::WorldToLocal);
+		}
+	}
+
+	return Compiler->CollectionTransform(InputIndex, CollectionParameters, TransformType);
+}
+
+void UMaterialExpressionCollectionTransform::GetCaption(TArray<FString>& OutCaptions) const
+{
+	FString TypePrefix;
+
+	if (Collection)
+	{
+		int32 ParameterIndex = -1;
+		int32 ComponentIndex = -1;
+		Collection->GetParameterIndex(ParameterId, ParameterIndex, ComponentIndex);
+
+		if (ParameterIndex == -1)
+		{
+			TypePrefix = TEXT("(Needs name)");
+		}
+		else if (ComponentIndex != -1)
+		{
+			TypePrefix = TEXT("(Needs vector type)");
+		}
+		else if (TransformType == EParameterCollectionTransformType::Position || TransformType == EParameterCollectionTransformType::Projection)
+		{
+			if (ParameterIndex + 4 > Collection->GetTotalVectorStorage())
+			{
+				TypePrefix = TEXT("(Needs 4 vectors)");
+			}
+			else if (TransformType == EParameterCollectionTransformType::Position)
+			{
+				TypePrefix = TEXT("(Position)");
+			}
+			else
+			{
+				TypePrefix = TEXT("(Projection)");
+			}
+		}
+		else if (TransformType == EParameterCollectionTransformType::Vector)
+		{
+			if (ParameterIndex + 3 > Collection->GetTotalVectorStorage())
+			{
+				TypePrefix = TEXT("(Needs 3 vectors)");
+			}
+			else
+			{
+				TypePrefix = TEXT("(Vector)");
+			}
+		}
+		else
+		{
+			if (ParameterIndex + 5 > Collection->GetTotalVectorStorage())
+			{
+				TypePrefix = TEXT("(Needs 5 vectors)");
+			}
+			else if (TransformType == EParameterCollectionTransformType::LocalToWorld)
+			{
+				TypePrefix = TEXT("(Local to World)");
+			}
+			else
+			{
+				check(TransformType == EParameterCollectionTransformType::WorldToLocal);
+				TypePrefix = TEXT("(World to Local)");
+			}
+		}
+	}
+	else
+	{
+		TypePrefix = TEXT("(Needs collection)");
+	}
+
+	OutCaptions.Add(TypePrefix + TEXT(" Collection Transform"));
+
+	if (Collection)
+	{
+		OutCaptions.Add(Collection->GetName());
+		OutCaptions.Add(FString(TEXT("'")) + ParameterName.ToString() + TEXT("'"));
+	}
+	else
+	{
+		OutCaptions.Add(TEXT("Unspecified"));
+	}
+}
+
+
+bool UMaterialExpressionCollectionTransform::MatchesSearchQuery(const TCHAR* SearchQuery)
+{
+	if (ParameterName.ToString().Contains(SearchQuery))
+	{
+		return true;
+	}
+
+	if (Collection && Collection->GetName().Contains(SearchQuery))
+	{
+		return true;
+	}
+
+	return Super::MatchesSearchQuery(SearchQuery);
+}
+#endif // WITH_EDITOR
+
 //
 //	UMaterialExpressionLightmapUVs
 //
@@ -25796,7 +26038,7 @@ int32 UMaterialExpressionCurveAtlasRowParameter::Compile(class FMaterialCompiler
 		int32 UV = Compiler->AppendVector(Arg1, Arg2);
 
 		// Sample texture
-		return Compiler->TextureSample(AtlasCode, UV, SAMPLERTYPE_LinearColor, INDEX_NONE, INDEX_NONE, TMVM_None, SSM_Clamp_WorldGroupSettings, AtlasRef, false);
+		return Compiler->TextureSample(AtlasCode, UV, SAMPLERTYPE_LinearColor, INDEX_NONE, INDEX_NONE, TMVM_None, SSM_Clamp_WorldGroupSettings, TGM_None, AtlasRef, false);
 	}
 	else
 	{
diff --git a/Engine/Source/Runtime/Engine/Private/Materials/ParameterCollection.cpp b/Engine/Source/Runtime/Engine/Private/Materials/ParameterCollection.cpp
--- a/Engine/Source/Runtime/Engine/Private/Materials/ParameterCollection.cpp
+++ b/Engine/Source/Runtime/Engine/Private/Materials/ParameterCollection.cpp
@@ -487,6 +487,11 @@ void UMaterialParameterCollection::GetParameterIndex(const FGuid& Id, int32& Out
 	}
 }
 
+int32 UMaterialParameterCollection::GetTotalVectorStorage() const
+{
+	return FMath::DivideAndRoundUp(ScalarParameters.Num(), 4) + VectorParameters.Num();
+}
+
 void UMaterialParameterCollection::GetParameterNames(TArray<FName>& OutParameterNames, bool bVectorParameters) const
 {
 	if (bVectorParameters)
diff --git a/Engine/Source/Runtime/Engine/Public/MaterialCompiler.h b/Engine/Source/Runtime/Engine/Public/MaterialCompiler.h
--- a/Engine/Source/Runtime/Engine/Public/MaterialCompiler.h
+++ b/Engine/Source/Runtime/Engine/Public/MaterialCompiler.h
@@ -9,6 +9,7 @@
 #include "CoreMinimal.h"
 #include "Misc/Guid.h"
 #include "MaterialShared.h"
+#include "Materials/MaterialExpressionCollectionTransform.h"
 #include "Materials/MaterialExpressionSpeedTree.h"
 #include "Materials/MaterialExpressionTextureSample.h"
 #include "Materials/MaterialExpressionLocalPosition.h"
@@ -201,6 +202,8 @@ public:
 	virtual int32 GetCurrentFunctionStackDepth() = 0;
 
 	virtual int32 AccessCollectionParameter(UMaterialParameterCollection* ParameterCollection, int32 ParameterIndex, int32 ComponentIndex) = 0;	
+	virtual int32 CollectionTransform(int32 InputIndex, const TStaticArray<int32, 5>& CollectionParameters, EParameterCollectionTransformType TransformType) = 0;
+
 	virtual int32 NumericParameter(EMaterialParameterType ParameterType, FName ParameterName, const UE::Shader::FValue& DefaultValue) = 0;
 
 	ENGINE_API int32 ScalarParameter(FName ParameterName, float DefaultValue);
@@ -286,7 +289,7 @@ public:
 
 	virtual int32 TextureCoordinate(uint32 CoordinateIndex, bool UnMirrorU, bool UnMirrorV) = 0;
 	virtual void  SetPotentiallyManipulateTexCoords() = 0;
-	virtual int32 TextureSample(int32 Texture,int32 Coordinate,enum EMaterialSamplerType SamplerType,int32 MipValue0Index=INDEX_NONE,int32 MipValue1Index=INDEX_NONE,ETextureMipValueMode MipValueMode=TMVM_None,ESamplerSourceMode SamplerSource=SSM_FromTextureAsset,int32 TextureReferenceIndex=INDEX_NONE, bool AutomaticViewMipBias=false, bool AdaptiveVirtualTexture=false, bool EnableFeedback = true) = 0;
+	virtual int32 TextureSample(int32 Texture,int32 Coordinate,enum EMaterialSamplerType SamplerType,int32 MipValue0Index=INDEX_NONE,int32 MipValue1Index=INDEX_NONE,ETextureMipValueMode MipValueMode=TMVM_None,ESamplerSourceMode SamplerSource=SSM_FromTextureAsset,ETextureGatherMode GatherMode=TGM_None,int32 TextureReferenceIndex=INDEX_NONE, bool AutomaticViewMipBias=false, bool AdaptiveVirtualTexture=false, bool EnableFeedback = true) = 0;
 	virtual int32 TextureProperty(int32 InTexture, EMaterialExposedTextureProperty Property) = 0;
 
 	virtual int32 TextureDecalMipmapLevel(int32 TextureSizeInput) = 0;
@@ -887,6 +890,8 @@ public:
 	virtual int32 TruncateLWC(int32 Code) override { return Compiler->TruncateLWC(Code); }
 
 	virtual int32 AccessCollectionParameter(UMaterialParameterCollection* ParameterCollection, int32 ParameterIndex, int32 ComponentIndex) override { return Compiler->AccessCollectionParameter(ParameterCollection, ParameterIndex, ComponentIndex); }
+	virtual int32 CollectionTransform(int32 InputIndex, const TStaticArray<int32, 5>& CollectionParameters, EParameterCollectionTransformType TransformType) override { return Compiler->CollectionTransform(InputIndex, CollectionParameters, TransformType); }
+
 	virtual int32 NumericParameter(EMaterialParameterType ParameterType, FName ParameterName, const UE::Shader::FValue& DefaultValue) override { return Compiler->NumericParameter(ParameterType, ParameterName, DefaultValue); }
 
 	virtual int32 Constant(float X) override { return Compiler->Constant(X); }
@@ -957,9 +962,9 @@ public:
 	virtual int32 If(int32 A, int32 B, int32 AGreaterThanB, int32 AEqualsB, int32 ALessThanB, int32 Threshold) override { return Compiler->If(A, B, AGreaterThanB, AEqualsB, ALessThanB, Threshold); }
 	virtual int32 Switch(int32 SwitchValueInput, int32 DefaultInput, TArray<int32>& CompiledInputs) override { return Compiler->Switch(SwitchValueInput, DefaultInput, CompiledInputs); };
 	
-	virtual int32 TextureSample(int32 InTexture, int32 Coordinate, enum EMaterialSamplerType SamplerType, int32 MipValue0Index, int32 MipValue1Index, ETextureMipValueMode MipValueMode, ESamplerSourceMode SamplerSource, int32 TextureReferenceIndex, bool AutomaticViewMipBias, bool AdaptiveVirtualTexture, bool EnableFeedback) override
+	virtual int32 TextureSample(int32 InTexture, int32 Coordinate, enum EMaterialSamplerType SamplerType, int32 MipValue0Index, int32 MipValue1Index, ETextureMipValueMode MipValueMode, ESamplerSourceMode SamplerSource, ETextureGatherMode GatherMode, int32 TextureReferenceIndex, bool AutomaticViewMipBias, bool AdaptiveVirtualTexture, bool EnableFeedback) override
 	{
-		return Compiler->TextureSample(InTexture, Coordinate, SamplerType, MipValue0Index, MipValue1Index, MipValueMode, SamplerSource, TextureReferenceIndex, AutomaticViewMipBias, AdaptiveVirtualTexture);
+		return Compiler->TextureSample(InTexture, Coordinate, SamplerType, MipValue0Index, MipValue1Index, MipValueMode, SamplerSource, GatherMode, TextureReferenceIndex, AutomaticViewMipBias, AdaptiveVirtualTexture);
 	}
 	virtual int32 TextureProperty(int32 InTexture, EMaterialExposedTextureProperty Property) override
 	{
diff --git a/Engine/Source/Runtime/Engine/Public/Materials/MaterialExpressionCollectionTransform.h b/Engine/Source/Runtime/Engine/Public/Materials/MaterialExpressionCollectionTransform.h
new file mode 100644
--- /dev/null
+++ b/Engine/Source/Runtime/Engine/Public/Materials/MaterialExpressionCollectionTransform.h
@@ -0,0 +1,73 @@
+// Copyright Epic Games, Inc. All Rights Reserved.
+
+/**
+ * MaterialExpressionCollectionTransform.h - a node that references a transform (5 consecutive vector elements) in a MaterialParameterCollection
+ */
+
+#pragma once
+
+#include "CoreMinimal.h"
+#include "UObject/ObjectMacros.h"
+#include "Misc/Guid.h"
+#include "Materials/MaterialExpression.h"
+#include "MaterialExpressionCollectionTransform.generated.h"
+
+struct FPropertyChangedEvent;
+
+UENUM()
+enum class EParameterCollectionTransformType : uint8
+{
+	Position UMETA(ToolTip = "Float 4x4 matrix (includes translation, 4 input elements, output float3 or float4 depending on input)"),
+	Vector UMETA(ToolTip = "Float 3x3 matrix (rotation only, 3 input elements, output float3)"),
+	Projection UMETA(ToolTip = "Float 4x4 projection matrix (optimized assuming only m11 m22 m33 m34 m43 m44 non-zero, 4 input elements, output float4)"),
+	LocalToWorld UMETA(ToolTip = "LWC matrix (float 4x4, post tile offset, 5 vector elements, output LWC float3)"),
+	WorldToLocal UMETA(ToolTip = "LWC inverse matrix (float 4x4, pre tile offset, 5 vector elements, output float3 or float4 depending on input)")
+};
+
+// Node that uses 3 to 5 consecutive vector elements of a Material Parameter Collection as a Transform matrix
+UCLASS(hidecategories = object, MinimalAPI)
+class UMaterialExpressionCollectionTransform : public UMaterialExpression
+{
+	GENERATED_UCLASS_BODY()
+
+	UPROPERTY()
+	FExpressionInput Input;
+
+	/** The Parameter Collection to use. */
+	UPROPERTY(EditAnywhere, Category=MaterialExpressionCollectionTransform)
+	TObjectPtr<class UMaterialParameterCollection> Collection;
+
+	/** Name of the parameter being referenced. */
+	UPROPERTY(EditAnywhere, Category=MaterialExpressionCollectionTransform)
+	FName ParameterName;
+
+	/** Id that is set from the name, and used to handle renaming of collection parameters. */
+	UPROPERTY()
+	FGuid ParameterId;
+
+	UPROPERTY(EditAnywhere, Category=MaterialExpressionCollectionTransform)
+	EParameterCollectionTransformType TransformType;
+
+	//~ Begin UObject Interface.
+#if WITH_EDITOR
+	virtual void PostEditChangeProperty(FPropertyChangedEvent& PropertyChangedEvent) override;
+#endif // WITH_EDITOR
+	virtual void PostLoad() override;
+	//~ End UObject Interface.
+
+	//~ Begin UMaterialExpression Interface
+#if WITH_EDITOR
+	virtual int32 Compile(class FMaterialCompiler* Compiler, int32 OutputIndex) override;
+	virtual void GetCaption(TArray<FString>& OutCaptions) const override;
+
+	virtual bool HasAParameterName() const override { return true; }
+	virtual FName GetParameterName() const override { return ParameterName; }
+	virtual void SetParameterName(const FName& Name) override { ParameterName = Name; }
+	virtual void ValidateParameterName(const bool bAllowDuplicateName) override {};
+
+	virtual bool MatchesSearchQuery( const TCHAR* SearchQuery ) override;
+
+	virtual bool GenerateHLSLExpression(FMaterialHLSLGenerator& Generator, UE::HLSLTree::FScope& Scope, int32 OutputIndex, UE::HLSLTree::FExpression const*& OutExpression) const override;
+#endif
+	//~ End UMaterialExpression Interface
+};
diff --git a/Engine/Source/Runtime/Engine/Public/Materials/MaterialExpressionTextureSample.h b/Engine/Source/Runtime/Engine/Public/Materials/MaterialExpressionTextureSample.h
--- a/Engine/Source/Runtime/Engine/Public/Materials/MaterialExpressionTextureSample.h
+++ b/Engine/Source/Runtime/Engine/Public/Materials/MaterialExpressionTextureSample.h
@@ -13,6 +13,18 @@ struct FPropertyChangedEvent;
 enum ETextureMipValueMode : int;
 enum ESamplerSourceMode : int;
 
+UENUM()
+enum ETextureGatherMode : int
+{
+	TGM_None UMETA(DisplayName = "None"),
+	TGM_Red UMETA(DisplayName = "Red"),
+	TGM_Green UMETA(DisplayName = "Green"),
+	TGM_Blue UMETA(DisplayName = "Blue"),
+	TGM_Alpha UMETA(DisplayName = "Alpha"),
+
+	TGM_MAX,
+};
+
 UCLASS(collapsecategories, hidecategories=Object, MinimalAPI)
 class UMaterialExpressionTextureSample : public UMaterialExpressionTextureBase
 {
@@ -56,6 +68,13 @@ class UMaterialExpressionTextureSample : public UMaterialExpressionTextureBase
 	UPROPERTY(EditAnywhere, Category=MaterialExpressionTextureSample, Meta = (ShowAsInputPin = "Advanced"))
 	TEnumAsByte<ESamplerSourceMode> SamplerSource;
 
+	/**
+	 * Whether to do a Gather of the given component from 4 neighboring texels (xyzw in counter clockwise order starting with the sample to the lower left).
+	 * Only works with Texture2D or TextureCube.  Doesn't work with virtual textures.
+	 */
+	UPROPERTY(EditAnywhere, Category = MaterialExpressionTextureSample)
+	TEnumAsByte<ETextureGatherMode> GatherMode;
+
 	/** Whether the texture should be sampled with per view mip biasing for sharper output with Temporal AA. */
 	UPROPERTY(EditAnywhere, Category = MaterialExpressionTextureSample)
 	uint8 AutomaticViewMipBias : 1;
diff --git a/Engine/Source/Runtime/Engine/Public/Materials/MaterialParameterCollection.h b/Engine/Source/Runtime/Engine/Public/Materials/MaterialParameterCollection.h
--- a/Engine/Source/Runtime/Engine/Public/Materials/MaterialParameterCollection.h
+++ b/Engine/Source/Runtime/Engine/Public/Materials/MaterialParameterCollection.h
@@ -154,6 +154,9 @@ class UMaterialParameterCollection : public UObject
 	/** Gets a vector and component index for the given parameter, used when compiling materials, to know where to access a certain parameter. */
 	ENGINE_API void GetParameterIndex(const FGuid& Id, int32& OutIndex, int32& OutComponentIndex) const;
 
+	/** Returns total vector storage of the parameter collection, used for math involving OutIndex returned by GetParameterIndex above, when compiling materials. */
+	ENGINE_API int32 GetTotalVectorStorage() const;
+
 	/** Populates an array with either scalar or vector parameter names. */
 	ENGINE_API void GetParameterNames(TArray<FName>& OutParameterNames, bool bVectorParameters) const;
 
diff --git a/Engine/Source/Runtime/Renderer/Private/SceneCaptureRendering.cpp b/Engine/Source/Runtime/Renderer/Private/SceneCaptureRendering.cpp
--- a/Engine/Source/Runtime/Renderer/Private/SceneCaptureRendering.cpp
+++ b/Engine/Source/Runtime/Renderer/Private/SceneCaptureRendering.cpp
@@ -48,6 +48,8 @@
 #include "RenderCaptureInterface.h"
 #include "CustomRenderPassSceneCapture.h"
 #include "SceneRenderBuilder.h"
+#include "Materials/MaterialParameterCollection.h"
+#include "Materials/MaterialParameterCollectionInstance.h"
 
 bool GSceneCaptureAllowRenderInMainRenderer = true;
 static FAutoConsoleVariableRef CVarSceneCaptureAllowRenderInMainRenderer(
@@ -1120,6 +1122,87 @@ void FScene::UpdateSceneCaptureContents(USceneCaptureComponent2D* CaptureCompone
 			}
 		}
 
+		UMaterialParameterCollection* CollectionTransformTarget = CaptureComponent->CollectionTransform.Get();
+		if (CollectionTransformTarget)
+		{
+			UMaterialParameterCollectionInstance* MaterialParameterCollectionInstance = World->GetParameterCollectionInstance(CollectionTransformTarget);
+			bool bCollectionModified = false;
+
+			// Find the parameters in the vector array
+			int32 ParameterWorldToLocal = INDEX_NONE;
+			int32 ParameterProjection = INDEX_NONE;
+			for (int32 ParameterIndex = 0; ParameterIndex < CollectionTransformTarget->VectorParameters.Num(); ParameterIndex++)
+			{
+				FName ParameterName = CollectionTransformTarget->VectorParameters[ParameterIndex].ParameterName;
+				if (!ParameterName.IsNone())
+				{
+					if (ParameterName == CaptureComponent->CollectionTransformWorldToLocal)
+					{
+						ParameterWorldToLocal = ParameterIndex;
+					}
+					else if (ParameterName == CaptureComponent->CollectionTransformProjection)
+					{
+						ParameterProjection = ParameterIndex;
+					}
+				}
+			}
+
+			// Structure to hold a matrix, plus LWC tile offset.
+			struct FMatrixPlusTileOffset
+			{
+				FMatrix44f Matrix;
+				FVector3f TileOffset;
+				float Padding = 0.0f;
+			};
+
+			// Ensure there's space for 5 output vectors for the world to local matrix
+			if (ParameterWorldToLocal != INDEX_NONE && ParameterWorldToLocal + 5 <= CollectionTransformTarget->VectorParameters.Num())
+			{
+				// Generate a view matrix (world space to local space) in LWC format.  This involves generating a tile origin using math
+				// copied from FRelativeViewMatrices::Create, then subtracting that from the ViewLocation.  The tile offset gets subtracted
+				// from the world position in the shader, before applying the matrix.
+				const double TileSize = FLargeWorldRenderScalar::GetTileSize();
+				const FVector ViewOriginTile = FLargeWorldRenderScalar::MakeQuantizedTile(ViewLocation, 8.0);
+
+				FMatrix ViewMatrix = FTranslationMatrix(-(ViewLocation - ViewOriginTile * TileSize)) * ViewRotationMatrix;
+
+				FMatrixPlusTileOffset MatrixPlusTileOffset;
+				MatrixPlusTileOffset.Matrix = FMatrix44f(ViewMatrix);
+				MatrixPlusTileOffset.TileOffset = (FVector3f)(ViewOriginTile);
+
+				// Store the vectors to the collection instance
+				const FLinearColor* MatrixPlusTileOffsetVectors = (const FLinearColor*)&MatrixPlusTileOffset;
+				for (int32 ElementIndex = 0; ElementIndex < 5; ElementIndex++)
+				{
+					MaterialParameterCollectionInstance->SetVectorParameterValue(CollectionTransformTarget->VectorParameters[ParameterWorldToLocal + ElementIndex].ParameterName, MatrixPlusTileOffsetVectors[ElementIndex]);
+				}
+
+				bCollectionModified = true;
+			}
+
+			// Ensure there's space for 4 output vectors for the projection matrix
+			if (ParameterProjection != INDEX_NONE && ParameterProjection + 4 <= CollectionTransformTarget->VectorParameters.Num())
+			{
+				FMatrix44f ProjectionMatrixf = FMatrix44f(ProjectionMatrix);
+
+				// Store the vectors to the collection instance
+				const FLinearColor* MatrixVectors = (const FLinearColor*)&ProjectionMatrixf;
+				for (int32 ElementIndex = 0; ElementIndex < 4; ElementIndex++)
+				{
+					MaterialParameterCollectionInstance->SetVectorParameterValue(CollectionTransformTarget->VectorParameters[ParameterProjection + ElementIndex].ParameterName, MatrixVectors[ElementIndex]);
+				}
+
+				bCollectionModified = true;
+			}
+
+			if (bCollectionModified)
+			{
+				// Rendering runs after the world tick, where deferred material parameter collection updates normally occur,
+				// so we need to manually update here, or the update will be delayed by a frame.
+				MaterialParameterCollectionInstance->DeferredUpdateRenderState(false);
+			}
+		}
+
 		// As optimization for depth capture modes, render scene capture as additional render passes inside the main renderer.
 		if (GSceneCaptureAllowRenderInMainRenderer && CaptureComponent->ShouldRenderInMainRenderer())
 		{
diff --git a/Engine/Source/Runtime/Renderer/Private/SceneRendering.cpp b/Engine/Source/Runtime/Renderer/Private/SceneRendering.cpp
--- a/Engine/Source/Runtime/Renderer/Private/SceneRendering.cpp
+++ b/Engine/Source/Runtime/Renderer/Private/SceneRendering.cpp
@@ -4995,13 +4995,19 @@ void FRendererModule::BeginRenderingViewFamilies(FCanvas* Canvas, TConstArrayVie
 
 		FSceneRenderBuilder SceneRenderBuilder(Scene);
 
+		bool bShowHitProxies = (Canvas->GetHitProxyConsumer() != nullptr);
+		if (!bShowHitProxies)
+		{
+			// Custom render passes need to be added before the main scene renderer is allocated, otherwise their update will be delayed by a frame
+			SceneCaptureUpdateDeferredCapturesInternal(Scene, ViewFamilies, SceneRenderBuilder, /*bCustomRenderPassOnly=*/ true);
+		}
+
 		TArray<FSceneRenderer*, FConcurrentLinearArrayAllocator> SceneRenderers = SceneRenderBuilder.CreateLinkedSceneRenderers(ViewFamilies, Canvas->GetHitProxyConsumer());
 		SetupDebugViewModes(SceneRenderers);
 
-		bool bShowHitProxies = (Canvas->GetHitProxyConsumer() != nullptr);
 		if (!bShowHitProxies)
 		{
-			SceneCaptureUpdateDeferredCapturesInternal(Scene, ViewFamilies, SceneRenderBuilder);
+			SceneCaptureUpdateDeferredCapturesInternal(Scene, ViewFamilies, SceneRenderBuilder, /*bCustomRenderPassOnly=*/ false);
 		}
 
 		if (!bShowHitProxies)
