Update EEditorPropertyValueState to work in BP

#rnx

diff --git a/Engine/Source/Runtime/Engine/Classes/Kismet/KismetSystemLibrary.h b/Engine/Source/Runtime/Engine/Classes/Kismet/KismetSystemLibrary.h
--- a/Engine/Source/Runtime/Engine/Classes/Kismet/KismetSystemLibrary.h
+++ b/Engine/Source/Runtime/Engine/Classes/Kismet/KismetSystemLibrary.h
@@ -73,7 +73,7 @@ namespace EQuitPreference
 
 /** States a property value queried by IsEditorPropertyOverridden can be in */
 UENUM()
-enum class EEditorPropertyValueState
+enum class EEditorPropertyValueState : uint8
 {
 	/** This property is in a default state */
 	Default,
