added experimental bRunTransactionally flag to FTickFunction
remove previous hacks to actor+component tick functions for running them transactionally, moved the call to FTickFunctionTask

diff --git a/Engine/Source/Runtime/Engine/Classes/Engine/EngineBaseTypes.h b/Engine/Source/Runtime/Engine/Classes/Engine/EngineBaseTypes.h
--- a/Engine/Source/Runtime/Engine/Classes/Engine/EngineBaseTypes.h
+++ b/Engine/Source/Runtime/Engine/Classes/Engine/EngineBaseTypes.h
@@ -220,6 +220,9 @@ public:
 	/** If false, this tick will run on the game thread, otherwise it will run on any thread in parallel with the game thread and in parallel with other "async ticks" **/
 	uint8 bRunOnAnyThread:1;
 
+	/** (experimental) if true, the tick function will be run transactionally */
+	uint8 bRunTransactionally:1;
+
 private:
 
 	enum class ETickState : uint8
@@ -517,8 +520,6 @@ struct FActorTickFunction : public FTickFunction
 		* @param MyCompletionGraphEvent - completion event for this task. Useful for holding the completetion of this task until certain child tasks are complete.
 	**/
 	ENGINE_API virtual void ExecuteTick(float DeltaTime, ELevelTick TickType, ENamedThreads::Type CurrentThread, const FGraphEventRef& MyCompletionGraphEvent) override;
-	ENGINE_API void ExecuteActorTickImpl(float DeltaTime, ELevelTick TickType, ENamedThreads::Type CurrentThread, const FGraphEventRef& MyCompletionGraphEvent);
-
 	/** Abstract function to describe this tick. Used to print messages about illegal cycles in the dependency graph **/
 	ENGINE_API virtual FString DiagnosticMessage() override;
 	ENGINE_API virtual FName DiagnosticContext(bool bDetailed) override;
diff --git a/Engine/Source/Runtime/Engine/Private/Actor.cpp b/Engine/Source/Runtime/Engine/Private/Actor.cpp
--- a/Engine/Source/Runtime/Engine/Private/Actor.cpp
+++ b/Engine/Source/Runtime/Engine/Private/Actor.cpp
@@ -77,7 +77,6 @@
 #include "WorldPartition/DataLayer/DataLayerManager.h"
 #include "WorldPartition/ContentBundle/ContentBundlePaths.h"
 #include "WorldPartition/ActorInstanceGuids.h"
-#include "UObject/RemoteObjectTransfer.h"
 
 DEFINE_LOG_CATEGORY(LogActor);
 
@@ -179,6 +178,7 @@ void AActor::InitializeDefaults()
 	PrimaryActorTick.bCanEverTick = false;
 	PrimaryActorTick.bStartWithTickEnabled = true;
 	PrimaryActorTick.bAllowTickBatching = true;
+	PrimaryActorTick.bRunTransactionally = true;
 	PrimaryActorTick.SetTickFunctionEnable(false); 
 	bAsyncPhysicsTickEnabled = false;
 
@@ -268,20 +268,6 @@ void AActor::InitializeDefaults()
 }
 
 void FActorTickFunction::ExecuteTick(float DeltaTime, enum ELevelTick TickType, ENamedThreads::Type CurrentThread, const FGraphEventRef& MyCompletionGraphEvent)
-{
-#if UE_WITH_REMOTE_OBJECT_HANDLE
-	auto ExecuteTickWork = [this, DeltaTime, TickType, CurrentThread, &MyCompletionGraphEvent]()
-	{
-		ExecuteActorTickImpl(DeltaTime, TickType, CurrentThread, MyCompletionGraphEvent);
-	};
-
-	UE::RemoteObject::Transfer::ExecuteTransactional(ExecuteTickWork);
-#else
-	ExecuteActorTickImpl(DeltaTime, TickType, CurrentThread, MyCompletionGraphEvent);
-#endif
-}
-
-void FActorTickFunction::ExecuteActorTickImpl(float DeltaTime, enum ELevelTick TickType, ENamedThreads::Type CurrentThread, const FGraphEventRef& MyCompletionGraphEvent)
 {
 	if (IsValid(Target))
 	{
diff --git a/Engine/Source/Runtime/Engine/Private/Components/ActorComponent.cpp b/Engine/Source/Runtime/Engine/Private/Components/ActorComponent.cpp
--- a/Engine/Source/Runtime/Engine/Private/Components/ActorComponent.cpp
+++ b/Engine/Source/Runtime/Engine/Private/Components/ActorComponent.cpp
@@ -42,7 +42,6 @@
 #include "Materials/MaterialInterface.h"
 #include "ObjectCacheContext.h"
 #include "AutoRTFM/AutoRTFM.h"
-#include "UObject/RemoteObjectTransfer.h"
 #include "Streaming/AsyncRegisterLevelContext.h"
 
 #if WITH_EDITOR
@@ -555,6 +554,7 @@ UActorComponent::UActorComponent(const FObjectInitializer& ObjectInitializer /*=
 	PrimaryComponentTick.bStartWithTickEnabled = true;
 	PrimaryComponentTick.bCanEverTick = false;
 	PrimaryComponentTick.bAllowTickBatching = true;
+	PrimaryComponentTick.bRunTransactionally = true;
 	PrimaryComponentTick.SetTickFunctionEnable(false);
 
 	MarkedForEndOfFrameUpdateArrayIndex = INDEX_NONE;
@@ -1416,22 +1416,10 @@ void FActorComponentTickFunction::ExecuteTick(float DeltaTime, enum ELevelTick T
 {
 	TRACE_CPUPROFILER_EVENT_SCOPE(FActorComponentTickFunction::ExecuteTick);
 
-#if UE_WITH_REMOTE_OBJECT_HANDLE
-	auto ExecuteTickWork = [this, DeltaTime, TickType, CurrentThread, &MyCompletionGraphEvent]()
-		{
-			ExecuteTickHelper(Target, Target->bTickInEditor, DeltaTime, TickType, [this, TickType](float DilatedTime)
-			{
-				Target->TickComponent(DilatedTime, TickType, this);
-			});
-		};
-
-	UE::RemoteObject::Transfer::ExecuteTransactional(ExecuteTickWork);
-#else
 	ExecuteTickHelper(Target, Target->bTickInEditor, DeltaTime, TickType, [this, TickType](float DilatedTime)
 	{
 		Target->TickComponent(DilatedTime, TickType, this);
 	});
-#endif
 }
 
 FString FActorComponentTickFunction::DiagnosticMessage()
diff --git a/Engine/Source/Runtime/Engine/Private/TickTaskManager.cpp b/Engine/Source/Runtime/Engine/Private/TickTaskManager.cpp
--- a/Engine/Source/Runtime/Engine/Private/TickTaskManager.cpp
+++ b/Engine/Source/Runtime/Engine/Private/TickTaskManager.cpp
@@ -12,6 +12,7 @@
 #include "Async/ParallelFor.h"
 #include "Misc/TimeGuard.h"
 #include "Templates/UniquePtr.h"
+#include "UObject/RemoteObjectTransfer.h"
 
 DEFINE_LOG_CATEGORY_STATIC(LogTick, Log, All);
 
@@ -303,7 +304,26 @@ public:
 			SCOPE_TIME_GUARD_DELEGATE_MS(NameFunction, 4);
 #endif
 			LIGHTWEIGHT_TIME_GUARD_BEGIN(FTickFunctionTask, GTimeguardThresholdMS);
+
+#if UE_WITH_REMOTE_OBJECT_HANDLE
+			auto ExecuteTickWork = [this, CurrentThread, &MyCompletionGraphEvent]()
+			{
+#endif
 			Target->ExecuteTick(Target->CalculateDeltaTime(Context.DeltaSeconds, Context.World), Context.TickType, CurrentThread, MyCompletionGraphEvent);
+
+#if UE_WITH_REMOTE_OBJECT_HANDLE
+			};
+
+			if (Target->bRunTransactionally)
+			{
+				UE::RemoteObject::Transfer::ExecuteTransactional(ExecuteTickWork);
+			}
+			else
+			{
+				ExecuteTickWork();
+			}
+#endif
+
 			LIGHTWEIGHT_TIME_GUARD_END(FTickFunctionTask, Target->DiagnosticMessage());
 		}
 		Target->ClearTaskInformation();  // This is stale and a good time to clear it for safety
