Fix crash on exiting PIE when data-driven-cvar callbacks fire in Blueprint during teardown

#rnx

diff --git a/Engine/Source/Runtime/Engine/Private/UnrealEngine.cpp b/Engine/Source/Runtime/Engine/Private/UnrealEngine.cpp
--- a/Engine/Source/Runtime/Engine/Private/UnrealEngine.cpp
+++ b/Engine/Source/Runtime/Engine/Private/UnrealEngine.cpp
@@ -13791,24 +13791,30 @@ const TArray<class ULocalPlayer*>& HandleFakeLocalPlayersList()
 	return FakeEmptyLocalPlayers;
 }
 
-const TArray<class ULocalPlayer*>& UEngine::GetGamePlayers(UWorld *World) const
+const TArray<class ULocalPlayer*>& UEngine::GetGamePlayers(UWorld* World) const
 {
-	const FWorldContext &Context = GetWorldContextFromWorldChecked(World);
-	if ( Context.OwningGameInstance == NULL )
+	const FWorldContext* Context = GetWorldContextFromWorld(World);
+	if (Context && Context->OwningGameInstance)
+	{
+		return Context->OwningGameInstance->GetLocalPlayers();
+	}
+	else
 	{
 		return HandleFakeLocalPlayersList();
 	}
-	return Context.OwningGameInstance->GetLocalPlayers();
 }
 
-const TArray<class ULocalPlayer*>& UEngine::GetGamePlayers(const UGameViewportClient *Viewport) const
+const TArray<class ULocalPlayer*>& UEngine::GetGamePlayers(const UGameViewportClient* Viewport) const
 {
-	const FWorldContext &Context = GetWorldContextFromGameViewportChecked(Viewport);
-	if ( Context.OwningGameInstance == NULL )
+	const FWorldContext* Context = GetWorldContextFromGameViewport(Viewport);
+	if (Context && Context->OwningGameInstance)
+	{
+		return Context->OwningGameInstance->GetLocalPlayers();
+	}
+	else
 	{
 		return HandleFakeLocalPlayersList();
 	}
-	return Context.OwningGameInstance->GetLocalPlayers();
 }
 
 ULocalPlayer* UEngine::FindFirstLocalPlayerFromControllerId(int32 ControllerId) const
