Fixed bug occurring when trying to duplicate a default Material with an incorrect EditorOnlyData name.

[RN] Fixed bug occurring when trying to duplicate a default Material with an incorrect EditorOnlyData name.

diff --git a/Engine/Source/Runtime/Engine/Private/Materials/MaterialInterface.cpp b/Engine/Source/Runtime/Engine/Private/Materials/MaterialInterface.cpp
--- a/Engine/Source/Runtime/Engine/Private/Materials/MaterialInterface.cpp
+++ b/Engine/Source/Runtime/Engine/Private/Materials/MaterialInterface.cpp
@@ -112,16 +112,32 @@ static EMaterialGetParameterValueFlags MakeParameterValueFlags(bool bOveriddenOn
 
 #if WITH_EDITORONLY_DATA
 
-/** Deletes any invalid EditorOnlyData object with specified name that isn't the one assigned to specified interface. */
-static void DisposeInvalidEditorOnlyData(UMaterialInterface* MaterialInterface, const FString& EditorOnlyDataName)
+// Deletes any invalid EditorOnlyData object with specified name that isn't the one assigned to specified interface.
+// Returns whether the MI's EditorOnlyData was renamed.
+static bool SafeRenameEditorOnlyData(UMaterialInterface* MaterialInterface, const FString& EditorOnlyDataName)
 {
-	UObject* EditorOnlyExisting = StaticFindObject(/*Class=*/ nullptr, MaterialInterface, *EditorOnlyDataName, true);
-	if (EditorOnlyExisting && EditorOnlyExisting != MaterialInterface->GetEditorOnlyData())
+	UObject* EditorOnlyCurrent = MaterialInterface->GetEditorOnlyData();
+	UObject* EditorOnlyWithIncorrectName = StaticFindObject(/*Class=*/ nullptr, MaterialInterface, *EditorOnlyDataName, true);
+	if (EditorOnlyWithIncorrectName && EditorOnlyWithIncorrectName != EditorOnlyCurrent)
 	{
+		if (EditorOnlyWithIncorrectName->IsRooted())
+		{
+			UE_LOG(LogMaterial, Warning, TEXT("MaterialInterface '%s' editor only data ('%s') could not be marked for garbage bacause it was rooted. Please resave this material."), *MaterialInterface->GetName(), *EditorOnlyWithIncorrectName->GetName());
+			return false;
+		}
 		FName UniqueDummyName = MakeUniqueObjectName(GetTransientPackage(), UMaterialInterfaceEditorOnlyData::StaticClass());
-		EditorOnlyExisting->Rename(*UniqueDummyName.ToString(), GetTransientPackage(), REN_NonTransactional | REN_DontCreateRedirectors);
-		EditorOnlyExisting->MarkAsGarbage();
+		EditorOnlyWithIncorrectName->Rename(*UniqueDummyName.ToString(), GetTransientPackage(), REN_NonTransactional | REN_DontCreateRedirectors);
+		EditorOnlyWithIncorrectName->MarkAsGarbage();
+	}
+
+	if (EditorOnlyCurrent->GetName() != EditorOnlyDataName)
+	{
+		UE_LOG(LogMaterial, Display, TEXT("MaterialInterface %s has a incorrectly name EditorOnlyData '%s'. This may cause issues when running the editor on cooked data. Trying to rename it to the correct name '%s'."), *MaterialInterface->GetName(), *EditorOnlyCurrent->GetName(), *EditorOnlyDataName);
+
+		return EditorOnlyCurrent->Rename(*EditorOnlyDataName, nullptr, REN_NonTransactional | REN_DontCreateRedirectors);
 	}
+
+	return false;
 }
 
 #endif
@@ -763,14 +779,10 @@ void UMaterialInterface::PostDuplicate(bool bDuplicateForPIE)
 	SetLightingGuid();
 
 #if WITH_EDITORONLY_DATA
-	// Duplication will create a new editor only data but this MI already links a populated data. This makes sure that the
-	// duplicate instance EOD has the correct name (necessary for running the editor on cooked data).
+	// If this MI references its editor only data, make sure it has the correct name, or the editor
+	// won't run properly when running on cooked data.
 	FString EditorOnlyDataName = UE::MaterialInterface::GetEditorOnlyDataName(*GetName());
-	DisposeInvalidEditorOnlyData(this, EditorOnlyDataName);
-	if (GetEditorOnlyData())
-	{
-		GetEditorOnlyData()->Rename(*EditorOnlyDataName, nullptr, REN_NonTransactional | REN_DontCreateRedirectors);
-	}
+	SafeRenameEditorOnlyData(this, EditorOnlyDataName);
 #endif // WITH_EDITORONLY_DATA
 }
 
@@ -1923,18 +1935,11 @@ void UMaterialInterface::PreSave(FObjectPreSaveContext ObjectSaveContext)
 	}
 
 #if WITH_EDITORONLY_DATA
-	// Make sure the EditorOnlyData is named correctly. Some MI assets have a differently named editor only data that can cause problems in
-	// the editor running on cooked data.
-	UMaterialInterfaceEditorOnlyData* EditorOnly = GetEditorOnlyData();
+	// Make sure the EditorOnlyData is named correctly. Some MI assets have a differently named editor only data that can cause problems in the editor running on cooked data.
 	FString EditorOnlyDataName = UE::MaterialInterface::GetEditorOnlyDataName(*GetName());
-	if (!ObjectSaveContext.IsProceduralSave() && EditorOnly && EditorOnly->GetName() != EditorOnlyDataName)
+	if (!ObjectSaveContext.IsProceduralSave())
 	{
-		UE_LOG(LogMaterial, Display, TEXT("MaterialInterface %s has a incorrectly name EditorOnlyData '%s'. This may cause issues when running the editor on cooked data. Trying to rename it to the correct name '%s'."), *GetName(), *EditorOnly->GetName(), *EditorOnlyDataName);
-		DisposeInvalidEditorOnlyData(this, EditorOnlyDataName);
-		if (EditorOnly)
-		{
-			EditorOnly->Rename(*EditorOnlyDataName, nullptr, REN_NonTransactional | REN_DontCreateRedirectors);
-		}
+		SafeRenameEditorOnlyData(this, EditorOnlyDataName);
 	}
 #endif // WITH_EDITORONLY_DATA
 
@@ -2064,8 +2069,7 @@ bool UMaterialInterface::Rename(const TCHAR* NewName, UObject* NewOuter, ERename
 	if (bRenamed && NewName && EditorOnlyData)
 	{
 		FString EditorOnlyDataName = UE::MaterialInterface::GetEditorOnlyDataName(NewName);
-		DisposeInvalidEditorOnlyData(this, EditorOnlyDataName);
-		bRenamed = EditorOnlyData->Rename(*EditorOnlyDataName, nullptr, Flags);
+		bRenamed = SafeRenameEditorOnlyData(this, EditorOnlyDataName);
 	}
 #endif
 	return bRenamed;
