Added support for multiple levels to be made visible or invisible at the same time using LevelStreaming.MaximumMakingVisibleLevels and LevelStreaming.MaximumMakingInvisibleLevels.
These can be used in conjunction with incremental component pre-registration/pre-unregistration and asynchronous phyiscs body creation/destruction (see p.Chaos.EnableAsyncInitBody, LevelStreaming.AllowIncrementalPreRegisterComponents and LevelStreaming.AllowIncrementalPreUnregisterComponents) to maximize the usage of the time limits for adding a level to the world and/or removing a level from the world.

diff --git a/Engine/Source/Runtime/Engine/Classes/Engine/Level.h b/Engine/Source/Runtime/Engine/Classes/Engine/Level.h
--- a/Engine/Source/Runtime/Engine/Classes/Engine/Level.h
+++ b/Engine/Source/Runtime/Engine/Classes/Engine/Level.h
@@ -390,6 +390,14 @@ enum class EActorPackagingScheme : uint8
 };
 #endif
 
+// Whether the level is in the process of being made visible or invisible (refer to bIsVisible when None).
+enum class ELevelPendingVisibilityState : uint8
+{
+	None,
+	MakingVisible,
+	MakingInvisible
+};
+
 //
 // The level object.  Contains the level's actor list, BSP information, and brush list.
 // Every Level has a World as its Outer and can be used as the PersistentLevel, however,
@@ -664,6 +672,12 @@ public:
 
 	UE_DEPRECATED(5.5, "This property shouldn't be used anymore.")
 	uint8										bHasCurrentActorCalledPreRegister:1;
+
+private:
+	/** Whether the level is in the process of being made visible or invisible (refer to bIsVisible when None). */
+	ELevelPendingVisibilityState				PendingVisibilityState:2;
+
+public:
 	/** The current stage for incrementally updating actor components in the level*/
 	EIncrementalComponentState					IncrementalComponentState;
 	/** Current index into actors array for updating components.							*/
@@ -848,6 +862,7 @@ private:
 	TSet<UActorComponent*> PreUnregisteringActorComponents;
 	TSet<UActorComponent*> PreRegisteringActorComponents;
 	friend struct FLevelRegistrationAccessor;
+	friend struct FLevelPendingVisibilityStateAccessor;
 
 protected:
 
@@ -1199,6 +1214,9 @@ public:
 
 	ENGINE_API bool IsMapBuildDataOwner() const;
 
+	/** Returns level's pending visibility state */
+	ENGINE_API ELevelPendingVisibilityState GetPendingVisibilityState() const;
+
 	/** Sets whether this level is a lighting scenario and handles propagating the change. */
 	ENGINE_API void SetLightingScenario(bool bNewIsLightingScenario);
 
@@ -1590,6 +1608,17 @@ private:
 	friend class UWorld;
 };
 
+struct FLevelPendingVisibilityStateAccessor
+{
+private:
+	ENGINE_API static void SetPendingVisibilityState(ULevel* InLevel, ELevelPendingVisibilityState InState)
+	{
+		InLevel->PendingVisibilityState = InState;
+	}
+
+	friend class UWorld;
+};
+
 #if WITH_EDITOR
 struct ENGINE_API FScopedOverrideSpawningLevelMountPointObject
 {
diff --git a/Engine/Source/Runtime/Engine/Classes/Engine/World.h b/Engine/Source/Runtime/Engine/Classes/Engine/World.h
--- a/Engine/Source/Runtime/Engine/Classes/Engine/World.h
+++ b/Engine/Source/Runtime/Engine/Classes/Engine/World.h
@@ -74,6 +74,7 @@ class FSceneView;
 struct FUniqueNetIdRepl;
 struct FEncryptionKeyResponse;
 struct FParticlePerfStats;
+enum class ELevelPendingVisibilityState : uint8;
 
 template<typename,typename> class TOctree2;
 
@@ -968,10 +969,24 @@ public:
 	bool HasStreamingLevelsToConsider() const { return StreamingLevelsToConsider.GetStreamingLevels().Num() > 0; }
 
 	/** Returns the level, if any, in the process of being made visible */
-	ULevel* GetCurrentLevelPendingVisibility() const { return CurrentLevelPendingVisibility; }
+	UE_DEPRECATED(5.6, "Use IsLevelMakingVisible instead.")
+	ULevel* GetCurrentLevelPendingVisibility() const { return nullptr; }
+
+	/** Returns true if the level is in the process of being made visible */
+	bool IsLevelMakingVisible(const ULevel* InLevel) const;
+
+	/** Returns true if any level is in the process of being made visible */
+	bool HasAnyLevelMakingVisible() const { return MakingVisibleLevels.Num() > 0; }
 
 	/** Returns the level, if any, in the process of being made invisible */
-	ULevel* GetCurrentLevelPendingInvisibility() const { return CurrentLevelPendingInvisibility; }
+	UE_DEPRECATED(5.6, "Use IsLevelMakingInvisible instead.")
+	ULevel* GetCurrentLevelPendingInvisibility() const { return nullptr; }
+
+	/** Returns true if the level is in the process of being made invisible */
+	bool IsLevelMakingInvisible(const ULevel* InLevel) const;
+
+	/** Returns true if any level is in the process of being made invisible */
+	bool HasAnyLevelMakingInvisible() const { return MakingInvisibleLevels.Num() > 0; }
 
 	/** Add a streaming level to the list of streamed levels to consider. */
 	void AddStreamingLevel(ULevelStreaming* StreamingLevelToAdd);
@@ -1030,6 +1045,9 @@ public:
 
 private:
 
+	/** Sets level's PendingVisibilityState */
+	void SetLevelPendingVisibilityState(ULevel* InLevel, ELevelPendingVisibilityState InState);
+
 	/** Returns wether AddToWorld should be skipped on a given level */
 	bool CanAddLoadedLevelToWorld(ULevel* Level) const;
 
@@ -1039,13 +1057,13 @@ private:
 	/** Stores whether the game world supports for a client to use "making invisible" transaction requests to the server. */
 	mutable TOptional<bool> bSupportsMakingInvisibleTransactionRequests;
 
-	/** Pointer to the current level in the queue to be made visible, NULL if none are pending. */
+	/** Level(s) in the queue to be made visible. */
 	UPROPERTY(Transient)
-	TObjectPtr<class ULevel>								CurrentLevelPendingVisibility;
+	TArray<TObjectPtr<class ULevel>>						MakingVisibleLevels;
 
-	/** Pointer to the current level in the queue to be made invisible, NULL if none are pending. */
+	/** Level(s) in the queue to be made invisible. */
 	UPROPERTY(Transient)
-	TObjectPtr<class ULevel>								CurrentLevelPendingInvisibility;
+	TArray<TObjectPtr<class ULevel>>						MakingInvisibleLevels;
 
 	/** NetDriver for capturing network traffic to record demos */
 	UPROPERTY()
diff --git a/Engine/Source/Runtime/Engine/Private/Level.cpp b/Engine/Source/Runtime/Engine/Private/Level.cpp
--- a/Engine/Source/Runtime/Engine/Private/Level.cpp
+++ b/Engine/Source/Runtime/Engine/Private/Level.cpp
@@ -3893,6 +3893,11 @@ bool ULevel::IsMapBuildDataOwner() const
 	return bIsMapBuildDataOwner;
 }
 
+ELevelPendingVisibilityState ULevel::GetPendingVisibilityState() const
+{
+	return PendingVisibilityState;
+}
+
 UMapBuildDataRegistry* ULevel::GetOrCreateMapBuildData()
 {
 	if (!MapBuildData 
@@ -4986,7 +4991,7 @@ void ULevel::RemoveUserDataOfClass(TSubclassOf<UAssetUserData> InUserDataClass)
 
 bool ULevel::HasVisibilityChangeRequestPending() const
 {
-	return (OwningWorld && ( this == OwningWorld->GetCurrentLevelPendingVisibility() || this == OwningWorld->GetCurrentLevelPendingInvisibility() ) );
+	return (OwningWorld && ( OwningWorld->IsLevelMakingVisible(this) || OwningWorld->IsLevelMakingInvisible(this) ) );
 }
 
 #if WITH_EDITOR
diff --git a/Engine/Source/Runtime/Engine/Private/LevelStreaming.cpp b/Engine/Source/Runtime/Engine/Private/LevelStreaming.cpp
--- a/Engine/Source/Runtime/Engine/Private/LevelStreaming.cpp
+++ b/Engine/Source/Runtime/Engine/Private/LevelStreaming.cpp
@@ -662,7 +662,7 @@ ELevelStreamingTargetState ULevelStreaming::DetermineTargetState() const
 	{
 	case ELevelStreamingState::MakingVisible:
 		ensure(LoadedLevel);
-		if (!ShouldBeVisible() && GetWorld()->GetCurrentLevelPendingVisibility() != LoadedLevel)
+		if (!ShouldBeVisible() && !World->IsLevelMakingVisible(LoadedLevel))
 		{
 			// Since level doesn't need to be visible anymore, change TargetState to ELevelStreamingTargetState::LoadedNotVisible.
 			// Next UpdateStreamingState will handle switching CurrentState to ELevelStreamingState::LoadedNotVisible.
@@ -853,7 +853,7 @@ bool ULevelStreaming::CanMakeInvisible()
 	// This is one hypothetical example. Detecting this will prevent any future case that could trigger an infinte loop in BlockTillLevelStreamingCompleted.
 	//
 	UWorld* World = GetWorld();
-	if (World && World->IsGameWorld() && LoadedLevel && (LoadedLevel == World->GetCurrentLevelPendingInvisibility()))
+	if (World && World->IsGameWorld() && LoadedLevel && World->IsLevelMakingInvisible(LoadedLevel))
 	{
 		return true;
 	}
@@ -878,7 +878,7 @@ bool ULevelStreaming::CanMakeVisible()
 	{
 		// Once the Level becomes the current pending visibility level, this function must return true or else it could
 		// block indefinitely if called inside UWorld::BlockTillLevelStreamingCompleted (same reason as CanMakeInvisible but with World's CurrentLevelPendingVisibility)
-		if (LoadedLevel == World->GetCurrentLevelPendingVisibility())
+		if (World->IsLevelMakingVisible(LoadedLevel))
 		{
 			return true;
 		}
@@ -931,7 +931,7 @@ void ULevelStreaming::UpdateNetVisibilityTransactionState(bool bInShouldBeVisibl
 			if (!bInShouldBeVisible && ShouldClientUseMakingInvisibleTransactionRequest())
 			{
 				// If this is a client request to make invisible, we will conditionally send a notification to the server before we make the level invisible
-				if (!LoadedLevel || (World->GetCurrentLevelPendingInvisibility() != LoadedLevel))
+				if (!LoadedLevel || !World->IsLevelMakingInvisible(LoadedLevel))
 				{
 					NetVisibilityState.bHasClientPendingRequest = true;
 					NetVisibilityState.PendingRequestType = ENetLevelVisibilityRequest::MakingInvisible;
@@ -940,7 +940,7 @@ void ULevelStreaming::UpdateNetVisibilityTransactionState(bool bInShouldBeVisibl
 			else if (bInShouldBeVisible && ShouldClientUseMakingVisibleTransactionRequest())
 			{
 				// If this is a client request to make visible, we will conditionally send a notification to the server before we make the level visible
-				if (!LoadedLevel || (World->GetCurrentLevelPendingVisibility() != LoadedLevel))
+				if (!LoadedLevel || !World->IsLevelMakingVisible(LoadedLevel))
 				{
 					NetVisibilityState.bHasClientPendingRequest = true;
 					NetVisibilityState.PendingRequestType = ENetLevelVisibilityRequest::MakingVisible;
@@ -1039,7 +1039,7 @@ void ULevelStreaming::UpdateStreamingState(bool& bOutUpdateAgain, bool& bOutRede
 				// Calling CanMakeVisible will trigger a request for visibility if necessary
 				if (!CanMakeVisible())
 				{
-					check(LoadedLevel != World->GetCurrentLevelPendingVisibility());
+					check(!World->IsLevelMakingVisible(LoadedLevel))
 					break;
 				}
 
@@ -1081,7 +1081,7 @@ void ULevelStreaming::UpdateStreamingState(bool& bOutUpdateAgain, bool& bOutRede
 			// We do not want to have any changes in flights when ending play, so before making invisible we wait for server acknowledgment
 			if (!CanMakeInvisible())
 			{
-				check(LoadedLevel != World->GetCurrentLevelPendingInvisibility());
+				check(!World->IsLevelMakingInvisible(LoadedLevel));
 				break;
 			}
 
@@ -1447,8 +1447,7 @@ void ULevelStreaming::PrepareLoadedLevel(ULevel* InLevel, UPackage* InLevelPacka
 
 	if (ensure(LevelOwningWorld))
 	{
-		ULevel* PendingLevelVisOrInvis = (LevelOwningWorld->GetCurrentLevelPendingVisibility() ? LevelOwningWorld->GetCurrentLevelPendingVisibility() : LevelOwningWorld->GetCurrentLevelPendingInvisibility());
-		if (PendingLevelVisOrInvis && PendingLevelVisOrInvis == LoadedLevel)
+		if (LoadedLevel && (LevelOwningWorld->IsLevelMakingVisible(LoadedLevel) || LevelOwningWorld->IsLevelMakingInvisible(LoadedLevel)))
 		{
 			// We can't change current loaded level if it's still processing visibility request
 			// On next UpdateLevelStreaming call this loaded package will be found in memory by RequestLevel function in case visibility request has finished
@@ -1573,8 +1572,7 @@ bool ULevelStreaming::RequestLevel(UWorld* PersistentWorld, bool bAllowLevelLoad
 	}
 
 	// Can not load new level now either, we're still processing visibility for this one
-	ULevel* PendingLevelVisOrInvis = (PersistentWorld->GetCurrentLevelPendingVisibility() ? PersistentWorld->GetCurrentLevelPendingVisibility() : PersistentWorld->GetCurrentLevelPendingInvisibility());
-    if (PendingLevelVisOrInvis && PendingLevelVisOrInvis == LoadedLevel)
+    if (LoadedLevel && (PersistentWorld->IsLevelMakingVisible(LoadedLevel) || PersistentWorld->IsLevelMakingInvisible(LoadedLevel)))
     {
 		UE_LOG(LogLevelStreaming, Verbose, TEXT("Delaying load of new level %s, because %s still processing visibility request."), *DesiredPackageName.ToString(), *CachedLoadedLevelPackageName.ToString());
 		return false;
@@ -2380,7 +2378,7 @@ EStreamingStatus ULevelStreaming::GetLevelStreamingStatus() const
 		{
 			return LEVEL_Visible;
 		}
-		else if ((CurrentState == ELevelStreamingState::MakingVisible) && (GetWorld()->GetCurrentLevelPendingVisibility() == LoadedLevel))
+		else if ((CurrentState == ELevelStreamingState::MakingVisible) && (GetWorld()->IsLevelMakingVisible(LoadedLevel)))
 		{
 			return LEVEL_MakingVisible;
 		}
diff --git a/Engine/Source/Runtime/Engine/Private/LevelTick.cpp b/Engine/Source/Runtime/Engine/Private/LevelTick.cpp
--- a/Engine/Source/Runtime/Engine/Private/LevelTick.cpp
+++ b/Engine/Source/Runtime/Engine/Private/LevelTick.cpp
@@ -1823,7 +1823,7 @@ void UWorld::CleanupActors()
 		// Don't compact actors array for levels that are currently in the process of being made visible as the
 		// code that spreads this work across several frames relies on the actor count not changing as it keeps
 		// an index into the array.
-		if( ensure(Level != nullptr) && (CurrentLevelPendingVisibility != Level) )
+		if( ensure(Level != nullptr) && !IsLevelMakingVisible(Level))
 		{
 			// Actor 0 (world info) and 1 (default brush) are special and should never be removed from the actor array even if NULL
 			const int32 FirstDynamicIndex = 2;
diff --git a/Engine/Source/Runtime/Engine/Private/World.cpp b/Engine/Source/Runtime/Engine/Private/World.cpp
--- a/Engine/Source/Runtime/Engine/Private/World.cpp
+++ b/Engine/Source/Runtime/Engine/Private/World.cpp
@@ -154,6 +154,26 @@ TRACE_DECLARE_INT_COUNTER(NumStreamingLevelsToConsider, TEXT("LevelStreamingProf
 
 #define LOCTEXT_NAMESPACE "World"
 
+static int32 GMaximumMakingVisibleLevels = 1;
+static FAutoConsoleVariableRef CVarMaximumMakingVisibleLevels(
+	TEXT("LevelStreaming.MaximumMakingVisibleLevels"),
+	GMaximumMakingVisibleLevels,
+	TEXT("Sets the maximum number of levels in the process of being made visible (default is 1). \n")
+	TEXT("Can be used in conjunction with incremental component pre-registration & asynchronous physics state creation \n")
+	TEXT("and set to a value greater than 1 to maximize the usage of the time limit allowed by s.LevelStreamingActorsUpdateTimeLimit \n")
+	TEXT("(see p.Chaos.EnableAsyncInitBody and LevelStreaming.AllowIncrementalPreRegisterComponents)."),
+	ECVF_Default);
+
+static int32 GMaximumMakingInvisibleLevels = 1;
+static FAutoConsoleVariableRef CVarMaximumMakingInvisibleLevels(
+	TEXT("LevelStreaming.MaximumMakingInvisibleLevels"),
+	GMaximumMakingInvisibleLevels,
+	TEXT("Sets the maximum number of levels in the process of being made invisible (default is 1). \n")
+	TEXT("Can be used in conjunction with incremental component pre-unregistration & asynchronous physics state destruction \n")
+	TEXT("and set to a value greater than 1 to maximize the usage of the time limit allowed by s.UnregisterComponentsTimeLimit \n")
+	TEXT("(see p.Chaos.EnableAsyncInitBody and LevelStreaming.AllowIncrementalPreUnregisterComponents)."),
+	ECVF_Default);
+
 static int32 bDisableRemapScriptActors = 0;
 FAutoConsoleVariableRef CVarDisableRemapScriptActors(TEXT("net.DisableRemapScriptActors"), bDisableRemapScriptActors, TEXT("When set, disables name remapping of compiled script actors (for networking)"));
 
@@ -3207,6 +3227,41 @@ const AServerStreamingLevelsVisibility* UWorld::GetServerStreamingLevelsVisibili
 	return ServerStreamingLevelsVisibility;
 }
 
+bool UWorld::IsLevelMakingVisible(const ULevel* InLevel) const
+{
+	const bool bIsMakingVisible = InLevel && (InLevel->GetPendingVisibilityState() == ELevelPendingVisibilityState::MakingVisible);
+	check(bIsMakingVisible == MakingVisibleLevels.Contains(InLevel));
+	return bIsMakingVisible;
+}
+
+bool UWorld::IsLevelMakingInvisible(const ULevel* InLevel) const
+{
+	const bool bIsMakingInvisible = InLevel && (InLevel->GetPendingVisibilityState() == ELevelPendingVisibilityState::MakingInvisible);
+	check(bIsMakingInvisible == MakingInvisibleLevels.Contains(InLevel));
+	return bIsMakingInvisible;
+}
+
+void UWorld::SetLevelPendingVisibilityState(ULevel* InLevel, ELevelPendingVisibilityState InState)
+{
+	FLevelPendingVisibilityStateAccessor::SetPendingVisibilityState(InLevel, InState);
+
+	switch (InState)
+	{
+	case ELevelPendingVisibilityState::None:
+		MakingVisibleLevels.Remove(InLevel);
+		MakingInvisibleLevels.Remove(InLevel);
+		break;
+	case ELevelPendingVisibilityState::MakingVisible:
+		check(!MakingInvisibleLevels.Contains(InLevel));
+		MakingVisibleLevels.Add(InLevel);
+		break;
+	case ELevelPendingVisibilityState::MakingInvisible:
+		check(!MakingVisibleLevels.Contains(InLevel));
+		MakingInvisibleLevels.Add(InLevel);
+		break;
+	}
+}
+
 void UWorld::AddToWorld( ULevel* Level, const FTransform& LevelTransform, bool bConsiderTimeLimit, FNetLevelVisibilityTransactionId TransactionId, ULevelStreaming* InOwningLevelStreaming)
 {
 	TRACE_CPUPROFILER_EVENT_SCOPE(AddToWorld);
@@ -3227,14 +3282,15 @@ void UWorld::AddToWorld( ULevel* Level, const FTransform& LevelTransform, bool b
 
 	const double StartTime = FPlatformTime::Seconds();
 
-	bool bExecuteNextStep = CurrentLevelPendingVisibility == Level || CurrentLevelPendingVisibility == NULL;
+	const bool bIsLevelMakingVisible = IsLevelMakingVisible(Level);
+	bool bExecuteNextStep = bIsLevelMakingVisible || (MakingVisibleLevels.Num() < FMath::Max<int32>(1, GMaximumMakingVisibleLevels));
 	bool bPerformedLastStep	= false;
 	const bool bIsGameWorld = IsGameWorld();
 
 	if (bExecuteNextStep)
 	{
 		// Once CurrentLevelPendingVisibility is set, don't call CanAddLoadedLevelToWorld as we must finish processing CurrentLevelPendingVisibility
-		if ((CurrentLevelPendingVisibility == nullptr) && !CanAddLoadedLevelToWorld(Level))
+		if (!bIsLevelMakingVisible && !CanAddLoadedLevelToWorld(Level))
 		{
 			bExecuteNextStep = false;
 		}
@@ -3267,10 +3323,13 @@ void UWorld::AddToWorld( ULevel* Level, const FTransform& LevelTransform, bool b
 
 		// Remove cumulated time since UpdateLevelStreaming
 		TimeLimit = FMath::Max<double>(0.0, TimeLimit - GAddToWorldTimeCumul);
+
+		// Make sure to test time limit (which includes GAddToWorldTimeCumul) before allowing to begin making visible a level
+		bExecuteNextStep = (!bConsiderTimeLimit || !IsTimeLimitExceeded(TEXT("computing time limit"), StartTime, Level, TimeLimit));
 	}
 
 	// Don't make this level visible if it's currently being made invisible
-	if( bExecuteNextStep && CurrentLevelPendingVisibility == NULL && CurrentLevelPendingInvisibility != Level )
+	if( bExecuteNextStep && !bIsLevelMakingVisible && !IsLevelMakingInvisible(Level) )
 	{
 		TRACE_CPUPROFILER_EVENT_SCOPE(AddToWorld_CurrentLevelPendingInvisibility);
 		TRACE_BEGIN_REGION(*WriteToString<256>(TEXT("AddToWorld: "), Level->GetOutermost()->GetName()));
@@ -3282,8 +3341,8 @@ void UWorld::AddToWorld( ULevel* Level, const FTransform& LevelTransform, bool b
 			FLevelStreamingDelegates::OnLevelBeginMakingVisible.Broadcast(this, OwningLevelStreaming, Level);
 		}
 		
-		// Mark level as being the one in process of being made visible.
-		CurrentLevelPendingVisibility = Level;
+		// Mark level as being in process of being made visible.
+		SetLevelPendingVisibilityState(Level, ELevelPendingVisibilityState::MakingVisible);
 		FLevelRegistrationAccessor::CreateAsyncRegisterLevelContext(Level);
 
 		// Add to the UWorld's array of levels, which causes it to be rendered et al.
@@ -3528,7 +3587,7 @@ void UWorld::AddToWorld( ULevel* Level, const FTransform& LevelTransform, bool b
 		ResetLevelFlagsOnLevelAddedToWorld(Level);
 
 		// Finished making level visible - allow other levels to be added to the world.
-		CurrentLevelPendingVisibility = nullptr;
+		SetLevelPendingVisibilityState(Level, ELevelPendingVisibilityState::None);
 		FLevelRegistrationAccessor::DestroyAsyncRegisterLevelContext(Level);
 
 		if (bIsGameWorld && !Level->bClientOnlyVisible)
@@ -3678,7 +3737,7 @@ void UWorld::RemoveFromWorld( ULevel* Level, bool bAllowIncrementalRemoval, FNet
 	const bool bIsCandidateForRemoval = ULevelStreaming::AllowIncrementalRemovalWhilePendingVisibility() ? true : !bAllowIncrementalRemoval;
 	
 	// To be removed from the world a world must be visible and not pending being made visible (this may be redundent, but for safety)
-	if ( ((CurrentLevelPendingVisibility == nullptr) || (bIsCandidateForRemoval && (CurrentLevelPendingVisibility != Level))) && (Level->bIsVisible || Level->bIsBeingRemoved) )
+	if ( (!HasAnyLevelMakingVisible() || (bIsCandidateForRemoval && !IsLevelMakingVisible(Level))) && (Level->bIsVisible || Level->bIsBeingRemoved) )
 	{
 #if PERF_TRACK_DETAILED_ASYNC_STATS
 		// Keep track of timing.
@@ -3689,23 +3748,28 @@ void UWorld::RemoveFromWorld( ULevel* Level, bool bAllowIncrementalRemoval, FNet
 		if ( bAllowIncrementalRemoval && GLevelStreamingUnregisterComponentsTimeLimit > 0.0 )
 		{
 			bFinishRemovingLevel = false;
-			if (CurrentLevelPendingInvisibility == nullptr)
+			if (!IsLevelMakingInvisible(Level) && (MakingInvisibleLevels.Num() < FMath::Max<int32>(1, GMaximumMakingInvisibleLevels)))
 			{
-				TRACE_BEGIN_REGION(*WriteToString<256>(TEXT("RemoveFromWorld: "), Level->GetPackage()->GetName()));
+				// Make sure to test time limit (which includes GRemoveFromWorldUnregisterComponentTimeCumul) before allowing to begin making invisible a level
+				double TimeLimit = FMath::Max<double>(0.0, GLevelStreamingUnregisterComponentsTimeLimit - GRemoveFromWorldUnregisterComponentTimeCumul);
+				if (TimeLimit > 0.0)
+				{
+					TRACE_BEGIN_REGION(*WriteToString<256>(TEXT("RemoveFromWorld: "), Level->GetPackage()->GetName()));
 
-				double StartTime = FPlatformTime::Seconds();
-				// Mark level as being the one in process of being made invisible. 
-				// This will prevent this level from being unloaded or made visible in the meantime
-				CurrentLevelPendingInvisibility = Level;
+					double StartTime = FPlatformTime::Seconds();
+					// Mark level as being in process of being made invisible. 
+					// This will prevent this level from being unloaded or made visible in the meantime
+					SetLevelPendingVisibilityState(Level, ELevelPendingVisibilityState::MakingInvisible);
 
-				BeginRemoval();
+					BeginRemoval();
 
-				// Take into consideration time taken by BeginRemoval
-				double DeltaTime = (FPlatformTime::Seconds() - StartTime) * 1000;
-				GRemoveFromWorldUnregisterComponentTimeCumul += DeltaTime;
+					// Take into consideration time taken by BeginRemoval
+					double DeltaTime = (FPlatformTime::Seconds() - StartTime) * 1000;
+					GRemoveFromWorldUnregisterComponentTimeCumul += DeltaTime;
+				}
 			}
 
-			if (CurrentLevelPendingInvisibility == Level)
+			if (IsLevelMakingInvisible(Level))
 			{
 				double TimeLimit = FMath::Max<double>(0.0, GLevelStreamingUnregisterComponentsTimeLimit - GRemoveFromWorldUnregisterComponentTimeCumul);
 				if (TimeLimit > 0.0)
@@ -3723,7 +3787,7 @@ void UWorld::RemoveFromWorld( ULevel* Level, bool bAllowIncrementalRemoval, FNet
 						if (Level->IncrementalUnregisterComponents(NumComponentsToUnregister, &Context))
 						{
 							// We're done, so the level can be removed
-							CurrentLevelPendingInvisibility = nullptr;
+							SetLevelPendingVisibilityState(Level, ELevelPendingVisibilityState::None);
 							bFinishRemovingLevel = true;
 							break;
 						}
@@ -3738,12 +3802,12 @@ void UWorld::RemoveFromWorld( ULevel* Level, bool bAllowIncrementalRemoval, FNet
 				}
 			}
 		}
-		else if (Level == CurrentLevelPendingInvisibility)
+		else if (IsLevelMakingInvisible(Level))
 		{
 			// Finish current level pending invisibility without time limit
 			check(Level->bIsBeingRemoved);
 			while (!Level->IncrementalUnregisterComponents(MAX_int32)) {}
-			CurrentLevelPendingInvisibility = nullptr;
+			SetLevelPendingVisibilityState(Level, ELevelPendingVisibilityState::None);
 			bFinishRemovingLevel = true;
 		}
 		else
@@ -4390,18 +4454,18 @@ void* UWorld::InternalGetStateStream(uint32 Id)
 
 bool UWorld::CanAddLoadedLevelToWorld(ULevel* Level) const
 {
-	if (CurrentLevelPendingVisibility == nullptr)
+	check(!IsLevelMakingVisible(Level))
+
+	// Allow world partition to decide wether a level should be added or not to the world
+	if (const IWorldPartitionCell* Cell = Level->GetWorldPartitionRuntimeCell())
 	{
-		// Allow world partition to decide wether a level should be added or not to the world
-		if (const IWorldPartitionCell* Cell = Level->GetWorldPartitionRuntimeCell())
+		const UWorld* OuterWorld = Cell->GetOuterWorld();
+		if (UWorldPartition* OuterWorldPartition = OuterWorld ? OuterWorld->GetWorldPartition() : nullptr)
 		{
-			const UWorld* OuterWorld = Cell->GetOuterWorld();
-			if (UWorldPartition* OuterWorldPartition = OuterWorld ? OuterWorld->GetWorldPartition() : nullptr)
-			{
-				return OuterWorldPartition->CanAddCellToWorld(Cell);
-			}
+			return OuterWorldPartition->CanAddCellToWorld(Cell);
 		}
 	}
+
 	return true;
 }
 
@@ -4514,6 +4578,30 @@ void UWorld::UpdateLevelStreaming()
 	// Store current number of pending unload levels, it may change in loop bellow
 	const int32 NumLevelsPendingPurge = FLevelStreamingGCHelper::GetNumLevelsPendingPurge();
 
+	{
+		// Make sure to respect the order of making visible and making invisible queues
+		TRACE_CPUPROFILER_EVENT_SCOPE(UpdateLevelVisibilityPriority);
+		static constexpr int32 MaxPriority = MAX_int32;
+		int32 Priority = MaxPriority;
+		for (ULevel* Level : MakingVisibleLevels)
+		{
+			if (ULevelStreaming* StreamingLevel = ULevelStreaming::FindStreamingLevel(Level))
+			{
+				check(StreamingLevelsToConsider.GetStreamingLevels().Contains(StreamingLevel));
+				StreamingLevel->SetPriority(Priority--);
+			}
+		}
+
+		for (ULevel* Level : MakingInvisibleLevels)
+		{
+			if (ULevelStreaming* StreamingLevel = ULevelStreaming::FindStreamingLevel(Level))
+			{
+				check(StreamingLevelsToConsider.GetStreamingLevels().Contains(StreamingLevel));
+				StreamingLevel->SetPriority(Priority--);
+			}
+		}
+	}
+
 	{
 		UE::Private::World::FStreamingLevelsToConsiderIterationScope Scope(StreamingLevelsToConsider);
 		const int32 NumStreamingLevelsToConsider = StreamingLevelsToConsider.GetStreamingLevels().Num();
@@ -4869,7 +4957,7 @@ void UWorld::ConditionallyBuildStreamingData()
 
 bool UWorld::IsVisibilityRequestPending() const
 {
-	return (CurrentLevelPendingVisibility != nullptr || CurrentLevelPendingInvisibility != nullptr);
+	return HasAnyLevelMakingVisible() || HasAnyLevelMakingInvisible();
 }
 
 bool UWorld::AreAlwaysLoadedLevelsLoaded() const
diff --git a/Engine/Source/Runtime/Engine/Private/WorldPartition/HLOD/HLODRuntimeSubsystem.cpp b/Engine/Source/Runtime/Engine/Private/WorldPartition/HLOD/HLODRuntimeSubsystem.cpp
--- a/Engine/Source/Runtime/Engine/Private/WorldPartition/HLOD/HLODRuntimeSubsystem.cpp
+++ b/Engine/Source/Runtime/Engine/Private/WorldPartition/HLOD/HLODRuntimeSubsystem.cpp
@@ -534,8 +534,7 @@ bool UWorldPartitionHLODRuntimeSubsystem::ShouldPerformWarmupForCell(const UWorl
 
 	// If processing for this cell in Add/RemoveFromWorld() has already started, it's too late for warmup
 	const ULevel* CellLevel = InCell->GetLevel();
-	const bool bCurrentlyProcessingLevel = CellLevel == World->GetCurrentLevelPendingVisibility() ||
-										   CellLevel == World->GetCurrentLevelPendingInvisibility();
+	const bool bCurrentlyProcessingLevel = World->IsLevelMakingVisible(CellLevel) || World->IsLevelMakingInvisible(CellLevel);
 	if (bCurrentlyProcessingLevel)
 	{
 		return false;
diff --git a/Engine/Source/Runtime/Engine/Private/WorldPartition/WorldPartitionLevelStreamingDynamic.cpp b/Engine/Source/Runtime/Engine/Private/WorldPartition/WorldPartitionLevelStreamingDynamic.cpp
--- a/Engine/Source/Runtime/Engine/Private/WorldPartition/WorldPartitionLevelStreamingDynamic.cpp
+++ b/Engine/Source/Runtime/Engine/Private/WorldPartition/WorldPartitionLevelStreamingDynamic.cpp
@@ -232,8 +232,7 @@ bool UWorldPartitionLevelStreamingDynamic::RequestLevel(UWorld* InPersistentWorl
 	const FString WorldAssetPackageName = GetWorldAssetPackageName();
 
 	// Can not load new level now either, we're still processing visibility for this one
-	ULevel* PendingLevelVisOrInvis = (InPersistentWorld->GetCurrentLevelPendingVisibility() ? InPersistentWorld->GetCurrentLevelPendingVisibility() : InPersistentWorld->GetCurrentLevelPendingInvisibility());
-	if (PendingLevelVisOrInvis && PendingLevelVisOrInvis == LoadedLevel)
+	if (LoadedLevel && (InPersistentWorld->IsLevelMakingVisible(LoadedLevel) || InPersistentWorld->IsLevelMakingInvisible(LoadedLevel)))
 	{
 		UE_LOG(LogLevelStreaming, Verbose, TEXT("Delaying load of new level %s, because still processing visibility request."), *WorldAssetPackageName);
 		return false;
diff --git a/Engine/Source/Runtime/Engine/Private/WorldPartition/WorldPartitionSubsystem.cpp b/Engine/Source/Runtime/Engine/Private/WorldPartition/WorldPartitionSubsystem.cpp
--- a/Engine/Source/Runtime/Engine/Private/WorldPartition/WorldPartitionSubsystem.cpp
+++ b/Engine/Source/Runtime/Engine/Private/WorldPartition/WorldPartitionSubsystem.cpp
@@ -1332,17 +1332,27 @@ void UWorldPartitionSubsystem::UpdateStreamingStateInternal(const UWorld* InWorl
 			Algo::Sort(PendingToActivateCells, [](const UWorldPartitionRuntimeCell* CellA, const UWorldPartitionRuntimeCell* CellB) { return CellA->SortCompare(CellB) < 0; });
 		}
 
-		const ULevel* CurrentPendingVisibility = World->GetCurrentLevelPendingVisibility();
-		const UWorldPartitionRuntimeCell* CurrentPendingVisibilityCell = CurrentPendingVisibility ? Cast<const UWorldPartitionRuntimeCell>(CurrentPendingVisibility->GetWorldPartitionRuntimeCell()) : nullptr;
-		auto SetCellsStreamingPriority = [CurrentPendingVisibilityCell](const TArray<const UWorldPartitionRuntimeCell*>& InCells, int32 InPriorityBias = 0)
+		auto SetCellsStreamingPriority = [World](const TArray<const UWorldPartitionRuntimeCell*>& InCells, int32 InPriorityBias = 0)
 		{
 			const int32 MaxPrio = InCells.Num() + InPriorityBias;
-			int32 Prio = MaxPrio;
-			for (const UWorldPartitionRuntimeCell* Cell : InCells)
+			int32 SortedPriority = MaxPrio;
+			if (World->HasAnyLevelMakingVisible())
 			{
-				// Make sure that the current pending visibility level is the most important so that level streaming update will process it right away
-				const int32 SortedPriority = (Cell == CurrentPendingVisibilityCell) ? MaxPrio + 1 : Prio--;
-				Cell->SetStreamingPriority(SortedPriority);
+				// Skip cells referencing current pending visible levels as their priority is already handled by UWorld::UpdateStreamingLevel
+				for (const UWorldPartitionRuntimeCell* Cell : InCells)
+				{
+					if (const ULevel* Level = Cell->GetLevel(); Level && (Level->GetPendingVisibilityState() != ELevelPendingVisibilityState::MakingVisible))
+					{
+						Cell->SetStreamingPriority(SortedPriority--);
+					}
+				}
+			}
+			else
+			{
+				for (const UWorldPartitionRuntimeCell* Cell : InCells)
+				{
+					Cell->SetStreamingPriority(SortedPriority--);
+				}
 			}
 		};
 
