Limit the size of debug diagnostics used when cooking animation sequences.

diff --git a/Engine/Source/Runtime/Engine/Private/Animation/AnimSequence.cpp b/Engine/Source/Runtime/Engine/Private/Animation/AnimSequence.cpp
--- a/Engine/Source/Runtime/Engine/Private/Animation/AnimSequence.cpp
+++ b/Engine/Source/Runtime/Engine/Private/Animation/AnimSequence.cpp
@@ -968,7 +968,8 @@ public:
 		Writer.BeginObject();
 		Writer << "DDCKeyHash" << WriteToString<64>(DDCKeyHash);
 		Writer << "DDCKey" << DDCKey;
-		Writer << "Model" << ModelStateString;
+		//Limiting the diagnostic size to the first 1000 chars
+		Writer << "Model" << FStringView(ModelStateString).Left(1000);
 		Writer << "RawDataGuidEqual" << bRawDataGuidEqual;
 		Writer.EndObject();
 
