Fix crash in FInstancedStaticMeshSceneProxy::SetupRayTracingDynamicInstances()

diff --git a/Engine/Source/Runtime/Engine/Classes/Engine/InstancedStaticMesh.h b/Engine/Source/Runtime/Engine/Classes/Engine/InstancedStaticMesh.h
--- a/Engine/Source/Runtime/Engine/Classes/Engine/InstancedStaticMesh.h
+++ b/Engine/Source/Runtime/Engine/Classes/Engine/InstancedStaticMesh.h
@@ -539,7 +539,7 @@ protected:
 	float InstanceLODDistanceScale;
 
 #if RHI_RAYTRACING
-	ENGINE_API void SetupRayTracingDynamicInstances(int32 NumDynamicInstances, int32 LODIndex);
+	ENGINE_API void SetupRayTracingDynamicInstances(FRHICommandListBase& RHICmdList, int32 NumDynamicInstances, int32 LODIndex);
 
 	struct FRayTracingDynamicData
 	{
diff --git a/Engine/Source/Runtime/Engine/Private/InstancedStaticMesh.cpp b/Engine/Source/Runtime/Engine/Private/InstancedStaticMesh.cpp
--- a/Engine/Source/Runtime/Engine/Private/InstancedStaticMesh.cpp
+++ b/Engine/Source/Runtime/Engine/Private/InstancedStaticMesh.cpp
@@ -1713,7 +1713,7 @@ void FInstancedStaticMeshSceneProxy::GetDynamicRayTracingInstances(FRayTracingIn
 			}
 		}
 
-		SetupRayTracingDynamicInstances(SimulatedInstances, LODIndex);
+		SetupRayTracingDynamicInstances(Collector.GetRHICommandList(), SimulatedInstances, LODIndex);
 
 		ActiveInstances.AddZeroed(SimulatedInstances);
 
@@ -1951,7 +1951,7 @@ void FInstancedStaticMeshSceneProxy::GetDynamicRayTracingInstances(FRayTracingIn
 }
 
 
-void FInstancedStaticMeshSceneProxy::SetupRayTracingDynamicInstances(int32 NumDynamicInstances, int32 LODIndex)
+void FInstancedStaticMeshSceneProxy::SetupRayTracingDynamicInstances(FRHICommandListBase& RHICmdList, int32 NumDynamicInstances, int32 LODIndex)
 {
 	if (NumDynamicInstances == RayTracingDynamicData.Num() && LODIndex == CachedRayTracingLODIndex)
 	{
@@ -1994,7 +1994,7 @@ void FInstancedStaticMeshSceneProxy::SetupRayTracingDynamicInstances(int32 NumDy
 		Initializer.bFastBuild = true;
 
 		DynamicData.DynamicGeometry.SetInitializer(MoveTemp(Initializer));
-		DynamicData.DynamicGeometry.InitResource(FRHICommandListImmediate::Get());
+		DynamicData.DynamicGeometry.InitResource(RHICmdList);
 	}
 }
 
