- Fix for CSV stats and assets with number post fixes

#rnx

diff --git a/Engine/Source/Runtime/Engine/Private/Particles/ParticlePerfStats.cpp b/Engine/Source/Runtime/Engine/Private/Particles/ParticlePerfStats.cpp
--- a/Engine/Source/Runtime/Engine/Private/Particles/ParticlePerfStats.cpp
+++ b/Engine/Source/Runtime/Engine/Private/Particles/ParticlePerfStats.cpp
@@ -375,18 +375,18 @@ void FParticlePerfStats::PopulateStatNames(const FName InName)
 	if (CSVStat_Total == NAME_None)
 	{
 		const FString StatName = InName.ToString();
-		CSVStat_Total = *FString::Printf(TEXT("Total/%s"), *StatName);
-		CSVStat_GTOnly = *FString::Printf(TEXT("GTOnly/%s"), *StatName);
-		CSVStat_InstAvgGT = *FString::Printf(TEXT("InstAvgGT/%s"), *StatName);
-		CSVStat_RT = *FString::Printf(TEXT("RT/%s"), *StatName);
-		CSVStat_InstAvgRT = *FString::Printf(TEXT("InstAvgRT/%s"), *StatName);
-		CSVStat_GPU = *FString::Printf(TEXT("GPU/%s"), *StatName);
-		CSVStat_InstAvgGPU = *FString::Printf(TEXT("InstAvgGPU/%s"), *StatName);
-		CSVStat_Count = *FString::Printf(TEXT("Count/%s"), *StatName);
-		CSVStat_Activation = *FString::Printf(TEXT("Activation/%s"), *StatName);
-		CSVStat_Waits = *FString::Printf(TEXT("Waits/%s"), *StatName);
-		CSVStat_Culled = *FString::Printf(TEXT("Culled/%s"), *StatName);
-		CSVStat_MemoryKB = *FString::Printf(TEXT("MemoryKB/%s"), *StatName);
+		CSVStat_Total = FName(*FString::Printf(TEXT("Total/%s"), *StatName), 0, false);
+		CSVStat_GTOnly = FName(*FString::Printf(TEXT("GTOnly/%s"), *StatName), 0, false);
+		CSVStat_InstAvgGT = FName(*FString::Printf(TEXT("InstAvgGT/%s"), *StatName), 0, false);
+		CSVStat_RT = FName(*FString::Printf(TEXT("RT/%s"), *StatName), 0, false);
+		CSVStat_InstAvgRT = FName(*FString::Printf(TEXT("InstAvgRT/%s"), *StatName), 0, false);
+		CSVStat_GPU = FName(*FString::Printf(TEXT("GPU/%s"), *StatName), 0, false);
+		CSVStat_InstAvgGPU = FName(*FString::Printf(TEXT("InstAvgGPU/%s"), *StatName), 0, false);
+		CSVStat_Count = FName(*FString::Printf(TEXT("Count/%s"), *StatName), 0, false);
+		CSVStat_Activation = FName(*FString::Printf(TEXT("Activation/%s"), *StatName), 0, false);
+		CSVStat_Waits = FName(*FString::Printf(TEXT("Waits/%s"), *StatName), 0, false);
+		CSVStat_Culled = FName(*FString::Printf(TEXT("Culled/%s"), *StatName), 0, false);
+		CSVStat_MemoryKB = FName(*FString::Printf(TEXT("MemoryKB/%s"), *StatName), 0, false);
 	}
 }
 
