Renamed UAnimBankMeshComponent and related PCG types to UInstancedSkinnedMeshComponent

diff --git a/Engine/Plugins/PCG/Shaders/Private/Elements/PCGAnimBankMeshSpawner.usf b/Engine/Plugins/PCG/Shaders/Private/Elements/PCGAnimBankMeshSpawner.usf
deleted file mode 100644
diff --git a/Engine/Plugins/PCG/Shaders/Private/Elements/PCGSkinnedMeshSpawner.usf b/Engine/Plugins/PCG/Shaders/Private/Elements/PCGSkinnedMeshSpawner.usf
new file mode 100644
--- /dev/null
+++ b/Engine/Plugins/PCG/Shaders/Private/Elements/PCGSkinnedMeshSpawner.usf
@@ -0,0 +1,188 @@
+// Copyright Epic Games, Inc. All Rights Reserved.
+
+// Skinned mesh spawner kernel. Selects a primitive either using a primitive cumulative distribution function (CDF) or by an attribute, and packs attributes in custom floats.
+// Passes through points to output.
+
+// When spawning by attribute, we don't know in advance how many instances of each mesh will be spawned, so the worst case NumPoints * NumPrims is allocated in the GPU Scene.
+
+#include "/Engine/Private/ComputeShaderUtils.ush"
+#include "/Engine/Private/Quaternion.ush"
+#include "/Plugin/PCG/Private/PCGShaderUtils.ush"
+
+// TODO: should be code-gen'd from PCGComputeCommon.h
+#define MAX_NUM_ATTRS 128
+#define NUM_RESERVED_ATTRS 32
+
+void CopyAttributesToOutput(uint In_DataIndex, uint In_ElementIndex, uint Out_DataIndex, uint Out_ElementIndex);
+
+[numthreads(64, 1, 1)]
+void Main(uint3 GroupId : SV_GroupID, uint GroupIndex : SV_GroupIndex)
+{
+	// Mark the kernel as having executed. Must run before we early out via thread index, because the kernel is still 'executed' even if the number of
+	// threads to iterate on is zero. Even if GetNumThreads() returns 0, the kernel will still have been dispatched on a single thread to set this flag.
+	if (all(GroupId == 0) && GroupIndex == 0)
+	{
+		Out_SetAsExecutedInternal();
+	}
+
+	const uint ThreadIndex = GetUnWrappedDispatchThreadId(GroupId, GroupIndex, 64);
+
+	uint In_DataIndex, In_ElementIndex;
+	if (!In_GetThreadData(ThreadIndex, In_DataIndex, In_ElementIndex))
+	{
+		return;
+	}
+
+	uint Out_DataIndex, Out_ElementIndex;
+	if (!Out_GetThreadData(ThreadIndex, Out_DataIndex, Out_ElementIndex))
+	{
+		// Should never happen.
+		return;
+	}
+
+	if (In_IsPointRemoved(In_DataIndex, In_ElementIndex))
+	{
+		// Propagate the 'Removed' status to the output point if input point is removed. Otherwise, this point will not be culled.
+		Out_RemovePoint(Out_DataIndex, Out_ElementIndex);
+
+		return;
+	}
+
+	uint PrimitiveIndex = (uint)-1;
+
+	const int NumPrimitives = InstanceData_GetNumPrimitives();
+
+	if (NumPrimitives > 0)
+	{
+		const uint SelectorAttributeId = ABSpawner_GetSelectorAttributeId();
+		
+		if (SelectorAttributeId != (uint)-1)
+		{
+			// Select using string key
+			const int MeshStringKey = In_GetStringKey(In_DataIndex, In_ElementIndex, SelectorAttributeId);
+			PrimitiveIndex = ABSpawner_GetPrimitiveIndexFromStringKey(MeshStringKey);
+		}
+		else
+		{
+			// Randomly select a primitive by inverting the cumulative distribution function.
+			
+			// Exactly matches UPCGBlueprintHelpers::GetRandomStreamFromPoint - three seeds passed in to consecutive arguments of ComputeSeed
+			// (argument count & order matters).
+			const uint SeedPoint = In_GetSeed(In_DataIndex, In_ElementIndex);
+			const uint SeedSettings = GetSettingsSeed();
+			const uint SeedComponent = GetComponentSeed();
+			uint Seed = ComputeSeed(SeedPoint, SeedSettings, SeedComponent);
+
+			const float RandomDraw = FRand(Seed);
+
+			for (PrimitiveIndex = 0; PrimitiveIndex < uint(NumPrimitives) - 1; ++PrimitiveIndex)
+			{
+				if (RandomDraw < ABSpawner_GetPrimitiveSelectionCDF(PrimitiveIndex))
+				{
+					break;
+				}
+			}
+		}
+	}
+
+	// Get instance index to write to. Will return a negative value if we have saturated instances for this primitive, in which case we will not push the data.
+	const int InstanceIndexInt = (PrimitiveIndex != (uint)-1) ? InstanceData_GetIndexToWriteTo(PrimitiveIndex) : -1;
+
+	// Write instance transform, only if not culled by primitive instance index
+	if (InstanceIndexInt >= 0)
+	{
+		const uint InstanceIndex = InstanceIndexInt;
+		const float4x4 Transform = In_GetPointTransform(In_DataIndex, In_ElementIndex);
+		InstanceData_WriteInstanceLocalToWorld(InstanceIndex, transpose(Transform));
+
+		// Write instance custom float data.
+		const uint NumAttributes = ABSpawner_GetNumAttributes();
+		for (uint AttributeIndex = 0; AttributeIndex < NumAttributes; ++AttributeIndex)
+		{
+			const uint4 AttributeIdOffsetStride = ABSpawner_GetAttributeIdOffsetStride(AttributeIndex);
+			const uint AttributeId = AttributeIdOffsetStride[0];
+			const uint AttributeOffset = AttributeIdOffsetStride[1];
+			const uint AttributeStride = AttributeIdOffsetStride[2];
+
+			for (uint FloatIndex = 0; FloatIndex < AttributeStride; ++FloatIndex)
+			{
+				const uint AttributeValueAddress = In_GetElementAddressInternal(In_DataIndex, In_ElementIndex, AttributeId);
+				// Assumes floating point data. TODO add type information and deal with different types in different ways?
+				InstanceData_WriteCustomFloatRaw(InstanceIndex, AttributeOffset + FloatIndex, In_LoadBufferInternal(AttributeValueAddress + 4 * FloatIndex));
+			}
+		}
+	}
+
+	// Write output point data (pass through).
+	Out_SetPosition(Out_DataIndex, Out_ElementIndex, In_GetPosition(In_DataIndex, In_ElementIndex));
+	Out_SetRotation(Out_DataIndex, Out_ElementIndex, In_GetRotation(In_DataIndex, In_ElementIndex));
+	Out_SetScale(Out_DataIndex, Out_ElementIndex, In_GetScale(In_DataIndex, In_ElementIndex));
+	Out_SetColor(Out_DataIndex, Out_ElementIndex, In_GetColor(In_DataIndex, In_ElementIndex));
+	Out_SetSeed(Out_DataIndex, Out_ElementIndex, In_GetSeed(In_DataIndex, In_ElementIndex));
+	Out_SetDensity(Out_DataIndex, Out_ElementIndex, In_GetDensity(In_DataIndex, In_ElementIndex));
+	Out_SetSteepness(Out_DataIndex, Out_ElementIndex, In_GetSteepness(In_DataIndex, In_ElementIndex));
+
+	// Either apply AB bounds to points or pass bounds through.
+	if (PrimitiveIndex != (uint)-1 && ABSpawner_ShouldApplyBounds())
+	{
+		Out_SetBoundsMin(Out_DataIndex, Out_ElementIndex, ABSpawner_GetPrimitiveMeshBoundsMin(PrimitiveIndex));
+		Out_SetBoundsMax(Out_DataIndex, Out_ElementIndex, ABSpawner_GetPrimitiveMeshBoundsMax(PrimitiveIndex));
+	}
+	else
+	{
+		Out_SetBoundsMin(Out_DataIndex, Out_ElementIndex, In_GetBoundsMin(In_DataIndex, In_ElementIndex));
+		Out_SetBoundsMax(Out_DataIndex, Out_ElementIndex, In_GetBoundsMax(In_DataIndex, In_ElementIndex));
+	}
+
+	CopyAttributesToOutput(In_DataIndex, In_ElementIndex, Out_DataIndex, Out_ElementIndex);
+
+	if (ABSpawner_GetSelectedMeshAttributeId() != (uint)-1)
+	{
+		const int OutputMeshStringKey = (PrimitiveIndex != (uint)-1) ? ABSpawner_GetPrimitiveStringKey(PrimitiveIndex) : -1;
+		Out_SetStringKey(Out_DataIndex, Out_ElementIndex, ABSpawner_GetSelectedMeshAttributeId(), OutputMeshStringKey);
+	}
+}
+
+// TODO pass in attribute IDs to copy.
+void CopyAttributesToOutput(uint In_DataIndex, uint In_ElementIndex, uint Out_DataIndex, uint Out_ElementIndex)
+{
+	const uint NumAttributes = Out_GetDataNumAttributesInternal(Out_DataIndex);
+	uint NumAttributesProcessed = 0;
+	
+	for (int AttributeId = NUM_RESERVED_ATTRS; AttributeId < MAX_NUM_ATTRS; ++AttributeId)
+	{
+		const uint Out_Stride = Out_GetAttributeStrideInternal(Out_DataIndex, AttributeId);
+		
+		if (Out_Stride == 0)
+		{
+			// No output attribute to write to.
+			continue;
+		}
+		
+		const uint In_Stride = In_GetAttributeStrideInternal(In_DataIndex, AttributeId);
+		const uint Out_ElementAddress = Out_GetElementAddressInternal(Out_DataIndex, Out_ElementIndex, AttributeId);
+		
+		if (In_Stride == Out_Stride)
+		{
+			const uint In_ElementAddress = In_GetElementAddressInternal(In_DataIndex, In_ElementIndex, AttributeId);
+			
+			for (uint Offset = 0; Offset < Out_Stride; Offset += 4)
+			{
+				Out_StoreBufferInternal(Out_ElementAddress + Offset, In_LoadBufferInternal(In_ElementAddress + Offset));
+			}
+		}
+		else
+		{
+			for (uint Offset = 0; Offset < Out_Stride; Offset += 4)
+			{
+				// Initialize output data (only header part of buffer will be initialized).
+				Out_StoreBufferInternal(Out_ElementAddress + Offset, 0u);
+			}
+		}
+		
+		if (++NumAttributesProcessed >= NumAttributes)
+		{
+			break; // We can early-out when we've looked at all the possible attributes
+		}
+	}
+}
\ No newline at end of file
diff --git a/Engine/Plugins/PCG/Source/PCG/Private/Compute/DataInterfaces/Elements/PCGAnimBankMeshSpawnerDataInterface.cpp b/Engine/Plugins/PCG/Source/PCG/Private/Compute/DataInterfaces/Elements/PCGAnimBankMeshSpawnerDataInterface.cpp
deleted file mode 100644
diff --git a/Engine/Plugins/PCG/Source/PCG/Private/Compute/DataInterfaces/Elements/PCGAnimBankMeshSpawnerDataInterface.h b/Engine/Plugins/PCG/Source/PCG/Private/Compute/DataInterfaces/Elements/PCGAnimBankMeshSpawnerDataInterface.h
deleted file mode 100644
diff --git a/Engine/Plugins/PCG/Source/PCG/Private/Compute/DataInterfaces/Elements/PCGSkinnedMeshSpawnerDataInterface.cpp b/Engine/Plugins/PCG/Source/PCG/Private/Compute/DataInterfaces/Elements/PCGSkinnedMeshSpawnerDataInterface.cpp
new file mode 100644
--- /dev/null
+++ b/Engine/Plugins/PCG/Source/PCG/Private/Compute/DataInterfaces/Elements/PCGSkinnedMeshSpawnerDataInterface.cpp
@@ -0,0 +1,250 @@
+// Copyright Epic Games, Inc. All Rights Reserved.
+
+#include "PCGSkinnedMeshSpawnerDataInterface.h"
+
+#include "Compute/PCGComputeCommon.h"
+#include "Compute/PCGDataBinding.h"
+#include "Elements/PCGSkinnedMeshSpawner.h"
+
+#include "GlobalRenderResources.h"
+#include "RHIResources.h"
+#include "ShaderCompilerCore.h"
+#include "ShaderCore.h"
+#include "ShaderParameterMetadataBuilder.h"
+#include "ComputeFramework/ShaderParamTypeDefinition.h"
+
+void UPCGSkinnedMeshSpawnerDataInterface::GetSupportedInputs(TArray<FShaderFunctionDefinition>& OutFunctions) const
+{
+	OutFunctions.AddDefaulted_GetRef()
+		.SetName(TEXT("ABSpawner_GetSelectorAttributeId"))
+		.AddReturnType(EShaderFundamentalType::Uint); // Attribute id to get mesh path string key from, or invalid if we should use CDF instead.
+
+	OutFunctions.AddDefaulted_GetRef()
+		.SetName(TEXT("ABSpawner_GetNumAttributes"))
+		.AddReturnType(EShaderFundamentalType::Uint); // Num attributes
+
+	OutFunctions.AddDefaulted_GetRef()
+		.SetName(TEXT("ABSpawner_GetNumPrimitives"))
+		.AddReturnType(EShaderFundamentalType::Uint); // Num primitives
+
+	OutFunctions.AddDefaulted_GetRef()
+		.SetName(TEXT("ABSpawner_ShouldApplyBounds"))
+		.AddReturnType(EShaderFundamentalType::Bool);
+
+	OutFunctions.AddDefaulted_GetRef()
+		.SetName(TEXT("ABSpawner_GetPrimitiveMeshBoundsMin"))
+		.AddReturnType(EShaderFundamentalType::Float, 3) // Local bounds min
+		.AddParam(EShaderFundamentalType::Uint); // InPrimitiveIndex
+
+	OutFunctions.AddDefaulted_GetRef()
+		.SetName(TEXT("ABSpawner_GetPrimitiveMeshBoundsMax"))
+		.AddReturnType(EShaderFundamentalType::Float, 3) // Local bounds max
+		.AddParam(EShaderFundamentalType::Uint); // InPrimitiveIndex
+
+	OutFunctions.AddDefaulted_GetRef()
+		.SetName(TEXT("ABSpawner_GetAttributeIdOffsetStride"))
+		.AddReturnType(EShaderFundamentalType::Uint, 4)
+		.AddParam(EShaderFundamentalType::Uint); // InAttributeIndex
+
+	OutFunctions.AddDefaulted_GetRef()
+		.SetName(TEXT("ABSpawner_GetPrimitiveStringKey"))
+		.AddReturnType(EShaderFundamentalType::Int) // String key
+		.AddParam(EShaderFundamentalType::Uint); // InPrimitiveIndex
+
+	OutFunctions.AddDefaulted_GetRef()
+		.SetName(TEXT("ABSpawner_GetPrimitiveIndexFromStringKey"))
+		.AddReturnType(EShaderFundamentalType::Uint) // Primitive index
+		.AddParam(EShaderFundamentalType::Int); // InMeshPathStringKey
+
+	OutFunctions.AddDefaulted_GetRef()
+		.SetName(TEXT("ABSpawner_GetPrimitiveSelectionCDF"))
+		.AddReturnType(EShaderFundamentalType::Float) // CDF value
+		.AddParam(EShaderFundamentalType::Uint); // InPrimitiveIndex
+
+	OutFunctions.AddDefaulted_GetRef()
+		.SetName(TEXT("ABSpawner_GetSelectedMeshAttributeId"))
+		.AddReturnType(EShaderFundamentalType::Uint); // Attribute id to output mesh path string key to
+}
+
+BEGIN_SHADER_PARAMETER_STRUCT(FPCGSkinnedMeshSpawnerDataInterfaceParameters,)
+SHADER_PARAMETER_ARRAY(FUintVector4, AttributeIdOffsetStrides, [UPCGSkinnedMeshSpawnerDataInterface::MAX_ATTRIBUTES])
+	SHADER_PARAMETER_ARRAY(FIntVector4, PrimitiveStringKeys, [PCGComputeConstants::MAX_PRIMITIVE_COMPONENTS_PER_SPAWNER])
+	SHADER_PARAMETER_ARRAY(FVector4f, PrimitiveMeshBoundsMin, [PCGComputeConstants::MAX_PRIMITIVE_COMPONENTS_PER_SPAWNER])
+	SHADER_PARAMETER_ARRAY(FVector4f, PrimitiveMeshBoundsMax, [PCGComputeConstants::MAX_PRIMITIVE_COMPONENTS_PER_SPAWNER])
+	SHADER_PARAMETER_SCALAR_ARRAY(float, SelectionCDF, [PCGComputeConstants::MAX_PRIMITIVE_COMPONENTS_PER_SPAWNER])
+	SHADER_PARAMETER(uint32, NumAttributes)
+	SHADER_PARAMETER(uint32, NumPrimitives)
+	SHADER_PARAMETER(int32, SelectorAttributeId)
+	SHADER_PARAMETER(int32, SelectedMeshAttributeId)
+	SHADER_PARAMETER(uint32, ApplyBounds)
+END_SHADER_PARAMETER_STRUCT()
+
+void UPCGSkinnedMeshSpawnerDataInterface::GetShaderParameters(TCHAR const* UID, FShaderParametersMetadataBuilder& InOutBuilder, FShaderParametersMetadataAllocations& InOutAllocations) const
+{
+	InOutBuilder.AddNestedStruct<FPCGSkinnedMeshSpawnerDataInterfaceParameters>(UID);
+}
+
+void UPCGSkinnedMeshSpawnerDataInterface::GetHLSL(FString& OutHLSL, FString const& InDataInterfaceName) const
+{
+	TMap<FString, FStringFormatArg> TemplateArgs =
+	{
+		{ TEXT("DataInterfaceName"), InDataInterfaceName },
+		{ TEXT("MaxAttributes"), MAX_ATTRIBUTES },
+		{ TEXT("MaxPrimitives"), PCGComputeConstants::MAX_PRIMITIVE_COMPONENTS_PER_SPAWNER },
+	};
+
+	OutHLSL += FString::Format(TEXT(
+		"int {DataInterfaceName}_SelectorAttributeId;\n"
+		"uint {DataInterfaceName}_NumAttributes;\n"
+		"uint {DataInterfaceName}_NumPrimitives;\n"
+		"uint {DataInterfaceName}_ApplyBounds;"
+		"int {DataInterfaceName}_SelectedMeshAttributeId;\n"
+		"uint4 {DataInterfaceName}_AttributeIdOffsetStrides[{MaxAttributes}];\n"
+		"float4 {DataInterfaceName}_PrimitiveMeshBoundsMin[{MaxPrimitives}];"
+		"float4 {DataInterfaceName}_PrimitiveMeshBoundsMax[{MaxPrimitives}];"
+		"int4 {DataInterfaceName}_PrimitiveStringKeys[{MaxPrimitives}];"
+		"DECLARE_SCALAR_ARRAY(float, {DataInterfaceName}_SelectionCDF, {MaxPrimitives});\n"
+		"\n"
+		"int ABSpawner_GetSelectorAttributeId_{DataInterfaceName}()\n"
+		"{\n"
+		"	return {DataInterfaceName}_SelectorAttributeId;\n"
+		"}\n"
+		"\n"
+		"uint ABSpawner_GetNumAttributes_{DataInterfaceName}()\n"
+		"{\n"
+		"	return {DataInterfaceName}_NumAttributes;\n"
+		"}\n"
+		"\n"
+		"uint4 ABSpawner_GetAttributeIdOffsetStride_{DataInterfaceName}(uint InAttributeIndex)\n"
+		"{\n"
+		"	return {DataInterfaceName}_AttributeIdOffsetStrides[InAttributeIndex];\n"
+		"}\n"
+		"\n"
+		"uint ABSpawner_GetNumPrimitives_{DataInterfaceName}()\n"
+		"{\n"
+		"	return {DataInterfaceName}_NumPrimitives;\n"
+		"}\n"
+		"\n"
+		"bool ABSpawner_ShouldApplyBounds_{DataInterfaceName}()\n"
+		"{\n"
+		"	return {DataInterfaceName}_ApplyBounds > 0;\n"
+		"}\n"
+		"\n"
+		"float3 ABSpawner_GetPrimitiveMeshBoundsMin_{DataInterfaceName}(uint InPrimitiveIndex)\n"
+		"{\n"
+		"	return {DataInterfaceName}_PrimitiveMeshBoundsMin[InPrimitiveIndex].xyz;\n"
+		"}\n"
+		"\n"
+		"float3 ABSpawner_GetPrimitiveMeshBoundsMax_{DataInterfaceName}(uint InPrimitiveIndex)\n"
+		"{\n"
+		"	return {DataInterfaceName}_PrimitiveMeshBoundsMax[InPrimitiveIndex].xyz;\n"
+		"}\n"
+		"\n"
+		"int ABSpawner_GetPrimitiveStringKey_{DataInterfaceName}(uint InPrimitiveIndex)\n"
+		"{\n"
+		"	return {DataInterfaceName}_PrimitiveStringKeys[InPrimitiveIndex][0];\n"
+		"}\n"
+		"\n"
+		"uint ABSpawner_GetPrimitiveIndexFromStringKey_{DataInterfaceName}(int InMeshPathStringKey)\n"
+		"{\n"
+		"	for (uint Index = 0; Index < {DataInterfaceName}_NumPrimitives; ++Index)\n"
+		"	{\n"
+		"		if ({DataInterfaceName}_PrimitiveStringKeys[Index][0] == InMeshPathStringKey)\n"
+		"		{\n"
+		"			return Index;\n"
+		"		}\n"
+		"	}\n"
+		"	\n"
+		"	return (uint)-1;\n"
+		"}\n"
+		"\n"
+		"float ABSpawner_GetPrimitiveSelectionCDF_{DataInterfaceName}(uint InPrimitiveIndex)\n"
+		"{\n"
+		"	return GET_SCALAR_ARRAY_ELEMENT({DataInterfaceName}_SelectionCDF, InPrimitiveIndex);\n"
+		"}\n"
+		"\n"
+		"int ABSpawner_GetSelectedMeshAttributeId_{DataInterfaceName}()\n"
+		"{\n"
+		"	return {DataInterfaceName}_SelectedMeshAttributeId;\n"
+		"}\n"
+	), TemplateArgs);
+}
+
+UComputeDataProvider* UPCGSkinnedMeshSpawnerDataInterface::CreateDataProvider(TObjectPtr<UObject> InBinding, uint64 InInputMask, uint64 InOutputMask) const
+{
+	TRACE_CPUPROFILER_EVENT_SCOPE(UPCGSkinnedMeshSpawnerDataInterface::CreateDataProvider);
+	UPCGDataBinding* Binding = CastChecked<UPCGDataBinding>(InBinding);
+
+	TObjectPtr<UPCGSkinnedMeshSpawnerDataProvider> DataProvider = NewObject<UPCGSkinnedMeshSpawnerDataProvider>();
+	DataProvider->Settings = CastChecked<UPCGSkinnedMeshSpawnerSettings>(Settings);
+	
+	// If there were 0 input points for this execution, we will not have created any primitives, so check for null.
+	if (const FPCGSpawnerPrimitives* Primitives = Binding->MeshSpawnersToPrimitives.Find(Settings))
+	{
+		DataProvider->AttributeIdOffsetStrides = Primitives->AttributeIdOffsetStrides;
+		DataProvider->SelectionCDF = Primitives->SelectionCDF;
+		DataProvider->PrimitiveStringKeys = Primitives->PrimitiveStringKeys;
+		DataProvider->PrimitiveMeshBounds = Primitives->PrimitiveMeshBounds;
+
+		// We only know name at compile time (not type), so we resolved this Id at execution time by inspecting the data arriving on the input pin. 
+		DataProvider->SelectorAttributeId = Primitives->SelectorAttributeId;
+		// We know the name and type of the selected mesh attribute statically and declared the attribute in GetKernelAttributeKeys,
+		// so the attribute ID should be present in the attribute table.
+		// TODO: DataProvider->SelectedMeshAttributeId = Binding->GetAttributeId(DataProvider->Settings->OutMeshAttribute, EPCGKernelAttributeType::StringKey);
+		ensure(DataProvider->SelectedMeshAttributeId != INDEX_NONE);
+	}
+
+	return DataProvider;
+}
+
+FComputeDataProviderRenderProxy* UPCGSkinnedMeshSpawnerDataProvider::GetRenderProxy()
+{
+	return new FPCGSkinnedMeshSpawnerDataProviderProxy(AttributeIdOffsetStrides, SelectorAttributeId, PrimitiveStringKeys, SelectionCDF, SelectedMeshAttributeId, PrimitiveMeshBounds);
+}
+
+bool FPCGSkinnedMeshSpawnerDataProviderProxy::IsValid(FValidationData const& InValidationData) const
+{
+	if (InValidationData.ParameterStructSize != sizeof(FParameters))
+	{
+		return false;
+	}
+
+	return true;
+}
+
+void FPCGSkinnedMeshSpawnerDataProviderProxy::GatherDispatchData(FDispatchData const& InDispatchData)
+{
+	const TStridedView<FParameters> ParameterArray = MakeStridedParameterView<FParameters>(InDispatchData);
+	for (int32 InvocationIndex = 0; InvocationIndex < ParameterArray.Num(); ++InvocationIndex)
+	{
+		FParameters& Parameters = ParameterArray[InvocationIndex];
+
+		Parameters.NumAttributes = AttributeIdOffsetStrides.Num();
+		Parameters.NumPrimitives = SelectionCDF.Num();
+		Parameters.SelectorAttributeId = SelectorAttributeId;
+		Parameters.SelectedMeshAttributeId = SelectedMeshAttributeId;
+
+		for (int32 Index = 0; Index < AttributeIdOffsetStrides.Num(); ++Index)
+		{
+			Parameters.AttributeIdOffsetStrides[Index] = AttributeIdOffsetStrides[Index];
+		}
+
+		for (int32 Index = 0; Index < PrimitiveStringKeys.Num(); ++Index)
+		{
+			Parameters.PrimitiveStringKeys[Index] = FIntVector4(PrimitiveStringKeys[Index], /*Unused*/0, /*Unused*/0, /*Unused*/0);
+		}
+
+		for (int32 Index = 0; Index < SelectionCDF.Num(); ++Index)
+		{
+			GET_SCALAR_ARRAY_ELEMENT(Parameters.SelectionCDF, Index) = SelectionCDF[Index];
+		}
+
+		Parameters.ApplyBounds = PrimitiveMeshBounds.IsEmpty() ? 0 : 1;
+
+		for (int32 Index = 0; Index < PrimitiveMeshBounds.Num(); ++Index)
+		{
+			Parameters.PrimitiveMeshBoundsMin[Index] = FVector4f(PrimitiveMeshBounds[Index].Min.X, PrimitiveMeshBounds[Index].Min.Y, PrimitiveMeshBounds[Index].Min.Z, /*Unused*/0.0f);
+			Parameters.PrimitiveMeshBoundsMax[Index] = FVector4f(PrimitiveMeshBounds[Index].Max.X, PrimitiveMeshBounds[Index].Max.Y, PrimitiveMeshBounds[Index].Max.Z, /*Unused*/0.0f);
+		}
+	}
+}
diff --git a/Engine/Plugins/PCG/Source/PCG/Private/Compute/DataInterfaces/Elements/PCGSkinnedMeshSpawnerDataInterface.h b/Engine/Plugins/PCG/Source/PCG/Private/Compute/DataInterfaces/Elements/PCGSkinnedMeshSpawnerDataInterface.h
new file mode 100644
--- /dev/null
+++ b/Engine/Plugins/PCG/Source/PCG/Private/Compute/DataInterfaces/Elements/PCGSkinnedMeshSpawnerDataInterface.h
@@ -0,0 +1,105 @@
+// Copyright Epic Games, Inc. All Rights Reserved.
+
+#pragma once
+
+#include "Compute/DataInterfaces/PCGComputeDataInterface.h"
+#include "ComputeFramework/ComputeDataProvider.h"
+
+#include "PCGSkinnedMeshSpawnerDataInterface.generated.h"
+
+class FPCGSkinnedMeshSpawnerDataInterfaceParameters;
+class UPCGSkinnedMeshSpawnerSettings;
+class UPCGSettings;
+
+/** Data Interface to marshal Skinned Mesh Spawner settings to the GPU. */
+UCLASS(ClassGroup = (Procedural))
+class UPCGSkinnedMeshSpawnerDataInterface : public UPCGComputeDataInterface
+{
+	GENERATED_BODY()
+
+public:
+	//~ Begin UComputeDataInterface Interface
+	TCHAR const* GetClassName() const override { return TEXT("PCGSkinnedMeshSpawner"); }
+	void GetSupportedInputs(TArray<FShaderFunctionDefinition>& OutFunctions) const override;
+	void GetShaderParameters(TCHAR const* UID, FShaderParametersMetadataBuilder& InOutBuilder, FShaderParametersMetadataAllocations& InOutAllocations) const override;
+	void GetHLSL(FString& OutHLSL, FString const& InDataInterfaceName) const override;
+	UComputeDataProvider* CreateDataProvider(TObjectPtr<UObject> InBinding, uint64 InInputMask, uint64 InOutputMask) const override;
+	//~ End UComputeDataInterface Interface
+
+	UPROPERTY()
+	TObjectPtr<const UPCGSettings> Settings = nullptr;
+
+	static constexpr uint32 MAX_ATTRIBUTES = 64;
+};
+
+UCLASS()
+class UPCGSkinnedMeshSpawnerDataProvider : public UComputeDataProvider
+{
+	GENERATED_BODY()
+
+public:
+	//~ Begin UComputeDataProvider Interface
+	FComputeDataProviderRenderProxy* GetRenderProxy() override;
+	//~ End UComputeDataProvider Interface
+
+	UPROPERTY()
+	TObjectPtr<const UPCGSkinnedMeshSpawnerSettings> Settings = nullptr;
+
+	UPROPERTY()
+	TArray<FUintVector4> AttributeIdOffsetStrides;
+
+	UPROPERTY()
+	TArray<int32> PrimitiveStringKeys;
+
+	UPROPERTY()
+	TArray<FBox> PrimitiveMeshBounds;
+
+	UPROPERTY()
+	TArray<float> SelectionCDF;
+
+	UPROPERTY()
+	int32 SelectorAttributeId = INDEX_NONE;
+
+	UPROPERTY()
+	int32 NumInputPoints = 0;
+
+	UPROPERTY()
+	int32 SelectedMeshAttributeId = INDEX_NONE;
+};
+
+class FPCGSkinnedMeshSpawnerDataProviderProxy : public FComputeDataProviderRenderProxy
+{
+public:
+	FPCGSkinnedMeshSpawnerDataProviderProxy(
+		const TArray<FUintVector4>& InAttributeIdOffsetStrides,
+		int32 InSelectorAttributeId,
+		const TArray<int32>& InPrimitiveStringKeys,
+		TArray<float> InSelectionCDF,
+		int32 InSelectedMeshAttributeId,
+		const TArray<FBox>& InPrimitiveMeshBounds)
+		: AttributeIdOffsetStrides(InAttributeIdOffsetStrides)
+		, SelectionCDF(InSelectionCDF)
+		, SelectorAttributeId(InSelectorAttributeId)
+		, PrimitiveStringKeys(InPrimitiveStringKeys)
+		, PrimitiveMeshBounds(InPrimitiveMeshBounds)
+		, SelectedMeshAttributeId(InSelectedMeshAttributeId)
+	{}
+
+	//~ Begin FComputeDataProviderRenderProxy Interface
+	bool IsValid(FValidationData const& InValidationData) const override;
+	void GatherDispatchData(FDispatchData const& InDispatchData) override;
+	//~ End FComputeDataProviderRenderProxy Interface
+
+protected:
+	using FParameters = FPCGSkinnedMeshSpawnerDataInterfaceParameters;
+
+	TArray<FUintVector4> AttributeIdOffsetStrides;
+	TArray<float> SelectionCDF;
+
+	int32 SelectorAttributeId = INDEX_NONE;
+	TArray<int32> PrimitiveStringKeys;
+
+	TArray<FBox> PrimitiveMeshBounds;
+
+	int32 SelectedMeshAttributeId = INDEX_NONE;
+};
diff --git a/Engine/Plugins/PCG/Source/PCG/Private/Elements/PCGAnimBankMeshSpawner.cpp b/Engine/Plugins/PCG/Source/PCG/Private/Elements/PCGAnimBankMeshSpawner.cpp
deleted file mode 100644
diff --git a/Engine/Plugins/PCG/Source/PCG/Private/Elements/PCGAnimBankMeshSpawnerContext.cpp b/Engine/Plugins/PCG/Source/PCG/Private/Elements/PCGAnimBankMeshSpawnerContext.cpp
deleted file mode 100644
diff --git a/Engine/Plugins/PCG/Source/PCG/Private/Elements/PCGSkinnedMeshSpawner.cpp b/Engine/Plugins/PCG/Source/PCG/Private/Elements/PCGSkinnedMeshSpawner.cpp
new file mode 100644
--- /dev/null
+++ b/Engine/Plugins/PCG/Source/PCG/Private/Elements/PCGSkinnedMeshSpawner.cpp
@@ -0,0 +1,713 @@
+// Copyright Epic Games, Inc. All Rights Reserved.
+
+#include "Elements/PCGSkinnedMeshSpawner.h"
+
+#include "PCGComponent.h"
+#include "PCGCustomVersion.h"
+#include "PCGManagedResource.h"
+#include "Compute/PCGComputeCommon.h"
+#include "Compute/PCGDataBinding.h"
+#include "Compute/DataInterfaces/PCGInstanceDataInterface.h"
+#include "Compute/DataInterfaces/Elements/PCGSkinnedMeshSpawnerDataInterface.h"
+#include "Data/PCGPointData.h"
+#include "Data/PCGSpatialData.h"
+#include "Elements/PCGSkinnedMeshSpawnerContext.h"
+#include "Graph/PCGGPUGraphCompilationContext.h"
+#include "Helpers/PCGActorHelpers.h"
+#include "Helpers/PCGHelpers.h"
+#include "InstanceDataPackers/PCGSkinnedMeshInstanceDataPackerBase.h"
+#include "MeshSelectors/PCGSkinnedMeshSelector.h"
+
+#include "Components/InstancedSkinnedMeshComponent.h"
+#include "Engine/SkinnedAsset.h"
+#include "Materials/MaterialInterface.h"
+#include "ShaderCompilerCore.h"
+#include "UObject/Package.h"
+
+#include UE_INLINE_GENERATED_CPP_BY_NAME(PCGSkinnedMeshSpawner)
+
+#define LOCTEXT_NAMESPACE "PCGSkinnedMeshSpawnerElement"
+
+UPCGSkinnedMeshSpawnerSettings::UPCGSkinnedMeshSpawnerSettings(const FObjectInitializer &ObjectInitializer)
+{
+	// Implementation note: this should not have been done here (it should have been null), as it causes issues with copy & paste
+	// when the thing to paste does not have that class for its instance.
+	// However, removing it makes it that any object actually using the instance created by default would be lost.
+	if (!this->HasAnyFlags(RF_ClassDefaultObject))
+	{
+		MeshSelectorParameters = ObjectInitializer.CreateDefaultSubobject<UPCGSkinnedMeshSelector>(this, TEXT("DefaultSelectorInstance"));
+	}
+}
+
+TCHAR const* UPCGSkinnedMeshSpawnerSettings::TemplateFilePath = TEXT("/Plugin/PCG/Private/Elements/PCGSkinnedMeshSpawner.usf");
+
+bool UPCGSkinnedMeshSpawnerSettings::IsKernelValid(FPCGContext* InContext, bool bQuiet) const
+{
+	TRACE_CPUPROFILER_EVENT_SCOPE(UPCGSkinnedMeshSpawnerSettings::IsKernelValid);
+
+	if (!Super::IsKernelValid(InContext, bQuiet))
+	{
+		return false;
+	}
+
+	if (MeshSelectorParameters == nullptr)
+	{
+		return false;
+	}
+
+	if (!MeshSelectorParameters->BankAttribute.IsValid())
+	{
+		PCG_KERNEL_VALIDATION_ERR(InContext, this, bQuiet, LOCTEXT("NoBankAttribute", "Mesh selector bank attribute is invalid."));
+		return false;
+	}
+
+	if (!MeshSelectorParameters->SequenceIndexAttribute.IsValid())
+	{
+		PCG_KERNEL_VALIDATION_ERR(InContext, this, bQuiet, LOCTEXT("NoSequenceIndexAttribute", "Mesh selector sequence index attribute is invalid."));
+		return false;
+	}
+
+	if (!MeshSelectorParameters->MeshAttribute.IsValid())
+	{
+		PCG_KERNEL_VALIDATION_ERR(InContext, this, bQuiet, LOCTEXT("NoMeshAttribute", "Mesh selector mesh attribute is invalid."));
+		return false;
+	}
+
+	// Currently instance packers must be able to specify a full list of attribute names upfront, to build the attribute table at compile time.
+	// TODO: We should be able to augment a static attribute table with new attributes at execution time, which will allow other types like regex.
+	if (InstanceDataPackerParameters && !InstanceDataPackerParameters->GetAttributeNames(/*OutNames=*/nullptr))
+	{
+		PCG_KERNEL_VALIDATION_ERR(InContext, this, bQuiet, LOCTEXT("InvalidInstancePacker", "Selected instance packer does not support GPU execution."));
+		return false;
+	}
+
+	return true;
+}
+
+FString UPCGSkinnedMeshSpawnerSettings::GetCookedKernelSource(const FPCGKernelAttributeTable& InAttributeStaticTable) const
+{
+	FString TemplateFile;
+	ensure(LoadShaderSourceFile(TemplateFilePath, EShaderPlatform::SP_PCD3D_SM5, &TemplateFile, nullptr));
+	return TemplateFile;
+}
+
+bool UPCGSkinnedMeshSpawnerSettings::ComputeOutputPinDataDesc(const UPCGPin* OutputPin, const UPCGDataBinding* Binding, FPCGDataCollectionDesc& OutDesc) const
+{
+	check(OutputPin);
+	check(Binding);
+
+	// Code assumes single output pin.
+	if (!ensure(OutputPin->Properties.Label == PCGPinConstants::DefaultOutputLabel))
+	{
+		return false;
+	}
+
+	check(MeshSelectorParameters);
+
+	// First inspect data coming from CPU. Will return description if this node is a CPU node or in a different compute graph.
+	if (Binding->ComputeCPUOutputPinDataDesc(OutputPin, OutDesc))
+	{
+		return true;
+	}
+
+	const FPCGDataCollectionDesc InputPinDesc = PCGDataForGPUHelpers::ComputeInputPinDataDesc(this, PCGPinConstants::DefaultInputLabel, Binding);
+	
+	// First forward input data to output.
+	FPCGDataCollectionDesc OutputPinDesc = InputPinDesc;
+
+	OutDesc = OutputPinDesc;
+	return true;
+}
+
+void UPCGSkinnedMeshSpawnerSettings::GetKernelAttributeKeys(TArray<FPCGKernelAttributeKey>& OutKeys) const
+{
+}
+
+int UPCGSkinnedMeshSpawnerSettings::ComputeKernelThreadCount(const UPCGDataBinding* Binding) const
+{
+	const UPCGNode* Node = CastChecked<UPCGNode>(GetOuter());
+	const UPCGPin* Pin = Node->GetInputPin(PCGPinConstants::DefaultInputLabel);
+
+	return PCGDataForGPUHelpers::ComputeInputPinDataDesc(Pin, Binding).ComputeDataElementCount(EPCGDataType::Point);
+}
+
+#if WITH_EDITOR
+void UPCGSkinnedMeshSpawnerSettings::CreateAdditionalInputDataInterfaces(FPCGGPUCompilationContext& InOutContext, UObject* InObjectOuter, TArray<TObjectPtr<UComputeDataInterface>>& OutDataInterfaces) const
+{
+	Super::CreateAdditionalInputDataInterfaces(InOutContext, InObjectOuter, OutDataInterfaces);
+
+	TObjectPtr<UPCGSkinnedMeshSpawnerDataInterface> NodeDI = InOutContext.NewObject_AnyThread<UPCGSkinnedMeshSpawnerDataInterface>(InObjectOuter);
+	NodeDI->Settings = this;
+	OutDataInterfaces.Add(NodeDI);
+}
+
+void UPCGSkinnedMeshSpawnerSettings::CreateAdditionalOutputDataInterfaces(FPCGGPUCompilationContext& InOutContext, UObject* InObjectOuter, TArray<TObjectPtr<UComputeDataInterface>>& OutDataInterfaces) const
+{
+	Super::CreateAdditionalOutputDataInterfaces(InOutContext, InObjectOuter, OutDataInterfaces);
+
+	UPCGInstanceDataInterface* InstanceDI = InOutContext.NewObject_AnyThread<UPCGInstanceDataInterface>(InObjectOuter);
+	InstanceDI->SetProducerSettings(this);
+	InstanceDI->InputPinProvidingData = PCGPinConstants::DefaultInputLabel;
+	OutDataInterfaces.Add(InstanceDI);
+}
+
+FText UPCGSkinnedMeshSpawnerSettings::GetDefaultNodeTitle() const
+{
+	return LOCTEXT("NodeTitle", "Instanced Skinned Mesh Spawner");
+}
+
+void UPCGSkinnedMeshSpawnerSettings::ApplyDeprecation(UPCGNode* InOutNode)
+{
+	check(InOutNode);
+	Super::ApplyDeprecation(InOutNode);
+}
+#endif
+
+TArray<FPCGPinProperties> UPCGSkinnedMeshSpawnerSettings::InputPinProperties() const
+{
+	// Note: If executing on the GPU, we need to prevent multiple connections on inputs, since it is not supported at this time.
+	// Also note: Since the ShouldExecuteOnGPU() is already tied to structural changes, we don't need to implement any logic for this in GetChangeTypeForProperty()
+	const bool bAllowMultipleConnections = !ShouldExecuteOnGPU();
+
+	TArray<FPCGPinProperties> Properties;
+	FPCGPinProperties& InputPinProperty = Properties.Emplace_GetRef(PCGPinConstants::DefaultInputLabel, EPCGDataType::Point, bAllowMultipleConnections);
+	InputPinProperty.SetRequiredPin();
+
+	return Properties;
+}
+
+FPCGElementPtr UPCGSkinnedMeshSpawnerSettings::CreateElement() const
+{
+	return MakeShared<FPCGSkinnedMeshSpawnerElement>();
+}
+
+FPCGContext* FPCGSkinnedMeshSpawnerElement::CreateContext()
+{
+	return new FPCGSkinnedMeshSpawnerContext();
+}
+
+bool FPCGSkinnedMeshSpawnerElement::PrepareDataInternal(FPCGContext* InContext) const
+{
+	TRACE_CPUPROFILER_EVENT_SCOPE(FPCGSkinnedMeshSpawnerElement::PrepareDataInternal);
+	// TODO : time-sliced implementation
+	FPCGSkinnedMeshSpawnerContext* Context = static_cast<FPCGSkinnedMeshSpawnerContext*>(InContext);
+	const UPCGSkinnedMeshSpawnerSettings* Settings = Context->GetInputSettings<UPCGSkinnedMeshSpawnerSettings>();
+	check(Settings);
+
+	if (!Settings->MeshSelectorParameters)
+	{
+		PCGE_LOG(Error, GraphAndLog, LOCTEXT("InvalidMeshSelectorInstance", "Invalid MeshSelector instance, try reselecting the MeshSelector type"));
+		return true;
+	}
+
+	if (!Context->SourceComponent.Get())
+	{
+		return true;
+	}
+
+#if WITH_EDITOR
+	// In editor, we always want to generate this data for inspection & to prevent caching issues
+	const bool bGenerateOutput = true;
+#else
+	const bool bGenerateOutput = Context->Node && Context->Node->IsOutputPinConnected(PCGPinConstants::DefaultOutputLabel);
+#endif
+
+	// Check if we can reuse existing resources
+	bool& bSkippedDueToReuse = Context->bSkippedDueToReuse;
+
+	if (!Context->bReuseCheckDone)
+	{
+		// Compute CRC if it has not been computed (it likely isn't, but this is to futureproof this)
+		if (!Context->DependenciesCrc.IsValid())
+		{
+			GetDependenciesCrc(Context->InputData, Settings, Context->SourceComponent.Get(), Context->DependenciesCrc);
+		}
+		
+		if (Context->DependenciesCrc.IsValid())
+		{
+			TArray<UPCGManagedISKMComponent*> MISKMCs;
+			Context->SourceComponent->ForEachManagedResource([&MISKMCs, &Context, Settings](UPCGManagedResource* InResource)
+			{
+				if (UPCGManagedISKMComponent* Resource = Cast<UPCGManagedISKMComponent>(InResource))
+				{
+					if (Resource->GetCrc().IsValid() && Resource->GetCrc() == Context->DependenciesCrc)
+					{
+						MISKMCs.Add(Resource);
+					}
+				}
+			});
+
+			for (UPCGManagedISKMComponent* MISKMC : MISKMCs)
+			{
+				if (!MISKMC->IsMarkedUnused() && Settings->bWarnOnIdenticalSpawn)
+				{
+					// TODO: Revisit if the stack is added to the managed components at creation
+					PCGLog::LogWarningOnGraph(LOCTEXT("IdenticalABMCSpawn", "Identical Instanced Skinned Mesh Component spawn occurred. It may be beneficial to re-check graph logic for identical spawn conditions (same mesh descriptor at same location, etc) or repeated nodes."), Context);
+				}
+
+				MISKMC->MarkAsReused();
+			}
+
+			if (!MISKMCs.IsEmpty())
+			{
+				bSkippedDueToReuse = true;
+			}
+		}
+
+		Context->bReuseCheckDone = true;
+	}
+
+	// Early out - if we've established we could reuse resources and there is no need to generate an output, quit now
+	if (!bGenerateOutput && bSkippedDueToReuse)
+	{
+		return true;
+	}
+
+	// perform mesh selection
+	TArray<FPCGTaggedData> Inputs = Context->InputData.GetInputs();
+	TArray<FPCGTaggedData>& Outputs = Context->OutputData.TaggedData;
+
+	while(Context->CurrentInputIndex < Inputs.Num())
+	{
+		if (!Context->bCurrentInputSetup)
+		{
+			const FPCGTaggedData& Input = Inputs[Context->CurrentInputIndex];
+			const UPCGSpatialData* SpatialData = Cast<UPCGSpatialData>(Input.Data);
+
+			if (!SpatialData)
+			{
+				PCGE_LOG(Error, GraphAndLog, LOCTEXT("InvalidInputData", "Invalid input data"));
+				++Context->CurrentInputIndex;
+				continue;
+			}
+
+			const UPCGPointData* PointData = SpatialData->ToPointData(Context);
+			if (!PointData)
+			{
+				PCGE_LOG(Error, GraphAndLog, LOCTEXT("NoPointDataInInput", "Unable to get point data from input"));
+				++Context->CurrentInputIndex;
+				continue;
+			}
+
+			AActor* TargetActor = Settings->TargetActor.Get() ? Settings->TargetActor.Get() : Context->GetTargetActor(nullptr);
+			if (!TargetActor)
+			{
+				PCGE_LOG(Error, GraphAndLog, LOCTEXT("InvalidTargetActor", "Invalid target actor. Ensure TargetActor member is initialized when creating SpatialData."));
+				++Context->CurrentInputIndex;
+				continue;
+			}
+
+			if (bGenerateOutput)
+			{
+				FPCGTaggedData& Output = Outputs.Add_GetRef(Input);
+
+				UPCGPointData* OutputPointData = FPCGContext::NewObject_AnyThread<UPCGPointData>(Context);
+				OutputPointData->InitializeFromData(PointData);
+
+				Output.Data = OutputPointData;
+				check(!Context->CurrentOutputPointData);
+				Context->CurrentOutputPointData = OutputPointData;
+			}
+
+			FPCGSkinnedMeshSpawnerContext::FPackedInstanceListData& InstanceListData = Context->MeshInstancesData.Emplace_GetRef();
+			InstanceListData.TargetActor = TargetActor;
+			InstanceListData.SpatialData = PointData;
+
+			Context->CurrentPointData = PointData;
+			Context->bCurrentInputSetup = true;
+		}
+
+		// TODO: If we know we re-use the mesh components, we should not run the Selection, as it can be pretty costly.
+		// At the moment, the selection is filling the output point data, so it is necessary to run it. But we should just hit the cache in that case.
+		if (!Context->bSelectionDone)
+		{
+			check(Context->CurrentPointData);
+			Context->bSelectionDone = Settings->MeshSelectorParameters->SelectInstances(*Context, Settings, Context->CurrentPointData, Context->MeshInstancesData.Last().MeshInstances, Context->CurrentOutputPointData);
+		}
+
+		if (!Context->bSelectionDone)
+		{
+			return false;
+		}
+
+		// If we need the output but would otherwise skip the resource creation, we don't need to run the instance packing part of the processing
+		if (!bSkippedDueToReuse)
+		{
+			TArray<FPCGSkinnedMeshPackedCustomData>& PackedCustomData = Context->MeshInstancesData.Last().PackedCustomData;
+			const TArray<FPCGSkinnedMeshInstanceList>& MeshInstances = Context->MeshInstancesData.Last().MeshInstances;
+
+			if (PackedCustomData.Num() != MeshInstances.Num())
+			{
+				PackedCustomData.SetNum(MeshInstances.Num());
+			}
+
+			if (Settings->InstanceDataPackerParameters)
+			{
+				for (int32 InstanceListIndex = 0; InstanceListIndex < MeshInstances.Num(); ++InstanceListIndex)
+				{
+					Settings->InstanceDataPackerParameters->PackInstances(*Context, Context->CurrentPointData, MeshInstances[InstanceListIndex], PackedCustomData[InstanceListIndex]);
+				}
+			}
+		}
+
+		// We're done - cleanup for next iteration if we still have time
+		++Context->CurrentInputIndex;
+		Context->ResetInputIterationData();
+
+		// Continue on to next iteration if there is time left, otherwise, exit here
+		if (Context->AsyncState.ShouldStop() && Context->CurrentInputIndex < Inputs.Num())
+		{
+			return false;
+		}
+	}
+
+	IPCGAsyncLoadingContext* AsyncLoadingContext = static_cast<IPCGAsyncLoadingContext*>(Context);
+
+	if (Context->CurrentInputIndex == Inputs.Num() && !AsyncLoadingContext->WasLoadRequested() && !Context->MeshInstancesData.IsEmpty() && !Settings->bSynchronousLoad)
+	{
+		TArray<FSoftObjectPath> ObjectsToLoad;
+		for (const FPCGSkinnedMeshSpawnerContext::FPackedInstanceListData& InstanceData : Context->MeshInstancesData)
+		{
+			for (const FPCGSkinnedMeshInstanceList& MeshInstanceList : InstanceData.MeshInstances)
+			{
+				if (!MeshInstanceList.Descriptor.SkinnedAsset.IsNull())
+				{
+					ObjectsToLoad.AddUnique(MeshInstanceList.Descriptor.SkinnedAsset.ToSoftObjectPath());
+				}
+
+			#if 0 // AB-TODO
+				for (const TSoftObjectPtr<UMaterialInterface>& OverrideMaterial : MeshInstanceList.Descriptor.OverrideMaterials)
+				{
+					if (!OverrideMaterial.IsNull())
+					{
+						ObjectsToLoad.AddUnique(OverrideMaterial.ToSoftObjectPath());
+					}
+				}
+
+			#endif
+			}
+		}
+
+		return AsyncLoadingContext->RequestResourceLoad(Context, std::move(ObjectsToLoad), /*bAsynchronous=*/true);
+	}
+
+	return true;
+}
+
+bool FPCGSkinnedMeshSpawnerElement::ExecuteInternal(FPCGContext* InContext) const
+{
+	TRACE_CPUPROFILER_EVENT_SCOPE(FPCGSkinnedMeshSpawnerElement::Execute);
+	FPCGSkinnedMeshSpawnerContext* Context = static_cast<FPCGSkinnedMeshSpawnerContext*>(InContext);
+	const UPCGSkinnedMeshSpawnerSettings* Settings = Context->GetInputSettings<UPCGSkinnedMeshSpawnerSettings>();
+	check(Settings && !Settings->ShouldExecuteOnGPU());
+
+	while(!Context->MeshInstancesData.IsEmpty())
+	{
+		const FPCGSkinnedMeshSpawnerContext::FPackedInstanceListData& InstanceList = Context->MeshInstancesData.Last();
+		check(Context->bSkippedDueToReuse || InstanceList.MeshInstances.Num() == InstanceList.PackedCustomData.Num());
+
+		const bool bTargetActorValid = (InstanceList.TargetActor && IsValid(InstanceList.TargetActor));
+
+		if (bTargetActorValid)
+		{
+			while (Context->CurrentDataIndex < InstanceList.MeshInstances.Num())
+			{
+				const FPCGSkinnedMeshInstanceList& MeshInstance = InstanceList.MeshInstances[Context->CurrentDataIndex];
+				// We always have mesh instances, but if we are in re-use, we don't compute the packed custom data.
+				const FPCGSkinnedMeshPackedCustomData* PackedCustomData = InstanceList.PackedCustomData.IsValidIndex(Context->CurrentDataIndex) ? &InstanceList.PackedCustomData[Context->CurrentDataIndex] : nullptr;
+				SpawnSkinnedMeshInstances(Context, MeshInstance, InstanceList.TargetActor, PackedCustomData);
+
+				// Now that the mesh is loaded/spawned, set the bounds to out points if requested.
+				if (MeshInstance.Descriptor.SkinnedAsset && Settings->bApplyMeshBoundsToPoints)
+				{
+					if (TMap<UPCGPointData*, TArray<int32>>* OutPointDataToPointIndex = Context->MeshToOutPoints.Find(MeshInstance.Descriptor.SkinnedAsset))
+					{
+						const FBox Bounds = MeshInstance.Descriptor.SkinnedAsset->GetBounds().GetBox();
+						for (TPair<UPCGPointData*, TArray<int32>>& It : *OutPointDataToPointIndex)
+						{
+							check(It.Key);
+							TArray<FPCGPoint>& OutPoints = It.Key->GetMutablePoints();
+							for (int32 Index : It.Value)
+							{
+								FPCGPoint& Point = OutPoints[Index];
+								Point.BoundsMin = Bounds.Min;
+								Point.BoundsMax = Bounds.Max;
+							}
+						}
+					}
+				}
+
+				++Context->CurrentDataIndex;
+
+				if (Context->AsyncState.ShouldStop())
+				{
+					break;
+				}
+			}
+		}
+
+		if (!bTargetActorValid || Context->CurrentDataIndex == InstanceList.MeshInstances.Num())
+		{
+			Context->MeshInstancesData.RemoveAtSwap(Context->MeshInstancesData.Num() - 1);
+			Context->CurrentDataIndex = 0;
+		}
+
+		if (Context->AsyncState.ShouldStop())
+		{
+			break;
+		}
+	}
+
+	const bool bFinishedExecution = Context->MeshInstancesData.IsEmpty();
+	if (bFinishedExecution)
+	{
+		if (AActor* TargetActor = Settings->TargetActor.Get() ? Settings->TargetActor.Get() : Context->GetTargetActor(nullptr))
+		{
+			for (UFunction* Function : PCGHelpers::FindUserFunctions(TargetActor->GetClass(), Settings->PostProcessFunctionNames, { UPCGFunctionPrototypes::GetPrototypeWithNoParams() }, Context))
+			{
+				TargetActor->ProcessEvent(Function, nullptr);
+			}
+		}
+	}
+
+	return bFinishedExecution;
+}
+
+bool FPCGSkinnedMeshSpawnerElement::CanExecuteOnlyOnMainThread(FPCGContext* Context) const
+{
+	// PrepareData can call UPCGManagedComponent::MarkAsReused which registers the mesh component, which can go into Chaos code that asserts if not on main thread.
+	// TODO: We can likely re-enable multi-threading for PrepareData if we move the call to MarkAsReused to Execute. There should hopefully not be
+	// wider contention on resources resources are not shared across nodes and are also per-component.
+	return !Context || Context->CurrentPhase == EPCGExecutionPhase::Execute || Context->CurrentPhase == EPCGExecutionPhase::PrepareData;
+}
+
+void FPCGSkinnedMeshSpawnerElement::SpawnSkinnedMeshInstances(FPCGSkinnedMeshSpawnerContext* Context, const FPCGSkinnedMeshInstanceList& InstanceList, AActor* TargetActor, const FPCGSkinnedMeshPackedCustomData* InPackedCustomData) const
+{
+	// Populate the mesh component from the previously prepared entries
+	TRACE_CPUPROFILER_EVENT_SCOPE(FPCGSkinnedMeshSpawnerElement::Execute::PopulateAB);
+
+	if (InstanceList.Instances.Num() == 0)
+	{
+		return;
+	}
+
+	// Will be synchronously loaded if not loaded. But by default it should already have been loaded asynchronously in PrepareData, so this is free.
+	USkinnedAsset* LoadedMesh = InstanceList.Descriptor.SkinnedAsset.LoadSynchronous();
+
+	if (!LoadedMesh)
+	{
+		// Either we have no mesh (so nothing to do) or the mesh couldn't be loaded
+		if (InstanceList.Descriptor.SkinnedAsset.IsValid())
+		{
+			PCGE_LOG(Error, GraphAndLog, FText::Format(LOCTEXT("MeshLoadFailed", "Unable to load mesh '{0}'"), FText::FromString(InstanceList.Descriptor.SkinnedAsset.ToString())));
+		}
+
+		return;
+	}
+
+	// Don't spawn meshes if we reuse the ISMCs, but we still want to be sure that the mesh is loaded at least (for operations downstream).
+	if (Context->bSkippedDueToReuse)
+	{
+		return;
+	}
+
+#if 0 // AB-TODO
+	for (TSoftObjectPtr<UMaterialInterface> OverrideMaterial : InstanceList.Descriptor.OverrideMaterials)
+	{
+		// Will be synchronously loaded if not loaded. But by default it should already have been loaded asynchronously in PrepareData, so this is free.
+		if (OverrideMaterial.IsValid() && !OverrideMaterial.LoadSynchronous())
+		{
+			PCGE_LOG(Error, GraphAndLog, FText::Format(LOCTEXT("OverrideMaterialLoadFailed", "Unable to load override material '{0}'"), FText::FromString(OverrideMaterial.ToString())));
+			return;
+		}
+	}
+#endif
+
+	// If we spawn the meshes, we should have computed a packed custom data.
+	if (!ensure(InPackedCustomData))
+	{
+		return;
+	}
+
+	const FPCGSkinnedMeshPackedCustomData& PackedCustomData = *InPackedCustomData;
+
+	FPCGSkinnedMeshComponentBuilderParams Params;
+	Params.Descriptor = InstanceList.Descriptor;
+	Params.NumCustomDataFloats = PackedCustomData.NumCustomDataFloats;
+
+	// If the root actor we're binding to is movable, then the component should be movable by default
+	if (USceneComponent* SceneComponent = TargetActor->GetRootComponent())
+	{
+		Params.Descriptor.Mobility = SceneComponent->Mobility;
+	}
+
+	const UPCGSkinnedMeshSpawnerSettings* Settings = Context->GetInputSettings<UPCGSkinnedMeshSpawnerSettings>();
+	check(Settings);
+
+	Params.SettingsCrc = Settings->GetSettingsCrc();
+	ensure(Params.SettingsCrc.IsValid());
+
+	UPCGManagedISKMComponent* MISKMC = UPCGActorHelpers::GetOrCreateManagedABMC(TargetActor, Context->SourceComponent.Get(), Params, Context);
+
+	check(MISKMC);
+	MISKMC->SetCrc(Context->DependenciesCrc);
+
+	// Keep track of all touched resources in the context, because if the execution is cancelled during the SMS execution
+	// we cannot easily guarantee that the state (esp. vs CRCs) is going to be entirely valid
+	Context->TouchedResources.Emplace(MISKMC);
+
+	UInstancedSkinnedMeshComponent* ISKMC = MISKMC->GetComponent();
+	check(ISKMC);
+
+	const int32 PreExistingInstanceCount = ISKMC->GetInstanceCount();
+	const int32 NewInstanceCount = InstanceList.Instances.Num();
+	const int32 NumCustomDataFloats = PackedCustomData.NumCustomDataFloats;
+
+	check((ISKMC->GetNumCustomDataFloats() == 0 && PreExistingInstanceCount == 0) || ISKMC->GetNumCustomDataFloats() == NumCustomDataFloats);
+	ISKMC->SetNumCustomDataFloats(NumCustomDataFloats);
+
+	TArray<FTransform> Transforms;
+	TArray<int32> BankIndices;
+
+	Transforms.Reserve(NewInstanceCount);
+	BankIndices.Reserve(NewInstanceCount);
+
+	// TODO: Remove allocs/copies
+	for (const FPCGSkinnedMeshInstance& Instance : InstanceList.Instances)
+	{
+		Transforms.Emplace(Instance.Transform);
+		BankIndices.Emplace(Instance.BankIndex);
+	}
+
+	// Populate the instances
+	TArray<FPrimitiveInstanceId> NewIds = ISKMC->AddInstances(Transforms, BankIndices, /*bShouldReturnIndices=*/true, /*bWorldSpace=*/true);
+
+	// Copy new CustomData into the ABMC PerInstanceSMCustomData
+	if (NumCustomDataFloats > 0)
+	{
+		for (int32 NewIndex = 0; NewIndex < NewInstanceCount; ++NewIndex)
+		{
+			ISKMC->SetCustomData(NewIds[NewIndex], MakeArrayView(&PackedCustomData.CustomData[NewIndex * NumCustomDataFloats], NumCustomDataFloats));
+		}
+	}
+
+	ISKMC->UpdateBounds();
+	ISKMC->OptimizeInstanceData();
+
+	{
+		PCGE_LOG(Verbose, LogOnly, FText::Format(LOCTEXT("GenerationInfo", "Added {0} instances of '{1}' on actor '{2}'"),
+			InstanceList.Instances.Num(), FText::FromString(InstanceList.Descriptor.SkinnedAsset->GetFName().ToString()), FText::FromString(TargetActor->GetFName().ToString())));
+	}
+}
+
+void FPCGSkinnedMeshSpawnerElement::AbortInternal(FPCGContext* InContext) const
+{
+	TRACE_CPUPROFILER_EVENT_SCOPE(FPCGSkinnedMeshSpawnerElement::AbortInternal);
+	// It is possible to Abort a ready task with no context yet
+	if (!InContext)
+	{
+		return;
+	}
+
+	FPCGSkinnedMeshSpawnerContext* Context = static_cast<FPCGSkinnedMeshSpawnerContext*>(InContext);
+
+	// Any resources we've touched during the execution of this node can potentially be in a "not-quite complete state" especially if we have multiple sources of data writing to the same ISMC.
+	// In this case, we're aiming to mark the resources as "Unused" so they are picked up to be removed during the component's OnProcessGraphAborted, which is why we call Release here.
+	for (TWeakObjectPtr<UPCGManagedISKMComponent> ManagedResource : Context->TouchedResources)
+	{
+		if(ManagedResource.IsValid())
+		{
+			TSet<TSoftObjectPtr<AActor>> Dummy;
+			ManagedResource->Release(/*bHardRelease=*/false, Dummy);
+		}
+	}
+}
+
+void UPCGSkinnedMeshSpawnerSettings::PostLoad()
+{
+	Super::PostLoad();
+
+	const EObjectFlags Flags = GetMaskedFlags(RF_PropagateToSubObjects) | RF_Transactional;
+
+	if (!MeshSelectorParameters)
+	{
+		MeshSelectorParameters = NewObject<UPCGSkinnedMeshSelector>(this, UPCGSkinnedMeshSelector::StaticClass(), NAME_None, GetMaskedFlags(RF_PropagateToSubObjects));
+	}
+	else
+	{
+		MeshSelectorParameters->SetFlags(Flags);
+	}
+
+	if (!InstanceDataPackerParameters)
+	{
+		RefreshInstancePacker();
+	}
+	else
+	{
+		InstanceDataPackerParameters->SetFlags(Flags);
+	}
+}
+
+#if WITH_EDITOR
+void UPCGSkinnedMeshSpawnerSettings::PostEditChangeProperty(FPropertyChangedEvent& PropertyChangedEvent) 
+{
+	if (PropertyChangedEvent.Property)
+	{
+		const FName& PropertyName = PropertyChangedEvent.Property->GetFName();
+
+		if (PropertyName == GET_MEMBER_NAME_CHECKED(UPCGSkinnedMeshSpawnerSettings, InstanceDataPackerType))
+		{
+			RefreshInstancePacker();
+		}
+	}
+
+	Super::PostEditChangeProperty(PropertyChangedEvent);
+}
+
+bool UPCGSkinnedMeshSpawnerSettings::CanEditChange(const FProperty* InProperty) const
+{
+	return Super::CanEditChange(InProperty);
+}
+#endif
+
+void UPCGSkinnedMeshSpawnerSettings::SetInstancePackerType(TSubclassOf<UPCGSkinnedMeshInstanceDataPackerBase> InInstancePackerType)
+{
+	if (!InstanceDataPackerParameters || InInstancePackerType != InstanceDataPackerType)
+	{
+		if (InInstancePackerType != InstanceDataPackerType)
+		{
+			InstanceDataPackerType = InInstancePackerType;
+		}
+		
+		RefreshInstancePacker();
+	}
+}
+
+void UPCGSkinnedMeshSpawnerSettings::RefreshInstancePacker()
+{
+	if (InstanceDataPackerType)
+	{
+		ensure(IsInGameThread());
+
+		if (InstanceDataPackerParameters)
+		{
+		#if WITH_EDITOR
+			InstanceDataPackerParameters->Rename(nullptr, GetTransientPackage(), REN_DontCreateRedirectors);
+		#endif
+			InstanceDataPackerParameters->MarkAsGarbage();
+			InstanceDataPackerParameters = nullptr;
+		}
+
+		const EObjectFlags Flags = GetMaskedFlags(RF_PropagateToSubObjects);
+		InstanceDataPackerParameters = NewObject<UPCGSkinnedMeshInstanceDataPackerBase>(this, InstanceDataPackerType, NAME_None, Flags);
+	}
+	else
+	{
+		InstanceDataPackerParameters = nullptr;
+	}
+}
+
+FPCGSkinnedMeshSpawnerContext::FPackedInstanceListData::FPackedInstanceListData() = default;
+FPCGSkinnedMeshSpawnerContext::FPackedInstanceListData::~FPackedInstanceListData() = default;
+
+#undef LOCTEXT_NAMESPACE
diff --git a/Engine/Plugins/PCG/Source/PCG/Private/Elements/PCGSkinnedMeshSpawnerContext.cpp b/Engine/Plugins/PCG/Source/PCG/Private/Elements/PCGSkinnedMeshSpawnerContext.cpp
new file mode 100644
--- /dev/null
+++ b/Engine/Plugins/PCG/Source/PCG/Private/Elements/PCGSkinnedMeshSpawnerContext.cpp
@@ -0,0 +1,17 @@
+// Copyright Epic Games, Inc. All Rights Reserved.
+
+#include "Elements/PCGSkinnedMeshSpawnerContext.h"
+
+void FPCGSkinnedMeshSpawnerContext::ResetInputIterationData()
+{
+	bCurrentInputSetup = false;
+	bSelectionDone = false;
+	bPartitionDone = false;
+	CurrentPointData = nullptr;
+	CurrentOutputPointData = nullptr;
+	MaterialOverrideHelper.Reset();
+	CurrentPointIndex = 0;
+	MeshToValueKey.Reset();
+	AttributeOverridePartition.Reset();
+	OverriddenDescriptors.Reset();
+}
\ No newline at end of file
diff --git a/Engine/Plugins/PCG/Source/PCG/Private/Helpers/PCGActorHelpers.cpp b/Engine/Plugins/PCG/Source/PCG/Private/Helpers/PCGActorHelpers.cpp
--- a/Engine/Plugins/PCG/Source/PCG/Private/Helpers/PCGActorHelpers.cpp
+++ b/Engine/Plugins/PCG/Source/PCG/Private/Helpers/PCGActorHelpers.cpp
@@ -11,7 +11,7 @@
 #include "EngineUtils.h"
 #include "Components/HierarchicalInstancedStaticMeshComponent.h"
 #include "Components/SplineMeshComponent.h"
-#include "Components/AnimBankMeshComponent.h"
+#include "Components/InstancedSkinnedMeshComponent.h"
 #include "Engine/StaticMesh.h"
 #include "Engine/SkinnedAsset.h"
 #include "Animation/AnimBank.h"
@@ -237,13 +237,13 @@ UPCGManagedISMComponent* UPCGActorHelpers::GetOrCreateManagedISMC(AActor* InTarg
 	return Resource;
 }
 
-UPCGManagedAnimBankMeshComponent* UPCGActorHelpers::GetOrCreateManagedABMC(AActor* InTargetActor, UPCGComponent* InSourceComponent, const FPCGAnimBankMeshComponentBuilderParams& InParams, FPCGContext* OptionalContext)
+UPCGManagedISKMComponent* UPCGActorHelpers::GetOrCreateManagedABMC(AActor* InTargetActor, UPCGComponent* InSourceComponent, const FPCGSkinnedMeshComponentBuilderParams& InParams, FPCGContext* OptionalContext)
 {
 	check(InTargetActor && InSourceComponent);
 
 	TRACE_CPUPROFILER_EVENT_SCOPE(UPCGActorHelpers::GetOrCreateManagedABMC);
 
-	FAnimBankMeshComponentDescriptor Descriptor(InParams.Descriptor);
+	FSkinnedMeshComponentDescriptor Descriptor(InParams.Descriptor);
 
 	const USkinnedAsset* SkinnedAsset = Descriptor.SkinnedAsset;
 	if (!SkinnedAsset)
@@ -251,25 +251,25 @@ UPCGManagedAnimBankMeshComponent* UPCGActorHelpers::GetOrCreateManagedABMC(AActo
 		return nullptr;
 	}
 
-	auto AddTagsToComponent = [InSourceComponent, &InParams](UAnimBankMeshComponent* ABMC)
+	auto AddTagsToComponent = [InSourceComponent, &InParams](UInstancedSkinnedMeshComponent* ISKMC)
 	{
-		ABMC->ComponentTags.AddUnique(PCGHelpers::DefaultPCGTag);
-		ABMC->ComponentTags.AddUnique(InSourceComponent->GetFName());
+		ISKMC->ComponentTags.AddUnique(PCGHelpers::DefaultPCGTag);
+		ISKMC->ComponentTags.AddUnique(InSourceComponent->GetFName());
 			
-		for (FName ComponentTag : InParams.Descriptor.ComponentTags)
+		for (const FName& ComponentTag : InParams.Descriptor.ComponentTags)
 		{
-			ABMC->ComponentTags.AddUnique(ComponentTag);
+			ISKMC->ComponentTags.AddUnique(ComponentTag);
 		}
 	};
 
 	if (!Descriptor.ComponentClass)
 	{
-		Descriptor.ComponentClass = UAnimBankMeshComponent::StaticClass();
+		Descriptor.ComponentClass = UInstancedSkinnedMeshComponent::StaticClass();
 	}
 
 	{
 		TRACE_CPUPROFILER_EVENT_SCOPE(UPCGActorHelpers::GetOrCreateManagedABMC::FindMatchingMABMC);
-		UPCGManagedAnimBankMeshComponent* MatchingResource = nullptr;
+		UPCGManagedISKMComponent* MatchingResource = nullptr;
 		if(InParams.SettingsCrc.IsValid())
 		{
 			InSourceComponent->ForEachManagedResource([&MatchingResource, &InParams, &InTargetActor, &Descriptor](UPCGManagedResource* InResource)
@@ -280,20 +280,20 @@ UPCGManagedAnimBankMeshComponent* UPCGActorHelpers::GetOrCreateManagedABMC(AActo
 					return;
 				}
 
-				if (UPCGManagedAnimBankMeshComponent* Resource = Cast<UPCGManagedAnimBankMeshComponent>(InResource))
+				if (UPCGManagedISKMComponent* Resource = Cast<UPCGManagedISKMComponent>(InResource))
 				{
 					if (!Resource->GetSettingsCrc().IsValid() || Resource->GetSettingsCrc() != InParams.SettingsCrc || !Resource->CanBeUsed())
 					{
 						return;
 					}
 
-					if (UAnimBankMeshComponent* ABMC = Resource->GetComponent())
+					if (UInstancedSkinnedMeshComponent* ISKMC = Resource->GetComponent())
 					{
-						if (IsValid(ABMC) &&
-							ABMC->GetOwner() == InTargetActor &&
-							ABMC->GetNumCustomDataFloats() == InParams.NumCustomDataFloats &&
+						if (IsValid(ISKMC) &&
+							ISKMC->GetOwner() == InTargetActor &&
+							ISKMC->GetNumCustomDataFloats() == InParams.NumCustomDataFloats &&
 							Resource->GetDescriptor() == Descriptor &&
-							ABMC->HasAnyFlags(RF_Transient) == InParams.bTransient)
+							ISKMC->HasAnyFlags(RF_Transient) == InParams.bTransient)
 						{
 							MatchingResource = Resource;
 						}
@@ -307,11 +307,11 @@ UPCGManagedAnimBankMeshComponent* UPCGActorHelpers::GetOrCreateManagedABMC(AActo
 			TRACE_CPUPROFILER_EVENT_SCOPE(UPCGActorHelpers::GetOrCreateManagedABMC::MarkAsUsed);
 			MatchingResource->MarkAsUsed();
 
-			UAnimBankMeshComponent* ABMC = Cast<UAnimBankMeshComponent>(MatchingResource->GeneratedComponent.Get());
-			if (ensure(ABMC))
+			UInstancedSkinnedMeshComponent* ISKMC = Cast<UInstancedSkinnedMeshComponent>(MatchingResource->GeneratedComponent.Get());
+			if (ensure(ISKMC))
 			{
-				ABMC->Modify(!InSourceComponent->IsInPreviewMode());
-				AddTagsToComponent(ABMC);
+				ISKMC->Modify(!InSourceComponent->IsInPreviewMode());
+				AddTagsToComponent(ISKMC);
 			}
 
 			return MatchingResource;
@@ -331,24 +331,24 @@ UPCGManagedAnimBankMeshComponent* UPCGActorHelpers::GetOrCreateManagedABMC(AActo
 		ObjectFlags |= RF_Transient;
 	}
 
-	UAnimBankMeshComponent* ABMC = NewObject<UAnimBankMeshComponent>(InTargetActor, Descriptor.ComponentClass, MakeUniqueObjectName(InTargetActor, Descriptor.ComponentClass, FName(ComponentName)), ObjectFlags);
-	Descriptor.InitComponent(ABMC);
-	ABMC->SetNumCustomDataFloats(InParams.NumCustomDataFloats);
+	UInstancedSkinnedMeshComponent* ISKMC = NewObject<UInstancedSkinnedMeshComponent>(InTargetActor, Descriptor.ComponentClass, MakeUniqueObjectName(InTargetActor, Descriptor.ComponentClass, FName(ComponentName)), ObjectFlags);
+	Descriptor.InitComponent(ISKMC);
+	ISKMC->SetNumCustomDataFloats(InParams.NumCustomDataFloats);
 
-	ABMC->RegisterComponent();
-	InTargetActor->AddInstanceComponent(ABMC);
+	ISKMC->RegisterComponent();
+	InTargetActor->AddInstanceComponent(ISKMC);
 
-	if (!ABMC->AttachToComponent(InTargetActor->GetRootComponent(), FAttachmentTransformRules(EAttachmentRule::KeepRelative, EAttachmentRule::KeepWorld, EAttachmentRule::KeepWorld, false)))
+	if (!ISKMC->AttachToComponent(InTargetActor->GetRootComponent(), FAttachmentTransformRules(EAttachmentRule::KeepRelative, EAttachmentRule::KeepWorld, EAttachmentRule::KeepWorld, false)))
 	{
 		PCGLog::Component::LogComponentAttachmentFailedWarning(OptionalContext);
 	}
 	
-	// Implementation note: Because we've used the FAnimBankMeshComponentDescriptor here (Descriptor vs InParams.Descriptor) which takes care of the loading, we still need to apply tags manually.
-	AddTagsToComponent(ABMC);
+	// Implementation note: Because we've used the FSkinnedMeshComponentDescriptor here (Descriptor vs InParams.Descriptor) which takes care of the loading, we still need to apply tags manually.
+	AddTagsToComponent(ISKMC);
 
 	// Create managed resource on source component
-	UPCGManagedAnimBankMeshComponent* Resource = NewObject<UPCGManagedAnimBankMeshComponent>(InSourceComponent, NAME_None, InParams.bTransient ? RF_Transient : RF_NoFlags);
-	Resource->SetComponent(ABMC);
+	UPCGManagedISKMComponent* Resource = NewObject<UPCGManagedISKMComponent>(InSourceComponent, NAME_None, InParams.bTransient ? RF_Transient : RF_NoFlags);
+	Resource->SetComponent(ISKMC);
 	Resource->SetDescriptor(Descriptor);
 	if (InTargetActor->GetRootComponent())
 	{
diff --git a/Engine/Plugins/PCG/Source/PCG/Private/InstanceDataPackers/PCGAnimBankMeshInstanceDataPackerBase.cpp b/Engine/Plugins/PCG/Source/PCG/Private/InstanceDataPackers/PCGAnimBankMeshInstanceDataPackerBase.cpp
deleted file mode 100644
diff --git a/Engine/Plugins/PCG/Source/PCG/Private/InstanceDataPackers/PCGAnimBankMeshInstanceDataPackerByAttribute.cpp b/Engine/Plugins/PCG/Source/PCG/Private/InstanceDataPackers/PCGAnimBankMeshInstanceDataPackerByAttribute.cpp
deleted file mode 100644
diff --git a/Engine/Plugins/PCG/Source/PCG/Private/InstanceDataPackers/PCGAnimBankMeshInstanceDataPackerByRegex.cpp b/Engine/Plugins/PCG/Source/PCG/Private/InstanceDataPackers/PCGAnimBankMeshInstanceDataPackerByRegex.cpp
deleted file mode 100644
diff --git a/Engine/Plugins/PCG/Source/PCG/Private/InstanceDataPackers/PCGSkinnedMeshInstanceDataPackerBase.cpp b/Engine/Plugins/PCG/Source/PCG/Private/InstanceDataPackers/PCGSkinnedMeshInstanceDataPackerBase.cpp
new file mode 100644
--- /dev/null
+++ b/Engine/Plugins/PCG/Source/PCG/Private/InstanceDataPackers/PCGSkinnedMeshInstanceDataPackerBase.cpp
@@ -0,0 +1,325 @@
+// Copyright Epic Games, Inc. All Rights Reserved.
+
+#include "InstanceDataPackers/PCGSkinnedMeshInstanceDataPackerBase.h"
+
+#include "PCGElement.h"
+#include "Elements/Metadata/PCGMetadataElementCommon.h"
+#include "MeshSelectors/PCGSkinnedMeshSelector.h"
+#include "Metadata/PCGMetadata.h"
+#include "Metadata/Accessors/IPCGAttributeAccessor.h"
+#include "Metadata/Accessors/PCGAttributeAccessorKeys.h"
+
+#include UE_INLINE_GENERATED_CPP_BY_NAME(PCGSkinnedMeshInstanceDataPackerBase)
+
+void UPCGSkinnedMeshInstanceDataPackerBase::PackInstances_Implementation(FPCGContext& Context, const UPCGSpatialData* InSpatialData, const FPCGSkinnedMeshInstanceList& InstanceList, FPCGSkinnedMeshPackedCustomData& OutPackedCustomData) const
+{
+	PCGE_LOG_C(Error, GraphAndLog, &Context, NSLOCTEXT("PCGSkinnedMeshInstanceDataPackerBase", "SkinnedMeshInstanceDataPackerBaseFailed", "Unable to execute SkinnedMeshInstanceDataPacker pure virtual base function, override the PackInstances function or use a default implementation."));
+}
+
+bool UPCGSkinnedMeshInstanceDataPackerBase::AddTypeToPacking(int TypeId, FPCGSkinnedMeshPackedCustomData& OutPackedCustomData) const
+{
+	switch (TypeId)
+	{
+	case PCG::Private::MetadataTypes<bool>::Id: // fall-through
+	case PCG::Private::MetadataTypes<float>::Id: // fall-through
+	case PCG::Private::MetadataTypes<double>::Id: // fall-through
+	case PCG::Private::MetadataTypes<int32>::Id: // fall-through
+	case PCG::Private::MetadataTypes<int64>::Id:
+		OutPackedCustomData.NumCustomDataFloats += 1;
+		break;
+	case PCG::Private::MetadataTypes<FVector2D>::Id:
+		OutPackedCustomData.NumCustomDataFloats += 2;
+		break;
+	case PCG::Private::MetadataTypes<FRotator>::Id: // fall-through
+	case PCG::Private::MetadataTypes<FVector>::Id:
+		OutPackedCustomData.NumCustomDataFloats += 3;
+		break;
+	case PCG::Private::MetadataTypes<FVector4>::Id: // fall-through
+	case PCG::Private::MetadataTypes<FQuat>::Id:
+		OutPackedCustomData.NumCustomDataFloats += 4;
+		break;
+	default:
+		return false;
+	}
+
+	return true;
+}
+
+void UPCGSkinnedMeshInstanceDataPackerBase::PackCustomDataFromAttributes(const FPCGSkinnedMeshInstanceList& InstanceList, const UPCGMetadata* Metadata, const TArray<FName>& AttributeNames, FPCGSkinnedMeshPackedCustomData& OutPackedCustomData) const
+{
+	if (!Metadata)
+	{
+		UE_LOG(LogPCG, Error, TEXT("Invalid metadata"));
+		return;
+	}
+
+	TArray<const FPCGMetadataAttributeBase*> Attributes;
+	for (const FName& AttributeName : AttributeNames)
+	{
+		Attributes.Add(Metadata->GetConstAttribute(AttributeName));
+	}
+
+	PackCustomDataFromAttributes(InstanceList, Attributes, OutPackedCustomData);
+}
+
+void UPCGSkinnedMeshInstanceDataPackerBase::PackCustomDataFromAttributes(const FPCGSkinnedMeshInstanceList& InstanceList, const TArray<const FPCGMetadataAttributeBase*>& Attributes, FPCGSkinnedMeshPackedCustomData& OutPackedCustomData) const
+{
+	const UPCGPointData* PointData = InstanceList.PointData.Get();
+	for (const int32 Index : InstanceList.InstancePointIndices)
+	{
+		const PCGMetadataEntryKey EntryKey = PointData ? PointData->GetPoints()[Index].MetadataEntry : Index;
+		for (const FPCGMetadataAttributeBase* AttributeBase : Attributes)
+		{
+			check(AttributeBase);
+
+			switch (AttributeBase->GetTypeId())
+			{
+			case PCG::Private::MetadataTypes<bool>::Id:
+			{
+				const FPCGMetadataAttribute<bool>* Attribute = static_cast<const FPCGMetadataAttribute<bool>*>(AttributeBase);
+				check(Attribute);
+
+				const bool Value = Attribute->GetValueFromItemKey(EntryKey);
+				OutPackedCustomData.CustomData.Add(Value);
+				break;
+			}
+			case PCG::Private::MetadataTypes<float>::Id:
+			{
+				const FPCGMetadataAttribute<float>* Attribute = static_cast<const FPCGMetadataAttribute<float>*>(AttributeBase);
+				check(Attribute);
+
+				const float Value = Attribute->GetValueFromItemKey(EntryKey);
+				OutPackedCustomData.CustomData.Add(Value);
+				break;
+			}
+			case PCG::Private::MetadataTypes<double>::Id:
+			{
+				const FPCGMetadataAttribute<double>* Attribute = static_cast<const FPCGMetadataAttribute<double>*>(AttributeBase);
+				check(Attribute);
+
+				const double Value = Attribute->GetValueFromItemKey(EntryKey);
+				OutPackedCustomData.CustomData.Add(Value);
+				break;
+			}
+			case PCG::Private::MetadataTypes<int32>::Id:
+			{
+				const FPCGMetadataAttribute<int32>* Attribute = static_cast<const FPCGMetadataAttribute<int32>*>(AttributeBase);
+				check(Attribute);
+
+				const float Value = static_cast<float>(Attribute->GetValueFromItemKey(EntryKey));
+				OutPackedCustomData.CustomData.Add(Value);
+				break;
+			}
+			case PCG::Private::MetadataTypes<int64>::Id:
+			{
+				const FPCGMetadataAttribute<int64>* Attribute = static_cast<const FPCGMetadataAttribute<int64>*>(AttributeBase);
+				check(Attribute);
+
+				const float Value = static_cast<float>(Attribute->GetValueFromItemKey(EntryKey));
+				OutPackedCustomData.CustomData.Add(Value);
+				break;
+			}
+			case PCG::Private::MetadataTypes<FRotator>::Id:
+			{
+				const FPCGMetadataAttribute<FRotator>* Attribute = static_cast<const FPCGMetadataAttribute<FRotator>*>(AttributeBase);
+				check(Attribute);
+
+				const FRotator Value = Attribute->GetValueFromItemKey(EntryKey);
+				OutPackedCustomData.CustomData.Add(Value.Roll);
+				OutPackedCustomData.CustomData.Add(Value.Pitch);
+				OutPackedCustomData.CustomData.Add(Value.Yaw);
+				break;
+			}
+			case PCG::Private::MetadataTypes<FVector2D>::Id:
+			{
+				const FPCGMetadataAttribute<FVector2D>* Attribute = static_cast<const FPCGMetadataAttribute<FVector2D>*>(AttributeBase);
+				check(Attribute);
+
+				const FVector2D Value = Attribute->GetValueFromItemKey(EntryKey);
+				OutPackedCustomData.CustomData.Add(Value.X);
+				OutPackedCustomData.CustomData.Add(Value.Y);
+				break;
+			}
+			case PCG::Private::MetadataTypes<FVector>::Id:
+			{
+				const FPCGMetadataAttribute<FVector>* Attribute = static_cast<const FPCGMetadataAttribute<FVector>*>(AttributeBase);
+				check(Attribute);
+
+				const FVector Value = Attribute->GetValueFromItemKey(EntryKey);
+				OutPackedCustomData.CustomData.Add(Value.X);
+				OutPackedCustomData.CustomData.Add(Value.Y);
+				OutPackedCustomData.CustomData.Add(Value.Z);
+				break;
+			}
+			case PCG::Private::MetadataTypes<FVector4>::Id:
+			{
+				const FPCGMetadataAttribute<FVector4>* Attribute = static_cast<const FPCGMetadataAttribute<FVector4>*>(AttributeBase);
+				check(Attribute);
+
+				const FVector4 Value = Attribute->GetValueFromItemKey(EntryKey);
+				OutPackedCustomData.CustomData.Add(Value.X);
+				OutPackedCustomData.CustomData.Add(Value.Y);
+				OutPackedCustomData.CustomData.Add(Value.Z);
+				OutPackedCustomData.CustomData.Add(Value.W);
+				break;
+			}
+			case PCG::Private::MetadataTypes<FQuat>::Id:
+			{
+				const FPCGMetadataAttribute<FQuat>* Attribute = static_cast<const FPCGMetadataAttribute<FQuat>*>(AttributeBase);
+				check(Attribute);
+
+				const FQuat Value = Attribute->GetValueFromItemKey(EntryKey);
+				OutPackedCustomData.CustomData.Add(Value.X);
+				OutPackedCustomData.CustomData.Add(Value.Y);
+				OutPackedCustomData.CustomData.Add(Value.Z);
+				OutPackedCustomData.CustomData.Add(Value.W);
+				break;
+			}
+			default:
+				break;
+			}
+		}
+	}
+}
+
+void UPCGSkinnedMeshInstanceDataPackerBase::PackCustomDataFromAccessors(const FPCGSkinnedMeshInstanceList& InstanceList, TArray<TUniquePtr<const IPCGAttributeAccessor>> Accessors, TArray<TUniquePtr<const IPCGAttributeAccessorKeys>> AccessorKeys, FPCGSkinnedMeshPackedCustomData& OutPackedCustomData) const
+{
+	const int32 TotalNumCustomDataFloats = OutPackedCustomData.NumCustomDataFloats * InstanceList.Instances.Num();
+
+	if (!ensure(Accessors.Num() == AccessorKeys.Num()))
+	{
+		// Zero out the custom data to avoid sending garbage data to the GPU.
+		OutPackedCustomData.CustomData.SetNumZeroed(TotalNumCustomDataFloats);
+		return;
+	}
+
+	OutPackedCustomData.CustomData.SetNumUninitialized(TotalNumCustomDataFloats);
+
+	if (TotalNumCustomDataFloats == 0)
+	{
+		return;
+	}
+
+	// Index of the accessor's attribute/property for a single instance (e.g. if the
+	// accessors are [float, FVector, float], AccessorPackedIndex will be [0, 1, 4])
+	int AccessorPackedIndex = 0;
+	const EPCGAttributeAccessorFlags Flags = EPCGAttributeAccessorFlags::AllowBroadcast | EPCGAttributeAccessorFlags::AllowConstructible;
+
+	for (int AccessorIndex = 0; AccessorIndex < Accessors.Num(); ++AccessorIndex)
+	{
+		const IPCGAttributeAccessor* Accessor = Accessors[AccessorIndex].Get();
+		const IPCGAttributeAccessorKeys* Keys = AccessorKeys[AccessorIndex].Get();
+
+		check(Accessor && Keys);
+
+		switch (Accessor->GetUnderlyingType())
+		{
+		case PCG::Private::MetadataTypes<bool>::Id: // fall-through
+		case PCG::Private::MetadataTypes<float>::Id: // fall-through
+		case PCG::Private::MetadataTypes<double>::Id: // fall-through
+		case PCG::Private::MetadataTypes<int32>::Id: // fall-through
+		case PCG::Private::MetadataTypes<int64>::Id:
+		{
+			auto PackFloats = [&OutPackedCustomData, AccessorPackedIndex](float Value, int32 Index)
+			{
+				const int PackedIndex = Index * OutPackedCustomData.NumCustomDataFloats + AccessorPackedIndex;
+				OutPackedCustomData.CustomData[PackedIndex] = Value;
+			};
+
+			ensure(PCGMetadataElementCommon::ApplyOnAccessor<float>(*Keys, *Accessor, PackFloats, Flags));
+
+			AccessorPackedIndex += 1;
+
+			break;
+		}
+		case PCG::Private::MetadataTypes<FRotator>::Id:
+		{
+			auto PackFloats = [&OutPackedCustomData, AccessorPackedIndex](const FRotator& Value, int32 Index)
+			{
+				const int PackedIndex = Index * OutPackedCustomData.NumCustomDataFloats + AccessorPackedIndex;
+
+				OutPackedCustomData.CustomData[PackedIndex + 0] = static_cast<float>(Value.Roll);
+				OutPackedCustomData.CustomData[PackedIndex + 1] = static_cast<float>(Value.Pitch);
+				OutPackedCustomData.CustomData[PackedIndex + 2] = static_cast<float>(Value.Yaw);
+			};
+
+			ensure(PCGMetadataElementCommon::ApplyOnAccessor<FRotator>(*Keys, *Accessor, PackFloats, Flags));
+
+			AccessorPackedIndex += 3;
+
+			break;
+		}
+		case PCG::Private::MetadataTypes<FVector2D>::Id:
+		{
+			auto PackFloats = [&OutPackedCustomData, AccessorPackedIndex](const FVector2D& Value, int32 Index)
+			{
+				const int PackedIndex = Index * OutPackedCustomData.NumCustomDataFloats + AccessorPackedIndex;
+
+				OutPackedCustomData.CustomData[PackedIndex + 0] = static_cast<float>(Value.X);
+				OutPackedCustomData.CustomData[PackedIndex + 1] = static_cast<float>(Value.Y);
+			};
+
+			ensure(PCGMetadataElementCommon::ApplyOnAccessor<FVector2D>(*Keys, *Accessor, PackFloats, Flags));
+
+			AccessorPackedIndex += 2;
+
+			break;
+		}
+		case PCG::Private::MetadataTypes<FVector>::Id:
+		{
+			auto PackFloats = [&OutPackedCustomData, AccessorPackedIndex](const FVector& Value, int32 Index)
+			{
+				const int PackedIndex = Index * OutPackedCustomData.NumCustomDataFloats + AccessorPackedIndex;
+
+				OutPackedCustomData.CustomData[PackedIndex + 0] = static_cast<float>(Value.X);
+				OutPackedCustomData.CustomData[PackedIndex + 1] = static_cast<float>(Value.Y);
+				OutPackedCustomData.CustomData[PackedIndex + 2] = static_cast<float>(Value.Z);
+			};
+
+			ensure(PCGMetadataElementCommon::ApplyOnAccessor<FVector>(*Keys, *Accessor, PackFloats, Flags));
+
+			AccessorPackedIndex += 3;
+
+			break;
+		}
+		case PCG::Private::MetadataTypes<FVector4>::Id:
+		{
+			auto PackFloats = [&OutPackedCustomData, AccessorPackedIndex](const FVector4& Value, int32 Index)
+			{
+				const int PackedIndex = Index * OutPackedCustomData.NumCustomDataFloats + AccessorPackedIndex;
+
+				OutPackedCustomData.CustomData[PackedIndex + 0] = static_cast<float>(Value.X);
+				OutPackedCustomData.CustomData[PackedIndex + 1] = static_cast<float>(Value.Y);
+				OutPackedCustomData.CustomData[PackedIndex + 2] = static_cast<float>(Value.Z);
+				OutPackedCustomData.CustomData[PackedIndex + 3] = static_cast<float>(Value.W);
+			};
+
+			ensure(PCGMetadataElementCommon::ApplyOnAccessor<FVector4>(*Keys, *Accessor, PackFloats, Flags));
+
+			AccessorPackedIndex += 4;
+
+			break;
+		}
+		case PCG::Private::MetadataTypes<FQuat>::Id:
+		{
+			auto PackFloats = [&OutPackedCustomData, AccessorPackedIndex](const FQuat& Value, int32 Index)
+			{
+				const int PackedIndex = Index * OutPackedCustomData.NumCustomDataFloats + AccessorPackedIndex;
+
+				OutPackedCustomData.CustomData[PackedIndex + 0] = static_cast<float>(Value.X);
+				OutPackedCustomData.CustomData[PackedIndex + 1] = static_cast<float>(Value.Y);
+				OutPackedCustomData.CustomData[PackedIndex + 2] = static_cast<float>(Value.Z);
+				OutPackedCustomData.CustomData[PackedIndex + 3] = static_cast<float>(Value.W);
+			};
+
+			ensure(PCGMetadataElementCommon::ApplyOnAccessor<FQuat>(*Keys, *Accessor, PackFloats, Flags));
+
+			AccessorPackedIndex += 4;
+
+			break;
+		}
+		default:
+			break;
+		}
+	}
+}
diff --git a/Engine/Plugins/PCG/Source/PCG/Private/InstanceDataPackers/PCGSkinnedMeshInstanceDataPackerByAttribute.cpp b/Engine/Plugins/PCG/Source/PCG/Private/InstanceDataPackers/PCGSkinnedMeshInstanceDataPackerByAttribute.cpp
new file mode 100644
--- /dev/null
+++ b/Engine/Plugins/PCG/Source/PCG/Private/InstanceDataPackers/PCGSkinnedMeshInstanceDataPackerByAttribute.cpp
@@ -0,0 +1,89 @@
+// Copyright Epic Games, Inc. All Rights Reserved.
+
+#include "InstanceDataPackers/PCGSkinnedMeshInstanceDataPackerByAttribute.h"
+
+#include "PCGContext.h"
+#include "PCGElement.h"
+#include "Data/PCGSpatialData.h"
+#include "InstanceDataPackers/PCGSkinnedMeshInstanceDataPackerBase.h"
+#include "Metadata/Accessors/PCGAttributeAccessorHelpers.h"
+#include "MeshSelectors/PCGSkinnedMeshSelector.h"
+
+#include UE_INLINE_GENERATED_CPP_BY_NAME(PCGSkinnedMeshInstanceDataPackerByAttribute)
+
+#define LOCTEXT_NAMESPACE "PCGSkinnedMeshInstanceDataPackerByAttribute"
+
+#if WITH_EDITOR
+void UPCGSkinnedMeshInstanceDataPackerByAttribute::PostLoad()
+{
+	Super::PostLoad();
+}
+#endif // WITH_EDITOR
+
+void UPCGSkinnedMeshInstanceDataPackerByAttribute::PackInstances_Implementation(FPCGContext& Context, const UPCGSpatialData* InSpatialData, const FPCGSkinnedMeshInstanceList& InstanceList, FPCGSkinnedMeshPackedCustomData& OutPackedCustomData) const
+{
+	if (!InSpatialData || !InSpatialData->Metadata)
+	{
+		PCGLog::InputOutput::LogInvalidInputDataError(&Context);
+		return;
+	}
+
+	TArray<TUniquePtr<const IPCGAttributeAccessor>> SelectedAccessors;
+	TArray<TUniquePtr<const IPCGAttributeAccessorKeys>> SelectedKeys;
+
+	SelectedAccessors.Reserve(AttributeSelectors.Num());
+	SelectedKeys.Reserve(AttributeSelectors.Num());
+
+	// Find attributes and calculate NumCustomDataFloats
+	for (const FPCGAttributePropertyInputSelector& Selector : AttributeSelectors)
+	{
+		TUniquePtr<const IPCGAttributeAccessor> Accessor = PCGAttributeAccessorHelpers::CreateConstAccessor(InSpatialData, Selector);
+		TUniquePtr<const IPCGAttributeAccessorKeys> Keys;
+		const UPCGPointData* PointData = InstanceList.PointData.Get();
+		TArray<const PCGMetadataEntryKey> ExtractedKeys;
+		
+		if (InSpatialData == PointData)
+		{
+			Keys = MakeUnique<const FPCGAttributeAccessorKeysPointsSubset>(PointData->GetPoints(), InstanceList.InstancePointIndices);
+		}
+		else
+		{
+			// Convert indices to entry keys
+			ExtractedKeys.Reserve(InstanceList.InstancePointIndices.Num());
+			Algo::Transform(InstanceList.InstancePointIndices, ExtractedKeys, [](const int32 Index) -> PCGMetadataEntryKey{ return Index; });
+			Keys = MakeUnique<const FPCGAttributeAccessorKeysEntries>(ExtractedKeys);
+		}
+
+		if (!Accessor.IsValid() || !Keys.IsValid())
+		{
+			PCGLog::Metadata::LogFailToCreateAccessorError(Selector, &Context);
+			continue;
+		}
+
+		if (!AddTypeToPacking(Accessor->GetUnderlyingType(), OutPackedCustomData))
+		{
+			PCGLog::LogWarningOnGraph(FText::Format(LOCTEXT("AttributeInvalidType", "Attribute/property '{0}' is not a valid type - skipped."), Selector.GetDisplayText()), &Context);
+			continue;
+		}
+
+		SelectedAccessors.Add(std::move(Accessor));
+		SelectedKeys.Add(std::move(Keys));
+	}
+
+	PackCustomDataFromAccessors(InstanceList, std::move(SelectedAccessors), std::move(SelectedKeys), OutPackedCustomData);
+}
+
+bool UPCGSkinnedMeshInstanceDataPackerByAttribute::GetAttributeNames(TArray<FName>* OutNames)
+{
+	if (OutNames)
+	{
+		for (const FPCGAttributePropertyInputSelector& AttributeSelector : AttributeSelectors)
+		{
+			OutNames->Add(AttributeSelector.GetAttributeName());
+		}
+	}
+
+	return true;
+}
+
+#undef LOCTEXT_NAMESPACE
diff --git a/Engine/Plugins/PCG/Source/PCG/Private/InstanceDataPackers/PCGSkinnedMeshInstanceDataPackerByRegex.cpp b/Engine/Plugins/PCG/Source/PCG/Private/InstanceDataPackers/PCGSkinnedMeshInstanceDataPackerByRegex.cpp
new file mode 100644
--- /dev/null
+++ b/Engine/Plugins/PCG/Source/PCG/Private/InstanceDataPackers/PCGSkinnedMeshInstanceDataPackerByRegex.cpp
@@ -0,0 +1,70 @@
+// Copyright Epic Games, Inc. All Rights Reserved.
+
+#include "InstanceDataPackers/PCGSkinnedMeshInstanceDataPackerByRegex.h"
+
+#include "PCGContext.h"
+#include "PCGElement.h"
+#include "Data/PCGSpatialData.h"
+#include "InstanceDataPackers/PCGSkinnedMeshInstanceDataPackerBase.h"
+
+#include "Internationalization/Regex.h"
+
+#include UE_INLINE_GENERATED_CPP_BY_NAME(PCGSkinnedMeshInstanceDataPackerByRegex)
+
+#define LOCTEXT_NAMESPACE "PCGSkinnedMeshInstanceDataPackerByRegex"
+
+void UPCGSkinnedMeshInstanceDataPackerByRegex::PackInstances_Implementation(FPCGContext& Context, const UPCGSpatialData* InSpatialData, const FPCGSkinnedMeshInstanceList& InstanceList, FPCGSkinnedMeshPackedCustomData& OutPackedCustomData) const
+{
+	if (!InSpatialData || !InSpatialData->Metadata)
+	{
+		PCGE_LOG_C(Error, GraphAndLog, &Context, LOCTEXT("InvalidInputData", "Invalid input data"));
+		return;
+	}
+	
+	TArray<FName> AttributeNames;
+	TArray<EPCGMetadataTypes> AttributeTypes;
+	InSpatialData->Metadata->GetAttributes(AttributeNames, AttributeTypes);
+
+	TArray<const FPCGMetadataAttributeBase*> SelectedAttributes;
+	TSet<PCGMetadataAttributeKey> FoundAttributes;
+
+	for (const FString& Regex : RegexPatterns)
+	{
+		const FRegexPattern Pattern(Regex);
+	
+		// Match our Regex Pattern against every Attribute
+		for (int Index = 0; Index < AttributeNames.Num(); ++Index)
+		{
+			const FName& AttributeName = AttributeNames[Index];
+			FRegexMatcher RegexMatcher(Pattern, AttributeName.ToString());
+	
+			if (!RegexMatcher.FindNext()) 
+			{
+				continue;
+			}
+	
+			const FPCGMetadataAttributeBase* AttributeBase = InSpatialData->Metadata->GetConstAttribute(AttributeName);
+			check(AttributeBase);
+
+			// Avoid adding the same attribute multiple times if it is captured by different regex patterns
+			if (FoundAttributes.Contains(AttributeBase->AttributeId))
+			{
+				continue;
+			}
+
+			FoundAttributes.Add(AttributeBase->AttributeId);
+	
+			if (!AddTypeToPacking(AttributeBase->GetTypeId(), OutPackedCustomData))
+			{
+				PCGE_LOG_C(Warning, GraphAndLog, &Context, FText::Format(LOCTEXT("InvalidAttributeType", "Attribute name '{0}' is not of a valid type"), FText::FromName(AttributeName)));
+				continue;
+			}
+	
+			SelectedAttributes.Add(AttributeBase);
+		}
+	}
+	
+	PackCustomDataFromAttributes(InstanceList, SelectedAttributes, OutPackedCustomData);
+}
+
+#undef LOCTEXT_NAMESPACE
diff --git a/Engine/Plugins/PCG/Source/PCG/Private/MeshSelectors/PCGAnimBankMeshDescriptor.cpp b/Engine/Plugins/PCG/Source/PCG/Private/MeshSelectors/PCGAnimBankMeshDescriptor.cpp
deleted file mode 100644
diff --git a/Engine/Plugins/PCG/Source/PCG/Private/MeshSelectors/PCGAnimBankMeshSelector.cpp b/Engine/Plugins/PCG/Source/PCG/Private/MeshSelectors/PCGAnimBankMeshSelector.cpp
deleted file mode 100644
diff --git a/Engine/Plugins/PCG/Source/PCG/Private/MeshSelectors/PCGSkinnedMeshDescriptor.cpp b/Engine/Plugins/PCG/Source/PCG/Private/MeshSelectors/PCGSkinnedMeshDescriptor.cpp
new file mode 100644
--- /dev/null
+++ b/Engine/Plugins/PCG/Source/PCG/Private/MeshSelectors/PCGSkinnedMeshDescriptor.cpp
@@ -0,0 +1,54 @@
+// Copyright Epic Games, Inc. All Rights Reserved.
+
+#include "MeshSelectors/PCGSkinnedMeshDescriptor.h"
+#include "Components/InstancedSkinnedMeshComponent.h"
+
+#include UE_INLINE_GENERATED_CPP_BY_NAME(PCGSkinnedMeshDescriptor)
+
+FPCGSoftSkinnedMeshComponentDescriptor::FPCGSoftSkinnedMeshComponentDescriptor()
+: FSoftSkinnedMeshComponentDescriptor()
+{
+	// Override defaults from base class
+	ComponentClass = UInstancedSkinnedMeshComponent::StaticClass();
+	Mobility = EComponentMobility::Static;
+	// GW-TODO: bGenerateOverlapEvents = 0;
+	// GW-TODO: BodyInstance.SetCollisionProfileName(UCollisionProfile::NoCollision_ProfileName);
+}
+
+FPCGSoftSkinnedMeshComponentDescriptor::FPCGSoftSkinnedMeshComponentDescriptor(const FSkinnedMeshComponentDescriptor& Other)
+	: Super(Other)
+{}
+
+void FPCGSoftSkinnedMeshComponentDescriptor::InitFrom(const UInstancedSkinnedMeshComponent* Component, bool bInitBodyInstance)
+{
+	ComponentTags = Component->ComponentTags;
+	Super::InitFrom(Component, bInitBodyInstance);
+}
+
+void FPCGSoftSkinnedMeshComponentDescriptor::InitComponent(UInstancedSkinnedMeshComponent* ABMComponent) const
+{
+	ABMComponent->ComponentTags = ComponentTags;
+	Super::InitComponent(ABMComponent);
+}
+
+int32 FPCGSoftSkinnedMeshComponentDescriptor::GetOrAddBankIndex(const FSoftAnimBankItem& BankItem)
+{
+	int32 BankIndex = INDEX_NONE;
+
+	for (int32 ExistingBankIndex = 0; ExistingBankIndex < BankItems.Num(); ++ExistingBankIndex)
+	{
+		if (BankItems[ExistingBankIndex] == BankItem)
+		{
+			BankIndex = ExistingBankIndex;
+			break;
+		}
+	}
+
+	if (BankIndex == INDEX_NONE)
+	{
+		BankIndex = BankItems.Num();
+		BankItems.Emplace(BankItem);
+	}
+
+	return BankIndex;
+}
diff --git a/Engine/Plugins/PCG/Source/PCG/Private/MeshSelectors/PCGSkinnedMeshSelector.cpp b/Engine/Plugins/PCG/Source/PCG/Private/MeshSelectors/PCGSkinnedMeshSelector.cpp
new file mode 100644
--- /dev/null
+++ b/Engine/Plugins/PCG/Source/PCG/Private/MeshSelectors/PCGSkinnedMeshSelector.cpp
@@ -0,0 +1,382 @@
+// Copyright Epic Games, Inc. All Rights Reserved.
+
+#include "MeshSelectors/PCGSkinnedMeshSelector.h"
+
+#include "Data/PCGPointData.h"
+#include "Data/PCGSpatialData.h"
+#include "Elements/PCGSkinnedMeshSpawner.h"
+#include "Elements/PCGSkinnedMeshSpawnerContext.h"
+#include "Elements/Metadata/PCGMetadataElementCommon.h"
+#include "MeshSelectors/PCGSkinnedMeshSelector.h"
+#include "Metadata/PCGObjectPropertyOverride.h"
+#include "Metadata/PCGMetadataPartitionCommon.h"
+#include "Metadata/Accessors/PCGAttributeAccessorHelpers.h"
+
+#include "Algo/AnyOf.h"
+#include "Animation/AnimBank.h"
+
+#include UE_INLINE_GENERATED_CPP_BY_NAME(PCGSkinnedMeshSelector)
+
+#define LOCTEXT_NAMESPACE "PCGSkinnedMeshSelector"
+
+namespace PCGSkinnedMeshSelector
+{
+	// Returns variation based on mesh, material overrides and reverse culling
+	FPCGSkinnedMeshInstanceList& GetInstanceList(
+		TArray<FPCGSkinnedMeshInstanceList>& InstanceLists,
+		const FPCGSoftSkinnedMeshComponentDescriptor& TemplateDescriptor,
+		TSoftObjectPtr<USkinnedAsset> Asset,
+		const TArray<TSoftObjectPtr<UMaterialInterface>>& MaterialOverrides,
+		bool bReverseCulling,
+		const UPCGPointData* InPointData,
+		const int AttributePartitionIndex = INDEX_NONE)
+	{
+		for (FPCGSkinnedMeshInstanceList& InstanceList : InstanceLists)
+		{
+			if (InstanceList.Descriptor.SkinnedAsset == Asset &&
+				// GW-TODO: InstanceList.Descriptor.bReverseCulling == bReverseCulling &&
+				// GW-TODO: InstanceList.Descriptor.OverrideMaterials == MaterialOverrides &&
+				InstanceList.AttributePartitionIndex == AttributePartitionIndex)
+			{
+				return InstanceList;
+			}
+		}
+
+		FPCGSkinnedMeshInstanceList& NewInstanceList = InstanceLists.Emplace_GetRef(TemplateDescriptor);
+		NewInstanceList.Descriptor.SkinnedAsset = Asset;
+		// GW-TODO: NewInstanceList.Descriptor.OverrideMaterials = MaterialOverrides;
+		// GW-TODO: NewInstanceList.Descriptor.bReverseCulling = bReverseCulling;
+		NewInstanceList.AttributePartitionIndex = AttributePartitionIndex;
+		NewInstanceList.PointData = InPointData;
+
+		return NewInstanceList;
+	}
+}
+
+void UPCGSkinnedMeshSelector::PostLoad()
+{
+	Super::PostLoad();
+
+	// TODO: Remove if/when FBodyInstance is updated or replaced
+	// Necessary to update the collision Response Container from the Response Array
+	TemplateDescriptor.PostLoadFixup(this);
+}
+
+bool UPCGSkinnedMeshSelector::SelectInstances(
+	FPCGSkinnedMeshSpawnerContext& Context,
+	const UPCGSkinnedMeshSpawnerSettings* Settings,
+	const UPCGPointData* InPointData,
+	TArray<FPCGSkinnedMeshInstanceList>& OutMeshInstances,
+	UPCGPointData* OutPointData) const
+{
+	TRACE_CPUPROFILER_EVENT_SCOPE(UPCGMeshSelector::SelectInstances);
+
+	if (!InPointData)
+	{
+		PCGE_LOG_C(Error, GraphAndLog, &Context, LOCTEXT("InputMissingData", "Missing input data"));
+		return true;
+	}
+
+	if (!InPointData->Metadata)
+	{
+		PCGE_LOG_C(Error, GraphAndLog, &Context, LOCTEXT("InputMissingMetadata", "Unable to get metadata from input"));
+		return true;
+	}
+
+	if (!InPointData->Metadata->HasAttribute(BankAttribute.GetAttributeName()))
+	{
+		PCGE_LOG_C(Error, GraphAndLog, &Context, FText::Format(LOCTEXT("AttributeNotInMetadata", "Attribute '{0}' is not in the metadata"), BankAttribute.GetDisplayText()));
+		return true;
+	}
+
+	if (!InPointData->Metadata->HasAttribute(SequenceIndexAttribute.GetAttributeName()))
+	{
+		PCGE_LOG_C(Error, GraphAndLog, &Context, FText::Format(LOCTEXT("AttributeNotInMetadata", "Attribute '{0}' is not in the metadata"), SequenceIndexAttribute.GetDisplayText()));
+		return true;
+	}
+
+	if (!InPointData->Metadata->HasAttribute(MeshAttribute.GetAttributeName()))
+	{
+		PCGE_LOG_C(Error, GraphAndLog, &Context, FText::Format(LOCTEXT("AttributeNotInMetadata", "Attribute '{0}' is not in the metadata"), MeshAttribute.GetDisplayText()));
+		return true;
+	}
+
+	const FPCGMetadataAttributeBase* BankAttributeBase = InPointData->Metadata->GetConstAttribute(BankAttribute.GetAttributeName());
+	check(BankAttributeBase);
+
+	const FPCGMetadataAttributeBase* SequenceIndexAttributeBase = InPointData->Metadata->GetConstAttribute(SequenceIndexAttribute.GetAttributeName());
+	check(SequenceIndexAttributeBase);
+
+	const FPCGMetadataAttributeBase* MeshAttributeBase = InPointData->Metadata->GetConstAttribute(MeshAttribute.GetAttributeName());
+	check(MeshAttributeBase);
+
+	// Validate that the "bank" attribute is of the right type
+	if (!PCG::Private::IsOfTypes<FSoftObjectPath, FString>(BankAttributeBase->GetTypeId()))
+	{
+		PCGE_LOG_C(Error, GraphAndLog, &Context, FText::Format(LOCTEXT("BankAttributeInvalidType", "Attribute '{0}' is not of valid type (must be FString or FSoftObjectPath)"), BankAttribute.GetDisplayText()));
+		return true;
+	}
+
+	// Validate that the "sequence index" attribute is of the right type
+	if (!PCG::Private::IsOfTypes<int32>(SequenceIndexAttributeBase->GetTypeId()))
+	{
+		PCGE_LOG_C(Error, GraphAndLog, &Context, FText::Format(LOCTEXT("SeqIndexAttributeInvalidType", "Attribute '{0}' is not of valid type (must be int32)"), SequenceIndexAttribute.GetDisplayText()));
+		return true;
+	}
+
+	// Validate that the "mesh" attribute is of the right type
+	if (!PCG::Private::IsOfTypes<FSoftObjectPath, FString>(MeshAttributeBase->GetTypeId()))
+	{
+		PCGE_LOG_C(Error, GraphAndLog, &Context, FText::Format(LOCTEXT("MeshAttributeInvalidType", "Attribute '{0}' is not of valid type (must be FString or FSoftObjectPath)"), MeshAttribute.GetDisplayText()));
+		return true;
+	}
+
+	FPCGMeshMaterialOverrideHelper& MaterialOverrideHelper = Context.MaterialOverrideHelper;
+	if (!MaterialOverrideHelper.IsInitialized())
+	{
+		MaterialOverrideHelper.Initialize(Context, bUseAttributeMaterialOverrides, {} /*TemplateDescriptor.OverrideMaterials*/, MaterialOverrideAttributes, InPointData->Metadata);
+	}
+
+	if (!MaterialOverrideHelper.IsValid())
+	{
+		return true;
+	}
+
+	// ByAttribute takes in SoftObjectPaths per point in the metadata, so we can pass those directly into the outgoing pin if it exists
+	if (Context.CurrentPointIndex == 0 && OutPointData)
+	{
+		TRACE_CPUPROFILER_EVENT_SCOPE(UPCGMeshSelector::SetupOutPointData);
+		OutPointData->SetPoints(InPointData->GetPoints());
+	}
+
+	TRACE_CPUPROFILER_EVENT_SCOPE(UPCGSkinnedMeshSelector::SelectEntries);
+
+	if (!Context.bPartitionDone)
+	{
+		TRACE_CPUPROFILER_EVENT_SCOPE(UPCGSkinnedMeshSelector::SelectEntries::BuildingPartition);
+
+		TArray<FPCGObjectPropertyOverrideDescription> PropertyOverrides = Settings->SkinnedMeshComponentPropertyOverrides;
+		const FString SkinnedAssetPropertyString = GET_MEMBER_NAME_CHECKED(FPCGSoftSkinnedMeshComponentDescriptor, SkinnedAsset).ToString();
+
+		// Add the skinned asset override to the list only if not already provided
+		if (!Algo::AnyOf(PropertyOverrides, [&SkinnedAssetPropertyString](const FPCGObjectPropertyOverrideDescription& PropertyOverride) { return PropertyOverride.PropertyTarget == SkinnedAssetPropertyString; }))
+		{
+			PropertyOverrides.Emplace(MeshAttribute, SkinnedAssetPropertyString);
+		}
+
+		// Validate all the selectors are actual FPCGSoftSkinnedMeshComponentDescriptor properties
+		TArray<FPCGAttributePropertySelector> ValidSelectorOverrides;
+
+		for (const FPCGObjectPropertyOverrideDescription& PropertyOverride : PropertyOverrides)
+		{
+			if (FPCGSoftSkinnedMeshComponentDescriptor::StaticStruct()->FindPropertyByName(FName(PropertyOverride.PropertyTarget)))
+			{
+				ValidSelectorOverrides.Emplace_GetRef() = FPCGAttributePropertySelector::CreateFromOtherSelector<FPCGAttributePropertySelector>(PropertyOverride.InputSource);
+			}
+			else
+			{
+				PCGE_LOG_C(Error, GraphAndLog, &Context, FText::Format(LOCTEXT("OverriddenPropertyNotFound", "Property '{0}' not a valid property with an ABM Descriptor. It will be ignored."), FText::FromString(PropertyOverride.PropertyTarget)));
+			}
+		}
+
+		// If there are valid overrides, partition the points on those attributes so that an instance can be created for each
+		if (!ValidSelectorOverrides.IsEmpty())
+		{
+			Context.AttributeOverridePartition = PCGMetadataPartitionCommon::AttributeGenericPartition(InPointData, ValidSelectorOverrides, &Context, Settings->bSilenceOverrideAttributeNotFoundErrors);
+		}
+
+		// Set the descriptors to match the partition count. Uninitialized, because it will be copied from the template below
+		Context.OverriddenDescriptors.Reserve(Context.AttributeOverridePartition.Num());
+		for (int I = 0; I < Context.AttributeOverridePartition.Num(); ++I)
+		{
+			FPCGSoftSkinnedMeshComponentDescriptor& Descriptor = Context.OverriddenDescriptors.Add_GetRef(TemplateDescriptor);
+
+			// If partition is empty (which can happen, esp. for the default partition on the default value, we'll just skip it here.
+			if (Context.AttributeOverridePartition[I].IsEmpty())
+			{
+				continue;
+			}
+
+			// Use the Object Override to map the user's input selector and property to the descriptor
+			FPCGObjectOverrides Overrides(&Descriptor);
+			Overrides.Initialize(PropertyOverrides, &Descriptor, InPointData, &Context);
+
+			// Since they are already partitioned and identical, we can just use the value on the first point
+			check(!Context.AttributeOverridePartition[I].IsEmpty());
+			int32 AnyPointIndexOnThisPartition = Context.AttributeOverridePartition[I][0];
+			Overrides.Apply(AnyPointIndexOnThisPartition);
+		}
+
+		Context.PointBankValues.Reset();
+		Context.PointSequenceIndexValues.Reset();
+
+		// Given partitioning is expensive, check if we're out of time for this frame
+		Context.bPartitionDone = true;
+		if (Context.ShouldStop())
+		{
+			return false;
+		}
+	}
+
+	// Assign points to entries
+	int32 CurrentPartitionIndex = Context.CurrentPointIndex; // misnomer but we're reusing another concept from the context
+	const TArray<FPCGPoint>& Points = InPointData->GetPoints();
+
+	{
+		TRACE_CPUPROFILER_EVENT_SCOPE(UPCGSkinnedMeshSelector::SelectEntries::PushingPointsToInstanceLists);
+
+		if (Context.PointBankValues.IsEmpty() || Context.PointSequenceIndexValues.IsEmpty())
+		{
+			PCGAttributeAccessorHelpers::ExtractAllValues(InPointData, BankAttribute, Context.PointBankValues, &Context);
+			PCGAttributeAccessorHelpers::ExtractAllValues(InPointData, SequenceIndexAttribute, Context.PointSequenceIndexValues, &Context);
+		}
+
+		// TODO: Revisit this when attribute partitioning is returned in a more optimized form
+		// The partition index is used to assign the point to the correct partition's instance
+		while (CurrentPartitionIndex < Context.AttributeOverridePartition.Num())
+		{
+			const int32 PartitionIndex = CurrentPartitionIndex++;
+			TArray<int32>& Partition = Context.AttributeOverridePartition[PartitionIndex];
+			const FPCGSoftSkinnedMeshComponentDescriptor& CurrentPartitionDescriptor = Context.OverriddenDescriptors[PartitionIndex];
+			
+			if (Partition.IsEmpty() || CurrentPartitionDescriptor.SkinnedAsset.IsNull())
+			{
+				continue;
+			}
+
+			bool bSkip = false;
+
+			for (const FSoftAnimBankItem& BankItem : CurrentPartitionDescriptor.BankItems)
+			{
+				if (BankItem.BankAsset.IsNull())
+				{
+					bSkip = true;
+					break;
+				}
+			}
+
+			if (bSkip)
+			{
+				continue;
+			}
+
+			// Setup data for mesh bounds computation
+			if (OutPointData && Settings->bApplyMeshBoundsToPoints)
+			{
+				TRACE_CPUPROFILER_EVENT_SCOPE(UPCGSkinnedMeshSelector::SelectEntries::PushPointsToApplyMeshBounds);
+				TArray<int32>& PointIndices = Context.MeshToOutPoints.FindOrAdd(CurrentPartitionDescriptor.SkinnedAsset).FindOrAdd(OutPointData);
+				PointIndices.Append(Partition);
+			}
+
+			// Separate the inverse determinant instances so we can push them to a different ABM
+			TArray<int32> ReverseInstances;
+			{
+				TRACE_CPUPROFILER_EVENT_SCOPE(UPCGSkinnedMeshSelector::SelectEntries::BuildReverseInstances);
+				int32 WriteIndex = 0;
+				for (int32 ReadIndex = 0; ReadIndex < Partition.Num(); ++ReadIndex)
+				{
+					const FPCGPoint& Point = Points[Partition[ReadIndex]];
+					if (Point.Transform.GetDeterminant() < 0)
+					{
+						ReverseInstances.Add(Partition[ReadIndex]);
+					}
+					else
+					{
+						if (WriteIndex != ReadIndex)
+						{
+							Partition[WriteIndex] = Partition[ReadIndex];
+						}
+						++WriteIndex;
+					}
+				}
+
+				Partition.SetNum(WriteIndex);
+			}
+
+			auto AddPointsToInstanceList = [&OutMeshInstances, &CurrentPartitionDescriptor, &Context, &MaterialOverrideHelper, &Points, PartitionIndex, InPointData](const TArray<int32>& PointIndices, bool bReverseTransform)
+			{
+				if (MaterialOverrideHelper.OverridesMaterials())
+				{
+					for (int32 PointIndex = 0; PointIndex < PointIndices.Num(); ++PointIndex)
+					{
+						const int32 InstanceToPoint = PointIndices[PointIndex];
+
+						const FPCGPoint& Point = Points[InstanceToPoint];
+						FPCGSkinnedMeshInstanceList& InstanceList = PCGSkinnedMeshSelector::GetInstanceList(
+							OutMeshInstances,
+							CurrentPartitionDescriptor,
+							CurrentPartitionDescriptor.SkinnedAsset,
+							MaterialOverrideHelper.GetMaterialOverrides(Point.MetadataEntry),
+							bReverseTransform,
+							InPointData,
+							PartitionIndex
+						);
+
+						FSoftAnimBankItem BankItem;
+						{
+							BankItem.BankAsset = Context.PointBankValues[InstanceToPoint];
+							BankItem.SequenceIndex = Context.PointSequenceIndexValues[InstanceToPoint];
+						}
+
+						FPCGSkinnedMeshInstance Instance;
+						Instance.BankIndex = InstanceList.Descriptor.GetOrAddBankIndex(BankItem);
+						Instance.Transform = Point.Transform;
+
+						InstanceList.Instances.Emplace(Instance);
+						InstanceList.InstancePointIndices.Emplace(PointIndex);
+					}
+				}
+				else
+				{
+					TArray<TSoftObjectPtr<UMaterialInterface>> DummyMaterialList;
+					FPCGSkinnedMeshInstanceList& InstanceList = PCGSkinnedMeshSelector::GetInstanceList(
+						OutMeshInstances,
+						CurrentPartitionDescriptor,
+						CurrentPartitionDescriptor.SkinnedAsset,
+						DummyMaterialList,
+						bReverseTransform,
+						InPointData,
+						PartitionIndex
+					);
+
+					check(InstanceList.Instances.Num() == InstanceList.InstancePointIndices.Num());
+					const int32 InstanceOffset = InstanceList.Instances.Num();
+					InstanceList.Instances.SetNum(InstanceOffset + PointIndices.Num());
+					InstanceList.InstancePointIndices.Append(PointIndices);
+
+					for (int32 PointIndex = 0; PointIndex < PointIndices.Num(); ++PointIndex)
+					{
+						FPCGSkinnedMeshInstance& Instance = InstanceList.Instances[InstanceOffset + PointIndex];
+
+						const int32 InstanceToPoint = PointIndices[PointIndex];
+
+						FSoftAnimBankItem BankItem;
+						{
+							BankItem.BankAsset = Context.PointBankValues[InstanceToPoint];
+							BankItem.SequenceIndex = Context.PointSequenceIndexValues[InstanceToPoint];
+						}
+
+						Instance.BankIndex = InstanceList.Descriptor.GetOrAddBankIndex(BankItem);
+						Instance.Transform = Points[InstanceToPoint].Transform;
+					}
+				}
+			};
+
+			{
+				TRACE_CPUPROFILER_EVENT_SCOPE(UPCGSkinnedMeshSelector::SelectEntries::AddPointsToInstanceList);
+				AddPointsToInstanceList(Partition, false);
+				AddPointsToInstanceList(ReverseInstances, true);
+			}
+
+			if (Context.ShouldStop())
+			{
+				break;
+			}
+		}
+	}
+
+	Context.CurrentPointIndex = CurrentPartitionIndex; // misnomer, but we're using the same context
+	return CurrentPartitionIndex == Context.AttributeOverridePartition.Num();
+}
+
+#undef LOCTEXT_NAMESPACE
diff --git a/Engine/Plugins/PCG/Source/PCG/Private/PCGManagedResource.cpp b/Engine/Plugins/PCG/Source/PCG/Private/PCGManagedResource.cpp
--- a/Engine/Plugins/PCG/Source/PCG/Private/PCGManagedResource.cpp
+++ b/Engine/Plugins/PCG/Source/PCG/Private/PCGManagedResource.cpp
@@ -8,7 +8,7 @@
 
 #include "Algo/AnyOf.h"
 #include "Components/InstancedStaticMeshComponent.h"
-#include "Components/AnimBankMeshComponent.h"
+#include "Components/InstancedSkinnedMeshComponent.h"
 #include "Components/SplineMeshComponent.h"
 #include "Engine/Level.h"
 #include "UObject/Package.h"
@@ -879,16 +879,16 @@ void UPCGManagedISMComponent::SetComponent(UInstancedStaticMeshComponent* InComp
 	CachedRawComponentPtr = InComponent;
 }
 
-void UPCGManagedAnimBankMeshComponent::PostLoad()
+void UPCGManagedISKMComponent::PostLoad()
 {
 	Super::PostLoad();
 
 	if (!bHasDescriptor)
 	{
-		if (UAnimBankMeshComponent* ABMC = GetComponent())
+		if (UInstancedSkinnedMeshComponent* ISKMC = GetComponent())
 		{
-			FAnimBankMeshComponentDescriptor NewDescriptor;
-			NewDescriptor.InitFrom(ABMC);
+			FSkinnedMeshComponentDescriptor NewDescriptor;
+			NewDescriptor.InitFrom(ISKMC);
 
 			SetDescriptor(NewDescriptor);
 		}
@@ -898,19 +898,19 @@ void UPCGManagedAnimBankMeshComponent::PostLoad()
 	GetComponent();
 }
 
-void UPCGManagedAnimBankMeshComponent::ForgetComponent()
+void UPCGManagedISKMComponent::ForgetComponent()
 {
 	Super::ForgetComponent();
 	CachedRawComponentPtr = nullptr;
 }
 
-void UPCGManagedAnimBankMeshComponent::SetDescriptor(const FAnimBankMeshComponentDescriptor& InDescriptor)
+void UPCGManagedISKMComponent::SetDescriptor(const FSkinnedMeshComponentDescriptor& InDescriptor)
 {
 	bHasDescriptor = true;
 	Descriptor = InDescriptor;
 }
 
-bool UPCGManagedAnimBankMeshComponent::ReleaseIfUnused(TSet<TSoftObjectPtr<AActor>>& OutActorsToDelete)
+bool UPCGManagedISKMComponent::ReleaseIfUnused(TSet<TSoftObjectPtr<AActor>>& OutActorsToDelete)
 {
 	if (Super::ReleaseIfUnused(OutActorsToDelete) || !GetComponent())
 	{
@@ -928,16 +928,16 @@ bool UPCGManagedAnimBankMeshComponent::ReleaseIfUnused(TSet<TSoftObjectPtr<AActo
 	}
 }
 
-void UPCGManagedAnimBankMeshComponent::ResetComponent()
+void UPCGManagedISKMComponent::ResetComponent()
 {
-	if (UAnimBankMeshComponent* ABMC = GetComponent())
+	if (UInstancedSkinnedMeshComponent* ISKMC = GetComponent())
 	{
-		ABMC->ClearInstances();
-		ABMC->UpdateBounds();
+		ISKMC->ClearInstances();
+		ISKMC->UpdateBounds();
 	}
 }
 
-void UPCGManagedAnimBankMeshComponent::MarkAsUsed()
+void UPCGManagedISKMComponent::MarkAsUsed()
 {
 	const bool bWasMarkedUnused = bIsMarkedUnused;
 	Super::MarkAsUsed();
@@ -947,12 +947,12 @@ void UPCGManagedAnimBankMeshComponent::MarkAsUsed()
 		return;
 	}
 
-	if (UAnimBankMeshComponent* ABMC = GetComponent())
+	if (UInstancedSkinnedMeshComponent* ISKMC = GetComponent())
 	{
 		const bool bHasPreviousRootLocation = bHasRootLocation;
 
 		// Keep track of the current root location so if we reuse this later we are able to update this appropriately
-		if (USceneComponent* RootComponent = ABMC->GetAttachmentRoot())
+		if (USceneComponent* RootComponent = ISKMC->GetAttachmentRoot())
 		{
 			bHasRootLocation = true;
 			RootLocation = RootComponent->GetComponentLocation();
@@ -963,55 +963,55 @@ void UPCGManagedAnimBankMeshComponent::MarkAsUsed()
 			RootLocation = FVector::ZeroVector;
 		}
 
-		if (bHasPreviousRootLocation != bHasRootLocation || (ABMC->GetComponentLocation() - RootLocation).SquaredLength() > UE_DOUBLE_SMALL_NUMBER)
+		if (bHasPreviousRootLocation != bHasRootLocation || (ISKMC->GetComponentLocation() - RootLocation).SquaredLength() > UE_DOUBLE_SMALL_NUMBER)
 		{
 			// Reset the rotation/scale to be identity otherwise if the root component transform has changed, the final transform will be wrong
 			// Since this is technically 'moving' the ABM, we need to unregister it before moving otherwise we could get a warning that we're moving a component with static mobility
-			ABMC->UnregisterComponent();
-			ABMC->SetWorldTransform(FTransform(FQuat::Identity, RootLocation, FVector::OneVector));
-			ABMC->RegisterComponent();
+			ISKMC->UnregisterComponent();
+			ISKMC->SetWorldTransform(FTransform(FQuat::Identity, RootLocation, FVector::OneVector));
+			ISKMC->RegisterComponent();
 		}
 	}
 }
 
-void UPCGManagedAnimBankMeshComponent::MarkAsReused()
+void UPCGManagedISKMComponent::MarkAsReused()
 {
 	Super::MarkAsReused();
 
-	if (UAnimBankMeshComponent* ABMC = GetComponent())
+	if (UInstancedSkinnedMeshComponent* ISKMC = GetComponent())
 	{
 		// Reset the rotation/scale to be identity otherwise if the root component transform has changed, the final transform will be wrong
 		FVector TentativeRootLocation = RootLocation;
 
 		if (!bHasRootLocation)
 		{
-			if (USceneComponent* RootComponent = ABMC->GetAttachmentRoot())
+			if (USceneComponent* RootComponent = ISKMC->GetAttachmentRoot())
 			{
 				TentativeRootLocation = RootComponent->GetComponentLocation();
 			}
 		}
 
-		if ((ABMC->GetComponentLocation() - TentativeRootLocation).SquaredLength() > UE_DOUBLE_SMALL_NUMBER)
+		if ((ISKMC->GetComponentLocation() - TentativeRootLocation).SquaredLength() > UE_DOUBLE_SMALL_NUMBER)
 		{
 			// Since this is technically 'moving' the ABM, we need to unregister it before moving otherwise we could get a warning that we're moving a component with static mobility
-			ABMC->UnregisterComponent();
-			ABMC->SetWorldTransform(FTransform(FQuat::Identity, TentativeRootLocation, FVector::OneVector));
-			ABMC->RegisterComponent();
+			ISKMC->UnregisterComponent();
+			ISKMC->SetWorldTransform(FTransform(FQuat::Identity, TentativeRootLocation, FVector::OneVector));
+			ISKMC->RegisterComponent();
 		}
 	}
 }
 
-void UPCGManagedAnimBankMeshComponent::SetRootLocation(const FVector& InRootLocation)
+void UPCGManagedISKMComponent::SetRootLocation(const FVector& InRootLocation)
 {
 	bHasRootLocation = true;
 	RootLocation = InRootLocation;
 }
 
-UAnimBankMeshComponent* UPCGManagedAnimBankMeshComponent::GetComponent() const
+UInstancedSkinnedMeshComponent* UPCGManagedISKMComponent::GetComponent() const
 {
 	if (!CachedRawComponentPtr)
 	{
-		UAnimBankMeshComponent* GeneratedComponentPtr = Cast<UAnimBankMeshComponent>(GeneratedComponent.Get());
+		UInstancedSkinnedMeshComponent* GeneratedComponentPtr = Cast<UInstancedSkinnedMeshComponent>(GeneratedComponent.Get());
 
 		// Implementation note:
 		// There is no surefire way to make sure that we can use the raw pointer UNLESS it is from the same owner
@@ -1026,7 +1026,7 @@ UAnimBankMeshComponent* UPCGManagedAnimBankMeshComponent::GetComponent() const
 	return CachedRawComponentPtr;
 }
 
-void UPCGManagedAnimBankMeshComponent::SetComponent(UAnimBankMeshComponent* InComponent)
+void UPCGManagedISKMComponent::SetComponent(UInstancedSkinnedMeshComponent* InComponent)
 {
 	GeneratedComponent = InComponent;
 	CachedRawComponentPtr = InComponent;
diff --git a/Engine/Plugins/PCG/Source/PCG/Public/Elements/PCGAnimBankMeshSpawner.h b/Engine/Plugins/PCG/Source/PCG/Public/Elements/PCGAnimBankMeshSpawner.h
deleted file mode 100644
diff --git a/Engine/Plugins/PCG/Source/PCG/Public/Elements/PCGAnimBankMeshSpawnerContext.h b/Engine/Plugins/PCG/Source/PCG/Public/Elements/PCGAnimBankMeshSpawnerContext.h
deleted file mode 100644
diff --git a/Engine/Plugins/PCG/Source/PCG/Public/Elements/PCGSkinnedMeshSpawner.h b/Engine/Plugins/PCG/Source/PCG/Public/Elements/PCGSkinnedMeshSpawner.h
new file mode 100644
--- /dev/null
+++ b/Engine/Plugins/PCG/Source/PCG/Public/Elements/PCGSkinnedMeshSpawner.h
@@ -0,0 +1,130 @@
+// Copyright Epic Games, Inc. All Rights Reserved.
+
+#pragma once
+
+#include "PCGSettings.h"
+#include "Metadata/PCGObjectPropertyOverride.h"
+#include "Engine/CollisionProfile.h"
+#include "PCGSkinnedMeshSpawner.generated.h"
+
+class UPCGSkinnedMeshInstanceDataPackerBase;
+class UPCGSkinnedMeshSelector;
+class UPCGSpatialData;
+struct FPCGContext;
+struct FPCGSkinnedMeshInstanceList;
+struct FPCGSkinnedMeshPackedCustomData;
+struct FPCGSkinnedMeshSpawnerContext;
+
+class UAnimBank;
+
+UCLASS(BlueprintType, ClassGroup = (Procedural), meta = (PrioritizeCategories = "Settings"))
+class PCG_API UPCGSkinnedMeshSpawnerSettings : public UPCGSettings
+{
+	GENERATED_BODY()
+
+public:
+	UPCGSkinnedMeshSpawnerSettings(const FObjectInitializer& ObjectInitializer);
+
+	// ~Begin UPCGSettings interface
+	virtual bool IsKernelValid(FPCGContext* InContext = nullptr, bool bQuiet = true) const;
+	virtual FString GetCookedKernelSource(const FPCGKernelAttributeTable& InAttributeStaticTable) const override;
+	virtual void GetKernelAttributeKeys(TArray<FPCGKernelAttributeKey>& OutKeys) const override;
+	virtual int ComputeKernelThreadCount(const UPCGDataBinding* Binding) const override;
+	virtual bool ComputeOutputPinDataDesc(const UPCGPin* OutputPin, const UPCGDataBinding* Binding, FPCGDataCollectionDesc& OutDesc) const override;
+#if WITH_EDITOR
+	virtual void CreateAdditionalInputDataInterfaces(FPCGGPUCompilationContext& InOutContext, UObject* InObjectOuter, TArray<TObjectPtr<UComputeDataInterface>>& OutDataInterfaces) const override;
+	virtual void CreateAdditionalOutputDataInterfaces(FPCGGPUCompilationContext& InOutContext, UObject* InObjectOuter, TArray<TObjectPtr<UComputeDataInterface>>& OutDataInterfaces) const override;
+	virtual bool DisplayExecuteOnGPUSetting() const override { return true; }
+	virtual FName GetDefaultNodeName() const override { return FName(TEXT("SkinnedMeshSpawner")); }
+	virtual FText GetDefaultNodeTitle() const override;
+	virtual EPCGSettingsType GetType() const override { return EPCGSettingsType::Spawner; }
+	virtual void ApplyDeprecation(UPCGNode* InOutNode) override;
+#endif
+	virtual bool UseSeed() const override { return true; }
+
+protected:
+	virtual TArray<FPCGPinProperties> InputPinProperties() const override;
+	virtual TArray<FPCGPinProperties> OutputPinProperties() const override { return Super::DefaultPointOutputPinProperties(); }
+	virtual FPCGElementPtr CreateElement() const override;
+	// ~End UPCGSettings interface
+
+public:
+	// ~Begin UObject interface
+	virtual void PostLoad() override;
+#if WITH_EDITOR
+	virtual void PostEditChangeProperty(FPropertyChangedEvent& PropertyChangedEvent) override;
+	virtual bool CanEditChange(const FProperty* InProperty) const override;
+	// ~End UObject interface
+#endif
+
+	UFUNCTION(BlueprintCallable, Category = Settings)
+	void SetInstancePackerType(TSubclassOf<UPCGSkinnedMeshInstanceDataPackerBase> InInstancePackerType);
+
+public:
+	UPROPERTY(BlueprintReadOnly, VisibleAnywhere, Instanced, Category = MeshSelector)
+	TObjectPtr<UPCGSkinnedMeshSelector> MeshSelectorParameters;
+
+	/**
+	 * Defines the method of custom data packing for spawned anim bank meshes. Note, Rotators are treated as 3 floats, while Quaternions are
+	 * treated as 4 floats. You can see an attribute's type in the 'Attribute List View' window, and use an 'Attribute Cast' node to cast to the desired type.
+	 */
+	UPROPERTY(BlueprintReadOnly, EditAnywhere, Category = InstanceDataPacker)
+	TSubclassOf<UPCGSkinnedMeshInstanceDataPackerBase> InstanceDataPackerType;
+
+	UPROPERTY(BlueprintReadOnly, VisibleAnywhere, Instanced, Category = InstanceDataPacker)
+	TObjectPtr<UPCGSkinnedMeshInstanceDataPackerBase> InstanceDataPackerParameters = nullptr;
+
+	/** Map an attribute directly to an anim bank descriptor property, the value of which will be overriden when generated.
+	 * Note: Currently only enabled using SelectByAttribute mesh selection.
+	 */
+	UPROPERTY(BlueprintReadWrite, EditAnywhere, Category = Settings)
+	TArray<FPCGObjectPropertyOverrideDescription> SkinnedMeshComponentPropertyOverrides;
+
+	/** Sets the BoundsMin and BoundsMax attributes of each point to reflect the AnimBank mesh spawned at its location */
+	UPROPERTY(BlueprintReadWrite, EditAnywhere, Category = Settings)
+	bool bApplyMeshBoundsToPoints = true;
+
+	UPROPERTY(BlueprintReadWrite, Category = Settings, meta = (PCG_Overridable))
+	TSoftObjectPtr<AActor> TargetActor;
+
+	/** Specify a list of functions to be called on the target actor after instances are spawned. Functions need to be parameter-less and with "CallInEditor" flag enabled. */
+	UPROPERTY(BlueprintReadWrite, EditAnywhere, Category = Settings)
+	TArray<FName> PostProcessFunctionNames;
+
+	/** Meshes/Materials will be synchronously loaded before spawning instead of asynchronously. */
+	UPROPERTY(BlueprintReadWrite, EditAnywhere, Category = "Settings|Debug")
+	bool bSynchronousLoad = false;
+
+	/** Opt-in option to silence errors when the property override attributes are not found. */
+	UPROPERTY(BlueprintReadWrite, EditAnywhere, Category = "Settings|Advanced")
+	bool bSilenceOverrideAttributeNotFoundErrors = false;
+
+	/** Adds a warning to the node on repeated spawning with identical conditions (ie. same mesh descriptor at same spawn location, etc). */
+	UPROPERTY(BlueprintReadWrite, EditAnywhere, Category = "Settings|Advanced")
+	bool bWarnOnIdenticalSpawn = true;
+
+protected:
+	void RefreshInstancePacker();
+
+private:
+	static TCHAR const* TemplateFilePath;
+};
+
+class PCG_API FPCGSkinnedMeshSpawnerElement : public IPCGElement
+{
+public:
+	virtual bool CanExecuteOnlyOnMainThread(FPCGContext* Context) const override;
+	virtual bool IsCacheable(const UPCGSettings* InSettings) const override { return false; }
+
+protected:
+	virtual FPCGContext* CreateContext() override;
+	virtual bool PrepareDataInternal(FPCGContext* Context) const override;
+	virtual bool ExecuteInternal(FPCGContext* Context) const override;
+	virtual void AbortInternal(FPCGContext* Context) const override;
+	
+	void SpawnSkinnedMeshInstances(
+		FPCGSkinnedMeshSpawnerContext* Context,
+		const FPCGSkinnedMeshInstanceList& InstanceList,
+		AActor* TargetActor,
+		const FPCGSkinnedMeshPackedCustomData* InPackedCustomData) const;
+};
diff --git a/Engine/Plugins/PCG/Source/PCG/Public/Elements/PCGSkinnedMeshSpawnerContext.h b/Engine/Plugins/PCG/Source/PCG/Public/Elements/PCGSkinnedMeshSpawnerContext.h
new file mode 100644
--- /dev/null
+++ b/Engine/Plugins/PCG/Source/PCG/Public/Elements/PCGSkinnedMeshSpawnerContext.h
@@ -0,0 +1,63 @@
+// Copyright Epic Games, Inc. All Rights Reserved.
+
+#pragma once
+
+#include "PCGContext.h"
+#include "PCGManagedResource.h"
+#include "Async/PCGAsyncLoadingContext.h"
+#include "MeshSelectors/PCGSkinnedMeshSelector.h"
+#include "InstanceDataPackers/PCGSkinnedMeshInstanceDataPackerBase.h"
+
+class AActor;
+class UPCGSpatialData;
+
+struct FPCGSkinnedMeshSpawnerContext : public FPCGContext, public IPCGAsyncLoadingContext
+{
+	struct FPackedInstanceListData
+	{
+		FPackedInstanceListData();
+		~FPackedInstanceListData();
+		AActor* TargetActor = nullptr;
+		const UPCGSpatialData* SpatialData = nullptr;
+		TArray<FPCGSkinnedMeshInstanceList> MeshInstances;
+		TArray<FPCGSkinnedMeshPackedCustomData> PackedCustomData;
+	};
+
+	TArray<FPackedInstanceListData> MeshInstancesData;
+	// Index of input in the context (for selection)
+	int32 CurrentInputIndex = 0;
+	// Index of MeshInstances/PackedCustomData in last MeshInstancesData element
+	int32 CurrentDataIndex = 0;
+
+	// Whole-execution variables
+	bool bReuseCheckDone = false;
+	bool bSkippedDueToReuse = false;
+
+	// Per-input context variables
+	bool bCurrentInputSetup = false;
+	bool bSelectionDone = false;
+	bool bPartitionDone = false;
+
+	const UPCGPointData* CurrentPointData = nullptr;
+	UPCGPointData* CurrentOutputPointData = nullptr;
+	FPCGMeshMaterialOverrideHelper MaterialOverrideHelper;
+	int32 CurrentPointIndex = 0;
+
+	// Used in all selectors if we have to change the out points bounds by the mesh bounds. Will be empty otherwise.
+	// We need to keep all point indices that will spawn this mesh, in all output point data.
+	TMap<TSoftObjectPtr<USkinnedAsset>, TMap<UPCGPointData*, TArray<int32>>> MeshToOutPoints;
+
+	TMap<TSoftObjectPtr<USkinnedAsset>, PCGMetadataValueKey> MeshToValueKey;
+
+	// Used for mesh property overrides
+	TArray<TArray<int32>> AttributeOverridePartition;
+	TArray<FPCGSoftSkinnedMeshComponentDescriptor> OverriddenDescriptors;
+
+	// Keeping track of all touched resources to allow for correct cleanup on abort
+	TArray<TWeakObjectPtr<UPCGManagedISKMComponent>> TouchedResources;
+
+	TArray<FSoftObjectPath> PointBankValues;
+	TArray<int32> PointSequenceIndexValues;
+
+	void ResetInputIterationData();
+};
\ No newline at end of file
diff --git a/Engine/Plugins/PCG/Source/PCG/Public/Helpers/PCGActorHelpers.h b/Engine/Plugins/PCG/Source/PCG/Public/Helpers/PCGActorHelpers.h
--- a/Engine/Plugins/PCG/Source/PCG/Public/Helpers/PCGActorHelpers.h
+++ b/Engine/Plugins/PCG/Source/PCG/Public/Helpers/PCGActorHelpers.h
@@ -5,7 +5,7 @@
 #include "PCGCrc.h"
 #include "Elements/PCGSplineMeshParams.h"
 #include "MeshSelectors/PCGISMDescriptor.h"
-#include "MeshSelectors/PCGAnimBankMeshDescriptor.h"
+#include "MeshSelectors/PCGSkinnedMeshDescriptor.h"
 #include "Engine/EngineTypes.h"
 #include "Engine/SplineMeshComponentDescriptor.h"
 #include "Engine/World.h"
@@ -21,12 +21,14 @@ struct FPCGContext;
 class UActorComponent;
 class UDataLayerInstance;
 class UInstancedStaticMeshComponent;
+class UInstancedSkinnedMeshComponent;
 class ULevel;
 class UMaterialInterface;
 class USplineMeshComponent;
 class UPCGComponent;
 class UPCGManagedISMComponent;
-class UPCGManagedAnimBankMeshComponent;
+class UPCGManagedISKMComponent;
+class UPCGManagedSkinnedMeshComponent;
 class UPCGManagedSplineMeshComponent;
 class UStaticMesh;
 class UWorld;
@@ -70,21 +72,21 @@ PRAGMA_ENABLE_DEPRECATION_WARNINGS
 	inline bool operator==(const FPCGISMComponentBuilderParams& Other) const { return Descriptor == Other.Descriptor && NumCustomDataFloats == Other.NumCustomDataFloats && bAllowDescriptorChanges == Other.bAllowDescriptorChanges; }
 };
 
-struct FPCGAnimBankMeshComponentBuilderParams
+struct FPCGSkinnedMeshComponentBuilderParams
 {
-	FPCGAnimBankMeshComponentBuilderParams() = default;
+	FPCGSkinnedMeshComponentBuilderParams() = default;
 
-	FPCGSoftAnimBankMeshComponentDescriptor Descriptor;
+	FPCGSoftSkinnedMeshComponentDescriptor Descriptor;
 	int32 NumCustomDataFloats = 0;
 	FPCGCrc SettingsCrc;
 	bool bTransient = false;
 
-	friend inline uint32 GetTypeHash(const FPCGAnimBankMeshComponentBuilderParams& Key)
+	friend inline uint32 GetTypeHash(const FPCGSkinnedMeshComponentBuilderParams& Key)
 	{
 		return HashCombine(GetTypeHash(Key.Descriptor), 1 + Key.NumCustomDataFloats);
 	}
 
-	inline bool operator==(const FPCGAnimBankMeshComponentBuilderParams& Other) const { return Descriptor == Other.Descriptor && NumCustomDataFloats == Other.NumCustomDataFloats; }
+	inline bool operator==(const FPCGSkinnedMeshComponentBuilderParams& Other) const { return Descriptor == Other.Descriptor && NumCustomDataFloats == Other.NumCustomDataFloats; }
 };
 
 struct FPCGSplineMeshComponentBuilderParameters
@@ -123,7 +125,7 @@ PRAGMA_ENABLE_DEPRECATION_WARNINGS
 	
 	static UInstancedStaticMeshComponent* GetOrCreateISMC(AActor* InTargetActor, UPCGComponent* SourceComponent, const FPCGISMComponentBuilderParams& Params, FPCGContext* OptionalContext = nullptr);
 	static UPCGManagedISMComponent* GetOrCreateManagedISMC(AActor* InTargetActor, UPCGComponent* SourceComponent, const FPCGISMComponentBuilderParams& Params, FPCGContext* OptionalContext = nullptr);
-	static UPCGManagedAnimBankMeshComponent* GetOrCreateManagedABMC(AActor* InTargetActor, UPCGComponent* SourceComponent, const FPCGAnimBankMeshComponentBuilderParams& Params, FPCGContext* OptionalContext = nullptr);
+	static UPCGManagedISKMComponent* GetOrCreateManagedABMC(AActor* InTargetActor, UPCGComponent* SourceComponent, const FPCGSkinnedMeshComponentBuilderParams& Params, FPCGContext* OptionalContext = nullptr);
 	static USplineMeshComponent* GetOrCreateSplineMeshComponent(AActor* InTargetActor, UPCGComponent* SourceComponent, const FPCGSplineMeshComponentBuilderParameters& Params, FPCGContext* OptionalContext = nullptr);
 	static UPCGManagedSplineMeshComponent* GetOrCreateManagedSplineMeshComponent(AActor* InTargetActor, UPCGComponent* SourceComponent, const FPCGSplineMeshComponentBuilderParameters& Params, FPCGContext* OptionalContext = nullptr);
 
diff --git a/Engine/Plugins/PCG/Source/PCG/Public/InstanceDataPackers/PCGAnimBankMeshInstanceDataPackerBase.h b/Engine/Plugins/PCG/Source/PCG/Public/InstanceDataPackers/PCGAnimBankMeshInstanceDataPackerBase.h
deleted file mode 100644
diff --git a/Engine/Plugins/PCG/Source/PCG/Public/InstanceDataPackers/PCGAnimBankMeshInstanceDataPackerByAttribute.h b/Engine/Plugins/PCG/Source/PCG/Public/InstanceDataPackers/PCGAnimBankMeshInstanceDataPackerByAttribute.h
deleted file mode 100644
diff --git a/Engine/Plugins/PCG/Source/PCG/Public/InstanceDataPackers/PCGAnimBankMeshInstanceDataPackerByRegex.h b/Engine/Plugins/PCG/Source/PCG/Public/InstanceDataPackers/PCGAnimBankMeshInstanceDataPackerByRegex.h
deleted file mode 100644
diff --git a/Engine/Plugins/PCG/Source/PCG/Public/InstanceDataPackers/PCGSkinnedMeshInstanceDataPackerBase.h b/Engine/Plugins/PCG/Source/PCG/Public/InstanceDataPackers/PCGSkinnedMeshInstanceDataPackerBase.h
new file mode 100644
--- /dev/null
+++ b/Engine/Plugins/PCG/Source/PCG/Public/InstanceDataPackers/PCGSkinnedMeshInstanceDataPackerBase.h
@@ -0,0 +1,58 @@
+// Copyright Epic Games, Inc. All Rights Reserved.
+
+#pragma once
+
+#include "PCGSkinnedMeshInstanceDataPackerBase.generated.h"
+
+class FPCGMetadataAttributeBase;
+class IPCGAttributeAccessor;
+class IPCGAttributeAccessorKeys;
+class UPCGMetadata;
+class UPCGSpatialData;
+struct FPCGContext;
+struct FPCGSkinnedMeshInstanceList;
+
+USTRUCT(BlueprintType)
+struct FPCGSkinnedMeshPackedCustomData
+{
+	GENERATED_BODY()
+
+	UPROPERTY(BlueprintReadWrite, EditAnywhere, Category = Settings)
+	int SequenceIndex = 0;
+
+	UPROPERTY(BlueprintReadWrite, EditAnywhere, Category = Settings)
+	int NumCustomDataFloats = 0;
+
+	UPROPERTY(BlueprintReadWrite, EditAnywhere, Category = Settings)
+	TArray<float> CustomData;
+};
+
+UCLASS(Abstract, BlueprintType, Blueprintable, ClassGroup = (Procedural))
+class PCG_API UPCGSkinnedMeshInstanceDataPackerBase : public UObject 
+{
+	GENERATED_BODY()
+
+public:
+	/** Defines the strategy for custom float data packing */
+	UFUNCTION(BlueprintNativeEvent, Category = InstancePacking)
+	void PackInstances(FPCGContext& Context, const UPCGSpatialData* InSpatialData, const FPCGSkinnedMeshInstanceList& InstanceList, FPCGSkinnedMeshPackedCustomData& OutPackedCustomData) const;
+
+	virtual void PackInstances_Implementation(FPCGContext& Context, const UPCGSpatialData* InSpatialData, const FPCGSkinnedMeshInstanceList& InstanceList, FPCGSkinnedMeshPackedCustomData& OutPackedCustomData) const;
+
+	/** Interprets Metadata TypeId and increments OutPackedCustomData.NumCustomDataFloats appropriately. Returns false if the type could not be interpreted. */
+	UFUNCTION(BlueprintCallable, Category = InstancePacking)
+	bool AddTypeToPacking(int TypeId, FPCGSkinnedMeshPackedCustomData& OutPackedCustomData) const;
+
+	/** Build a PackedCustomData by processing each attribute in order for each point in the InstanceList */
+	UFUNCTION(BlueprintCallable, Category = InstancePacking) 
+	void PackCustomDataFromAttributes(const FPCGSkinnedMeshInstanceList& InstanceList, const UPCGMetadata* Metadata, const TArray<FName>& AttributeNames, FPCGSkinnedMeshPackedCustomData& OutPackedCustomData) const;
+
+	/** Build a PackedCustomData by processing each attribute in order for each point in the InstanceList */
+	void PackCustomDataFromAttributes(const FPCGSkinnedMeshInstanceList& InstanceList, const TArray<const FPCGMetadataAttributeBase*>& Attributes, FPCGSkinnedMeshPackedCustomData& OutPackedCustomData) const;
+
+	/** Build a PackedCustomData by processing each accessor in order for each point in the InstanceList */
+	void PackCustomDataFromAccessors(const FPCGSkinnedMeshInstanceList& InstanceList, TArray<TUniquePtr<const IPCGAttributeAccessor>> Accessors, TArray<TUniquePtr<const IPCGAttributeAccessorKeys>> AccessorKeys, FPCGSkinnedMeshPackedCustomData& OutPackedCustomData) const;
+
+	/** If OutNames is not null, returns a list of all attributes that will be packed. Returns true if this list can be statically determined (prior to execution). */
+	virtual bool GetAttributeNames(TArray<FName>* OutNames) { return false; }
+};
diff --git a/Engine/Plugins/PCG/Source/PCG/Public/InstanceDataPackers/PCGSkinnedMeshInstanceDataPackerByAttribute.h b/Engine/Plugins/PCG/Source/PCG/Public/InstanceDataPackers/PCGSkinnedMeshInstanceDataPackerByAttribute.h
new file mode 100644
--- /dev/null
+++ b/Engine/Plugins/PCG/Source/PCG/Public/InstanceDataPackers/PCGSkinnedMeshInstanceDataPackerByAttribute.h
@@ -0,0 +1,27 @@
+// Copyright Epic Games, Inc. All Rights Reserved.
+
+#pragma once
+
+#include "PCGSkinnedMeshInstanceDataPackerBase.h"
+#include "Metadata/PCGAttributePropertySelector.h"
+
+#include "PCGSkinnedMeshInstanceDataPackerByAttribute.generated.h"
+
+UCLASS(MinimalAPI, BlueprintType, ClassGroup = (Procedural))
+class UPCGSkinnedMeshInstanceDataPackerByAttribute : public UPCGSkinnedMeshInstanceDataPackerBase
+{
+	GENERATED_BODY()
+
+public:
+#if WITH_EDITOR
+	// ~Begin UObject interface
+	virtual void PostLoad() override;
+	// ~End UObject interface
+#endif
+	virtual void PackInstances_Implementation(UPARAM(ref) FPCGContext& Context, const UPCGSpatialData* InSpatialData, UPARAM(ref) const FPCGSkinnedMeshInstanceList& InstanceList, FPCGSkinnedMeshPackedCustomData& OutPackedCustomData) const override;
+	virtual bool GetAttributeNames(TArray<FName>* OutNames) override;
+
+public:
+	UPROPERTY(BlueprintReadWrite, EditAnywhere, Category = InstanceDataPacker)
+	TArray<FPCGAttributePropertyInputSelector> AttributeSelectors;
+};
diff --git a/Engine/Plugins/PCG/Source/PCG/Public/InstanceDataPackers/PCGSkinnedMeshInstanceDataPackerByRegex.h b/Engine/Plugins/PCG/Source/PCG/Public/InstanceDataPackers/PCGSkinnedMeshInstanceDataPackerByRegex.h
new file mode 100644
--- /dev/null
+++ b/Engine/Plugins/PCG/Source/PCG/Public/InstanceDataPackers/PCGSkinnedMeshInstanceDataPackerByRegex.h
@@ -0,0 +1,19 @@
+// Copyright Epic Games, Inc. All Rights Reserved.
+
+#pragma once
+
+#include "PCGSkinnedMeshInstanceDataPackerBase.h"
+#include "PCGSkinnedMeshInstanceDataPackerByRegex.generated.h"
+
+UCLASS(MinimalAPI, BlueprintType, ClassGroup = (Procedural))
+class UPCGSkinnedMeshInstanceDataPackerByRegex : public UPCGSkinnedMeshInstanceDataPackerBase
+{
+	GENERATED_BODY()
+
+public:
+	virtual void PackInstances_Implementation(UPARAM(ref) FPCGContext& Context, const UPCGSpatialData* InSpatialData, UPARAM(ref) const FPCGSkinnedMeshInstanceList& InstanceList, FPCGSkinnedMeshPackedCustomData& OutPackedCustomData) const override;
+
+public:
+	UPROPERTY(BlueprintReadWrite, EditAnywhere, Category = InstanceDataPacker)
+	TArray<FString> RegexPatterns; 
+};
diff --git a/Engine/Plugins/PCG/Source/PCG/Public/MeshSelectors/PCGAnimBankMeshDescriptor.h b/Engine/Plugins/PCG/Source/PCG/Public/MeshSelectors/PCGAnimBankMeshDescriptor.h
deleted file mode 100644
diff --git a/Engine/Plugins/PCG/Source/PCG/Public/MeshSelectors/PCGAnimBankMeshSelector.h b/Engine/Plugins/PCG/Source/PCG/Public/MeshSelectors/PCGAnimBankMeshSelector.h
deleted file mode 100644
diff --git a/Engine/Plugins/PCG/Source/PCG/Public/MeshSelectors/PCGSkinnedMeshDescriptor.h b/Engine/Plugins/PCG/Source/PCG/Public/MeshSelectors/PCGSkinnedMeshDescriptor.h
new file mode 100644
--- /dev/null
+++ b/Engine/Plugins/PCG/Source/PCG/Public/MeshSelectors/PCGSkinnedMeshDescriptor.h
@@ -0,0 +1,43 @@
+// Copyright Epic Games, Inc. All Rights Reserved.
+
+#pragma once
+
+#include "Engine/CollisionProfile.h"
+#include "Engine/DataTable.h"
+#include "Animation/AnimBank.h"
+
+#include "PCGSkinnedMeshDescriptor.generated.h"
+
+USTRUCT()
+struct FPCGAnimBankDataRow : public FTableRowBase
+{
+	GENERATED_BODY()
+
+	UPROPERTY(EditAnywhere, Category=AnimBank)
+	TObjectPtr<USkeletalMesh> Mesh = nullptr;
+
+	UPROPERTY(EditAnywhere, Category=AnimBank)
+	TObjectPtr<UAnimBank> Bank = nullptr;
+
+	UPROPERTY(EditAnywhere, Category=AnimBank)
+	uint32 SequenceIndex = 0;
+};
+
+/** Convenience PCG-side component descriptor so we can adjust defaults to the most common use cases. */
+// Implementation note: the tags don't really need to contribute to the hash, so we will retain the base class !=, == and ComputeHash implementations
+USTRUCT()
+struct FPCGSoftSkinnedMeshComponentDescriptor : public FSoftSkinnedMeshComponentDescriptor
+{
+	GENERATED_BODY()
+
+	PCG_API FPCGSoftSkinnedMeshComponentDescriptor();
+	PCG_API explicit FPCGSoftSkinnedMeshComponentDescriptor(const FSkinnedMeshComponentDescriptor& Other);
+	PCG_API virtual void InitFrom(const UInstancedSkinnedMeshComponent* Component, bool bInitBodyInstance = true) override;
+	PCG_API virtual void InitComponent(UInstancedSkinnedMeshComponent* Component) const override;
+
+	PCG_API int32 GetOrAddBankIndex(const FSoftAnimBankItem& BankItem);
+
+public:
+	UPROPERTY(EditAnywhere, Category="Component Settings|Tags", meta = (DisplayAfter = "ComponentClass"))
+	TArray<FName> ComponentTags;
+};
diff --git a/Engine/Plugins/PCG/Source/PCG/Public/MeshSelectors/PCGSkinnedMeshSelector.h b/Engine/Plugins/PCG/Source/PCG/Public/MeshSelectors/PCGSkinnedMeshSelector.h
new file mode 100644
--- /dev/null
+++ b/Engine/Plugins/PCG/Source/PCG/Public/MeshSelectors/PCGSkinnedMeshSelector.h
@@ -0,0 +1,97 @@
+// Copyright Epic Games, Inc. All Rights Reserved.
+
+#pragma once
+
+#include "PCGPoint.h"
+#include "Data/PCGPointData.h"
+#include "Engine/CollisionProfile.h"
+#include "MeshSelectors/PCGSkinnedMeshDescriptor.h"
+#include "MeshSelectors/PCGMeshMaterialOverrideHelper.h"
+#include "Metadata/PCGMetadata.h"
+#include "PCGSkinnedMeshSelector.generated.h"
+
+class UPCGPointData;
+class UPCGSpatialData;
+class UAnimBank;
+struct FPCGContext;
+struct FPCGSkinnedMeshSpawnerContext;
+class UPCGSkinnedMeshSpawnerSettings;
+
+USTRUCT(BlueprintType)
+struct FPCGSkinnedMeshInstance
+{
+	GENERATED_BODY()
+
+	FTransform Transform;
+	int32 BankIndex = 0;
+};
+
+USTRUCT(BlueprintType)
+struct FPCGSkinnedMeshInstanceList
+{
+	GENERATED_BODY()
+
+	explicit FPCGSkinnedMeshInstanceList(const FPCGSoftSkinnedMeshComponentDescriptor& InDescriptor)
+	: Descriptor(InDescriptor)
+	, AttributePartitionIndex(INDEX_NONE)
+	{}
+	
+	FPCGSkinnedMeshInstanceList() = default;
+	~FPCGSkinnedMeshInstanceList() = default;
+	FPCGSkinnedMeshInstanceList(const FPCGSkinnedMeshInstanceList&) = default;
+	FPCGSkinnedMeshInstanceList(FPCGSkinnedMeshInstanceList&&) = default;
+	FPCGSkinnedMeshInstanceList& operator=(const FPCGSkinnedMeshInstanceList&) = default;
+	FPCGSkinnedMeshInstanceList& operator=(FPCGSkinnedMeshInstanceList&&) = default;
+
+	UPROPERTY(EditAnywhere, Category = Settings)
+	FPCGSoftSkinnedMeshComponentDescriptor Descriptor;
+
+	UPROPERTY(BlueprintReadWrite, EditAnywhere, Category = Settings)
+	TArray<FPCGSkinnedMeshInstance> Instances;
+
+	/** Tracks which partition the instance list belongs to. */
+	int64 AttributePartitionIndex;
+
+	UPROPERTY(BlueprintReadWrite, EditAnywhere, Category = Settings)
+	TWeakObjectPtr<const UPCGPointData> PointData;
+	
+	UPROPERTY(BlueprintReadWrite, EditAnywhere, Category = Settings)
+	TArray<int32> InstancePointIndices;
+};
+
+UCLASS(MinimalAPI, BlueprintType, ClassGroup = (Procedural))
+class UPCGSkinnedMeshSelector : public UObject 
+{
+	GENERATED_BODY()
+
+public:
+	// ~Begin UObject interface
+	virtual void PostLoad() override;
+	// ~End UObject interface
+
+	virtual bool SelectInstances(
+		FPCGSkinnedMeshSpawnerContext& Context,
+		const UPCGSkinnedMeshSpawnerSettings* Settings,
+		const UPCGPointData* InPointData,
+		TArray<FPCGSkinnedMeshInstanceList>& OutMeshInstances,
+		UPCGPointData* OutPointData) const;
+
+public:
+	UPROPERTY(BlueprintReadWrite, EditAnywhere, Category = MeshSelector)
+	FPCGAttributePropertyInputSelector BankAttribute;
+
+	UPROPERTY(BlueprintReadWrite, EditAnywhere, Category = MeshSelector)
+	FPCGAttributePropertyInputSelector SequenceIndexAttribute;
+
+	UPROPERTY(BlueprintReadWrite, EditAnywhere, Category = MeshSelector)
+	FPCGAttributePropertyInputSelector MeshAttribute;
+
+	UPROPERTY(EditAnywhere, Category = MeshSelector)
+	FPCGSoftSkinnedMeshComponentDescriptor TemplateDescriptor;
+
+	UPROPERTY(BlueprintReadWrite, EditAnywhere, Category = MeshSelector, meta = (InlineEditConditionToggle))
+	bool bUseAttributeMaterialOverrides = false;
+
+	UPROPERTY(BlueprintReadWrite, EditAnywhere, DisplayName = "By Attribute Material Overrides", Category = MeshSelector, meta = (EditCondition = "bUseAttributeMaterialOverrides"))
+	TArray<FName> MaterialOverrideAttributes;
+};
diff --git a/Engine/Plugins/PCG/Source/PCG/Public/PCGManagedResource.h b/Engine/Plugins/PCG/Source/PCG/Public/PCGManagedResource.h
--- a/Engine/Plugins/PCG/Source/PCG/Public/PCGManagedResource.h
+++ b/Engine/Plugins/PCG/Source/PCG/Public/PCGManagedResource.h
@@ -13,7 +13,7 @@
 #include "PCGManagedResource.generated.h"
 
 class UActorComponent;
-class UAnimBankMeshComponent;
+class UInstancedSkinnedMeshComponent;
 class UInstancedStaticMeshComponent;
 class USplineMeshComponent;
 
@@ -305,7 +305,7 @@ protected:
 };
 
 UCLASS(BlueprintType)
-class PCG_API UPCGManagedAnimBankMeshComponent : public UPCGManagedComponent
+class PCG_API UPCGManagedISKMComponent : public UPCGManagedComponent
 {
 	GENERATED_BODY()
 
@@ -326,11 +326,11 @@ public:
 	virtual void ForgetComponent() override;
 	//~End UPCGManagedComponents interface
 
-	UAnimBankMeshComponent* GetComponent() const;
-	void SetComponent(UAnimBankMeshComponent* InComponent);
+	UInstancedSkinnedMeshComponent* GetComponent() const;
+	void SetComponent(UInstancedSkinnedMeshComponent* InComponent);
 
-	void SetDescriptor(const FAnimBankMeshComponentDescriptor& InDescriptor);
-	const FAnimBankMeshComponentDescriptor& GetDescriptor() const { return Descriptor; }
+	void SetDescriptor(const FSkinnedMeshComponentDescriptor& InDescriptor);
+	const FSkinnedMeshComponentDescriptor& GetDescriptor() const { return Descriptor; }
 
 	void SetRootLocation(const FVector& InRootLocation);
 
@@ -342,7 +342,7 @@ protected:
 	bool bHasDescriptor = false;
 
 	UPROPERTY()
-	FAnimBankMeshComponentDescriptor Descriptor;
+	FSkinnedMeshComponentDescriptor Descriptor;
 
 	UPROPERTY()
 	bool bHasRootLocation = false;
@@ -353,8 +353,8 @@ protected:
 	UPROPERTY()
 	FPCGCrc SettingsCrc;
 
-	// Cached raw pointer to anim bank mesh component
-	mutable UAnimBankMeshComponent* CachedRawComponentPtr = nullptr;
+	// Cached raw pointer to ISKM component
+	mutable UInstancedSkinnedMeshComponent* CachedRawComponentPtr = nullptr;
 };
 
 UCLASS(BlueprintType)
diff --git a/Engine/Shaders/Private/Skinning/AnimBankMeshDebug.usf b/Engine/Shaders/Private/Skinning/AnimBankMeshDebug.usf
deleted file mode 100644
diff --git a/Engine/Shaders/Private/Skinning/SkinnedMeshDebug.usf b/Engine/Shaders/Private/Skinning/SkinnedMeshDebug.usf
new file mode 100644
--- /dev/null
+++ b/Engine/Shaders/Private/Skinning/SkinnedMeshDebug.usf
@@ -0,0 +1,100 @@
+// Copyright Epic Games, Inc. All Rights Reserved.
+
+#include "/Engine/Public/Platform.ush"
+#include "/Engine/Private/Common.ush"
+#include "/Engine/Private/SceneData.ush"
+#include "/Engine/Private/BoneTransform.ush"
+#include "/Engine/Private/Visualization.ush"
+#include "/Engine/Private/Nanite/NaniteDataDecode.ush"
+
+uint PersistentPrimitiveIndex;
+
+float3 GetBoneColor(uint BoneIndex, uint BoneDepth, uint BoneCount)
+{
+#if 0
+	float Scale = 20.0f;
+	return ColorMapTurbo((float(BoneDepth) / float(BoneCount)) * Scale);
+#else
+	return IntToColor(BoneIndex);
+#endif
+}
+
+struct FPackedBone
+{
+	uint BoneParent : 16;
+	uint BoneDepth  : 16;
+};
+
+FPackedBone GetPackedBone(FNaniteSkinningHeader SkinningHeader, uint BoneIndex)
+{
+	const uint BaseOffset = SkinningHeader.HierarchyBufferOffset * (uint)sizeof(uint);
+	const uint BoneOffset = BoneIndex * (uint)sizeof(uint);
+	return Scene.NaniteSkinning.BoneHierarchy.Load<FPackedBone>(BaseOffset + BoneOffset);
+}
+
+float4x4 MatrixFrom4x3(float4x3 M)
+{
+	return float4x4(
+		float4(M[0], 0.0f),
+		float4(M[1], 0.0f),
+		float4(M[2], 0.0f),
+		float4(M[3], 1.0f)
+	);
+}
+
+void LineVS(
+	in uint InstanceID		: SV_InstanceID,
+	in uint VertexID		: SV_VertexID,
+	out float4 OutPosition	: SV_POSITION,
+	out float4 OutColor		: TEXCOORD0
+)
+{
+	ResolvedView = ResolveView();
+
+	uint BoneIndex = VertexID >> 1u;
+
+	const FPrimitiveSceneData PrimitiveData = GetPrimitiveData(PersistentPrimitiveIndex);
+	const FInstanceSceneData InstanceData = GetInstanceSceneData(PrimitiveData.InstanceSceneDataOffset + InstanceID);
+	const FNaniteSkinningHeader SkinningHeader = LoadNaniteSkinningHeader(PersistentPrimitiveIndex);
+	const FPackedBone PackedBone = GetPackedBone(SkinningHeader, BoneIndex);
+
+	uint BoneDepth = PackedBone.BoneDepth;
+	if ((VertexID % 2u) == 0)
+	{
+		const uint ParentBoneIndex = PackedBone.BoneParent;
+		if (ParentBoneIndex < SkinningHeader.MaxTransformCount)
+		{
+			BoneIndex = ParentBoneIndex;
+			BoneDepth--;
+		}
+	}
+
+	const float4x4 BoneTransform = MatrixFrom4x3(LoadNaniteBoneTransform(SkinningHeader.TransformBufferOffset + InstanceData.SkinningData + BoneIndex));
+
+	float4x4 BoneObjectSpace;
+
+	BRANCH
+	if (SkinningHeader.bHasScale)
+	{
+		BoneObjectSpace = MatrixFromBoneTransformWithScale(LoadNaniteBoneObjectSpaceWithScale(SkinningHeader.ObjectSpaceBufferOffset, BoneIndex));
+	}
+	else
+	{
+		BoneObjectSpace = MatrixFromBoneTransform(LoadNaniteBoneObjectSpace(SkinningHeader.ObjectSpaceBufferOffset, BoneIndex));
+	}
+
+	const float3 LocalPosition = mul(BoneObjectSpace, BoneTransform)._m30_m31_m32;
+	const float4 TranslatedWorldPosition = float4(DFTransformLocalToTranslatedWorld(LocalPosition, InstanceData.LocalToWorld, ResolvedView.PreViewTranslation).xyz, 1.0f);
+
+	OutPosition = mul(TranslatedWorldPosition, ResolvedView.TranslatedWorldToClip);
+	OutColor = float4(GetBoneColor(VertexID >> 1u, BoneDepth, SkinningHeader.MaxTransformCount), 1.0f);
+}
+
+void LinePS(
+	 in float4 Position	: SV_POSITION,
+	 in float4 Color	: TEXCOORD0,
+	out float4 OutColor	: SV_Target0
+)
+{
+	OutColor = Color;
+}
\ No newline at end of file
diff --git a/Engine/Source/Runtime/Core/Public/UObject/UE5MainStreamObjectVersion.h b/Engine/Source/Runtime/Core/Public/UObject/UE5MainStreamObjectVersion.h
--- a/Engine/Source/Runtime/Core/Public/UObject/UE5MainStreamObjectVersion.h
+++ b/Engine/Source/Runtime/Core/Public/UObject/UE5MainStreamObjectVersion.h
@@ -374,7 +374,7 @@ struct FUE5MainStreamObjectVersion
 		VisualLoggerAddedSeparateWorldTime,
 
 		// Added support for InstanceDataManagerSerialization and changed format for the instances to FTransform3f
-		AnimBankInstanceDataSerializationV2,
+		SkinnedMeshInstanceDataSerializationV2,
 	
 		// Added new material compilation validation for runtime virtual textures
 		RuntimeVirtualTextureMaterialValidation,
diff --git a/Engine/Source/Runtime/Engine/Classes/Animation/AnimBank.h b/Engine/Source/Runtime/Engine/Classes/Animation/AnimBank.h
--- a/Engine/Source/Runtime/Engine/Classes/Animation/AnimBank.h
+++ b/Engine/Source/Runtime/Engine/Classes/Animation/AnimBank.h
@@ -22,7 +22,7 @@ struct FPropertyChangedEvent;
 
 struct FReferenceSkeleton;
 class USkinnedAsset;
-class UAnimBankMeshComponent;
+class UInstancedSkinnedMeshComponent;
 class FAnimBankBuildAsyncCacheTask;
 
 struct FSkinnedAssetMapping
@@ -347,28 +347,28 @@ inline uint32 GetTypeHash(const TArray<FSoftAnimBankItem>& InBankItems)
 }
 
 USTRUCT()
-struct FAnimBankMeshComponentDescriptor
+struct FSkinnedMeshComponentDescriptor
 {
 	GENERATED_BODY()
 
-	ENGINE_API FAnimBankMeshComponentDescriptor();
-	ENGINE_API explicit FAnimBankMeshComponentDescriptor(ENoInit);
-	ENGINE_API explicit FAnimBankMeshComponentDescriptor(const FAnimBankMeshComponentDescriptor&);
-	ENGINE_API explicit FAnimBankMeshComponentDescriptor(const FSoftAnimBankMeshComponentDescriptor&);
-	ENGINE_API virtual ~FAnimBankMeshComponentDescriptor();
+	ENGINE_API FSkinnedMeshComponentDescriptor();
+	ENGINE_API explicit FSkinnedMeshComponentDescriptor(ENoInit);
+	ENGINE_API explicit FSkinnedMeshComponentDescriptor(const FSkinnedMeshComponentDescriptor&);
+	ENGINE_API explicit FSkinnedMeshComponentDescriptor(const FSoftSkinnedMeshComponentDescriptor&);
+	ENGINE_API virtual ~FSkinnedMeshComponentDescriptor();
 
-	ENGINE_API UAnimBankMeshComponent* CreateComponent(UObject* Outer, FName Name = NAME_None, EObjectFlags ObjectFlags = EObjectFlags::RF_NoFlags) const;
+	ENGINE_API UInstancedSkinnedMeshComponent* CreateComponent(UObject* Outer, FName Name = NAME_None, EObjectFlags ObjectFlags = EObjectFlags::RF_NoFlags) const;
 
-	ENGINE_API virtual void InitFrom(const UAnimBankMeshComponent* Component, bool bInitBodyInstance = true);
+	ENGINE_API virtual void InitFrom(const UInstancedSkinnedMeshComponent* Component, bool bInitBodyInstance = true);
 	ENGINE_API virtual uint32 ComputeHash() const;
-	ENGINE_API virtual void InitComponent(UAnimBankMeshComponent* ISMComponent) const;
+	ENGINE_API virtual void InitComponent(UInstancedSkinnedMeshComponent* ISMComponent) const;
 
 	ENGINE_API void PostLoadFixup(UObject* Loader);
 
-	ENGINE_API bool operator!=(const FAnimBankMeshComponentDescriptor& Other) const;
-	ENGINE_API bool operator==(const FAnimBankMeshComponentDescriptor& Other) const;
+	ENGINE_API bool operator!=(const FSkinnedMeshComponentDescriptor& Other) const;
+	ENGINE_API bool operator==(const FSkinnedMeshComponentDescriptor& Other) const;
 
-	friend inline uint32 GetTypeHash(const FAnimBankMeshComponentDescriptor& Key)
+	friend inline uint32 GetTypeHash(const FSkinnedMeshComponentDescriptor& Key)
 	{
 		return Key.GetTypeHash();
 	}
@@ -390,7 +390,7 @@ public:
 	TEnumAsByte<EComponentMobility::Type> Mobility;
 
 	UPROPERTY(EditAnywhere, Category = "Component Settings")
-	TSubclassOf<UAnimBankMeshComponent> ComponentClass;
+	TSubclassOf<UInstancedSkinnedMeshComponent> ComponentClass;
 
 	UPROPERTY(EditAnywhere, Category = "Component Settings", meta = (DisplayAfter = "ComponentClass"))
 	TObjectPtr<class USkinnedAsset> SkinnedAsset;
@@ -400,28 +400,28 @@ public:
 };
 
 USTRUCT()
-struct FSoftAnimBankMeshComponentDescriptor
+struct FSoftSkinnedMeshComponentDescriptor
 {
 	GENERATED_BODY()
 
-	ENGINE_API FSoftAnimBankMeshComponentDescriptor();
-	ENGINE_API explicit FSoftAnimBankMeshComponentDescriptor(ENoInit);
-	ENGINE_API explicit FSoftAnimBankMeshComponentDescriptor(const FAnimBankMeshComponentDescriptor&);
-	ENGINE_API explicit FSoftAnimBankMeshComponentDescriptor(const FSoftAnimBankMeshComponentDescriptor&);
-	ENGINE_API virtual ~FSoftAnimBankMeshComponentDescriptor();
+	ENGINE_API FSoftSkinnedMeshComponentDescriptor();
+	ENGINE_API explicit FSoftSkinnedMeshComponentDescriptor(ENoInit);
+	ENGINE_API explicit FSoftSkinnedMeshComponentDescriptor(const FSkinnedMeshComponentDescriptor&);
+	ENGINE_API explicit FSoftSkinnedMeshComponentDescriptor(const FSoftSkinnedMeshComponentDescriptor&);
+	ENGINE_API virtual ~FSoftSkinnedMeshComponentDescriptor();
 
-	ENGINE_API UAnimBankMeshComponent* CreateComponent(UObject* Outer, FName Name = NAME_None, EObjectFlags ObjectFlags = EObjectFlags::RF_NoFlags) const;
+	ENGINE_API UInstancedSkinnedMeshComponent* CreateComponent(UObject* Outer, FName Name = NAME_None, EObjectFlags ObjectFlags = EObjectFlags::RF_NoFlags) const;
 
-	ENGINE_API virtual void InitFrom(const UAnimBankMeshComponent* Component, bool bInitBodyInstance = true);
+	ENGINE_API virtual void InitFrom(const UInstancedSkinnedMeshComponent* Component, bool bInitBodyInstance = true);
 	ENGINE_API virtual uint32 ComputeHash() const;
-	ENGINE_API virtual void InitComponent(UAnimBankMeshComponent* ISMComponent) const;
+	ENGINE_API virtual void InitComponent(UInstancedSkinnedMeshComponent* ISMComponent) const;
 
 	ENGINE_API void PostLoadFixup(UObject* Loader);
 
-	ENGINE_API bool operator!=(const FSoftAnimBankMeshComponentDescriptor& Other) const;
-	ENGINE_API bool operator==(const FSoftAnimBankMeshComponentDescriptor& Other) const;
+	ENGINE_API bool operator!=(const FSoftSkinnedMeshComponentDescriptor& Other) const;
+	ENGINE_API bool operator==(const FSoftSkinnedMeshComponentDescriptor& Other) const;
 
-	friend inline uint32 GetTypeHash(const FSoftAnimBankMeshComponentDescriptor& Key)
+	friend inline uint32 GetTypeHash(const FSoftSkinnedMeshComponentDescriptor& Key)
 	{
 		return Key.GetTypeHash();
 	}
@@ -443,7 +443,7 @@ public:
 	TEnumAsByte<EComponentMobility::Type> Mobility;
 
 	UPROPERTY(EditAnywhere, Category = "Component Settings")
-	TSubclassOf<UAnimBankMeshComponent> ComponentClass;
+	TSubclassOf<UInstancedSkinnedMeshComponent> ComponentClass;
 
 	UPROPERTY()
 	TSoftObjectPtr<class USkinnedAsset> SkinnedAsset;
diff --git a/Engine/Source/Runtime/Engine/Classes/Components/AnimBankMeshComponent.h b/Engine/Source/Runtime/Engine/Classes/Components/AnimBankMeshComponent.h
deleted file mode 100644
diff --git a/Engine/Source/Runtime/Engine/Classes/Components/InstancedSkinnedMeshComponent.h b/Engine/Source/Runtime/Engine/Classes/Components/InstancedSkinnedMeshComponent.h
new file mode 100644
--- /dev/null
+++ b/Engine/Source/Runtime/Engine/Classes/Components/InstancedSkinnedMeshComponent.h
@@ -0,0 +1,278 @@
+// Copyright Epic Games, Inc. All Rights Reserved.
+
+#pragma once
+
+#include "Components/StaticMeshComponent.h"
+#include "Components/SkeletalMeshComponent.h"
+#include "Rendering/NaniteInterface.h"
+#include "Animation/AnimBank.h"
+#include "InstanceDataSceneProxy.h"
+#include "InstanceData/InstanceDataManager.h"
+#include "InstancedSkinnedMeshComponent.generated.h"
+
+class FPrimitiveSceneProxy;
+class USkinnedMeshComponent;
+class UMaterialInterface;
+class UTexture;
+class UAnimBank;
+struct FAnimBankRecordHandle;
+
+USTRUCT()
+struct FSkinnedMeshInstanceData
+{
+	GENERATED_USTRUCT_BODY()
+
+	UPROPERTY(EditAnywhere, Category=Instances)
+	FTransform3f Transform;
+
+	UPROPERTY(EditAnywhere, Category = Animation)
+	uint32 BankIndex;
+
+	FSkinnedMeshInstanceData()
+	: Transform(FTransform3f::Identity)
+	, BankIndex(0)
+	{
+	}
+
+	FSkinnedMeshInstanceData(const FTransform3f& InTransform, uint32 InBankIndex)
+	: Transform(InTransform)
+	, BankIndex(InBankIndex)
+	{
+	}
+
+	friend FArchive& operator<<(FArchive& Ar, FSkinnedMeshInstanceData& InstanceData)
+	{
+		Ar << InstanceData.Transform;
+		Ar << InstanceData.BankIndex;
+		return Ar;
+	}
+};
+
+UCLASS(ClassGroup=Rendering, hidecategories=(Object,Activation,Collision,"Components|Activation",Physics), editinlinenew, meta=(BlueprintSpawnableComponent), Blueprintable, MinimalAPI)
+class UInstancedSkinnedMeshComponent : public USkinnedMeshComponent
+{
+	GENERATED_UCLASS_BODY()
+	
+protected:
+	/** Array of instances, bulk serialized. */
+	UPROPERTY(EditAnywhere, SkipSerialization, DisplayName="Instances", Category=Instances, meta=(MakeEditWidget=true, EditFixedOrder))
+	TArray<FSkinnedMeshInstanceData> InstanceData;
+
+	/** Defines the number of floats that will be available per instance for custom data */
+	UPROPERTY(EditAnywhere, Category=Instances, AdvancedDisplay)
+	int32 NumCustomDataFloats = 0;
+
+	/** Array of custom data for instances. This will contains NumCustomDataFloats*InstanceCount entries. The entries are represented sequantially, in instance order. Can be read in a material and manipulated through Blueprints.
+	*	Example: If NumCustomDataFloats is 1, then each entry will belong to an instance. Custom data 0 will belong to Instance 0. Custom data 1 will belong to Instance 1 etc.
+	*	Example: If NumCustomDataFloats is 2, then each pair of sequential entries belong to an instance. Custom data 0 and 1 will belong to Instance 0. Custom data 2 and 3 will belong to Instance 2 etc.
+	*/
+	UPROPERTY(EditAnywhere, EditFixedSize, SkipSerialization, DisplayName="Custom Data", Category=Instances, AdvancedDisplay, meta=(EditFixedOrder))
+	TArray<float> InstanceCustomData;
+
+	/** Screen space footprint (in primary view) cutoff which dictates the far distance the instance will play back animation. 
+	 *  Using 0.0 (the default) falls back to a global threshold, a negative value disables the cutoff.
+	 */
+	UPROPERTY(EditAnywhere, BlueprintReadOnly, Category = Animation)
+	float AnimationMinScreenSize = 0.0f;
+
+	/** If true, this component will avoid serializing its per instance data / those properties will also not be editable */
+	UPROPERTY()
+	uint8 bInheritPerInstanceData : 1;
+
+	UPROPERTY(EditAnywhere, BlueprintReadOnly, Category = Animation, meta = (ShowOnlyInnerProperties))
+	TArray<struct FAnimBankItem> AnimBankItems;
+
+public:
+
+	ENGINE_API UInstancedSkinnedMeshComponent(FVTableHelper& Helper);
+	virtual ~UInstancedSkinnedMeshComponent();
+
+	//~ Begin UObject Interface
+	virtual void Serialize(FArchive& Ar) override;
+	virtual void PostLoad() override;
+	virtual void BeginDestroy() override;
+#if WITH_EDITOR
+	virtual void PostEditChangeChainProperty(FPropertyChangedChainEvent& PropertyChangedEvent) override;
+	virtual void PostEditUndo() override;
+	virtual void BeginCacheForCookedPlatformData(const ITargetPlatform* TargetPlatform) override;
+	virtual bool IsCachedCookedPlatformDataLoaded( const ITargetPlatform* TargetPlatform ) override;
+#endif
+	//~ End UObject Interface
+
+	//~ Begin UActorComponent Interface
+	virtual void OnRegister() override;
+	virtual void SendRenderInstanceData_Concurrent() override;
+	virtual void TickComponent(float DeltaTime, enum ELevelTick TickType, FActorComponentTickFunction* ThisTickFunction) override;
+
+	virtual void CreateRenderState_Concurrent(FRegisterComponentContext* Context) override;
+	virtual void DestroyRenderState_Concurrent() override;
+
+	virtual TStructOnScope<FActorComponentInstanceData> GetComponentInstanceData() const override;
+	//~ End UActorComponent Interface
+
+	//~ Begin USceneComponent Interface
+	virtual FBoxSphereBounds CalcBounds(const FTransform& LocalToWorld) const override;
+	virtual FPrimitiveSceneProxy* CreateSceneProxy() override;
+	virtual void OnUpdateTransform(EUpdateTransformFlags UpdateTransformFlags, ETeleportType Teleport) override;
+	//~ End USceneComponent Interface
+
+	//~ Begin USkinnedMeshComponent Interface
+	virtual const Nanite::FResources* GetNaniteResources() const;
+	virtual void RefreshBoneTransforms(FActorComponentTickFunction* TickFunction = nullptr) override;
+#if WITH_EDITOR
+	/**
+	 * Called by the asset compilers (FAnimBankCompilingManager / FSkinnedAssetCompilingManager) when done to allow render state to be created.
+	 */
+	virtual void PostAssetCompilation();
+
+#endif
+	//~ End USkinnedMeshComponent Interface
+
+public:
+	/** Add an instance to this component. Transform is given in local space of this component unless bWorldSpace is set. */
+	UFUNCTION(BlueprintCallable, Category = "Components|InstancedSkinnedMesh")
+	ENGINE_API FPrimitiveInstanceId AddInstance(const FTransform& InstanceTransform, int32 BankIndex, bool bWorldSpace = false);
+
+	/** Add multiple instances to this component. Transform is given in local space of this component unless bWorldSpace is set. */
+	UFUNCTION(BlueprintCallable, Category = "Components|InstancedSkinnedMesh")
+	ENGINE_API TArray<FPrimitiveInstanceId> AddInstances(const TArray<FTransform>& Transforms, const TArray<int32>& BankIndices, bool bShouldReturnIds, bool bWorldSpace = false);
+
+	/** Update custom data for specific instance */
+	UFUNCTION(BlueprintCallable, Category = "Components|InstancedSkinnedMesh")
+	ENGINE_API bool SetCustomDataValue(FPrimitiveInstanceId InstanceId, int32 CustomDataIndex, float CustomDataValue);
+
+	/** Update all custom data values for specific instance, the size of the array view must match the NumCustomDataFloats */
+	ENGINE_API bool SetCustomData(FPrimitiveInstanceId InstanceId, TArrayView<const float> CustomDataFloats);
+
+	/** Update number of custom data entries per instance. This applies to all instances and will reallocate the full custom data buffer and reset all values to 0 */
+	UFUNCTION(BlueprintCallable, Category = "Components|InstancedSkinnedMesh")
+	ENGINE_API void SetNumCustomDataFloats(int32 InNumCustomDataFloats);
+
+	/** Get the transform for the instance specified. Instance is returned in local space of this component unless bWorldSpace is set.  Returns True on success. */
+	UFUNCTION(BlueprintCallable, Category = "Components|InstancedSkinnedMesh")
+	ENGINE_API bool GetInstanceTransform(FPrimitiveInstanceId InstanceId, FTransform& OutInstanceTransform, bool bWorldSpace = false) const;
+
+	/** Remove the instance specified. Returns True on success. */
+	UFUNCTION(BlueprintCallable, Category = "Components|InstancedSkinnedMesh")
+	ENGINE_API bool RemoveInstance(FPrimitiveInstanceId InstanceId);
+
+	/** Remove the instances specified. Returns True on success. */
+	UFUNCTION(BlueprintCallable, Category = "Components|InstancedSkinnedMesh")
+	ENGINE_API void RemoveInstances(const TArray<FPrimitiveInstanceId>& InstancesToRemove);
+
+	/** Clear all instances being rendered by this component. */
+	UFUNCTION(BlueprintCallable, Category = "Components|InstancedSkinnedMesh")
+	ENGINE_API void ClearInstances();
+
+	UFUNCTION(BlueprintCallable, Category = "Components|InstancedSkinnedMesh")
+	ENGINE_API int32 AddAnimBankItem(const FAnimBankItem& BankItem);
+
+	/**
+	 * Optimize the instance data by performing sorting according to spatial hash on the _source_ data.
+	 * Note that this reorders the instances and thus any indexing will change. By default resets the ID mapping to identity.
+	 *  @param bShouldRetainIdMap	If true, the id mapping is updated instead of reset to identity, this retains the validity of the IDs but adds some memory and storage cost (for the ID mapping).
+	 */
+	UFUNCTION(BlueprintCallable, Category = "Components|InstancedSkinnedMesh")
+	ENGINE_API void OptimizeInstanceData(bool bShouldRetainIdMap = false);
+
+	ENGINE_API bool IsEnabled() const;
+	ENGINE_API int32 GetInstanceCount() const;
+
+	ENGINE_API TArrayView<const struct FAnimBankItem> GetAnimBankItems() const;
+
+	ENGINE_API void SetAnimBankItems(TArrayView<const struct FAnimBankItem> InAnimBankItems);
+
+	int32 GetNumCustomDataFloats() const { return NumCustomDataFloats; }
+protected:
+
+	/** Handle changes that must happen before the proxy is recreated. */
+	void PreApplyComponentInstanceData(struct FInstancedSkinnedMeshComponentInstanceData* ComponentInstanceData);
+
+	/** Applies the cached component instance data to a newly blueprint constructed component. */
+	void ApplyComponentInstanceData(struct FInstancedSkinnedMeshComponentInstanceData* ComponentInstanceData);
+
+	FInstanceDataManagerSourceDataDesc GetComponentDesc(ERHIFeatureLevel::Type FeatureLevel);
+
+#if WITH_EDITOR
+	/** One bit per instance if the instance is selected. */
+	TBitArray<> SelectedInstances;
+#endif
+
+	bool bIsInstanceDataApplyCompleted = true;
+
+	void CreateHitProxyData(TArray<TRefCountPtr<HHitProxy>>& HitProxies);
+
+	/** 
+	 */
+	FPrimitiveInstanceId AddInstanceInternal(int32 InstanceIndex, const FTransform& InstanceTransform, int32 BankIndex, bool bWorldSpace);
+
+	/** 
+	 */
+	bool RemoveInstanceInternal(int32 InstanceIndex, bool InstanceAlreadyRemoved);
+
+private:
+	void ApplyInheritedPerInstanceData(const UInstancedSkinnedMeshComponent* InArchetype);
+	bool ShouldInheritPerInstanceData(const UInstancedSkinnedMeshComponent* InArchetype) const;
+	bool ShouldInheritPerInstanceData() const;
+
+	/** Sets up new instance data to sensible defaults, creates physics counterparts if possible. */
+	void SetupNewInstanceData(FSkinnedMeshInstanceData& InOutNewInstanceData, int32 InInstanceIndex, const FTransform3f& InInstanceTransform, int32 BankIndex);
+
+	FInstanceDataManager InstanceDataManager;
+
+	friend class FInstancedSkinnedMeshSceneProxy;
+	friend class FInstancedSkinnedMeshObject;
+	friend struct FInstancedSkinnedMeshComponentInstanceData;
+
+private:
+	void SetSkinnedAssetCallback();
+};
+
+/** Helper class used to preserve state across blueprint re-instancing */
+USTRUCT()
+struct FInstancedSkinnedMeshComponentInstanceData : public FSceneComponentInstanceData
+{
+	GENERATED_BODY()
+
+public:
+	FInstancedSkinnedMeshComponentInstanceData() = default;
+	
+	FInstancedSkinnedMeshComponentInstanceData(const UInstancedSkinnedMeshComponent* InComponent)
+	: FSceneComponentInstanceData(InComponent)
+	, SkinnedAsset(InComponent->GetSkinnedAsset())
+	{
+	}
+
+	virtual ~FInstancedSkinnedMeshComponentInstanceData() = default;
+
+	virtual bool ContainsData() const override
+	{
+		return true;
+	}
+
+	virtual void ApplyToComponent(UActorComponent* Component, const ECacheApplyPhase CacheApplyPhase) override
+	{
+		// The Super::ApplyToComponent will cause the scene proxy to be recreated, so we must do what we can to make sure the state is ok before that.
+		CastChecked<UInstancedSkinnedMeshComponent>(Component)->PreApplyComponentInstanceData(this);
+		Super::ApplyToComponent(Component, CacheApplyPhase);
+		CastChecked<UInstancedSkinnedMeshComponent>(Component)->ApplyComponentInstanceData(this);
+	}
+
+	virtual void AddReferencedObjects(FReferenceCollector& Collector) override
+	{
+		Super::AddReferencedObjects(Collector);
+		Collector.AddReferencedObject(SkinnedAsset);
+	}
+
+public:
+	UPROPERTY()
+	TObjectPtr<USkinnedAsset> SkinnedAsset = nullptr;
+
+	UPROPERTY()
+	TArray<FSkinnedMeshInstanceData> InstanceData;
+
+	TBitArray<> SelectedInstances;
+
+	UPROPERTY()
+	bool bHasPerInstanceHitProxies = false;
+};
diff --git a/Engine/Source/Runtime/Engine/Private/Animation/AnimBank.cpp b/Engine/Source/Runtime/Engine/Private/Animation/AnimBank.cpp
--- a/Engine/Source/Runtime/Engine/Private/Animation/AnimBank.cpp
+++ b/Engine/Source/Runtime/Engine/Private/Animation/AnimBank.cpp
@@ -38,7 +38,7 @@
 #include "Serialization/MemoryWriter.h"
 #endif
 
-#include "Components/AnimBankMeshComponent.h"
+#include "Components/InstancedSkinnedMeshComponent.h"
 
 #define LOCTEXT_NAMESPACE "AnimBank"
 
@@ -1397,17 +1397,17 @@ bool FSoftAnimBankItem::operator==(const FSoftAnimBankItem& Other) const
 	return BankAsset == Other.BankAsset && SequenceIndex == Other.SequenceIndex;
 }
 
-FAnimBankMeshComponentDescriptor::FAnimBankMeshComponentDescriptor()
+FSkinnedMeshComponentDescriptor::FSkinnedMeshComponentDescriptor()
 {
 	// Make sure we have proper defaults
-	InitFrom(UAnimBankMeshComponent::StaticClass()->GetDefaultObject<UAnimBankMeshComponent>());
+	InitFrom(UInstancedSkinnedMeshComponent::StaticClass()->GetDefaultObject<UInstancedSkinnedMeshComponent>());
 }
 
-FAnimBankMeshComponentDescriptor::FAnimBankMeshComponentDescriptor(ENoInit) {}
-FAnimBankMeshComponentDescriptor::FAnimBankMeshComponentDescriptor(const FAnimBankMeshComponentDescriptor& InDescriptor) = default;
-FAnimBankMeshComponentDescriptor::~FAnimBankMeshComponentDescriptor() = default;
+FSkinnedMeshComponentDescriptor::FSkinnedMeshComponentDescriptor(ENoInit) {}
+FSkinnedMeshComponentDescriptor::FSkinnedMeshComponentDescriptor(const FSkinnedMeshComponentDescriptor& InDescriptor) = default;
+FSkinnedMeshComponentDescriptor::~FSkinnedMeshComponentDescriptor() = default;
 
-FAnimBankMeshComponentDescriptor::FAnimBankMeshComponentDescriptor(const FSoftAnimBankMeshComponentDescriptor& InDescriptor)
+FSkinnedMeshComponentDescriptor::FSkinnedMeshComponentDescriptor(const FSoftSkinnedMeshComponentDescriptor& InDescriptor)
 {
 	SkinnedAsset = InDescriptor.SkinnedAsset.LoadSynchronous();
 	Algo::Transform(InDescriptor.BankItems, BankItems, [](const FSoftAnimBankItem& BankItem) { return FAnimBankItem(BankItem); });
@@ -1417,14 +1417,14 @@ FAnimBankMeshComponentDescriptor::FAnimBankMeshComponentDescriptor(const FSoftAn
 	Hash = InDescriptor.Hash;
 }
 
-UAnimBankMeshComponent* FAnimBankMeshComponentDescriptor::CreateComponent(UObject* Outer, FName Name, EObjectFlags ObjectFlags) const
+UInstancedSkinnedMeshComponent* FSkinnedMeshComponentDescriptor::CreateComponent(UObject* Outer, FName Name, EObjectFlags ObjectFlags) const
 {
-	UAnimBankMeshComponent* Component = NewObject<UAnimBankMeshComponent>(Outer, ComponentClass, Name, ObjectFlags);
+	UInstancedSkinnedMeshComponent* Component = NewObject<UInstancedSkinnedMeshComponent>(Outer, ComponentClass, Name, ObjectFlags);
 	InitComponent(Component);
 	return Component;
 }
 
-void FAnimBankMeshComponentDescriptor::InitFrom(const UAnimBankMeshComponent* Template, bool bInitBodyInstance)
+void FSkinnedMeshComponentDescriptor::InitFrom(const UInstancedSkinnedMeshComponent* Template, bool bInitBodyInstance)
 {
 	check(Template);
 
@@ -1434,9 +1434,9 @@ void FAnimBankMeshComponentDescriptor::InitFrom(const UAnimBankMeshComponent* Te
 	BankItems = Template->GetAnimBankItems();
 }
 
-uint32 FAnimBankMeshComponentDescriptor::ComputeHash() const
+uint32 FSkinnedMeshComponentDescriptor::ComputeHash() const
 {
-	FAnimBankMeshComponentDescriptor& MutableSelf = *const_cast<FAnimBankMeshComponentDescriptor*>(this);
+	FSkinnedMeshComponentDescriptor& MutableSelf = *const_cast<FSkinnedMeshComponentDescriptor*>(this);
 	Hash = 0; // we don't want the hash to impact the calculation
 	
 	FArchiveCrc32 CrcArchive;
@@ -1446,24 +1446,24 @@ uint32 FAnimBankMeshComponentDescriptor::ComputeHash() const
 	return Hash;
 }
 
-void FAnimBankMeshComponentDescriptor::InitComponent(UAnimBankMeshComponent* Component) const
+void FSkinnedMeshComponentDescriptor::InitComponent(UInstancedSkinnedMeshComponent* Component) const
 {
 	Component->SetSkinnedAsset(SkinnedAsset);
 	Component->Mobility = Mobility;
 	Component->SetAnimBankItems(BankItems);
 }
 
-void FAnimBankMeshComponentDescriptor::PostLoadFixup(UObject* Loader)
+void FSkinnedMeshComponentDescriptor::PostLoadFixup(UObject* Loader)
 {
 	check(Loader);
 }
 
-bool FAnimBankMeshComponentDescriptor::operator!=(const FAnimBankMeshComponentDescriptor& Other) const
+bool FSkinnedMeshComponentDescriptor::operator!=(const FSkinnedMeshComponentDescriptor& Other) const
 {
 	return !(*this == Other);
 }
 
-bool FAnimBankMeshComponentDescriptor::operator==(const FAnimBankMeshComponentDescriptor& Other) const
+bool FSkinnedMeshComponentDescriptor::operator==(const FSkinnedMeshComponentDescriptor& Other) const
 {
 	return
 		(Hash == 0 || Other.Hash == 0 || Hash == Other.Hash) && // Check hash first, other checks are in case of Hash collision
@@ -1473,17 +1473,17 @@ bool FAnimBankMeshComponentDescriptor::operator==(const FAnimBankMeshComponentDe
 		Mobility == Other.Mobility;
 }
 
-FSoftAnimBankMeshComponentDescriptor::FSoftAnimBankMeshComponentDescriptor()
+FSoftSkinnedMeshComponentDescriptor::FSoftSkinnedMeshComponentDescriptor()
 {
 	// Make sure we have proper defaults
-	InitFrom(UAnimBankMeshComponent::StaticClass()->GetDefaultObject<UAnimBankMeshComponent>());
+	InitFrom(UInstancedSkinnedMeshComponent::StaticClass()->GetDefaultObject<UInstancedSkinnedMeshComponent>());
 }
 
-FSoftAnimBankMeshComponentDescriptor::FSoftAnimBankMeshComponentDescriptor(ENoInit) {}
-FSoftAnimBankMeshComponentDescriptor::FSoftAnimBankMeshComponentDescriptor(const FSoftAnimBankMeshComponentDescriptor& InDescriptor) = default;
-FSoftAnimBankMeshComponentDescriptor::~FSoftAnimBankMeshComponentDescriptor() = default;
+FSoftSkinnedMeshComponentDescriptor::FSoftSkinnedMeshComponentDescriptor(ENoInit) {}
+FSoftSkinnedMeshComponentDescriptor::FSoftSkinnedMeshComponentDescriptor(const FSoftSkinnedMeshComponentDescriptor& InDescriptor) = default;
+FSoftSkinnedMeshComponentDescriptor::~FSoftSkinnedMeshComponentDescriptor() = default;
 
-FSoftAnimBankMeshComponentDescriptor::FSoftAnimBankMeshComponentDescriptor(const FAnimBankMeshComponentDescriptor& InDescriptor)
+FSoftSkinnedMeshComponentDescriptor::FSoftSkinnedMeshComponentDescriptor(const FSkinnedMeshComponentDescriptor& InDescriptor)
 {
 	SkinnedAsset = InDescriptor.SkinnedAsset;
 	Algo::Transform(InDescriptor.BankItems, BankItems, [](const FAnimBankItem& BankItem) { return FSoftAnimBankItem(BankItem); });
@@ -1493,14 +1493,14 @@ FSoftAnimBankMeshComponentDescriptor::FSoftAnimBankMeshComponentDescriptor(const
 	Hash = InDescriptor.Hash;
 }
 
-UAnimBankMeshComponent* FSoftAnimBankMeshComponentDescriptor::CreateComponent(UObject* Outer, FName Name, EObjectFlags ObjectFlags) const
+UInstancedSkinnedMeshComponent* FSoftSkinnedMeshComponentDescriptor::CreateComponent(UObject* Outer, FName Name, EObjectFlags ObjectFlags) const
 {
-	UAnimBankMeshComponent* Component = NewObject<UAnimBankMeshComponent>(Outer, ComponentClass, Name, ObjectFlags);
+	UInstancedSkinnedMeshComponent* Component = NewObject<UInstancedSkinnedMeshComponent>(Outer, ComponentClass, Name, ObjectFlags);
 	InitComponent(Component);
 	return Component;
 }
 
-void FSoftAnimBankMeshComponentDescriptor::InitFrom(const UAnimBankMeshComponent* Template, bool bInitBodyInstance)
+void FSoftSkinnedMeshComponentDescriptor::InitFrom(const UInstancedSkinnedMeshComponent* Template, bool bInitBodyInstance)
 {
 	check(Template);
 
@@ -1510,7 +1510,7 @@ void FSoftAnimBankMeshComponentDescriptor::InitFrom(const UAnimBankMeshComponent
 	BankItems = Template->GetAnimBankItems();
 }
 
-uint32 FSoftAnimBankMeshComponentDescriptor::ComputeHash() const
+uint32 FSoftSkinnedMeshComponentDescriptor::ComputeHash() const
 {
 	Hash = 0; // we don't want the hash to impact the calculation
 	
@@ -1521,7 +1521,7 @@ uint32 FSoftAnimBankMeshComponentDescriptor::ComputeHash() const
 	return Hash;
 }
 
-void FSoftAnimBankMeshComponentDescriptor::InitComponent(UAnimBankMeshComponent* Component) const
+void FSoftSkinnedMeshComponentDescriptor::InitComponent(UInstancedSkinnedMeshComponent* Component) const
 {
 	Component->SetSkinnedAsset(SkinnedAsset.LoadSynchronous());
 	Component->Mobility = Mobility;
@@ -1533,17 +1533,17 @@ void FSoftAnimBankMeshComponentDescriptor::InitComponent(UAnimBankMeshComponent*
 	Component->SetAnimBankItems(AnimBankItems);
 }
 
-void FSoftAnimBankMeshComponentDescriptor::PostLoadFixup(UObject* Loader)
+void FSoftSkinnedMeshComponentDescriptor::PostLoadFixup(UObject* Loader)
 {
 	check(Loader);
 }
 
-bool FSoftAnimBankMeshComponentDescriptor::operator!=(const FSoftAnimBankMeshComponentDescriptor& Other) const
+bool FSoftSkinnedMeshComponentDescriptor::operator!=(const FSoftSkinnedMeshComponentDescriptor& Other) const
 {
 	return !(*this == Other);
 }
 
-bool FSoftAnimBankMeshComponentDescriptor::operator==(const FSoftAnimBankMeshComponentDescriptor& Other) const
+bool FSoftSkinnedMeshComponentDescriptor::operator==(const FSoftSkinnedMeshComponentDescriptor& Other) const
 {
 	return
 		(Hash == 0 || Other.Hash == 0 || Hash == Other.Hash) && // Check hash first, other checks are in case of Hash collision
diff --git a/Engine/Source/Runtime/Engine/Private/Animation/AnimBankCompiler.cpp b/Engine/Source/Runtime/Engine/Private/Animation/AnimBankCompiler.cpp
--- a/Engine/Source/Runtime/Engine/Private/Animation/AnimBankCompiler.cpp
+++ b/Engine/Source/Runtime/Engine/Private/Animation/AnimBankCompiler.cpp
@@ -15,7 +15,7 @@
 #include "Animation/AnimBank.h"
 #include "ProfilingDebugging/CountersTrace.h"
 #include "AnimationSequenceCompiler.h"
-#include "Components/AnimBankMeshComponent.h"
+#include "Components/InstancedSkinnedMeshComponent.h"
 
 #define LOCTEXT_NAMESPACE "AnimBankCompiler"
 
@@ -219,7 +219,7 @@ void FAnimBankCompilingManager::PostCompilation(UAnimBank* AnimBank)
 		AnimBank->InitResources();
 
 		FObjectCacheContextScope ObjectCacheScope;
-		for (UAnimBankMeshComponent* Component : ObjectCacheScope.GetContext().GetAnimBankMeshComponents(AnimBank))
+		for (UInstancedSkinnedMeshComponent* Component : ObjectCacheScope.GetContext().GetInstancedSkinnedMeshComponents(AnimBank))
 		{
 			Component->PostAssetCompilation();
 		}
diff --git a/Engine/Source/Runtime/Engine/Private/Components/AnimBankMeshComponent.cpp b/Engine/Source/Runtime/Engine/Private/Components/AnimBankMeshComponent.cpp
deleted file mode 100644
diff --git a/Engine/Source/Runtime/Engine/Private/Components/InstancedSkinnedMeshComponent.cpp b/Engine/Source/Runtime/Engine/Private/Components/InstancedSkinnedMeshComponent.cpp
new file mode 100644
--- /dev/null
+++ b/Engine/Source/Runtime/Engine/Private/Components/InstancedSkinnedMeshComponent.cpp
@@ -0,0 +1,1647 @@
+// Copyright Epic Games, Inc. All Rights Reserved.
+
+#include "Components/InstancedSkinnedMeshComponent.h"
+#include "Animation/AnimSequence.h"
+#include "Animation/Skeleton.h"
+#include "ComponentRecreateRenderStateContext.h"
+#include "Elements/Framework/EngineElementsLibrary.h"
+#include "Elements/Interfaces/TypedElementWorldInterface.h"
+#include "Elements/SMInstance/SMInstanceElementData.h" // For SMInstanceElementDataUtil::SMInstanceElementsEnabled
+#include "Elements/SMInstance/SMInstanceElementId.h"
+#include "Engine/StaticMesh.h"
+#include "HitProxies.h"
+#include "NaniteSceneProxy.h"
+#include "PrimitiveSceneInfo.h"
+#include "Rendering/RenderCommandPipes.h"
+#include "Rendering/SkeletalMeshRenderData.h"
+#include "SceneInterface.h"
+#include "SkeletalRenderPublic.h"
+#include "SkinningDefinitions.h"
+#include "SkinnedMeshSceneProxyDesc.h"
+#include "InstanceData/InstanceUpdateChangeSet.h"
+#include "InstanceData/InstanceDataUpdateUtils.h"
+
+#include UE_INLINE_GENERATED_CPP_BY_NAME(InstancedSkinnedMeshComponent)
+
+static TAutoConsoleVariable<int32> CVarInstancedSkinnedMeshesForceRefPose(
+	TEXT("r.InstancedSkinnedMeshes.ForceRefPose"),
+	0,
+	TEXT("Whether to force ref pose for instanced skinned meshes"),
+	ECVF_RenderThreadSafe
+);
+
+static TAutoConsoleVariable<int32> CVarInstancedSkinnedMeshesSampledBounds(
+	TEXT("r.InstancedSkinnedMeshes.SampledBounds"),
+	1,
+	TEXT("Whether to use sampled bounds for anim bank meshes"),
+	FConsoleVariableDelegate::CreateLambda([](IConsoleVariable* InVariable)
+	{
+		FGlobalComponentRecreateRenderStateContext Context;
+	}),
+	ECVF_RenderThreadSafe
+);
+
+class FInstancedSkinnedMeshSceneProxy : public Nanite::FSkinnedSceneProxy
+{
+public:
+	using Super = Nanite::FSkinnedSceneProxy;
+
+	FInstancedSkinnedMeshSceneProxy(const Nanite::FMaterialAudit& MaterialAudit, UInstancedSkinnedMeshComponent* InComponent, FSkeletalMeshRenderData* InRenderData)
+	: Super(MaterialAudit, InComponent, InRenderData, false /* bAllowScale */)
+	, AnimationMinScreenSize(InComponent->AnimationMinScreenSize)
+	{
+	#if WITH_EDITOR
+		const bool bSupportInstancePicking = HasPerInstanceHitProxies() && SMInstanceElementDataUtil::SMInstanceElementsEnabled();
+		HitProxyMode = bSupportInstancePicking ? EHitProxyMode::PerInstance : EHitProxyMode::MaterialSection;
+
+		if (HitProxyMode == EHitProxyMode::PerInstance)
+		{
+			bHasSelectedInstances = InComponent->SelectedInstances.Find(true) != INDEX_NONE;
+			if (bHasSelectedInstances)
+			{
+				// If we have selected indices, mark scene proxy as selected.
+				SetSelection_GameThread(true);
+			}
+		}
+	#endif
+		const bool bForceRefPose = CVarInstancedSkinnedMeshesForceRefPose.GetValueOnAnyThread() != 0;
+		const bool bUseAnimBank = !bForceRefPose && InComponent->AnimBankItems.Num() > 0;
+
+		InstanceDataSceneProxy = InComponent->InstanceDataManager.GetOrCreateProxy();
+		SetupInstanceSceneDataBuffers(InstanceDataSceneProxy->GeInstanceSceneDataBuffers());
+
+		// Note: ideally this would be picked up from the Flags.bHasPerInstanceDynamicData above, but that path is not great at the moment.
+		bAlwaysHasVelocity = true;
+
+		if (bUseAnimBank)
+		{
+			static FGuid AnimBankGPUProviderId(ANIM_BANK_GPU_TRANSFORM_PROVIDER_GUID);
+			static FGuid AnimBankCPUProviderId(ANIM_BANK_CPU_TRANSFORM_PROVIDER_GUID);
+
+			static const auto AnimBankGPUVar = IConsoleManager::Get().FindTConsoleVariableDataInt(TEXT("r.AnimBank.GPU"));
+			if (AnimBankGPUVar && AnimBankGPUVar->GetValueOnAnyThread() == 1)
+			{
+				TransformProviderId = AnimBankGPUProviderId;
+			}
+			else
+			{
+				TransformProviderId = AnimBankCPUProviderId;
+			}
+
+			AnimBankItems = InComponent->AnimBankItems;
+			AnimBankIds.Reset(AnimBankItems.Num());
+			AnimBankHandles.Reset(AnimBankItems.Num());
+			UniqueAnimationCount = AnimBankItems.Num();
+		}
+		else
+		{
+			UniqueAnimationCount = 1; // Ref Pose
+			static FGuid RefPoseProviderId(REF_POSE_TRANSFORM_PROVIDER_GUID);
+			TransformProviderId = RefPoseProviderId;
+		}
+	}
+
+	virtual void CreateRenderThreadResources(FRHICommandListBase& RHICmdList) override
+	{
+		Super::CreateRenderThreadResources(RHICmdList);
+
+		TArray<FAnimBankDesc> Descs;
+		Descs.Reserve(AnimBankItems.Num());
+
+		for (const FAnimBankItem& Item : AnimBankItems)
+		{
+			if (Item.BankAsset == nullptr || SkinnedAsset == nullptr)
+			{
+				Descs.Emplace(FAnimBankDesc());
+				continue;
+			}
+
+			const FAnimBankData& BankData = Item.BankAsset->GetData();
+
+			if (Item.SequenceIndex >= BankData.Entries.Num())
+			{
+				Descs.Emplace(FAnimBankDesc());
+				continue;
+			}
+
+			const FAnimBankEntry& BankEntry = BankData.Entries[Item.SequenceIndex];
+
+			FAnimBankDesc& Desc = Descs.Emplace_GetRef();
+
+			Desc.BankAsset		= MakeWeakObjectPtr(Item.BankAsset);
+			Desc.SequenceIndex	= Item.SequenceIndex;
+			Desc.Asset			= MakeWeakObjectPtr(SkinnedAsset);
+			Desc.Position		= BankEntry.Position;
+			Desc.PlayRate		= BankEntry.PlayRate;
+			Desc.bLooping		= BankEntry.IsLooping()   ? 1 : 0;
+			Desc.bAutoStart		= BankEntry.IsAutoStart() ? 1 : 0;
+		}
+
+		AnimBankHandles = GetScene().RegisterAnimBank(Descs);
+		AnimBankIds.Reset(AnimBankHandles.Num());
+		for (const FAnimBankRecordHandle& Handle : AnimBankHandles)
+		{
+			AnimBankIds.Emplace((uint64)Handle.Id);
+		}
+	}
+
+	virtual void DestroyRenderThreadResources() override
+	{
+		GetScene().UnregisterAnimBank(AnimBankHandles);
+		AnimBankHandles.Reset();
+		AnimBankIds.Reset();
+		Super::DestroyRenderThreadResources();
+	}
+
+	// FSkinnedSceneProxy interface
+	virtual const TConstArrayView<uint64> GetAnimationProviderData(bool& bOutValid) const override
+	{
+		bOutValid = (AnimBankIds.Num() == UniqueAnimationCount);
+		return AnimBankIds;
+	}
+	virtual float GetAnimationMinScreenSize() const  override { return AnimationMinScreenSize; }
+
+public:
+	// FPrimitiveSceneProxy interface.
+	virtual SIZE_T GetTypeHash() const override
+	{
+		static size_t UniquePointer;
+		return reinterpret_cast<size_t>(&UniquePointer);
+	}
+
+private:
+	TSharedPtr<FInstanceDataSceneProxy, ESPMode::ThreadSafe> InstanceDataSceneProxy; 
+	TArray<uint64> AnimBankIds;
+	TArray<FAnimBankRecordHandle> AnimBankHandles;
+	TArray<FAnimBankItem> AnimBankItems;
+	uint64 BankRegistration = 0;
+	float AnimationMinScreenSize = 0.0f;
+};
+
+/** 
+* Stores the updated matrices needed to skin the verts.
+* Created by the game thread and sent to the rendering thread as an update 
+*/
+class FInstancedSkinnedMeshData
+{
+public:
+	FInstancedSkinnedMeshData(
+		const USkinnedAsset* InSkinnedAsset,
+		TConstArrayView<FAnimBankItem> AnimBankItems,
+		int32 InLODIndex,
+		EPreviousBoneTransformUpdateMode InPreviousBoneTransformUpdateMode
+	);
+
+	virtual ~FInstancedSkinnedMeshData();
+
+	// Current reference pose to local space transforms
+	TArray<FMatrix44f> ReferenceToLocal;
+
+	// Previous reference pose to local space transforms
+	TArray<FMatrix44f> PrevReferenceToLocal;
+
+	TArray<FMatrix3x4> CurrentBoneTransforms;
+	TArray<FMatrix3x4> PreviousBoneTransforms;
+
+#if !(UE_BUILD_SHIPPING || UE_BUILD_TEST) 
+	// Component space bone transforms
+	TArray<FTransform> ComponentSpaceTransforms;
+#endif
+
+	// Current LOD for bones being updated
+	int32 LODIndex;
+
+	// Returns the size of memory allocated by render data
+	void GetResourceSizeEx(FResourceSizeEx& CumulativeResourceSize);
+
+private:
+
+	enum class ETransformsToUpdate
+	{
+		Current,
+		Previous,
+	};
+};
+
+FInstancedSkinnedMeshData::FInstancedSkinnedMeshData(
+	const USkinnedAsset* InSkinnedAsset,
+	TConstArrayView<FAnimBankItem> AnimBankItems,
+	int32 InLODIndex,
+	EPreviousBoneTransformUpdateMode InPreviousBoneTransformUpdateMode
+)
+:	LODIndex(InLODIndex)
+{
+	int32 BoneCount = InSkinnedAsset->GetSkeleton()->GetReferenceSkeleton().GetRawBoneNum();
+	ReferenceToLocal.SetNumUninitialized(BoneCount);
+	for (int32 TransformIndex = 0; TransformIndex < BoneCount; ++TransformIndex)
+	{
+		ReferenceToLocal[TransformIndex].SetIdentity();
+	}
+
+	const int64 ReferenceToLocalCount = int64(ReferenceToLocal.Num());
+	const FMatrix44f* ReferenceToLocalPtr = ReferenceToLocal.GetData();
+
+	const int32 UniqueAnimationCount = FMath::Max(AnimBankItems.Num(), 1 /* Ref Pose */);
+
+	CurrentBoneTransforms.SetNumUninitialized(BoneCount * UniqueAnimationCount);
+	FMatrix3x4* CurrentBoneTransformsPtr = CurrentBoneTransforms.GetData();
+	TransposeTransforms(CurrentBoneTransformsPtr, ReferenceToLocalPtr, ReferenceToLocalCount);
+
+	// TODO: Temp, optimize out
+	for (int32 AnimationIndex = 1; AnimationIndex < UniqueAnimationCount; ++AnimationIndex)
+	{
+		FMemory::Memcpy(CurrentBoneTransformsPtr + (BoneCount * AnimationIndex), CurrentBoneTransformsPtr, sizeof(FMatrix3x4) * BoneCount);
+	}
+
+	PrevReferenceToLocal = ReferenceToLocal;
+	PreviousBoneTransforms = CurrentBoneTransforms;
+
+#if 0
+	UpdateRefToLocalMatrices(ReferenceToLocal, InComponent, InRenderData, LODIndex);
+	UpdateBonesRemovedByLOD(ReferenceToLocal, InComponent, ETransformsToUpdate::Current);
+
+	switch (InPreviousBoneTransformUpdateMode)
+	{
+	case EPreviousBoneTransformUpdateMode::None:
+		// Use previously uploaded buffer
+		// TODO: Nanite-Skinning, optimize scene extension upload to keep cached GPU representation using PreviousBoneTransformRevisionNumber
+		// For now we'll just redundantly update and upload previous transforms
+		UpdatePreviousRefToLocalMatrices(PrevReferenceToLocal, InComponent, InRenderData, LODIndex);
+		UpdateBonesRemovedByLOD(PrevReferenceToLocal, InComponent, ETransformsToUpdate::Previous);
+		break;
+
+	case EPreviousBoneTransformUpdateMode::UpdatePrevious:
+		UpdatePreviousRefToLocalMatrices(PrevReferenceToLocal, InComponent, InRenderData, LODIndex);
+		UpdateBonesRemovedByLOD(PrevReferenceToLocal, InComponent, ETransformsToUpdate::Previous);
+		break;
+
+	case EPreviousBoneTransformUpdateMode::DuplicateCurrentToPrevious:
+		// TODO: Nanite-Skinning likely possible we can just return ReferenceToLocal here rather than cloning it into previous
+		// Need to make sure it's safe when next update mode = None
+		PrevReferenceToLocal = ReferenceToLocal;
+		break;
+	}
+
+#if !(UE_BUILD_SHIPPING || UE_BUILD_TEST)
+	ComponentSpaceTransforms = InComponent->GetComponentSpaceTransforms();
+#endif
+#endif
+}
+
+FInstancedSkinnedMeshData::~FInstancedSkinnedMeshData() = default;
+
+void FInstancedSkinnedMeshData::GetResourceSizeEx(FResourceSizeEx& CumulativeResourceSize)
+{
+	CumulativeResourceSize.AddDedicatedSystemMemoryBytes(sizeof(*this));
+	CumulativeResourceSize.AddDedicatedSystemMemoryBytes(ReferenceToLocal.GetAllocatedSize());
+}
+
+class FInstancedSkinnedMeshObject : public FSkeletalMeshObject
+{
+public:
+	FInstancedSkinnedMeshObject(const FSkinnedMeshSceneProxyDesc& InMeshDesc, TConstArrayView<FAnimBankItem> InAnimBankItems, FSkeletalMeshRenderData* InRenderData, ERHIFeatureLevel::Type InFeatureLevel)
+		: FSkeletalMeshObject(InMeshDesc, InRenderData, InFeatureLevel)
+		, AnimBankItems(InAnimBankItems)
+	{
+		for (int32 LODIndex = 0; LODIndex < InRenderData->LODRenderData.Num(); ++LODIndex)
+		{
+			LODs.Emplace(InFeatureLevel, InRenderData, LODIndex);
+		}
+
+		InitResources(InMeshDesc);
+	}
+
+	FInstancedSkinnedMeshObject(USkinnedMeshComponent* InComponent, FSkeletalMeshRenderData* InRenderData, ERHIFeatureLevel::Type InFeatureLevel)
+	: FInstancedSkinnedMeshObject(FSkinnedMeshSceneProxyDesc(InComponent), CastChecked<UInstancedSkinnedMeshComponent>(InComponent)->AnimBankItems, InRenderData, InFeatureLevel)
+	{
+	}
+
+	virtual ~FInstancedSkinnedMeshObject()
+	{
+	}
+
+	virtual void InitResources(const FSkinnedMeshSceneProxyDesc& InMeshDesc) override
+	{
+		for (int32 LODIndex = 0; LODIndex < LODs.Num(); ++LODIndex)
+		{
+			FInstancedSkinnedMeshObjectLOD& LOD = LODs[LODIndex];
+
+			// Skip LODs that have their render data stripped
+			if (LOD.RenderData->LODRenderData[LODIndex].GetNumVertices() > 0)
+			{
+				const FSkelMeshComponentLODInfo* InitLODInfo = nullptr;
+				if (InMeshDesc.LODInfo.IsValidIndex(LODIndex))
+				{
+					InitLODInfo = &InMeshDesc.LODInfo[LODIndex];
+				}
+
+				LOD.InitResources(InitLODInfo);
+			}
+		}
+	}
+
+	virtual void ReleaseResources() override
+	{
+		for (int32 LODIndex = 0; LODIndex < LODs.Num(); ++LODIndex)
+		{
+			FInstancedSkinnedMeshObjectLOD& LOD = LODs[LODIndex];
+			LOD.ReleaseResources();
+		}
+	}
+	
+	virtual void Update(
+		int32 LODIndex,
+		const FSkinnedMeshSceneProxyDynamicData& InDynamicData,
+		const FPrimitiveSceneProxy* InSceneProxy,
+		const USkinnedAsset* InSkinnedAsset,
+		const FMorphTargetWeightMap& InActiveMorphTargets,
+		const TArray<float>& MorphTargetWeights,
+		EPreviousBoneTransformUpdateMode PreviousBoneTransformUpdateMode,
+		const FExternalMorphWeightData& InExternalMorphWeightData) override
+	{
+		// Create the new dynamic data for use by the rendering thread
+		// this data is only deleted when another update is sent
+		TUniquePtr<FInstancedSkinnedMeshData> NewDynamicData = MakeUnique<FInstancedSkinnedMeshData>(InSkinnedAsset, AnimBankItems, LODIndex, PreviousBoneTransformUpdateMode);
+
+		uint64 FrameNumberToPrepare = GFrameCounter;
+		uint32 RevisionNumber = 0;
+
+		if (InSceneProxy)
+		{
+			RevisionNumber = InDynamicData.GetBoneTransformRevisionNumber();
+		}
+
+		// Queue a call to update this data
+		{
+			FInstancedSkinnedMeshObject* MeshObject = this;
+			ENQUEUE_RENDER_COMMAND(SkelMeshObjectUpdateDataCommand)(UE::RenderCommandPipe::SkeletalMesh,
+				[MeshObject, FrameNumberToPrepare, RevisionNumber, NewDynamicData=MoveTemp(NewDynamicData)](FRHICommandList& RHICmdList) mutable
+				{
+					FScopeCycleCounter Context(MeshObject->GetStatId());
+					MeshObject->UpdateDynamicData_RenderThread(RHICmdList, MoveTemp(NewDynamicData), FrameNumberToPrepare, RevisionNumber);
+				}
+			);
+		}
+	}
+
+	void UpdateDynamicData_RenderThread(
+		FRHICommandList& RHICmdList,
+		TUniquePtr<FInstancedSkinnedMeshData>&& InDynamicData,
+		uint64 FrameNumberToPrepare,
+		uint32 RevisionNumber
+	)
+	{
+		// Update with new data
+		DynamicData = MoveTemp(InDynamicData);
+		check(DynamicData);
+
+		check(IsInParallelRenderingThread());
+	}
+
+	virtual const FVertexFactory* GetSkinVertexFactory(const FSceneView* View, int32 LODIndex, int32 ChunkIdx, ESkinVertexFactoryMode VFMode = ESkinVertexFactoryMode::Default) const override
+	{
+		check(LODs.IsValidIndex(LODIndex));
+		return nullptr;
+		//return &LODs[LODIndex].VertexFactory;
+	}
+
+	virtual const FVertexFactory* GetStaticSkinVertexFactory(int32 LODIndex, int32 ChunkIdx, ESkinVertexFactoryMode VFMode = ESkinVertexFactoryMode::Default) const override
+	{
+		check(LODs.IsValidIndex(LODIndex));
+		return nullptr;
+		//return &LODs[LODIndex].VertexFactory;
+	}
+
+	virtual TArray<FTransform>* GetComponentSpaceTransforms() const override
+	{
+	#if !(UE_BUILD_SHIPPING || UE_BUILD_TEST)
+		if (DynamicData)
+		{
+			return &(DynamicData->ComponentSpaceTransforms);
+		}
+		else
+	#endif
+		{
+			return nullptr;
+		}
+	}
+
+	virtual const TArray<FMatrix44f>& GetReferenceToLocalMatrices() const override
+	{
+		return DynamicData->ReferenceToLocal;
+	}
+
+	virtual const TArray<FMatrix44f>& GetPrevReferenceToLocalMatrices() const override
+	{
+		return DynamicData->PrevReferenceToLocal;
+	}
+
+	const TArray<FMatrix3x4>* GetCurrentBoneTransforms() const override
+	{
+		return &DynamicData->CurrentBoneTransforms;
+	}
+
+	const TArray<FMatrix3x4>* GetPreviousBoneTransforms() const override
+	{
+		return &DynamicData->PreviousBoneTransforms;
+	}
+
+	virtual int32 GetLOD() const override
+	{
+		// WorkingMinDesiredLODLevel can be a LOD that's not loaded, so need to clamp it to the first loaded LOD
+		return 0;
+		//return FMath::Max<int32>(WorkingMinDesiredLODLevel, SkeletalMeshRenderData->CurrentFirstLODIdx);
+		/*if (DynamicData)
+		{
+			return DynamicData->LODIndex;
+		}
+		else
+		{
+			return 0;
+		}*/
+	}
+
+	virtual bool HaveValidDynamicData() const override
+	{
+		return (DynamicData != nullptr);
+	}
+
+	virtual void GetResourceSizeEx(FResourceSizeEx& CumulativeResourceSize) override
+	{
+		CumulativeResourceSize.AddDedicatedSystemMemoryBytes(sizeof(*this));
+
+		if (DynamicData)
+		{
+			DynamicData->GetResourceSizeEx(CumulativeResourceSize);
+		}
+
+		CumulativeResourceSize.AddDedicatedSystemMemoryBytes(LODs.GetAllocatedSize());
+
+		for (int32 Index = 0; Index < LODs.Num(); ++Index)
+		{
+			LODs[Index].GetResourceSizeEx(CumulativeResourceSize);
+		}
+	}
+
+	virtual void UpdateSkinWeightBuffer(const TArrayView<const FSkelMeshComponentLODInfo> InLODInfo) override
+	{
+		for (int32 LODIndex = 0; LODIndex < LODs.Num(); ++LODIndex)
+		{
+			FInstancedSkinnedMeshObjectLOD& LOD = LODs[LODIndex];
+
+			// Skip LODs that have their render data stripped
+			if (LOD.RenderData->LODRenderData[LODIndex].GetNumVertices() > 0)
+			{
+				const FSkelMeshComponentLODInfo* UpdateLODInfo = nullptr;
+				if (InLODInfo.IsValidIndex(LODIndex))
+				{
+					UpdateLODInfo = &InLODInfo[LODIndex];
+				}
+
+				LOD.UpdateSkinWeights(UpdateLODInfo);
+			}
+		}
+	}
+
+	virtual bool IsNaniteMesh() const override { return true; }
+
+private:
+	TUniquePtr<FInstancedSkinnedMeshData> DynamicData;
+
+	struct FInstancedSkinnedMeshObjectLOD
+	{
+		const FSkeletalMeshRenderData* RenderData;
+		int32 LODIndex;
+		bool bInitialized;
+
+		FInstancedSkinnedMeshObjectLOD(ERHIFeatureLevel::Type InFeatureLevel, const FSkeletalMeshRenderData* InRenderData, int32 InLOD)
+		: RenderData(InRenderData)
+		, LODIndex(InLOD)
+		, bInitialized(false)
+		{
+		}
+
+		void InitResources(const FSkelMeshComponentLODInfo* InLODInfo)
+		{
+			check(RenderData);
+			check(RenderData->LODRenderData.IsValidIndex(LODIndex));
+
+			const FSkeletalMeshLODRenderData& LODData = RenderData->LODRenderData[LODIndex];
+			(void)LODData;
+
+		#if RHI_RAYTRACING
+			if (IsRayTracingEnabled() && RenderData->bSupportRayTracing)
+			{
+			}
+		#endif
+
+			bInitialized = true;
+		}
+
+		void ReleaseResources()
+		{
+			bInitialized = false;
+		}
+
+		void GetResourceSizeEx(FResourceSizeEx& CumulativeResourceSize)
+		{
+		}
+
+		void UpdateSkinWeights(const FSkelMeshComponentLODInfo* InLODInfo)
+		{
+			check(RenderData);
+			check(RenderData->LODRenderData.IsValidIndex(LODIndex));
+
+			//FSkeletalMeshLODRenderData& LODData = RenderData->LODRenderData[LODIndex];
+			//MeshObjectWeightBuffer = FSkeletalMeshObject::GetSkinWeightVertexBuffer(LODData, InLODInfo);
+		}
+	};
+
+	TConstArrayView<FAnimBankItem> AnimBankItems;
+	TArray<FInstancedSkinnedMeshObjectLOD> LODs;
+
+	mutable int32 CachedLOD;
+};
+
+static FSkeletalMeshObject* CreateInstancedSkinnedMeshObjectFn(void* UserData, USkinnedMeshComponent* InComponent, FSkeletalMeshRenderData* InRenderData, ERHIFeatureLevel::Type InFeatureLevel)
+{
+	if (InComponent->ShouldNaniteSkin())
+	{
+		return ::new FInstancedSkinnedMeshObject(InComponent, InRenderData, InFeatureLevel);
+	}
+
+	return nullptr;
+}
+
+UInstancedSkinnedMeshComponent::UInstancedSkinnedMeshComponent(FVTableHelper& Helper)
+: Super(Helper)
+, bInheritPerInstanceData(false)
+, InstanceDataManager(this)
+{
+}
+
+UInstancedSkinnedMeshComponent::UInstancedSkinnedMeshComponent(const FObjectInitializer& ObjectInitializer) 
+: Super(ObjectInitializer) 
+, bInheritPerInstanceData(false)
+, InstanceDataManager(this)
+{
+}
+
+UInstancedSkinnedMeshComponent::~UInstancedSkinnedMeshComponent()
+{
+}
+
+struct FSkinnedMeshInstanceData_Deprecated
+{
+	FMatrix Transform;
+	uint32 BankIndex;
+	uint32 Padding[3]; // Need to respect 16 byte alignment for bulk-serialization
+
+	FSkinnedMeshInstanceData_Deprecated()
+	: Transform(FMatrix::Identity)
+	, BankIndex(0)
+	{
+		Padding[0] = 0;
+		Padding[1] = 0;
+		Padding[2] = 0;
+	}
+
+	FSkinnedMeshInstanceData_Deprecated(const FMatrix& InTransform, uint32 InBankIndex)
+	: Transform(InTransform)
+	, BankIndex(InBankIndex)
+	{
+		Padding[0] = 0;
+		Padding[1] = 0;
+		Padding[2] = 0;
+	}
+
+	friend FArchive& operator<<(FArchive& Ar, FSkinnedMeshInstanceData_Deprecated& InstanceData)
+	{
+		// @warning BulkSerialize: FSkinnedMeshInstanceData is serialized as memory dump
+		// See TArray::BulkSerialize for detailed description of implied limitations.
+		Ar << InstanceData.Transform;
+		Ar << InstanceData.BankIndex;
+		Ar << InstanceData.Padding[0];
+		Ar << InstanceData.Padding[1];
+		Ar << InstanceData.Padding[2];
+		return Ar;
+	}
+};
+
+void UInstancedSkinnedMeshComponent::Serialize(FArchive& Ar)
+{
+	LLM_SCOPE(ELLMTag::SkeletalMesh);
+	Super::Serialize(Ar);
+
+	Ar.UsingCustomVersion(FUE5MainStreamObjectVersion::GUID);
+
+	bool bCooked = Ar.IsCooking();
+	Ar << bCooked;
+
+	// Inherit properties when bEditableWhenInherited == false || bInheritPerInstanceData == true (when the component isn't a template and we are persisting data)
+	const UInstancedSkinnedMeshComponent* Archetype = Cast<UInstancedSkinnedMeshComponent>(GetArchetype());
+	const bool bInheritSkipSerializationProperties = ShouldInheritPerInstanceData(Archetype) && Ar.IsPersistent();
+	
+	// Check if we need have SkipSerialization property data to load/save
+	bool bHasSkipSerializationPropertiesData = !bInheritSkipSerializationProperties;
+	Ar << bHasSkipSerializationPropertiesData;
+
+	if (Ar.IsLoading())
+	{
+		// Read existing data if it was serialized
+		TArray<FSkinnedMeshInstanceData> TempInstanceData;
+		TArray<float> TempInstanceCustomData;
+
+		if (bHasSkipSerializationPropertiesData)
+		{
+			if (Ar.CustomVer(FUE5MainStreamObjectVersion::GUID) < FUE5MainStreamObjectVersion::SkinnedMeshInstanceDataSerializationV2)
+			{
+				TArray<FSkinnedMeshInstanceData_Deprecated> TempInstanceData_Deprecated;
+				TempInstanceData_Deprecated.BulkSerialize(Ar, false /* force per element serialization */);
+
+				TempInstanceData.Reserve(TempInstanceData_Deprecated.Num());
+				for (const auto& Item : TempInstanceData_Deprecated)
+				{
+					TempInstanceData.Emplace(FTransform3f(FMatrix44f(Item.Transform)), Item.BankIndex);
+				}
+			}
+			else
+			{
+				Ar << TempInstanceData;
+			}
+			TempInstanceCustomData.BulkSerialize(Ar);
+		}
+
+		// If we should inherit use Archetype Data
+		if (bInheritSkipSerializationProperties)
+		{
+			ApplyInheritedPerInstanceData(Archetype);
+		} 
+		// It is possible for a component to lose its BP archetype between a save / load so in this case we have no per instance data (usually this component gets deleted through construction script)
+		else if (bHasSkipSerializationPropertiesData)
+		{
+			InstanceData = MoveTemp(TempInstanceData);
+			InstanceCustomData = MoveTemp(TempInstanceCustomData);
+		}
+	}
+	else if (bHasSkipSerializationPropertiesData)
+	{
+		Ar << InstanceData;
+		InstanceCustomData.BulkSerialize(Ar);
+	}
+
+#if WITH_EDITOR
+	if (Ar.IsTransacting())
+	{
+		Ar << SelectedInstances;
+	}
+#endif
+
+	if (Ar.CustomVer(FUE5MainStreamObjectVersion::GUID) >= FUE5MainStreamObjectVersion::SkinnedMeshInstanceDataSerializationV2)
+	{
+		InstanceDataManager.Serialize(Ar, bCooked);
+	}
+	else if (Ar.IsLoading())
+	{
+		// Prior to this the id mapping was not saved so we need to reset it.
+		InstanceDataManager.Reset(InstanceData.Num());
+	}
+
+	if (bCooked)
+	{
+		if (Ar.IsLoading())
+		{
+			InstanceDataManager.ReadCookedRenderData(Ar);
+		}
+#if WITH_EDITOR
+		else if (Ar.IsSaving())
+		{
+			InstanceDataManager.WriteCookedRenderData(Ar, GetComponentDesc(GMaxRHIFeatureLevel));
+		}
+#endif
+	}
+}
+
+void UInstancedSkinnedMeshComponent::PostLoad()
+{
+	Super::PostLoad();
+
+#if WITH_EDITOR
+	if (GIsEditor)
+	{
+		SetSkinnedAssetCallback();
+	}
+#endif
+}
+
+void UInstancedSkinnedMeshComponent::BeginDestroy()
+{
+	Super::BeginDestroy();
+}
+
+void UInstancedSkinnedMeshComponent::OnRegister()
+{
+	Super::OnRegister();
+}
+
+void UInstancedSkinnedMeshComponent::TickComponent(float DeltaTime, ELevelTick TickType, FActorComponentTickFunction* ThisTickFunction)
+{
+	Super::TickComponent(DeltaTime, TickType, ThisTickFunction);
+}
+
+bool UInstancedSkinnedMeshComponent::IsEnabled() const
+{
+	USkeletalMesh* SkeletalMeshPtr = Cast<USkeletalMesh>(GetSkinnedAsset());
+	if (SkeletalMeshPtr && SkeletalMeshPtr->GetResourceForRendering())
+	{
+		return GetInstanceCount() > 0;
+	}
+
+	return false;
+}
+
+int32 UInstancedSkinnedMeshComponent::GetInstanceCount() const
+{
+	return InstanceData.Num();
+}
+
+TArrayView<const struct FAnimBankItem> UInstancedSkinnedMeshComponent::GetAnimBankItems() const
+{
+	return MakeArrayView(AnimBankItems);
+}
+
+void UInstancedSkinnedMeshComponent::SetAnimBankItems(TArrayView<const struct FAnimBankItem> InAnimBankItems)
+{
+	AnimBankItems = InAnimBankItems;
+	// We use the transform dirty state to drive the update of the animation data (to defer the need to add more bits), so we mark those as dirty here.
+	InstanceDataManager.TransformsChangedAll();
+	MarkRenderStateDirty();
+}
+
+template <typename ArrayType>
+inline void ReorderArray(ArrayType& InOutDataArray, const TArray<int32>& OldIndexArray, int32 ElementStride = 1)
+{
+	// TODO: this fails for FSkinnedMeshInstanceData, it is not really a problem here because the ctor just performs initializaiton and we'll overwrite it anyway 
+	// static_assert(std::is_trivially_constructible_v<decltype(InOutDataArray[0])>, "Must have trivially constructible type to safely call SetNumUninitialized below");
+	ArrayType TmpDataArray = MoveTemp(InOutDataArray);
+	InOutDataArray.Empty(TmpDataArray.Num());
+	for (int32 NewIndex = 0; NewIndex < TmpDataArray.Num(); ++NewIndex)
+	{
+		int32 OldIndex = OldIndexArray[NewIndex];
+		for (int32 SubIndex = 0; SubIndex < ElementStride; ++SubIndex)
+		{
+			InOutDataArray.Add(TmpDataArray[OldIndex * ElementStride + SubIndex]);
+		}
+	}
+}
+
+void UInstancedSkinnedMeshComponent::OptimizeInstanceData(bool bShouldRetainIdMap)
+{
+	// compute the optimal order 
+	TArray<int32> IndexRemap = InstanceDataManager.Optimize(GetComponentDesc(GMaxRHIFeatureLevel), bShouldRetainIdMap);
+	
+	if (!IndexRemap.IsEmpty())
+	{
+		// Reorder instances according to the remap
+		ReorderArray(InstanceData, IndexRemap);
+		ReorderArray(InstanceCustomData, IndexRemap, NumCustomDataFloats);
+#if WITH_EDITOR
+		ReorderArray(SelectedInstances, IndexRemap);
+#endif
+	}
+}
+
+void UInstancedSkinnedMeshComponent::ApplyInheritedPerInstanceData(const UInstancedSkinnedMeshComponent* InArchetype)
+{
+	check(InArchetype);
+	InstanceData = InArchetype->InstanceData;
+	InstanceCustomData = InArchetype->InstanceCustomData;
+	NumCustomDataFloats = InArchetype->NumCustomDataFloats;
+}
+
+bool UInstancedSkinnedMeshComponent::ShouldInheritPerInstanceData() const
+{
+	return ShouldInheritPerInstanceData(Cast<UInstancedSkinnedMeshComponent>(GetArchetype()));
+}
+
+bool UInstancedSkinnedMeshComponent::ShouldInheritPerInstanceData(const UInstancedSkinnedMeshComponent* InArchetype) const
+{
+	return (bInheritPerInstanceData || !bEditableWhenInherited) && InArchetype && InArchetype->IsInBlueprint() && !IsTemplate();
+}
+
+void UInstancedSkinnedMeshComponent::SetupNewInstanceData(FSkinnedMeshInstanceData& InOutNewInstanceData, int32 InInstanceIndex, const FTransform3f& InInstanceTransform, int32 InBankIndex)
+{
+	InOutNewInstanceData.Transform = InInstanceTransform;
+	InOutNewInstanceData.BankIndex = InBankIndex;
+
+	if (bPhysicsStateCreated)
+	{
+		// ..
+	}
+}
+
+const Nanite::FResources* UInstancedSkinnedMeshComponent::GetNaniteResources() const
+{
+	return Super::GetNaniteResources();
+}
+
+#if WITH_EDITOR
+
+void UInstancedSkinnedMeshComponent::PostAssetCompilation()
+{
+	InstanceDataManager.ClearChangeTracking();
+	MarkRenderStateDirty();
+}
+
+#endif 
+
+FInstanceDataManagerSourceDataDesc UInstancedSkinnedMeshComponent::GetComponentDesc(ERHIFeatureLevel::Type FeatureLevel)
+{
+	FInstanceDataManagerSourceDataDesc ComponentDesc;
+
+
+	ComponentDesc.PrimitiveMaterialDesc = GetUsedMaterialPropertyDesc(FeatureLevel);
+
+	FInstanceDataFlags Flags;
+	Flags.bHasPerInstanceRandom = ComponentDesc.PrimitiveMaterialDesc.bAnyMaterialHasPerInstanceRandom;
+	Flags.bHasPerInstanceCustomData = ComponentDesc.PrimitiveMaterialDesc.bAnyMaterialHasPerInstanceCustomData && NumCustomDataFloats != 0;
+#if WITH_EDITOR
+	Flags.bHasPerInstanceEditorData = GIsEditor != 0 && bHasPerInstanceHitProxies;
+#endif
+	
+	const bool bForceRefPose = CVarInstancedSkinnedMeshesForceRefPose.GetValueOnAnyThread() != 0;
+	const bool bUseAnimBank = !bForceRefPose && AnimBankItems.Num() > 0;
+
+	Flags.bHasPerInstanceHierarchyOffset = false;
+	Flags.bHasPerInstanceLocalBounds = bUseAnimBank && AnimBankItems.Num() > 1;
+	Flags.bHasPerInstanceDynamicData = false;
+	Flags.bHasPerInstanceSkinningData = true;
+
+
+	Flags.bHasPerInstanceLMSMUVBias = false;//IsStaticLightingAllowed();
+
+	ComponentDesc.Flags = Flags;
+	
+	// TODO: rename
+	ComponentDesc.MeshBounds = GetSkinnedAsset()->GetBounds();
+	ComponentDesc.NumCustomDataFloats = NumCustomDataFloats;
+	ComponentDesc.NumInstances = InstanceData.Num();
+
+	ComponentDesc.PrimitiveLocalToWorld = GetRenderMatrix();
+	ComponentDesc.ComponentMobility = Mobility;
+
+	const FReferenceSkeleton& RefSkeleton = GetSkinnedAsset()->GetRefSkeleton();
+	uint32 MaxBoneTransformCount = RefSkeleton.GetRawBoneNum();
+
+	ComponentDesc.BuildChangeSet = [&, MaxBoneTransformCount, MeshBounds = ComponentDesc.MeshBounds](FInstanceUpdateChangeSet& ChangeSet)
+	{
+		// publish data
+		ChangeSet.GetTransformWriter().Gather([&](int32 InstanceIndex) -> FRenderTransform { return FRenderTransform(InstanceData[InstanceIndex].Transform.ToMatrixWithScale()); });
+		ChangeSet.GetCustomDataWriter().Gather(MakeArrayView(InstanceCustomData), NumCustomDataFloats);
+
+		ChangeSet.GetSkinningDataWriter().Gather(
+			[&](int32 InstanceIndex) -> uint32 
+			{ 
+				return InstanceData[InstanceIndex].BankIndex * MaxBoneTransformCount * 2u; 
+			});
+
+		ChangeSet.GetLocalBoundsWriter().Gather(
+			[&](int32 InstanceIndex) -> FRenderBounds 
+			{				
+				uint32 BankIndex = InstanceData[InstanceIndex].BankIndex;
+				if (BankIndex < uint32(AnimBankItems.Num()))
+				{
+					const FAnimBankItem& BankItem = AnimBankItems[BankIndex];
+					if (BankItem.BankAsset != nullptr)
+					{
+						const FAnimBankData& BankData = BankItem.BankAsset->GetData();
+						if (BankItem.SequenceIndex < BankData.Entries.Num())
+						{
+							return BankData.Entries[BankItem.SequenceIndex].SampledBounds;
+						}
+					}
+				}
+				return MeshBounds; 
+			});
+
+#if WITH_EDITOR
+	if (ChangeSet.Flags.bHasPerInstanceEditorData)
+	{
+		// TODO: the way hit proxies are managed seems daft, why don't we just add them when needed and store them in an array alonside the instances?
+		//       this will always force us to update all the hit proxy data for every instances.
+		TArray<TRefCountPtr<HHitProxy>> HitProxies;
+		CreateHitProxyData(HitProxies);
+		ChangeSet.SetEditorData(HitProxies, SelectedInstances);
+	}
+#endif
+
+
+	};
+
+	return ComponentDesc;
+}
+
+void UInstancedSkinnedMeshComponent::SendRenderInstanceData_Concurrent()
+{
+	Super::SendRenderInstanceData_Concurrent();
+
+	// If the primitive isn't hidden update its instances.
+	const bool bDetailModeAllowsRendering = true;//DetailMode <= GetCachedScalabilityCVars().DetailMode;
+	// The proxy may not be created, this can happen when a SM is async loading for example.
+	if (bDetailModeAllowsRendering && (ShouldRender() || bCastHiddenShadow || bAffectIndirectLightingWhileHidden || bRayTracingFarField))
+	{
+		if (SceneProxy != nullptr)
+		{
+			// Make sure the instance data proxy is up to date:
+			if (InstanceDataManager.FlushChanges(GetComponentDesc(SceneProxy->GetScene().GetFeatureLevel())))
+			{
+				UpdateBounds();
+				GetWorld()->Scene->UpdatePrimitiveInstances(this);
+			}
+		}
+		else
+		{
+			UpdateBounds();
+			GetWorld()->Scene->AddPrimitive(this);
+		}
+	}
+}
+
+
+void UInstancedSkinnedMeshComponent::CreateRenderState_Concurrent(FRegisterComponentContext* Context)
+{
+	MeshObjectFactory = &CreateInstancedSkinnedMeshObjectFn;
+	Super::CreateRenderState_Concurrent(Context);
+}
+
+void UInstancedSkinnedMeshComponent::DestroyRenderState_Concurrent()
+{
+	Super::DestroyRenderState_Concurrent();
+}
+
+FPrimitiveSceneProxy* UInstancedSkinnedMeshComponent::CreateSceneProxy()
+{
+	LLM_SCOPE(ELLMTag::SkeletalMesh);
+
+	ERHIFeatureLevel::Type SceneFeatureLevel = GetWorld()->GetFeatureLevel();
+	FPrimitiveSceneProxy* Result = nullptr;
+	FSkeletalMeshRenderData* SkelMeshRenderData = GetSkeletalMeshRenderData();
+
+#if WITH_EDITOR
+	if (!bIsInstanceDataApplyCompleted)
+	{
+		return nullptr;
+	}
+#endif
+
+	const USkinnedAsset* SkinnedAssetPtr = GetSkinnedAsset();
+	if (InstanceData.Num() == 0 || SkinnedAssetPtr == nullptr || SkinnedAssetPtr->IsCompiling())
+	{
+		return nullptr;
+	}
+
+	for (FAnimBankItem& BankItem : AnimBankItems)
+	{
+		if (BankItem.BankAsset && BankItem.BankAsset->IsCompiling())
+		{
+			return nullptr;
+		}
+	}
+
+	if (CheckPSOPrecachingAndBoostPriority() && GetPSOPrecacheProxyCreationStrategy() == EPSOPrecacheProxyCreationStrategy::DelayUntilPSOPrecached)
+	{
+		UE_LOG(LogAnimBank, Verbose, TEXT("Skipping CreateSceneProxy for UInstancedSkinnedMeshComponent %s (UInstancedSkinnedMeshComponent PSOs are still compiling)"), *GetFullName());
+		return nullptr;
+	}
+
+	// Only create a scene proxy for rendering if properly initialized
+	if (SkelMeshRenderData &&
+		SkelMeshRenderData->LODRenderData.IsValidIndex(GetPredictedLODLevel()) &&
+		!bHideSkin &&
+		MeshObject)
+	{
+		// Only create a scene proxy if the bone count being used is supported, or if we don't have a skeleton (this is the case with destructibles)
+		int32 MinLODIndex = ComputeMinLOD();
+		int32 MaxBonesPerChunk = SkelMeshRenderData->GetMaxBonesPerSection(MinLODIndex);
+		int32 MaxSupportedNumBones = MeshObject->IsCPUSkinned() ? MAX_int32 : FGPUBaseSkinVertexFactory::GetMaxGPUSkinBones();
+		if (MaxBonesPerChunk <= MaxSupportedNumBones)
+		{
+			if (ShouldNaniteSkin() && IsEnabled())
+			{
+				Nanite::FMaterialAudit NaniteMaterials{};
+				AuditMaterials(this, NaniteMaterials, true /* Set material usage flags */);
+
+				const bool bIsMaskingAllowed = Nanite::IsMaskingAllowed(GetWorld(), false /* force Nanite for masked */);
+				if (NaniteMaterials.IsValid(bIsMaskingAllowed))
+				{
+					Result = ::new FInstancedSkinnedMeshSceneProxy(NaniteMaterials, this, SkelMeshRenderData);
+				}
+			}
+
+			if (Result == nullptr)
+			{
+				Result = Super::CreateSceneProxy();
+			}
+
+			// Unclear exactly how this is supposed to work with a non-instanced proxy - will be interesting...
+			if (Result)
+			{
+				InstanceDataManager.FlushChanges(GetComponentDesc(Result->GetScene().GetFeatureLevel()));
+			}
+		}
+	}
+
+#if !(UE_BUILD_SHIPPING || UE_BUILD_TEST)
+	SendRenderDebugPhysics(Result);
+#endif
+
+	return Result;
+}
+
+void UInstancedSkinnedMeshComponent::OnUpdateTransform(EUpdateTransformFlags UpdateTransformFlags, ETeleportType Teleport)
+{
+	Super::OnUpdateTransform(UpdateTransformFlags, Teleport);
+	InstanceDataManager.PrimitiveTransformChanged();
+}
+
+#if WITH_EDITOR
+
+void UInstancedSkinnedMeshComponent::PostEditChangeChainProperty(FPropertyChangedChainEvent& PropertyChangedEvent)
+{
+	// Always clear the change tracking because in the editor, attributes may have been set without any sort of notification
+	InstanceDataManager.ClearChangeTracking();
+	if (PropertyChangedEvent.Property != nullptr)
+	{
+		if (PropertyChangedEvent.Property->GetFName() == GET_MEMBER_NAME_CHECKED(UInstancedSkinnedMeshComponent, InstanceData))
+		{
+			if (PropertyChangedEvent.ChangeType == EPropertyChangeType::ArrayAdd
+				|| PropertyChangedEvent.ChangeType == EPropertyChangeType::Duplicate)
+			{
+				int32 AddedAtIndex = PropertyChangedEvent.GetArrayIndex(PropertyChangedEvent.Property->GetFName().ToString());
+				check(AddedAtIndex != INDEX_NONE);
+
+				AddInstanceInternal(
+					AddedAtIndex,
+					PropertyChangedEvent.ChangeType == EPropertyChangeType::ArrayAdd ? FTransform::Identity : FTransform(InstanceData[AddedAtIndex].Transform),
+					PropertyChangedEvent.ChangeType == EPropertyChangeType::ArrayAdd ? 0 : InstanceData[AddedAtIndex].BankIndex,
+					/*bWorldSpace*/false
+				);
+
+				// added via the property editor, so we will want to interactively work with instances
+				//bHasPerInstanceHitProxies = true;
+			}
+			else if (PropertyChangedEvent.ChangeType == EPropertyChangeType::ArrayRemove)
+			{
+				int32 RemovedAtIndex = PropertyChangedEvent.GetArrayIndex(PropertyChangedEvent.Property->GetFName().ToString());
+				check(RemovedAtIndex != INDEX_NONE);
+
+				RemoveInstanceInternal(RemovedAtIndex, true);
+			}
+			else if (PropertyChangedEvent.ChangeType == EPropertyChangeType::ArrayClear)
+			{
+				ClearInstances();
+			}
+			else if (PropertyChangedEvent.ChangeType == EPropertyChangeType::ValueSet)
+			{
+			}
+			MarkRenderStateDirty();
+		}
+		else if (PropertyChangedEvent.Property->GetFName() == GET_MEMBER_NAME_CHECKED(FSkinnedMeshInstanceData, Transform)
+			 || PropertyChangedEvent.Property->GetFName() == GET_MEMBER_NAME_CHECKED(FSkinnedMeshInstanceData, BankIndex))
+		{
+			MarkRenderStateDirty();
+		}
+		else if (PropertyChangedEvent.Property->GetFName() == "NumCustomDataFloats")
+		{
+			SetNumCustomDataFloats(NumCustomDataFloats);
+		}
+		else if (PropertyChangedEvent.PropertyChain.GetActiveMemberNode()->GetValue()->GetFName() == "InstanceCustomData")
+		{
+			int32 ChangedCustomValueIndex = PropertyChangedEvent.GetArrayIndex(PropertyChangedEvent.Property->GetFName().ToString());
+			if (ensure(NumCustomDataFloats > 0))
+			{
+				int32 InstanceIndex = ChangedCustomValueIndex / NumCustomDataFloats;
+			}
+			MarkRenderStateDirty();
+		}
+		else if (
+			(PropertyChangedEvent.Property->GetFName() == GET_MEMBER_NAME_CHECKED(UInstancedSkinnedMeshComponent, AnimBankItems)) ||
+			(PropertyChangedEvent.Property->GetFName() == GET_MEMBER_NAME_CHECKED(FAnimBankItem, BankAsset))
+			)
+		{
+			for (FAnimBankItem& BankItem : AnimBankItems)
+			{
+				// TODO: BankItem.ValidatePosition();
+
+				// Make sure the animation skeleton is valid
+				bool bInvalid = true;
+				if (BankItem.BankAsset && BankItem.BankAsset->Asset && BankItem.BankAsset->Asset->GetSkeleton())
+				{
+					USkeletalMesh* SkeletalMeshPtr = Cast<USkeletalMesh>(GetSkinnedAsset());
+					if (SkeletalMeshPtr)
+					{
+						// Make sure the skeletons match!
+						if (BankItem.BankAsset->Asset->GetSkeleton() == SkeletalMeshPtr->GetSkeleton())
+						{
+							bInvalid = false;
+						}
+					}
+				}
+
+				if (bInvalid)
+				{
+					UE_LOG(LogAnimation, Warning, TEXT("Invalid animation skeleton"));
+					BankItem.BankAsset = nullptr;
+				}
+			}
+
+			MarkRenderStateDirty();
+		}
+	}
+	Super::PostEditChangeChainProperty(PropertyChangedEvent);
+}
+
+void UInstancedSkinnedMeshComponent::PostEditUndo()
+{
+	Super::PostEditUndo();
+	MarkRenderStateDirty();
+}
+
+void UInstancedSkinnedMeshComponent::BeginCacheForCookedPlatformData(const ITargetPlatform* TargetPlatform)
+{
+	for (FAnimBankItem& BankItem : AnimBankItems)
+	{
+		if (BankItem.BankAsset)
+		{
+			BankItem.BankAsset->BeginCacheForCookedPlatformData(TargetPlatform);
+		}
+	}
+
+}
+
+bool UInstancedSkinnedMeshComponent::IsCachedCookedPlatformDataLoaded( const ITargetPlatform* TargetPlatform ) 
+{
+	for (FAnimBankItem& BankItem : AnimBankItems)
+	{
+		if (BankItem.BankAsset->IsCompiling())
+		{
+			return false;
+		}
+	}
+
+	return Super::IsCachedCookedPlatformDataLoaded(TargetPlatform);
+}
+
+#endif
+
+TStructOnScope<FActorComponentInstanceData> UInstancedSkinnedMeshComponent::GetComponentInstanceData() const
+{
+	TStructOnScope<FActorComponentInstanceData> ComponentInstanceData;
+#if WITH_EDITOR
+	ComponentInstanceData.InitializeAs<FInstancedSkinnedMeshComponentInstanceData>(this);
+	FInstancedSkinnedMeshComponentInstanceData* SkinnedMeshInstanceData = ComponentInstanceData.Cast<FInstancedSkinnedMeshComponentInstanceData>();
+
+	// Back up per-instance info (this is strictly for Comparison in UInstancedSkinnedMeshComponent::ApplyComponentInstanceData 
+	// as this Property will get serialized by base class FActorComponentInstanceData through FComponentPropertyWriter which uses the PPF_ForceTaggedSerialization to backup all properties even the custom serialized ones
+	SkinnedMeshInstanceData->InstanceData = InstanceData;
+
+	// Back up instance selection
+	SkinnedMeshInstanceData->SelectedInstances = SelectedInstances;
+
+	// Back up per-instance hit proxies
+	SkinnedMeshInstanceData->bHasPerInstanceHitProxies = bHasPerInstanceHitProxies;
+#endif
+	return ComponentInstanceData;
+}
+
+void UInstancedSkinnedMeshComponent::PreApplyComponentInstanceData(struct FInstancedSkinnedMeshComponentInstanceData* InstancedMeshData)
+{
+#if WITH_EDITOR
+	// Prevent proxy recreate while traversing the ::ApplyToComponent stack
+	bIsInstanceDataApplyCompleted = false;
+#endif
+}
+
+void UInstancedSkinnedMeshComponent::ApplyComponentInstanceData(struct FInstancedSkinnedMeshComponentInstanceData* InstancedMeshData)
+{
+#if WITH_EDITOR
+	check(InstancedMeshData);
+
+	ON_SCOPE_EXIT
+	{
+		bIsInstanceDataApplyCompleted = true;
+	};
+
+	if (GetSkinnedAsset() != InstancedMeshData->SkinnedAsset)
+	{
+		return;
+	}
+
+	// If we should inherit from archetype do it here after data was applied and before comparing (RerunConstructionScript will serialize SkipSerialization properties and reapply them even if we want to inherit them)
+	const UInstancedSkinnedMeshComponent* Archetype = Cast<UInstancedSkinnedMeshComponent>(GetArchetype());
+	if (ShouldInheritPerInstanceData(Archetype))
+	{
+		ApplyInheritedPerInstanceData(Archetype);
+	}
+
+	SelectedInstances = InstancedMeshData->SelectedInstances;
+	bHasPerInstanceHitProxies = InstancedMeshData->bHasPerInstanceHitProxies;
+#endif
+}
+
+FBoxSphereBounds UInstancedSkinnedMeshComponent::CalcBounds(const FTransform& LocalToWorld) const
+{
+	//SCOPE_CYCLE_COUNTER(STAT_CalcSkelMeshBounds);
+
+	const USkinnedAsset* SkinnedAssetPtr = GetSkinnedAsset();
+	if (SkinnedAssetPtr && InstanceData.Num() > 0)
+	{
+		const FMatrix BoundTransformMatrix = LocalToWorld.ToMatrixWithScale();
+
+		FBoxSphereBounds::Builder BoundsBuilder;
+
+		const bool bUseSampledBounds = CVarInstancedSkinnedMeshesSampledBounds.GetValueOnAnyThread() != 0;
+		if (bUseSampledBounds && AnimBankItems.Num() > 0)
+		{
+			// Trade per sequence bounds (tigher fitting) for faster builds with high instance counts.
+			const bool bFastBuild = false;
+			if (bFastBuild)
+			{
+				FBox MergedBounds;
+
+				for (const FAnimBankItem& BankItem : AnimBankItems)
+				{
+					if (BankItem.BankAsset != nullptr)
+					{
+						const FAnimBankData& BankData = BankItem.BankAsset->GetData();
+						for (const FAnimBankEntry& BankEntry : BankData.Entries)
+						{
+							MergedBounds += BankEntry.SampledBounds.GetBox();
+						}
+					}
+				}
+
+				if (MergedBounds.IsValid)
+				{
+					for (int32 InstanceIndex = 0; InstanceIndex < InstanceData.Num(); InstanceIndex++)
+					{
+						const FSkinnedMeshInstanceData& Instance = InstanceData[InstanceIndex];
+						BoundsBuilder += MergedBounds.TransformBy(FTransform(Instance.Transform) * LocalToWorld);
+					}
+				}
+			}
+			else
+			{
+				for (int32 InstanceIndex = 0; InstanceIndex < InstanceData.Num(); InstanceIndex++)
+				{
+					const FSkinnedMeshInstanceData& Instance = InstanceData[InstanceIndex];
+					if (Instance.BankIndex < uint32(AnimBankItems.Num()))
+					{
+						const FAnimBankItem& BankItem = AnimBankItems[Instance.BankIndex];
+						if (BankItem.BankAsset != nullptr
+#if WITH_EDITOR
+							&& !BankItem.BankAsset->IsCompiling()
+#endif
+							)
+						{
+							const FAnimBankData& BankData = BankItem.BankAsset->GetData();
+							if (BankItem.SequenceIndex < BankData.Entries.Num())
+							{
+								const FBox BankBounds = BankData.Entries[BankItem.SequenceIndex].SampledBounds.GetBox();
+								BoundsBuilder += BankBounds.TransformBy(FTransform(Instance.Transform) * LocalToWorld);
+							}
+						}
+					}
+				}
+			}
+
+			return BoundsBuilder;
+		}
+		else
+		{
+			const FBox InstanceBounds = SkinnedAssetPtr->GetBounds().GetBox();
+			if (InstanceBounds.IsValid)
+			{
+				for (int32 InstanceIndex = 0; InstanceIndex < InstanceData.Num(); InstanceIndex++)
+				{
+					BoundsBuilder += InstanceBounds.TransformBy(FTransform(InstanceData[InstanceIndex].Transform) * LocalToWorld);
+				}
+
+				return BoundsBuilder;
+			}
+		}
+	}
+
+	return CalcMeshBound(FVector3f::ZeroVector, false, LocalToWorld);
+}
+
+void UInstancedSkinnedMeshComponent::SetSkinnedAssetCallback()
+{
+	MarkRenderStateDirty();
+}
+
+void UInstancedSkinnedMeshComponent::RefreshBoneTransforms(FActorComponentTickFunction* TickFunction)
+{
+	//SCOPE_CYCLE_COUNTER(STAT_RefreshBoneTransforms);
+
+	// Can't do anything without a SkinnedAsset
+	if (!GetSkinnedAsset())
+	{
+		return;
+	}
+
+	// Do nothing more if no bones in skeleton.
+	if (GetNumComponentSpaceTransforms() == 0)
+	{
+		return;
+	}
+
+	// We need the mesh space bone transforms now for renderer to get delta from ref pose:
+	//FillComponentSpaceTransforms();
+	//FinalizeBoneTransform();
+
+	//UpdateChildTransforms();
+	UpdateBounds();
+	MarkRenderTransformDirty();
+	MarkRenderDynamicDataDirty();
+
+	//bNeedsRefreshTransform = false;
+}
+
+FPrimitiveInstanceId UInstancedSkinnedMeshComponent::AddInstance(const FTransform& InstanceTransform, int32 BankIndex, bool bWorldSpace)
+{
+	return AddInstanceInternal(InstanceData.Num(), InstanceTransform, BankIndex, bWorldSpace);
+}
+
+TArray<FPrimitiveInstanceId> UInstancedSkinnedMeshComponent::AddInstances(const TArray<FTransform>& Transforms, const TArray<int32>& BankIndices, bool bShouldReturnIds, bool bWorldSpace)
+{
+	TArray<FPrimitiveInstanceId> NewInstanceIds;
+	if (Transforms.IsEmpty() || (Transforms.Num() != BankIndices.Num()))
+	{
+		return NewInstanceIds;
+	}
+
+	Modify();
+
+	const int32 NumToAdd = Transforms.Num();
+
+	if (bShouldReturnIds)
+	{
+		NewInstanceIds.SetNumUninitialized(NumToAdd);
+	}
+
+	// Reserve memory space
+	const int32 NewNumInstances = InstanceData.Num() + NumToAdd;
+	InstanceData.Reserve(NewNumInstances);
+	InstanceCustomData.Reserve(NumCustomDataFloats * NewNumInstances);
+#if WITH_EDITOR
+	SelectedInstances.Reserve(NewNumInstances);
+#endif
+
+	for (int32 AddIndex = 0; AddIndex < NumToAdd; ++AddIndex)
+	{
+		const FTransform& Transform = Transforms[AddIndex];
+		const int32 BankIndex = BankIndices[AddIndex];
+		FPrimitiveInstanceId InstanceId = AddInstanceInternal(InstanceData.Num(), Transform, BankIndex, bWorldSpace);
+		if (bShouldReturnIds)
+		{
+			NewInstanceIds[AddIndex] = InstanceId;
+		}
+	}
+
+	return NewInstanceIds;
+}
+
+bool UInstancedSkinnedMeshComponent::SetCustomDataValue(FPrimitiveInstanceId InstanceId, int32 CustomDataIndex, float CustomDataValue)
+{
+	int32 InstanceIndex = InstanceDataManager.IdToIndex(InstanceId);
+
+	if (!InstanceData.IsValidIndex(InstanceIndex) || CustomDataIndex < 0 || CustomDataIndex >= NumCustomDataFloats)
+	{
+		return false;
+	}
+
+	Modify();
+
+	InstanceDataManager.CustomDataChanged(InstanceIndex);
+	InstanceCustomData[InstanceIndex * NumCustomDataFloats + CustomDataIndex] = CustomDataValue;
+
+	return true;
+}
+
+bool UInstancedSkinnedMeshComponent::SetCustomData(FPrimitiveInstanceId InstanceId, TArrayView<const float> CustomDataFloats)
+{
+	int32 InstanceIndex = InstanceDataManager.IdToIndex(InstanceId);
+
+	if (!InstanceData.IsValidIndex(InstanceIndex) || CustomDataFloats.Num() == 0)
+	{
+		return false;
+	}
+
+	Modify();
+
+	const int32 NumToCopy = FMath::Min(CustomDataFloats.Num(), NumCustomDataFloats);
+	InstanceDataManager.CustomDataChanged(InstanceIndex);
+	FMemory::Memcpy(&InstanceCustomData[InstanceIndex * NumCustomDataFloats], CustomDataFloats.GetData(), NumToCopy * CustomDataFloats.GetTypeSize());
+	return true;
+}
+
+void UInstancedSkinnedMeshComponent::SetNumCustomDataFloats(int32 InNumCustomDataFloats)
+{
+	if (FMath::Max(InNumCustomDataFloats, 0) != NumCustomDataFloats)
+	{
+		NumCustomDataFloats = FMath::Max(InNumCustomDataFloats, 0);
+	}
+
+	if (InstanceData.Num() * NumCustomDataFloats != InstanceCustomData.Num())
+	{
+		InstanceDataManager.NumCustomDataChanged();
+
+		// Clear out and reinit to 0
+		InstanceCustomData.Empty(InstanceData.Num() * NumCustomDataFloats);
+		InstanceCustomData.SetNumZeroed(InstanceData.Num() * NumCustomDataFloats);
+	}
+}
+
+bool UInstancedSkinnedMeshComponent::GetInstanceTransform(FPrimitiveInstanceId InstanceId, FTransform& OutInstanceTransform, bool bWorldSpace) const
+{
+	int32 InstanceIndex = InstanceDataManager.IdToIndex(InstanceId);
+	if (!InstanceData.IsValidIndex(InstanceIndex))
+	{
+		return false;
+	}
+
+	const FSkinnedMeshInstanceData& Instance = InstanceData[InstanceIndex];
+
+	OutInstanceTransform = FTransform(Instance.Transform);
+	if (bWorldSpace)
+	{
+		OutInstanceTransform = OutInstanceTransform * GetComponentTransform();
+	}
+
+	return true;
+}
+
+bool UInstancedSkinnedMeshComponent::RemoveInstance(FPrimitiveInstanceId InstanceId)
+{
+	int32 InstanceIndex = InstanceDataManager.IdToIndex(InstanceId);
+	if (InstanceData.IsValidIndex(InstanceIndex))
+	{
+		Modify();
+		return RemoveInstanceInternal(InstanceIndex, false);
+	}
+	return false;
+}
+
+void UInstancedSkinnedMeshComponent::RemoveInstances(const TArray<FPrimitiveInstanceId>& InstancesToRemove)
+{
+	Modify();
+
+	for (FPrimitiveInstanceId InstanceId : InstancesToRemove)
+	{
+		int32 InstanceIndex = InstanceDataManager.IdToIndex(InstanceId);
+		RemoveInstanceInternal(InstanceIndex, false);
+	}
+}
+
+void UInstancedSkinnedMeshComponent::ClearInstances()
+{
+	Modify();
+
+	// Clear all the per-instance data
+	InstanceData.Empty();
+	InstanceCustomData.Empty();
+
+#if WITH_EDITOR
+	SelectedInstances.Empty();
+#endif
+	InstanceDataManager.ClearInstances();
+}
+
+int32 UInstancedSkinnedMeshComponent::AddAnimBankItem(const FAnimBankItem& BankItem)
+{
+	int32 ItemIndex = AnimBankItems.Num();
+	AnimBankItems.Emplace(BankItem);
+	return ItemIndex;
+}
+
+struct HSkinnedMeshInstance : public HHitProxy
+{
+	TObjectPtr<UInstancedSkinnedMeshComponent> Component;
+	int32 InstanceIndex;
+
+	DECLARE_HIT_PROXY(ENGINE_API);
+	HSkinnedMeshInstance(UInstancedSkinnedMeshComponent* InComponent, int32 InInstanceIndex)
+	: HHitProxy(HPP_World)
+	, Component(InComponent)
+	, InstanceIndex(InInstanceIndex)
+	{
+	}
+
+	virtual void AddReferencedObjects(FReferenceCollector& Collector) override
+	{
+		Collector.AddReferencedObject(Component);
+	}
+
+	virtual FTypedElementHandle GetElementHandle() const override
+	{
+	#if WITH_EDITOR
+		if (Component)
+		{
+		#if 0
+			if (true)//if (CVarEnableViewportSMInstanceSelection.GetValueOnAnyThread() != 0)
+			{
+				// Prefer per-instance selection if available
+				// This may fail to return a handle if the feature is disabled, or if per-instance editing is disabled for this component
+				if (FTypedElementHandle ElementHandle = UEngineElementsLibrary::AcquireEditorSMInstanceElementHandle(Component, InstanceIndex))
+				{
+					return ElementHandle;
+				}
+			}
+		#endif
+
+			// If per-instance selection isn't possible, fallback to general per-component selection (which may choose to select the owner actor instead)
+			return UEngineElementsLibrary::AcquireEditorComponentElementHandle(Component);
+		}
+	#endif	// WITH_EDITOR
+		return FTypedElementHandle();
+	}
+
+	virtual EMouseCursor::Type GetMouseCursor() override
+	{
+		return EMouseCursor::Crosshairs;
+	}
+};
+
+IMPLEMENT_HIT_PROXY(HSkinnedMeshInstance, HHitProxy);
+
+void UInstancedSkinnedMeshComponent::CreateHitProxyData(TArray<TRefCountPtr<HHitProxy>>& HitProxies)
+{
+	if (GIsEditor && bHasPerInstanceHitProxies)
+	{
+		int32 NumProxies = InstanceData.Num();
+		HitProxies.Empty(NumProxies);
+
+		for (int32 InstanceIdx = 0; InstanceIdx < NumProxies; ++InstanceIdx)
+		{
+			HitProxies.Add(new HSkinnedMeshInstance(this, InstanceIdx));
+		}
+	}
+	else
+	{
+		HitProxies.Empty();
+	}
+}
+
+FPrimitiveInstanceId UInstancedSkinnedMeshComponent::AddInstanceInternal(int32 InstanceIndex, const FTransform& InstanceTransform, int32 BankIndex, bool bWorldSpace)
+{
+	// This happens because the editor modifies the InstanceData array _before_ callbacks. If we could change the UI to not do that we could remove this ugly hack.
+	if (!InstanceData.IsValidIndex(InstanceIndex))
+	{
+		check(InstanceIndex == InstanceData.Num());
+		InstanceData.AddDefaulted();
+	}
+
+	FPrimitiveInstanceId InstanceId = InstanceDataManager.Add(InstanceIndex);
+
+	const FTransform3f LocalTransform = FTransform3f(bWorldSpace ? InstanceTransform.GetRelativeTransform(GetComponentTransform()) : InstanceTransform);
+	SetupNewInstanceData(InstanceData[InstanceIndex], InstanceIndex, LocalTransform, BankIndex);
+
+	// Add custom data to instance
+	InstanceCustomData.AddZeroed(NumCustomDataFloats);
+
+#if WITH_EDITOR
+	SelectedInstances.Add(false);
+#endif
+
+	return InstanceId;
+}
+
+bool UInstancedSkinnedMeshComponent::RemoveInstanceInternal(int32 InstanceIndex, bool bInstanceAlreadyRemoved)
+{
+	if (!ensure(bInstanceAlreadyRemoved || InstanceData.IsValidIndex(InstanceIndex)))
+	{
+		return false;
+	}
+	InstanceDataManager.RemoveAtSwap(InstanceIndex);
+	
+	// remove instance
+	if (!bInstanceAlreadyRemoved)
+	{
+		InstanceData.RemoveAtSwap(InstanceIndex, EAllowShrinking::No);
+	}
+	
+	if (InstanceCustomData.IsValidIndex(InstanceIndex * NumCustomDataFloats))
+	{
+		InstanceCustomData.RemoveAtSwap(InstanceIndex * NumCustomDataFloats, NumCustomDataFloats, EAllowShrinking::No);
+	}
+
+#if WITH_EDITOR
+	// remove selection flag if array is filled in
+	if (SelectedInstances.IsValidIndex(InstanceIndex))
+	{
+		SelectedInstances.RemoveAtSwap(InstanceIndex);
+	}
+#endif
+	return true;
+}
diff --git a/Engine/Source/Runtime/Engine/Private/ObjectCacheContext.cpp b/Engine/Source/Runtime/Engine/Private/ObjectCacheContext.cpp
--- a/Engine/Source/Runtime/Engine/Private/ObjectCacheContext.cpp
+++ b/Engine/Source/Runtime/Engine/Private/ObjectCacheContext.cpp
@@ -16,7 +16,7 @@
 #include "HAL/LowLevelMemStats.h"
 #include "Misc/ScopeRWLock.h"
 #include "Components/ComponentInterfaceIterator.h"
-#include "Components/AnimBankMeshComponent.h"
+#include "Components/InstancedSkinnedMeshComponent.h"
 
 #define LOCTEXT_NAMESPACE "ObjectCache"
 
@@ -832,16 +832,16 @@ TObjectCacheIterator<USkinnedMeshComponent> FObjectCacheContext::GetSkinnedMeshC
 		[](USkinnedMeshComponent* Component, int32 ) { return Component->GetSkinnedAsset(); });
 }
 
-TObjectCacheIterator<UAnimBankMeshComponent> FObjectCacheContext::GetAnimBankMeshComponents()
+TObjectCacheIterator<UInstancedSkinnedMeshComponent> FObjectCacheContext::GetInstancedSkinnedMeshComponents()
 {
-	return AnimBankMeshCache.Get<TObjectIterator<UAnimBankMeshComponent>>();
+	return InstancedSkinnedMeshCache.Get<TObjectIterator<UInstancedSkinnedMeshComponent>>();
 }
 
-TObjectCacheIterator<UAnimBankMeshComponent> FObjectCacheContext::GetAnimBankMeshComponents(UAnimBank* InAnimBank)
+TObjectCacheIterator<UInstancedSkinnedMeshComponent> FObjectCacheContext::GetInstancedSkinnedMeshComponents(UAnimBank* InAnimBank)
 {
-	return AnimBankMeshCache.Get<TObjectIterator<UAnimBankMeshComponent>>(InAnimBank, 
-		[](UAnimBankMeshComponent* Component){ return Component->GetAnimBankItems().Num();  },
-		[](UAnimBankMeshComponent* Component, int32 Index)->const UAnimBank* { return Component->GetAnimBankItems()[Index].BankAsset; }
+	return InstancedSkinnedMeshCache.Get<TObjectIterator<UInstancedSkinnedMeshComponent>>(InAnimBank,
+		[](UInstancedSkinnedMeshComponent* Component){ return Component->GetAnimBankItems().Num();  },
+		[](UInstancedSkinnedMeshComponent* Component, int32 Index)->const UAnimBank* { return Component->GetAnimBankItems()[Index].BankAsset; }
 	);
 }
 
diff --git a/Engine/Source/Runtime/Engine/Public/ObjectCacheContext.h b/Engine/Source/Runtime/Engine/Public/ObjectCacheContext.h
--- a/Engine/Source/Runtime/Engine/Public/ObjectCacheContext.h
+++ b/Engine/Source/Runtime/Engine/Public/ObjectCacheContext.h
@@ -20,7 +20,7 @@ class USkeletalMesh;
 class USkinnedAsset;
 class IPrimitiveComponent;
 class IStaticMeshComponent;
-class UAnimBankMeshComponent;
+class UInstancedSkinnedMeshComponent;
 class UAnimBank;
 
 /**
@@ -75,8 +75,8 @@ public:
 	ENGINE_API TObjectCacheIterator<IStaticMeshComponent>  GetStaticMeshComponents();
 	ENGINE_API TObjectCacheIterator<USkinnedMeshComponent> GetSkinnedMeshComponents();
 	ENGINE_API TObjectCacheIterator<USkinnedMeshComponent> GetSkinnedMeshComponents(USkinnedAsset* InSkinnedAsset);
-	ENGINE_API TObjectCacheIterator<UAnimBankMeshComponent> GetAnimBankMeshComponents();
-	ENGINE_API TObjectCacheIterator<UAnimBankMeshComponent> GetAnimBankMeshComponents(UAnimBank* InAnimBank);
+	ENGINE_API TObjectCacheIterator<UInstancedSkinnedMeshComponent> GetInstancedSkinnedMeshComponents();
+	ENGINE_API TObjectCacheIterator<UInstancedSkinnedMeshComponent> GetInstancedSkinnedMeshComponents(UAnimBank* InAnimBank);
 	ENGINE_API TObjectCacheIterator<IStaticMeshComponent>  GetStaticMeshComponents(UStaticMesh* InStaticMesh);
 	ENGINE_API TObjectCacheIterator<UMaterialInterface>    GetMaterialsAffectedByTexture(UTexture* InTexture);
 	ENGINE_API TObjectCacheIterator<UMaterialInterface>    GetMaterialsAffectedByTextureCollection(UTextureCollection* InTextureCollection);
@@ -120,7 +120,7 @@ private:
 	};
 
 	TComponentTypeCache<USkinnedMeshComponent, USkinnedAsset> SkinnedMeshCache;
-	TComponentTypeCache<UAnimBankMeshComponent, UAnimBank> AnimBankMeshCache;
+	TComponentTypeCache<UInstancedSkinnedMeshComponent, UAnimBank> InstancedSkinnedMeshCache;
 	TOptional<TArray<IPrimitiveComponent*>> PrimitiveComponents;
 	TComponentTypeCache<IStaticMeshComponent, UStaticMesh> StaticMeshCache;
 };
diff --git a/Engine/Source/Runtime/Renderer/Private/Skinning/AnimBankMeshDebugView.cpp b/Engine/Source/Runtime/Renderer/Private/Skinning/AnimBankMeshDebugView.cpp
deleted file mode 100644
diff --git a/Engine/Source/Runtime/Renderer/Private/Skinning/AnimBankMeshDebugView.h b/Engine/Source/Runtime/Renderer/Private/Skinning/AnimBankMeshDebugView.h
deleted file mode 100644
diff --git a/Engine/Source/Runtime/Renderer/Private/Skinning/SkinnedMeshDebugView.cpp b/Engine/Source/Runtime/Renderer/Private/Skinning/SkinnedMeshDebugView.cpp
new file mode 100644
--- /dev/null
+++ b/Engine/Source/Runtime/Renderer/Private/Skinning/SkinnedMeshDebugView.cpp
@@ -0,0 +1,229 @@
+// Copyright Epic Games, Inc. All Rights Reserved.
+
+#include "SkinnedMeshDebugView.h"
+#include "Components/InstancedSkinnedMeshComponent.h"
+
+#include "GlobalRenderResources.h"
+#include "GlobalShader.h"
+#include "RenderGraphBuilder.h"
+#include "RenderGraphResources.h"
+#include "RenderGraphUtils.h"
+#include "TextureResource.h"
+
+#include "SceneView.h"
+#include "SceneInterface.h"
+#include "ScenePrivate.h"
+
+#include "PostProcess/PostProcessing.h"
+#include "PostProcess/PostProcessMaterialInputs.h"
+#include "Nanite/NaniteSkinningSceneExtension.h"
+
+#include "AnimationRuntime.h"
+
+TAutoConsoleVariable<bool> CVarSkinnedMeshDebugDraw(
+	TEXT("r.InstancedSkinnedMeshes.DebugDraw"), false,
+	TEXT("Whether to enable instanced skinned mesh debug draw."),
+	ECVF_RenderThreadSafe
+);
+
+class FSkinnedMeshDrawLineVS : public FGlobalShader
+{
+	DECLARE_GLOBAL_SHADER(FSkinnedMeshDrawLineVS);
+	SHADER_USE_PARAMETER_STRUCT(FSkinnedMeshDrawLineVS, FGlobalShader);
+
+	static bool ShouldCompilePermutation(const FGlobalShaderPermutationParameters& Parameters)
+	{
+		// TODO: Have an animation bank test
+		return DoesPlatformSupportNanite(Parameters.Platform);
+	}
+
+	static void ModifyCompilationEnvironment(const FGlobalShaderPermutationParameters& Parameters, FShaderCompilerEnvironment& OutEnvironment)
+	{
+		FGlobalShader::ModifyCompilationEnvironment(Parameters, OutEnvironment);
+		OutEnvironment.SetDefine(TEXT("VF_SUPPORTS_PRIMITIVE_SCENE_DATA"), 1);
+		OutEnvironment.CompilerFlags.Add(CFLAG_HLSL2021);
+	}
+
+	BEGIN_SHADER_PARAMETER_STRUCT(FParameters, )
+		SHADER_PARAMETER_STRUCT_REF(FViewUniformShaderParameters, View)
+		SHADER_PARAMETER_RDG_UNIFORM_BUFFER(FSceneUniformParameters, Scene)
+		//SHADER_PARAMETER_RDG_UNIFORM_BUFFER(FBatchedPrimitiveParameters, BatchedPrimitive)
+		SHADER_PARAMETER(uint32, PersistentPrimitiveIndex)
+	END_SHADER_PARAMETER_STRUCT()
+};
+
+class FSkinnedMeshDrawLinePS : public FGlobalShader
+{
+	DECLARE_GLOBAL_SHADER(FSkinnedMeshDrawLinePS);
+	SHADER_USE_PARAMETER_STRUCT(FSkinnedMeshDrawLinePS, FGlobalShader);
+
+	static bool ShouldCompilePermutation(const FGlobalShaderPermutationParameters& Parameters)
+	{
+		// TODO: Have an animation bank test
+		return DoesPlatformSupportNanite(Parameters.Platform);
+	}
+
+	static void ModifyCompilationEnvironment(const FGlobalShaderPermutationParameters& Parameters, FShaderCompilerEnvironment& OutEnvironment)
+	{
+		FGlobalShader::ModifyCompilationEnvironment(Parameters, OutEnvironment);
+		OutEnvironment.CompilerFlags.Add(CFLAG_HLSL2021);
+	}
+
+	BEGIN_SHADER_PARAMETER_STRUCT(FParameters, )
+		RENDER_TARGET_BINDING_SLOTS()
+	END_SHADER_PARAMETER_STRUCT()
+};
+
+BEGIN_SHADER_PARAMETER_STRUCT(FSkinnedMeshDrawLineParameters, )
+	SHADER_PARAMETER_STRUCT_INCLUDE(FSkinnedMeshDrawLineVS::FParameters, VSParameters)
+	SHADER_PARAMETER_STRUCT_INCLUDE(FSkinnedMeshDrawLinePS::FParameters, PSParameters)
+END_SHADER_PARAMETER_STRUCT()
+
+IMPLEMENT_GLOBAL_SHADER(FSkinnedMeshDrawLineVS, "/Engine/Private/Skinning/SkinnedMeshDebug.usf", "LineVS", SF_Vertex);
+IMPLEMENT_GLOBAL_SHADER(FSkinnedMeshDrawLinePS, "/Engine/Private/Skinning/SkinnedMeshDebug.usf", "LinePS", SF_Pixel);
+
+#if UE_ENABLE_DEBUG_DRAWING
+
+FSkinnedMeshDebugViewExtension::FSkinnedMeshDebugViewExtension(const FAutoRegister& AutoRegister)
+: FSceneViewExtensionBase(AutoRegister)
+{
+}
+
+void FSkinnedMeshDebugViewExtension::SubscribeToPostProcessingPass(EPostProcessingPass PassId, const FSceneView& View, FAfterPassCallbackDelegateArray& InOutPassCallbacks, bool bIsPassEnabled)
+{
+	if (PassId == EPostProcessingPass::VisualizeDepthOfField)
+	{
+		if (CVarSkinnedMeshDebugDraw.GetValueOnAnyThread())
+		{
+			InOutPassCallbacks.Add(FAfterPassCallbackDelegate::CreateRaw(this, &FSkinnedMeshDebugViewExtension::PostProcessPass_RenderThread));
+		}
+	}
+}
+
+FScreenPassTexture FSkinnedMeshDebugViewExtension::PostProcessPass_RenderThread(FRDGBuilder& GraphBuilder, const FSceneView& View, const FPostProcessMaterialInputs& InOutInputs)
+{
+	FScreenPassTextureSlice SceneColorSlice = InOutInputs.GetInput(EPostProcessMaterialInput::SceneColor);
+	check(SceneColorSlice.IsValid());
+
+	FScreenPassTexture SceneColor(SceneColorSlice);
+
+	FScreenPassRenderTarget Output = InOutInputs.OverrideOutput;
+
+	// If the override output is provided, it means that this is the last pass in post processing.
+	if (!Output.IsValid())
+	{
+		Output = FScreenPassRenderTarget(SceneColor.Texture, SceneColor.ViewRect, View.GetOverwriteLoadAction());
+	}
+
+	RenderSkeletons(GraphBuilder, View, Output);
+
+	return MoveTemp(Output);
+}
+
+void FSkinnedMeshDebugViewExtension::RenderSkeletons(FRDGBuilder& GraphBuilder, const FSceneView& View, const FScreenPassRenderTarget& Output)
+{
+	if (View.Family == nullptr || View.Family->Scene == nullptr)
+	{
+		return;
+	}
+
+	const FScene* Scene = View.Family->Scene->GetRenderScene();
+	if (Scene == nullptr)
+	{
+		return;
+	}
+
+	const Nanite::FSkinningSceneExtension* SkinningExtension = Scene->GetExtensionPtr<Nanite::FSkinningSceneExtension>();
+	if (SkinningExtension == nullptr)
+	{
+		return;
+	}
+
+	TArray<FPrimitiveSceneInfo*> SkinnedPrimitives;
+	SkinningExtension->GetSkinnedPrimitives(SkinnedPrimitives);
+
+	TArray<FSkinnedMeshPrimitive> Primitives;
+	Primitives.Reserve(SkinnedPrimitives.Num());
+
+	for (const FPrimitiveSceneInfo* SceneInfo : SkinnedPrimitives)
+	{
+		if (SceneInfo == nullptr || !SceneInfo->GetPersistentIndex().IsValid() || SceneInfo->Proxy == nullptr)
+		{
+			continue;
+		}
+
+		if (!SceneInfo->Proxy->IsSkinnedMesh() || !SceneInfo->Proxy->IsNaniteMesh())
+		{
+			continue;
+		}
+
+		auto* SkinnedProxy = static_cast<Nanite::FSkinnedSceneProxy*>(SceneInfo->Proxy);
+
+		FSkinnedMeshPrimitive& Primitive = Primitives.AddDefaulted_GetRef();
+		Primitive.Index = SceneInfo->GetPersistentIndex();
+		Primitive.BoneCount = SkinnedProxy->GetMaxBoneTransformCount();
+		Primitive.InstanceCount = SceneInfo->GetNumInstanceSceneDataEntries();
+	}
+
+	if (Primitives.Num() == 0)
+	{
+		return;
+	}
+
+	const FGlobalShaderMap* ShaderMap = GetGlobalShaderMap(View.GetFeatureLevel());
+
+	TShaderMapRef<FSkinnedMeshDrawLineVS> VertexShader(ShaderMap);
+	TShaderMapRef<FSkinnedMeshDrawLinePS> PixelShader(ShaderMap);
+
+	if (VertexShader.IsNull() || PixelShader.IsNull())
+	{
+		return;
+	}
+
+	const FIntRect ViewRect = Output.ViewRect;
+
+	TRDGUniformBufferRef<FSceneUniformParameters> SceneUniformParameters = GetSceneUniformBufferRef(GraphBuilder, View);
+
+	FSkinnedMeshDrawLineParameters* PassParameters = GraphBuilder.AllocParameters<FSkinnedMeshDrawLineParameters>();
+	PassParameters->VSParameters.View = View.ViewUniformBuffer;
+	PassParameters->VSParameters.Scene = SceneUniformParameters;
+	PassParameters->PSParameters.RenderTargets[0] = Output.GetRenderTargetBinding();
+
+	GraphBuilder.AddPass(
+		RDG_EVENT_NAME("RenderSkinnedMeshDebug"),
+		PassParameters,
+		ERDGPassFlags::Raster,
+		[VertexShader, PixelShader, PassParameters, ViewRect = ViewRect, Primitives = MoveTemp(Primitives)](FRDGAsyncTask, FRHICommandList& RHICmdList)
+		{
+			FGraphicsPipelineStateInitializer GraphicsPSOInit;
+			RHICmdList.ApplyCachedRenderTargets(GraphicsPSOInit);
+			RHICmdList.SetViewport(ViewRect.Min.X, ViewRect.Min.Y, 0.0f, ViewRect.Max.X, ViewRect.Max.Y, 1.0f);
+
+			GraphicsPSOInit.DepthStencilState	= TStaticDepthStencilState<false, CF_Always>::GetRHI();
+			GraphicsPSOInit.BlendState			= TStaticBlendState<CW_RGBA, BO_Add, BF_SourceAlpha, BF_InverseSourceAlpha, BO_Add, BF_One, BF_Zero>::GetRHI();
+			GraphicsPSOInit.RasterizerState		= TStaticRasterizerState<FM_Solid, CM_None>::GetRHI();
+			GraphicsPSOInit.PrimitiveType		= PT_LineList;
+
+			GraphicsPSOInit.BoundShaderState.VertexDeclarationRHI	= GEmptyVertexDeclaration.VertexDeclarationRHI;
+			GraphicsPSOInit.BoundShaderState.VertexShaderRHI		= VertexShader.GetVertexShader();
+			GraphicsPSOInit.BoundShaderState.PixelShaderRHI			= PixelShader.GetPixelShader();
+
+			SetGraphicsPipelineState(RHICmdList, GraphicsPSOInit, 0);
+
+			for (const FSkinnedMeshPrimitive& Primitive : Primitives)
+			{
+				const uint32 BoneCount = Primitive.BoneCount;
+				const uint32 InstanceCount = Primitive.InstanceCount;
+
+				check(Primitive.Index.Index != INDEX_NONE);
+				PassParameters->VSParameters.PersistentPrimitiveIndex = uint32(Primitive.Index.Index);
+
+				SetShaderParameters(RHICmdList, VertexShader, VertexShader.GetVertexShader(), PassParameters->VSParameters);
+				SetShaderParameters(RHICmdList, PixelShader, PixelShader.GetPixelShader(), PassParameters->PSParameters);
+				RHICmdList.DrawPrimitive(0, BoneCount, InstanceCount);
+			}
+		}
+	);
+}
+
+#endif
diff --git a/Engine/Source/Runtime/Renderer/Private/Skinning/SkinnedMeshDebugView.h b/Engine/Source/Runtime/Renderer/Private/Skinning/SkinnedMeshDebugView.h
new file mode 100644
--- /dev/null
+++ b/Engine/Source/Runtime/Renderer/Private/Skinning/SkinnedMeshDebugView.h
@@ -0,0 +1,51 @@
+// Copyright Epic Games, Inc. All Rights Reserved.
+
+#pragma once
+
+#include "CoreMinimal.h"
+#include "SceneViewExtension.h"
+#include "RenderGraphFwd.h"
+#include "MeshPassProcessor.h"
+
+#if UE_ENABLE_DEBUG_DRAWING
+
+struct FSkinnedMeshPrimitive
+{
+	FPersistentPrimitiveIndex Index;
+	uint32 BoneCount = 0;
+	uint32 InstanceCount = 0;
+};
+
+class FSkinnedMeshDebugViewExtension : public FSceneViewExtensionBase
+{
+public:
+	FSkinnedMeshDebugViewExtension(const FAutoRegister& AutoRegister);
+
+	virtual void SubscribeToPostProcessingPass(EPostProcessingPass PassId, const FSceneView& View, FAfterPassCallbackDelegateArray& InOutPassCallbacks, bool bIsPassEnabled) override;
+
+	FScreenPassTexture PostProcessPass_RenderThread(FRDGBuilder& GraphBuilder, const FSceneView& View, const struct FPostProcessMaterialInputs& InOutInputs);
+
+	void RenderSkeletons(class FRDGBuilder& GraphBuilder, const class FSceneView& View, const struct FScreenPassRenderTarget& Output);
+
+private:
+	inline static TSharedPtr<FSkinnedMeshDebugViewExtension, ESPMode::ThreadSafe> Instance;
+
+	struct FStaticConstructor
+	{
+		FStaticConstructor()
+		{
+			FCoreDelegates::OnPostEngineInit.AddLambda([]()
+			{ 
+				Instance = FSceneViewExtensions::NewExtension<FSkinnedMeshDebugViewExtension>();
+			});
+			FCoreDelegates::OnEnginePreExit.AddLambda([]()
+			{ 
+				Instance = nullptr;
+			});
+		}
+	};
+
+	static inline FStaticConstructor StaticConstructor;
+};
+
+#endif
