[AutoRTFM] Disable memory validation for (another) UObject constructor.

The reinitialization of a UObject can upset the validator.

diff --git a/Engine/Source/Runtime/CoreUObject/Private/UObject/UObjectGlobals.cpp b/Engine/Source/Runtime/CoreUObject/Private/UObject/UObjectGlobals.cpp
--- a/Engine/Source/Runtime/CoreUObject/Private/UObject/UObjectGlobals.cpp
+++ b/Engine/Source/Runtime/CoreUObject/Private/UObject/UObjectGlobals.cpp
@@ -3839,7 +3839,8 @@ void UObject::PostInitProperties()
 	FOverridableManager::Get().ClearOverrides(*this);
 }
 
-// UObjects can be reconstructed after use, which can upset the validator.
+// Construct the UObject in the open with memory validation disabled as UObjects can be 
+// reconstructed, and the re-initialization of the object can upset the AutoRTFM memory validator.
 UE_AUTORTFM_ALWAYS_OPEN_NO_MEMORY_VALIDATION
 UObject::UObject()
 {
@@ -3857,7 +3858,9 @@ UObject::UObject()
 	const_cast<FObjectInitializer&>(ObjectInitializer).FinalizeSubobjectClassInitialization();
 }
 
-UE_AUTORTFM_ALWAYS_OPEN
+// Construct the UObject in the open with memory validation disabled as UObjects can be 
+// reconstructed, and the re-initialization of the object can upset the AutoRTFM memory validator.
+UE_AUTORTFM_ALWAYS_OPEN_NO_MEMORY_VALIDATION
 UObject::UObject(const FObjectInitializer& ObjectInitializer)
 {
 #if UE_WITH_REMOTE_OBJECT_HANDLE
diff --git a/Engine/Source/Runtime/Engine/Private/Tests/AutoRTFM/AutoRTFMActorComponentTests.cpp b/Engine/Source/Runtime/Engine/Private/Tests/AutoRTFM/AutoRTFMActorComponentTests.cpp
--- a/Engine/Source/Runtime/Engine/Private/Tests/AutoRTFM/AutoRTFMActorComponentTests.cpp
+++ b/Engine/Source/Runtime/Engine/Private/Tests/AutoRTFM/AutoRTFMActorComponentTests.cpp
@@ -11,6 +11,7 @@
 #include "AutoRTFMTestLevel.h"
 #include "AutoRTFMTestObject.h"
 #include "AutoRTFMTestPrimitiveComponent.h"
+#include "AutoRTFMTestCameraShake.h"
 #include "Chaos/Core.h"
 #include "Components/SkeletalMeshComponent.h"
 #include "Engine/LevelStreamingPersistent.h"
@@ -125,6 +126,7 @@ protected:
 			FString Desc = FString::Printf(TEXT("SKIPPED test '%s'. AutoRTFM disabled."), GetData(TestName));
 			ExecutionInfo.AddEvent(FAutomationEvent(EAutomationEventType::Info, Desc));
 		}
+
 		return true;
 	}
 
@@ -514,6 +516,27 @@ AUTORTFM_ACTOR_COMPONENT_TEST(ReconstructObject)
 	TestTrueExpr(123 == Old->ActorCategory);
 }
 
+AUTORTFM_ACTOR_COMPONENT_TEST(ReconstructCameraShake) // SOL-7529
+{
+	FName Name = TEXT("ShakyMcShakeface");
+	UObject* Outer = static_cast<UObject*>(GetTransientPackage());
+	UAutoRTFMTestCameraShake* Old = NewObject<UAutoRTFMTestCameraShake>(Outer, Name);
+	TestTrueExpr(1.0 == Old->ShakeScale);
+	Old->ShakeScale = 123;
+
+	UAutoRTFMTestCameraShake* New = nullptr;
+
+	AutoRTFM::Testing::Abort([&]
+		{
+			// Reconstruct the object by using the same name, then immediately abort.
+			New = NewObject<UAutoRTFMTestCameraShake>(Outer, Name);
+			AutoRTFM::AbortTransaction();
+		});
+
+	TestTrueExpr(nullptr == New);
+	TestTrueExpr(123 == Old->ShakeScale);
+}
+
 }  // anonymous namespace
 
 #undef AUTORTFM_ACTOR_COMPONENT_TEST
diff --git a/Engine/Source/Runtime/Engine/Private/Tests/AutoRTFM/AutoRTFMTestCameraShake.h b/Engine/Source/Runtime/Engine/Private/Tests/AutoRTFM/AutoRTFMTestCameraShake.h
new file mode 100644
--- /dev/null
+++ b/Engine/Source/Runtime/Engine/Private/Tests/AutoRTFM/AutoRTFMTestCameraShake.h
@@ -0,0 +1,13 @@
+// Copyright Epic Games, Inc. All Rights Reserved.
+
+#pragma once
+
+#include "Camera/CameraShakeBase.h"
+#include "AutoRTFMTestCameraShake.generated.h"
+
+UCLASS(NotBlueprintable, HideDropdown)
+class UAutoRTFMTestCameraShake : public UCameraShakeBase
+{
+public:
+	GENERATED_BODY()
+};
