Marked ShouldComponentIgnoreHitResult as virtual

diff --git a/Engine/Source/Runtime/Engine/Classes/Components/PrimitiveComponent.h b/Engine/Source/Runtime/Engine/Classes/Components/PrimitiveComponent.h
--- a/Engine/Source/Runtime/Engine/Classes/Components/PrimitiveComponent.h
+++ b/Engine/Source/Runtime/Engine/Classes/Components/PrimitiveComponent.h
@@ -1086,7 +1086,7 @@ public:
 	FMaskFilter GetMoveIgnoreMask() const { return MoveIgnoreMask; }
 
 	/** Should the hit result be ignored based on this component */
-	ENGINE_API bool ShouldComponentIgnoreHitResult(FHitResult const& TestHit, EMoveComponentFlags MoveFlags);
+	ENGINE_API virtual bool ShouldComponentIgnoreHitResult(FHitResult const& TestHit, EMoveComponentFlags MoveFlags);
 
 	/** Set the mask filter checked when others move into us. */
 	void SetMaskFilterOnBodyInstance(FMaskFilter InMaskFilter) { BodyInstance.SetMaskFilter(InMaskFilter); }
