Feature request: make UTextRenderComponent blueprintable

#rnx

diff --git a/Engine/Source/Runtime/Engine/Classes/Components/TextRenderComponent.h b/Engine/Source/Runtime/Engine/Classes/Components/TextRenderComponent.h
--- a/Engine/Source/Runtime/Engine/Classes/Components/TextRenderComponent.h
+++ b/Engine/Source/Runtime/Engine/Classes/Components/TextRenderComponent.h
@@ -40,7 +40,7 @@ enum EVerticalTextAligment : int
 /**
  * Renders text in the world with given font. Contains usual font related attributes such as Scale, Alignment, Color etc.
  */
-UCLASS(ClassGroup=Rendering, hidecategories=(Object,LOD,Physics,TextureStreaming,Activation,"Components|Activation",Collision), editinlinenew, meta=(BlueprintSpawnableComponent = ""), MinimalAPI)
+UCLASS(Blueprintable, ClassGroup=Rendering, hidecategories=(Object,LOD,Physics,TextureStreaming,Activation,"Components|Activation",Collision), editinlinenew, meta=(BlueprintSpawnableComponent), MinimalAPI)
 class UTextRenderComponent : public UPrimitiveComponent
 {
 	GENERATED_UCLASS_BODY()
