Implement External Collisions for ChaosClothAsset.

diff --git a/Engine/Plugins/ChaosClothAsset/Source/ChaosClothAssetEngine/Private/ChaosClothAsset/ClothComponent.cpp b/Engine/Plugins/ChaosClothAsset/Source/ChaosClothAssetEngine/Private/ChaosClothAsset/ClothComponent.cpp
--- a/Engine/Plugins/ChaosClothAsset/Source/ChaosClothAssetEngine/Private/ChaosClothAsset/ClothComponent.cpp
+++ b/Engine/Plugins/ChaosClothAsset/Source/ChaosClothAssetEngine/Private/ChaosClothAsset/ClothComponent.cpp
@@ -21,15 +21,16 @@ CSV_DECLARE_CATEGORY_MODULE_EXTERN(ENGINE_API, Animation);
 
 UChaosClothComponent::UChaosClothComponent(const FObjectInitializer& ObjectInitializer)
 	: Super(ObjectInitializer)
-#if WITH_EDITORONLY_DATA
-	, bSimulateInEditor(0)
-#endif
 	, bUseAttachedParentAsPoseComponent(1)  // By default use the parent component as leader pose component
 	, bWaitForParallelTask(0)
 	, bEnableSimulation(1)
 	, bSuspendSimulation(0)
 	, bBindToLeaderComponent(0)
-	, CollisionSources(MakeUnique<UE::Chaos::ClothAsset::FCollisionSources>(this))
+	, bCollideWithEnvironment(0)
+#if WITH_EDITORONLY_DATA
+	, bSimulateInEditor(0)
+#endif
+	, CollisionSources(MakeUnique<UE::Chaos::ClothAsset::FCollisionSources>(this, (bool)bCollideWithEnvironment))
 {
 	PrimaryComponentTick.EndTickGroup = TG_PostPhysics;
 }
@@ -183,6 +184,10 @@ PRAGMA_ENABLE_DEPRECATION_WARNINGS
 		{
 			bTickInEditor = bSimulateInEditor;
 		}
+		if(Property->GetFName() == GET_MEMBER_NAME_CHECKED(UChaosClothComponent, bCollideWithEnvironment))
+		{
+			CollisionSources->SetCollideWithEnvironment(bCollideWithEnvironment);
+		}
 	}
 
 	Super::PostEditChangeProperty(PropertyChangedEvent);
@@ -227,6 +232,9 @@ void UChaosClothComponent::OnRegister()
 	// Update render visibility, so that an empty LODs doesn't unnecessarily go to render
 	UpdateVisibility();
 
+	// Update CollisionSources
+	CollisionSources->SetCollideWithEnvironment(bCollideWithEnvironment);
+
 	// Register the dataflow simulation interface
 	UE::Dataflow::RegisterSimulationInterface(this);
 }
@@ -566,6 +574,12 @@ void UChaosClothComponent::ResetCollisionSources()
 	CollisionSources->Reset();
 }
 
+void UChaosClothComponent::SetCollideWithEnvironment(bool bCollide)
+{
+	bCollideWithEnvironment = bCollide;
+	CollisionSources->SetCollideWithEnvironment(bCollideWithEnvironment);
+}
+
 void UChaosClothComponent::StartNewParallelSimulation(float DeltaTime)
 {
 	if (ClothSimulationProxy.IsValid())
diff --git a/Engine/Plugins/ChaosClothAsset/Source/ChaosClothAssetEngine/Private/ChaosClothAsset/CollisionSources.cpp b/Engine/Plugins/ChaosClothAsset/Source/ChaosClothAssetEngine/Private/ChaosClothAsset/CollisionSources.cpp
--- a/Engine/Plugins/ChaosClothAsset/Source/ChaosClothAssetEngine/Private/ChaosClothAsset/CollisionSources.cpp
+++ b/Engine/Plugins/ChaosClothAsset/Source/ChaosClothAssetEngine/Private/ChaosClothAsset/CollisionSources.cpp
@@ -4,12 +4,14 @@
 #include "ChaosCloth/ChaosClothingSimulationCollider.h"
 #include "Components/SkinnedMeshComponent.h"
 #include "Engine/SkinnedAsset.h"
+#include "PhysicsEngine/EnvironmentalCollisions.h"
 #include "PhysicsEngine/PhysicsAsset.h"
 
 namespace UE::Chaos::ClothAsset
 {
-	FCollisionSources::FCollisionSources(USkinnedMeshComponent* InOwnerComponent)
+	FCollisionSources::FCollisionSources(USkinnedMeshComponent* InOwnerComponent, bool bInCollideWithEnvironment)
 		: OwnerComponent(InOwnerComponent)
+		, bCollideWithEnvironment(bInCollideWithEnvironment)
 	{}
 
 	FCollisionSources::~FCollisionSources() = default;
@@ -95,7 +97,7 @@ namespace UE::Chaos::ClothAsset
 		++Version;
 	}
 
-	void FCollisionSources::ExtractCollisionData(FClothCollisionData& CollisionData)
+	void FCollisionSources::ExtractNonEnvironmentalCollisionData(FClothCollisionData& CollisionData)
 	{
 		CollisionData.Reset();
 		if (OwnerComponent.IsValid())
@@ -104,15 +106,42 @@ namespace UE::Chaos::ClothAsset
 			{
 				CollisionSource.ExtractCollisionData(*OwnerComponent, CollisionData);
 			}
+
+		}
+	}
+
+	void FCollisionSources::ExtractEnvironmentalCollisionData(FClothCollisionData& CollisionData)
+	{
+		CollisionData.Reset();
+		if (OwnerComponent.IsValid())
+		{
+			if (bCollideWithEnvironment)
+			{
+				FEnvironmentalCollisions::AppendCollisionDataFromEnvironment(OwnerComponent.Get(), CollisionData);
+			}
 		}
 	}
 
 	void FCollisionSourcesProxy::ExtractCollisionData()
 	{
+		bool bCollisionDataChanged = false;
 		if (Version != CollisionSources.Version)
 		{
-			CollisionSources.ExtractCollisionData(CollisionData);
+			CollisionSources.ExtractNonEnvironmentalCollisionData(NonEnvironmentalCollisionData);
 			Version = CollisionSources.Version;
+			bCollisionDataChanged = true;
+		}
+
+		// EnvironmentalCollisionData changes every frame, so if we had some before, it's now stale.
+		bCollisionDataChanged |= !EnvironmentalCollisionData.IsEmpty();
+
+		CollisionSources.ExtractEnvironmentalCollisionData(EnvironmentalCollisionData);
+		bCollisionDataChanged |= !EnvironmentalCollisionData.IsEmpty();
+
+		if (bCollisionDataChanged)
+		{
+			CollisionData = NonEnvironmentalCollisionData;
+			CollisionData.Append(EnvironmentalCollisionData);
 		}
 	}
 
diff --git a/Engine/Plugins/ChaosClothAsset/Source/ChaosClothAssetEngine/Private/ChaosClothAsset/CollisionSources.h b/Engine/Plugins/ChaosClothAsset/Source/ChaosClothAssetEngine/Private/ChaosClothAsset/CollisionSources.h
--- a/Engine/Plugins/ChaosClothAsset/Source/ChaosClothAssetEngine/Private/ChaosClothAsset/CollisionSources.h
+++ b/Engine/Plugins/ChaosClothAsset/Source/ChaosClothAssetEngine/Private/ChaosClothAsset/CollisionSources.h
@@ -17,7 +17,7 @@ namespace UE::Chaos::ClothAsset
 	class FCollisionSources final
 	{
 	public:
-		explicit FCollisionSources(USkinnedMeshComponent* InOwnerComponent);
+		explicit FCollisionSources(USkinnedMeshComponent* InOwnerComponent, bool bInCollideWithEnvironment = false);
 		~FCollisionSources();
 
 		void Add(USkinnedMeshComponent* SourceComponent, const UPhysicsAsset* SourcePhysicsAsset, bool bUseSphylsOnly = false);
@@ -28,6 +28,11 @@ namespace UE::Chaos::ClothAsset
 
 		void Reset();
 
+		void SetCollideWithEnvironment(bool bCollide)
+		{
+			bCollideWithEnvironment = bCollide;
+		}
+
 	private:
 		friend class FCollisionSourcesProxy;
 
@@ -51,11 +56,13 @@ namespace UE::Chaos::ClothAsset
 			void ExtractCollisionData(const USkinnedMeshComponent& InOwnerComponent, FClothCollisionData& CollisionData);
 		};
 
-		void ExtractCollisionData(FClothCollisionData& CollisionData);
+		void ExtractNonEnvironmentalCollisionData(FClothCollisionData& CollisionData);
+		void ExtractEnvironmentalCollisionData(FClothCollisionData& CollisionData);
 
 		const TWeakObjectPtr<USkinnedMeshComponent> OwnerComponent;
 		TArray<FCollisionSource> CollisionSources;
 		int32 Version = INDEX_NONE;
+		bool bCollideWithEnvironment = false;
 	};
 
 	/**
@@ -74,6 +81,9 @@ namespace UE::Chaos::ClothAsset
 
 	protected:
 		FCollisionSources& CollisionSources;
+		FClothCollisionData NonEnvironmentalCollisionData;
+		FClothCollisionData EnvironmentalCollisionData;
+
 		FClothCollisionData CollisionData;
 		int32 Version = INDEX_NONE;
 	};
diff --git a/Engine/Plugins/ChaosClothAsset/Source/ChaosClothAssetEngine/Public/ChaosClothAsset/ClothComponent.h b/Engine/Plugins/ChaosClothAsset/Source/ChaosClothAssetEngine/Public/ChaosClothAsset/ClothComponent.h
--- a/Engine/Plugins/ChaosClothAsset/Source/ChaosClothAssetEngine/Public/ChaosClothAsset/ClothComponent.h
+++ b/Engine/Plugins/ChaosClothAsset/Source/ChaosClothAssetEngine/Public/ChaosClothAsset/ClothComponent.h
@@ -129,6 +129,15 @@ public:
 	/** Return all collision sources currently assigned to this component. */
 	UE::Chaos::ClothAsset::FCollisionSources& GetCollisionSources() const { return *CollisionSources; }
 
+	/** Set whether or not to collide with the environment. */
+	UFUNCTION(BlueprintCallable, Category = "ClothComponent", Meta = (Keywords = "Chaos Cloth Collide Environment"))
+	void SetCollideWithEnvironment(bool bCollide);
+
+	/** Set whether or not to collision with the environment is enabled. */
+	UFUNCTION(BlueprintCallable, Category = "ClothComponent", Meta = (Keywords = "Chaos Cloth Collide Environment"))
+	bool GetCollideWithEnvironment() const { return bCollideWithEnvironment; }
+
+
 	/**
 	 * Return the property collections holding the runtime properties for this cloth component (one per LOD).
 	 * This might be different from the cloth asset's since the component's properties can be modified in code or in blueprints.
@@ -220,20 +229,19 @@ private:
 
 	friend UE::Chaos::ClothAsset::FClothComponentCacheAdapter;
 
-#if WITH_EDITORONLY_DATA
-	/** Cloth asset used by this component. */
-	UE_DEPRECATED(5.1, "This property isn't deprecated, but getter and setter must be used at all times to preserve correct operations.")
-	UPROPERTY(EditAnywhere, Transient, Setter = SetClothAsset, BlueprintSetter = SetClothAsset, Getter = GetClothAsset, BlueprintGetter = GetClothAsset, Category = ClothComponent)
-	TObjectPtr<UChaosClothAsset> ClothAsset;
-
-	/** Whether to run the simulation in editor. */
-	UPROPERTY(EditInstanceOnly, Transient, Category = ClothComponent)
-	uint8 bSimulateInEditor : 1;
-#endif
-
 	/* Solver dataflow asset used to advance in time */
 	UPROPERTY(EditAnywhere, Category = ClothComponent, meta=(EditConditionHides), AdvancedDisplay)
 	FDataflowSimulationAsset SimulationAsset;
+
+	/** Blend amount between the skinned (=0) and the simulated pose (=1). */
+	UPROPERTY(Interp, Category = ClothComponent)
+	float BlendWeight = 1.f;
+
+	/** This scale is applied to all cloth geometry (e.g., cloth meshes and collisions) in order to simulate in a different scale space than world.This scale is not applied to distance-based simulation parameters such as MaxDistance.
+	* This property is currently only read by the cloth solver when creating cloth actors, but may become animatable in the future.
+	*/
+	UPROPERTY(EditAnywhere, Category = ClothComponent, meta = (UIMin = 0.0, UIMax = 10.0, ClampMin = 0.0, ClampMax = 10000.0))
+	float ClothGeometryScale = 1.f;
 	
 	/** If enabled, and the parent is another Skinned Mesh Component (e.g. another Cloth Component, Poseable Mesh Component, Skeletal Mesh Component, ...etc.), use its pose. */
 	UPROPERTY(EditAnywhere, Category = ClothComponent)
@@ -263,18 +271,19 @@ private:
 	UPROPERTY(Interp, Category = ClothComponent)
 	uint8 bReset : 1;
 
-	/** Blend amount between the skinned (=0) and the simulated pose (=1). */
-	UPROPERTY(Interp, Category = ClothComponent)
-	float BlendWeight = 1.f;
+	/** Collide with the environment. */
+	UPROPERTY(EditAnywhere, Category = ClothComponent)
+	uint8 bCollideWithEnvironment : 1;
 
-	/** This scale is applied to all cloth geometry (e.g., cloth meshes and collisions) in order to simulate in a different scale space than world.This scale is not applied to distance-based simulation parameters such as MaxDistance.
-	* This property is currently only read by the cloth solver when creating cloth actors, but may become animatable in the future.
-	*/
-	UPROPERTY(EditAnywhere, Category = ClothComponent, meta = (UIMin = 0.0, UIMax = 10.0, ClampMin = 0.0, ClampMax = 10000.0))
-	float ClothGeometryScale = 1.f;
+#if WITH_EDITORONLY_DATA
+	/** Whether to run the simulation in editor. */
+	UPROPERTY(EditInstanceOnly, Transient, Category = ClothComponent)
+	uint8 bSimulateInEditor : 1;
 
-#if WITH_EDITOR
-	bool bTickOnceInEditor = false;
+	/** Cloth asset used by this component. */
+	UE_DEPRECATED(5.1, "This property isn't deprecated, but getter and setter must be used at all times to preserve correct operations.")
+	UPROPERTY(EditAnywhere, Transient, Setter = SetClothAsset, BlueprintSetter = SetClothAsset, Getter = GetClothAsset, BlueprintGetter = GetClothAsset, Category = ClothComponent)
+	TObjectPtr<UChaosClothAsset> ClothAsset;
 #endif
 
 	UPROPERTY(Transient)
@@ -290,4 +299,8 @@ private:
 
 	// External sources for collision
 	TUniquePtr<UE::Chaos::ClothAsset::FCollisionSources> CollisionSources;
+
+#if WITH_EDITOR
+	bool bTickOnceInEditor = false;
+#endif
 };
diff --git a/Engine/Source/Runtime/ClothingSystemRuntimeInterface/Public/ClothCollisionData.h b/Engine/Source/Runtime/ClothingSystemRuntimeInterface/Public/ClothCollisionData.h
--- a/Engine/Source/Runtime/ClothingSystemRuntimeInterface/Public/ClothCollisionData.h
+++ b/Engine/Source/Runtime/ClothingSystemRuntimeInterface/Public/ClothCollisionData.h
@@ -18,6 +18,11 @@ struct FClothCollisionData
 
 	CLOTHINGSYSTEMRUNTIMEINTERFACE_API void AppendTransformed(const FClothCollisionData& InOther, const TArray<FTransform>& BoneTransforms);
 
+	bool IsEmpty() const
+	{
+		return Spheres.IsEmpty() && SphereConnections.IsEmpty() && Convexes.IsEmpty() && Boxes.IsEmpty();
+	}
+
 	// Sphere data
 	UPROPERTY(EditAnywhere, Category = Collison)
 	TArray<FClothCollisionPrim_Sphere> Spheres;
diff --git a/Engine/Source/Runtime/ClothingSystemRuntimeInterface/Public/ClothingSimulationInterface.h b/Engine/Source/Runtime/ClothingSystemRuntimeInterface/Public/ClothingSimulationInterface.h
--- a/Engine/Source/Runtime/ClothingSystemRuntimeInterface/Public/ClothingSimulationInterface.h
+++ b/Engine/Source/Runtime/ClothingSystemRuntimeInterface/Public/ClothingSimulationInterface.h
@@ -40,6 +40,7 @@ protected:
 
 	friend class USkeletalMeshComponent;
 	friend class FParallelClothTask;
+	friend struct FEnvironmentalCollisions;
 
 	/**
 	 * Create an actor for this simulation from the data in InAsset
diff --git a/Engine/Source/Runtime/Engine/Classes/PhysicsEngine/EnvironmentalCollisions.h b/Engine/Source/Runtime/Engine/Classes/PhysicsEngine/EnvironmentalCollisions.h
new file mode 100644
--- /dev/null
+++ b/Engine/Source/Runtime/Engine/Classes/PhysicsEngine/EnvironmentalCollisions.h
@@ -0,0 +1,11 @@
+// Copyright Epic Games, Inc. All Rights Reserved.
+#pragma once
+#include "EngineDefines.h"
+
+struct FClothCollisionData;
+class USceneComponent;
+
+struct FEnvironmentalCollisions
+{
+	static ENGINE_API void AppendCollisionDataFromEnvironment(const USceneComponent* SceneComponent, FClothCollisionData& CollisionData);
+};
diff --git a/Engine/Source/Runtime/Engine/Private/PhysicsEngine/EnvironmentalCollisions.cpp b/Engine/Source/Runtime/Engine/Private/PhysicsEngine/EnvironmentalCollisions.cpp
new file mode 100644
--- /dev/null
+++ b/Engine/Source/Runtime/Engine/Private/PhysicsEngine/EnvironmentalCollisions.cpp
@@ -0,0 +1,355 @@
+// Copyright Epic Games, Inc. All Rights Reserved.
+
+#include "PhysicsEngine/EnvironmentalCollisions.h"
+#include "Chaos/Capsule.h"
+#include "Chaos/Convex.h"
+#include "Chaos/ImplicitFwd.h"
+#include "Chaos/ImplicitObjectScaled.h"
+#include "Chaos/PhysicsObjectInterface.h"
+#include "Components/PrimitiveComponent.h"
+#include "Components/SceneComponent.h"
+#include "Components/SkeletalMeshComponent.h"
+#include "Engine/OverlapResult.h"
+#include "Engine/World.h"
+#include "Misc/ConfigCacheIni.h"
+#include "ClothCollisionData.h"
+
+void FEnvironmentalCollisions::AppendCollisionDataFromEnvironment(const USceneComponent* SceneComponent, FClothCollisionData& CollisionData)
+{
+
+	// Read config on first call
+	struct FEnvironmentCollisionConfig
+	{
+		float Padding = 2.f;  // Extra padding added to the bounds so that the collision can still be detected after being resolved
+		float Thickness = 2.f;  // Extra thickness added to edgy collision shapes (cubes & convexes)
+		int32 MaxShapes = 32;  // Limit the number of extracted shapes per component as these collisions are very expensive
+
+		FEnvironmentCollisionConfig()
+		{
+			if (GConfig)
+			{
+				GConfig->GetFloat(TEXT("ClothSettings"), TEXT("EnvironmentCollisionPadding"), Padding, GEngineIni);
+				GConfig->GetFloat(TEXT("ClothSettings"), TEXT("EnvironmentCollisionThickness"), Thickness, GEngineIni);
+				GConfig->GetInt(TEXT("ClothSettings"), TEXT("EnvironmentCollisionMaxShapes"), MaxShapes, GEngineIni);
+			}
+		}
+	};
+	static FEnvironmentCollisionConfig EnvironmentCollisionConfig;
+
+	TArray<FOverlapResult> Overlaps;
+
+	FCollisionObjectQueryParams ObjectParams;
+
+	ObjectParams.AddObjectTypesToQuery(ECollisionChannel::ECC_WorldStatic);
+	// to collide with other clothing objects
+	ObjectParams.AddObjectTypesToQuery(ECollisionChannel::ECC_PhysicsBody);
+
+	FCollisionQueryParams Params(SCENE_QUERY_STAT(ClothOverlapComponents), false);
+	const FBoxSphereBounds PaddedBounds = SceneComponent->Bounds.ExpandBy(EnvironmentCollisionConfig.Padding);
+
+	SceneComponent->GetWorld()->OverlapMultiByObjectType(Overlaps, PaddedBounds.Origin, FQuat::Identity, ObjectParams, FCollisionShape::MakeBox(PaddedBounds.BoxExtent), Params);
+
+	for (int32 OverlapIdx = 0; OverlapIdx < Overlaps.Num(); ++OverlapIdx)
+	{
+		const TWeakObjectPtr<UPrimitiveComponent>& Component = Overlaps[OverlapIdx].Component;
+		if (Component.IsValid())
+		{
+			ECollisionChannel Channel = Component->GetCollisionObjectType();
+
+			if (Channel == ECollisionChannel::ECC_WorldStatic)
+			{
+				// Static world geo
+				if (!Component->BodyInstance.IsValidBodyInstance())
+				{
+					// Move to next component, this one has no valid physics
+					continue;
+				}
+
+				bool bSuccessfulRead = false;
+
+				const FPhysicsActorHandle& ActorRef = Component->BodyInstance.GetActorReferenceWithWelding();
+				FPhysicsCommand::ExecuteRead(ActorRef, [&](const FPhysicsActorHandle& Actor)
+					{
+						TArray<FPhysicsShapeHandle> AllShapes;
+						const int32 NumSyncShapes = Component->BodyInstance.GetAllShapes_AssumesLocked(AllShapes);
+
+						if (NumSyncShapes == 0 || NumSyncShapes > EnvironmentCollisionConfig.MaxShapes)
+						{
+							// Either no shapes or too complicated to consider
+							return;
+						}
+
+						using namespace Chaos;
+
+						const FReal Thickness = (FReal)EnvironmentCollisionConfig.Thickness;
+
+						auto AddSphere = [&CollisionData](const FImplicitSphere3& ImplicitSphere, const FTransform& ComponentToClothTransform, const FVec3& Scale = FVec3::OneVector)
+							{
+								CollisionData.Spheres.Emplace(
+									ImplicitSphere.GetRadiusf() * Scale.X,  // Assumes uniform scale
+									ComponentToClothTransform.TransformPosition(FVec3(ImplicitSphere.GetCenterf())));
+							};
+
+						auto AddBox = [&CollisionData, Thickness](const FImplicitBox3& ImplicitBox, const FTransform& ComponentToClothTransform, const FVec3& Scale = FVec3::OneVector)
+							{
+								CollisionData.Boxes.Emplace(
+									ComponentToClothTransform.TransformPosition(ImplicitBox.GetCenter()),
+									ComponentToClothTransform.GetRotation(),
+									ImplicitBox.Extents() * Scale * (FReal)0.5 + Thickness);
+							};
+
+						auto AddCapsule = [&CollisionData](const FImplicitCapsule3& ImplicitCapsule, const FTransform& ComponentToClothTransform, const FVec3& Scale = FVec3::OneVector)
+							{
+								const int32 BaseSphereIndex = CollisionData.Spheres.Num();
+
+								const FReal Radius = ImplicitCapsule.GetRadiusf() * Scale.X;  // Assumes uniform scale
+								const FReal HalfHeight = ImplicitCapsule.GetHeightf() * (FReal)0.5;
+								const FVector HalfSegment = ComponentToClothTransform.GetUnitAxis(EAxis::X) * HalfHeight * Scale.X;  // Assumes uniform scale
+								const FVector TransformedCenter = ComponentToClothTransform.TransformPosition(FVec3(ImplicitCapsule.GetCenterf()));
+
+								CollisionData.Spheres.Emplace(Radius, TransformedCenter + HalfSegment);
+								CollisionData.Spheres.Emplace(Radius, TransformedCenter - HalfSegment);
+
+								CollisionData.SphereConnections.Emplace(BaseSphereIndex, BaseSphereIndex + 1);
+							};
+
+						auto AddConvex = [&CollisionData, Thickness](const FImplicitConvex3& ImplicitConvex, const FMatrix& ComponentToClothMatrix, const FVec3* const Scale = nullptr)
+							{
+								TArray<FClothCollisionPrim_ConvexFace> Faces;
+								const int32 NumPlanes = ImplicitConvex.NumPlanes();
+								Faces.SetNum(NumPlanes);
+
+								TArray<FVector> SurfacePoints;
+								const int32 NumSurfacePoints = ImplicitConvex.NumVertices();
+								SurfacePoints.SetNumUninitialized(NumSurfacePoints);
+
+								TArray<TArray<int32, TInlineAllocator<4>>, TInlineAllocator<16>> PointFaces;
+								PointFaces.SetNum(NumSurfacePoints);
+
+								for (int32 FaceIndex = 0; FaceIndex < NumPlanes; ++FaceIndex)
+								{
+									const TPlaneConcrete<FReal, 3>& Plane = ImplicitConvex.GetPlane(FaceIndex);
+									const FVec3& Normal = Plane.Normal();
+									const FVec3 Base = Plane.X() + Normal * (FReal)Thickness;
+
+									Faces[FaceIndex].Plane = FPlane(FVector(Base), FVector(Normal)).TransformBy(ComponentToClothMatrix);
+
+									const int32 NumFaceIndices = ImplicitConvex.NumPlaneVertices(FaceIndex);
+									Faces[FaceIndex].Indices.SetNumUninitialized(NumFaceIndices);
+
+									for (int32 Index = 0; Index < NumFaceIndices; ++Index)
+									{
+										const int32 PointIndex = ImplicitConvex.GetPlaneVertex(FaceIndex, Index);
+										Faces[FaceIndex].Indices[Index] = PointIndex;
+
+										PointFaces[PointIndex].Add(FaceIndex);
+									}
+								}
+
+								for (int32 PointIndex = 0; PointIndex < NumSurfacePoints; ++PointIndex)
+								{
+									check(PointFaces[PointIndex].Num() >= 3);
+									const int32 Index0 = PointFaces[PointIndex][0];
+									const int32 Index1 = PointFaces[PointIndex][1];
+									const int32 Index2 = PointFaces[PointIndex][2];
+
+									if (!FMath::IntersectPlanes3(SurfacePoints[PointIndex], Faces[Index0].Plane, Faces[Index1].Plane, Faces[Index2].Plane))
+									{
+										SurfacePoints[PointIndex] = ComponentToClothMatrix.TransformPosition(FVector(ImplicitConvex.GetVertex(PointIndex)));
+									}
+								}
+
+								CollisionData.Convexes.Emplace(MoveTemp(Faces), MoveTemp(SurfacePoints));
+							};
+
+						FTransform ClothComponentTransform = SceneComponent->GetComponentTransform();
+						ClothComponentTransform.RemoveScaling();  // The environment collision shape doesn't need the scale of the cloth skeletal mesh applied to it (but it does need the source scale from its component transform)
+						const FTransform ComponentToClothBaseTransform = Component->GetComponentTransform() * ClothComponentTransform.Inverse();
+
+						bool bHasSimpleCollision = false;
+
+						for (FPhysicsShapeHandle& ShapeHandle : AllShapes)
+						{
+							FTransform ComponentToClothTransform = ComponentToClothBaseTransform;
+
+							const FImplicitObject* ImplicitObject = &ShapeHandle.GetGeometry();
+							EImplicitObjectType ImplicitType = ImplicitObject->GetType();
+
+							// Transformed implicits
+							if (ImplicitType == ImplicitObjectType::Transformed)
+							{
+								const TImplicitObjectTransformed<FReal, 3>& ImplicitTransformed = ImplicitObject->GetObjectChecked<TImplicitObjectTransformed<FReal, 3>>();
+								ImplicitObject = ImplicitTransformed.GetTransformedObject();
+								ImplicitType = ImplicitObject->GetType();
+
+								ComponentToClothTransform = ImplicitTransformed.GetTransform() * ComponentToClothTransform;
+								UE_LOG(LogSkeletalMesh, Verbose, TEXT("Found transformed environmental collision"));
+							}
+
+							switch (ImplicitType)
+							{
+								// Base implicits
+							case ImplicitObjectType::Sphere:
+							{
+								const FImplicitSphere3& ImplicitSphere = ImplicitObject->GetObjectChecked<FImplicitSphere3>();;
+								AddSphere(ImplicitSphere, ComponentToClothTransform);
+							}
+							UE_LOG(LogSkeletalMesh, Verbose, TEXT("Found Sphere cloth environmental collision in [%s]"), *Component->GetOwner()->GetFName().ToString());
+							bHasSimpleCollision = true;
+							break;
+							case ImplicitObjectType::Box:
+							{
+								const FImplicitBox3& ImplicitBox = ImplicitObject->GetObjectChecked<FImplicitBox3>();
+								AddBox(ImplicitBox, ComponentToClothTransform);
+							}
+							UE_LOG(LogSkeletalMesh, Verbose, TEXT("Found Box cloth environmental collision in [%s]"), *Component->GetOwner()->GetFName().ToString());
+							bHasSimpleCollision = true;
+							break;
+							case ImplicitObjectType::Capsule:
+							{
+								const FImplicitCapsule3& ImplicitCapsule = ImplicitObject->GetObjectChecked<FImplicitCapsule3>();
+								AddCapsule(ImplicitCapsule, ComponentToClothTransform);
+								UE_LOG(LogSkeletalMesh, Verbose, TEXT("Found Capsule cloth environmental collision in [%s]"), *Component->GetOwner()->GetFName().ToString());
+								bHasSimpleCollision = true;
+							}
+							break;
+							case ImplicitObjectType::Convex:
+							{
+								const FImplicitConvex3& ImplicitConvex = ImplicitObject->GetObjectChecked<FImplicitConvex3>();
+								AddConvex(ImplicitConvex, ComponentToClothTransform.ToMatrixNoScale());
+								UE_LOG(LogSkeletalMesh, Verbose, TEXT("Found Convex cloth environmental collision in [%s]"), *Component->GetOwner()->GetFName().ToString());
+								bHasSimpleCollision = true;
+							}
+							break;
+
+							// Instanced implicits
+							case ImplicitObjectType::IsInstanced | ImplicitObjectType::Sphere:
+							{
+								const TImplicitObjectInstanced<FImplicitSphere3>& ImplicitInstanced = ImplicitObject->GetObjectChecked<TImplicitObjectInstanced<FImplicitSphere3>>();
+								check(ImplicitInstanced.Object());
+								const FImplicitSphere3& ImplicitSphere = *ImplicitInstanced.GetInstancedObject();
+								AddSphere(ImplicitSphere, ComponentToClothTransform);
+								UE_LOG(LogSkeletalMesh, Verbose, TEXT("Found Instanced Sphere cloth environmental collision in [%s]"), *Component->GetOwner()->GetFName().ToString());
+								bHasSimpleCollision = true;
+							}
+							break;
+							case ImplicitObjectType::IsInstanced | ImplicitObjectType::Box:
+							{
+								const TImplicitObjectInstanced<FImplicitBox3>& ImplicitInstanced = ImplicitObject->GetObjectChecked<TImplicitObjectInstanced<FImplicitBox3>>();
+								check(ImplicitInstanced.Object());
+								const FImplicitBox3& ImplicitBox = *ImplicitInstanced.GetInstancedObject();
+								AddBox(ImplicitBox, ComponentToClothTransform);
+								UE_LOG(LogSkeletalMesh, Verbose, TEXT("Found Instanced Box cloth environmental collision in [%s]"), *Component->GetOwner()->GetFName().ToString());
+								bHasSimpleCollision = true;
+							}
+							break;
+							case ImplicitObjectType::IsInstanced | ImplicitObjectType::Capsule:
+							{
+								const TImplicitObjectInstanced<FImplicitCapsule3>& ImplicitInstanced = ImplicitObject->GetObjectChecked<TImplicitObjectInstanced<FImplicitCapsule3>>();
+								check(ImplicitInstanced.Object());
+								const FImplicitCapsule3& ImplicitCapsule = *ImplicitInstanced.GetInstancedObject();
+								AddCapsule(ImplicitCapsule, ComponentToClothTransform);
+								UE_LOG(LogSkeletalMesh, Verbose, TEXT("Found Instanced Capsule cloth environmental collision in [%s]"), *Component->GetOwner()->GetFName().ToString());
+								bHasSimpleCollision = true;
+							}
+							break;
+							case ImplicitObjectType::IsInstanced | ImplicitObjectType::Convex:
+							{
+								const TImplicitObjectInstanced<FImplicitConvex3>& ImplicitInstanced = ImplicitObject->GetObjectChecked<TImplicitObjectInstanced<FImplicitConvex3>>();
+								check(ImplicitInstanced.Object());
+								const FImplicitConvex3& ImplicitConvex = *ImplicitInstanced.GetInstancedObject();
+								AddConvex(ImplicitConvex, ComponentToClothTransform.ToMatrixNoScale());
+								UE_LOG(LogSkeletalMesh, Verbose, TEXT("Found Instanced Convex cloth environmental collision in [%s]"), *Component->GetOwner()->GetFName().ToString());
+								bHasSimpleCollision = true;
+							}
+							break;
+
+							// Scaled implicits
+							case ImplicitObjectType::IsScaled | ImplicitObjectType::Sphere:
+							{
+								const TImplicitObjectScaled<FImplicitSphere3>& ImplicitScaled = ImplicitObject->GetObjectChecked<TImplicitObjectScaled<FImplicitSphere3>>();
+								check(ImplicitScaled.Object());
+								const FImplicitSphere3& ImplicitSphere = *ImplicitScaled.GetUnscaledObject();
+								ensure(FVector::DistSquared(ComponentToClothTransform.GetScale3D(), FVector(ImplicitScaled.GetScale())) < UE_KINDA_SMALL_NUMBER);
+								AddSphere(ImplicitSphere, ComponentToClothTransform, ImplicitScaled.GetScale());
+								UE_LOG(LogSkeletalMesh, Verbose, TEXT("Found Scaled Sphere cloth environmental collision in [%s]"), *Component->GetOwner()->GetFName().ToString());
+								bHasSimpleCollision = true;
+							}
+							break;
+							case ImplicitObjectType::IsScaled | ImplicitObjectType::Box:
+							{
+								const TImplicitObjectScaled<FImplicitBox3>& ImplicitScaled = ImplicitObject->GetObjectChecked<TImplicitObjectScaled<FImplicitBox3>>();
+								check(ImplicitScaled.Object());
+								const FImplicitBox3& ImplicitBox = *ImplicitScaled.GetUnscaledObject();
+								ensure(FVector::DistSquared(ComponentToClothTransform.GetScale3D(), FVector(ImplicitScaled.GetScale())) < UE_KINDA_SMALL_NUMBER);
+								AddBox(ImplicitBox, ComponentToClothTransform, ImplicitScaled.GetScale());
+								UE_LOG(LogSkeletalMesh, Verbose, TEXT("Found Scaled Box cloth environmental collision in [%s]"), *Component->GetOwner()->GetFName().ToString());
+								bHasSimpleCollision = true;
+							}
+							break;
+							case ImplicitObjectType::IsScaled | ImplicitObjectType::Capsule:
+							{
+								const TImplicitObjectScaled<FImplicitCapsule3>& ImplicitScaled = ImplicitObject->GetObjectChecked<TImplicitObjectScaled<FImplicitCapsule3>>();
+								check(ImplicitScaled.Object());
+								const FImplicitCapsule3& ImplicitCapsule = *ImplicitScaled.GetUnscaledObject();
+								ensure(FVector::DistSquared(ComponentToClothTransform.GetScale3D(), FVector(ImplicitScaled.GetScale())) < UE_KINDA_SMALL_NUMBER);
+								AddCapsule(ImplicitCapsule, ComponentToClothTransform, ImplicitScaled.GetScale());
+								UE_LOG(LogSkeletalMesh, Verbose, TEXT("Found Scaled Capsule cloth environmental collision in [%s]"), *Component->GetOwner()->GetFName().ToString());
+								bHasSimpleCollision = true;
+							}
+							break;
+							case ImplicitObjectType::IsScaled | ImplicitObjectType::Convex:
+							{
+								const TImplicitObjectScaled<FImplicitConvex3>& ImplicitScaled = ImplicitObject->GetObjectChecked<TImplicitObjectScaled<FImplicitConvex3>>();
+								check(ImplicitScaled.Object());
+								const FImplicitConvex3& ImplicitConvex = *ImplicitScaled.GetUnscaledObject();
+								ensure(FVector::DistSquared(ComponentToClothTransform.GetScale3D(), FVector(ImplicitScaled.GetScale())) < UE_KINDA_SMALL_NUMBER);
+								AddConvex(ImplicitConvex, ComponentToClothTransform.ToMatrixWithScale());
+								UE_LOG(LogSkeletalMesh, Verbose, TEXT("Found Scaled Convex cloth environmental collision in [%s]"), *Component->GetOwner()->GetFName().ToString());
+								bHasSimpleCollision = true;
+							}
+							break;
+
+							// Triangle mesh
+							case ImplicitObjectType::TriangleMesh:
+							case ImplicitObjectType::IsInstanced | ImplicitObjectType::TriangleMesh:
+							case ImplicitObjectType::IsScaled | ImplicitObjectType::TriangleMesh:
+								// TODO: We could eventually want to collide cloth against triangle meshes,
+								//       however the concept of simple vs complex shape might need to be clarified
+								//       as it currently iterates over all shape to discard the triangle mesh ones.
+								UE_LOG(LogSkeletalMesh, Verbose, TEXT("Found unusable Triangle Mesh cloth environmental collision in [%s]"), !Component->GetOwner() ? TEXT("Unknown") : *Component->GetOwner()->GetFName().ToString());
+								break;
+
+							default:
+								UE_LOG(LogSkeletalMesh, Verbose, TEXT("Found unsupported collision type during environmental collision with the cloth in [%s]"), !Component->GetOwner() ? TEXT("Unknown") : *Component->GetOwner()->GetFName().ToString());
+								break;
+							}
+						}
+						bSuccessfulRead = true;
+					});
+			}
+			else if (Channel == ECollisionChannel::ECC_PhysicsBody)
+			{
+				// Possibly a skeletal mesh, extract it's clothing collisions if necessary
+				USkeletalMeshComponent* SkelComp = Cast<USkeletalMeshComponent>(Component.Get());
+
+				if (SkelComp && SkelComp->GetSkeletalMeshAsset())
+				{
+					if ((UPrimitiveComponent*)SkelComp == (UPrimitiveComponent*)SceneComponent)
+					{
+						// Same mesh, move to next component
+						continue;
+					}
+
+					if (SkelComp->GetClothingSimulation())
+					{
+						// append skeletal component collisions
+						FClothCollisionData SkelCollisionData;
+						SkelComp->GetClothingSimulation()->GetCollisions(SkelCollisionData, false);
+						CollisionData.Append(SkelCollisionData);
+					}
+				}
+			}
+		}
+	}
+}
diff --git a/Engine/Source/Runtime/Engine/Private/SkeletalMeshComponentPhysics.cpp b/Engine/Source/Runtime/Engine/Private/SkeletalMeshComponentPhysics.cpp
--- a/Engine/Source/Runtime/Engine/Private/SkeletalMeshComponentPhysics.cpp
+++ b/Engine/Source/Runtime/Engine/Private/SkeletalMeshComponentPhysics.cpp
@@ -12,6 +12,7 @@
 #include "Engine/World.h"
 #include "Modules/ModuleManager.h"
 #include "Physics/Experimental/PhysScene_Chaos.h"
+#include "PhysicsEngine/EnvironmentalCollisions.h"
 #include "PhysicsProxy/SingleParticlePhysicsProxy.h"
 #include "AnimationRuntime.h"
 #include "ClothCollisionData.h"
@@ -3094,25 +3095,6 @@ void USkeletalMeshComponent::CopyChildrenClothCollisionsToParent()
 
 void USkeletalMeshComponent::ProcessClothCollisionWithEnvironment()
 {
-	// Read config on first call
-	struct FEnvironmentCollisionConfig
-	{
-		float Padding = 2.f;  // Extra padding added to the bounds so that the collision can still be detected after being resolved
-		float Thickness = 2.f;  // Extra thickness added to edgy collision shapes (cubes & convexes)
-		int32 MaxShapes = 32;  // Limit the number of extracted shapes per component as these collisions are very expensive
-
-		FEnvironmentCollisionConfig()
-		{
-			if (GConfig)
-			{
-				GConfig->GetFloat(TEXT("ClothSettings"), TEXT("EnvironmentCollisionPadding"), Padding, GEngineIni);
-				GConfig->GetFloat(TEXT("ClothSettings"), TEXT("EnvironmentCollisionThickness"), Thickness, GEngineIni);
-				GConfig->GetInt(TEXT("ClothSettings"), TEXT("EnvironmentCollisionMaxShapes"), MaxShapes, GEngineIni);
-			}
-		}
-	};
-	static FEnvironmentCollisionConfig EnvironmentCollisionConfig;
-
 	// don't handle collision detection if this component is in editor
 	if(!GetWorld()->IsGameWorld() || !ClothingSimulation)
 	{
@@ -3120,324 +3102,7 @@ void USkeletalMeshComponent::ProcessClothCollisionWithEnvironment()
 	}
 
 	FClothCollisionData NewCollisionData;
-
-	TArray<FOverlapResult> Overlaps;
-
-	FCollisionObjectQueryParams ObjectParams;
-
-	ObjectParams.AddObjectTypesToQuery(ECollisionChannel::ECC_WorldStatic);
-	// to collide with other clothing objects
-	ObjectParams.AddObjectTypesToQuery(ECollisionChannel::ECC_PhysicsBody);
-
-	FCollisionQueryParams Params(SCENE_QUERY_STAT(ClothOverlapComponents), false);
-	const FBoxSphereBounds PaddedBounds = Bounds.ExpandBy(EnvironmentCollisionConfig.Padding);
-
-	GetWorld()->OverlapMultiByObjectType(Overlaps, PaddedBounds.Origin, FQuat::Identity, ObjectParams, FCollisionShape::MakeBox(PaddedBounds.BoxExtent), Params);
-
-	for (int32 OverlapIdx=0; OverlapIdx<Overlaps.Num(); ++OverlapIdx)
-	{
-		const TWeakObjectPtr<UPrimitiveComponent>& Component = Overlaps[OverlapIdx].Component;
-		if (Component.IsValid())
-		{ 
-			ECollisionChannel Channel = Component->GetCollisionObjectType();
-
-			if(Channel == ECollisionChannel::ECC_WorldStatic)
-			{
-				// Static world geo
-				if(!Component->BodyInstance.IsValidBodyInstance())
-				{
-					// Move to next component, this one has no valid physics
-					continue;
-				}
-
-				bool bSuccessfulRead = false;
-
-				const FPhysicsActorHandle& ActorRef = Component->BodyInstance.GetActorReferenceWithWelding();
-				FPhysicsCommand::ExecuteRead(ActorRef, [&](const FPhysicsActorHandle& Actor)
-				{
-					TArray<FPhysicsShapeHandle> AllShapes;
-					const int32 NumSyncShapes = Component->BodyInstance.GetAllShapes_AssumesLocked(AllShapes);
-
-					if (NumSyncShapes == 0 || NumSyncShapes > EnvironmentCollisionConfig.MaxShapes)
-					{
-						// Either no shapes or too complicated to consider
-						return;
-					}
-
-					using namespace Chaos;
-
-					const FReal Thickness = (FReal)EnvironmentCollisionConfig.Thickness;
-
-					auto AddSphere = [&NewCollisionData](const FImplicitSphere3& ImplicitSphere, const FTransform& ComponentToClothTransform, const FVec3& Scale = FVec3::OneVector)
-					{
-						NewCollisionData.Spheres.Emplace(
-							ImplicitSphere.GetRadiusf() * Scale.X,  // Assumes uniform scale
-							ComponentToClothTransform.TransformPosition(FVec3(ImplicitSphere.GetCenterf())));
-					};
-
-					auto AddBox = [&NewCollisionData, Thickness](const FImplicitBox3& ImplicitBox, const FTransform& ComponentToClothTransform, const FVec3& Scale = FVec3::OneVector)
-						{
-							NewCollisionData.Boxes.Emplace(
-								ComponentToClothTransform.TransformPosition(ImplicitBox.GetCenter()),
-								ComponentToClothTransform.GetRotation(),
-								ImplicitBox.Extents() * Scale * (FReal)0.5 + Thickness);
-						};
-
-					auto AddCapsule = [&NewCollisionData](const FImplicitCapsule3& ImplicitCapsule, const FTransform& ComponentToClothTransform, const FVec3& Scale = FVec3::OneVector)
-						{
-							const int32 BaseSphereIndex = NewCollisionData.Spheres.Num();
-
-							const FReal Radius = ImplicitCapsule.GetRadiusf() * Scale.X;  // Assumes uniform scale
-							const FReal HalfHeight = ImplicitCapsule.GetHeightf() * (FReal)0.5;
-							const FVector HalfSegment = ComponentToClothTransform.GetUnitAxis(EAxis::X) * HalfHeight * Scale.X;  // Assumes uniform scale
-							const FVector TransformedCenter = ComponentToClothTransform.TransformPosition(FVec3(ImplicitCapsule.GetCenterf()));
-
-							NewCollisionData.Spheres.Emplace(Radius, TransformedCenter + HalfSegment);
-							NewCollisionData.Spheres.Emplace(Radius, TransformedCenter - HalfSegment);
-
-							NewCollisionData.SphereConnections.Emplace(BaseSphereIndex, BaseSphereIndex + 1);
-						};
-
-					auto AddConvex = [&NewCollisionData, Thickness](const FImplicitConvex3& ImplicitConvex, const FMatrix& ComponentToClothMatrix, const FVec3* const Scale = nullptr)
-						{
-							TArray<FClothCollisionPrim_ConvexFace> Faces;
-							const int32 NumPlanes = ImplicitConvex.NumPlanes();
-							Faces.SetNum(NumPlanes);
-
-							TArray<FVector> SurfacePoints;
-							const int32 NumSurfacePoints = ImplicitConvex.NumVertices();
-							SurfacePoints.SetNumUninitialized(NumSurfacePoints);
-
-							TArray<TArray<int32, TInlineAllocator<4>>, TInlineAllocator<16>> PointFaces;
-							PointFaces.SetNum(NumSurfacePoints);
-
-							for (int32 FaceIndex = 0; FaceIndex < NumPlanes; ++FaceIndex)
-							{
-								const TPlaneConcrete<FReal, 3>& Plane = ImplicitConvex.GetPlane(FaceIndex);
-								const FVec3& Normal = Plane.Normal();
-								const FVec3 Base = Plane.X() + Normal * (FReal)Thickness;
-
-								Faces[FaceIndex].Plane = FPlane(FVector(Base), FVector(Normal)).TransformBy(ComponentToClothMatrix);
-
-								const int32 NumFaceIndices = ImplicitConvex.NumPlaneVertices(FaceIndex);
-								Faces[FaceIndex].Indices.SetNumUninitialized(NumFaceIndices);
-
-								for (int32 Index = 0; Index < NumFaceIndices; ++Index)
-								{
-									const int32 PointIndex = ImplicitConvex.GetPlaneVertex(FaceIndex, Index);
-									Faces[FaceIndex].Indices[Index] = PointIndex;
-
-									PointFaces[PointIndex].Add(FaceIndex);
-								}
-							}
-
-							for (int32 PointIndex = 0; PointIndex < NumSurfacePoints; ++PointIndex)
-							{
-								check(PointFaces[PointIndex].Num() >= 3);
-								const int32 Index0 = PointFaces[PointIndex][0];
-								const int32 Index1 = PointFaces[PointIndex][1];
-								const int32 Index2 = PointFaces[PointIndex][2];
-
-								if (!FMath::IntersectPlanes3(SurfacePoints[PointIndex], Faces[Index0].Plane, Faces[Index1].Plane, Faces[Index2].Plane))
-								{
-									SurfacePoints[PointIndex] = ComponentToClothMatrix.TransformPosition(FVector(ImplicitConvex.GetVertex(PointIndex)));
-								}
-							}
-
-							NewCollisionData.Convexes.Emplace(MoveTemp(Faces), MoveTemp(SurfacePoints));
-						};
-
-					FTransform ClothComponentTransform = GetComponentTransform();
-					ClothComponentTransform.RemoveScaling();  // The environment collision shape doesn't need the scale of the cloth skeletal mesh applied to it (but it does need the source scale from its component transform)
-					const FTransform ComponentToClothBaseTransform = Component->GetComponentTransform() * ClothComponentTransform.Inverse();
-
-					bool bHasSimpleCollision = false;
-
-					for (FPhysicsShapeHandle& ShapeHandle : AllShapes)
-					{
-						FTransform ComponentToClothTransform = ComponentToClothBaseTransform;
-
-						const FImplicitObject* ImplicitObject = &ShapeHandle.GetGeometry();
-						EImplicitObjectType ImplicitType = ImplicitObject->GetType();
-
-						// Transformed implicits
-						if (ImplicitType == ImplicitObjectType::Transformed)
-						{
-							const TImplicitObjectTransformed<FReal, 3>& ImplicitTransformed = ImplicitObject->GetObjectChecked<TImplicitObjectTransformed<FReal, 3>>();
-							ImplicitObject = ImplicitTransformed.GetTransformedObject();
-							ImplicitType = ImplicitObject->GetType();
-
-							ComponentToClothTransform = ImplicitTransformed.GetTransform() * ComponentToClothTransform;
-							UE_LOG(LogSkeletalMesh, Verbose, TEXT("Found transformed environmental collision"));
-						}
-
-						switch (ImplicitType)
-						{
-						// Base implicits
-						case ImplicitObjectType::Sphere:
-							{
-								const FImplicitSphere3& ImplicitSphere = ImplicitObject->GetObjectChecked<FImplicitSphere3>();;
-								AddSphere(ImplicitSphere, ComponentToClothTransform);
-							}
-							UE_LOG(LogSkeletalMesh, Verbose, TEXT("Found Sphere cloth environmental collision in [%s]"), *Component->GetOwner()->GetFName().ToString());
-							bHasSimpleCollision = true;
-							break;
-						case ImplicitObjectType::Box:
-							{
-								const FImplicitBox3& ImplicitBox = ImplicitObject->GetObjectChecked<FImplicitBox3>();
-								AddBox(ImplicitBox, ComponentToClothTransform);
-							}
-							UE_LOG(LogSkeletalMesh, Verbose, TEXT("Found Box cloth environmental collision in [%s]"), *Component->GetOwner()->GetFName().ToString());
-							bHasSimpleCollision = true;
-							break;
-						case ImplicitObjectType::Capsule:
-							{
-								const FImplicitCapsule3& ImplicitCapsule = ImplicitObject->GetObjectChecked<FImplicitCapsule3>();
-								AddCapsule(ImplicitCapsule, ComponentToClothTransform);
-								UE_LOG(LogSkeletalMesh, Verbose, TEXT("Found Capsule cloth environmental collision in [%s]"), *Component->GetOwner()->GetFName().ToString());
-								bHasSimpleCollision = true;
-							}
-							break;
-						case ImplicitObjectType::Convex:
-							{
-								const FImplicitConvex3& ImplicitConvex = ImplicitObject->GetObjectChecked<FImplicitConvex3>();
-								AddConvex(ImplicitConvex, ComponentToClothTransform.ToMatrixNoScale());
-								UE_LOG(LogSkeletalMesh, Verbose, TEXT("Found Convex cloth environmental collision in [%s]"), *Component->GetOwner()->GetFName().ToString());
-								bHasSimpleCollision = true;
-							}
-							break;
-
-						// Instanced implicits
-						case ImplicitObjectType::IsInstanced | ImplicitObjectType::Sphere:
-							{
-								const TImplicitObjectInstanced<FImplicitSphere3>& ImplicitInstanced = ImplicitObject->GetObjectChecked<TImplicitObjectInstanced<FImplicitSphere3>>();
-								check(ImplicitInstanced.Object());
-								const FImplicitSphere3& ImplicitSphere = *ImplicitInstanced.GetInstancedObject();
-								AddSphere(ImplicitSphere, ComponentToClothTransform);
-								UE_LOG(LogSkeletalMesh, Verbose, TEXT("Found Instanced Sphere cloth environmental collision in [%s]"), *Component->GetOwner()->GetFName().ToString());
-								bHasSimpleCollision = true;
-							}
-							break;
-						case ImplicitObjectType::IsInstanced | ImplicitObjectType::Box:
-							{
-								const TImplicitObjectInstanced<FImplicitBox3>& ImplicitInstanced = ImplicitObject->GetObjectChecked<TImplicitObjectInstanced<FImplicitBox3>>();
-								check(ImplicitInstanced.Object());
-								const FImplicitBox3& ImplicitBox = *ImplicitInstanced.GetInstancedObject();
-								AddBox(ImplicitBox, ComponentToClothTransform);
-								UE_LOG(LogSkeletalMesh, Verbose, TEXT("Found Instanced Box cloth environmental collision in [%s]"), *Component->GetOwner()->GetFName().ToString());
-								bHasSimpleCollision = true;
-							}
-							break;
-						case ImplicitObjectType::IsInstanced | ImplicitObjectType::Capsule:
-							{
-								const TImplicitObjectInstanced<FImplicitCapsule3>& ImplicitInstanced = ImplicitObject->GetObjectChecked<TImplicitObjectInstanced<FImplicitCapsule3>>();
-								check(ImplicitInstanced.Object());
-								const FImplicitCapsule3& ImplicitCapsule = *ImplicitInstanced.GetInstancedObject();
-								AddCapsule(ImplicitCapsule, ComponentToClothTransform);
-								UE_LOG(LogSkeletalMesh, Verbose, TEXT("Found Instanced Capsule cloth environmental collision in [%s]"), *Component->GetOwner()->GetFName().ToString());
-								bHasSimpleCollision = true;
-							}
-							break;
-						case ImplicitObjectType::IsInstanced | ImplicitObjectType::Convex:
-							{
-								const TImplicitObjectInstanced<FImplicitConvex3>& ImplicitInstanced = ImplicitObject->GetObjectChecked<TImplicitObjectInstanced<FImplicitConvex3>>();
-								check(ImplicitInstanced.Object());
-								const FImplicitConvex3& ImplicitConvex = *ImplicitInstanced.GetInstancedObject();
-								AddConvex(ImplicitConvex, ComponentToClothTransform.ToMatrixNoScale());
-								UE_LOG(LogSkeletalMesh, Verbose, TEXT("Found Instanced Convex cloth environmental collision in [%s]"), *Component->GetOwner()->GetFName().ToString());
-								bHasSimpleCollision = true;
-							}
-							break;
-
-						// Scaled implicits
-						case ImplicitObjectType::IsScaled | ImplicitObjectType::Sphere:
-							{
-								const TImplicitObjectScaled<FImplicitSphere3>& ImplicitScaled = ImplicitObject->GetObjectChecked<TImplicitObjectScaled<FImplicitSphere3>>();
-								check(ImplicitScaled.Object());
-								const FImplicitSphere3& ImplicitSphere = *ImplicitScaled.GetUnscaledObject();
-								ensure(FVector::DistSquared(ComponentToClothTransform.GetScale3D(), FVector(ImplicitScaled.GetScale())) < UE_KINDA_SMALL_NUMBER);
-								AddSphere(ImplicitSphere, ComponentToClothTransform, ImplicitScaled.GetScale());
-								UE_LOG(LogSkeletalMesh, Verbose, TEXT("Found Scaled Sphere cloth environmental collision in [%s]"), *Component->GetOwner()->GetFName().ToString());
-								bHasSimpleCollision = true;
-							}
-							break;
-						case ImplicitObjectType::IsScaled | ImplicitObjectType::Box:
-							{
-								const TImplicitObjectScaled<FImplicitBox3>& ImplicitScaled = ImplicitObject->GetObjectChecked<TImplicitObjectScaled<FImplicitBox3>>();
-								check(ImplicitScaled.Object());
-								const FImplicitBox3& ImplicitBox = *ImplicitScaled.GetUnscaledObject();
-								ensure(FVector::DistSquared(ComponentToClothTransform.GetScale3D(), FVector(ImplicitScaled.GetScale())) < UE_KINDA_SMALL_NUMBER);
-								AddBox(ImplicitBox, ComponentToClothTransform, ImplicitScaled.GetScale());
-								UE_LOG(LogSkeletalMesh, Verbose, TEXT("Found Scaled Box cloth environmental collision in [%s]"), *Component->GetOwner()->GetFName().ToString());
-								bHasSimpleCollision = true;
-							}
-							break;
-						case ImplicitObjectType::IsScaled | ImplicitObjectType::Capsule:
-							{
-								const TImplicitObjectScaled<FImplicitCapsule3>& ImplicitScaled = ImplicitObject->GetObjectChecked<TImplicitObjectScaled<FImplicitCapsule3>>();
-								check(ImplicitScaled.Object());
-								const FImplicitCapsule3& ImplicitCapsule = *ImplicitScaled.GetUnscaledObject();
-								ensure(FVector::DistSquared(ComponentToClothTransform.GetScale3D(), FVector(ImplicitScaled.GetScale())) < UE_KINDA_SMALL_NUMBER);
-								AddCapsule(ImplicitCapsule, ComponentToClothTransform, ImplicitScaled.GetScale());
-								UE_LOG(LogSkeletalMesh, Verbose, TEXT("Found Scaled Capsule cloth environmental collision in [%s]"), *Component->GetOwner()->GetFName().ToString());
-								bHasSimpleCollision = true;
-							}
-							break;
-						case ImplicitObjectType::IsScaled | ImplicitObjectType::Convex:
-							{
-								const TImplicitObjectScaled<FImplicitConvex3>& ImplicitScaled = ImplicitObject->GetObjectChecked<TImplicitObjectScaled<FImplicitConvex3>>();
-								check(ImplicitScaled.Object());
-								const FImplicitConvex3& ImplicitConvex = *ImplicitScaled.GetUnscaledObject();
-								ensure(FVector::DistSquared(ComponentToClothTransform.GetScale3D(), FVector(ImplicitScaled.GetScale())) < UE_KINDA_SMALL_NUMBER);
-								AddConvex(ImplicitConvex, ComponentToClothTransform.ToMatrixWithScale());
-								UE_LOG(LogSkeletalMesh, Verbose, TEXT("Found Scaled Convex cloth environmental collision in [%s]"), *Component->GetOwner()->GetFName().ToString());
-								bHasSimpleCollision = true;
-							}
-							break;
-
-						// Triangle mesh
-						case ImplicitObjectType::TriangleMesh:
-						case ImplicitObjectType::IsInstanced | ImplicitObjectType::TriangleMesh:
-						case ImplicitObjectType::IsScaled | ImplicitObjectType::TriangleMesh:
-							// TODO: We could eventually want to collide cloth against triangle meshes,
-							//       however the concept of simple vs complex shape might need to be clarified
-							//       as it currently iterates over all shape to discard the triangle mesh ones.
-							UE_LOG(LogSkeletalMesh, Verbose, TEXT("Found unusable Triangle Mesh cloth environmental collision in [%s]"), !Component->GetOwner() ? TEXT("Unknown") : *Component->GetOwner()->GetFName().ToString());
-							break;
-
-						default: 
-							UE_LOG(LogSkeletalMesh, Verbose, TEXT("Found unsupported collision type during environmental collision with the cloth in [%s]"), !Component->GetOwner() ? TEXT("Unknown") : *Component->GetOwner()->GetFName().ToString());
-							break;
-						}
-					}
-					bSuccessfulRead = true;
-				});
-			}
-			else if(Channel == ECollisionChannel::ECC_PhysicsBody)
-			{
-				// Possibly a skeletal mesh, extract it's clothing collisions if necessary
-				USkeletalMeshComponent* SkelComp = Cast<USkeletalMeshComponent>(Component.Get());
-
-				if(SkelComp && SkelComp->GetSkeletalMeshAsset())
-				{
-					if(SkelComp == this)
-					{
-						// Same mesh, move to next component
-						continue;
-					}
-
-					if(SkelComp->ClothingSimulation)
-					{
-						// append skeletal component collisions
-						FClothCollisionData SkelCollisionData;
-						SkelComp->ClothingSimulation->GetCollisions(SkelCollisionData, false);
-						NewCollisionData.Append(SkelCollisionData);
-					}
-				}
-			}
-		}
-	}
-
+	FEnvironmentalCollisions::AppendCollisionDataFromEnvironment(this, NewCollisionData);
 	ClothingSimulation->AddExternalCollisions(NewCollisionData);
 }
 
