Starvation handling

diff --git a/Engine/Source/Runtime/Engine/Private/AudioDecompress.cpp b/Engine/Source/Runtime/Engine/Private/AudioDecompress.cpp
--- a/Engine/Source/Runtime/Engine/Private/AudioDecompress.cpp
+++ b/Engine/Source/Runtime/Engine/Private/AudioDecompress.cpp
@@ -20,7 +20,7 @@ namespace AudioDecompressPrivate
 #endif //WITH_AUDIO_DECODER_DIAGNOSTICS
 
 	static int32 bEnableRetainingNextPlayingChunkCVar = 0;
-	FAutoConsoleVariableRef CVarForceDecoderNegativeSamplesOnWave(
+	FAutoConsoleVariableRef CVarEnableRetainingNextPlayingChunk(
 		TEXT("au.streamcache.RetainNextChunk"),
 		bEnableRetainingNextPlayingChunkCVar,
 		TEXT("Keep a reference to next to play chunk" ),
@@ -73,7 +73,6 @@ IStreamedCompressedInfo::IStreamedCompressedInfo()
 	, StreamSeekBlockIndex(INDEX_NONE)
 	, StreamSeekBlockOffset(0)
 {
-	StartTimeInCycles = FPlatformTime::Cycles64();
 }
 
 IStreamedCompressedInfo::~IStreamedCompressedInfo() = default;
@@ -252,24 +251,24 @@ bool IStreamedCompressedInfo::StreamCompressedInfoInternal(const FSoundWaveProxy
 	// Get the zeroth chunk of data (should always be loaded)
 	CurrentChunkIndex = 0;
 	uint32 ChunkSize = 0;
-	const uint8* FirstChunk = GetLoadedChunk(StreamingSoundWave, CurrentChunkIndex, ChunkSize);
+	const uint8* ZerothChunk = GetLoadedChunk(StreamingSoundWave, CurrentChunkIndex, ChunkSize);
 
 	// If there is a seek-table present in the chunk, we'll need to parse that first.
 	if (InWaveProxy->GetSoundWaveData()->HasChunkSeekTable(CurrentChunkIndex))
 	{
 		uint32 SeekTableEnd = 0;
-		if (ensureMsgf(FStreamedAudioChunkSeekTable::Parse(FirstChunk, ChunkSize, SeekTableEnd, GetCurrentSeekTable()),
+		if (ensureMsgf(FStreamedAudioChunkSeekTable::Parse(ZerothChunk, ChunkSize, SeekTableEnd, GetCurrentSeekTable()),
 			TEXT("Failed to parse seektable in '%s' chunk=%d"), *StreamingSoundWave->GetFName().ToString(), CurrentChunkIndex))
 		{
 			ChunkSize -= SeekTableEnd;
-			FirstChunk += SeekTableEnd;
+			ZerothChunk += SeekTableEnd;
 		}
 	}
 	
 	bIsStreaming = true;
-	if (FirstChunk)
+	if (ZerothChunk)
 	{
-		bool HeaderReadResult = ReadCompressedInfo(FirstChunk, ChunkSize, QualityInfo);
+		bool HeaderReadResult = ReadCompressedInfo(ZerothChunk, ChunkSize, QualityInfo);
 		
 		// If we've read through the entirety of the zeroth chunk while parsing the header, move onto the next chunk.
 		if (SrcBufferOffset >= ChunkSize)
@@ -420,48 +419,14 @@ bool IStreamedCompressedInfo::StreamCompressedData(uint8* Destination, bool bLoo
 		}
 		else
 		{
-			// Still not loaded, zero remainder of current buffer.
-			// For Chunk 1
-			//  For Load on Demand, this is expected, because it will likely wait for the first chunk to load.
-			//  For Prime on Load, it could happen if the load is still pending for the First Chunk.
-			//  For Retain on Load, it could also happen if the retained data has been purged for the cache. (or we're in the editor)
-			// For All other Chunks other than > 1
-			//  An underrun (starvation) has occured. Warn the user.
+			// We don't have a chunk yet. If we're just starting an on demand stream this is expected.
 			const ESoundWaveLoadingBehavior Behavior = StreamingSoundWave->GetLoadingBehavior();
-			const float TimeSoFar = FPlatformTime::ToSeconds64(FPlatformTime::Cycles64() - StartTimeInCycles);
-			if (CurrentChunkIndex == 1)
-			{
-				switch(Behavior)
-				{
-				case ESoundWaveLoadingBehavior::RetainOnLoad:
-				{	// We don't inline in the editor, so don't warn.
-					const bool bIsEditor = GEngine ? GEngine->IsEditor() : false;
-					UE_CLOG(!bIsEditor, LogAudioStreamCaching, Verbose, TEXT("First Audio Chunk (%s) not loaded and we are 'RetainOnLoad'. Likely because the stream cache has purged it. FailCount=%d"),
-						*StreamingSoundWave->GetFName().ToString(), ++PrintChunkFailMessageCount);
-					break;
-				}
-				case ESoundWaveLoadingBehavior::PrimeOnLoad:
-				{
-					UE_LOG(LogAudioStreamCaching, Verbose, TEXT("First Audio Chunk (%s) not loaded and we are 'PrimeOnLoad'. Likely because the stream cache is still loading it. Attempt=%d"),
-						*StreamingSoundWave->GetFName().ToString(), ++PrintChunkFailMessageCount);
-					break;
-				}
-				case ESoundWaveLoadingBehavior::LoadOnDemand:
-				{
-					UE_LOG(LogAudioStreamCaching, VeryVerbose, TEXT("First Audio Chunk (%s) not loaded and we are 'LoadOnDemand'. This is expected while we load the first chunk. Attempt=%d"),
-						*StreamingSoundWave->GetFName().ToString(), ++PrintChunkFailMessageCount);
-					break;
-				}
-				default:
-					break;
-				}
-			}
-			else // If we're mid play back, and we hit this point. We've underrun, log a warning.
+			const bool bLoadOnDemandStarting = Behavior == ESoundWaveLoadingBehavior::LoadOnDemand && CurrentChunkIndex==1;
+			if (!bLoadOnDemandStarting)
 			{
-				UE_LOG(LogAudioStreamCaching, Verbose, TEXT("Chunk %d not yet loaded for playing SoundWave '%s'. Attempt=%d, LoadBehavior=%s, Latency=%2.2f secs, LikelyReason: Not loaded fast enough, IO Saturation." ),
-					CurrentChunkIndex, *StreamingSoundWave->GetFName().ToString(), ++PrintChunkFailMessageCount, EnumToString(Behavior), TimeSoFar);
+				IStreamingManager::Get().GetAudioStreamingManager().HandleStarvation(StreamingSoundWave, CurrentChunkIndex, OutNumBytesStreamed, BufferSize);
 			}
-			
+			// Still not loaded, zero remainder of current buffer.
 			ZeroBuffer(Destination, BufferSize);
 			return false;
 		}
@@ -578,8 +543,7 @@ bool IStreamedCompressedInfo::StreamCompressedData(uint8* Destination, bool bLoo
 				}
 				else
 				{
-					UE_LOG(LogAudioStreamCaching, Warning, TEXT("Chunk %d not yet loaded for playing SoundWave '%s'. LikelyReason: Not loaded fast enough, IO Saturation."),
-						CurrentChunkIndex, *StreamingSoundWave->GetFName().ToString());
+					IStreamingManager::Get().GetAudioStreamingManager().HandleStarvation(StreamingSoundWave, CurrentChunkIndex, OutNumBytesStreamed, BufferSize);
 
 					SrcBufferDataSize = 0;
 					// this is zero padding, so don't include it in the OutNumBytesStreamed count
diff --git a/Engine/Source/Runtime/Engine/Private/AudioStreamingCache.cpp b/Engine/Source/Runtime/Engine/Private/AudioStreamingCache.cpp
--- a/Engine/Source/Runtime/Engine/Private/AudioStreamingCache.cpp
+++ b/Engine/Source/Runtime/Engine/Private/AudioStreamingCache.cpp
@@ -18,6 +18,7 @@ AudioStreaming.cpp: Implementation of audio streaming classes.
 #include "ProfilingDebugging/CountersTrace.h"
 #include "Audio/AudioDebug.h"
 #include "Audio/AudioDebugDraw.h"
+#include "IO/IoBuffer.h"
 
 DEFINE_LOG_CATEGORY(LogAudioStreamCaching);
 
@@ -185,6 +186,23 @@ static FAutoConsoleCommand GFlushAudioCacheCommand(
 	})
 );
 
+
+static const TCHAR* LexToString(const EAsyncIOPriorityAndFlags In)
+{
+#define CASE_TO_STRING(X) case X: return TEXT(#X)
+	switch (In & AIOP_PRIORITY_MASK)
+	{
+		CASE_TO_STRING(AIOP_MIN);
+		CASE_TO_STRING(AIOP_Low);
+		CASE_TO_STRING(AIOP_BelowNormal);
+		CASE_TO_STRING(AIOP_Normal);
+		CASE_TO_STRING(AIOP_High);
+		CASE_TO_STRING(AIOP_CriticalPath);
+		default: return TEXT("Unknown");
+	}
+#undef CASE_TO_STRING
+}
+
 static EAsyncIOPriorityAndFlags ToIoPriority(const int32 InValue)
 {
 	switch (InValue)
@@ -323,6 +341,39 @@ bool FAudioChunkCache::FChunkKey::operator==(const FChunkKey& Other) const
 #endif
 }
 
+void FAudioChunkCache::FCacheElement::WaitForAsyncLoadCompletion(bool bCancel)
+{
+#if WITH_EDITORONLY_DATA
+	if (DDCTask.IsValid() && !DDCTask->IsDone())
+	{
+		if (bCancel)
+		{
+			DDCTask->Cancel();
+		}
+
+		DDCTask->EnsureCompletion(false);
+	}
+#endif
+
+	if (Read.IsPending())
+	{
+		if (bCancel)
+		{
+			Read.Cancel();
+		}
+		else
+		{
+			static constexpr uint32 WaitTimeMs = 1000;
+			for (int32 Attempt = 0; Attempt < 10 && Read.WaitFor(WaitTimeMs); Attempt++)
+			{
+				UE_CLOG(Attempt % 5 == 0, LogAudioStreamCaching, Warning,
+					TEXT("Still Waiting for Sound: %s:%i"),*Key.SoundWaveName.ToString(), Key.ChunkIndex);
+			}
+		}
+		Read = {};
+	}
+}
+
 #if DEBUG_STREAM_CACHE
 bool FAudioChunkCache::FCacheElement::IsBeingPlayed() const
 {
@@ -516,6 +567,15 @@ void FCachedAudioStreamingManager::RemoveMemoryCountedFeature(const FAudioStream
 	}
 }
 
+void FCachedAudioStreamingManager::HandleStarvation(const FSoundWaveProxyPtr& SoundWave, const uint32 ChunkIndex, const uint32 InBytesProduced, const uint32 InBytesExpected) 
+{
+	// If this sound wave is managed by a cache, use that to get the chunk:
+	if (FAudioChunkCache* Cache = GetCacheForWave(SoundWave))
+	{
+		Cache->HandleStarvation(SoundWave, ChunkIndex, InBytesProduced, InBytesExpected);
+	}
+}
+
 FAudioChunkHandle FCachedAudioStreamingManager::GetLoadedChunk(const FSoundWaveProxyPtr& SoundWave, uint32 ChunkIndex, bool bBlockForLoad, bool bForImmediatePlayback) const
 {
 	LLM_SCOPE(ELLMTag::AudioStreamCache);
@@ -578,7 +638,7 @@ FAudioChunkHandle FCachedAudioStreamingManager::GetLoadedChunk(const FSoundWaveP
 #endif
 			);
 
-			uint64 LookupIDForNextChunk = Cache->AddOrTouchChunk(NextChunkKey, SoundWaveData, [](EAudioChunkLoadResult) {}, ENamedThreads::AnyThread, false);
+			const uint64 LookupIDForNextChunk = Cache->AddOrTouchChunk(NextChunkKey, SoundWaveData, [](EAudioChunkLoadResult) {}, ENamedThreads::AnyThread, bForImmediatePlayback); 
 			if (LookupIDForNextChunk == InvalidAudioStreamCacheLookupID)
 			{
 				UE_LOG(LogAudioStreamCaching, Warning, TEXT("Failed to add follow-up chunk for %s. This should not happen with our current TrimMemory() logic"), *SoundWave->GetFName().ToString());
@@ -791,6 +851,21 @@ uint64 FAudioChunkCache::AddOrTouchChunk(const FChunkKey& InKey, const TSharedPt
 	if (FoundElement)
 	{
 		TouchElement(FoundElement);
+		if (FoundElement->Read.IsPending())
+		{
+			// Bump priority if we've changed from loading to playing.
+			const EAsyncIOPriorityAndFlags Priority = GetAsyncPriorityForChunk(InKey, bNeededForPlayback);
+			if (FoundElement->Priority < Priority)
+			{
+				UE_CLOG(FoundElement->WantsFilteredLogging(), LogAudioStreamCaching, Verbose,
+					TEXT("[%s] - Bumping priority of existing read, now we're in playback), chunk %d, ObjectKey=%s"),
+						*WriteToString<128>(FoundElement->Key.SoundWaveName), FoundElement->Key.ChunkIndex, *WriteToString<40>(FoundElement->Key.ObjectKey));
+
+				FoundElement->Read.UpdatePriority(Priority);
+				FoundElement->Priority = Priority;
+			}
+		}
+			
 		if (FoundElement->bIsLoaded)
 		{
 			ExecuteOnLoadCompleteCallback(EAudioChunkLoadResult::AlreadyLoaded, OnLoadCompleted, CallbackThread);
@@ -1872,6 +1947,12 @@ void FAudioChunkCache::KickOffAsyncLoad(FCacheElement* CacheElement, const FChun
 	{
 		if (CacheElement->IsLoadInProgress())
 		{
+			// Update priority if it's already pending...		
+			if (CacheElement->Read.IsPending() && AsyncIOPriority > CacheElement->Priority)
+			{
+				CacheElement->Read.UpdatePriority(AsyncIOPriority);
+				CacheElement->Priority = AsyncIOPriority;
+			}
 			UE_CLOG(CacheElement->WantsFilteredLogging(), LogAudioTiming, Verbose, TEXT("[%s] - Waiting for async load while kicking async load, chunk %d, ObjectKey=%s"), *WriteToString<128>(CacheElement->Key.SoundWaveName), CacheElement->Key.ChunkIndex, *WriteToString<40>(CacheElement->Key.ObjectKey));
 			CacheElement->WaitForAsyncLoadCompletion(true);
 		}
@@ -1887,17 +1968,9 @@ void FAudioChunkCache::KickOffAsyncLoad(FCacheElement* CacheElement, const FChun
 		NumberOfLoadsInFlight.Increment();
 		AddLoadCounters(AsyncIOPriority, 1);
 
-		FBulkDataIORequestCallBack AsyncFileCallBack = [this, OnLoadCompleted, CacheElement, InKey, ChunkDataSize, CallbackThread, AsyncIOPriority](bool bWasCancelled, IBulkDataIORequest*)
+		FBulkDataRequest::FCompletionCallback AsyncFileCallBack = [this, OnLoadCompleted, CacheElement, InKey, ChunkDataSize, CallbackThread, AsyncIOPriority](FBulkDataRequest::EStatus Status) 
 		{
-			// Take ownership of the read request and close the storage
-			IBulkDataIORequest* LocalReadRequest = (IBulkDataIORequest*)FPlatformAtomics::InterlockedExchangePtr((void* volatile*)&CacheElement->ReadRequest, (void*)0x1);
-
-			if (LocalReadRequest && (void*)LocalReadRequest != (void*)0x1)
-			{
-				// Delete the request to avoid hording space in pak cache
-				TGraphTask<FClearAudioChunkCacheReadRequestTask>::CreateTask().ConstructAndDispatchWhenReady(LocalReadRequest);
-			}
-
+			const bool bWasCancelled = Status == FBulkDataRequest::EStatus::Cancelled;
 			// Diagnostic logging.
 			const bool bMatched = CacheElement->Key == InKey;
 			UE_CLOG(CacheElement->WantsFilteredLogging(), LogAudioTiming, Verbose, TEXT("[%s] - AsyncReadComplete index %d, ObjectKey=%s"), *WriteToString<128>(CacheElement->Key.SoundWaveName), CacheElement->Key.ChunkIndex, *WriteToString<40>(CacheElement->Key.ObjectKey));
@@ -1907,7 +1980,8 @@ void FAudioChunkCache::KickOffAsyncLoad(FCacheElement* CacheElement, const FChun
 			// Populate key and DataSize. The async read request was set up to write directly into CacheElement->ChunkData.
 			CacheElement->Key = InKey;
 			CacheElement->ChunkDataSize = ChunkDataSize;
-			CacheElement->bIsLoaded = true;
+			CacheElement->bIsLoaded = Status == FBulkDataRequest::EStatus::Ok;
+
 			if (bWasCancelled)
 			{
 				FMemory::Memzero(CacheElement->ChunkData, CacheElement->ChunkDataSize);
@@ -1929,7 +2003,7 @@ void FAudioChunkCache::KickOffAsyncLoad(FCacheElement* CacheElement, const FChun
 		CacheElement->DebugInfo.TimeLoadStartedCycles = FPlatformTime::Cycles64();
 #endif
 
-		CacheElement->ReadRequest = nullptr;
+		CacheElement->Read = {};
 		if (Chunk->BulkData.IsBulkDataLoaded())
 		{
 			UE_CLOG(CacheElement->WantsFilteredLogging(), LogAudioTiming, Verbose, TEXT("[%s] - Chunk %d has bulk data loaded, , ObjectKey=%s"), *WriteToString<128>(CacheElement->Key.SoundWaveName), CacheElement->Key.ChunkIndex, *WriteToString<40>(CacheElement->Key.ObjectKey));
@@ -1949,7 +2023,7 @@ void FAudioChunkCache::KickOffAsyncLoad(FCacheElement* CacheElement, const FChun
 #endif //DEBUG_STREAM_CACHE
 			
 			// Fire the callback (this will mark it load completed etc).
-			AsyncFileCallBack(false, nullptr);
+			AsyncFileCallBack(FBulkDataRequest::EStatus::Ok);
 		}
 		else
 		{
@@ -1965,24 +2039,18 @@ void FAudioChunkCache::KickOffAsyncLoad(FCacheElement* CacheElement, const FChun
 #if !UE_BUILD_SHIPPING
 			UE_CLOG(CacheElement->bWantsExtraLogging, LogAudioStreamCaching, Display, TEXT("[%s] - CreateStreamingRequest: chunk %d, ObjectKey=%s"), *WriteToString<128>(CacheElement->Key.SoundWaveName), CacheElement->Key.ChunkIndex, *WriteToString<40>(CacheElement->Key.ObjectKey));
 #endif
+			
 			TRACE_IOSTORE_METADATA_SCOPE_TAG(InKey.SoundWaveName);
-			IBulkDataIORequest* LocalReadRequest = Chunk->BulkData.CreateStreamingRequest(0, ChunkDataSize, AsyncIOPriority | AIOP_FLAG_DONTCACHE, &AsyncFileCallBack, CacheElement->ChunkData);
-			if (!LocalReadRequest)
-			{
-				UE_LOG(LogAudioStreamCaching, Error, TEXT("Chunk load in audio LRU cache failed."));
-				OnLoadCompleted(EAudioChunkLoadResult::ChunkOutOfBounds);
-				NumberOfLoadsInFlight.Decrement();
-				AddLoadCounters(AsyncIOPriority, -1);
-			}
-			else if (FPlatformAtomics::InterlockedCompareExchangePointer((void* volatile*)&CacheElement->ReadRequest, LocalReadRequest, nullptr) == (void*)0x1)
-			{
-				// The request is completed before we can store it. Just delete it
-				TGraphTask<FClearAudioChunkCacheReadRequestTask>::CreateTask().ConstructAndDispatchWhenReady(LocalReadRequest);
-			}
+			FBulkDataBatchRequest::FBatchBuilder Builder = FBulkDataBatchRequest::NewBatch(1);
+			FIoBuffer Buffer(FIoBuffer::Wrap, CacheElement->ChunkData,CacheElement->ChunkDataSize);
+			Builder.Read(Chunk->BulkData, 0, CacheElement->ChunkDataSize, AsyncIOPriority | AIOP_FLAG_DONTCACHE, Buffer);
+			Builder.Issue(MoveTemp(AsyncFileCallBack), CacheElement->Read);
+			CacheElement->Priority = AsyncIOPriority;
 		}
 	}
 }
 
+
 EAsyncIOPriorityAndFlags FAudioChunkCache::GetAsyncPriorityForChunk(const FChunkKey& InKey, bool bNeededForPlayback)
 {
 	if (bNeededForPlayback)
@@ -2035,6 +2103,39 @@ void FAudioChunkCache::SetCacheLookupIDForChunk(const FChunkKey& InChunkKey, uin
 }
 
 
+void FAudioChunkCache::HandleStarvation(
+	const FSoundWaveProxyPtr& SoundWave, const uint32 ChunkIndex, const uint32 InBytesProduced, const uint32 InBytesExpected)
+{
+	const FChunkKey Key(
+		SoundWave->GetSoundWaveData()
+		, ChunkIndex
+#if WITH_EDITOR
+		, (uint32)SoundWave->GetCurrentChunkRevision()
+#endif
+	);
+	if (FCacheElement* FoundElement = FindElementForKey(Key))
+	{
+		const float PercentProduced = InBytesExpected > 0.f ? InBytesProduced / InBytesExpected : 0.f;
+		const float Starvometer = 1.f - PercentProduced;
+		
+		// Stream is playing, and we're starving, crank priority of read
+		static constexpr EAsyncIOPriorityAndFlags StarvingPriority = AIOP_MAX;
+		if (FoundElement->Read.IsPending() && FoundElement->Priority != StarvingPriority)
+		{
+#if DEBUG_STREAM_CACHE
+			const float TimeStartedSecs = FPlatformTime::ToSeconds64(FPlatformTime::Cycles64()-FoundElement->DebugInfo.TimeLoadStartedCycles);
+			UE_CLOG(FoundElement->WantsFilteredLogging(), LogAudioTiming, Verbose, TEXT("FAudioChunkCache::HandleStarvation Wave:%s, Chunk:%d, LoadRequested=-%.2f (secs), OldPriority=%s, NewPriority=%s, Starvometer=%2.f"),
+				*SoundWave->GetFName().ToString(), ChunkIndex, TimeStartedSecs, LexToString(FoundElement->Priority), LexToString(StarvingPriority), Starvometer);
+#endif //DEBUG_STREAM_CACHE
+
+			FoundElement->Read.UpdatePriority(AIOP_MAX);
+			FoundElement->Priority = StarvingPriority;
+		}
+		
+		return;
+	}
+	UE_LOG(LogAudioTiming, Verbose, TEXT("Decoder Starving waiting for Wave:%s, Chunk:%d (not in cache!)"), *SoundWave->GetFName().ToString(), ChunkIndex)
+}
 
 int32 FCachedAudioStreamingManager::RenderStatAudioStreaming(UWorld* World, FViewport* Viewport, FCanvas* Canvas, int32 X, int32 Y, const FVector* ViewLocation, const FRotator* ViewRotation)
 {
diff --git a/Engine/Source/Runtime/Engine/Public/AudioDecompress.h b/Engine/Source/Runtime/Engine/Public/AudioDecompress.h
--- a/Engine/Source/Runtime/Engine/Public/AudioDecompress.h
+++ b/Engine/Source/Runtime/Engine/Public/AudioDecompress.h
@@ -340,10 +340,6 @@ protected:
 	bool bStoringEndOfFile;
 	/** The current chunk index in the streamed chunks. */
 	int32 CurrentChunkIndex;
-	/** Whether or not to print the chunk fail message. */
-	int32 PrintChunkFailMessageCount = 0;
-	/** A counter of when we started the last request (for gauging latency) */
-	uint64 StartTimeInCycles = 0;
 	/** Number of bytes of padding used, overridden in some implementations. Defaults to 0. */
 	uint32 SrcBufferPadding;
 	/** Chunk Handle to ensure that this chunk of streamed audio is not deleted while we are using it. */
diff --git a/Engine/Source/Runtime/Engine/Public/AudioStreamingCache.h b/Engine/Source/Runtime/Engine/Public/AudioStreamingCache.h
--- a/Engine/Source/Runtime/Engine/Public/AudioStreamingCache.h
+++ b/Engine/Source/Runtime/Engine/Public/AudioStreamingCache.h
@@ -207,6 +207,8 @@ public:
 	TPair<int, int> DebugVisualDisplay(UWorld* World, FViewport* Viewport, FCanvas* Canvas, int32 X, int32 Y, const FVector* ViewLocation, const FRotator* ViewRotation) const;
 	TPair<int, int> DebugBirdsEyeDisplay(UWorld* World, FViewport* Viewport, FCanvas* Canvas, int32 X, int32 Y, const FVector* ViewLocation, const FRotator* ViewRotation) const;
 
+	void HandleStarvation(const FSoundWaveProxyPtr& SoundWave, const uint32 ChunkIndex, const uint32 InBytesProduced, const uint32 InBytesExpected);
+
 	void IncrementCacheOverflowCounter()
 	{
 		CacheOverflowCount.Increment();
@@ -321,8 +323,8 @@ private:
 		TUniquePtr<FAsyncStreamDerivedChunkTask> DDCTask;
 #endif
 
-		// Handle to our async read request operation.
-		IBulkDataIORequest* ReadRequest;
+		FBulkDataBatchRequest Read;
+		EAsyncIOPriorityAndFlags Priority = AIOP_MIN;
 
 #if DEBUG_STREAM_CACHE
 		FCacheElementDebugInfo DebugInfo;
@@ -335,39 +337,11 @@ private:
 			, LessRecentElement(nullptr)
 			, CacheLookupID(InCacheIndex)
 			, bIsLoaded(false)
-			, ReadRequest(nullptr)
 		{
 		}
 
-		void WaitForAsyncLoadCompletion(bool bCancel)
-		{
-#if WITH_EDITORONLY_DATA
-			if (DDCTask.IsValid() && !DDCTask->IsDone())
-			{
-				if (bCancel)
-				{
-					DDCTask->Cancel();
-				}
-				
-				DDCTask->EnsureCompletion(false);
-			}
-#endif
-
-			// Take ownership and close the storage
-			IBulkDataIORequest* LocalReadRequest = (IBulkDataIORequest*)FPlatformAtomics::InterlockedExchangePtr((void* volatile*)&ReadRequest, (void*)0x1);
-
-			if (LocalReadRequest && (void*)LocalReadRequest != (void*)0x1)
-			{
-				if (bCancel)
-				{
-					LocalReadRequest->Cancel();
-				}
-				
-				LocalReadRequest->WaitCompletion();
-				delete LocalReadRequest;
-			}
-		}
-
+		void WaitForAsyncLoadCompletion(bool bCancel);
+	
 		bool IsLoadInProgress() const
 		{
 			return !bIsLoaded;
@@ -437,6 +411,9 @@ private:
 	// Number of async load operations we have currently in flight.
 	FThreadSafeCounter NumberOfLoadsInFlight;
 
+	// Time begin play came in. (so we can validate load times against load screen/vs level loads).
+	uint64 TimeOfBeginPlayCycles=0;
+
 	// Critical section: only used when we are modifying element positions in the cache. This only happens in TouchElement, EvictLeastRecentChunk, and TrimMemory.
 	// Individual cache elements should be thread safe to access.
 	mutable TDontCopy<FCriticalSection> CacheMutationCriticalSection;
@@ -586,6 +563,8 @@ protected:
 	virtual void AddMemoryCountedFeature(const FAudioStreamCacheMemoryHandle& Feature) override;
 	virtual void RemoveMemoryCountedFeature(const FAudioStreamCacheMemoryHandle& Feature) override;
 
+	virtual void HandleStarvation(const FSoundWaveProxyPtr& SoundWave, const uint32 ChunkIndex, const uint32 InBytesProduced, const uint32 InBytesExpected) override;
+	
 	/**
 	 * Returns which cache this sound wave should be in,
 	 * based on the size of this sound wave's chunk,
diff --git a/Engine/Source/Runtime/Engine/Public/ContentStreaming.h b/Engine/Source/Runtime/Engine/Public/ContentStreaming.h
--- a/Engine/Source/Runtime/Engine/Public/ContentStreaming.h
+++ b/Engine/Source/Runtime/Engine/Public/ContentStreaming.h
@@ -561,6 +561,12 @@ struct IAudioStreamingManager : public IStreamingManager
 	 */
 	virtual void SetProfilingMode(bool bEnabled) = 0;
 
+	/**
+	 * In the event that the decoder playing back the chunks is in a starvation state due
+	 * to an outstanding chunk load, handle that here.
+	 */
+	virtual void HandleStarvation(const FSoundWaveProxyPtr& SoundWave, const uint32 ChunkIndex, const uint32 InBytesProduced, const uint32 InBytesExpected ) {}
+
 protected:
 	friend FAudioChunkHandle;
 	friend FAudioStreamCacheMemoryHandle;
