Fix for asset during audio device preinit when there is no-config defined.

diff --git a/Engine/Source/Runtime/Engine/Private/AudioDeviceManager.cpp b/Engine/Source/Runtime/Engine/Private/AudioDeviceManager.cpp
--- a/Engine/Source/Runtime/Engine/Private/AudioDeviceManager.cpp
+++ b/Engine/Source/Runtime/Engine/Private/AudioDeviceManager.cpp
@@ -353,12 +353,21 @@ void FAudioDeviceManager::RegisterAudioInfoFactories()
 {
 	// Load any Engine.ini defined modules necessary for registering format factories.
 	TArray<FString> AudioInfoModules;
-	if (GConfig->GetArray(TEXT("Audio"), TEXT("AudioInfoModules"), AudioInfoModules, GEngineIni))
+	if (!GConfig->GetArray(TEXT("Audio"), TEXT("AudioInfoModules"), AudioInfoModules, GEngineIni))
 	{
-		for (const FString& i : AudioInfoModules)
-		{
-			FModuleManager::Get().LoadModuleChecked(*i);
-		}
+		// If this is simply not defined, default to sensible list of internal formats.
+		static const TCHAR* DefaultInfoModuleNames[] = {
+			TEXT("OpusAudioDecoder"), TEXT("VorbisAudioDecoder"), TEXT("AdpcmAudioDecoder"), TEXT("BinkAudioDecoder"), TEXT("RadAudioDecoder")
+		};
+		AudioInfoModules.Append(DefaultInfoModuleNames);
+		UE_LOG(LogAudio, Warning, TEXT("Audio:AudioInfoModules is not defined, defaulting to built in formats. (%s)"),
+			*FString::JoinBy(AudioInfoModules, TEXT(", "), [](const FString& i) { return i; }));
+	}
+
+	// Load any necessary audio modules.
+	for (const FString& i : AudioInfoModules)
+	{
+		FModuleManager::Get().LoadModuleChecked(*i);
 	}
 		
 	// Sanity check we have all the Factories we need to run now by 
