Multithread server
* Tick task groups can now be executed in multithread on dedicated servers when UE_DEFAULT_MULTITHREAD_SERVER is true
* Fixed forked multithread server always having tick task group set to run in single-thread

diff --git a/Engine/Source/Runtime/Core/Private/GenericPlatform/GenericPlatformProcess.cpp b/Engine/Source/Runtime/Core/Private/GenericPlatform/GenericPlatformProcess.cpp
--- a/Engine/Source/Runtime/Core/Private/GenericPlatform/GenericPlatformProcess.cpp
+++ b/Engine/Source/Runtime/Core/Private/GenericPlatform/GenericPlatformProcess.cpp
@@ -24,6 +24,7 @@
 #include "Async/TaskGraphInterfaces.h"
 
 #ifndef DEFAULT_NO_THREADING
+	// By default we assume the platform supports multithreading
 	#define DEFAULT_NO_THREADING 0
 #endif
 
diff --git a/Engine/Source/Runtime/Core/Private/Misc/App.cpp b/Engine/Source/Runtime/Core/Private/Misc/App.cpp
--- a/Engine/Source/Runtime/Core/Private/Misc/App.cpp
+++ b/Engine/Source/Runtime/Core/Private/Misc/App.cpp
@@ -14,6 +14,11 @@
 #include "Misc/NetworkVersion.h"
 #include "Misc/SecureHash.h"
 
+#ifndef UE_DEFAULT_MULTITHREAD_SERVER
+	// Controls if the dedicated server should enable multithreading or not
+	#define UE_DEFAULT_MULTITHREAD_SERVER 0
+#endif
+
 DEFINE_LOG_CATEGORY_STATIC(LogApp, Log, All);
 
 /* FApp static initialization
@@ -302,6 +307,12 @@ bool FApp::ShouldUseThreadingForPerformance()
 	return !OnlyOneThread || bForceUsePerfThreads;
 }
 
+bool FApp::IsMultithreadServer()
+{
+	static bool bIsMultithreadServer = IsRunningDedicatedServer() && UE_DEFAULT_MULTITHREAD_SERVER;
+	return bIsMultithreadServer;
+}
+
 FTimecode FApp::GetTimecode()
 {
 	if (CurrentFrameTime.IsSet())
diff --git a/Engine/Source/Runtime/Core/Private/Misc/Fork.cpp b/Engine/Source/Runtime/Core/Private/Misc/Fork.cpp
--- a/Engine/Source/Runtime/Core/Private/Misc/Fork.cpp
+++ b/Engine/Source/Runtime/Core/Private/Misc/Fork.cpp
@@ -7,16 +7,6 @@
 #include "Misc/CommandLine.h"
 #include "Misc/Parse.h"
 
-#ifndef DEFAULT_SERVER_FAKE_FORKS
-	// Defines if the server process should always simulate forking even on platforms that do not support it.
-	#define DEFAULT_SERVER_FAKE_FORKS 0
-#endif
-
-#ifndef DEFAULT_FORK_PROCESS_MULTITHREAD
-	// Defines if forked child processes should support multithreading by default
-	#define DEFAULT_FORK_PROCESS_MULTITHREAD 0
-#endif
-
 namespace ForkInternal
 {
 	/** Set to true on a child process right after the fork point */
diff --git a/Engine/Source/Runtime/Core/Public/Misc/App.h b/Engine/Source/Runtime/Core/Public/Misc/App.h
--- a/Engine/Source/Runtime/Core/Public/Misc/App.h
+++ b/Engine/Source/Runtime/Core/Public/Misc/App.h
@@ -494,6 +494,11 @@ public:
 	 */
 	static CORE_API bool ShouldUseThreadingForPerformance();
 
+	/**
+	 * Returns if the application is a dedicated server and wants to be multi-threaded.
+	 */
+	static CORE_API bool IsMultithreadServer();
+
 	/**
 	 * Checks whether application is in benchmark mode.
 	 *
diff --git a/Engine/Source/Runtime/Core/Public/Misc/Fork.h b/Engine/Source/Runtime/Core/Public/Misc/Fork.h
--- a/Engine/Source/Runtime/Core/Public/Misc/Fork.h
+++ b/Engine/Source/Runtime/Core/Public/Misc/Fork.h
@@ -6,9 +6,15 @@
 #include "HAL/PlatformAffinity.h"
 
 #ifndef DEFAULT_SERVER_FAKE_FORKS
+	// Defines if the server process should always simulate forking even on platforms that do not support it.
 	#define DEFAULT_SERVER_FAKE_FORKS 0
 #endif
 
+#ifndef DEFAULT_FORK_PROCESS_MULTITHREAD
+	// Defines if forked child processes should support multithreading by default
+	#define DEFAULT_FORK_PROCESS_MULTITHREAD 0
+#endif
+
 class FRunnable;
 class FRunnableThread;
 
diff --git a/Engine/Source/Runtime/Engine/Private/TickTaskManager.cpp b/Engine/Source/Runtime/Engine/Private/TickTaskManager.cpp
--- a/Engine/Source/Runtime/Engine/Private/TickTaskManager.cpp
+++ b/Engine/Source/Runtime/Engine/Private/TickTaskManager.cpp
@@ -10,6 +10,7 @@
 #include "Stats/StatsTrace.h"
 #include "TickTaskManagerInterface.h"
 #include "Async/ParallelFor.h"
+#include "Misc/Fork.h"
 #include "Misc/TimeGuard.h"
 #include "Templates/UniquePtr.h"
 #include "UObject/RemoteObjectTransfer.h"
@@ -511,6 +512,8 @@ class FTickTaskSequencer
 	bool				bLogTicks;
 	/** If true, log tick prerequisites when logging ticks */
 	bool				bLogTicksShowPrerequistes;
+	/** If true, tick everything from the main game thread */
+	bool				bSingleThreadMode = false;
 
 public:
 
@@ -522,26 +525,37 @@ public:
 	}
 
 	/** Return true if we should be running in single threaded mode, ala dedicated server */
-	FORCEINLINE static bool SingleThreadedMode()
+	static bool SingleThreadedMode()
 	{
-		static TOptional<bool> bSingleThreaded;
-		if (!bSingleThreaded.IsSet())
+#if DEFAULT_FORK_PROCESS_MULTITHREAD
+		// In a Forked multithread process ignore FPlatformProcess::SupportsMultithreading because the parent process never supports multithreading and the function will always return false
+		if (FForkProcessHelper::SupportsMultithreadingPostFork())
 		{
-			if (FForkProcessHelper::IsForkedMultithreadInstance())
-			{
-				bSingleThreaded = false;
-			}
-			else if (!FApp::ShouldUseThreadingForPerformance() || IsRunningDedicatedServer() || FPlatformMisc::NumberOfCores() < 3 || !FPlatformProcess::SupportsMultithreading())
-			{
-				bSingleThreaded = true;
-			}
-			else
+			// Stay in single-thread mode until the forked process decides it's safe to run in multithread
+			return FForkProcessHelper::IsForkedMultithreadInstance() == false;
+		}
+		else
+#endif
+		{
+			auto IsSingleThreadOnce = []()
 			{
-				bSingleThreaded = false;
-			}
+				// Are we a dedicated server that supports multithreading
+				if (FApp::IsMultithreadServer() && FPlatformProcess::SupportsMultithreading())
+				{
+					return false;
+				}
+				else if (!FApp::ShouldUseThreadingForPerformance() || IsRunningDedicatedServer() || FPlatformMisc::NumberOfCores() < 3 || !FPlatformProcess::SupportsMultithreading())
+				{
+					return true;
+				}
+
+				return false;
+			};
+
+			static const bool bIsSingleThread = IsSingleThreadOnce();
+			return bIsSingleThread;
 		}
-		return *bSingleThreaded;
-	}
+	};
 
 	/** Accessor that will return a valid task pointer or null if not actually scheduled */
 	FORCEINLINE static FTickGraphTask* GetGraphTask(const FTickFunction* TickFunction)
@@ -802,7 +816,7 @@ public:
 
 		{
 			SCOPE_CYCLE_COUNTER(STAT_ReleaseTickGroup);
-			if (SingleThreadedMode() || CVarAllowAsyncTickDispatch.GetValueOnGameThread() == 0)
+			if (bSingleThreadMode || CVarAllowAsyncTickDispatch.GetValueOnGameThread() == 0)
 			{
 				DispatchTickGroup(ENamedThreads::GameThread, WorldTickGroup);
 			}
@@ -814,7 +828,7 @@ public:
 			}
 		}
 
-		if (bBlockTillComplete || SingleThreadedMode())
+		if (bBlockTillComplete || bSingleThreadMode)
 		{
 			SCOPE_CYCLE_COUNTER(STAT_ReleaseTickGroup_Block);
 			for (ETickingGroup Block = WaitForTickGroup; Block <= WorldTickGroup; Block = ETickingGroup(Block + 1))
@@ -824,7 +838,7 @@ public:
 				{
 					TRACE_CPUPROFILER_EVENT_SCOPE(TickCompletionEvents);
 					FTaskGraphInterface::Get().WaitUntilTasksComplete(TickCompletionEvents[Block], ENamedThreads::GameThread);
-					if (SingleThreadedMode() || Block == TG_NewlySpawned || CVarAllowAsyncTickCleanup.GetValueOnGameThread() == 0 || TickCompletionEvents[Block].Num() < 50)
+					if (bSingleThreadMode || Block == TG_NewlySpawned || CVarAllowAsyncTickCleanup.GetValueOnGameThread() == 0 || TickCompletionEvents[Block].Num() < 50)
 					{
 						ResetTickGroup(Block);
 					}
@@ -853,12 +867,15 @@ public:
 		bLogTicks = !!CVarLogTicks.GetValueOnGameThread();
 		bLogTicksShowPrerequistes = !!CVarLogTicksShowPrerequistes.GetValueOnGameThread();
 
+		// Always cache the setting at the start of the tick process because in some rare cases (forking) the process can switch from single-thread to multi-thread mid-tick
+		bSingleThreadMode = SingleThreadedMode();
+
 		if (bLogTicks)
 		{
 			UE_LOG(LogTick, Log, TEXT("tick %6llu ---------------------------------------- Start Frame"),(uint64)GFrameCounter);
 		}
 
-		if (SingleThreadedMode())
+		if (bSingleThreadMode)
 		{
 			bAllowConcurrentTicks = false;
 		}
