Add commandline option to ignore blueprint cook dependencies.  Meant for use in when correctness of incremental cooked output isn't required.

diff --git a/Engine/Source/Runtime/Engine/Private/Blueprint.cpp b/Engine/Source/Runtime/Engine/Private/Blueprint.cpp
--- a/Engine/Source/Runtime/Engine/Private/Blueprint.cpp
+++ b/Engine/Source/Runtime/Engine/Private/Blueprint.cpp
@@ -876,6 +876,12 @@ void UBlueprint::DebuggingWorldRegistrationHelper(UObject* ObjectProvidingWorld,
 void UBlueprint::PreSaveCook(FObjectPreSaveContext ObjectSaveContext)
 {
 	using namespace UE::Cook;
+	static bool bIgnoreBlueprintCookDependencies = FParse::Param(FCommandLine::Get(), TEXT("IgnoreBlueprintCookDependencies"));
+	if (bIgnoreBlueprintCookDependencies)
+	{
+		return;
+	}
+	
 	ObjectSaveContext.AddCookBuildDependency(
 		UE::Private::BlueprintDependencies::RecordCookDependencies(this));
 }
