Log full path if available when a client fails to find or spawn a replicated actor received from the server.

diff --git a/Engine/Source/Runtime/Engine/Private/DataChannel.cpp b/Engine/Source/Runtime/Engine/Private/DataChannel.cpp
--- a/Engine/Source/Runtime/Engine/Private/DataChannel.cpp
+++ b/Engine/Source/Runtime/Engine/Private/DataChannel.cpp
@@ -3249,7 +3249,7 @@ void UActorChannel::ProcessBunch( FInBunch & Bunch )
 			}
 
 			check( !bSpawnedNewActor );
-			UE_LOG(LogNet, Warning, TEXT("UActorChannel::ProcessBunch: SerializeNewActor failed to find/spawn actor. Actor: %s, Channel: %i"), *GetFullNameSafe(NewChannelActor), ChIndex);
+			UE_LOG(LogNet, Warning, TEXT("UActorChannel::ProcessBunch: SerializeNewActor failed to find/spawn actor. Actor: %s, Channel: %i, FullNetGUIDPath: %s"), *GetFullNameSafe(NewChannelActor), ChIndex, *Connection->Driver->GuidCache->FullNetGUIDPath(ActorNetGUID));
 			Broken = 1;
 
 			if (!Connection->IsInternalAck()
