[NetPhysics] Dynamic number of redundant inputs per packet for autonomous proxy to server for NetworkPhysicsComponent.

diff --git a/Engine/Source/Runtime/Engine/Private/PhysicsEngine/NetworkPhysicsComponent.cpp b/Engine/Source/Runtime/Engine/Private/PhysicsEngine/NetworkPhysicsComponent.cpp
--- a/Engine/Source/Runtime/Engine/Private/PhysicsEngine/NetworkPhysicsComponent.cpp
+++ b/Engine/Source/Runtime/Engine/Private/PhysicsEngine/NetworkPhysicsComponent.cpp
@@ -23,11 +23,27 @@ namespace PhysicsReplicationCVars
 	namespace ResimulationCVars
 	{
 		int32 RedundantInputs = 2;
-		static FAutoConsoleVariableRef CVarResimRedundantInputs(TEXT("np2.Resim.RedundantInputs"), RedundantInputs, TEXT("How many extra inputs to send with each unreliable network message, to account for packetloss. From owning client to server and back to owning client."));
+		static FAutoConsoleVariableRef CVarResimRedundantInputs(TEXT("np2.Resim.RedundantInputs"), RedundantInputs, TEXT("How many extra inputs to send with each unreliable network message, to account for packetloss. From owning client to server and back to owning client. NOTE: This is disabled while np2.Resim.DynamicInputScaling.Enabled is enabled."));
 		int32 RedundantRemoteInputs = 1;
 		static FAutoConsoleVariableRef CVarResimRedundantRemoteInputs(TEXT("np2.Resim.RedundantRemoteInputs"), RedundantRemoteInputs, TEXT("How many extra inputs to send with each unreliable network message, to account for packetloss. From server to remote clients."));
 		int32 RedundantStates = 0;
 		static FAutoConsoleVariableRef CVarResimRedundantStates(TEXT("np2.Resim.RedundantStates"), RedundantStates, TEXT("How many extra states to send with each unreliable network message, to account for packetloss."));
+
+		bool bDynamicInputScalingEnabled = true;
+		static FAutoConsoleVariableRef CVarDynamicInputScalingEnabled(TEXT("np2.Resim.DynamicInputScaling.Enabled"), bDynamicInputScalingEnabled, TEXT("Enable dynmic scaling of number of inputs sent from owning client to the server to account for packet loss. The server will control the value based on how often the server has a hole in its input buffer. NOTE: This overrides np2.Resim.RedundantInputs"));
+		float DynamicInputScalingMaxInputsPercent = 0.1f;
+		static FAutoConsoleVariableRef CVarDynamicInputScalingMaxInputsPercent(TEXT("np2.Resim.DynamicInputScaling.MaxInputsPercent"), DynamicInputScalingMaxInputsPercent, TEXT("Default 0.1 (= 10%, value in percent as multiplier). Sets the max scalable number of inputs to network from owning client to server as a percentage of the physics fixed tick-rate. 10% of 30Hz = 3 inputs at max."));
+		int32 DynamicInputScalingMinInputs = 2;
+		static FAutoConsoleVariableRef CVarDynamicInputScalingMinInputs(TEXT("np2.Resim.DynamicInputScaling.MinInputs"), DynamicInputScalingMinInputs, TEXT("Default 2. Sets the minimum scalable number of inputs to network from owning client to server."));
+		float DynamicInputScalingIncreaseAverageMultiplier = 0.2f;
+		static FAutoConsoleVariableRef CVarDynamicInputScalingIncreaseAverageMultiplier(TEXT("np2.Resim.DynamicInputScaling.IncreaseAverageMultiplier"), DynamicInputScalingIncreaseAverageMultiplier, TEXT("Default 0.2 (= 20%). Multiplier for how fast the average input scaling value increases. NOTE it's recommended to have a higher value than np2.Resim.DynamicInputScaling.DecreaseAverageMultiplier so the average can grow quick when network conditions gets worse."));
+		float DynamicInputScalingDecreaseAverageMultiplier = 0.1f;
+		static FAutoConsoleVariableRef CVarDynamicInputScalingDecreaseAverageMultiplier(TEXT("np2.Resim.DynamicInputScaling.DecreaseAverageMultiplier"), DynamicInputScalingDecreaseAverageMultiplier, TEXT("Default 0.1 (= 10%). Multiplier for how fast the average input scaling value decreases. NOTE it's recommended to have a lower value than np2.Resim.DynamicInputScaling.IncreaseAverageMultiplier so the average doesn't try to decrease too quickly which can cause repeated desyncs."));
+		float DynamicInputScalingIncreaseTimeInterval = 2.0f;
+		static FAutoConsoleVariableRef CVarDynamicInputScalingIncreaseTimeInterval(TEXT("np2.Resim.DynamicInputScaling.IncreaseTimeInterval"), DynamicInputScalingIncreaseTimeInterval, TEXT("Default 2.0 (value in seconds). How often dynamic scaling can increase the number of inputs to send." ));
+		float DynamicInputScalingDecreaseTimeInterval = 10.0f;
+		static FAutoConsoleVariableRef CVarDynamicInputScalingDecreaseTimeInterval(TEXT("np2.Resim.DynamicInputScaling.DecreaseTimeInterval"), DynamicInputScalingDecreaseTimeInterval, TEXT("Default 10.0 (value in seconds). How often dynamic scaling can decrease the number of inputs to send."));
+
 		bool bAllowRewindToClosestState = true;
 		static FAutoConsoleVariableRef CVarResimAllowRewindToClosestState(TEXT("np2.Resim.AllowRewindToClosestState"), bAllowRewindToClosestState, TEXT("When rewinding to a specific frame, if the client doens't have state data for that frame, use closest data available. Only affects the first rewind frame, when FPBDRigidsEvolution is set to Reset."));
 		bool bCompareStateToTriggerRewind = false;
@@ -510,11 +526,11 @@ void UNetworkPhysicsComponent::InitializeComponent()
 
 			if (ReplicatedInputs.History)
 			{
-				ReplicatedInputs.History->ResizeDataHistory(InputsToNetwork);
+				ReplicatedInputs.History->ResizeDataHistory(InputsToNetwork_OwnerDefault);
 			}
 			if (ReplicatedRemoteInputs.History)
 			{
-				ReplicatedRemoteInputs.History->ResizeDataHistory(InputsToNetworkRemote);
+				ReplicatedRemoteInputs.History->ResizeDataHistory(InputsToNetwork_Simulated);
 			}
 			if (ReplicatedStates.History)
 			{
@@ -532,7 +548,8 @@ void UNetworkPhysicsComponent::InitializeComponent()
 				// Create async component to run on Physics Thread
 				NetworkPhysicsComponent_Internal = Solver->CreateAndRegisterSimCallbackObject_External<FAsyncNetworkPhysicsComponent>();
 				NetworkPhysicsComponent_Internal->RootPhysicsObject = RootPhysicsObject;
-				NetworkPhysicsComponent_Internal->InputsToNetwork = InputsToNetwork;
+				NetworkPhysicsComponent_Internal->InputsToNetwork_OwnerDefault = InputsToNetwork_OwnerDefault;
+				NetworkPhysicsComponent_Internal->InputsToNetwork_Simulated = InputsToNetwork_Simulated;
 				NetworkPhysicsComponent_Internal->StatesToNetwork = StatesToNetwork;
 				NetworkPhysicsComponent_Internal->bCompareStateToTriggerRewind = bCompareStateToTriggerRewind;
 				NetworkPhysicsComponent_Internal->bCompareStateToTriggerRewindIncludeSimProxies = bCompareStateToTriggerRewindIncludeSimProxies;
@@ -588,6 +605,7 @@ void UNetworkPhysicsComponent::GetLifetimeReplicatedProps(TArray< FLifetimePrope
 	ReplicatedParamsAll.RepNotifyCondition = REPNOTIFY_Always;
 	ReplicatedParamsAll.bIsPushBased = true;
 
+	DOREPLIFETIME_CONDITION(UNetworkPhysicsComponent, InputsToNetwork_Owner, COND_OwnerOnly);
 	DOREPLIFETIME_WITH_PARAMS_FAST(UNetworkPhysicsComponent, ReplicatedInputs, ReplicatedParamsOwner);
 	DOREPLIFETIME_WITH_PARAMS_FAST(UNetworkPhysicsComponent, ReplicatedRemoteInputs, ReplicatedParamsRemote);
 	DOREPLIFETIME_WITH_PARAMS_FAST(UNetworkPhysicsComponent, ReplicatedStates, ReplicatedParamsAll);
@@ -608,11 +626,16 @@ void UNetworkPhysicsComponent::NetworkMarshaledData()
 		const bool bIsServer = HasServerWorld();
 		while (Chaos::TSimCallbackOutputHandle<FAsyncNetworkPhysicsComponentOutput> AsyncOutput = NetworkPhysicsComponent_Internal->PopFutureOutputData_External())
 		{
+			if (AsyncOutput->InputsToNetwork_Owner.IsSet())
+			{
+				// Only marshaled from PT to GT on the server, InputsToNetwork_Owner is a replicated property towards the owner
+				InputsToNetwork_Owner = *AsyncOutput->InputsToNetwork_Owner;
+			}
+
 			// Unimportant / Unreliable
 			if (bEnableUnreliableFlow
 				&& AsyncOutput->InputData
-				&& AsyncOutput->InputData->HasDataInHistory()
-				&& AsyncOutput->InputData->CopyAllData(*ReplicatedInputs.History, /*bIncludeUnimportant*/ true, /*bIncludeImportant*/ bEnableReliableFlow == false))
+				&& AsyncOutput->InputData->HasDataInHistory())
 			{
 				if (bIsServer)
 				{
@@ -620,17 +643,29 @@ void UNetworkPhysicsComponent::NetworkMarshaledData()
 					{
 						// Send inputs to remote clients after getting marshaled from PT if server is the one controlling the component
 						ReplicatedRemoteInputs.History->SetRecordDataIncremental(true); // Only record data that is newer than already cached data
-						ReplicatedInputs.History->CopyAllData(*ReplicatedRemoteInputs.History, /*bIncludeUnimportant*/ true, /*bIncludeImportant*/ true);
-						MARK_PROPERTY_DIRTY_FROM_NAME(UNetworkPhysicsComponent, ReplicatedRemoteInputs, this);
+						if (AsyncOutput->InputData->CopyAllData(*ReplicatedRemoteInputs.History, /*bIncludeUnimportant*/ true, /*bIncludeImportant*/ true))
+						{
+							MARK_PROPERTY_DIRTY_FROM_NAME(UNetworkPhysicsComponent, ReplicatedRemoteInputs, this);
+						}
 					}
 
+					// TODO, Only copy/replicate data if InputFrame is INDEX_NONE i.e. the input has been altered by the server
+
 					// Server sends inputs through property replication to owning client
-					MARK_PROPERTY_DIRTY_FROM_NAME(UNetworkPhysicsComponent, ReplicatedInputs, this);
+					if (AsyncOutput->InputData->CopyAllData(*ReplicatedInputs.History, /*bIncludeUnimportant*/ true, /*bIncludeImportant*/ bEnableReliableFlow == false))
+					{
+						MARK_PROPERTY_DIRTY_FROM_NAME(UNetworkPhysicsComponent, ReplicatedInputs, this);
+					}
 				}
-				else if (IsLocallyControlled())
+				else if (IsLocallyControlled()) // Client-side
 				{
-					// Clients send inputs through an RPC to the server
-					ServerReceiveInputData(ReplicatedInputs);
+					ReplicatedInputs.History->ResizeDataHistory(AsyncOutput->InputData->GetHistorySize(), EAllowShrinking::Yes);
+
+					if (AsyncOutput->InputData->CopyAllData(*ReplicatedInputs.History, /*bIncludeUnimportant*/ true, /*bIncludeImportant*/ bEnableReliableFlow == false))
+					{
+						// Clients send inputs through an RPC to the server
+						ServerReceiveInputData(ReplicatedInputs);
+					}
 				}
 			}
 
@@ -644,7 +679,7 @@ void UNetworkPhysicsComponent::NetworkMarshaledData()
 						continue;
 					}
 
-					ReplicatedImportantInput.History->ResizeDataHistory(InputImportant->GetHistorySize());
+					ReplicatedImportantInput.History->ResizeDataHistory(InputImportant->GetHistorySize(), EAllowShrinking::Yes);
 					if (InputImportant->CopyAllData(*ReplicatedImportantInput.History, /*bIncludeUnimportant*/ false, /*bIncludeImportant*/ true))
 					{
 						if (bIsServer)
@@ -681,7 +716,7 @@ void UNetworkPhysicsComponent::NetworkMarshaledData()
 							continue;
 						}
 
-						ReplicatedImportantState.History->ResizeDataHistory(StateImportant->GetHistorySize());
+						ReplicatedImportantState.History->ResizeDataHistory(StateImportant->GetHistorySize(), EAllowShrinking::Yes);
 						if (StateImportant->CopyAllData(*ReplicatedImportantState.History, /*bIncludeUnimportant*/ false, /*bIncludeImportant*/ true))
 						{
 							MulticastReceiveImportantStateData(ReplicatedImportantState);
@@ -738,10 +773,11 @@ void UNetworkPhysicsComponent::OnRep_SetReplicatedStates()
 	{
 		if (!AsyncInput->StateData)
 		{
-			AsyncInput->StateData = StateHelper->CreateUniqueRewindHistory(StatesToNetwork);
+			AsyncInput->StateData = StateHelper->CreateUniqueRewindHistory(ReplicatedStates.History->GetHistorySize());
 		}
 
 		AsyncInput->StateData->ResetFast();
+		AsyncInput->StateData->ResizeDataHistory(ReplicatedStates.History->GetHistorySize(), EAllowShrinking::No);
 		ReplicatedStates.History->CopyAllData(*AsyncInput->StateData.Get(), /*bIncludeUnimportant*/ true, /*bIncludeImportant*/ true);
 	}
 
@@ -775,10 +811,11 @@ void UNetworkPhysicsComponent::OnRep_SetReplicatedInputs()
 	{
 		if (!AsyncInput->InputData)
 		{
-			AsyncInput->InputData = InputHelper->CreateUniqueRewindHistory(InputsToNetwork);
+			AsyncInput->InputData = InputHelper->CreateUniqueRewindHistory(ReplicatedInputs.History->GetHistorySize());
 		}
 
 		AsyncInput->InputData->ResetFast();
+		AsyncInput->InputData->ResizeDataHistory(ReplicatedInputs.History->GetHistorySize(), EAllowShrinking::No);
 		ReplicatedInputs.History->CopyAllData(*AsyncInput->InputData.Get(), /*bIncludeUnimportant*/ true, /*bIncludeImportant*/ true);
 	}
 }
@@ -794,10 +831,11 @@ void UNetworkPhysicsComponent::OnRep_SetReplicatedRemoteInputs()
 	{
 		if (!AsyncInput->InputData)
 		{
-			AsyncInput->InputData = InputHelper->CreateUniqueRewindHistory(InputsToNetwork);
+			AsyncInput->InputData = InputHelper->CreateUniqueRewindHistory(ReplicatedRemoteInputs.History->GetHistorySize());
 		}
 
 		AsyncInput->InputData->ResetFast();
+		AsyncInput->InputData->ResizeDataHistory(ReplicatedRemoteInputs.History->GetHistorySize(), EAllowShrinking::No);
 		ReplicatedRemoteInputs.History->CopyAllData(*AsyncInput->InputData.Get(), /*bIncludeUnimportant*/ true, /*bIncludeImportant*/ true);
 	}
 }
@@ -819,7 +857,7 @@ void UNetworkPhysicsComponent::ServerReceiveInputData_Implementation(const FNetw
 	{
 		if (!AsyncInput->InputData)
 		{
-			AsyncInput->InputData = InputHelper->CreateUniqueRewindHistory(InputsToNetwork);
+			AsyncInput->InputData = InputHelper->CreateUniqueRewindHistory(ClientInputs.History->GetHistorySize());
 		}
 
 		// Validate data in the received inputs
@@ -829,6 +867,7 @@ void UNetworkPhysicsComponent::ServerReceiveInputData_Implementation(const FNetw
 		}
 
 		AsyncInput->InputData->ResetFast();
+		AsyncInput->InputData->ResizeDataHistory(ClientInputs.History->GetHistorySize(), EAllowShrinking::No);
 		ClientInputs.History->CopyAllData(*AsyncInput->InputData.Get(), /*bIncludeUnimportant*/ true, /*bIncludeImportant*/ true);
 
 		// Send received inputs to remote clients
@@ -1036,6 +1075,7 @@ void UNetworkPhysicsComponent::UpdateAsyncComponent(const bool bFullUpdate)
 						AsyncInput->NetworkPhysicsTickOffset = PC->GetNetworkPhysicsTickOffset();
 					}
 				}
+				AsyncInput->InputsToNetwork_Owner = InputsToNetwork_Owner;
 			}
 
 			// bIsLocallyControlled is marshaled outside of the bFullUpdate because it's not always set when last bFullUpdate is called.
@@ -1226,6 +1266,11 @@ void FAsyncNetworkPhysicsComponent::ConsumeAsyncInput(const int32 PhysicsStep)
 			{
 				NetworkPhysicsTickOffset = *AsyncInput->NetworkPhysicsTickOffset;
 			}
+			if (AsyncInput->InputsToNetwork_Owner.IsSet())
+			{
+				// Only marshaled from GT to PT on the client
+				InputsToNetwork_Owner = *AsyncInput->InputsToNetwork_Owner;
+			}
 			if (AsyncInput->PhysicsReplicationMode.IsSet())
 			{
 				PhysicsReplicationMode = *AsyncInput->PhysicsReplicationMode;
@@ -1283,11 +1328,13 @@ void FAsyncNetworkPhysicsComponent::ConsumeAsyncInput(const int32 PhysicsStep)
 
 		/** Continuously marshaled data */
 		{
+			const bool bIsServer = IsServer();
+
 			/** Receive data helper */
 			auto ReceiveHelper = [&](Chaos::FBaseRewindHistory* History, Chaos::FBaseRewindHistory* ReceiveData, const bool bImportant, const bool bCompareData)
 			{
-				const bool bCompareDataForRewind = bCompareData && !IsServer();
-				const int32 ResimFrame = History->ReceiveNewData(*ReceiveData, (IsServer() ? 0 : NetworkPhysicsTickOffset), bCompareDataForRewind, bImportant);
+				const bool bCompareDataForRewind = bCompareData && !bIsServer;
+				const int32 ResimFrame = History->ReceiveNewData(*ReceiveData, (bIsServer ? 0 : NetworkPhysicsTickOffset), bCompareDataForRewind, bImportant);
 				if (bCompareDataForRewind)
 				{
 					TriggerResimulation(ResimFrame);
@@ -1295,7 +1342,7 @@ void FAsyncNetworkPhysicsComponent::ConsumeAsyncInput(const int32 PhysicsStep)
 
 #if DEBUG_NETWORK_PHYSICS
 				{
-					FString NetRoleString = IsServer() ? FString("SERVER") : (IsLocallyControlled() ? FString("AUTONO") : FString("PROXY "));
+					FString NetRoleString = bIsServer ? FString("SERVER") : (IsLocallyControlled() ? FString("AUTONO") : FString("PROXY "));
 					ReceiveData->DebugData(FString::Printf(TEXT("%s | PT | RECEIVE DATA | LatestFrame: %d | bImportant: %d | Name: %s"), *NetRoleString, ReceiveData->GetLatestFrame(), bImportant, *GetActorName()));
 				}
 #endif
@@ -1313,7 +1360,7 @@ void FAsyncNetworkPhysicsComponent::ConsumeAsyncInput(const int32 PhysicsStep)
 				/* Extract latest received input from client on the server, to be used if input buffer runs empty
 				* TODO, improve flow to not require this before ReceiveHelper */
 				FNetworkPhysicsData* PhysicsData = nullptr;
-				if (IsServer() && LatestInputReceiveData)
+				if (bIsServer && LatestInputReceiveData)
 				{
 					PhysicsData = LatestInputReceiveData.Get();
 					if (AsyncInput->InputData->ExtractData(AsyncInput->InputData->GetLatestFrame(), false, PhysicsData, true) == false)
@@ -1329,7 +1376,7 @@ void FAsyncNetworkPhysicsComponent::ConsumeAsyncInput(const int32 PhysicsStep)
 				}
 
 				// Validate data in the received inputs on the server
-				if (!ComponentSettings.GetValidateDataOnGameThread() && IsServer() && ActorComponent.IsValid() && ActorComponent.Get()->IsBeingDestroyed() == false)
+				if (!ComponentSettings.GetValidateDataOnGameThread() && bIsServer && ActorComponent.IsValid() && ActorComponent.Get()->IsBeingDestroyed() == false)
 				{
 					AsyncInput->InputData->ValidateDataInHistory(ActorComponent.Get());
 				}
@@ -1337,7 +1384,7 @@ void FAsyncNetworkPhysicsComponent::ConsumeAsyncInput(const int32 PhysicsStep)
 
 				/* If the server-side input history doesn't have any entries ahead of the current physics tick, the input buffer is empty, inject the latest received input as the input for the current tick.
 				* This happens during a desync where the client is far behind the server */
-				if (IsServer() && InputHistory->GetLatestFrame() < PhysicsStep && PhysicsData)
+				if (bIsServer && InputHistory->GetLatestFrame() < PhysicsStep && PhysicsData)
 				{
 #if DEBUG_NETWORK_PHYSICS
 					UE_LOG(LogChaos, Log, TEXT("SERVER | PT | Input Buffer Empty, Injecting Received Input at frame %d || LocalFrame = %d || ServerFrame = %d || InputFrame = %d || Data: %s || Actor: %s")
@@ -1389,7 +1436,6 @@ FAsyncNetworkPhysicsComponentOutput& FAsyncNetworkPhysicsComponent::GetAsyncOutp
 	if ((IsLocallyControlled() || IsServer()) && AsyncOutput.InputData == nullptr && InputHistory != nullptr)
 	{
 		AsyncOutput.InputData = InputHistory->CreateNew();
-		AsyncOutput.InputData->ResizeDataHistory(InputsToNetwork);
 	}
 
 	// StateData marshal from PT to GT is needed for: Server
@@ -1419,8 +1465,8 @@ void FAsyncNetworkPhysicsComponent::OnPreProcessInputs_Internal(const int32 Phys
 
 #if DEBUG_NETWORK_PHYSICS
 	{
-		const int32 InputBufferSize = (IsServer() && InputHistory) ? (InputHistory->GetLatestFrame() - PhysicsStep) : 0;
-		const FString NetRoleString = IsServer() ? FString("SERVER") : (IsLocallyControlled() ? FString("AUTONO") : FString("PROXY "));
+		const int32 InputBufferSize = (bIsServer && InputHistory) ? (InputHistory->GetLatestFrame() - PhysicsStep) : 0;
+		const FString NetRoleString = bIsServer ? FString("SERVER") : (IsLocallyControlled() ? FString("AUTONO") : FString("PROXY "));
 		UE_LOG(LogChaos, Log, TEXT("%s | PT | OnPreProcessInputsInternal | At Frame %d | IsResim: %d | FirstResimFrame: %d | InputBuffer: %d | Name = %s"), *NetRoleString, PhysicsStep, bIsSolverResim, bIsSolverReset, InputBufferSize, *GetActorName());
 	}
 #endif
@@ -1428,7 +1474,7 @@ void FAsyncNetworkPhysicsComponent::OnPreProcessInputs_Internal(const int32 Phys
 	if (ActorComponent.IsValid() && ActorComponent.Get()->IsBeingDestroyed() == false)
 	{
 		// Apply replicated state on clients if we are resimulating or on simulated proxies if setting is enabled
-		const bool bApplySimProxyState = ComponentSettings.GetApplySimProxyStateAtRuntime() && !IsServer() && !IsLocallyControlled();
+		const bool bApplySimProxyState = ComponentSettings.GetApplySimProxyStateAtRuntime() && !bIsServer && !IsLocallyControlled();
 		if ((bApplySimProxyState || bIsSolverResim) && StateHistory && StateData)
 		{
 			FNetworkPhysicsData* PhysicsData = StateData.Get();
@@ -1452,13 +1498,12 @@ void FAsyncNetworkPhysicsComponent::OnPreProcessInputs_Internal(const int32 Phys
 		}
 
 		// Apply replicated inputs on server and simulated proxies if setting is enabled, and on local player if we are resimulating
-		const bool bApplySimProxyInput = ComponentSettings.GetApplySimProxyInputAtRuntime() && !IsServer() && !IsLocallyControlled();
-		const bool bApplyServerInput = IsServer() && !IsLocallyControlled();
+		const bool bApplySimProxyInput = ComponentSettings.GetApplySimProxyInputAtRuntime() && !bIsServer && !IsLocallyControlled();
+		const bool bApplyServerInput = bIsServer && !IsLocallyControlled();
 		if ((bApplyServerInput || bApplySimProxyInput || bIsSolverResim) && InputHistory && InputData)
 		{
 			FNetworkPhysicsData* PhysicsData = InputData.Get();
 			int32 NextExpectedLocalFrame = PhysicsData->LocalFrame + 1;
-			PhysicsData->LocalFrame = PhysicsStep;
 
 			// There are important inputs earlier than upcoming input to apply
 			if (NewImportantInputFrame < NextExpectedLocalFrame && !bIsSolverResim)
@@ -1523,12 +1568,23 @@ void FAsyncNetworkPhysicsComponent::OnPreProcessInputs_Internal(const int32 Phys
 				// If the extracted input data was altered (extrapolated, interpolated) on the server, record it into the history for it to get replicated to clients
 				if (bIsServer && PhysicsData->InputFrame == INDEX_NONE)
 				{
+					// Check that buffer is not empty
+					if (PhysicsStep <= InputHistory->GetLatestFrame())
+					{
+						if (PhysicsData->LocalFrame < PhysicsStep) 
+						{
+							// A missing input was detected and buffer is not empty, inform the owning client to send more inputs in each RPC to not get a gaps in the buffer
+							// NOTE: We don't send more extra inputs when the buffer runs empty since that case is corrected via time dilation, not sending extra inputs
+							MissingInputCount++;
+						}
+					}
 #if DEBUG_NETWORK_PHYSICS
-					if (PhysicsStep > InputHistory->GetLatestFrame())
+					else
 					{
 						UE_LOG(LogChaos, Log, TEXT("		Non-Determinism: Input buffer Empty, input for frame %d was extrapolated from frame: %d"), PhysicsStep, PhysicsData->LocalFrame);
 					}
 #endif
+
 					PhysicsData->bReceivedData = true; // Mark the input data as received so that it doesn't get overwritten by incoming client inputs
 					PhysicsData->bImportant = false;
 					PhysicsData->LocalFrame = PhysicsStep;
@@ -1571,7 +1627,7 @@ void FAsyncNetworkPhysicsComponent::OnPostProcessInputs_Internal(const int32 Phy
 
 #if DEBUG_NETWORK_PHYSICS
 	{
-		FString NetRoleString = IsServer() ? FString("SERVER") : (IsLocallyControlled() ? FString("AUTONO") : FString("PROXY "));
+		FString NetRoleString = bIsServer ? FString("SERVER") : (IsLocallyControlled() ? FString("AUTONO") : FString("PROXY "));
 		UE_LOG(LogChaos, Log, TEXT("%s | PT | OnPostProcessInputsInternal | At Frame %d | IsResim: %d | FirstResimFrame: %d | Name = %s"), *NetRoleString, PhysicsStep, bIsSolverResim, bIsSolverReset, *GetActorName());
 	}
 #endif
@@ -1648,6 +1704,12 @@ void FAsyncNetworkPhysicsComponent::SendInputData_Internal(FAsyncNetworkPhysicsC
 {
 	const bool bIsServer = IsServer();
 
+	if (bIsServer)
+	{
+		UpdateDynamicInputScaling();
+		AsyncOutput.InputsToNetwork_Owner = InputsToNetwork_Owner;
+	}
+
 	// Inputs are sent from the server or locally controlled actors/pawns
 	if (AsyncOutput.InputData && InputHistory && (IsLocallyControlled() || bIsServer))
 	{
@@ -1659,18 +1721,18 @@ void FAsyncNetworkPhysicsComponent::SendInputData_Internal(FAsyncNetworkPhysicsC
 		// -- Default / Unreliable Flow --
 		if (ComponentSettings.GetEnableUnreliableFlow())
 		{
-			const int32 FromFrame = FMath::Max(0, ToFrame - InputsToNetwork - 1); // Remove 1 since both ToFrame and FromFrame are inclusive
+			const uint16 NumInputsToNetwork = bIsServer ? InputsToNetwork_Simulated : InputsToNetwork_Owner;
+			const int32 FromFrame = FMath::Max(0, ToFrame - NumInputsToNetwork - 1); // Remove 1 since both ToFrame and FromFrame are inclusive
+
+			AsyncOutput.InputData->ResizeDataHistory(NumInputsToNetwork);
 
-			// Resize marshaling history if needed
-			AsyncOutput.InputData->ResizeDataHistory(InputsToNetwork);
-			
 			if (InputHistory->CopyData(*AsyncOutput.InputData, FromFrame, ToFrame, /*bIncludeUnimportant*/ true, /*bIncludeImportant*/ ComponentSettings.GetEnableReliableFlow() == false))
 			{
 #if DEBUG_NETWORK_PHYSICS
 				{
 					const int32 LocalFrame = GetRigidSolver()->GetCurrentFrame();
-					const int32 ServerFrame = IsServer() ? LocalFrame : LocalFrame + GetNetworkPhysicsTickOffset();
-					FString NetRoleString = IsServer() ? FString("SERVER") : (IsLocallyControlled() ? FString("AUTONO") : FString("PROXY "));
+					const int32 ServerFrame = bIsServer ? LocalFrame : LocalFrame + GetNetworkPhysicsTickOffset();
+					FString NetRoleString = bIsServer ? FString("SERVER") : (IsLocallyControlled() ? FString("AUTONO") : FString("PROXY "));
 					AsyncOutput.InputData->DebugData(FString::Printf(TEXT("%s | PT | SendInputData_Internal | UNRELIABLE | CurrentLocalFrame = %d | CurrentServerFrame = %d | Name: %s"), *NetRoleString, LocalFrame, ServerFrame, *GetActorName()));
 				}
 #endif
@@ -1699,8 +1761,8 @@ void FAsyncNetworkPhysicsComponent::SendInputData_Internal(FAsyncNetworkPhysicsC
 #if DEBUG_NETWORK_PHYSICS
 					{
 						const int32 LocalFrame = GetRigidSolver()->GetCurrentFrame();
-						const int32 ServerFrame = IsServer() ? LocalFrame : LocalFrame + GetNetworkPhysicsTickOffset();
-						FString NetRoleString = IsServer() ? FString("SERVER") : (IsLocallyControlled() ? FString("AUTONO") : FString("PROXY "));
+						const int32 ServerFrame = bIsServer ? LocalFrame : LocalFrame + GetNetworkPhysicsTickOffset();
+						FString NetRoleString = bIsServer ? FString("SERVER") : (IsLocallyControlled() ? FString("AUTONO") : FString("PROXY "));
 						AsyncOutput.InputDataImportant[Idx]->DebugData(FString::Printf(TEXT("%s | PT | SendInputData_Internal | RELIABLE | CurrentLocalFrame = %d | CurrentServerFrame = %d | Name: %s"), *NetRoleString, LocalFrame, ServerFrame, *GetActorName()));
 					}
 #endif
@@ -1848,6 +1910,62 @@ const float FAsyncNetworkPhysicsComponent::GetCurrentInputDecay(const FNetworkPh
 	return InputDecay;
 }
 
+void FAsyncNetworkPhysicsComponent::UpdateDynamicInputScaling()
+{
+	if (!PhysicsReplicationCVars::ResimulationCVars::bDynamicInputScalingEnabled)
+	{
+		InputsToNetwork_Owner = InputsToNetwork_OwnerDefault;
+		return;
+	}
+
+	if (!IsServer())
+	{
+		return;
+	}
+
+	Chaos::FPBDRigidsSolver* RigidSolver = GetRigidSolver();
+	if (!RigidSolver)
+	{
+		return;
+	}
+
+	const float TimeSinceLastDynamicScaling = RigidSolver->GetSolverTime() - TimeOfLastDynamicInputScaling;
+
+	if (MissingInputCount > 0)
+	{
+		if (TimeSinceLastDynamicScaling > PhysicsReplicationCVars::ResimulationCVars::DynamicInputScalingIncreaseTimeInterval)
+		{
+			const uint16 MaxInputsValue = static_cast<uint16>(FMath::CeilToInt32(PhysicsReplicationCVars::ResimulationCVars::DynamicInputScalingMaxInputsPercent / (RigidSolver->GetAsyncDeltaTime())));
+
+			// Increase the amount of inputs the owner sends
+			InputsToNetwork_Owner++;
+
+			// Update the average value for minimum clamping
+			DynamicInputScalingAverageInputs += (static_cast<float>(InputsToNetwork_Owner) - DynamicInputScalingAverageInputs) * PhysicsReplicationCVars::ResimulationCVars::DynamicInputScalingIncreaseAverageMultiplier;
+			
+			// Clamp to maximum valid value
+			InputsToNetwork_Owner = FMath::Min(InputsToNetwork_Owner, MaxInputsValue);
+
+			TimeOfLastDynamicInputScaling = RigidSolver->GetSolverTime();
+			MissingInputCount = 0;
+		}
+	}
+	else if (TimeSinceLastDynamicScaling > PhysicsReplicationCVars::ResimulationCVars::DynamicInputScalingDecreaseTimeInterval)
+	{
+		// Decrease the amount of inputs the owner sends
+		InputsToNetwork_Owner--;
+
+		// Update the average value for minimum clamping, perform before clamping to allow for decreasing average even if the clamp might still round up.
+		DynamicInputScalingAverageInputs += (static_cast<float>(InputsToNetwork_Owner) - DynamicInputScalingAverageInputs) * PhysicsReplicationCVars::ResimulationCVars::DynamicInputScalingDecreaseAverageMultiplier;
+
+		// Clamp to minimum valid value
+		const uint16 MinInputsValue = FMath::Max(static_cast<uint16>(FMath::RoundToInt(DynamicInputScalingAverageInputs), static_cast<uint16>(PhysicsReplicationCVars::ResimulationCVars::DynamicInputScalingMinInputs)));
+		InputsToNetwork_Owner = FMath::Max(InputsToNetwork_Owner, MinInputsValue);
+
+		TimeOfLastDynamicInputScaling = RigidSolver->GetSolverTime();
+	}
+}
+
 void FAsyncNetworkPhysicsComponent::RegisterDataHistoryInRewindData()
 {
 	if (Chaos::FPhysicsSolverBase* BaseSolver = GetSolver())
diff --git a/Engine/Source/Runtime/Engine/Private/PlayerController.cpp b/Engine/Source/Runtime/Engine/Private/PlayerController.cpp
--- a/Engine/Source/Runtime/Engine/Private/PlayerController.cpp
+++ b/Engine/Source/Runtime/Engine/Private/PlayerController.cpp
@@ -6304,7 +6304,7 @@ FAsyncPhysicsTimestamp APlayerController::GetPhysicsTimestamp(float DeltaSeconds
 
 void APlayerController::UpdateServerAsyncPhysicsTickOffset()
 {
-	if (NetworkPhysicsCvars::TickOffsetUpdateIntervalTime <= 0 || TickOffsetSyncCountdown < (NetworkPhysicsCvars::TickOffsetUpdateIntervalTime / 1000.0f))
+	if (NetworkPhysicsCvars::TickOffsetUpdateIntervalTime <= 0 || TickOffsetSyncCountdown < (NetworkPhysicsCvars::TickOffsetUpdateIntervalTime * 0.001f))
 	{
 		return;
 	}
@@ -6344,7 +6344,7 @@ void APlayerController::ServerSendLatestAsyncPhysicsTimestamp_Implementation(FAs
 
 	// Get current server timestamp and add the frame buffer to the ServerFrame
 	FAsyncPhysicsTimestamp ActualTimestamp = GetPhysicsTimestamp();
-	const int32 BufferTickSize = FMath::CeilToInt((NetworkPhysicsCvars::TickOffsetBufferTime / 1000.f) / AsyncFixedTimeStepSize);
+	const int32 BufferTickSize = FMath::CeilToInt((NetworkPhysicsCvars::TickOffsetBufferTime * 0.001f) / AsyncFixedTimeStepSize);
 	ActualTimestamp.ServerFrame += BufferTickSize;
 
 	// Mark offset as assigned when we get a valid predicted server frame.
@@ -6359,8 +6359,8 @@ void APlayerController::ServerSendLatestAsyncPhysicsTimestamp_Implementation(FAs
 	// Send update to client if offset is not assigned or over correction limits
 	// Note that we are sending the current ServerFrame along with the frame buffer added, to the client.
 	if (!bNetworkPhysicsTickOffsetAssigned
-		|| TimestampTimeDiff > (NetworkPhysicsCvars::TickOffsetCorrectionSizeTimeLimit / 1000.f)
-		|| CurrentTime - NetworkPhysicsTickOffsetDesyncAccumulatedTime > (NetworkPhysicsCvars::TickOffsetCorrectionTimeLimit / 1000.0f))
+		|| TimestampTimeDiff > (NetworkPhysicsCvars::TickOffsetCorrectionSizeTimeLimit * 0.001f)
+		|| CurrentTime - NetworkPhysicsTickOffsetDesyncAccumulatedTime > (NetworkPhysicsCvars::TickOffsetCorrectionTimeLimit * 0.001f))
 	{
 		Timestamp.ServerFrame = ActualTimestamp.ServerFrame;
 		NetworkPhysicsTickOffsetDesyncAccumulatedTime = CurrentTime;
diff --git a/Engine/Source/Runtime/Engine/Public/Physics/NetworkPhysicsComponent.h b/Engine/Source/Runtime/Engine/Public/Physics/NetworkPhysicsComponent.h
--- a/Engine/Source/Runtime/Engine/Public/Physics/NetworkPhysicsComponent.h
+++ b/Engine/Source/Runtime/Engine/Public/Physics/NetworkPhysicsComponent.h
@@ -855,24 +855,26 @@ private:
 	/** Setup data type and history in async input and trigger the creation of data history in the async component */
 	ENGINE_API void CreateAsyncDataHistory();
 
-	/** Set number of inputs to send over the network with each message, clamped to 1 as minimum 
-	* From owning client to server and server to owning client. */
-	void SetNumberOfInputsToNetwork(uint16 NumInputs)
+	/** Set default number of inputs to send over the network with each message, clamped to 1 as minimum 
+	* From owning client to server.
+	* NOTE: this will be overridden if np2.Resim.DynamicInputScaling.Enabled is enabled. */
+	void SetNumberOfInputsToNetwork(const uint16 NumInputs)
 	{
-		InputsToNetwork = FMath::Max((uint16)(NumInputs), (uint16)1);
+		InputsToNetwork_OwnerDefault = FMath::Max(NumInputs, static_cast<uint16>(1));
+		InputsToNetwork_Owner = InputsToNetwork_OwnerDefault;
 	}
 
 	/** Set number of inputs to send over the network with each message, clamped to 1 as minimum.
-	* From server to remote clients. */
-	void SetNumberOfRemoteInputsToNetwork(uint16 NumInputs)
+	* From server to remote clients (simulated proxies) */
+	void SetNumberOfRemoteInputsToNetwork(const uint16 NumInputs)
 	{
-		InputsToNetworkRemote = FMath::Max((uint16)(NumInputs), (uint16)1);
+		InputsToNetwork_Simulated = FMath::Max(NumInputs, static_cast<uint16>(1));
 	}
 
 	/** Set number of states to send over the network with each message, clamped to 1 as minimum */
-	void SetNumberOfStatesToNetwork(uint16 NumInputs)
+	void SetNumberOfStatesToNetwork(const uint16 NumInputs)
 	{
-		StatesToNetwork = FMath::Max((uint16)(NumInputs), (uint16)1);
+		StatesToNetwork = FMath::Max(NumInputs, static_cast<uint16>(1));
 	}
 
 public:
@@ -910,11 +912,15 @@ private:
 	// Helper for the creation of state data and history with correct derived type
 	TUniquePtr<FNetworkPhysicsDataHelper> StateHelper;
 
-	// Send last N number of inputs each replication call from owning client to server and server to owning client
-	uint16 InputsToNetwork = 3;
+	// The number of inputs the owning client should send to the server with each RPC, replicated from the server. This is dynamically scaled based on when there are holes in the inputs buffer if np2.Resim.DynamicInputScaling.Enabled is enabled
+	UPROPERTY( Replicated )
+	uint16 InputsToNetwork_Owner = 3;
+
+	// The default value for number for InputsToNetwork_Owner, acts as the initial value and the cap when dynamically adjusting InputsToNetwork_Owner
+	uint16 InputsToNetwork_OwnerDefault = 3;
 
 	// Send last N number of inputs each replication call from server to remote clients
-	uint16 InputsToNetworkRemote = 2;
+	uint16 InputsToNetwork_Simulated = 2;
 
 	// Send last N number of states each replication call from server to remote clients
 	uint16 StatesToNetwork = 1;
@@ -1027,6 +1033,14 @@ public:
 	/** Returns the current amount of input decay during resimulation as a magnitude from 0.0 to 1.0. Returns 0 if not currently resimulating. */
 	UE_DEPRECATED(5.5, "Deprecated, use FAsyncNetworkPhysicsComponent::GetCurrentInputDecay instead.")
 		ENGINE_API const float GetCurrentInputDecay(FNetworkPhysicsData* PhysicsData);
+
+private:
+#if WITH_EDITORONLY_DATA
+	UE_DEPRECATED(5.6, "Deprecated, use InputsToNetwork_Owner instead")
+		uint16 InputsToNetwork = 3;
+	UE_DEPRECATED(5.6, "Deprecated, use InputsToNetwork_Simulated instead")
+		uint16 InputsToNetworkRemote = 2;
+#endif
 };
 
 /** DEPRECATED UE 5.4 */
@@ -1042,10 +1056,10 @@ FORCEINLINE void UNetworkPhysicsComponent::CreateDataHistory(UActorComponent* Hi
 	InputHelper = MakeUnique<TNetworkPhysicsDataHelper<typename PhysicsTraits::InputsType>>();
 	StateHelper = MakeUnique<TNetworkPhysicsDataHelper<typename PhysicsTraits::StatesType>>();
 
-	ReplicatedInputs.History = MakeUnique<TNetRewindHistory<typename PhysicsTraits::InputsType>>(InputsToNetwork);
+	ReplicatedInputs.History = MakeUnique<TNetRewindHistory<typename PhysicsTraits::InputsType>>(InputsToNetwork_OwnerDefault);
 	ReplicatedInputs.Owner = this;
 
-	ReplicatedRemoteInputs.History = MakeUnique<TNetRewindHistory<typename PhysicsTraits::InputsType>>(InputsToNetworkRemote);
+	ReplicatedRemoteInputs.History = MakeUnique<TNetRewindHistory<typename PhysicsTraits::InputsType>>(InputsToNetwork_Simulated);
 	ReplicatedRemoteInputs.Owner = this;
 
 	ReplicatedStates.History = MakeUnique<TNetRewindHistory<typename PhysicsTraits::StatesType>>(StatesToNetwork);
@@ -1067,10 +1081,10 @@ FORCEINLINE void UNetworkPhysicsComponent::CreateInputHistory(UActorComponent* H
 {
 	InputHelper = MakeUnique<TNetworkPhysicsDataHelper<InputsType>>();
 
-	ReplicatedInputs.History = MakeUnique<TNetRewindHistory<InputsType>>(InputsToNetwork);
+	ReplicatedInputs.History = MakeUnique<TNetRewindHistory<InputsType>>(InputsToNetwork_OwnerDefault);
 	ReplicatedInputs.Owner = this;
 
-	ReplicatedRemoteInputs.History = MakeUnique<TNetRewindHistory<InputsType>>(InputsToNetworkRemote);
+	ReplicatedRemoteInputs.History = MakeUnique<TNetRewindHistory<InputsType>>(InputsToNetwork_Simulated);
 	ReplicatedRemoteInputs.Owner = this;
 
 	ReplicatedImportantInput.History = MakeUnique<TNetRewindHistory<InputsType>>(1);
@@ -1090,6 +1104,7 @@ struct FAsyncNetworkPhysicsComponentInput : public Chaos::FSimCallbackInput
 	TOptional<ENetMode> NetMode;
 	TOptional<ENetRole> NetRole;
 	TOptional<int32> NetworkPhysicsTickOffset;
+	TOptional<uint16> InputsToNetwork_Owner;
 	TOptional<EPhysicsReplicationMode> PhysicsReplicationMode;
 	TOptional<TObjectPtr<UActorComponent>> ActorComponent;
 	TOptional<FString> ActorName;
@@ -1114,6 +1129,7 @@ struct FAsyncNetworkPhysicsComponentInput : public Chaos::FSimCallbackInput
 		NetMode.Reset();
 		NetRole.Reset();
 		NetworkPhysicsTickOffset.Reset();
+		InputsToNetwork_Owner.Reset();
 		PhysicsReplicationMode.Reset();
 		ActorComponent.Reset();
 		ActorName.Reset();
@@ -1143,6 +1159,8 @@ struct FAsyncNetworkPhysicsComponentInput : public Chaos::FSimCallbackInput
 
 struct FAsyncNetworkPhysicsComponentOutput : public Chaos::FSimCallbackOutput
 {
+	TOptional<uint16> InputsToNetwork_Owner;
+
 	TUniquePtr<Chaos::FBaseRewindHistory> InputData;
 	TUniquePtr<Chaos::FBaseRewindHistory> StateData;
 
@@ -1151,6 +1169,8 @@ struct FAsyncNetworkPhysicsComponentOutput : public Chaos::FSimCallbackOutput
 
 	void Reset()
 	{
+		InputsToNetwork_Owner.Reset();
+
 		if (InputData)
 		{
 			InputData->ResetFast();
@@ -1240,6 +1260,9 @@ private:
 	/** Populate state data in AsyncOutput to send over the network */
 	void SendStateData_Internal(FAsyncNetworkPhysicsComponentOutput& AsyncOutput, const int32 PhysicsStep);
 
+	/** Update function on the server to manage how many inputs the owning client should send with each RPC */
+	void UpdateDynamicInputScaling();
+
 private:
 	friend UNetworkPhysicsComponent;
 
@@ -1279,10 +1302,17 @@ private:
 	// Local temporary states data used by pre/post process inputs functions
 	TUniquePtr<FNetworkPhysicsData> StateData;
 
-	// Send last N number of inputs each replication call to patch up holes due to packet loss
-	uint16 InputsToNetwork = 3;
+	// Send last N number of inputs each replication call
+	uint16 InputsToNetwork_OwnerDefault = 3; // Default value for owning client
+	uint16 InputsToNetwork_Owner = 3; // From owning client, i.e. autonomous proxy or client owning an actor with bIsRelayingLocalInputs enabled
+	uint16 InputsToNetwork_Simulated = 2; // To simulated proxies
+	
+	// Properties for dynamic scaling of inputs
+	float TimeOfLastDynamicInputScaling = 0.0f;
+	float DynamicInputScalingAverageInputs = 0.0f;
+	int32 MissingInputCount = 0;
 
-	// Send last N number of states each replication call to patch up holes due to packet loss
+	// Send last N number of states each replication call
 	uint16 StatesToNetwork = 1;
 
 	// Cache predicted states and then compare incoming states via FNetworkPhysicsData::CompareData to trigger a resim if they desync
diff --git a/Engine/Source/Runtime/Engine/Public/Physics/NetworkPhysicsSettingsComponent.h b/Engine/Source/Runtime/Engine/Public/Physics/NetworkPhysicsSettingsComponent.h
--- a/Engine/Source/Runtime/Engine/Public/Physics/NetworkPhysicsSettingsComponent.h
+++ b/Engine/Source/Runtime/Engine/Public/Physics/NetworkPhysicsSettingsComponent.h
@@ -397,7 +397,8 @@ struct FNetworkPhysicsSettingsNetworkPhysicsComponent
 
 	UPROPERTY(config, EditDefaultsOnly, Category = "Overrides", Meta = (InlineEditConditionToggle))
 	uint32 bOverrideRedundantInputs : 1;
-	// Overrides CVar: np2.Resim.RedundantInputs -- How many extra inputs to send with each unreliable network message, to account for packetloss. From owning client to server and server to owning client.
+	/* Overrides CVar : np2.Resim.RedundantInputs -- How many extra inputs to send with each unreliable network message, to account for packetloss.From owning client to server and server to owning client.
+	 * NOTE: This is disabled while np2.Resim.DynamicInputScaling.Enabled is enabled. */
 	UPROPERTY(config, EditDefaultsOnly, Category = "Overrides", Meta = (EditCondition = "bOverrideRedundantInputs"))
 	uint16 RedundantInputs = (uint16)PhysicsReplicationCVars::ResimulationCVars::RedundantInputs;
 	const uint16 GetRedundantInputs() const { return bOverrideRedundantInputs ? RedundantInputs : (uint16)PhysicsReplicationCVars::ResimulationCVars::RedundantInputs; }
diff --git a/Engine/Source/Runtime/Experimental/Chaos/Public/RewindData.h b/Engine/Source/Runtime/Experimental/Chaos/Public/RewindData.h
--- a/Engine/Source/Runtime/Experimental/Chaos/Public/RewindData.h
+++ b/Engine/Source/Runtime/Experimental/Chaos/Public/RewindData.h
@@ -238,7 +238,7 @@ public :
 
 				const int32 DeltaFrame = FMath::Abs(ExtractFrame - DataHistory[MinFrameIndex].LocalFrame);
 
-				ExtractedData.LocalFrame = ExtractFrame;
+				// Don't update LocalFrame here so that it can be used to check if this data was grabbed from the correct frame or not
 				ExtractedData.ServerFrame = DataHistory[MinFrameIndex].ServerFrame + DeltaFrame;
 				ExtractedData.InputFrame = INDEX_NONE; // Clear InputFrame since this history entry is now altered and doesn't correspond to the source entry with the same InputFrame
 
