[NetPhysics] Fix cast issue in recent submit.

diff --git a/Engine/Source/Runtime/Engine/Private/PhysicsEngine/NetworkPhysicsComponent.cpp b/Engine/Source/Runtime/Engine/Private/PhysicsEngine/NetworkPhysicsComponent.cpp
--- a/Engine/Source/Runtime/Engine/Private/PhysicsEngine/NetworkPhysicsComponent.cpp
+++ b/Engine/Source/Runtime/Engine/Private/PhysicsEngine/NetworkPhysicsComponent.cpp
@@ -1959,7 +1959,7 @@ void FAsyncNetworkPhysicsComponent::UpdateDynamicInputScaling()
 		DynamicInputScalingAverageInputs += (static_cast<float>(InputsToNetwork_Owner) - DynamicInputScalingAverageInputs) * PhysicsReplicationCVars::ResimulationCVars::DynamicInputScalingDecreaseAverageMultiplier;
 
 		// Clamp to minimum valid value
-		const uint16 MinInputsValue = FMath::Max(static_cast<uint16>(FMath::RoundToInt(DynamicInputScalingAverageInputs), static_cast<uint16>(PhysicsReplicationCVars::ResimulationCVars::DynamicInputScalingMinInputs)));
+		const uint16 MinInputsValue = static_cast<uint16>(FMath::Max(static_cast<uint16>(FMath::RoundToInt(DynamicInputScalingAverageInputs)), static_cast<uint16>(PhysicsReplicationCVars::ResimulationCVars::DynamicInputScalingMinInputs)));
 		InputsToNetwork_Owner = FMath::Max(InputsToNetwork_Owner, MinInputsValue);
 
 		TimeOfLastDynamicInputScaling = RigidSolver->GetSolverTime();
