Also allow 0 value for CVarForceLOD

diff --git a/Engine/Source/Runtime/Engine/Private/HierarchicalInstancedStaticMesh.cpp b/Engine/Source/Runtime/Engine/Private/HierarchicalInstancedStaticMesh.cpp
--- a/Engine/Source/Runtime/Engine/Private/HierarchicalInstancedStaticMesh.cpp
+++ b/Engine/Source/Runtime/Engine/Private/HierarchicalInstancedStaticMesh.cpp
@@ -1758,7 +1758,7 @@ void FHierarchicalStaticMeshSceneProxy::GetDynamicMeshElements(const TArray<cons
 					}
 					int32 UseMaxLOD = InstanceParams.LODs;
 
-					int32 Force = CVarForceLOD.GetValueOnRenderThread() >= 0 ? CVarForceLOD.GetValueOnRenderThread() : (ForcedLodModel > 0 ? ForcedLodModel : -1); 
+					int32 Force = CVarForceLOD.GetValueOnRenderThread() >= 0 ? CVarForceLOD.GetValueOnRenderThread() : (ForcedLodModel >= 0 ? ForcedLodModel : -1); 
 					if (Force >= 0)
 					{
 						UseMinLOD = FMath::Clamp(Force, 0, InstanceParams.LODs - 1);
