Fix FNaniteSplineMeshSceneProxy incorrectly initializing bHasRayTracingRepresentation/bDynamicRayTracingGeometry when static mesh asset doesn't support ray tracing.

diff --git a/Engine/Source/Runtime/Engine/Private/SplineMeshSceneProxy.cpp b/Engine/Source/Runtime/Engine/Private/SplineMeshSceneProxy.cpp
--- a/Engine/Source/Runtime/Engine/Private/SplineMeshSceneProxy.cpp
+++ b/Engine/Source/Runtime/Engine/Private/SplineMeshSceneProxy.cpp
@@ -340,12 +340,12 @@ FNaniteSplineMeshSceneProxy::FNaniteSplineMeshSceneProxy(const Nanite::FMaterial
 	SetupInstanceSceneDataBuffers(&SplineMeshInstanceData);
 
 #if RHI_RAYTRACING
-	// We only need to init vertex factories for Nanite spline meshes if they can be ray traced
-	if (IsRayTracingEnabled())
+	if (bSupportRayTracing)
 	{
 		bHasRayTracingRepresentation = true;
 		bDynamicRayTracingGeometry = true;
 
+		// We only need to init vertex factories for Nanite spline meshes if they can be ray traced
 		const ERHIFeatureLevel::Type FeatureLevel = GetScene().GetFeatureLevel();
 		for (int32 LODIndex = 0; LODIndex < RenderData->LODResources.Num(); LODIndex++)
 		{
