PR #12641: Add autocast to GetAssetUserDataOfClass

diff --git a/Engine/Source/Runtime/Engine/Classes/Interfaces/Interface_AssetUserData.h b/Engine/Source/Runtime/Engine/Classes/Interfaces/Interface_AssetUserData.h
--- a/Engine/Source/Runtime/Engine/Classes/Interfaces/Interface_AssetUserData.h
+++ b/Engine/Source/Runtime/Engine/Classes/Interfaces/Interface_AssetUserData.h
@@ -29,7 +29,7 @@ class IInterface_AssetUserData
 	*
 	* @return	The instance of the UAssetUserData class contained, or null if it doesn't exist
 	*/
-	UFUNCTION(BlueprintCallable, Category = AssetUserData)
+	UFUNCTION(BlueprintCallable, Category=AssetUserData, meta=(DeterminesOutputType="InUserDataClass"))
 	virtual UAssetUserData* GetAssetUserDataOfClass(TSubclassOf<UAssetUserData> InUserDataClass)
 	{
 		return nullptr;
