Fix missing resource transitions for read/write of vertex buffers from DeformerGraph.

diff --git a/Engine/Plugins/Animation/DeformerGraph/Source/OptimusCore/Private/DataInterfaces/OptimusDataInterfaceSkinnedMeshRead.cpp b/Engine/Plugins/Animation/DeformerGraph/Source/OptimusCore/Private/DataInterfaces/OptimusDataInterfaceSkinnedMeshRead.cpp
--- a/Engine/Plugins/Animation/DeformerGraph/Source/OptimusCore/Private/DataInterfaces/OptimusDataInterfaceSkinnedMeshRead.cpp
+++ b/Engine/Plugins/Animation/DeformerGraph/Source/OptimusCore/Private/DataInterfaces/OptimusDataInterfaceSkinnedMeshRead.cpp
@@ -150,13 +150,12 @@ void UOptimusSkinnedMeshReadDataInterface::GetDefines(FComputeKernelDefinitionSe
 		FString::FromInt(static_cast<uint32>(EOptimusSkinnedMeshReadReadableOutputBuffer::VertexColor))));
 }
 
-
-
 void UOptimusSkinnedMeshReadDataInterface::GetPermutations(FComputeKernelPermutationVector& OutPermutationVector) const
 {
 	OutPermutationVector.AddPermutation(ReadableOutputBufferPermutationName, static_cast<uint32>(EOptimusSkinnedMeshReadReadableOutputBuffer::NumPermutations));
 }
 
+
 UComputeDataProvider* UOptimusSkinnedMeshReadDataInterface::CreateDataProvider(TObjectPtr<UObject> InBinding, uint64 InInputMask, uint64 InOutputMask) const
 {
 	UOptimusSkinnedMeshReadDataProvider* Provider = NewObject<UOptimusSkinnedMeshReadDataProvider>();
@@ -207,7 +206,7 @@ bool FOptimusSkinnedMeshReadDataProviderProxy::IsValid(FValidationData const& In
 	return true;
 }
 
-void FOptimusSkinnedMeshReadDataProviderProxy::AllocateResources(FRDGBuilder& GraphBuilder)
+void FOptimusSkinnedMeshReadDataProviderProxy::AllocateResources(FRDGBuilder& GraphBuilder, FAllocationData const& InAllocationData)
 {
 	// Allocate required buffers
 	const int32 LodIndex = SkeletalMeshObject->GetLOD();
@@ -221,7 +220,7 @@ void FOptimusSkinnedMeshReadDataProviderProxy::AllocateResources(FRDGBuilder& Gr
 	
 	if (InputMask & static_cast<uint64>(ESkinnedMeshReadDataInterfaceInputSelectorMask::Position) )
 	{
-		FRDGBuffer* PositionBuffer = FSkeletalMeshDeformerHelpers::AllocateVertexFactoryPositionBuffer(GraphBuilder, SkeletalMeshObject, LodIndex, bLodJustChanged,TEXT("OptimusSkinnedMeshPosition"));
+		FRDGBuffer* PositionBuffer = FSkeletalMeshDeformerHelpers::AllocateVertexFactoryPositionBuffer(GraphBuilder, InAllocationData.ExternalAccessQueue, SkeletalMeshObject, LodIndex, bLodJustChanged,TEXT("OptimusSkinnedMeshPosition"));
 		PositionBufferUAV = GraphBuilder.CreateUAV(PositionBuffer, PF_R32_FLOAT, ERDGUnorderedAccessViewFlags::SkipBarrier);
 	}
 	else
@@ -235,7 +234,7 @@ void FOptimusSkinnedMeshReadDataProviderProxy::AllocateResources(FRDGBuilder& Gr
 	if ((InputMask & static_cast<uint64>(ESkinnedMeshReadDataInterfaceInputSelectorMask::TangentX)) ||
 		(InputMask & static_cast<uint64>(ESkinnedMeshReadDataInterfaceInputSelectorMask::TangentZ)))
 	{
-		FRDGBuffer* TangentBuffer = FSkeletalMeshDeformerHelpers::AllocateVertexFactoryTangentBuffer(GraphBuilder, SkeletalMeshObject, LodIndex, TEXT("OptimusSkinnedMeshTangent"));
+		FRDGBuffer* TangentBuffer = FSkeletalMeshDeformerHelpers::AllocateVertexFactoryTangentBuffer(GraphBuilder, InAllocationData.ExternalAccessQueue, SkeletalMeshObject, LodIndex, TEXT("OptimusSkinnedMeshTangent"));
 		TangentBufferUAV = GraphBuilder.CreateUAV(TangentBuffer, TangentsFormat, ERDGUnorderedAccessViewFlags::SkipBarrier);
 	}
 	else
@@ -245,7 +244,7 @@ void FOptimusSkinnedMeshReadDataProviderProxy::AllocateResources(FRDGBuilder& Gr
 
 	if (InputMask & static_cast<uint64>(ESkinnedMeshReadDataInterfaceInputSelectorMask::Color))
 	{
-		FRDGBuffer* ColorBuffer = FSkeletalMeshDeformerHelpers::AllocateVertexFactoryColorBuffer(GraphBuilder, SkeletalMeshObject, LodIndex, TEXT("OptimusSkinnedMeshColor"));
+		FRDGBuffer* ColorBuffer = FSkeletalMeshDeformerHelpers::AllocateVertexFactoryColorBuffer(GraphBuilder, InAllocationData.ExternalAccessQueue, SkeletalMeshObject, LodIndex, TEXT("OptimusSkinnedMeshColor"));
 		// using RGBA here and do a manual fetch swizzle in shader instead of BGRA directly because some Mac does not support it. See GMetalBufferFormats[PF_B8G8R8A8] 
 		ColorBufferUAV = GraphBuilder.CreateUAV(ColorBuffer, PF_R8G8B8A8, ERDGUnorderedAccessViewFlags::SkipBarrier);
 	}
@@ -257,10 +256,6 @@ void FOptimusSkinnedMeshReadDataProviderProxy::AllocateResources(FRDGBuilder& Gr
 	FSkeletalMeshDeformerHelpers::UpdateVertexFactoryBufferOverrides(GraphBuilder.RHICmdList, SkeletalMeshObject, LodIndex);
 }
 
-
-
-
-
 void FOptimusSkinnedMeshReadDataProviderProxy::GatherPermutations(FPermutationData& InOutPermutationData) const
 {
 	EOptimusSkinnedMeshReadReadableOutputBuffer ReadableOutputBuffers = EOptimusSkinnedMeshReadReadableOutputBuffer::None; 
@@ -286,9 +281,6 @@ void FOptimusSkinnedMeshReadDataProviderProxy::GatherPermutations(FPermutationDa
 	}
 }
 
-
-
-
 void FOptimusSkinnedMeshReadDataProviderProxy::GatherDispatchData(FDispatchData const& InDispatchData)
 {
 	const int32 LodIndex = SkeletalMeshObject->GetLOD();
@@ -297,8 +289,6 @@ void FOptimusSkinnedMeshReadDataProviderProxy::GatherDispatchData(FDispatchData
 
 	FRHIShaderResourceView* NullSRVBinding = GWhiteVertexBufferWithSRV->ShaderResourceViewRHI.GetReference();
 
-
-	
 	const TStridedView<FParameters> ParameterArray = MakeStridedParameterView<FParameters>(InDispatchData);
 	for (int32 InvocationIndex = 0; InvocationIndex < ParameterArray.Num(); ++InvocationIndex)
 	{
diff --git a/Engine/Plugins/Animation/DeformerGraph/Source/OptimusCore/Private/DataInterfaces/OptimusDataInterfaceSkinnedMeshRead.h b/Engine/Plugins/Animation/DeformerGraph/Source/OptimusCore/Private/DataInterfaces/OptimusDataInterfaceSkinnedMeshRead.h
--- a/Engine/Plugins/Animation/DeformerGraph/Source/OptimusCore/Private/DataInterfaces/OptimusDataInterfaceSkinnedMeshRead.h
+++ b/Engine/Plugins/Animation/DeformerGraph/Source/OptimusCore/Private/DataInterfaces/OptimusDataInterfaceSkinnedMeshRead.h
@@ -85,7 +85,7 @@ public:
 
 	//~ Begin FComputeDataProviderRenderProxy Interface
 	bool IsValid(FValidationData const& InValidationData) const override;
-	void AllocateResources(FRDGBuilder& GraphBuilder) override;
+	void AllocateResources(FRDGBuilder& GraphBuilder, FAllocationData const& InAllocationData) override;
 	void GatherPermutations(FPermutationData& InOutPermutationData) const override;
 	void GatherDispatchData(FDispatchData const& InDispatchData) override;
 	//~ End FComputeDataProviderRenderProxy Interface
diff --git a/Engine/Plugins/Animation/DeformerGraph/Source/OptimusCore/Private/DataInterfaces/OptimusDataInterfaceSkinnedMeshWrite.cpp b/Engine/Plugins/Animation/DeformerGraph/Source/OptimusCore/Private/DataInterfaces/OptimusDataInterfaceSkinnedMeshWrite.cpp
--- a/Engine/Plugins/Animation/DeformerGraph/Source/OptimusCore/Private/DataInterfaces/OptimusDataInterfaceSkinnedMeshWrite.cpp
+++ b/Engine/Plugins/Animation/DeformerGraph/Source/OptimusCore/Private/DataInterfaces/OptimusDataInterfaceSkinnedMeshWrite.cpp
@@ -423,7 +423,7 @@ bool FOptimusSkinnedMeshWriteDataProviderProxy::IsValid(FValidationData const& I
 	return true;
 }
 
-void FOptimusSkinnedMeshWriteDataProviderProxy::AllocateResources(FRDGBuilder& GraphBuilder)
+void FOptimusSkinnedMeshWriteDataProviderProxy::AllocateResources(FRDGBuilder& GraphBuilder, FAllocationData const& InAllocationData)
 {
 	// Allocate required buffers
 	const int32 LodIndex = SkeletalMeshObject->GetLOD();
@@ -437,7 +437,7 @@ void FOptimusSkinnedMeshWriteDataProviderProxy::AllocateResources(FRDGBuilder& G
 
 	if (OutputMask & static_cast<uint64>(ESkinnedMeshWriteDataInterfaceOutputSelectorMask::Position))
 	{
-		PositionBuffer = FSkeletalMeshDeformerHelpers::AllocateVertexFactoryPositionBuffer(GraphBuilder, SkeletalMeshObject, LodIndex, bLodJustChanged,TEXT("OptimusSkinnedMeshPosition"));
+		PositionBuffer = FSkeletalMeshDeformerHelpers::AllocateVertexFactoryPositionBuffer(GraphBuilder, InAllocationData.ExternalAccessQueue, SkeletalMeshObject, LodIndex, bLodJustChanged,TEXT("OptimusSkinnedMeshPosition"));
 		PositionBufferUAV = GraphBuilder.CreateUAV(PositionBuffer, PF_R32_FLOAT, ERDGUnorderedAccessViewFlags::SkipBarrier);
 	}
 	else
@@ -459,7 +459,7 @@ void FOptimusSkinnedMeshWriteDataProviderProxy::AllocateResources(FRDGBuilder& G
 	if ((OutputMask & static_cast<uint64>(ESkinnedMeshWriteDataInterfaceOutputSelectorMask::TangentX)) ||
 		(OutputMask & static_cast<uint64>(ESkinnedMeshWriteDataInterfaceOutputSelectorMask::TangentZ)))
 	{
-		TangentBuffer = FSkeletalMeshDeformerHelpers::AllocateVertexFactoryTangentBuffer(GraphBuilder, SkeletalMeshObject, LodIndex, TEXT("OptimusSkinnedMeshTangent"));
+		TangentBuffer = FSkeletalMeshDeformerHelpers::AllocateVertexFactoryTangentBuffer(GraphBuilder, InAllocationData.ExternalAccessQueue, SkeletalMeshObject, LodIndex, TEXT("OptimusSkinnedMeshTangent"));
 		TangentBufferUAV = GraphBuilder.CreateUAV(TangentBuffer, TangentsFormat, ERDGUnorderedAccessViewFlags::SkipBarrier);
 	}
 	else
@@ -477,7 +477,7 @@ void FOptimusSkinnedMeshWriteDataProviderProxy::AllocateResources(FRDGBuilder& G
 
 	if (OutputMask & static_cast<uint64>(ESkinnedMeshWriteDataInterfaceOutputSelectorMask::Color))
 	{
-		ColorBuffer = FSkeletalMeshDeformerHelpers::AllocateVertexFactoryColorBuffer(GraphBuilder, SkeletalMeshObject, LodIndex, TEXT("OptimusSkinnedMeshColor"));
+		ColorBuffer = FSkeletalMeshDeformerHelpers::AllocateVertexFactoryColorBuffer(GraphBuilder, InAllocationData.ExternalAccessQueue, SkeletalMeshObject, LodIndex, TEXT("OptimusSkinnedMeshColor"));
 		// using RGBA here and do a manual fetch swizzle in shader instead of BGRA directly because some Mac does not support it. See GMetalBufferFormats[PF_B8G8R8A8] 
 		ColorBufferUAV = GraphBuilder.CreateUAV(ColorBuffer, PF_R8G8B8A8, ERDGUnorderedAccessViewFlags::SkipBarrier);
 	}
diff --git a/Engine/Plugins/Animation/DeformerGraph/Source/OptimusCore/Private/DataInterfaces/OptimusDataInterfaceSkinnedMeshWrite.h b/Engine/Plugins/Animation/DeformerGraph/Source/OptimusCore/Private/DataInterfaces/OptimusDataInterfaceSkinnedMeshWrite.h
--- a/Engine/Plugins/Animation/DeformerGraph/Source/OptimusCore/Private/DataInterfaces/OptimusDataInterfaceSkinnedMeshWrite.h
+++ b/Engine/Plugins/Animation/DeformerGraph/Source/OptimusCore/Private/DataInterfaces/OptimusDataInterfaceSkinnedMeshWrite.h
@@ -124,7 +124,7 @@ class FOptimusSkinnedMeshWriteDataProviderProxy : public FComputeDataProviderRen
 public:
 	//~ Begin FComputeDataProviderRenderProxy Interface
 	bool IsValid(FValidationData const& InValidationData) const override;
-	void AllocateResources(FRDGBuilder& GraphBuilder) override;
+	void AllocateResources(FRDGBuilder& GraphBuilder, FAllocationData const& InAllocationData) override;
 	void GatherDispatchData(FDispatchData const& InDispatchData) override;
 	void GetReadbackData(TArray<FReadbackData>& OutReadbackData) const override;
 	//~ End FComputeDataProviderRenderProxy Interface
diff --git a/Engine/Plugins/Runtime/ComputeFramework/Source/ComputeFramework/Private/ComputeGraphWorker.cpp b/Engine/Plugins/Runtime/ComputeFramework/Source/ComputeFramework/Private/ComputeGraphWorker.cpp
--- a/Engine/Plugins/Runtime/ComputeFramework/Source/ComputeFramework/Private/ComputeGraphWorker.cpp
+++ b/Engine/Plugins/Runtime/ComputeFramework/Source/ComputeFramework/Private/ComputeGraphWorker.cpp
@@ -112,6 +112,8 @@ void FComputeGraphTaskWorker::SubmitWork(FRDGBuilder& GraphBuilder, FName InExec
 	SubmitDescs.Reset();
 	Shaders.Reset();
 
+	FRDGExternalAccessQueue ExternalAccessQueue;
+
 	TArray<FGraphInvocation> const& GraphInvocations = *FoundGraphInvocations;
 	for (int32 GraphIndex = 0; GraphIndex < GraphInvocations.Num(); ++GraphIndex)
 	{
@@ -208,7 +210,7 @@ void FComputeGraphTaskWorker::SubmitWork(FRDGBuilder& GraphBuilder, FName InExec
 		}
 
 		// Allocate RDG resources for all the data providers in the graph.
-		FComputeDataProviderRenderProxy::FAllocationData AllocationData { NumKernels };
+		FComputeDataProviderRenderProxy::FAllocationData AllocationData { NumKernels, ExternalAccessQueue };
 		for (int32 DataProviderIndex = 0; DataProviderIndex < GraphInvocation.DataProviderRenderProxies.Num(); ++DataProviderIndex)
 		{
 			FComputeDataProviderRenderProxy* DataProvider = GraphInvocation.DataProviderRenderProxies[DataProviderIndex];
@@ -351,6 +353,8 @@ void FComputeGraphTaskWorker::SubmitWork(FRDGBuilder& GraphBuilder, FName InExec
 		}
 	}
 
+	ExternalAccessQueue.Submit(GraphBuilder);
+
 	// Release any graph resources at the end of graph execution.
 	GraphBuilder.AddPostExecuteCallback(
 		[this, InExecutionGroupName]
diff --git a/Engine/Plugins/Runtime/ComputeFramework/Source/ComputeFramework/Public/ComputeFramework/ComputeDataProvider.h b/Engine/Plugins/Runtime/ComputeFramework/Source/ComputeFramework/Public/ComputeFramework/ComputeDataProvider.h
--- a/Engine/Plugins/Runtime/ComputeFramework/Source/ComputeFramework/Public/ComputeFramework/ComputeDataProvider.h
+++ b/Engine/Plugins/Runtime/ComputeFramework/Source/ComputeFramework/Public/ComputeFramework/ComputeDataProvider.h
@@ -10,6 +10,7 @@
 class FComputeDataProviderRenderProxy;
 struct FComputeKernelPermutationVector;
 class FRDGBuilder;
+class FRDGExternalAccessQueue;
 
 /**
  * Compute Framework Data Provider.
@@ -78,7 +79,9 @@ public:
 	/** Setup needed to allocate resources. */
 	struct FAllocationData
 	{
-		int32 NumGraphKernels;
+		int32 NumGraphKernels = 0;
+		/** Queue for AllocateResources() to add any resources that need to be externally (non-RDG) managed at the end of compute graph execution. */
+		FRDGExternalAccessQueue& ExternalAccessQueue;
 	};
 
 	// todo: Fix up derived classes to implement new AllocateResources().
diff --git a/Engine/Source/Runtime/Engine/Private/SkeletalMeshDeformerHelpers.cpp b/Engine/Source/Runtime/Engine/Private/SkeletalMeshDeformerHelpers.cpp
--- a/Engine/Source/Runtime/Engine/Private/SkeletalMeshDeformerHelpers.cpp
+++ b/Engine/Source/Runtime/Engine/Private/SkeletalMeshDeformerHelpers.cpp
@@ -4,8 +4,9 @@
 #include "DataDrivenShaderPlatformInfo.h"
 #include "RenderGraphBuilder.h"
 #include "RenderGraphResources.h"
-#include "SkeletalRenderGPUSkin.h"
+#include "RenderGraphUtils.h"
 #include "RenderingThread.h"
+#include "SkeletalRenderGPUSkin.h"
 
 FRHIShaderResourceView* FSkeletalMeshDeformerHelpers::GetBoneBufferForReading(
 	FSkeletalMeshObject const* InMeshObject,
@@ -138,7 +139,7 @@ FRDGBuffer* FSkeletalMeshDeformerHelpers::GetAllocatedColorBuffer(FRDGBuilder& G
 	return ColorBuffer;
 }
 
-FRDGBuffer* FSkeletalMeshDeformerHelpers::AllocateVertexFactoryPositionBuffer(FRDGBuilder& GraphBuilder, FSkeletalMeshObject* InMeshObject, int32 InLodIndex, bool bInLodJustChanged, TCHAR const* InBufferName)
+FRDGBuffer* FSkeletalMeshDeformerHelpers::AllocateVertexFactoryPositionBuffer(FRDGBuilder& GraphBuilder, FRDGExternalAccessQueue& ExternalAccesQueue, FSkeletalMeshObject* InMeshObject, int32 InLodIndex, bool bInLodJustChanged, TCHAR const* InBufferName)
 {
 	if (InMeshObject->IsCPUSkinned())
 	{
@@ -159,6 +160,7 @@ FRDGBuffer* FSkeletalMeshDeformerHelpers::AllocateVertexFactoryPositionBuffer(FR
 	if (DeformerGeometry.Position && Frame == DeformerGeometry.PositionUpdatedFrame)
 	{
 		PositionBuffer = GraphBuilder.RegisterExternalBuffer(DeformerGeometry.Position);
+		GraphBuilder.UseInternalAccessMode(PositionBuffer);
 	}
 	else
 	{
@@ -167,6 +169,7 @@ FRDGBuffer* FSkeletalMeshDeformerHelpers::AllocateVertexFactoryPositionBuffer(FR
 
 		PositionBuffer = GraphBuilder.CreateBuffer(FRDGBufferDesc::CreateBufferDesc(PosBufferBytesPerElement, NumVertices * PosBufferElementMultiplier), InBufferName, ERDGBufferFlags::None);
 		DeformerGeometry.Position = GraphBuilder.ConvertToExternalBuffer(PositionBuffer);
+
 		DeformerGeometry.PositionSRV = DeformerGeometry.Position->GetOrCreateSRV(GraphBuilder.RHICmdList, FRHIBufferSRVCreateInfo(PF_R32_FLOAT));
 
 		// Avoid using position buffer from the last time this LOD was active to compute motion vectors,
@@ -179,8 +182,6 @@ FRDGBuffer* FSkeletalMeshDeformerHelpers::AllocateVertexFactoryPositionBuffer(FR
 		}
 		
 		DeformerGeometry.PositionUpdatedFrame = Frame;
-		GraphBuilder.SetBufferAccessFinal(PositionBuffer, ERHIAccess::VertexOrIndexBuffer | ERHIAccess::SRVMask);
-
 		
 #if RHI_RAYTRACING
 		// Update ray tracing geometry whenever we recreate the position buffer.
@@ -200,10 +201,12 @@ FRDGBuffer* FSkeletalMeshDeformerHelpers::AllocateVertexFactoryPositionBuffer(FR
 #endif // RHI_RAYTRACING
 	}
 
+	ExternalAccesQueue.AddUnique(PositionBuffer, ERHIAccess::VertexOrIndexBuffer | ERHIAccess::SRVMask);
+
 	return PositionBuffer;
 }
 
-FRDGBuffer* FSkeletalMeshDeformerHelpers::AllocateVertexFactoryTangentBuffer(FRDGBuilder& GraphBuilder, FSkeletalMeshObject* InMeshObject, int32 InLodIndex, TCHAR const* InBufferName)
+FRDGBuffer* FSkeletalMeshDeformerHelpers::AllocateVertexFactoryTangentBuffer(FRDGBuilder& GraphBuilder, FRDGExternalAccessQueue& ExternalAccesQueue, FSkeletalMeshObject* InMeshObject, int32 InLodIndex, TCHAR const* InBufferName)
 {
 	if (InMeshObject->IsCPUSkinned())
 	{
@@ -224,6 +227,7 @@ FRDGBuffer* FSkeletalMeshDeformerHelpers::AllocateVertexFactoryTangentBuffer(FRD
 	if (DeformerGeometry.Tangent && Frame == DeformerGeometry.TangentUpdatedFrame)
 	{
 		TangentBuffer = GraphBuilder.RegisterExternalBuffer(DeformerGeometry.Tangent);
+		GraphBuilder.UseInternalAccessMode(TangentBuffer);
 	}
 	else
 	{
@@ -232,13 +236,14 @@ FRDGBuffer* FSkeletalMeshDeformerHelpers::AllocateVertexFactoryTangentBuffer(FRD
 		const EPixelFormat TangentsFormat = IsOpenGLPlatform(GMaxRHIShaderPlatform) ? PF_R16G16B16A16_SINT : PF_R16G16B16A16_SNORM;
 		DeformerGeometry.TangentSRV = DeformerGeometry.Tangent->GetOrCreateSRV(GraphBuilder.RHICmdList, FRHIBufferSRVCreateInfo(TangentsFormat));
 		DeformerGeometry.TangentUpdatedFrame = Frame;
-		GraphBuilder.SetBufferAccessFinal(TangentBuffer, ERHIAccess::VertexOrIndexBuffer | ERHIAccess::SRVMask);
 	}
 
+	ExternalAccesQueue.AddUnique(TangentBuffer, ERHIAccess::VertexOrIndexBuffer | ERHIAccess::SRVMask);
+
 	return TangentBuffer;
 }
 
-FRDGBuffer* FSkeletalMeshDeformerHelpers::AllocateVertexFactoryColorBuffer(FRDGBuilder& GraphBuilder, FSkeletalMeshObject* InMeshObject, int32 InLodIndex, TCHAR const* InBufferName)
+FRDGBuffer* FSkeletalMeshDeformerHelpers::AllocateVertexFactoryColorBuffer(FRDGBuilder& GraphBuilder, FRDGExternalAccessQueue& ExternalAccesQueue, FSkeletalMeshObject* InMeshObject, int32 InLodIndex, TCHAR const* InBufferName)
 {
 	if (InMeshObject->IsCPUSkinned())
 	{
@@ -259,6 +264,7 @@ FRDGBuffer* FSkeletalMeshDeformerHelpers::AllocateVertexFactoryColorBuffer(FRDGB
 	if (DeformerGeometry.Color && Frame == DeformerGeometry.ColorUpdatedFrame)
 	{
 		ColorBuffer = GraphBuilder.RegisterExternalBuffer(DeformerGeometry.Color);
+		GraphBuilder.UseInternalAccessMode(ColorBuffer);
 	}
 	else
 	{
@@ -266,9 +272,10 @@ FRDGBuffer* FSkeletalMeshDeformerHelpers::AllocateVertexFactoryColorBuffer(FRDGB
 		DeformerGeometry.Color = GraphBuilder.ConvertToExternalBuffer(ColorBuffer);
 		DeformerGeometry.ColorSRV = DeformerGeometry.Color->GetOrCreateSRV(GraphBuilder.RHICmdList, FRHIBufferSRVCreateInfo(PF_R8G8B8A8));
 		DeformerGeometry.ColorUpdatedFrame = Frame;
-		GraphBuilder.SetBufferAccessFinal(ColorBuffer, ERHIAccess::VertexOrIndexBuffer | ERHIAccess::SRVMask);
 	}
 
+	ExternalAccesQueue.AddUnique(ColorBuffer, ERHIAccess::VertexOrIndexBuffer | ERHIAccess::SRVMask);
+
 	return ColorBuffer;
 }
 
diff --git a/Engine/Source/Runtime/Engine/Public/SkeletalMeshDeformerHelpers.h b/Engine/Source/Runtime/Engine/Public/SkeletalMeshDeformerHelpers.h
--- a/Engine/Source/Runtime/Engine/Public/SkeletalMeshDeformerHelpers.h
+++ b/Engine/Source/Runtime/Engine/Public/SkeletalMeshDeformerHelpers.h
@@ -6,6 +6,7 @@
 #include "RenderGraphResources.h"
 
 class FRDGBuffer;
+class FRDGExternalAccessQueue;
 class FRHIShaderResourceView;
 class FSkeletalMeshObject;
 
@@ -71,21 +72,21 @@ public:
 	 * Ownership is handled by the MeshObject.
 	 * If we call this more than once for the same MeshObject in the same frame then we return the allocation from the first call.
 	 */
-	ENGINE_API static FRDGBuffer* AllocateVertexFactoryPositionBuffer(FRDGBuilder& GraphBuilder, FSkeletalMeshObject* InMeshObject, int32 InLodIndex, bool bInLodJustChanged, TCHAR const* InBufferName);
+	ENGINE_API static FRDGBuffer* AllocateVertexFactoryPositionBuffer(FRDGBuilder& GraphBuilder, FRDGExternalAccessQueue& ExternalAccesQueue, FSkeletalMeshObject* InMeshObject, int32 InLodIndex, bool bInLodJustChanged, TCHAR const* InBufferName);
 
 	/**
 	 * Allocate and bind a new tangent buffer and return it for writing.
 	 * Ownership is handled by the MeshObject.
 	 * If we call this more than once for the same MeshObject in the same frame then we return the allocation from the first call.
 	 */
-	ENGINE_API static FRDGBuffer* AllocateVertexFactoryTangentBuffer(FRDGBuilder& GraphBuilder,	FSkeletalMeshObject* InMeshObject, int32 InLodIndex, TCHAR const* InBufferName);
+	ENGINE_API static FRDGBuffer* AllocateVertexFactoryTangentBuffer(FRDGBuilder& GraphBuilder, FRDGExternalAccessQueue& ExternalAccesQueue, FSkeletalMeshObject* InMeshObject, int32 InLodIndex, TCHAR const* InBufferName);
 
 	/**
 	 * Allocate and bind a new color buffer and return it for writing.
 	 * Ownership is handled by the MeshObject.
 	 * If we call this more than once for the same MeshObject in the same frame then we return the allocation from the first call.
 	 */
-	ENGINE_API static FRDGBuffer* AllocateVertexFactoryColorBuffer(FRDGBuilder& GraphBuilder, FSkeletalMeshObject* InMeshObject, int32 InLodIndex, TCHAR const* InBufferName);
+	ENGINE_API static FRDGBuffer* AllocateVertexFactoryColorBuffer(FRDGBuilder& GraphBuilder, FRDGExternalAccessQueue& ExternalAccesQueue, FSkeletalMeshObject* InMeshObject, int32 InLodIndex, TCHAR const* InBufferName);
 
 	/** 
 	 * Update all of the MeshObject's passthrough vertex factories with the currently allocated vertex buffers. 
