Make FTickableGlobalObject creation transactional safe.

diff --git a/Engine/Source/Runtime/Engine/Private/Tests/AutoRTFM/AutoRTFMTickableTests.cpp b/Engine/Source/Runtime/Engine/Private/Tests/AutoRTFM/AutoRTFMTickableTests.cpp
new file mode 100644
--- /dev/null
+++ b/Engine/Source/Runtime/Engine/Private/Tests/AutoRTFM/AutoRTFMTickableTests.cpp
@@ -0,0 +1,43 @@
+// Copyright Epic Games, Inc. All Rights Reserved.
+
+#include "Misc/AutomationTest.h"
+#include "AutoRTFM/AutoRTFM.h"
+#include "Tickable.h"
+
+#if WITH_DEV_AUTOMATION_TESTS
+
+IMPLEMENT_SIMPLE_AUTOMATION_TEST(FAutoRTFMTickableTests, "AutoRTFM + FTickableObject", EAutomationTestFlags::EngineFilter | EAutomationTestFlags::ClientContext | EAutomationTestFlags::ServerContext | EAutomationTestFlags::CommandletContext)
+
+bool FAutoRTFMTickableTests::RunTest(const FString & Parameters)
+{
+	if (!AutoRTFM::ForTheRuntime::IsAutoRTFMRuntimeEnabled())
+	{
+		ExecutionInfo.AddEvent(FAutomationEvent(EAutomationEventType::Info, TEXT("SKIPPED 'FAutoRTFMTickableTests' test. AutoRTFM disabled.")));
+		return true;
+	}
+
+	struct FMyTickableGameObject final : FTickableGameObject
+	{
+		void Tick(float) override {}
+		TStatId GetStatId() const override { return TStatId(); }
+	};
+
+	AutoRTFM::ETransactionResult Result = AutoRTFM::Transact([&]
+	{
+		FMyTickableGameObject Tickable;
+		Tickable.SetTickableTickType(ETickableTickType::Always);
+		AutoRTFM::AbortTransaction();
+	});
+	TestTrueExpr(AutoRTFM::ETransactionResult::AbortedByRequest == Result);
+
+	Result = AutoRTFM::Transact([&]
+	{
+		FMyTickableGameObject Tickable;
+		Tickable.SetTickableTickType(ETickableTickType::Always);
+	});
+	TestTrueExpr(AutoRTFM::ETransactionResult::Committed == Result);
+
+	return true;
+}
+
+#endif //WITH_DEV_AUTOMATION_TESTS
diff --git a/Engine/Source/Runtime/Engine/Private/Tickable.cpp b/Engine/Source/Runtime/Engine/Private/Tickable.cpp
--- a/Engine/Source/Runtime/Engine/Private/Tickable.cpp
+++ b/Engine/Source/Runtime/Engine/Private/Tickable.cpp
@@ -9,14 +9,14 @@ DECLARE_CYCLE_STAT(TEXT("TickableGameObjects Time"), STAT_TickableGameObjectsTim
 void FTickableObjectBase::FTickableStatics::QueueTickableObjectForAdd(FTickableObjectBase* InTickable)
 {
 	// This only needs to lock the new object queue
-	FScopeLock NewTickableObjectsLock(&NewTickableObjectsCritical);
+	UE::TScopeLock NewTickableObjectsLock(NewTickableObjectsCritical);
 	NewTickableObjects.Add(InTickable, ETickableTickType::NewObject);
 }
 
 void FTickableObjectBase::FTickableStatics::SetTickTypeForTickableObject(FTickableObjectBase* InTickable, ETickableTickType NewTickType)
 {
-	FScopeLock TickableObjectsLock(&TickableObjectsCritical);
-	FScopeLock NewTickableObjectsLock(&NewTickableObjectsCritical);
+	UE::TScopeLock TickableObjectsLock(TickableObjectsCritical);
+	UE::TScopeLock NewTickableObjectsLock(NewTickableObjectsCritical);
 	
 	// Existing entries should never be set to new object
 	check(NewTickType != ETickableTickType::NewObject);
@@ -64,7 +64,7 @@ void FTickableObjectBase::FTickableStatics::StartTicking()
 {
 	check(!bIsTickingObjects);
 
-	FScopeLock NewTickableObjectsLock(&NewTickableObjectsCritical);
+	UE::TScopeLock NewTickableObjectsLock(NewTickableObjectsCritical);
 
 	for (TPair<FTickableObjectBase*, ETickableTickType> Pair : NewTickableObjects)
 	{
@@ -105,7 +105,7 @@ void FTickableObjectBase::FTickableStatics::FinishTicking()
 
 void FTickableObjectBase::SimpleTickObjects(FTickableStatics& Statics, TFunctionRef<void(FTickableObjectBase*)> TickFunc)
 {
-	FScopeLock LockTickableObjects(&Statics.TickableObjectsCritical);
+	UE::TScopeLock LockTickableObjects(Statics.TickableObjectsCritical);
 
 	Statics.StartTicking();
 
@@ -171,7 +171,7 @@ void FTickableGameObject::TickObjects(UWorld* World, const ELevelTick LevelTickT
 
 	{
 		// It's a long lock but it's ok, the only thing we can block here is the GC worker thread that destroys UObjects
-		FScopeLock LockTickableObjects(&Statics.TickableObjectsCritical);
+		UE::TScopeLock LockTickableObjects(Statics.TickableObjectsCritical);
 
 		Statics.StartTicking();
 
diff --git a/Engine/Source/Runtime/Engine/Public/Tickable.h b/Engine/Source/Runtime/Engine/Public/Tickable.h
--- a/Engine/Source/Runtime/Engine/Public/Tickable.h
+++ b/Engine/Source/Runtime/Engine/Public/Tickable.h
@@ -10,6 +10,7 @@
 #include "CoreMinimal.h"
 #include "Stats/Stats.h"
 #include "Misc/ScopeLock.h"
+#include "Misc/TransactionallySafeCriticalSection.h"
 #include "Engine/EngineBaseTypes.h"
 
 /**
@@ -49,13 +50,13 @@ protected:
 	struct FTickableStatics
 	{
 		/** This critical section should be locked during entire tick process */
-		FCriticalSection TickableObjectsCritical;
+		FTransactionallySafeCriticalSection TickableObjectsCritical;
 
 		/** List of objects that are fully ticking */
 		TArray<FTickableObjectBase::FTickableObjectEntry> TickableObjects;
 
 		/** Lock for modifying new list, this is automatically acquired by functions below */
-		FCriticalSection NewTickableObjectsCritical;
+		FTransactionallySafeCriticalSection NewTickableObjectsCritical;
 
 		/** Set of objects that have not yet been queried for tick type */
 		TMap<FTickableObjectBase*, ETickableTickType> NewTickableObjects;
