Adds r.TransparentViewportClear

PR #12653: Console variable to clear viewport to transparent

diff --git a/Engine/Source/Runtime/Engine/Private/GameViewportClient.cpp b/Engine/Source/Runtime/Engine/Private/GameViewportClient.cpp
--- a/Engine/Source/Runtime/Engine/Private/GameViewportClient.cpp
+++ b/Engine/Source/Runtime/Engine/Private/GameViewportClient.cpp
@@ -119,6 +119,13 @@ static TAutoConsoleVariable<float> CVarSecondaryScreenPercentage( // TODO: make
 	TEXT(" 1: override secondary screen percentage."),
 	ECVF_Default);
 
+static TAutoConsoleVariable<bool> CVarTransparentViewportClear(
+	TEXT("r.TransparentViewportClear"),
+	false,
+	TEXT("If true, then the viewport will clear to transparent instead of black.\n")
+	TEXT("This can be useful with r.PostProcessing.PropagateAlpha if the game view\n")
+	TEXT("does not cover the whole window."),
+	ECVF_Default);
 
 static TAutoConsoleVariable<bool> CVarRemapDeviceIdForOffsetPlayerGamepadIds(
 	TEXT("input.bRemapDeviceIdForOffsetPlayerGamepadIds"),
@@ -1786,9 +1793,10 @@ void UGameViewportClient::Draw(FViewport* InViewport, FCanvas* SceneCanvas)
 	// If the views don't cover the entire bounding rectangle, clear the entire buffer.
 	bool bBufferCleared = false;
 	bool bStereoscopicPass = (ViewFamily.Views.Num() != 0 && IStereoRendering::IsStereoEyeView(*ViewFamily.Views[0]));
-	if (ViewFamily.Views.Num() == 0 || TotalArea != (MaxX-MinX)*(MaxY-MinY) || bDisableWorldRendering || bStereoscopicPass)
+	bool bTransparentViewportClear = CVarTransparentViewportClear.GetValueOnAnyThread();
+	if (ViewFamily.Views.Num() == 0 || TotalArea != (MaxX-MinX)*(MaxY-MinY) || bDisableWorldRendering || bStereoscopicPass || bTransparentViewportClear)
 	{
-		if (bDisableWorldRendering || !bStereoscopicPass) // TotalArea computation does not work correctly for stereoscopic views
+		if (bDisableWorldRendering || !bStereoscopicPass || bTransparentViewportClear) // TotalArea computation does not work correctly for stereoscopic views
 		{
 			SceneCanvas->Clear(FLinearColor::Transparent);
 		}
