Adding scope counters for some audio-related gameplay statics functions & actor sound parameter interface

#rnx

diff --git a/Engine/Source/Runtime/Engine/Private/Audio/ActorSoundParameterInterface.cpp b/Engine/Source/Runtime/Engine/Private/Audio/ActorSoundParameterInterface.cpp
--- a/Engine/Source/Runtime/Engine/Private/Audio/ActorSoundParameterInterface.cpp
+++ b/Engine/Source/Runtime/Engine/Private/Audio/ActorSoundParameterInterface.cpp
@@ -21,6 +21,8 @@ namespace ActorSoundParameterInterfaceConsoleVariables
 
 void UActorSoundParameterInterface::Fill(const AActor* OwningActor, TArray<FAudioParameter>& OutParams)
 {
+	QUICK_SCOPE_CYCLE_COUNTER(UActorSoundParameterInterface_Fill);
+
 	TArray<const AActor*> Actors;
 	TArray<const UActorComponent*> Components;
 
@@ -56,6 +58,8 @@ void UActorSoundParameterInterface::Fill(const AActor* OwningActor, TArray<FAudi
 
 void UActorSoundParameterInterface::GetImplementers(const AActor* InActor, TArray<const AActor*>& OutActors, TArray<const UActorComponent*>& OutComponents)
 {
+	QUICK_SCOPE_CYCLE_COUNTER(UActorSoundParameterInterface_GetImplementers);
+
 	if (!InActor)
 	{
 		return;
diff --git a/Engine/Source/Runtime/Engine/Private/GameplayStatics.cpp b/Engine/Source/Runtime/Engine/Private/GameplayStatics.cpp
--- a/Engine/Source/Runtime/Engine/Private/GameplayStatics.cpp
+++ b/Engine/Source/Runtime/Engine/Private/GameplayStatics.cpp
@@ -1590,6 +1590,8 @@ void UGameplayStatics::SetGlobalListenerFocusParameters(const UObject* WorldCont
 
 void UGameplayStatics::PlaySound2D(const UObject* WorldContextObject, USoundBase* Sound, float VolumeMultiplier, float PitchMultiplier, float StartTime, USoundConcurrency* ConcurrencySettings, const AActor* OwningActor, bool bIsUISound)
 {
+	QUICK_SCOPE_CYCLE_COUNTER(UGameplayStatics_PlaySound2D);
+
 	if (!Sound || !GEngine || !GEngine->UseSound())
 	{
 		return;
@@ -1637,6 +1639,8 @@ void UGameplayStatics::PlaySound2D(const UObject* WorldContextObject, USoundBase
 
 UAudioComponent* UGameplayStatics::CreateSound2D(const UObject* WorldContextObject, USoundBase* Sound, float VolumeMultiplier, float PitchMultiplier, float StartTime, USoundConcurrency* ConcurrencySettings, bool bPersistAcrossLevelTransition, bool bAutoDestroy)
 {
+	QUICK_SCOPE_CYCLE_COUNTER(UGameplayStatics_CreateSound2D);
+
 	if (!Sound || !GEngine || !GEngine->UseSound())
 	{
 		return nullptr;
@@ -1687,6 +1691,8 @@ UAudioComponent* UGameplayStatics::SpawnSound2D(const UObject* WorldContextObjec
 
 void UGameplayStatics::PlaySoundAtLocation(const UObject* WorldContextObject, class USoundBase* Sound, FVector Location, FRotator Rotation, float VolumeMultiplier, float PitchMultiplier, float StartTime, class USoundAttenuation* AttenuationSettings, class USoundConcurrency* ConcurrencySettings, const AActor* OwningActor, const UInitialActiveSoundParams* InitialParams)
 {
+	QUICK_SCOPE_CYCLE_COUNTER(UGameplayStatics_PlaySoundAtLocation);
+
 	if (!Sound || !GEngine || !GEngine->UseSound())
 	{
 		return;
@@ -1716,6 +1722,8 @@ void UGameplayStatics::PlaySoundAtLocation(const UObject* WorldContextObject, cl
 
 UAudioComponent* UGameplayStatics::SpawnSoundAtLocation(const UObject* WorldContextObject, USoundBase* Sound, FVector Location, FRotator Rotation, float VolumeMultiplier, float PitchMultiplier, float StartTime, USoundAttenuation* AttenuationSettings, USoundConcurrency* ConcurrencySettings, bool bAutoDestroy)
 {
+	QUICK_SCOPE_CYCLE_COUNTER(UGameplayStatics_SpawnSoundAtLocation);
+
 	if (!Sound || !GEngine || !GEngine->UseSound())
 	{
 		return nullptr;
@@ -1761,6 +1769,8 @@ UAudioComponent* UGameplayStatics::SpawnSoundAtLocation(const UObject* WorldCont
 
 UAudioComponent* UGameplayStatics::SpawnSoundAttached(USoundBase* Sound, USceneComponent* AttachToComponent, FName AttachPointName, FVector Location, FRotator Rotation, EAttachLocation::Type LocationType, bool bStopWhenAttachedToDestroyed, float VolumeMultiplier, float PitchMultiplier, float StartTime, USoundAttenuation* AttenuationSettings, USoundConcurrency* ConcurrencySettings, bool bAutoDestroy)
 {
+	QUICK_SCOPE_CYCLE_COUNTER(UGameplayStatics_SpawnSoundAttached);
+
 	if (!Sound)
 	{
 		return nullptr;
