Reverted 37748880 since it prevent the default material on VolumetricCloud in editor.

#rnx
#tests pc editor dx12

diff --git a/Engine/Source/Runtime/Engine/Private/Components/VolumetricCloudComponent.cpp b/Engine/Source/Runtime/Engine/Private/Components/VolumetricCloudComponent.cpp
--- a/Engine/Source/Runtime/Engine/Private/Components/VolumetricCloudComponent.cpp
+++ b/Engine/Source/Runtime/Engine/Private/Components/VolumetricCloudComponent.cpp
@@ -57,16 +57,8 @@ UVolumetricCloudComponent::UVolumetricCloudComponent(const FObjectInitializer& O
 	, bRenderInMainPass(true)
 	, VolumetricCloudSceneProxy(nullptr)
 {
-	bool bUseDefaultMaterial = !HasAnyFlags(RF_ClassDefaultObject | RF_DefaultSubObject);
-	if (bUseDefaultMaterial && GetOuter() != nullptr)
-	{
-		bUseDefaultMaterial = !GetOuter()->HasAnyFlags(RF_ClassDefaultObject);
-	}
-	if (bUseDefaultMaterial)
-	{
-		static ConstructorHelpers::FObjectFinder<UMaterialInterface> VolumetricCloudDefaultMaterialRef(TEXT("/Engine/EngineSky/VolumetricClouds/m_SimpleVolumetricCloud_Inst.m_SimpleVolumetricCloud_Inst"));
-		Material = VolumetricCloudDefaultMaterialRef.Object;
-	}
+	static ConstructorHelpers::FObjectFinder<UMaterialInterface> VolumetricCloudDefaultMaterialRef(TEXT("/Engine/EngineSky/VolumetricClouds/m_SimpleVolumetricCloud_Inst.m_SimpleVolumetricCloud_Inst"));
+	Material = VolumetricCloudDefaultMaterialRef.Object;
 }
 
 UVolumetricCloudComponent::~UVolumetricCloudComponent()
