Namespace animation compression LogScale

diff --git a/Engine/Source/Runtime/Engine/Private/Animation/AnimCompress_PerTrackCompression.cpp b/Engine/Source/Runtime/Engine/Private/Animation/AnimCompress_PerTrackCompression.cpp
--- a/Engine/Source/Runtime/Engine/Private/Animation/AnimCompress_PerTrackCompression.cpp
+++ b/Engine/Source/Runtime/Engine/Private/Animation/AnimCompress_PerTrackCompression.cpp
@@ -154,24 +154,24 @@ protected:
 
 			if (bHasX)
 			{
-				X = FAnimationCompression_PerTrackUtils::CompressFixed16(V.X, LogScale);
+				X = FAnimationCompression_PerTrackUtils::CompressFixed16(V.X, AnimationCompressionUtils::LogScale);
 				AppendBytes(&X, sizeof(X));
 			}
 			if (bHasY)
 			{
-				Y = FAnimationCompression_PerTrackUtils::CompressFixed16(V.Y, LogScale);
+				Y = FAnimationCompression_PerTrackUtils::CompressFixed16(V.Y, AnimationCompressionUtils::LogScale);
 				AppendBytes(&Y, sizeof(Y));
 			}
 			if (bHasZ)
 			{
-				Z = FAnimationCompression_PerTrackUtils::CompressFixed16(V.Z, LogScale);
+				Z = FAnimationCompression_PerTrackUtils::CompressFixed16(V.Z, AnimationCompressionUtils::LogScale);
 				AppendBytes(&Z, sizeof(Z));
 			}
 
 			const FVector3f DecompressedV(
-				bHasX ? FAnimationCompression_PerTrackUtils::DecompressFixed16<LogScale>(X) : 0.0f,
-				bHasY ? FAnimationCompression_PerTrackUtils::DecompressFixed16<LogScale>(Y) : 0.0f,
-				bHasZ ? FAnimationCompression_PerTrackUtils::DecompressFixed16<LogScale>(Z) : 0.0f);
+				bHasX ? FAnimationCompression_PerTrackUtils::DecompressFixed16<AnimationCompressionUtils::LogScale>(X) : 0.0f,
+				bHasY ? FAnimationCompression_PerTrackUtils::DecompressFixed16<AnimationCompressionUtils::LogScale>(Y) : 0.0f,
+				bHasZ ? FAnimationCompression_PerTrackUtils::DecompressFixed16<AnimationCompressionUtils::LogScale>(Z) : 0.0f);
 
 			const float Error = (V - DecompressedV).Size();
 			MaxError = FMath::Max(MaxError, Error);
@@ -637,24 +637,24 @@ protected:
 
 			if (bHasX)
 			{
-				X = FAnimationCompression_PerTrackUtils::CompressFixed16(V.X, LogScale);
+				X = FAnimationCompression_PerTrackUtils::CompressFixed16(V.X, AnimationCompressionUtils::LogScale);
 				AppendBytes(&X, sizeof(X));
 			}
 			if (bHasY)
 			{
-				Y = FAnimationCompression_PerTrackUtils::CompressFixed16(V.Y, LogScale);
+				Y = FAnimationCompression_PerTrackUtils::CompressFixed16(V.Y, AnimationCompressionUtils::LogScale);
 				AppendBytes(&Y, sizeof(Y));
 			}
 			if (bHasZ)
 			{
-				Z = FAnimationCompression_PerTrackUtils::CompressFixed16(V.Z, LogScale);
+				Z = FAnimationCompression_PerTrackUtils::CompressFixed16(V.Z, AnimationCompressionUtils::LogScale);
 				AppendBytes(&Z, sizeof(Z));
 			}
 
 			const FVector3f DecompressedV(
-				bHasX ? FAnimationCompression_PerTrackUtils::DecompressFixed16<LogScale>(X) : 0.0f,
-				bHasY ? FAnimationCompression_PerTrackUtils::DecompressFixed16<LogScale>(Y) : 0.0f,
-				bHasZ ? FAnimationCompression_PerTrackUtils::DecompressFixed16<LogScale>(Z) : 0.0f);
+				bHasX ? FAnimationCompression_PerTrackUtils::DecompressFixed16<AnimationCompressionUtils::LogScale>(X) : 0.0f,
+				bHasY ? FAnimationCompression_PerTrackUtils::DecompressFixed16<AnimationCompressionUtils::LogScale>(Y) : 0.0f,
+				bHasZ ? FAnimationCompression_PerTrackUtils::DecompressFixed16<AnimationCompressionUtils::LogScale>(Z) : 0.0f);
 
 			const float Error = (V - DecompressedV).Size();
 			MaxError = FMath::Max(MaxError, Error);
diff --git a/Engine/Source/Runtime/Engine/Public/AnimationCompression.h b/Engine/Source/Runtime/Engine/Public/AnimationCompression.h
--- a/Engine/Source/Runtime/Engine/Public/AnimationCompression.h
+++ b/Engine/Source/Runtime/Engine/Public/AnimationCompression.h
@@ -35,11 +35,11 @@ ENGINE_API DECLARE_LOG_CATEGORY_EXTERN(LogAnimationCompression, Display, All);
 #define Quant11BitFactor  (1023.f)
 #define Quant11BitOffs    (1023)
 
-//@TODO: Explore different scales
-inline static constexpr int32 LogScale = 7;
-
 namespace AnimationCompressionUtils
 {
+	//@TODO: Explore different scales
+	inline static constexpr int32 LogScale = 7;
+
 	template<typename ValueType>
 	inline ValueType UnalignedRead(const void* Ptr)
 	{
@@ -907,7 +907,7 @@ public:
 			const uint16* RESTRICT TypedKeyData = (const uint16*)KeyData;
 			if (FormatFlags & 1)
 			{
-				Out.X = FAnimationCompression_PerTrackUtils::DecompressFixed16<LogScale>(*TypedKeyData++);
+				Out.X = FAnimationCompression_PerTrackUtils::DecompressFixed16<AnimationCompressionUtils::LogScale>(*TypedKeyData++);
 			}
 			else
 			{
@@ -916,7 +916,7 @@ public:
 
 			if (FormatFlags & 2)
 			{
-				Out.Y = FAnimationCompression_PerTrackUtils::DecompressFixed16<LogScale>(*TypedKeyData++);
+				Out.Y = FAnimationCompression_PerTrackUtils::DecompressFixed16<AnimationCompressionUtils::LogScale>(*TypedKeyData++);
 			}
 			else
 			{
@@ -925,7 +925,7 @@ public:
 
 			if (FormatFlags & 4)
 			{
-				Out.Z = FAnimationCompression_PerTrackUtils::DecompressFixed16<LogScale>(*TypedKeyData++);
+				Out.Z = FAnimationCompression_PerTrackUtils::DecompressFixed16<AnimationCompressionUtils::LogScale>(*TypedKeyData++);
 			}
 			else
 			{
@@ -1090,7 +1090,7 @@ public:
 			const uint16* RESTRICT TypedKeyData = (const uint16*)KeyData;
 			if (FormatFlags & 1)
 			{
-				Out.X = FAnimationCompression_PerTrackUtils::DecompressFixed16<LogScale>(*TypedKeyData++);
+				Out.X = FAnimationCompression_PerTrackUtils::DecompressFixed16<AnimationCompressionUtils::LogScale>(*TypedKeyData++);
 			}
 			else
 			{
@@ -1099,7 +1099,7 @@ public:
 
 			if (FormatFlags & 2)
 			{
-				Out.Y = FAnimationCompression_PerTrackUtils::DecompressFixed16<LogScale>(*TypedKeyData++);
+				Out.Y = FAnimationCompression_PerTrackUtils::DecompressFixed16<AnimationCompressionUtils::LogScale>(*TypedKeyData++);
 			}
 			else
 			{
@@ -1108,7 +1108,7 @@ public:
 
 			if (FormatFlags & 4)
 			{
-				Out.Z = FAnimationCompression_PerTrackUtils::DecompressFixed16<LogScale>(*TypedKeyData++);
+				Out.Z = FAnimationCompression_PerTrackUtils::DecompressFixed16<AnimationCompressionUtils::LogScale>(*TypedKeyData++);
 			}
 			else
 			{
