Fixed material renaming issue introduced in 38777582.

diff --git a/Engine/Source/Runtime/Engine/Private/Materials/MaterialInterface.cpp b/Engine/Source/Runtime/Engine/Private/Materials/MaterialInterface.cpp
--- a/Engine/Source/Runtime/Engine/Private/Materials/MaterialInterface.cpp
+++ b/Engine/Source/Runtime/Engine/Private/Materials/MaterialInterface.cpp
@@ -136,7 +136,7 @@ static bool SafeRenameEditorOnlyData(UMaterialInterface* MaterialInterface, cons
 		return EditorOnlyCurrent->Rename(*EditorOnlyDataName, nullptr, REN_NonTransactional | REN_DontCreateRedirectors);
 	}
 
-	return false;
+	return true;
 }
 
 #endif
