Add an alias to "SSSwapController" which is SplitScreen.RotateControllerIds. This is more discoverable and easier to find.

diff --git a/Engine/Source/Runtime/Engine/Private/GameViewportClient.cpp b/Engine/Source/Runtime/Engine/Private/GameViewportClient.cpp
--- a/Engine/Source/Runtime/Engine/Private/GameViewportClient.cpp
+++ b/Engine/Source/Runtime/Engine/Private/GameViewportClient.cpp
@@ -135,6 +135,21 @@ static TAutoConsoleVariable<bool> CVarRemapDeviceIdForOffsetPlayerGamepadIds(
 	TEXT("Note: This CVar will be removed in a future release, this is a temporary wrapper for bug fix behavior."),
 	ECVF_Default);
 
+static FAutoConsoleCommandWithWorld SplitScreenSwapControllersCommand
+(
+	TEXT("SplitScreen.RotateControllerIds"),
+	TEXT("Rotates controller ids among gameplayers, useful for testing splitscreen with only one controller. Alias to 'SSSwapController'"),
+	FConsoleCommandWithWorldDelegate::CreateLambda([](const UWorld* InWorld)
+		{
+			UGameViewportClient* Client = InWorld->GetGameViewport();
+			if (!Client)
+			{
+				return;
+			}
+			Client->SSSwapControllers();
+		})
+);
+
 #if CSV_PROFILER_STATS
 struct FCsvLocalPlayer
 {
