Fix for cloud component

#tests PC dx12 editor
#rnx

diff --git a/Engine/Source/Runtime/Engine/Classes/Components/VolumetricCloudComponent.h b/Engine/Source/Runtime/Engine/Classes/Components/VolumetricCloudComponent.h
--- a/Engine/Source/Runtime/Engine/Classes/Components/VolumetricCloudComponent.h
+++ b/Engine/Source/Runtime/Engine/Classes/Components/VolumetricCloudComponent.h
@@ -202,6 +202,7 @@ protected:
 	//~ Begin UActorComponent Interface.
 	virtual void CreateRenderState_Concurrent(FRegisterComponentContext* Context) override;
 	virtual void DestroyRenderState_Concurrent() override;
+	virtual void OnRegister() override;
 	//~ End UActorComponent Interface.
 
 public:
diff --git a/Engine/Source/Runtime/Engine/Private/Components/VolumetricCloudComponent.cpp b/Engine/Source/Runtime/Engine/Private/Components/VolumetricCloudComponent.cpp
--- a/Engine/Source/Runtime/Engine/Private/Components/VolumetricCloudComponent.cpp
+++ b/Engine/Source/Runtime/Engine/Private/Components/VolumetricCloudComponent.cpp
@@ -57,8 +57,6 @@ UVolumetricCloudComponent::UVolumetricCloudComponent(const FObjectInitializer& O
 	, bRenderInMainPass(true)
 	, VolumetricCloudSceneProxy(nullptr)
 {
-	static ConstructorHelpers::FObjectFinder<UMaterialInterface> VolumetricCloudDefaultMaterialRef(TEXT("/Engine/EngineSky/VolumetricClouds/m_SimpleVolumetricCloud_Inst.m_SimpleVolumetricCloud_Inst"));
-	Material = VolumetricCloudDefaultMaterialRef.Object;
 }
 
 UVolumetricCloudComponent::~UVolumetricCloudComponent()
@@ -142,6 +140,20 @@ void UVolumetricCloudComponent::Serialize(FArchive& Ar)
 	}
 }
 
+void UVolumetricCloudComponent::OnRegister()
+{
+	Super::OnRegister();
+
+	// Only load and assign a default the material if none is assigned.
+	// This prevent the loading of default the material in game, and provide a default material in editor when drag and dropping the component.
+	if (Material == nullptr)
+	{
+		// This is a similar behavior to material on meshes: a default material is always assigned if none are set.
+		check(IsInGameThread());
+		Material = LoadObject<UMaterialInterface>(NULL, TEXT("/Engine/EngineSky/VolumetricClouds/m_SimpleVolumetricCloud_Inst.m_SimpleVolumetricCloud_Inst"));
+	}
+}
+
 
 #define CLOUD_DECLARE_BLUEPRINT_SETFUNCTION(MemberType, MemberName) void UVolumetricCloudComponent::Set##MemberName(MemberType NewValue)\
 {\
