RepLayout now compares object properties using TObjectPtr instead of raw UObject* to prevent unnecessary object dereferences when enabling UE_WITH_REMOTE_OBJECT_HANDLE.

diff --git a/Engine/Source/Runtime/Engine/Private/RepLayout.cpp b/Engine/Source/Runtime/Engine/Private/RepLayout.cpp
--- a/Engine/Source/Runtime/Engine/Private/RepLayout.cpp
+++ b/Engine/Source/Runtime/Engine/Private/RepLayout.cpp
@@ -601,8 +601,8 @@ static FORCEINLINE bool CompareObject(
 	// This saves us from having to add referenced objects during GC
 	FObjectPropertyBase* ObjProperty = CastFieldChecked<FObjectPropertyBase>(Cmd.Property);
 
-	UObject* ObjectA = ObjProperty->GetObjectPropertyValue(A);
-	UObject* ObjectB = ObjProperty->GetObjectPropertyValue(B);
+	TObjectPtr<UObject> ObjectA = ObjProperty->GetObjectPtrPropertyValue(A);
+	TObjectPtr<UObject> ObjectB = ObjProperty->GetObjectPtrPropertyValue(B);
 
 	return ObjectA == ObjectB;
 }
