Optimization: AWorldPartitionHLOD now skips the incremental pre-register/pre-unregister components phases as it is not needed.

diff --git a/Engine/Source/Runtime/Engine/Classes/GameFramework/Actor.h b/Engine/Source/Runtime/Engine/Classes/GameFramework/Actor.h
--- a/Engine/Source/Runtime/Engine/Classes/GameFramework/Actor.h
+++ b/Engine/Source/Runtime/Engine/Classes/GameFramework/Actor.h
@@ -1510,6 +1510,12 @@ private:
 	ENGINE_API TArray<const UDataLayerInstance*> GetDataLayerInstancesInternal(bool bUseLevelContext, bool bIncludeParentDataLayers = true) const;
 	ENGINE_API const IWorldPartitionCell* GetWorldPartitionRuntimeCell() const;
 
+	/** Whether actor supports incremental pre-register of components. */
+	ENGINE_API virtual bool SupportsIncrementalPreRegisterComponents() const;
+
+	/** Whether actor supports incremental pre-unregister of components. */
+	ENGINE_API virtual bool SupportsIncrementalPreUnregisterComponents() const;
+
 	/** Incrementally pre-registers components associated with this actor (used during level streaming). */
 	ENGINE_API void IncrementalPreRegisterComponents();
 
@@ -4949,6 +4955,16 @@ FORCEINLINE_DEBUGGABLE bool AActor::IsNetMode(ENetMode Mode) const
 #endif
 }
 
+FORCEINLINE bool AActor::SupportsIncrementalPreRegisterComponents() const
+{
+	return true;
+}
+
+FORCEINLINE bool AActor::SupportsIncrementalPreUnregisterComponents() const
+{
+	return true;
+}
+
 #if WITH_EDITOR
 /** Callback for editor actor selection. This must be in engine instead of editor for AActor::IsSelectedInEditor to work */
 extern ENGINE_API TFunction<bool(const AActor*)> GIsActorSelectedInEditor;
diff --git a/Engine/Source/Runtime/Engine/Private/Actor.cpp b/Engine/Source/Runtime/Engine/Private/Actor.cpp
--- a/Engine/Source/Runtime/Engine/Private/Actor.cpp
+++ b/Engine/Source/Runtime/Engine/Private/Actor.cpp
@@ -5735,6 +5735,11 @@ static USceneComponent* GetNonPreRegisteredParent(UActorComponent* Component)
 
 void AActor::IncrementalPreUnregisterComponents()
 {
+	if (!SupportsIncrementalPreUnregisterComponents())
+	{
+		return;
+	}
+
 	TRACE_CPUPROFILER_EVENT_SCOPE(AActor::IncrementalPreUnregisterComponents);
 
 	TInlineComponentArray<UActorComponent*> Components;
@@ -5762,6 +5767,11 @@ void AActor::IncrementalPreUnregisterComponents()
 
 void AActor::IncrementalPreRegisterComponents()
 {
+	if (!SupportsIncrementalPreRegisterComponents())
+	{
+		return;
+	}
+
 	TRACE_CPUPROFILER_EVENT_SCOPE(AActor::IncrementalPreRegisterComponents);
 
 	// First call PreRegisterAllComponents (IncrementalPreRegisterComponents includes PreRegisterAllComponents)
diff --git a/Engine/Source/Runtime/Engine/Public/WorldPartition/HLOD/HLODActor.h b/Engine/Source/Runtime/Engine/Public/WorldPartition/HLOD/HLODActor.h
--- a/Engine/Source/Runtime/Engine/Public/WorldPartition/HLOD/HLODActor.h
+++ b/Engine/Source/Runtime/Engine/Public/WorldPartition/HLOD/HLODActor.h
@@ -90,6 +90,8 @@ protected:
 	ENGINE_API virtual void PreRegisterAllComponents() override;
 	ENGINE_API virtual void BeginPlay() override;
 	ENGINE_API virtual void EndPlay(const EEndPlayReason::Type EndPlayReason) override;
+	ENGINE_API virtual bool SupportsIncrementalPreRegisterComponents() const override;
+	ENGINE_API virtual bool SupportsIncrementalPreUnregisterComponents() const override;
 	virtual bool IsRuntimeOnly() const override { return true; }
 #if WITH_EDITOR
 	ENGINE_API virtual TUniquePtr<class FWorldPartitionActorDesc> CreateClassActorDesc() const override;
@@ -150,4 +152,14 @@ private:
 #endif	
 };
 
+FORCEINLINE bool AWorldPartitionHLOD::SupportsIncrementalPreRegisterComponents() const
+{
+	return false;
+}
+
+FORCEINLINE bool AWorldPartitionHLOD::SupportsIncrementalPreUnregisterComponents() const
+{
+	return false;
+}
+
 DEFINE_ACTORDESC_TYPE(AWorldPartitionHLOD, FHLODActorDesc);
