Remove dll export from SkeletalMeshComponentInstance data as it lives in a private header

diff --git a/Engine/Source/Runtime/Engine/Private/Components/SkeletalMeshComponentInstanceData.h b/Engine/Source/Runtime/Engine/Private/Components/SkeletalMeshComponentInstanceData.h
--- a/Engine/Source/Runtime/Engine/Private/Components/SkeletalMeshComponentInstanceData.h
+++ b/Engine/Source/Runtime/Engine/Private/Components/SkeletalMeshComponentInstanceData.h
@@ -9,7 +9,7 @@ class USkeletalMeshComponent;
 
 /** Saves internal SkeletalMesh (transient) state that gets lost at reconstruction for Blueprint created components. */
 USTRUCT()
-struct ENGINE_API FSkeletalMeshComponentInstanceData : public FSceneComponentInstanceData
+struct FSkeletalMeshComponentInstanceData : public FSceneComponentInstanceData
 {
 	GENERATED_BODY()
 	
