Optimization: Make all components of AWorldPartitionHLOD non relevant for navigation

diff --git a/Engine/Source/Runtime/Engine/Classes/GameFramework/Actor.h b/Engine/Source/Runtime/Engine/Classes/GameFramework/Actor.h
--- a/Engine/Source/Runtime/Engine/Classes/GameFramework/Actor.h
+++ b/Engine/Source/Runtime/Engine/Classes/GameFramework/Actor.h
@@ -1511,10 +1511,16 @@ private:
 	ENGINE_API const IWorldPartitionCell* GetWorldPartitionRuntimeCell() const;
 
 	/** Whether actor supports incremental pre-register of components. */
-	ENGINE_API virtual bool SupportsIncrementalPreRegisterComponents() const;
+	ENGINE_API virtual bool SupportsIncrementalPreRegisterComponents() const
+	{
+		return true;
+	}
 
 	/** Whether actor supports incremental pre-unregister of components. */
-	ENGINE_API virtual bool SupportsIncrementalPreUnregisterComponents() const;
+	ENGINE_API virtual bool SupportsIncrementalPreUnregisterComponents() const
+	{
+		return true;
+	}
 
 	/** Incrementally pre-registers components associated with this actor (used during level streaming). */
 	ENGINE_API void IncrementalPreRegisterComponents();
@@ -4955,16 +4961,6 @@ FORCEINLINE_DEBUGGABLE bool AActor::IsNetMode(ENetMode Mode) const
 #endif
 }
 
-FORCEINLINE bool AActor::SupportsIncrementalPreRegisterComponents() const
-{
-	return true;
-}
-
-FORCEINLINE bool AActor::SupportsIncrementalPreUnregisterComponents() const
-{
-	return true;
-}
-
 #if WITH_EDITOR
 /** Callback for editor actor selection. This must be in engine instead of editor for AActor::IsSelectedInEditor to work */
 extern ENGINE_API TFunction<bool(const AActor*)> GIsActorSelectedInEditor;
diff --git a/Engine/Source/Runtime/Engine/Public/WorldPartition/HLOD/HLODActor.h b/Engine/Source/Runtime/Engine/Public/WorldPartition/HLOD/HLODActor.h
--- a/Engine/Source/Runtime/Engine/Public/WorldPartition/HLOD/HLODActor.h
+++ b/Engine/Source/Runtime/Engine/Public/WorldPartition/HLOD/HLODActor.h
@@ -90,8 +90,18 @@ protected:
 	ENGINE_API virtual void PreRegisterAllComponents() override;
 	ENGINE_API virtual void BeginPlay() override;
 	ENGINE_API virtual void EndPlay(const EEndPlayReason::Type EndPlayReason) override;
-	ENGINE_API virtual bool SupportsIncrementalPreRegisterComponents() const override;
-	ENGINE_API virtual bool SupportsIncrementalPreUnregisterComponents() const override;
+	ENGINE_API virtual bool SupportsIncrementalPreRegisterComponents() const override
+	{
+		return false;
+	}
+	ENGINE_API virtual bool SupportsIncrementalPreUnregisterComponents() const override
+	{
+		return false;
+	}
+	ENGINE_API virtual bool IsComponentRelevantForNavigation(UActorComponent* Component) const override
+	{
+		return false;
+	}
 	virtual bool IsRuntimeOnly() const override { return true; }
 #if WITH_EDITOR
 	ENGINE_API virtual TUniquePtr<class FWorldPartitionActorDesc> CreateClassActorDesc() const override;
@@ -152,14 +162,4 @@ private:
 #endif	
 };
 
-FORCEINLINE bool AWorldPartitionHLOD::SupportsIncrementalPreRegisterComponents() const
-{
-	return false;
-}
-
-FORCEINLINE bool AWorldPartitionHLOD::SupportsIncrementalPreUnregisterComponents() const
-{
-	return false;
-}
-
 DEFINE_ACTORDESC_TYPE(AWorldPartitionHLOD, FHLODActorDesc);
