DeformerGraph move logic for last frame invalidation on LOD change to final application of vertex factory.
This will simplify a future change for how we manage the position/previous position buffers.
Also add/tidy deprecation macros from this and one previous change.

diff --git a/Engine/Plugins/Animation/DeformerGraph/Source/OptimusCore/Private/DataInterfaces/OptimusDataInterfaceSkinnedMeshRead.cpp b/Engine/Plugins/Animation/DeformerGraph/Source/OptimusCore/Private/DataInterfaces/OptimusDataInterfaceSkinnedMeshRead.cpp
--- a/Engine/Plugins/Animation/DeformerGraph/Source/OptimusCore/Private/DataInterfaces/OptimusDataInterfaceSkinnedMeshRead.cpp
+++ b/Engine/Plugins/Animation/DeformerGraph/Source/OptimusCore/Private/DataInterfaces/OptimusDataInterfaceSkinnedMeshRead.cpp
@@ -211,16 +211,18 @@ void FOptimusSkinnedMeshReadDataProviderProxy::AllocateResources(FRDGBuilder& Gr
 	// Allocate required buffers
 	const int32 LodIndex = SkeletalMeshObject->GetLOD();
 	
-	bool bLodJustChanged = false;
+	// Avoid using previous position buffer from when this LOD was last active to compute motion vectors.
+	// The position delta between that previous position (could be from any time ago) and the current position can be any crazy value that is not meaningful.
+	bool bInvalidatePreviousPosition = false;
 	if (LodIndex != *LastLodIndexPtr)
 	{
-		bLodJustChanged = true;	
+		bInvalidatePreviousPosition = true;
 		*LastLodIndexPtr = LodIndex;
 	}
 	
 	if (InputMask & static_cast<uint64>(ESkinnedMeshReadDataInterfaceInputSelectorMask::Position) )
 	{
-		FRDGBuffer* PositionBuffer = FSkeletalMeshDeformerHelpers::AllocateVertexFactoryPositionBuffer(GraphBuilder, InAllocationData.ExternalAccessQueue, SkeletalMeshObject, LodIndex, bLodJustChanged,TEXT("OptimusSkinnedMeshPosition"));
+		FRDGBuffer* PositionBuffer = FSkeletalMeshDeformerHelpers::AllocateVertexFactoryPositionBuffer(GraphBuilder, InAllocationData.ExternalAccessQueue, SkeletalMeshObject, LodIndex, TEXT("OptimusSkinnedMeshPosition"));
 		PositionBufferUAV = GraphBuilder.CreateUAV(PositionBuffer, PF_R32_FLOAT, ERDGUnorderedAccessViewFlags::SkipBarrier);
 	}
 	else
@@ -253,7 +255,7 @@ void FOptimusSkinnedMeshReadDataProviderProxy::AllocateResources(FRDGBuilder& Gr
 		ColorBufferUAV = GraphBuilder.CreateUAV(GraphBuilder.RegisterExternalBuffer(GWhiteVertexBufferWithRDG->Buffer), PF_A32B32G32R32F);
 	}
 
-	FSkeletalMeshDeformerHelpers::UpdateVertexFactoryBufferOverrides(GraphBuilder.RHICmdList, SkeletalMeshObject, LodIndex);
+	FSkeletalMeshDeformerHelpers::UpdateVertexFactoryBufferOverrides(GraphBuilder.RHICmdList, SkeletalMeshObject, LodIndex, bInvalidatePreviousPosition);
 }
 
 void FOptimusSkinnedMeshReadDataProviderProxy::GatherPermutations(FPermutationData& InOutPermutationData) const
diff --git a/Engine/Plugins/Animation/DeformerGraph/Source/OptimusCore/Private/DataInterfaces/OptimusDataInterfaceSkinnedMeshWrite.cpp b/Engine/Plugins/Animation/DeformerGraph/Source/OptimusCore/Private/DataInterfaces/OptimusDataInterfaceSkinnedMeshWrite.cpp
--- a/Engine/Plugins/Animation/DeformerGraph/Source/OptimusCore/Private/DataInterfaces/OptimusDataInterfaceSkinnedMeshWrite.cpp
+++ b/Engine/Plugins/Animation/DeformerGraph/Source/OptimusCore/Private/DataInterfaces/OptimusDataInterfaceSkinnedMeshWrite.cpp
@@ -428,16 +428,18 @@ void FOptimusSkinnedMeshWriteDataProviderProxy::AllocateResources(FRDGBuilder& G
 	// Allocate required buffers
 	const int32 LodIndex = SkeletalMeshObject->GetLOD();
 
-	bool bLodJustChanged = false;
+	// Avoid using previous position buffer from when this LOD was last active to compute motion vectors.
+	// The position delta between that previous position (could be from any time ago) and the current position can be any crazy value that is not meaningful.
+	bool bInvalidatePreviousPosition = false;
 	if (LodIndex != *LastLodIndexPtr)
 	{
-		bLodJustChanged = true;	
+		bInvalidatePreviousPosition = true;
 		*LastLodIndexPtr = LodIndex;
 	}
 
 	if (OutputMask & static_cast<uint64>(ESkinnedMeshWriteDataInterfaceOutputSelectorMask::Position))
 	{
-		PositionBuffer = FSkeletalMeshDeformerHelpers::AllocateVertexFactoryPositionBuffer(GraphBuilder, InAllocationData.ExternalAccessQueue, SkeletalMeshObject, LodIndex, bLodJustChanged,TEXT("OptimusSkinnedMeshPosition"));
+		PositionBuffer = FSkeletalMeshDeformerHelpers::AllocateVertexFactoryPositionBuffer(GraphBuilder, InAllocationData.ExternalAccessQueue, SkeletalMeshObject, LodIndex, TEXT("OptimusSkinnedMeshPosition"));
 		PositionBufferUAV = GraphBuilder.CreateUAV(PositionBuffer, PF_R32_FLOAT, ERDGUnorderedAccessViewFlags::SkipBarrier);
 	}
 	else
@@ -494,7 +496,7 @@ void FOptimusSkinnedMeshWriteDataProviderProxy::AllocateResources(FRDGBuilder& G
 		ColorBufferUAV = GraphBuilder.CreateUAV(GraphBuilder.RegisterExternalBuffer(GWhiteVertexBufferWithRDG->Buffer), PF_A32B32G32R32F);
 	}
 
-	FSkeletalMeshDeformerHelpers::UpdateVertexFactoryBufferOverrides(GraphBuilder.RHICmdList, SkeletalMeshObject, LodIndex);
+	FSkeletalMeshDeformerHelpers::UpdateVertexFactoryBufferOverrides(GraphBuilder.RHICmdList, SkeletalMeshObject, LodIndex, bInvalidatePreviousPosition);
 }
 
 void FOptimusSkinnedMeshWriteDataProviderProxy::GatherDispatchData(FDispatchData const& InDispatchData)
diff --git a/Engine/Source/Runtime/Engine/Private/SkeletalMeshDeformerHelpers.cpp b/Engine/Source/Runtime/Engine/Private/SkeletalMeshDeformerHelpers.cpp
--- a/Engine/Source/Runtime/Engine/Private/SkeletalMeshDeformerHelpers.cpp
+++ b/Engine/Source/Runtime/Engine/Private/SkeletalMeshDeformerHelpers.cpp
@@ -139,7 +139,7 @@ FRDGBuffer* FSkeletalMeshDeformerHelpers::GetAllocatedColorBuffer(FRDGBuilder& G
 	return ColorBuffer;
 }
 
-FRDGBuffer* FSkeletalMeshDeformerHelpers::AllocateVertexFactoryPositionBuffer(FRDGBuilder& GraphBuilder, FRDGExternalAccessQueue& ExternalAccesQueue, FSkeletalMeshObject* InMeshObject, int32 InLodIndex, bool bInLodJustChanged, TCHAR const* InBufferName)
+FRDGBuffer* FSkeletalMeshDeformerHelpers::AllocateVertexFactoryPositionBuffer(FRDGBuilder& GraphBuilder, FRDGExternalAccessQueue& ExternalAccesQueue, FSkeletalMeshObject* InMeshObject, int32 InLodIndex, TCHAR const* InBufferName)
 {
 	if (InMeshObject->IsCPUSkinned())
 	{
@@ -172,15 +172,6 @@ FRDGBuffer* FSkeletalMeshDeformerHelpers::AllocateVertexFactoryPositionBuffer(FR
 
 		DeformerGeometry.PositionSRV = DeformerGeometry.Position->GetOrCreateSRV(GraphBuilder.RHICmdList, FRHIBufferSRVCreateInfo(PF_R32_FLOAT));
 
-		// Avoid using position buffer from the last time this LOD was active to compute motion vectors,
-		// the position delta between that previous position (could be from any time ago) and the current position can be any crazy value that is not meaningful
-		// instead lets just set the motion vector to zero here
-		if (bInLodJustChanged)
-		{
-			DeformerGeometry.PrevPosition = DeformerGeometry.Position;
-			DeformerGeometry.PrevPositionSRV = DeformerGeometry.PositionSRV;
-		}
-		
 		DeformerGeometry.PositionUpdatedFrame = Frame;
 		
 #if RHI_RAYTRACING
@@ -206,6 +197,16 @@ FRDGBuffer* FSkeletalMeshDeformerHelpers::AllocateVertexFactoryPositionBuffer(FR
 	return PositionBuffer;
 }
 
+FRDGBuffer* FSkeletalMeshDeformerHelpers::AllocateVertexFactoryPositionBuffer(FRDGBuilder& GraphBuilder, FSkeletalMeshObject* InMeshObject, int32 InLodIndex, bool bInLodJustChanged, TCHAR const* InBufferName)
+{
+	FRDGExternalAccessQueue Queue;
+PRAGMA_DISABLE_DEPRECATION_WARNINGS
+	FRDGBuffer* Buffer = AllocateVertexFactoryPositionBuffer(GraphBuilder, Queue, InMeshObject, InLodIndex, InBufferName);
+PRAGMA_ENABLE_DEPRECATION_WARNINGS
+	Queue.Submit(GraphBuilder);
+	return Buffer;
+}
+
 FRDGBuffer* FSkeletalMeshDeformerHelpers::AllocateVertexFactoryTangentBuffer(FRDGBuilder& GraphBuilder, FRDGExternalAccessQueue& ExternalAccesQueue, FSkeletalMeshObject* InMeshObject, int32 InLodIndex, TCHAR const* InBufferName)
 {
 	if (InMeshObject->IsCPUSkinned())
@@ -243,6 +244,16 @@ FRDGBuffer* FSkeletalMeshDeformerHelpers::AllocateVertexFactoryTangentBuffer(FRD
 	return TangentBuffer;
 }
 
+FRDGBuffer* FSkeletalMeshDeformerHelpers::AllocateVertexFactoryTangentBuffer(FRDGBuilder& GraphBuilder, FSkeletalMeshObject* InMeshObject, int32 InLodIndex, TCHAR const* InBufferName)
+{
+	FRDGExternalAccessQueue Queue;
+PRAGMA_DISABLE_DEPRECATION_WARNINGS
+	FRDGBuffer* Buffer = AllocateVertexFactoryTangentBuffer(GraphBuilder, Queue, InMeshObject, InLodIndex, InBufferName);
+PRAGMA_ENABLE_DEPRECATION_WARNINGS
+	Queue.Submit(GraphBuilder);
+	return Buffer;
+}
+
 FRDGBuffer* FSkeletalMeshDeformerHelpers::AllocateVertexFactoryColorBuffer(FRDGBuilder& GraphBuilder, FRDGExternalAccessQueue& ExternalAccesQueue, FSkeletalMeshObject* InMeshObject, int32 InLodIndex, TCHAR const* InBufferName)
 {
 	if (InMeshObject->IsCPUSkinned())
@@ -279,12 +290,17 @@ FRDGBuffer* FSkeletalMeshDeformerHelpers::AllocateVertexFactoryColorBuffer(FRDGB
 	return ColorBuffer;
 }
 
-void FSkeletalMeshDeformerHelpers::UpdateVertexFactoryBufferOverrides(FSkeletalMeshObject* InMeshObject, int32 InLodIndex)
+FRDGBuffer* FSkeletalMeshDeformerHelpers::AllocateVertexFactoryColorBuffer(FRDGBuilder& GraphBuilder, FSkeletalMeshObject* InMeshObject, int32 InLodIndex, TCHAR const* InBufferName)
 {
-	UpdateVertexFactoryBufferOverrides(FRHICommandListImmediate::Get(), InMeshObject, InLodIndex);
+	FRDGExternalAccessQueue Queue;
+PRAGMA_DISABLE_DEPRECATION_WARNINGS
+	FRDGBuffer* Buffer = AllocateVertexFactoryColorBuffer(GraphBuilder, Queue, InMeshObject, InLodIndex, InBufferName);
+PRAGMA_ENABLE_DEPRECATION_WARNINGS
+	Queue.Submit(GraphBuilder);
+	return Buffer;
 }
 
-void FSkeletalMeshDeformerHelpers::UpdateVertexFactoryBufferOverrides(FRHICommandListBase& RHICmdList, FSkeletalMeshObject* InMeshObject, int32 InLodIndex)
+void FSkeletalMeshDeformerHelpers::UpdateVertexFactoryBufferOverrides(FRHICommandListBase& RHICmdList, FSkeletalMeshObject* InMeshObject, int32 InLodIndex, bool bInvalidatePreviousPosition)
 {
 	if (InMeshObject->IsCPUSkinned())
 	{
@@ -303,7 +319,7 @@ void FSkeletalMeshDeformerHelpers::UpdateVertexFactoryBufferOverrides(FRHIComman
 	{
 		Desc.StreamBuffers[FGPUSkinPassthroughVertexFactory::VertexPosition] = DeformerGeometry.Position->GetRHI();
 		Desc.SRVs[FGPUSkinPassthroughVertexFactory::Position] = DeformerGeometry.PositionSRV;
-		Desc.SRVs[FGPUSkinPassthroughVertexFactory::PreviousPosition] = DeformerGeometry.PrevPositionSRV;
+		Desc.SRVs[FGPUSkinPassthroughVertexFactory::PreviousPosition] = bInvalidatePreviousPosition ? nullptr : DeformerGeometry.PrevPositionSRV;
 		bAssignedAttributes = true;
 	}
 	if (DeformerGeometry.TangentSRV)
@@ -339,6 +355,20 @@ void FSkeletalMeshDeformerHelpers::UpdateVertexFactoryBufferOverrides(FRHIComman
 	}
 }
 
+void FSkeletalMeshDeformerHelpers::UpdateVertexFactoryBufferOverrides(FRHICommandListBase& RHICmdList, FSkeletalMeshObject* InMeshObject, int32 InLodIndex)
+{
+PRAGMA_DISABLE_DEPRECATION_WARNINGS
+	UpdateVertexFactoryBufferOverrides(RHICmdList, InMeshObject, InLodIndex, false);
+PRAGMA_ENABLE_DEPRECATION_WARNINGS
+}
+
+void FSkeletalMeshDeformerHelpers::UpdateVertexFactoryBufferOverrides(FSkeletalMeshObject* InMeshObject, int32 InLodIndex)
+{
+PRAGMA_DISABLE_DEPRECATION_WARNINGS
+	UpdateVertexFactoryBufferOverrides(FRHICommandListImmediate::Get(), InMeshObject, InLodIndex);
+PRAGMA_ENABLE_DEPRECATION_WARNINGS
+}
+
 void FSkeletalMeshDeformerHelpers::ResetVertexFactoryBufferOverrides(FSkeletalMeshObject* InMeshObject, int32 LODIndex)
 {
 	if (InMeshObject->IsCPUSkinned())
diff --git a/Engine/Source/Runtime/Engine/Public/SkeletalMeshDeformerHelpers.h b/Engine/Source/Runtime/Engine/Public/SkeletalMeshDeformerHelpers.h
--- a/Engine/Source/Runtime/Engine/Public/SkeletalMeshDeformerHelpers.h
+++ b/Engine/Source/Runtime/Engine/Public/SkeletalMeshDeformerHelpers.h
@@ -72,7 +72,7 @@ public:
 	 * Ownership is handled by the MeshObject.
 	 * If we call this more than once for the same MeshObject in the same frame then we return the allocation from the first call.
 	 */
-	ENGINE_API static FRDGBuffer* AllocateVertexFactoryPositionBuffer(FRDGBuilder& GraphBuilder, FRDGExternalAccessQueue& ExternalAccesQueue, FSkeletalMeshObject* InMeshObject, int32 InLodIndex, bool bInLodJustChanged, TCHAR const* InBufferName);
+	ENGINE_API static FRDGBuffer* AllocateVertexFactoryPositionBuffer(FRDGBuilder& GraphBuilder, FRDGExternalAccessQueue& ExternalAccesQueue, FSkeletalMeshObject* InMeshObject, int32 InLodIndex, TCHAR const* InBufferName);
 
 	/**
 	 * Allocate and bind a new tangent buffer and return it for writing.
@@ -92,10 +92,7 @@ public:
 	 * Update all of the MeshObject's passthrough vertex factories with the currently allocated vertex buffers. 
 	 * Usually call this after all AllocateVertexFactory*() functions for a frame.
 	 */
-	ENGINE_API static void UpdateVertexFactoryBufferOverrides(FRHICommandListBase& RHICmdList, FSkeletalMeshObject* InMeshObject, int32 InLodIndex);
-
-	UE_DEPRECATED(5.4, "UpdateVertexFactoryBufferOverrides requires a command list.")
-	ENGINE_API static void UpdateVertexFactoryBufferOverrides(FSkeletalMeshObject* InMeshObject, int32 InLodIndex);
+	ENGINE_API static void UpdateVertexFactoryBufferOverrides(FRHICommandListBase& RHICmdList, FSkeletalMeshObject* InMeshObject, int32 InLodIndex, bool bInvalidatePreviousPosition);
 
 	/** 
 	 * Release all of the the buffers that have been allocated through the AllocateVertexFactory*() functions.
@@ -104,4 +101,23 @@ public:
 	ENGINE_API static void ResetVertexFactoryBufferOverrides(FSkeletalMeshObject* InMeshObject, int32 InLodIndex);
 
 #pragma endregion SetInternals
+
+#pragma region Deprecated
+
+	UE_DEPRECATED(5.6, "AllocateVertexFactoryPositionBuffer requires a FRDGExternalAccessQueue.")
+	ENGINE_API static FRDGBuffer* AllocateVertexFactoryPositionBuffer(FRDGBuilder& GraphBuilder, FSkeletalMeshObject* InMeshObject, int32 InLodIndex, bool bInLodJustChanged, TCHAR const* InBufferName);
+
+	UE_DEPRECATED(5.6, "AllocateVertexFactoryTangentBuffer requires a FRDGExternalAccessQueue.")
+	ENGINE_API static FRDGBuffer* AllocateVertexFactoryTangentBuffer(FRDGBuilder& GraphBuilder, FSkeletalMeshObject* InMeshObject, int32 InLodIndex, TCHAR const* InBufferName);
+
+	UE_DEPRECATED(5.6, "AllocateVertexFactoryColorBuffer requires a FRDGExternalAccessQueue.")
+	ENGINE_API static FRDGBuffer* AllocateVertexFactoryColorBuffer(FRDGBuilder& GraphBuilder, FSkeletalMeshObject* InMeshObject, int32 InLodIndex, TCHAR const* InBufferName);
+
+	UE_DEPRECATED(5.6, "UpdateVertexFactoryBufferOverrides requires a bInvalidatePreviousPosition flag.")
+	ENGINE_API static void UpdateVertexFactoryBufferOverrides(FRHICommandListBase& RHICmdList, FSkeletalMeshObject* InMeshObject, int32 InLodIndex);
+
+	UE_DEPRECATED(5.4, "UpdateVertexFactoryBufferOverrides requires a command list.")
+	ENGINE_API static void UpdateVertexFactoryBufferOverrides(FSkeletalMeshObject* InMeshObject, int32 InLodIndex);
+
+#pragma endregion Deprecated
 };
