PR #9036: Fix to AudioStreamingCache chunk lookup ID

diff --git a/Engine/Source/Runtime/Engine/Private/AudioStreamingCache.cpp b/Engine/Source/Runtime/Engine/Private/AudioStreamingCache.cpp
--- a/Engine/Source/Runtime/Engine/Private/AudioStreamingCache.cpp
+++ b/Engine/Source/Runtime/Engine/Private/AudioStreamingCache.cpp
@@ -2099,7 +2099,8 @@ uint64 FAudioChunkCache::GetCacheLookupIDForChunk(const FChunkKey& InChunkKey) c
 void FAudioChunkCache::SetCacheLookupIDForChunk(const FChunkKey& InChunkKey, uint64 InCacheLookupID)
 {
 	FScopeLock Lock(&CacheMutationCriticalSection.Get());
-	CacheLookupIdMap.FindOrAdd(InChunkKey, InCacheLookupID);
+	uint64& ID = CacheLookupIdMap.FindOrAdd(InChunkKey);
+	ID = InCacheLookupID;
 }
 
 
