Thread-safety for DeviceProfiles::GetConsolidatedCVarValue API
- Deprecated GetConsolidatedCVars, replaced with ForEachConsolidatedCVar

diff --git a/Engine/Source/Runtime/Engine/Classes/DeviceProfiles/DeviceProfile.h b/Engine/Source/Runtime/Engine/Classes/DeviceProfiles/DeviceProfile.h
--- a/Engine/Source/Runtime/Engine/Classes/DeviceProfiles/DeviceProfile.h
+++ b/Engine/Source/Runtime/Engine/Classes/DeviceProfiles/DeviceProfile.h
@@ -134,9 +134,17 @@ public:
 
 	/** Lazily generate a consolidated list of CVars, recursing up the device profile hierarchy 
 	 *  This will not include any cvars from the device's selected fragments.
+	 *  This is _not_ thread-safe (for thread-safety use ForEachConsolidatedCVar or GetConsolidatedCVarValue API instead)
 	*/
+	UE_DEPRECATED(5.6, "GetConsolidatedCVars is no longer supported due to being non-thread safe, use ForEachConsolidatedCvar or GetConsolidatedCVarValue instead")
 	ENGINE_API const TMap<FString, FString>& GetConsolidatedCVars() const;
 
+	/** 
+	 * Iterates over all consolidated string key-value CVar pairs held in this device profile, or in any parent device profile.
+	 * @param	FuncRef		Iteration function (returning false will early out and stop subsequent calls)
+	 */
+	ENGINE_API void ForEachConsolidatedCVar(const TFunctionRef<bool(const FString&, const FString&)>& FuncRef) const;
+
 	/** 
 	 * Get the string value of a CVar that is held in this device profile, or in any parent device profile.
 	 * @param	CVarName		The name of the CVar to find
@@ -171,6 +179,7 @@ private:
 private:
 	/** Consolidated CVars, lazy initialized - access via GetConsolidatedCVars */
 	mutable TMap<FString, FString> ConsolidatedCVars;
+	mutable FCriticalSection CVarsLock;
 
 #endif
 
diff --git a/Engine/Source/Runtime/Engine/Private/DeviceProfiles/DeviceProfile.cpp b/Engine/Source/Runtime/Engine/Private/DeviceProfiles/DeviceProfile.cpp
--- a/Engine/Source/Runtime/Engine/Private/DeviceProfiles/DeviceProfile.cpp
+++ b/Engine/Source/Runtime/Engine/Private/DeviceProfiles/DeviceProfile.cpp
@@ -159,6 +159,7 @@ void UDeviceProfile::ValidateTextureLODGroups()
 void UDeviceProfile::HandleCVarsChanged()
 {
 	OnCVarsUpdated().ExecuteIfBound();
+	FScopeLock Lock(&CVarsLock);
 	ConsolidatedCVars.Reset();
 }
 
@@ -288,7 +289,10 @@ FString UDeviceProfile::GetCVarValue(const FString& CVarName) const
 
 bool UDeviceProfile::GetConsolidatedCVarValue(const TCHAR* CVarName, FString& OutString, bool bCheckDefaults /*=false*/) const
 {
+	FScopeLock Lock(&CVarsLock);
+	PRAGMA_DISABLE_DEPRECATION_WARNINGS
 	const FString* FoundValue = GetConsolidatedCVars().Find(CVarName);
+	PRAGMA_ENABLE_DEPRECATION_WARNINGS
 	if(FoundValue)
 	{
 		OutString = *FoundValue;
@@ -372,6 +376,7 @@ const TMap<FString, FString>& UDeviceProfile::GetConsolidatedCVars() const
 	};
 
 	// First build our own CVar map
+	FScopeLock Lock(&CVarsLock);
 	if(ConsolidatedCVars.Num() == 0)
 	{
 		BuildCVarMap(this, ConsolidatedCVars);
@@ -388,6 +393,21 @@ const TMap<FString, FString>& UDeviceProfile::GetConsolidatedCVars() const
 	return ConsolidatedCVars;
 }
 
+void UDeviceProfile::ForEachConsolidatedCVar(const TFunctionRef<bool(const FString&, const FString&)>& FuncRef) const
+{
+	FScopeLock Lock(&CVarsLock);
+	PRAGMA_DISABLE_DEPRECATION_WARNINGS
+	const TMap<FString, FString>& ConsolidatedCvars = GetConsolidatedCVars();
+	PRAGMA_ENABLE_DEPRECATION_WARNINGS
+	for (const TPair<FString, FString>& Pair : ConsolidatedCvars)
+	{
+		if (!FuncRef(Pair.Key, Pair.Value))
+		{
+			return;
+		}
+	}	
+}
+
 #endif
 
 
