Add a nullptr check in UMaterialExpressionPinBase::PostEditChangeProperty, fixing a crash when deleting an orphaned pin from a Material with bad data.

diff --git a/Engine/Source/Runtime/Engine/Private/Materials/MaterialExpressions.cpp b/Engine/Source/Runtime/Engine/Private/Materials/MaterialExpressions.cpp
--- a/Engine/Source/Runtime/Engine/Private/Materials/MaterialExpressions.cpp
+++ b/Engine/Source/Runtime/Engine/Private/Materials/MaterialExpressions.cpp
@@ -14096,7 +14096,7 @@ void UMaterialExpressionPinBase::PreEditChange(FProperty* PropertyAboutToChange)
 
 void UMaterialExpressionPinBase::PostEditChangeProperty(FPropertyChangedEvent& PropertyChangedEvent)
 {
-	if (PropertyChangedEvent.MemberProperty && GraphNode && SubgraphExpression->GraphNode)
+	if (PropertyChangedEvent.MemberProperty && GraphNode && SubgraphExpression && SubgraphExpression->GraphNode)
 	{
 		Modify();
 		Material->Modify();
