[WorldPartitionActorDesc] Fix stack overflow which could happen during PropagateActorToWorldUpdate

diff --git a/Engine/Source/Runtime/Engine/Private/WorldPartition/ActorDescContainer.cpp b/Engine/Source/Runtime/Engine/Private/WorldPartition/ActorDescContainer.cpp
--- a/Engine/Source/Runtime/Engine/Private/WorldPartition/ActorDescContainer.cpp
+++ b/Engine/Source/Runtime/Engine/Private/WorldPartition/ActorDescContainer.cpp
@@ -702,6 +702,24 @@ void UActorDescContainer::RemoveChildActorFromParentMap(FWorldPartitionActorDesc
 
 void UActorDescContainer::PropagateActorToWorldUpdate(FWorldPartitionActorDesc* ActorDesc)
 {
+	TSet<FGuid> CycleDetector;
+	PropagateActorToWorldUpdateInternal(ActorDesc, CycleDetector);
+}
+
+void UActorDescContainer::PropagateActorToWorldUpdateInternal(FWorldPartitionActorDesc* ActorDesc, TSet<FGuid>& CycleDetector)
+{
+	bool bAlreadyVisited = false;
+	CycleDetector.Add(ActorDesc->GetGuid(), &bAlreadyVisited);
+	if (bAlreadyVisited)
+	{
+		// Each actor can have at most one parent, which means that it can appear on at most one Children list in ParentActorToChildrenMap.
+		// Because of that, reaching the same actor twice, in one series of PropagateActorToWorldUpdate calls, means that we found a cycle.
+		// That can happen when the actor currently being saved becomes a descendant of one of its ex-descendants
+		// and that ex-descendant actor is not saved yet. In that case we can break the cycle here and stop propagating updates. 
+		// They'll be properly propagated once the "ex-descendant" actor mentioned above gets saved.
+		return;
+	}
+
 	if (const TSet<FGuid>* Children = ParentActorToChildrenMap.Find(ActorDesc->GetGuid()))
 	{
 		for (const FGuid& ChildGuid : *Children)
@@ -712,7 +730,7 @@ void UActorDescContainer::PropagateActorToWorldUpdate(FWorldPartitionActorDesc*
 				ChildActor->UpdateActorToWorld();
 				OnActorDescUpdated(ChildActor);
 
-				PropagateActorToWorldUpdate(ChildActor);
+				PropagateActorToWorldUpdateInternal(ChildActor, CycleDetector);
 			}
 		}
 	}
diff --git a/Engine/Source/Runtime/Engine/Public/WorldPartition/ActorDescContainer.h b/Engine/Source/Runtime/Engine/Public/WorldPartition/ActorDescContainer.h
--- a/Engine/Source/Runtime/Engine/Public/WorldPartition/ActorDescContainer.h
+++ b/Engine/Source/Runtime/Engine/Public/WorldPartition/ActorDescContainer.h
@@ -211,6 +211,7 @@ private:
 	void RemoveChildActorFromParentMap(FWorldPartitionActorDesc* ActorDesc);
 
 	void PropagateActorToWorldUpdate(FWorldPartitionActorDesc* ActorDesc);
+	void PropagateActorToWorldUpdateInternal(FWorldPartitionActorDesc* ActorDesc, TSet<FGuid>& CycleDetector);
 
 	static FUObjectAnnotationSparse<FDeletedObjectPlaceholderAnnotation, true> DeletedObjectPlaceholdersAnnotation;
 
