FDebugRenderSceneProxy: match wired and solid render side counts when both are used.

diff --git a/Engine/Source/Runtime/Engine/Private/DebugRenderSceneProxy.cpp b/Engine/Source/Runtime/Engine/Private/DebugRenderSceneProxy.cpp
--- a/Engine/Source/Runtime/Engine/Private/DebugRenderSceneProxy.cpp
+++ b/Engine/Source/Runtime/Engine/Private/DebugRenderSceneProxy.cpp
@@ -459,7 +459,7 @@ void FDebugRenderSceneProxy::FCircle::Draw(FPrimitiveDrawInterface* PDI, EDrawTy
 
 	if (InDrawType == SolidAndWireMeshes || InDrawType == SolidMesh)
 	{
-		GetDiscMesh(Center, XAxis, YAxis, Radius, 12, MaterialCache[Color.WithAlpha(InDrawAlpha * Color.A)], SDPG_World, ViewIndex, Collector);
+		GetDiscMesh(Center, XAxis, YAxis, Radius, (InDrawType == SolidAndWireMeshes) ? 9 : 12, MaterialCache[Color.WithAlpha(InDrawAlpha * Color.A)], SDPG_World, ViewIndex, Collector);
 	}
 }
 
@@ -474,7 +474,7 @@ void FDebugRenderSceneProxy::FWireCylinder::Draw(FPrimitiveDrawInterface* PDI, E
 
 	if (InDrawType == SolidAndWireMeshes || InDrawType == SolidMesh)
 	{
-		GetCylinderMesh(Base, XAxis, YAxis, Direction, Radius, HalfHeight, 16, MaterialCache[Color.WithAlpha(InDrawAlpha * Color.A)], SDPG_World, ViewIndex, Collector);
+		GetCylinderMesh(Base, XAxis, YAxis, Direction, Radius, HalfHeight, (InDrawType == SolidAndWireMeshes) ? 9 : 16, MaterialCache[Color.WithAlpha(InDrawAlpha * Color.A)], SDPG_World, ViewIndex, Collector);
 	}
 }
 
@@ -488,7 +488,7 @@ void FDebugRenderSceneProxy::FCone::Draw(FPrimitiveDrawInterface* PDI, EDrawType
 	}
 	if (InDrawType == SolidAndWireMeshes || InDrawType == SolidMesh)
 	{
-		GetConeMesh(ConeToWorld, Angle1, Angle2, 16, MaterialCache[Color.WithAlpha(InDrawAlpha * Color.A)], SDPG_World, ViewIndex, Collector);
+		GetConeMesh(ConeToWorld, Angle1, Angle2, (InDrawType == SolidAndWireMeshes) ? 9 : 16, MaterialCache[Color.WithAlpha(InDrawAlpha * Color.A)], SDPG_World, ViewIndex, Collector);
 	}
 }
 
@@ -518,6 +518,6 @@ void FDebugRenderSceneProxy::FCapsule::Draw(FPrimitiveDrawInterface* PDI, EDrawT
 	}
 	if (InDrawType == SolidAndWireMeshes || InDrawType == SolidMesh)
 	{
-		GetCapsuleMesh(Base, X, Y, Z, Color, Radius, HalfHeight, 16, MaterialCache[Color.WithAlpha(InDrawAlpha * Color.A)], SDPG_World, false, ViewIndex, Collector);
+		GetCapsuleMesh(Base, X, Y, Z, Color, Radius, HalfHeight, (InDrawType == SolidAndWireMeshes) ? 9 : 16, MaterialCache[Color.WithAlpha(InDrawAlpha * Color.A)], SDPG_World, false, ViewIndex, Collector);
 	}
 }
\ No newline at end of file
