Changed GetRequiredMemoryEstimate to CalculateRequiredMemoryEstimate, caching the result. This fixes an issue where CompressionResultSummary().GatherPreCompressionStats was being called twice for sequences not cached in the DDC

diff --git a/Engine/Source/Runtime/Engine/Private/Animation/AnimationCompressionDerivedData.cpp b/Engine/Source/Runtime/Engine/Private/Animation/AnimationCompressionDerivedData.cpp
--- a/Engine/Source/Runtime/Engine/Private/Animation/AnimationCompressionDerivedData.cpp
+++ b/Engine/Source/Runtime/Engine/Private/Animation/AnimationCompressionDerivedData.cpp
@@ -34,6 +34,7 @@ FAnimationSequenceAsyncCacheTask::FAnimationSequenceAsyncCacheTask(const FIoHash
 		, WeakAnimSequence(&InAnimSequence)
 		, CompressibleAnimPtr(InCompressibleAnimPtr)
 		, TargetPlatform(InTargetPlatform)
+		, RequiredMemory(0)
 {
 	check(!InAnimSequence.IsUnreachable());
 	check(!InAnimSequence.GetSkeleton()->IsUnreachable());
@@ -107,7 +108,7 @@ void FAnimationSequenceAsyncCacheTask::BeginCache(const FIoHash& KeyHash)
 	{
 		FQueuedThreadPool* ThreadPool = Anim::FAnimSequenceCompilingManager::Get().GetThreadPool();
 		const EQueuedWorkPriority BasePriority = Anim::FAnimSequenceCompilingManager::Get().GetBasePriority(AnimSequence);
-		const int64 RequiredMemory = GetRequiredMemoryEstimate();
+		CalculateRequiredMemoryEstimate();
 
 		CompressionStartTime = FPlatformTime::Seconds();
 
@@ -194,7 +195,7 @@ bool FAnimationSequenceAsyncCacheTask::BuildData() const
 	}
 
 	UE_LOG(LogAnimationCompression, Display, TEXT("Building compressed animation data for %s (Required Memory Estimate: %.2f MB)"),
-		*CompressibleAnimPtr->FullName, double(GetRequiredMemoryEstimate()) / (1024.0 * 1024.0));
+		*CompressibleAnimPtr->FullName, double(RequiredMemory) / (1024.0 * 1024.0));
 
 	check(CompressibleAnimPtr.IsValid());
 	FCompressibleAnimData& DataToCompress = *CompressibleAnimPtr.Get();
@@ -284,10 +285,9 @@ void FAnimationSequenceAsyncCacheTask::LaunchCompressionTask(const UE::FSharedSt
 		});
 }
 
-int64 FAnimationSequenceAsyncCacheTask::GetRequiredMemoryEstimate() const
+void FAnimationSequenceAsyncCacheTask::CalculateRequiredMemoryEstimate()
 {
-	int64 RequiredMemory = -1;
-
+	RequiredMemory = 0;
 	if (const UAnimSequence* AnimSequence = WeakAnimSequence.Get())
 	{
 		// Includes bone/curves
@@ -346,8 +346,6 @@ int64 FAnimationSequenceAsyncCacheTask::GetRequiredMemoryEstimate() const
 			}
 		}
 	}
-
-	return RequiredMemory;
 }
 
 int32 GSkipDDC = 0;
diff --git a/Engine/Source/Runtime/Engine/Private/Animation/AnimationCompressionDerivedData.h b/Engine/Source/Runtime/Engine/Private/Animation/AnimationCompressionDerivedData.h
--- a/Engine/Source/Runtime/Engine/Private/Animation/AnimationCompressionDerivedData.h
+++ b/Engine/Source/Runtime/Engine/Private/Animation/AnimationCompressionDerivedData.h
@@ -89,7 +89,7 @@ namespace UE::Anim
 		void EndCache(DerivedData::FCacheGetValueResponse&& Response);
 		bool BuildData() const;
 		void LaunchCompressionTask(const FSharedString& Name, const DerivedData::FCacheKey& Key);
-		int64 GetRequiredMemoryEstimate() const;
+		void CalculateRequiredMemoryEstimate();
 
 	private:
 		friend class FAnimationSequenceAsyncBuildWorker;
@@ -102,7 +102,7 @@ namespace UE::Anim
 		FCompressibleAnimPtr CompressibleAnimPtr;
 		const ITargetPlatform* TargetPlatform;
 		double CompressionStartTime;
-		
+		int64 RequiredMemory;
 	};
 #endif
 }
\ No newline at end of file
