Public API to provide ability to wait for animation compression to be finished (for specific assets)

diff --git a/Engine/Source/Runtime/Engine/Private/Animation/AnimationSequenceCompiler.cpp b/Engine/Source/Runtime/Engine/Private/Animation/AnimationSequenceCompiler.cpp
--- a/Engine/Source/Runtime/Engine/Private/Animation/AnimationSequenceCompiler.cpp
+++ b/Engine/Source/Runtime/Engine/Private/Animation/AnimationSequenceCompiler.cpp
@@ -15,6 +15,7 @@
 #include "Animation/Skeleton.h"
 #include "Animation/AnimSequence.h"
 #include "UObject/StrongObjectPtr.h"
+#include "Animation/IAnimationSequenceCompiler.h"
 
 #define LOCTEXT_NAMESPACE "AnimSequenceCompilingManager"
 
@@ -26,6 +27,11 @@ namespace UE::Anim
 		return Singleton;
 	}
 
+	void IAnimSequenceCompilingManager::FinishCompilation(TArrayView<UAnimSequence* const> InAnimSequences)
+	{
+		FAnimSequenceCompilingManager::Get().FinishCompilation(InAnimSequences); 
+	}
+
 	FAnimSequenceCompilingManager::FAnimSequenceCompilingManager()
 		: Notification(MakeUnique<FAsyncCompilationNotification>(FAnimSequenceCompilingManager::GetAssetNameFormat()))
 	{
diff --git a/Engine/Source/Runtime/Engine/Public/Animation/IAnimationSequenceCompiler.h b/Engine/Source/Runtime/Engine/Public/Animation/IAnimationSequenceCompiler.h
new file mode 100644
--- /dev/null
+++ b/Engine/Source/Runtime/Engine/Public/Animation/IAnimationSequenceCompiler.h
@@ -0,0 +1,20 @@
+// Copyright Epic Games, Inc. All Rights Reserved.
+
+#pragma once
+
+#if WITH_EDITOR
+#include "CoreMinimal.h"
+
+class UAnimSequence;
+
+namespace UE::Anim
+{
+	class IAnimSequenceCompilingManager
+	{
+	public:
+		/** Ensures that, in case of any, compilation tasks for the provided Animation Sequences has been completed */
+		ENGINE_API static void FinishCompilation(TArrayView<UAnimSequence* const> InAnimSequences);
+	};
+}
+
+#endif // WITH_EDITOR
\ No newline at end of file
