IsBone/CurveDataValid APIs

diff --git a/Engine/Source/Runtime/Engine/Private/Animation/AnimCompressionTypes.cpp b/Engine/Source/Runtime/Engine/Private/Animation/AnimCompressionTypes.cpp
--- a/Engine/Source/Runtime/Engine/Private/Animation/AnimCompressionTypes.cpp
+++ b/Engine/Source/Runtime/Engine/Private/Animation/AnimCompressionTypes.cpp
@@ -1605,16 +1605,38 @@ void FCompressedAnimSequence::ClearCompressedCurveData()
 bool FCompressedAnimSequence::IsValid(const UAnimSequence* AnimSequence, bool bLogInformation/*=false*/) const
 {
 #if WITH_EDITOR
-	const bool bHasBoneTracks = AnimSequence->GetDataModelInterface()->GetNumBoneTracks() != 0;
-	const bool bHasCurveData = AnimSequence->GetDataModelInterface()->GetNumberOfFloatCurves() != 0;
+	return IsBoneDataValid(AnimSequence, bLogInformation) && IsCurveDataValid(AnimSequence, bLogInformation);
+#else
+	return true;
+#endif
+}
 
+bool FCompressedAnimSequence::IsBoneDataValid(const UAnimSequence* AnimSequence, bool bLogInformation) const
+{
+#if WITH_EDITOR
+	const bool bHasBoneTracks = AnimSequence->GetDataModelInterface()->GetNumBoneTracks() != 0;
 	const bool bHasValidCompressedBoneData = CompressedDataStructure != nullptr || CompressedTrackToSkeletonMapTable.Num() == 0;
+	const bool bValidCompressedData = !bHasBoneTracks || bHasValidCompressedBoneData;
+	if(bLogInformation && !bValidCompressedData)
+	{
+		UE_LOG(LogAnimation, Warning, TEXT("%s: Num Bone Tracks: %i\nNum Curves: %i\nCompressed bone data: %i (%i)\nCompressed curve data %i\nValid additive: %i"), *AnimSequence->GetName(), AnimSequence->GetDataModelInterface()->GetNumBoneTracks(), AnimSequence->GetDataModelInterface()->GetNumberOfFloatCurves(), CompressedDataStructure != nullptr ? 1 : 0, CompressedTrackToSkeletonMapTable.Num(), CompressedCurveByteStream.Num() != 0 ? 1 : 0, AnimSequence->IsValidAdditive() ? 1 : 0);
+	}
+	return bValidCompressedData;
+#else
+	return true;
+#endif
+}
+
+bool FCompressedAnimSequence::IsCurveDataValid(const UAnimSequence* AnimSequence, bool bLogInformation) const
+{
+#if WITH_EDITOR
+	const bool bHasCurveData = AnimSequence->GetDataModelInterface()->GetNumberOfFloatCurves() != 0;
 	const bool bHasValidCompressedCurveData = CompressedCurveByteStream.Num() != 0;
-	const bool bValidCompressedData = (!bHasBoneTracks || bHasValidCompressedBoneData) && (!bHasCurveData || bHasValidCompressedCurveData || AnimSequence->IsValidAdditive());
+	const bool bValidCompressedData = (!bHasCurveData || bHasValidCompressedCurveData || AnimSequence->IsValidAdditive());
 
 	if(bLogInformation && !bValidCompressedData)
 	{
-		UE_LOG(LogAnimation, Warning, TEXT("Num Bone Tracks: %i\nNum Curves: %i\nCompressed bone data: %i (%i)\nCompressed curve data %i\nValid additive: %i"), AnimSequence->GetDataModelInterface()->GetNumBoneTracks(), AnimSequence->GetDataModelInterface()->GetNumberOfFloatCurves(), CompressedDataStructure != nullptr ? 1 : 0, CompressedTrackToSkeletonMapTable.Num(), CompressedCurveByteStream.Num() != 0 ? 1 : 0, AnimSequence->IsValidAdditive() ? 1 : 0);
+		UE_LOG(LogAnimation, Warning, TEXT("Num Curves: %i\nCompressed curve data %i\nValid additive: %i"), AnimSequence->GetDataModelInterface()->GetNumberOfFloatCurves(), CompressedCurveByteStream.Num() != 0 ? 1 : 0, AnimSequence->IsValidAdditive() ? 1 : 0);
 	}
 
 	return bValidCompressedData;
@@ -1685,27 +1707,27 @@ namespace UE::Anim::Compression
 #if WITH_EDITOR
 template <typename ArrayType>
 void UpdateSHAWithArray(FSHA1& Sha, const TArray<ArrayType>& Array)
-		{
+{
 	Sha.Update((uint8*)Array.GetData(), Array.Num() * Array.GetTypeSize());
-	}
+}
 
 void UpdateSHAWithRawTrack(FSHA1& Sha, const FRawAnimSequenceTrack& RawTrack)
-	{
+{
 	UpdateSHAWithArray(Sha, RawTrack.PosKeys);
 	UpdateSHAWithArray(Sha, RawTrack.RotKeys);
 	UpdateSHAWithArray(Sha, RawTrack.ScaleKeys);
-	}
+}
 
 template<class DataType>
 void UpdateWithData(FSHA1& Sha, const DataType& Data)
-	{
+{
 	Sha.Update((uint8*)(&Data), sizeof(DataType));
-	}
+}
 
 void UpdateSHAWithCurves(FSHA1& Sha, const FRawCurveTracks& InRawCurveData) 
-	{
+{
 	auto UpdateWithFloatCurve = [&Sha](const FRichCurve& Curve)
-		{
+	{
 		UpdateWithData(Sha, Curve.DefaultValue);
 		UpdateSHAWithArray(Sha, Curve.GetConstRefOfKeys());
 		UpdateWithData(Sha, Curve.PreInfinityExtrap);
@@ -1713,7 +1735,7 @@ void UpdateSHAWithCurves(FSHA1& Sha, const FRawCurveTracks& InRawCurveData)
 	};
 
 	for (const FFloatCurve& Curve : InRawCurveData.FloatCurves)
-			{
+	{
 		UpdateWithData(Sha, Curve.GetName());
 		UpdateWithFloatCurve(Curve.FloatCurve);
 	}
@@ -1733,8 +1755,8 @@ void UpdateSHAWithCurves(FSHA1& Sha, const FRawCurveTracks& InRawCurveData)
 		UpdateWithComponent(Curve.TranslationCurve);
 		UpdateWithComponent(Curve.RotationCurve);
 		UpdateWithComponent(Curve.ScaleCurve);
-		}
 	}
+}
 
 FGuid GenerateGuidFromRawAnimData(const TArray<FRawAnimSequenceTrack>& RawAnimationData, const FRawCurveTracks& RawCurveData)
 	{
diff --git a/Engine/Source/Runtime/Engine/Public/Animation/AnimCompressionTypes.h b/Engine/Source/Runtime/Engine/Public/Animation/AnimCompressionTypes.h
--- a/Engine/Source/Runtime/Engine/Public/Animation/AnimCompressionTypes.h
+++ b/Engine/Source/Runtime/Engine/Public/Animation/AnimCompressionTypes.h
@@ -849,6 +849,8 @@ public:
 	ENGINE_API void ClearCompressedCurveData();
 
 	ENGINE_API bool IsValid(const UAnimSequence* AnimSequence, bool bLogInformation=false) const;
+	ENGINE_API bool IsBoneDataValid(const UAnimSequence* AnimSequence, bool bLogInformation=false) const;
+	ENGINE_API bool IsCurveDataValid(const UAnimSequence* AnimSequence, bool bLogInformation=false) const;
 };
 
 struct FRootMotionReset
