Forward animation extraction context to blendspace sampling

diff --git a/Engine/Source/Runtime/Engine/Private/Animation/BlendSpace.cpp b/Engine/Source/Runtime/Engine/Private/Animation/BlendSpace.cpp
--- a/Engine/Source/Runtime/Engine/Private/Animation/BlendSpace.cpp
+++ b/Engine/Source/Runtime/Engine/Private/Animation/BlendSpace.cpp
@@ -1114,6 +1114,7 @@ void UBlendSpace::GetAnimationPose_Internal(TArray<FBlendSampleData>& BlendSampl
 		ChildrenCurves[ChildrenIdx].InitFrom(OutCurve);
 	}
 
+	FAnimExtractContext ChildExtractionContext = ExtractionContext;
 	// get all child atoms we interested in
 	for (int32 I = 0; I < NumPoses; ++I)
 	{
@@ -1145,7 +1146,9 @@ void UBlendSpace::GetAnimationPose_Internal(TArray<FBlendSampleData>& BlendSampl
 
 					FAnimationPoseData ChildAnimationPoseData = { Pose, ChildrenCurves[I], ChildrenAttributes[I] };
 					// first one always fills up the source one
-					Sample.Animation->GetAnimationPose(ChildAnimationPoseData, FAnimExtractContext(static_cast<double>(Time), ExtractionContext.bExtractRootMotion, BlendSampleDataCache[I].DeltaTimeRecord, ExtractionContext.bLooping));
+					ChildExtractionContext.CurrentTime = static_cast<double>(Time);
+					ChildExtractionContext.DeltaTimeRecord = BlendSampleDataCache[I].DeltaTimeRecord;
+					Sample.Animation->GetAnimationPose(ChildAnimationPoseData, ChildExtractionContext);
 				}
 				else
 				{
