[shaders] Remove the massive "debug description" strings from shader stats and the "debug" csv which appears to be the only point at which they are used. These are a huge memory (and network, in the case of mp cook) overhead and are of little use as far as anyone who works on shaders currently can tell. This allows re-enabling the multiprocess cook stat aggregation since removing this data reduces message sizes by multiple orders of magnitude (so no longer any risk of exceeding message size limits).

diff --git a/Engine/Source/Editor/UnrealEd/Private/ShaderStatsCollector.cpp b/Engine/Source/Editor/UnrealEd/Private/ShaderStatsCollector.cpp
--- a/Engine/Source/Editor/UnrealEd/Private/ShaderStatsCollector.cpp
+++ b/Engine/Source/Editor/UnrealEd/Private/ShaderStatsCollector.cpp
@@ -1,9 +1,10 @@
 // Copyright Epic Games, Inc. All Rights Reserved.
 
 #include "ShaderStatsCollector.h"
-#include "ShaderStats.h"
 #include "AnalyticsEventAttribute.h"
+#include "CookOnTheSide/CookLog.h"
 #include "Serialization/CompactBinaryWriter.h"
+#include "ShaderStats.h"
 
 static TAutoConsoleVariable<int> CVarShaderCompilerStatsPrintoutInterval(
 	TEXT("r.ShaderCompiler.StatsPrintoutInterval"),
@@ -166,11 +167,11 @@ void FShaderStatsAggregator::ClientTick(UE::Cook::FMPCollectorClientTickContext&
 		Writer.BeginObject();
 
 		FShaderCompilerStats LocalStats;
-		// Temporarily disable shader stats aggregation for multi-process cooks due to message sizes exceeding limits.
-		// GShaderCompilingManager->GetLocalStats(LocalStats);
+		GShaderCompilingManager->GetLocalStats(LocalStats);
 		LocalStats.WriteToCompactBinary(Writer);
 
 		Writer.EndObject();
+		UE_LOG(LogCook, Verbose, TEXT("Synching shader compile stats to cook director; message size %d bytes"), Writer.GetSaveSize());
 
 		Context.AddMessage(Writer.Save().AsObject());
 	}
diff --git a/Engine/Source/Runtime/Engine/Private/ShaderCompiler/ShaderCompilerStats.cpp b/Engine/Source/Runtime/Engine/Private/ShaderCompiler/ShaderCompilerStats.cpp
--- a/Engine/Source/Runtime/Engine/Private/ShaderCompiler/ShaderCompilerStats.cpp
+++ b/Engine/Source/Runtime/Engine/Private/ShaderCompiler/ShaderCompilerStats.cpp
@@ -139,57 +139,6 @@ void FShaderCompilerStats::WriteStats(FOutputDevice* Ar)
 			}
 		}
 	}
-	{
-
-		FString FileName = FString::Printf(TEXT("%s/MaterialStatsDebug/StatsDebug-%s.csv"), *FPaths::ProjectSavedDir(), *FDateTime::Now().ToString());
-		auto DebugWriter = IFileManager::Get().CreateFileWriter(*FileName);
-		FDiagnosticTableWriterCSV StatWriter(DebugWriter);
-		const TSparseArray<ShaderCompilerStats>& PlatformStats = GetShaderCompilerStats();
-		StatWriter.AddColumn(TEXT("Name"));
-		StatWriter.AddColumn(TEXT("Platform"));
-		StatWriter.AddColumn(TEXT("Compiles"));
-		StatWriter.AddColumn(TEXT("CompilesDouble"));
-		StatWriter.AddColumn(TEXT("Uses"));
-		StatWriter.AddColumn(TEXT("UsesDouble"));
-		StatWriter.AddColumn(TEXT("PermutationString"));
-		StatWriter.CycleRow();
-
-
-		for (int32 Platform = 0; Platform < PlatformStats.GetMaxIndex(); ++Platform)
-		{
-			if (PlatformStats.IsValidIndex(Platform))
-			{
-				const ShaderCompilerStats& Stats = PlatformStats[Platform];
-				for (const auto& Pair : Stats)
-				{
-					const FString& Path = Pair.Key;
-					const FShaderStats& SingleStats = Pair.Value;
-					for (const FShaderCompilerSinglePermutationStat& Stat : SingleStats.PermutationCompilations)
-					{
-						StatWriter.AddColumn(*Path);
-						StatWriter.AddColumn(TEXT("%u"), Platform);
-						StatWriter.AddColumn(TEXT("%u"), Stat.Compiled);
-						StatWriter.AddColumn(TEXT("%u"), Stat.CompiledDouble);
-						StatWriter.AddColumn(TEXT("%u"), Stat.Cooked);
-						StatWriter.AddColumn(TEXT("%u"), Stat.CookedDouble);
-						StatWriter.AddColumn(TEXT("%s"), *Stat.PermutationString);
-						StatWriter.CycleRow();
-					}
-				}
-
-			}
-		}
-
-		FString FullFileName = FPaths::ConvertRelativePathToFull(FileName);
-		if (Ar)
-		{
-			Ar->Logf(DebugText, *FullFileName);
-		}
-		else
-		{
-			UE_LOG(LogShaderCompilers, Log, DebugText, *FullFileName);
-		}
-	}
 #endif // ALLOW_DEBUG_FILES
 }
 
@@ -800,7 +749,7 @@ void FShaderCompilerStats::WriteToCompactBinary(FCbWriter& Writer)
 				Writer << "CompiledDouble" << Stat.CompiledDouble;
 				Writer << "Cooked" << Stat.Cooked;
 				Writer << "CookedDouble" << Stat.CookedDouble;
-				Writer << "PermutationString" << Stat.PermutationString;
+				Writer << "PermutationHash" << Stat.PermutationHash;
 				Writer.EndObject();
 			}
 			Writer.EndArray();
@@ -887,7 +836,7 @@ void FShaderCompilerStats::ReadFromCompactBinary(FCbObjectView& Reader)
 				FCbObjectView PermutationObject = CompilationField.AsObjectView();
 				uint32 Index = ShaderStats.PermutationCompilations.Emplace
 				(
-					FString(PermutationObject["PermutationString"].AsString()),
+					PermutationObject["PermutationHash"].AsUInt64(),
 					PermutationObject["Compiled"].AsUInt32(),
 					PermutationObject["Cooked"].AsUInt32()
 				);
@@ -1141,9 +1090,10 @@ void FShaderCompilerStats::RegisterCookedShaders(uint32 NumCooked, float Compile
 	FShaderStats& Stats = CompileStats[Platform].FindOrAdd(MaterialPath);
 	Stats.CompileTime += CompileTime;
 	bool bFound = false;
+	uint64 PermutationHash = FShaderCompilerSinglePermutationStat::GetPermutationHash(PermutationString);
 	for (FShaderCompilerSinglePermutationStat& Stat : Stats.PermutationCompilations)
 	{
-		if (PermutationString == Stat.PermutationString)
+		if (PermutationHash == Stat.PermutationHash)
 		{
 			bFound = true;
 			if (Stat.Cooked != 0)
@@ -1164,7 +1114,7 @@ void FShaderCompilerStats::RegisterCookedShaders(uint32 NumCooked, float Compile
 	}
 	if (!bFound)
 	{
-		Stats.PermutationCompilations.Emplace(PermutationString, 0, NumCooked);
+		Stats.PermutationCompilations.Emplace(PermutationHash, 0, NumCooked);
 	}
 }
 
@@ -1179,9 +1129,10 @@ void FShaderCompilerStats::RegisterCompiledShaders(uint32 NumCompiled, EShaderPl
 	FShaderStats& Stats = CompileStats[Platform].FindOrAdd(MaterialPath);
 
 	bool bFound = false;
+	uint64 PermutationHash = FShaderCompilerSinglePermutationStat::GetPermutationHash(PermutationString);
 	for (FShaderCompilerSinglePermutationStat& Stat : Stats.PermutationCompilations)
 	{
-		if (PermutationString == Stat.PermutationString)
+		if (PermutationHash == Stat.PermutationHash)
 		{
 			bFound = true;
 			if (Stat.Compiled != 0)
@@ -1204,7 +1155,7 @@ void FShaderCompilerStats::RegisterCompiledShaders(uint32 NumCompiled, EShaderPl
 
 	if (!bFound)
 	{
-		Stats.PermutationCompilations.Emplace(PermutationString, NumCompiled, 0);
+		Stats.PermutationCompilations.Emplace(PermutationHash, NumCompiled, 0);
 	}
 }
 
diff --git a/Engine/Source/Runtime/Engine/Public/ShaderCompiler.h b/Engine/Source/Runtime/Engine/Public/ShaderCompiler.h
--- a/Engine/Source/Runtime/Engine/Public/ShaderCompiler.h
+++ b/Engine/Source/Runtime/Engine/Public/ShaderCompiler.h
@@ -26,6 +26,7 @@
 #include "Containers/List.h"
 #include "Containers/Deque.h"
 #include "Hash/Blake3.h"
+#include "Hash/CityHash.h"
 #include "SceneTypes.h"
 #include "UObject/StrongObjectPtr.h"
 
@@ -652,15 +653,26 @@ struct FShaderCompilerSinglePermutationStat
 	GENERATED_BODY()
 
 	FShaderCompilerSinglePermutationStat()
-		: PermutationString()
+		: PermutationHash(0)
 		, Compiled(0)
 		, Cooked(0)
 		, CompiledDouble(0)
 		, CookedDouble(0)
 	{}
 
+	UE_DEPRECATED(5.6, "Use constructor accepting a u64 permutationstring hash instead of a string")
 	FShaderCompilerSinglePermutationStat(FString PermutationString, uint32 Compiled, uint32 Cooked)
-		: PermutationString(PermutationString)
+		: Compiled(Compiled)
+		, Cooked(Cooked)
+		, CompiledDouble(0)
+		, CookedDouble(0)
+
+	{
+		PermutationHash = GetPermutationHash(PermutationString);
+	}
+
+	FShaderCompilerSinglePermutationStat(uint64 PermutationHash, uint32 Compiled, uint32 Cooked)
+		: PermutationHash(PermutationHash)
 		, Compiled(Compiled)
 		, Cooked(Cooked)
 		, CompiledDouble(0)
@@ -668,8 +680,26 @@ struct FShaderCompilerSinglePermutationStat
 
 	{}
 
+	PRAGMA_DISABLE_DEPRECATION_WARNINGS
+	// Explicitly-defaulted ctors & assignment operators are needed temporarily due to deprecation
+	// of PermutationString field. Can be removed once the deprecation window for said field ends.
+	FShaderCompilerSinglePermutationStat(FShaderCompilerSinglePermutationStat&&) = default;
+	FShaderCompilerSinglePermutationStat(const FShaderCompilerSinglePermutationStat&) = default;
+	FShaderCompilerSinglePermutationStat& operator=(FShaderCompilerSinglePermutationStat&&) = default;
+	FShaderCompilerSinglePermutationStat& operator=(const FShaderCompilerSinglePermutationStat&) = default;
+	PRAGMA_ENABLE_DEPRECATION_WARNINGS
+
+	static uint64 GetPermutationHash(const FString& PermutationString)
+	{
+		return CityHash64(reinterpret_cast<const char*>(*PermutationString), PermutationString.Len() * sizeof(TCHAR));
+	}
+
+	UPROPERTY()
+	uint64 PermutationHash;
+
+	UE_DEPRECATED(5.6, "PermutationString is no longer stored due to memory overhead; use PermutationHash to uniquely identify permutation stats")
 	UPROPERTY()
-	FString PermutationString;
+	uint32 PermutationString;
 
 	UPROPERTY()
 	uint32 Compiled;
