Deformer Graph reuse existing vertex buffers in more cases. This reduces cost of churn of vertex factory uniform buffer which can cost as much as 20% of the compute graph execution CPU time.

diff --git a/Engine/Source/Runtime/Engine/Private/SkeletalMeshDeformerHelpers.cpp b/Engine/Source/Runtime/Engine/Private/SkeletalMeshDeformerHelpers.cpp
--- a/Engine/Source/Runtime/Engine/Private/SkeletalMeshDeformerHelpers.cpp
+++ b/Engine/Source/Runtime/Engine/Private/SkeletalMeshDeformerHelpers.cpp
@@ -149,6 +149,7 @@ FRDGBuffer* FSkeletalMeshDeformerHelpers::AllocateVertexFactoryPositionBuffer(FR
 	FSkeletalMeshRenderData const& SkeletalMeshRenderData = InMeshObject->GetSkeletalMeshRenderData();
 	FSkeletalMeshLODRenderData const* LodRenderData = SkeletalMeshRenderData.GetPendingFirstLOD(InLodIndex);
 	const int32 NumVertices = LodRenderData->GetNumVertices();
+	const FRDGBufferDesc BufferDesc = FRDGBufferDesc::CreateBufferDesc(PosBufferBytesPerElement, NumVertices * PosBufferElementMultiplier);
 
 	FSkeletalMeshObjectGPUSkin* MeshObjectGPU = static_cast<FSkeletalMeshObjectGPUSkin*>(InMeshObject);
 	FGPUBaseSkinVertexFactory const* BaseVertexFactory = MeshObjectGPU->GetBaseSkinVertexFactory(InLodIndex, 0);
@@ -157,20 +158,43 @@ FRDGBuffer* FSkeletalMeshDeformerHelpers::AllocateVertexFactoryPositionBuffer(FR
 	FMeshDeformerGeometry& DeformerGeometry = MeshObjectGPU->GetDeformerGeometry(InLodIndex);
 
 	FRDGBuffer* PositionBuffer = nullptr;
-	if (DeformerGeometry.Position && Frame == DeformerGeometry.PositionUpdatedFrame)
+	if (DeformerGeometry.Position.IsValid() && Frame == DeformerGeometry.PositionUpdatedFrame)
 	{
 		PositionBuffer = GraphBuilder.RegisterExternalBuffer(DeformerGeometry.Position);
 		GraphBuilder.UseInternalAccessMode(PositionBuffer);
 	}
 	else
 	{
-		DeformerGeometry.PrevPosition = DeformerGeometry.Position;
-		DeformerGeometry.PrevPositionSRV = DeformerGeometry.PositionSRV;
+		const bool bIsMatchingDescription = DeformerGeometry.Position.IsValid() && DeformerGeometry.Position->Desc == BufferDesc;
+		const bool bIsMatchingPreviousDescription = DeformerGeometry.PrevPosition.IsValid() && DeformerGeometry.PrevPosition->Desc == BufferDesc;
 
-		PositionBuffer = GraphBuilder.CreateBuffer(FRDGBufferDesc::CreateBufferDesc(PosBufferBytesPerElement, NumVertices * PosBufferElementMultiplier), InBufferName, ERDGBufferFlags::None);
-		DeformerGeometry.Position = GraphBuilder.ConvertToExternalBuffer(PositionBuffer);
-
-		DeformerGeometry.PositionSRV = DeformerGeometry.Position->GetOrCreateSRV(GraphBuilder.RHICmdList, FRHIBufferSRVCreateInfo(PF_R32_FLOAT));
+		if (bIsMatchingDescription && bIsMatchingPreviousDescription)
+		{
+			// Flip position buffers and return the current one.
+			DeformerGeometry.PrevPosition.Swap(DeformerGeometry.Position);
+			DeformerGeometry.PrevPositionSRV.Swap(DeformerGeometry.PositionSRV);
+			
+			PositionBuffer = GraphBuilder.RegisterExternalBuffer(DeformerGeometry.Position);
+			GraphBuilder.UseInternalAccessMode(PositionBuffer);
+		}
+		else
+		{
+			if (bIsMatchingDescription)
+			{
+				DeformerGeometry.PrevPosition = MoveTemp(DeformerGeometry.Position);
+				DeformerGeometry.PrevPositionSRV = MoveTemp(DeformerGeometry.PositionSRV);
+			}
+			else
+			{
+				DeformerGeometry.PrevPosition = nullptr;
+				DeformerGeometry.PrevPositionSRV = nullptr;
+			}
+
+			PositionBuffer = GraphBuilder.CreateBuffer(FRDGBufferDesc::CreateBufferDesc(PosBufferBytesPerElement, NumVertices * PosBufferElementMultiplier), InBufferName, ERDGBufferFlags::None);
+			DeformerGeometry.Position = GraphBuilder.ConvertToExternalBuffer(PositionBuffer);
+
+			DeformerGeometry.PositionSRV = DeformerGeometry.Position->GetOrCreateSRV(GraphBuilder.RHICmdList, FRHIBufferSRVCreateInfo(PF_R32_FLOAT));
+		}
 
 		DeformerGeometry.PositionUpdatedFrame = Frame;
 		
@@ -217,6 +241,7 @@ FRDGBuffer* FSkeletalMeshDeformerHelpers::AllocateVertexFactoryTangentBuffer(FRD
 	FSkeletalMeshRenderData const& SkeletalMeshRenderData = InMeshObject->GetSkeletalMeshRenderData();
 	FSkeletalMeshLODRenderData const* LodRenderData = SkeletalMeshRenderData.GetPendingFirstLOD(InLodIndex);
 	const int32 NumVertices = LodRenderData->GetNumVertices();
+	const FRDGBufferDesc BufferDesc = FRDGBufferDesc::CreateBufferDesc(TangentBufferBytesPerElement, NumVertices * TangentBufferElementMultiplier);
 
 	FSkeletalMeshObjectGPUSkin* MeshObjectGPU = static_cast<FSkeletalMeshObjectGPUSkin*>(InMeshObject);
 	FGPUBaseSkinVertexFactory const* BaseVertexFactory = MeshObjectGPU->GetBaseSkinVertexFactory(InLodIndex, 0);
@@ -225,20 +250,21 @@ FRDGBuffer* FSkeletalMeshDeformerHelpers::AllocateVertexFactoryTangentBuffer(FRD
 	FMeshDeformerGeometry& DeformerGeometry = MeshObjectGPU->GetDeformerGeometry(InLodIndex);
 
 	FRDGBuffer* TangentBuffer = nullptr;
-	if (DeformerGeometry.Tangent && Frame == DeformerGeometry.TangentUpdatedFrame)
+	if (DeformerGeometry.Tangent.IsValid() && DeformerGeometry.Tangent->Desc == BufferDesc)
 	{
 		TangentBuffer = GraphBuilder.RegisterExternalBuffer(DeformerGeometry.Tangent);
 		GraphBuilder.UseInternalAccessMode(TangentBuffer);
 	}
 	else
 	{
-		TangentBuffer = GraphBuilder.CreateBuffer(FRDGBufferDesc::CreateBufferDesc(TangentBufferBytesPerElement, NumVertices * TangentBufferElementMultiplier), InBufferName, ERDGBufferFlags::None);
+		TangentBuffer = GraphBuilder.CreateBuffer(BufferDesc, InBufferName, ERDGBufferFlags::None);
 		DeformerGeometry.Tangent = GraphBuilder.ConvertToExternalBuffer(TangentBuffer);
 		const EPixelFormat TangentsFormat = IsOpenGLPlatform(GMaxRHIShaderPlatform) ? PF_R16G16B16A16_SINT : PF_R16G16B16A16_SNORM;
 		DeformerGeometry.TangentSRV = DeformerGeometry.Tangent->GetOrCreateSRV(GraphBuilder.RHICmdList, FRHIBufferSRVCreateInfo(TangentsFormat));
-		DeformerGeometry.TangentUpdatedFrame = Frame;
 	}
 
+	DeformerGeometry.TangentUpdatedFrame = Frame;
+
 	ExternalAccesQueue.AddUnique(TangentBuffer, ERHIAccess::VertexOrIndexBuffer | ERHIAccess::SRVMask);
 
 	return TangentBuffer;
@@ -264,6 +290,7 @@ FRDGBuffer* FSkeletalMeshDeformerHelpers::AllocateVertexFactoryColorBuffer(FRDGB
 	FSkeletalMeshRenderData const& SkeletalMeshRenderData = InMeshObject->GetSkeletalMeshRenderData();
 	FSkeletalMeshLODRenderData const* LodRenderData = SkeletalMeshRenderData.GetPendingFirstLOD(InLodIndex);
 	const int32 NumVertices = LodRenderData->GetNumVertices();
+	const FRDGBufferDesc BufferDesc = FRDGBufferDesc::CreateBufferDesc(ColorBufferBytesPerElement, NumVertices);
 
 	FSkeletalMeshObjectGPUSkin* MeshObjectGPU = static_cast<FSkeletalMeshObjectGPUSkin*>(InMeshObject);
 	FGPUBaseSkinVertexFactory const* BaseVertexFactory = MeshObjectGPU->GetBaseSkinVertexFactory(InLodIndex, 0);
@@ -272,19 +299,20 @@ FRDGBuffer* FSkeletalMeshDeformerHelpers::AllocateVertexFactoryColorBuffer(FRDGB
 	FMeshDeformerGeometry& DeformerGeometry = MeshObjectGPU->GetDeformerGeometry(InLodIndex);
 
 	FRDGBuffer* ColorBuffer = nullptr;
-	if (DeformerGeometry.Color && Frame == DeformerGeometry.ColorUpdatedFrame)
+	if (DeformerGeometry.Color.IsValid() && DeformerGeometry.Color->Desc == BufferDesc)
 	{
 		ColorBuffer = GraphBuilder.RegisterExternalBuffer(DeformerGeometry.Color);
 		GraphBuilder.UseInternalAccessMode(ColorBuffer);
 	}
 	else
 	{
-		ColorBuffer = GraphBuilder.CreateBuffer(FRDGBufferDesc::CreateBufferDesc(ColorBufferBytesPerElement, NumVertices), InBufferName, ERDGBufferFlags::None);
+		ColorBuffer = GraphBuilder.CreateBuffer(BufferDesc, InBufferName, ERDGBufferFlags::None);
 		DeformerGeometry.Color = GraphBuilder.ConvertToExternalBuffer(ColorBuffer);
 		DeformerGeometry.ColorSRV = DeformerGeometry.Color->GetOrCreateSRV(GraphBuilder.RHICmdList, FRHIBufferSRVCreateInfo(PF_R8G8B8A8));
-		DeformerGeometry.ColorUpdatedFrame = Frame;
 	}
 
+	DeformerGeometry.ColorUpdatedFrame = Frame;
+
 	ExternalAccesQueue.AddUnique(ColorBuffer, ERHIAccess::VertexOrIndexBuffer | ERHIAccess::SRVMask);
 
 	return ColorBuffer;
