UE-184084 - AActor::CheckActorComponents has misleading warning

This warning is _really_ old, and not very useful to begin with. It'll always trip because the component template objects have the GEN_VARIABLE suffix, which instances don't have. This warning was never updated to reflect that.

It tends to pop up when users enable all warnings (ie: -logCmds="global warning").

diff --git a/Engine/Source/Runtime/Engine/Private/Actor.cpp b/Engine/Source/Runtime/Engine/Private/Actor.cpp
--- a/Engine/Source/Runtime/Engine/Private/Actor.cpp
+++ b/Engine/Source/Runtime/Engine/Private/Actor.cpp
@@ -369,15 +369,6 @@ bool AActor::CheckActorComponents() const
 			UE_LOG(LogCheckComponents, Error, TEXT("SerializedComponent is a template but I am not. Me = %s, Component = %s"), *this->GetFullName(), *Inner->GetFullName());
 			bResult = false;
 		}
-		UObject* Archetype = Inner->GetArchetype();
-		if (Archetype != Inner->GetClass()->GetDefaultObject())
-		{
-			if (Archetype != GetClass()->GetDefaultSubobjectByName(Inner->GetFName()))
-			{
-				UE_LOG(LogCheckComponents, Error, TEXT("SerializedComponent archetype is not the CDO nor a default subobject of my class. Me = %s, Component = %s, Archetype = %s"), *this->GetFullName(), *Inner->GetFullName(), *Archetype->GetFullName());
-				bResult = false;
-			}
-		}
 	}
 	return bResult;
 }
