[HWRT] Fix ISMs using incorrect VertexFactory.
- FInstancedStaticMeshSceneProxy::GetDynamicRayTracingInstances(...) was calling FStaticMeshSceneProxy::GetMeshElement(...) so didn't get the correct vertex factory.

diff --git a/Engine/Source/Runtime/Engine/Private/InstancedStaticMesh.cpp b/Engine/Source/Runtime/Engine/Private/InstancedStaticMesh.cpp
--- a/Engine/Source/Runtime/Engine/Private/InstancedStaticMesh.cpp
+++ b/Engine/Source/Runtime/Engine/Private/InstancedStaticMesh.cpp
@@ -1700,7 +1700,7 @@ void FInstancedStaticMeshSceneProxy::GetDynamicRayTracingInstances(FRayTracingIn
 			{
 				FMeshBatch MeshBatch;
 
-				if (!FStaticMeshSceneProxy::GetMeshElement(LODIndex, 0, SectionIdx, 0, false, false, MeshBatch))
+				if (!GetMeshElement(LODIndex, 0, SectionIdx, 0, false, false, MeshBatch))
 				{
 					MeshBatch.MaterialRenderProxy = UMaterial::GetDefaultMaterial(MD_Surface)->GetRenderProxy();
 					MeshBatch.VertexFactory = &RenderData->LODVertexFactories[LODIndex].VertexFactory;
@@ -1945,7 +1945,7 @@ void FInstancedStaticMeshSceneProxy::GetDynamicRayTracingInstances(FRayTracingIn
 			//#dxr_todo: so far we use the parent static mesh path to get material data
 			FMeshBatch MeshBatch;
 
-			bool bResult = FStaticMeshSceneProxy::GetMeshElement(LODIndex, 0, SectionIdx, 0, false, false, MeshBatch);
+			bool bResult = GetMeshElement(LODIndex, 0, SectionIdx, 0, false, false, MeshBatch);
 			if (!bResult)
 			{
 				// Hidden material
