Defer deformer graph modification of vertex factory to RDG graph execution.
This moves work from render thread to reduce overall wall time.
Also small fix to correctly set last position frame index for all paths.

diff --git a/Engine/Plugins/Animation/DeformerGraph/Source/OptimusCore/Private/DataInterfaces/OptimusDataInterfaceSkinnedMeshRead.cpp b/Engine/Plugins/Animation/DeformerGraph/Source/OptimusCore/Private/DataInterfaces/OptimusDataInterfaceSkinnedMeshRead.cpp
--- a/Engine/Plugins/Animation/DeformerGraph/Source/OptimusCore/Private/DataInterfaces/OptimusDataInterfaceSkinnedMeshRead.cpp
+++ b/Engine/Plugins/Animation/DeformerGraph/Source/OptimusCore/Private/DataInterfaces/OptimusDataInterfaceSkinnedMeshRead.cpp
@@ -255,7 +255,7 @@ void FOptimusSkinnedMeshReadDataProviderProxy::AllocateResources(FRDGBuilder& Gr
 		ColorBufferUAV = GraphBuilder.CreateUAV(GraphBuilder.RegisterExternalBuffer(GWhiteVertexBufferWithRDG->Buffer), PF_A32B32G32R32F);
 	}
 
-	FSkeletalMeshDeformerHelpers::UpdateVertexFactoryBufferOverrides(GraphBuilder.RHICmdList, SkeletalMeshObject, LodIndex, bInvalidatePreviousPosition);
+	FSkeletalMeshDeformerHelpers::UpdateVertexFactoryBufferOverrides(GraphBuilder, SkeletalMeshObject, LodIndex, bInvalidatePreviousPosition);
 }
 
 void FOptimusSkinnedMeshReadDataProviderProxy::GatherPermutations(FPermutationData& InOutPermutationData) const
diff --git a/Engine/Plugins/Animation/DeformerGraph/Source/OptimusCore/Private/DataInterfaces/OptimusDataInterfaceSkinnedMeshWrite.cpp b/Engine/Plugins/Animation/DeformerGraph/Source/OptimusCore/Private/DataInterfaces/OptimusDataInterfaceSkinnedMeshWrite.cpp
--- a/Engine/Plugins/Animation/DeformerGraph/Source/OptimusCore/Private/DataInterfaces/OptimusDataInterfaceSkinnedMeshWrite.cpp
+++ b/Engine/Plugins/Animation/DeformerGraph/Source/OptimusCore/Private/DataInterfaces/OptimusDataInterfaceSkinnedMeshWrite.cpp
@@ -496,7 +496,7 @@ void FOptimusSkinnedMeshWriteDataProviderProxy::AllocateResources(FRDGBuilder& G
 		ColorBufferUAV = GraphBuilder.CreateUAV(GraphBuilder.RegisterExternalBuffer(GWhiteVertexBufferWithRDG->Buffer), PF_A32B32G32R32F);
 	}
 
-	FSkeletalMeshDeformerHelpers::UpdateVertexFactoryBufferOverrides(GraphBuilder.RHICmdList, SkeletalMeshObject, LodIndex, bInvalidatePreviousPosition);
+	FSkeletalMeshDeformerHelpers::UpdateVertexFactoryBufferOverrides(GraphBuilder, SkeletalMeshObject, LodIndex, bInvalidatePreviousPosition);
 }
 
 void FOptimusSkinnedMeshWriteDataProviderProxy::GatherDispatchData(FDispatchData const& InDispatchData)
diff --git a/Engine/Source/Runtime/Engine/Private/SkeletalMeshDeformerHelpers.cpp b/Engine/Source/Runtime/Engine/Private/SkeletalMeshDeformerHelpers.cpp
--- a/Engine/Source/Runtime/Engine/Private/SkeletalMeshDeformerHelpers.cpp
+++ b/Engine/Source/Runtime/Engine/Private/SkeletalMeshDeformerHelpers.cpp
@@ -196,8 +196,6 @@ FRDGBuffer* FSkeletalMeshDeformerHelpers::AllocateVertexFactoryPositionBuffer(FR
 			DeformerGeometry.PositionSRV = DeformerGeometry.Position->GetOrCreateSRV(GraphBuilder.RHICmdList, FRHIBufferSRVCreateInfo(PF_R32_FLOAT));
 		}
 
-		DeformerGeometry.PositionUpdatedFrame = Frame;
-		
 #if RHI_RAYTRACING
 		// Update ray tracing geometry whenever we recreate the position buffer.
 		FSkeletalMeshRenderData& SkelMeshRenderData = MeshObjectGPU->GetSkeletalMeshRenderData();
@@ -216,6 +214,8 @@ FRDGBuffer* FSkeletalMeshDeformerHelpers::AllocateVertexFactoryPositionBuffer(FR
 #endif // RHI_RAYTRACING
 	}
 
+	DeformerGeometry.PositionUpdatedFrame = Frame;
+
 	ExternalAccesQueue.AddUnique(PositionBuffer, ERHIAccess::VertexOrIndexBuffer | ERHIAccess::SRVMask);
 
 	return PositionBuffer;
@@ -337,7 +337,7 @@ void FSkeletalMeshDeformerHelpers::UpdateVertexFactoryBufferOverrides(FRHIComman
 
 	FSkeletalMeshObjectGPUSkin* MeshObjectGPU = static_cast<FSkeletalMeshObjectGPUSkin*>(InMeshObject);
 	FMeshDeformerGeometry& DeformerGeometry = MeshObjectGPU->GetDeformerGeometry(InLodIndex);
-	
+
 	FGPUSkinPassthroughVertexFactory::FAddVertexAttributeDesc Desc;
 	Desc.FrameNumber = DeformerGeometry.PositionUpdatedFrame;
 
@@ -369,7 +369,7 @@ void FSkeletalMeshDeformerHelpers::UpdateVertexFactoryBufferOverrides(FRHIComman
 	}
 
 	const FSkeletalMeshObjectGPUSkin::FSkeletalMeshObjectLOD& LOD = MeshObjectGPU->LODs[InLodIndex];
-	const int32 NumSections = InMeshObject->GetRenderSections(InLodIndex).Num();
+	const int32 NumSections = MeshObjectGPU->GetRenderSections(InLodIndex).Num();
 	for (int32 SectionIndex = 0; SectionIndex < NumSections; ++SectionIndex)
 	{
 		FGPUBaseSkinVertexFactory const* BaseVertexFactory = MeshObjectGPU->GetBaseSkinVertexFactory(InLodIndex, SectionIndex);
@@ -383,6 +383,14 @@ void FSkeletalMeshDeformerHelpers::UpdateVertexFactoryBufferOverrides(FRHIComman
 	}
 }
 
+void FSkeletalMeshDeformerHelpers::UpdateVertexFactoryBufferOverrides(FRDGBuilder& GraphBuilder, FSkeletalMeshObject* InMeshObject, int32 InLodIndex, bool bInvalidatePreviousPosition)
+{
+	GraphBuilder.AddPass({}, ERDGPassFlags::None, [InMeshObject, InLodIndex, bInvalidatePreviousPosition](FRHICommandList& RHICmdList)
+	{
+		UpdateVertexFactoryBufferOverrides(RHICmdList, InMeshObject, InLodIndex, bInvalidatePreviousPosition);
+	});
+}
+
 void FSkeletalMeshDeformerHelpers::UpdateVertexFactoryBufferOverrides(FRHICommandListBase& RHICmdList, FSkeletalMeshObject* InMeshObject, int32 InLodIndex)
 {
 PRAGMA_DISABLE_DEPRECATION_WARNINGS
diff --git a/Engine/Source/Runtime/Engine/Public/SkeletalMeshDeformerHelpers.h b/Engine/Source/Runtime/Engine/Public/SkeletalMeshDeformerHelpers.h
--- a/Engine/Source/Runtime/Engine/Public/SkeletalMeshDeformerHelpers.h
+++ b/Engine/Source/Runtime/Engine/Public/SkeletalMeshDeformerHelpers.h
@@ -93,6 +93,7 @@ public:
 	 * Usually call this after all AllocateVertexFactory*() functions for a frame.
 	 */
 	ENGINE_API static void UpdateVertexFactoryBufferOverrides(FRHICommandListBase& RHICmdList, FSkeletalMeshObject* InMeshObject, int32 InLodIndex, bool bInvalidatePreviousPosition);
+	ENGINE_API static void UpdateVertexFactoryBufferOverrides(FRDGBuilder& GraphBuilder, FSkeletalMeshObject* InMeshObject, int32 InLodIndex, bool bInvalidatePreviousPosition);
 
 	/** 
 	 * Release all of the the buffers that have been allocated through the AllocateVertexFactory*() functions.
