Fix reachable garbage object warning from some legacy ASpatialHashRuntimeGridInfo actors
* Automatically fix ASpatialHashRuntimeGridInfo actors that were saved prior to the constructor using ObjectInitializer.DoNotCreateDefaultSubobject(TEXT("Sprite"))
* Those sprite components are MarkForGarbage in their UActorComponent::PostInitProperties and so we need to clear the component pointers from the owner actor

diff --git a/Engine/Source/Runtime/Core/Public/UObject/FortniteMainBranchObjectVersion.h b/Engine/Source/Runtime/Core/Public/UObject/FortniteMainBranchObjectVersion.h
--- a/Engine/Source/Runtime/Core/Public/UObject/FortniteMainBranchObjectVersion.h
+++ b/Engine/Source/Runtime/Core/Public/UObject/FortniteMainBranchObjectVersion.h
@@ -559,6 +559,9 @@ struct FFortniteMainBranchObjectVersion
 		// Added Dynamic Mesh Sculpt Layer serialization
 		DynamicMeshSerializeSculptLayers,
 
+		// Fix reachable garbage object warnings from some legacy ASpatialHashRuntimeGridInfo actors
+		SpatialHashRuntimeGridInfoSpriteFixup,
+
 		// -----<new versions can be added above this line>-------------------------------------------------
 		VersionPlusOne,
 		LatestVersion = VersionPlusOne - 1
diff --git a/Engine/Source/Runtime/Engine/Private/WorldPartition/WorldPartitionRuntimeSpatialHash.cpp b/Engine/Source/Runtime/Engine/Private/WorldPartition/WorldPartitionRuntimeSpatialHash.cpp
--- a/Engine/Source/Runtime/Engine/Private/WorldPartition/WorldPartitionRuntimeSpatialHash.cpp
+++ b/Engine/Source/Runtime/Engine/Private/WorldPartition/WorldPartitionRuntimeSpatialHash.cpp
@@ -25,6 +25,7 @@
 #include "WorldPartition/DataLayer/ExternalDataLayerAsset.h"
 #include "WorldPartition/DataLayer/ExternalDataLayerInstance.h"
 #include "WorldPartition/HLOD/HLODLayer.h"
+#include "Components/BillboardComponent.h"
 #include "Components/LineBatchComponent.h"
 #include "Engine/Engine.h"
 #include "Engine/Level.h"
@@ -36,6 +37,7 @@
 #include "Math/TransformCalculus2D.h"
 #include "Misc/ArchiveMD5.h"
 #include "Misc/HierarchicalLogArchive.h"
+#include "Serialization/ArchiveReplaceObjectRef.h"
 #include "UObject/UnrealType.h"
 #include "UObject/ObjectSaveContext.h"
 #include "UObject/Package.h"
@@ -1159,6 +1161,22 @@ ASpatialHashRuntimeGridInfo::ASpatialHashRuntimeGridInfo(const FObjectInitialize
 #endif
 }
 
+void ASpatialHashRuntimeGridInfo::Serialize(FArchive& Ar)
+{
+	Ar.UsingCustomVersion(FFortniteMainBranchObjectVersion::GUID);
+
+	Super::Serialize(Ar);
+
+#if WITH_EDITORONLY_DATA
+	if (Ar.IsLoading() && Ar.CustomVer(FFortniteMainBranchObjectVersion::GUID) < FFortniteMainBranchObjectVersion::SpatialHashRuntimeGridInfoSpriteFixup)
+	{
+		// Fixes ASpatialHashRuntimeGridInfo that were saved prior to the constructor using ObjectInitializer.DoNotCreateDefaultSubobject(TEXT("Sprite"))
+		// Those Components are MarkForGarbage in their UActorComponent::PostInitProperties and so we need to clear the Component pointers
+		FArchiveReplaceObjectRef<UObject>(this, { { GetSpriteComponent(), nullptr } });
+	}
+#endif
+}
+
 UWorldPartitionRuntimeSpatialHash::UWorldPartitionRuntimeSpatialHash(const FObjectInitializer& ObjectInitializer)
 	: Super(ObjectInitializer)
 #if WITH_EDITORONLY_DATA
diff --git a/Engine/Source/Runtime/Engine/Public/WorldPartition/WorldPartitionRuntimeSpatialHash.h b/Engine/Source/Runtime/Engine/Public/WorldPartition/WorldPartitionRuntimeSpatialHash.h
--- a/Engine/Source/Runtime/Engine/Public/WorldPartition/WorldPartitionRuntimeSpatialHash.h
+++ b/Engine/Source/Runtime/Engine/Public/WorldPartition/WorldPartitionRuntimeSpatialHash.h
@@ -283,6 +283,8 @@ class ASpatialHashRuntimeGridInfo : public AInfo
 {
 	GENERATED_UCLASS_BODY()
 
+	virtual void Serialize(FArchive& Ar) override;
+
 public:
 	UPROPERTY()
 	FSpatialHashRuntimeGrid	GridSettings;
