Add a cvar r.EnableDrawDebugHelpers to allow disabling of the draw debug helper functionality dynamically at runtime. Also flushes persistent lines/strings when it becomes disabled.

diff --git a/Engine/Source/Runtime/Engine/Private/DrawDebugHelpers.cpp b/Engine/Source/Runtime/Engine/Private/DrawDebugHelpers.cpp
--- a/Engine/Source/Runtime/Engine/Private/DrawDebugHelpers.cpp
+++ b/Engine/Source/Runtime/Engine/Private/DrawDebugHelpers.cpp
@@ -20,6 +20,32 @@
 ENGINE_API float GServerDrawDebugColorTintStrength = 0.75f;
 ENGINE_API FLinearColor GServerDrawDebugColorTint(0.0f, 0.0f, 0.0f, 1.0f);
 
+namespace UE::Private
+{
+	static void OnEnableDrawDebugHelpersChanged(IConsoleVariable* EnableDrawDebugHelpersVar)
+	{
+		// Clear out any persistently-drawn items if option is disabled
+		if (!EnableDrawDebugHelpersVar->GetBool())
+		{
+			for (const FWorldContext& WorldContext : GEngine->GetWorldContexts())
+			{
+				FlushPersistentDebugLines(WorldContext.World());
+				FlushDebugStrings(WorldContext.World());
+			}
+		}
+	}
+
+	static bool bEnableDrawDebugHelpers = true;
+	static FAutoConsoleVariableRef CVarEnableDrawDebugHelpers(
+		TEXT("r.EnableDrawDebugHelpers"),
+		bEnableDrawDebugHelpers,
+		TEXT(	"If true, then any 'DrawDebug' helper functions will be enabled.\nIf false, then the 'DrawDebug' helper functions will exit early and any persistent items will be removed.\nThis is only considered when ENABLE_DRAW_DEBUG is true"),
+		FConsoleVariableDelegate::CreateStatic(&OnEnableDrawDebugHelpersChanged));
+}
+
+// Checks the value of the r.EnableDrawDebugHelpers cvar, and returns early if it is false.
+#define RETURN_IF_DEBUG_DRAW_DISABLED		if (!UE::Private::bEnableDrawDebugHelpers) { return; }
+
 #if WITH_EDITOR
 
 FColor AdjustColorForServer(const FColor InColor)
@@ -101,6 +127,8 @@ static float GetDebugLineLifeTime(ULineBatchComponent* LineBatcher, float LifeTi
 
 void DrawDebugLine(const UWorld* InWorld, FVector const& LineStart, FVector const& LineEnd, FColor const& Color, bool bPersistentLines, float LifeTime, uint8 DepthPriority, float Thickness)
 {
+	RETURN_IF_DEBUG_DRAW_DISABLED;
+	
 	if (GEngine->GetNetMode(InWorld) != NM_DedicatedServer)
 	{
 		// this means foreground lines can't be persistent 
@@ -118,6 +146,8 @@ void DrawDebugLine(const UWorld* InWorld, FVector const& LineStart, FVector cons
 
 void DrawDebugPoint(const UWorld* InWorld, FVector const& Position, float Size, FColor const& Color, bool bPersistentLines, float LifeTime, uint8 DepthPriority)
 {
+	RETURN_IF_DEBUG_DRAW_DISABLED;
+	
 	if (GEngine->GetNetMode(InWorld) != NM_DedicatedServer)
 	{
 		// this means foreground lines can't be persistent 
@@ -135,6 +165,8 @@ void DrawDebugPoint(const UWorld* InWorld, FVector const& Position, float Size,
 
 void DrawDebugDirectionalArrow(const UWorld* InWorld, FVector const& LineStart, FVector const& LineEnd, float ArrowSize, FColor const& Color, bool bPersistentLines, float LifeTime, uint8 DepthPriority, float Thickness)
 {
+	RETURN_IF_DEBUG_DRAW_DISABLED;
+	
 	// no debug line drawing on dedicated server
 	if (GEngine->GetNetMode(InWorld) != NM_DedicatedServer)
 	{
@@ -157,6 +189,8 @@ void DrawDebugDirectionalArrow(const UWorld* InWorld, FVector const& LineStart,
 
 void DrawDebugBox(const UWorld* InWorld, FVector const& Center, FVector const& Box, FColor const& Color, bool bPersistentLines, float LifeTime, uint8 DepthPriority, float Thickness)
 {
+	RETURN_IF_DEBUG_DRAW_DISABLED;
+	
 	// no debug line drawing on dedicated server
 	if (GEngine->GetNetMode(InWorld) != NM_DedicatedServer)
 	{
@@ -175,6 +209,8 @@ void DrawDebugBox(const UWorld* InWorld, FVector const& Center, FVector const& B
 
 void DrawDebugBox(const UWorld* InWorld, FVector const& Center, FVector const& Box, const FQuat& Rotation, FColor const& Color, bool bPersistentLines, float LifeTime, uint8 DepthPriority, float Thickness)
 {
+	RETURN_IF_DEBUG_DRAW_DISABLED;
+	
 	// no debug line drawing on dedicated server
 	if (GEngine->GetNetMode(InWorld) != NM_DedicatedServer)
 	{
@@ -194,6 +230,8 @@ void DrawDebugBox(const UWorld* InWorld, FVector const& Center, FVector const& B
 
 void DrawDebugMesh(const UWorld* InWorld, TArray<FVector> const& Verts, TArray<int32> const& Indices, FColor const& Color, bool bPersistent, float LifeTime, uint8 DepthPriority)
 {
+	RETURN_IF_DEBUG_DRAW_DISABLED;
+	
 	// no debug line drawing on dedicated server
 	if (GEngine->GetNetMode(InWorld) != NM_DedicatedServer)
 	{
@@ -211,6 +249,8 @@ void DrawDebugMesh(const UWorld* InWorld, TArray<FVector> const& Verts, TArray<i
 
 void DrawDebugSolidBox(const UWorld* InWorld, FBox const& Box, FColor const& Color, const FTransform& Transform, bool bPersistent, float LifeTime, uint8 DepthPriority)
 {
+	RETURN_IF_DEBUG_DRAW_DISABLED;
+	
 	// no debug line drawing on dedicated server
 	if (GEngine->GetNetMode(InWorld) != NM_DedicatedServer)
 	{
@@ -227,7 +267,10 @@ void DrawDebugSolidBox(const UWorld* InWorld, FBox const& Box, FColor const& Col
 }
 
 void DrawDebugSolidBox(const UWorld* InWorld, FVector const& Center, FVector const& Extent, FColor const& Color, bool bPersistent, float LifeTime, uint8 DepthPriority)
-{	// No Rotation, so just use identity transform and build the box in the right place!
+{
+	RETURN_IF_DEBUG_DRAW_DISABLED;
+	
+	// No Rotation, so just use identity transform and build the box in the right place!
 	FBox Box = FBox::BuildAABB(Center, Extent);
 
 	DrawDebugSolidBox(InWorld, Box, Color, FTransform::Identity, bPersistent, LifeTime, DepthPriority);
@@ -235,6 +278,8 @@ void DrawDebugSolidBox(const UWorld* InWorld, FVector const& Center, FVector con
 
 void DrawDebugSolidBox(const UWorld* InWorld, FVector const& Center, FVector const& Extent, FQuat const& Rotation, FColor const& Color, bool bPersistent, float LifeTime, uint8 DepthPriority)
 {
+	RETURN_IF_DEBUG_DRAW_DISABLED;
+	
 	FTransform Transform(Rotation, Center, FVector(1.0f, 1.0f, 1.0f));	// Build transform from Rotation, Center with uniform scale of 1.0.
 	FBox Box = FBox::BuildAABB(FVector::ZeroVector, Extent);	// The Transform handles the Center location, so this box needs to be centered on origin.
 
@@ -244,11 +289,15 @@ void DrawDebugSolidBox(const UWorld* InWorld, FVector const& Center, FVector con
 /** Loc is an anchor point in the world to guide which part of the infinite plane to draw. */
 void DrawDebugSolidPlane(const UWorld* InWorld, FPlane const& P, FVector const& Loc, float Size, FColor const& Color, bool bPersistent, float LifeTime, uint8 DepthPriority)
 {
+	RETURN_IF_DEBUG_DRAW_DISABLED;
+	
 	DrawDebugSolidPlane(InWorld, P, Loc, FVector2D(Size, Size), Color, bPersistent, LifeTime, DepthPriority);
 }
 
 ENGINE_API void DrawDebugSolidPlane(const UWorld* InWorld, FPlane const& P, FVector const& Loc, FVector2D const& Extents, FColor const& Color, bool bPersistent/*=false*/, float LifeTime/*=-1*/, uint8 DepthPriority /*= 0*/)
 {
+	RETURN_IF_DEBUG_DRAW_DISABLED;
+	
 	// no debug line drawing on dedicated server
 	if(GEngine->GetNetMode(InWorld) != NM_DedicatedServer)
 	{
@@ -285,6 +334,8 @@ ENGINE_API void DrawDebugSolidPlane(const UWorld* InWorld, FPlane const& P, FVec
 
 void DrawDebugCoordinateSystem(const UWorld* InWorld, FVector const& AxisLoc, FRotator const& AxisRot, float Scale, bool bPersistentLines, float LifeTime, uint8 DepthPriority, float Thickness)
 {
+	RETURN_IF_DEBUG_DRAW_DISABLED;
+	
 	// no debug line drawing on dedicated server
 	if (GEngine->GetNetMode(InWorld) != NM_DedicatedServer)
 	{
@@ -310,6 +361,8 @@ void DrawDebugCoordinateSystem(const UWorld* InWorld, FVector const& AxisLoc, FR
 
 ENGINE_API void DrawDebugCrosshairs(const UWorld* InWorld, FVector const& AxisLoc, FRotator const& AxisRot, float Scale, const FColor& Color, bool bPersistentLines, float LifeTime, uint8 DepthPriority)
 {
+	RETURN_IF_DEBUG_DRAW_DISABLED;
+	
 	// no debug line drawing on dedicated server
 	if (GEngine->GetNetMode(InWorld) != NM_DedicatedServer)
 	{
@@ -336,6 +389,8 @@ ENGINE_API void DrawDebugCrosshairs(const UWorld* InWorld, FVector const& AxisLo
 
 static void InternalDrawDebugCircle(const UWorld* InWorld, const FMatrix& TransformMatrix, float Radius, int32 Segments, const FColor& Color, bool bPersistentLines, float LifeTime, uint8 DepthPriority, float Thickness = 0.f)
 {
+	RETURN_IF_DEBUG_DRAW_DISABLED;
+	
 	if (ULineBatchComponent* const LineBatcher = GetDebugLineBatcher(InWorld, bPersistentLines, LifeTime, (DepthPriority == SDPG_Foreground)))
 	{
 		const float LineLifeTime = GetDebugLineLifeTime(LineBatcher, LifeTime, bPersistentLines);
@@ -365,6 +420,8 @@ static void InternalDrawDebugCircle(const UWorld* InWorld, const FMatrix& Transf
 
 void DrawDebugCircle(const UWorld* InWorld, const FMatrix& TransformMatrix, float Radius, int32 Segments, const FColor& Color, bool bPersistentLines, float LifeTime, uint8 DepthPriority, float Thickness, bool bDrawAxis)
 {
+	RETURN_IF_DEBUG_DRAW_DISABLED;
+	
 	// no debug line drawing on dedicated server
 	if (GEngine->GetNetMode(InWorld) != NM_DedicatedServer)
 	{
@@ -396,6 +453,8 @@ void DrawDebugCircle(const UWorld* InWorld, const FMatrix& TransformMatrix, floa
 
 void DrawDebugCircle(const UWorld* InWorld, FVector Center, float Radius, int32 Segments, const FColor& Color, bool PersistentLines, float LifeTime, uint8 DepthPriority, float Thickness, FVector YAxis, FVector ZAxis, bool bDrawAxis)
 {
+	RETURN_IF_DEBUG_DRAW_DISABLED;
+	
 	FMatrix TM;
 	TM.SetOrigin(Center);
 	TM.SetAxis(0, FVector(1,0,0));
@@ -418,6 +477,8 @@ void DrawDebugCircle(const UWorld* InWorld, FVector Center, float Radius, int32
 
 void DrawDebugCircleArc(const UWorld* InWorld, const FVector& Center, float Radius, const FVector& Direction, float AngleWidth, int32 Segments, const FColor& Color, bool PersistentLines, float LifeTime, uint8 DepthPriority, float Thickness)
 {
+	RETURN_IF_DEBUG_DRAW_DISABLED;
+	
 	if (GEngine->GetNetMode(InWorld) != NM_DedicatedServer)
 	{
 		if (ULineBatchComponent* const LineBatcher = GetDebugLineBatcher(InWorld, PersistentLines, LifeTime, (DepthPriority == SDPG_Foreground)))
@@ -456,6 +517,8 @@ void DrawDebugCircleArc(const UWorld* InWorld, const FVector& Center, float Radi
 
 void DrawDebug2DDonut(const UWorld* InWorld, const FMatrix& TransformMatrix, float InnerRadius, float OuterRadius, int32 Segments, const FColor& Color, bool bPersistentLines, float LifeTime, uint8 DepthPriority, float Thickness)
 {
+	RETURN_IF_DEBUG_DRAW_DISABLED;
+	
 	// no debug line drawing on dedicated server
 	if (GEngine->GetNetMode(InWorld) != NM_DedicatedServer)
 	{
@@ -489,6 +552,8 @@ void DrawDebug2DDonut(const UWorld* InWorld, const FMatrix& TransformMatrix, flo
 
 void DrawDebugSphere(const UWorld* InWorld, FVector const& Center, float Radius, int32 Segments, FColor const& Color, bool bPersistentLines, float LifeTime, uint8 DepthPriority, float Thickness)
 {
+	RETURN_IF_DEBUG_DRAW_DISABLED;
+	
 	// no debug line drawing on dedicated server
 	if (GEngine->GetNetMode(InWorld) != NM_DedicatedServer)
 	{
@@ -507,6 +572,8 @@ void DrawDebugSphere(const UWorld* InWorld, FVector const& Center, float Radius,
 
 void DrawDebugCylinder(const UWorld* InWorld, FVector const& Start, FVector const& End, float Radius, int32 Segments, FColor const& Color, bool bPersistentLines, float LifeTime, uint8 DepthPriority, float Thickness)
 {
+	RETURN_IF_DEBUG_DRAW_DISABLED;
+	
 	// no debug line drawing on dedicated server
 	if (GEngine->GetNetMode(InWorld) != NM_DedicatedServer)
 	{
@@ -526,6 +593,8 @@ void DrawDebugCylinder(const UWorld* InWorld, FVector const& Start, FVector cons
 /** Used by gameplay when defining a cone by a vertical and horizontal dot products. */
 void DrawDebugAltCone(const UWorld* InWorld, FVector const& Origin, FRotator const& Rotation, float Length, float AngleWidth, float AngleHeight, FColor const& DrawColor, bool bPersistentLines, float LifeTime, uint8 DepthPriority, float Thickness)
 {
+	RETURN_IF_DEBUG_DRAW_DISABLED;
+	
 	if (GEngine->GetNetMode(InWorld) != NM_DedicatedServer)
 	{
 		// this means foreground lines can't be persistent 
@@ -594,6 +663,8 @@ void DrawDebugAltCone(const UWorld* InWorld, FVector const& Origin, FRotator con
 
 void DrawDebugCone(const UWorld* InWorld, FVector const& Origin, FVector const& Direction, float Length, float AngleWidth, float AngleHeight, int32 NumSides, FColor const& DrawColor, bool bPersistentLines, float LifeTime, uint8 DepthPriority, float Thickness)
 {
+	RETURN_IF_DEBUG_DRAW_DISABLED;
+	
 	// no debug line drawing on dedicated server
 	if (GEngine->GetNetMode(InWorld) != NM_DedicatedServer)
 	{
@@ -612,6 +683,8 @@ void DrawDebugCone(const UWorld* InWorld, FVector const& Origin, FVector const&
 
 void DrawDebugString(const UWorld* InWorld, FVector const& TextLocation, const FString& Text, class AActor* TestBaseActor, FColor const& TextColor, float Duration, bool bDrawShadow, float FontScale)
 {
+	RETURN_IF_DEBUG_DRAW_DISABLED;
+	
 	// no debug line drawing on dedicated server
 	if (GEngine->GetNetMode(InWorld) != NM_DedicatedServer)
 	{
@@ -652,6 +725,8 @@ void FlushDebugStrings( const UWorld* InWorld )
 
 void DrawDebugFrustum(const UWorld* InWorld, const FMatrix& FrustumToWorld, FColor const& Color, bool bPersistentLines, float LifeTime, uint8 DepthPriority, float Thickness)
 {
+	RETURN_IF_DEBUG_DRAW_DISABLED;
+	
 	// no debug line drawing on dedicated server
 	if (GEngine->GetNetMode(InWorld) != NM_DedicatedServer)
 	{
@@ -698,6 +773,8 @@ void DrawDebugFrustum(const UWorld* InWorld, const FMatrix& FrustumToWorld, FCol
 
 void DrawCircle(const UWorld* InWorld, const FVector& Base, const FVector& X, const FVector& Y, const FColor& Color, float Radius, int32 NumSides, bool bPersistentLines, float LifeTime, uint8 DepthPriority, float Thickness)
 {
+	RETURN_IF_DEBUG_DRAW_DISABLED;
+	
 	const float	AngleDelta = 2.0f * UE_PI / NumSides;
 	FVector	LastVertex = Base + X * Radius;
 
@@ -711,6 +788,8 @@ void DrawCircle(const UWorld* InWorld, const FVector& Base, const FVector& X, co
 
 void DrawDebugCapsule(const UWorld* InWorld, FVector const& Center, float HalfHeight, float Radius, const FQuat& Rotation, FColor const& Color, bool bPersistentLines, float LifeTime, uint8 DepthPriority, float Thickness)
 {
+	RETURN_IF_DEBUG_DRAW_DISABLED;
+	
 	// no debug line drawing on dedicated server
 	if (GEngine->GetNetMode(InWorld) != NM_DedicatedServer)
 	{
@@ -729,6 +808,8 @@ void DrawDebugCapsule(const UWorld* InWorld, FVector const& Center, float HalfHe
 
 void DrawDebugCamera(const UWorld* InWorld, FVector const& Location, FRotator const& Rotation, float FOVDeg, float Scale, FColor const& Color, bool bPersistentLines, float LifeTime, uint8 DepthPriority)
 {
+	RETURN_IF_DEBUG_DRAW_DISABLED;
+	
 	static float BaseScale = 4.f;
 	static FVector BaseProportions(2.f, 1.f, 1.5f);
 
@@ -775,12 +856,16 @@ void DrawDebugCamera(const UWorld* InWorld, FVector const& Location, FRotator co
 // https://en.wikipedia.org/wiki/Centripetal_Catmull%E2%80%93Rom_spline
 void DrawCentripetalCatmullRomSpline(const UWorld* InWorld, TConstArrayView<FVector> Points, FColor const& Color, float Alpha, int32 NumSamplesPerSegment, bool bPersistentLines, float LifeTime, uint8 DepthPriority, float Thickness)
 {
+	RETURN_IF_DEBUG_DRAW_DISABLED;
+	
 	TConstArrayView<FColor> Colors(&Color, 1);
 	DrawCentripetalCatmullRomSpline(InWorld, Points, Colors, Alpha, NumSamplesPerSegment, bPersistentLines, LifeTime, DepthPriority, Thickness);
 }
 
 void DrawCentripetalCatmullRomSpline(const UWorld* InWorld, TConstArrayView<FVector> Points, TConstArrayView<FColor> Colors, float Alpha, int32 NumSamplesPerSegment, bool bPersistentLines, float LifeTime, uint8 DepthPriority, float Thickness)
 {
+	RETURN_IF_DEBUG_DRAW_DISABLED;
+	
 	const int32 NumPoints = Points.Num();
 	const int32 NumColors = Colors.Num();
 	if (NumPoints > 1)
@@ -853,6 +938,8 @@ void DrawCentripetalCatmullRomSpline(const UWorld* InWorld, TConstArrayView<FVec
 
 void DrawDebugFloatHistory(UWorld const & WorldRef, FDebugFloatHistory const & FloatHistory, FTransform const & DrawTransform, FVector2D const & DrawSize, FColor const & DrawColor, bool const & bPersistent, float const & LifeTime, uint8 const & DepthPriority)
 {
+	RETURN_IF_DEBUG_DRAW_DISABLED;
+	
 	int const NumSamples = FloatHistory.GetNumSamples();
 	if (NumSamples >= 2)
 	{
@@ -898,6 +985,8 @@ void DrawDebugFloatHistory(UWorld const & WorldRef, FDebugFloatHistory const & F
 
 void DrawDebugFloatHistory(UWorld const & WorldRef, FDebugFloatHistory const & FloatHistory, FVector const & DrawLocation, FVector2D const & DrawSize, FColor const & DrawColor, bool const & bPersistent, float const & LifeTime, uint8 const & DepthPriority)
 {
+	RETURN_IF_DEBUG_DRAW_DISABLED;
+	
 	APlayerController * PlayerController = WorldRef.GetGameInstance() != nullptr ? WorldRef.GetGameInstance()->GetFirstLocalPlayerController() : nullptr;
 	FRotator const DrawRotation = (PlayerController && PlayerController->PlayerCameraManager) ? PlayerController->PlayerCameraManager->GetCameraRotation() : FRotator(0, 0, 0);
 
@@ -910,6 +999,8 @@ void DrawDebugFloatHistory(UWorld const & WorldRef, FDebugFloatHistory const & F
 
 void DrawDebugCanvas2DLine(UCanvas* Canvas, const FVector& Start, const FVector& End, const FLinearColor& LineColor)
 {
+	RETURN_IF_DEBUG_DRAW_DISABLED;
+	
 	FCanvasLineItem LineItem;
 	LineItem.Origin = Start;
 	LineItem.EndPos = End;
@@ -920,11 +1011,15 @@ void DrawDebugCanvas2DLine(UCanvas* Canvas, const FVector& Start, const FVector&
 
 void DrawDebugCanvasLine(UCanvas* Canvas, const FVector& Start, const FVector& End, const FLinearColor& LineColor)
 {
+	RETURN_IF_DEBUG_DRAW_DISABLED;
+	
 	DrawDebugCanvas2DLine(Canvas, Canvas->Project(Start), Canvas->Project(End), LineColor);
 }
 
 void DrawDebugCanvasCircle(UCanvas* Canvas, const FVector& Base, const FVector& X, const FVector& Y, FColor Color, float Radius, int32 NumSides)
 {
+	RETURN_IF_DEBUG_DRAW_DISABLED;
+	
 	const float	AngleDelta = 2.0f * UE_PI / NumSides;
 	FVector	LastVertex = Base + X * Radius;
 
@@ -938,6 +1033,8 @@ void DrawDebugCanvasCircle(UCanvas* Canvas, const FVector& Base, const FVector&
 
 void DrawDebugCanvasHalfCircle(UCanvas* Canvas, const FVector& Base, const FVector& X, const FVector& Y, FColor Color, float Radius, int32 NumSides)
 {
+	RETURN_IF_DEBUG_DRAW_DISABLED;
+	
 	const float	AngleDelta = 2.0f * UE_PI / NumSides;
 	FVector	LastVertex = Base + X * Radius;
 
@@ -951,6 +1048,8 @@ void DrawDebugCanvasHalfCircle(UCanvas* Canvas, const FVector& Base, const FVect
 
 void DrawDebugCanvasWireSphere(UCanvas* Canvas, const FVector& Base, FColor Color, float Radius, int32 NumSides)
 {
+	RETURN_IF_DEBUG_DRAW_DISABLED;
+	
 	DrawDebugCanvasCircle(Canvas, Base, FVector(1,0,0), FVector(0,1,0), Color, Radius, NumSides);
 	DrawDebugCanvasCircle(Canvas, Base, FVector(1,0,0), FVector(0,0,1), Color, Radius, NumSides);
 	DrawDebugCanvasCircle(Canvas, Base, FVector(0,1,0), FVector(0,0,1), Color, Radius, NumSides);
@@ -958,6 +1057,8 @@ void DrawDebugCanvasWireSphere(UCanvas* Canvas, const FVector& Base, FColor Colo
 
 void DrawDebugCanvasWireCone(UCanvas* Canvas, const FTransform& Transform, float ConeRadius, float ConeAngle, int32 ConeSides, FColor Color)
 {
+	RETURN_IF_DEBUG_DRAW_DISABLED;
+	
 	static const float TwoPI = 2.0f * UE_PI;
 	static const float ToRads = UE_PI / 180.0f;
 	static const float MaxAngle = 89.0f * ToRads + 0.001f;
@@ -1001,6 +1102,8 @@ void DrawDebugCanvasWireCone(UCanvas* Canvas, const FTransform& Transform, float
 
 void DrawDebugCanvasWireBox(UCanvas* Canvas, const FMatrix& Transform, const FBox& Box, FColor Color)
 {
+	RETURN_IF_DEBUG_DRAW_DISABLED;
+	
 	const FVector Vertices[] =
 	{
 		Transform.TransformPosition(FVector(Box.Min.X, Box.Min.Y, Box.Min.Z)),
@@ -1028,6 +1131,8 @@ void DrawDebugCanvasWireBox(UCanvas* Canvas, const FMatrix& Transform, const FBo
 
 void DrawDebugCanvasCapsule(UCanvas* Canvas, const FMatrix& Transform, float HalfLength, float Radius, const FColor& LineColor)
 {
+	RETURN_IF_DEBUG_DRAW_DISABLED;
+	
 	constexpr int32 DrawCollisionSides = 16;
 
 	FVector Origin = Transform.GetOrigin();
@@ -1065,6 +1170,8 @@ void DrawDebugCanvasCapsule(UCanvas* Canvas, const FMatrix& Transform, float Hal
 
 void DrawDebugCanvas2DLine(UCanvas* Canvas, const FVector2D& StartPosition, const FVector2D& EndPosition, const FLinearColor& LineColor, const float& LineThickness)
 {
+	RETURN_IF_DEBUG_DRAW_DISABLED;
+	
 	if (Canvas)
 	{
 		FCanvasLineItem LineItem(StartPosition, EndPosition);
@@ -1076,6 +1183,8 @@ void DrawDebugCanvas2DLine(UCanvas* Canvas, const FVector2D& StartPosition, cons
 
 void DrawDebugCanvas2DCircle(UCanvas* Canvas, const FVector2D& Center, float Radius, int32 NumSides, const FLinearColor& LineColor, const float& LineThickness)
 {
+	RETURN_IF_DEBUG_DRAW_DISABLED;
+	
 	const float	AngleDelta = 2.0f * UE_PI / NumSides;
 	FVector2D AxisX(1.f, 0.f);
 	FVector2D AxisY(0.f, -1.f);
@@ -1091,6 +1200,8 @@ void DrawDebugCanvas2DCircle(UCanvas* Canvas, const FVector2D& Center, float Rad
 
 void DrawDebugCanvas2DBox(UCanvas* Canvas, const FBox2D& Box, const FLinearColor& LineColor, const float& LineThickness)
 {
+	RETURN_IF_DEBUG_DRAW_DISABLED;
+	
 	if (Canvas)
 	{
 		FCanvasBoxItem BoxItem(Box.Min, Box.GetSize());
