Material Layer 2.0:
- Fixing the GUID assigned to the default layer node
- Displaying the the internal tree in the same layout as the ui to make debugging easier

diff --git a/Engine/Source/Runtime/Engine/Private/Materials/MaterialExpressions.cpp b/Engine/Source/Runtime/Engine/Private/Materials/MaterialExpressions.cpp
--- a/Engine/Source/Runtime/Engine/Private/Materials/MaterialExpressions.cpp
+++ b/Engine/Source/Runtime/Engine/Private/Materials/MaterialExpressions.cpp
@@ -18407,6 +18407,25 @@ FLayerNodeIdArray FMaterialLayersFunctionsTree::GetChildrenIds(FNodeId InNodeId)
 	return ChildrendIds;
 }
 
+FLayerNodeId FMaterialLayersFunctionsTree::GetChildrenTailId(FNodeId InNodeId) const
+{
+	FNodeId NextChildId = Root;
+	FNode* Node = GetNode(InNodeId);
+	if (Node)
+		NextChildId = Node->ChildrenHead;
+
+	while (IsValidId(NextChildId))
+	{
+		Node = GetNode(NextChildId);
+		if (Node && IsValidId(Node->NextSibling))
+			NextChildId = Node->NextSibling;
+		else
+			break;
+	}
+
+	return NextChildId;
+}
+
 int32 FMaterialLayersFunctionsTree::GetDepth(FNodeId InNodeId) const
 {
 	int32 Depth = 0;
@@ -18451,6 +18470,30 @@ FLayerNodeId FMaterialLayersFunctionsTree::GetSiblingHeadId(FNodeId InNodeId) co
 	return HeadNodeId;
 }
 
+FLayerNodeId FMaterialLayersFunctionsTree::GetPreviousSiblingId(FNodeId InNodeId) const
+{
+	// Go to the sibling head
+	// and walk next until next is the specified node
+	FLayerNodeId PrevSibling = GetSiblingHeadId(InNodeId);
+
+	// if the queried node IS the first of the siblings then no more prev siblings
+	if (!IsValidId(PrevSibling) || (PrevSibling == InNodeId))
+		return InvalidId;
+
+	// LEt's look for the Node whose next sibling IS the queried nodeID
+	while (IsValidId(PrevSibling))
+	{
+		FNode* Node = GetNode(PrevSibling); 
+		if (Node->NextSibling == InNodeId)
+		{
+			break;
+		}
+		PrevSibling = Node->NextSibling;
+	}
+
+	return PrevSibling;
+}
+
 FLayerNodeId FMaterialLayersFunctionsTree::AddNode(const FPayload& InPayload, FNodeId InParent, int32 InAtChildIndex)
 {
 	// Check Parent id
@@ -18765,7 +18808,7 @@ namespace MLFT {
 			{
 				TraverseNodes.Add({ InNodeId,  InDepth, InPayload });
 			};
-		LFRD->Tree.Traverse(Visitor);
+		LFRD->Tree.TraverseBottomUp(Visitor);
 
 		for (auto ItNode = TraverseNodes.begin(); ItNode != TraverseNodes.end(); ++ItNode)
 		{
@@ -19081,6 +19124,10 @@ void FMaterialLayersFunctions::AddDefaultBackgroundLayer()
 	AppendLayerNode(-1, -1);
 	FText LayerName = FText(LOCTEXT("Background", "Background"));
 	EditorOnly.LayerNames[0] = LayerName;
+	// Use a consistent Guid for the background layer
+	// Default constructor assigning different guids will break FStructUtils::AttemptToFindUninitializedScriptStructMembers
+	EditorOnly.LayerGuids[0] = (BackgroundGuid);
+
 #else
 
 	// Default to a non-blended "background" layer
diff --git a/Engine/Source/Runtime/Engine/Public/Materials/MaterialLayersFunctions.h b/Engine/Source/Runtime/Engine/Public/Materials/MaterialLayersFunctions.h
--- a/Engine/Source/Runtime/Engine/Public/Materials/MaterialLayersFunctions.h
+++ b/Engine/Source/Runtime/Engine/Public/Materials/MaterialLayersFunctions.h
@@ -256,9 +256,16 @@ public:
 	// Find the node which is the head of the list of siblings where InNodeID belongs
 	FNodeId GetSiblingHeadId(FNodeId InNodeId) const;
 
+	// Find the node which is the previous sibling to the specified node
+	// The Next sibling is directly available in the FNode struct, not the previous one
+	FNodeId GetPreviousSiblingId(FNodeId InNodeId) const;
+
 	// Gather the list of children for this node
 	TArray<FNodeId> GetChildrenIds(FNodeId InNodeId) const;
 
+	// Find the last of the children in the specified node
+	FNodeId GetChildrenTailId(FNodeId InNodeId) const;
+
 	// Return the depth of the node in the tree, the root is depth 0
 	// the first valid node in the tree under the root is depth 1
 	int32 GetDepth(FNodeId InNodeId) const;
@@ -287,7 +294,7 @@ public:
 	// Traverser Utility
 	using FNodeVisitor = void (*)(FNodeId InNodeId, int32 InDepth, int32 InSiblingNum, const FPayload& InPayload);
 	template <typename V>
-	FNodeId TraverseNode(V InVisitor, FNodeId InNodeId, int32 InDepth, int32 InSiblingNum) const
+	FNodeId TraverseNode(V InVisitor, FNodeId InNodeId, int32 InDepth, int32 InSiblingNum, bool ReverseOrder = false) const
 	{
 		// First plan for the next child node and sibling we will be visiting
 		// default with the InNodeId is the root case
@@ -307,6 +314,18 @@ public:
 			// And since we have visited a valid node, lets visit the children next and the next sibling after
 			NextChildId = Node.ChildrenHead;
 			NextSiblingId = Node.NextSibling;
+			if (ReverseOrder)
+			{
+				NextChildId = GetChildrenTailId(InNodeId);
+				NextSiblingId = GetPreviousSiblingId(InNodeId);
+			}
+		}
+		else
+		{
+			if (ReverseOrder)
+			{
+				NextChildId = GetChildrenTailId(InvalidId);
+			}
 		}
 
 		// Traverse children
@@ -315,7 +334,7 @@ public:
 			int32 SiblingNum = 0;
 			while (IsValidId(NextChildId))
 			{
-				NextChildId = TraverseNode(InVisitor, NextChildId, InDepth, SiblingNum);
+				NextChildId = TraverseNode(InVisitor, NextChildId, InDepth, SiblingNum, ReverseOrder);
 				SiblingNum++;
 			}
 			InDepth--;
@@ -329,6 +348,11 @@ public:
 	{
 		TraverseNode(InVisitor, InRootNodeId, 0, 0);
 	}
+	template <typename V>
+	void TraverseBottomUp(V InVisitor, FNodeId InRootNodeId = InvalidId) const
+	{
+		TraverseNode(InVisitor, InRootNodeId, 0, 0, true);
+	}
 
 	// Logging the tree data structure, useful for debugging
 	// generate a multiline string representing the tree nodes hierarchy with the payload info per node
