[Switch] Possible fix for live issue that doesn't reproduce in testing. Block the render thread while we're restarting scheduler worker threads. Switch has slightly different thread priorities than other platforms.

#rnx

#tests local, preflight

diff --git a/Engine/Source/Runtime/Engine/Private/UnrealEngine.cpp b/Engine/Source/Runtime/Engine/Private/UnrealEngine.cpp
--- a/Engine/Source/Runtime/Engine/Private/UnrealEngine.cpp
+++ b/Engine/Source/Runtime/Engine/Private/UnrealEngine.cpp
@@ -19388,12 +19388,31 @@ static void SetupThreadConfig(const TArray<FString>& Args)
 	check(IsInGameThread());
 	SetPriorityAndAffinityOnGameThread();
 	FlushRenderingCommands();
-	
+
+	// Block the render thread while we're restarting workers
+	UE::FManualResetEvent ReadyEvent;
+	UE::FManualResetEvent WaitEvent;
+	UE::FManualResetEvent ExitEvent;
+	ENQUEUE_RENDER_COMMAND(StallRT)(
+		[&](FRHICommandListImmediate& RHICmdList)
+		{
+			ReadyEvent.Notify();
+			WaitEvent.Wait();
+			ExitEvent.Notify();
+		}
+	);
+
+	// Wait until the RT is really idle
+	ReadyEvent.Wait();
+
 	int32 NumWorkerThreads = FPlatformMisc::NumberOfWorkerThreadsToSpawn();
 	int32 NumBackgroundWorkers = FMath::Max(1, NumWorkerThreads - FMath::Min<int>(2, NumWorkerThreads));
 	int32 NumForegroundWorkers = FMath::Max(1, NumWorkerThreads - NumBackgroundWorkers);
 	LowLevelTasks::FScheduler::Get().RestartWorkers(NumForegroundWorkers, NumBackgroundWorkers, FForkProcessHelper::IsForkedMultithreadInstance() ? FThread::Forkable : FThread::NonForkable, TaskThreadConfig.Priority, TaskBPThreadConfig.Priority, TaskThreadConfig.Affinity, TaskBPThreadConfig.Affinity);
 
+	WaitEvent.Notify();
+	ExitEvent.Wait();
+
 	UE_LOG(LogConsoleResponse, Display, TEXT("Task Priority %s"), *ThreadPriorityToString(TaskThreadConfig.Priority));
 	UE_LOG(LogConsoleResponse, Display, TEXT("Task Affinity %016llX"), TaskThreadConfig.Affinity);
 
