Disabling actor ticking and initialization when using a multi-server proxy.

diff --git a/Engine/Plugins/Runtime/MultiServerReplication/Source/MultiServerReplication/Private/MultiServerProxy.cpp b/Engine/Plugins/Runtime/MultiServerReplication/Source/MultiServerReplication/Private/MultiServerProxy.cpp
--- a/Engine/Plugins/Runtime/MultiServerReplication/Source/MultiServerReplication/Private/MultiServerProxy.cpp
+++ b/Engine/Plugins/Runtime/MultiServerReplication/Source/MultiServerReplication/Private/MultiServerProxy.cpp
@@ -320,6 +320,12 @@ bool UProxyNetDriver::InitBase(bool bInitAsClient, FNetworkNotify* InNotify, con
 	}
 
 	SetReplicateTransactionally(false);
+	
+#if UE_SUPPORT_FOR_ACTOR_TICK_DISABLE
+	// The proxy should only be replicating the exact state from the servers and passing it on
+	// to clients and not tick actors or call any user actor callbacks.
+	GetWorld()->EnableActorTickAndUserCallbacks(false);
+#endif
 
 	return bSuccess;
 }
diff --git a/Engine/Source/Runtime/Core/Public/Misc/Build.h b/Engine/Source/Runtime/Core/Public/Misc/Build.h
--- a/Engine/Source/Runtime/Core/Public/Misc/Build.h
+++ b/Engine/Source/Runtime/Core/Public/Misc/Build.h
@@ -524,3 +524,8 @@
 	#define UE_DEPRECATED_PROFILER_ENABLED 0
 #endif
 
+// A compile time flag used to enable support for disabling actor ticking and calls to user callbacks.
+// This functionality is not intended for general use and should be used with care.
+#ifndef UE_SUPPORT_FOR_ACTOR_TICK_DISABLE
+	#define UE_SUPPORT_FOR_ACTOR_TICK_DISABLE 0
+#endif
diff --git a/Engine/Source/Runtime/Engine/Classes/Engine/World.h b/Engine/Source/Runtime/Engine/Classes/Engine/World.h
--- a/Engine/Source/Runtime/Engine/Classes/Engine/World.h
+++ b/Engine/Source/Runtime/Engine/Classes/Engine/World.h
@@ -1043,6 +1043,20 @@ public:
 	UPROPERTY()
 	FString										StreamingLevelsPrefix;
 
+#if UE_SUPPORT_FOR_ACTOR_TICK_DISABLE
+	/** 
+	 * Enable/disable actor ticking and any user callbacks.
+	 *
+	 * Applies to the following callbacks:
+	 *		BeginPlay()
+	 *		EndPlay()
+	 *		BeginComponentOverlap()
+	 *		EndComponentOverlap()
+	 */
+	void EnableActorTickAndUserCallbacks(bool bEnable);
+	bool IsActorTickAndUserCallbacksEnabled() const { return bEnableActorTickAndUserCallbacks; };
+#endif
+
 private:
 
 	/** Sets level's PendingVisibilityState */
@@ -1069,6 +1083,10 @@ private:
 	UPROPERTY()
 	TObjectPtr<class UDemoNetDriver>						DemoNetDriver;
 
+#if UE_SUPPORT_FOR_ACTOR_TICK_DISABLE
+	bool bEnableActorTickAndUserCallbacks = true;
+#endif
+
 public:
 	/** Gets the demo net driver for this world. */
 	UDemoNetDriver* GetDemoNetDriver() const { return DemoNetDriver; }
diff --git a/Engine/Source/Runtime/Engine/Private/Actor.cpp b/Engine/Source/Runtime/Engine/Private/Actor.cpp
--- a/Engine/Source/Runtime/Engine/Private/Actor.cpp
+++ b/Engine/Source/Runtime/Engine/Private/Actor.cpp
@@ -2112,6 +2112,13 @@ void AActor::SetTickGroup(ETickingGroup NewTickGroup)
 
 void AActor::ClearComponentOverlaps()
 {
+#if UE_SUPPORT_FOR_ACTOR_TICK_DISABLE
+	if (!GetWorld()->IsActorTickAndUserCallbacksEnabled())
+	{
+		return;
+	}
+#endif
+
 	TInlineComponentArray<UPrimitiveComponent*> PrimitiveComponents;
 	GetComponents(PrimitiveComponents);
 
@@ -4470,7 +4477,12 @@ void AActor::DispatchBeginPlay(bool bFromLevelStreaming)
 		BeginReplication();
 #endif // UE_WITH_IRIS
 
-		BeginPlay();
+#if UE_SUPPORT_FOR_ACTOR_TICK_DISABLE
+		if (World->IsActorTickAndUserCallbacksEnabled())
+#endif
+		{
+			BeginPlay();
+		}
 
 		ensure(BeginPlayCallDepth - 1 == CurrentCallDepth);
 		BeginPlayCallDepth = CurrentCallDepth;
@@ -4482,7 +4494,11 @@ void AActor::DispatchBeginPlay(bool bFromLevelStreaming)
 			World->DestroyActor(this, true); 
 		}
 		
-		if (IsValidChecked(this))
+		if (IsValidChecked(this)
+#if UE_SUPPORT_FOR_ACTOR_TICK_DISABLE
+			&& World->IsActorTickAndUserCallbacksEnabled()
+#endif
+		)
 		{
 			// Initialize overlap state
 			UpdateInitialOverlaps(bFromLevelStreaming);
diff --git a/Engine/Source/Runtime/Engine/Private/LevelTick.cpp b/Engine/Source/Runtime/Engine/Private/LevelTick.cpp
--- a/Engine/Source/Runtime/Engine/Private/LevelTick.cpp
+++ b/Engine/Source/Runtime/Engine/Private/LevelTick.cpp
@@ -1401,6 +1401,9 @@ void UWorld::Tick( ELevelTick TickType, float DeltaSeconds )
 	bool bDoingActorTicks = 
 		(TickType!=LEVELTICK_TimeOnly)
 		&&	!bIsPaused
+#if UE_SUPPORT_FOR_ACTOR_TICK_DISABLE
+		&&  IsActorTickAndUserCallbacksEnabled()
+#endif
 		&&	(!NetDriver || !NetDriver->ServerConnection || NetDriver->ServerConnection->GetConnectionState()==USOCK_Open);
 
 	FLatentActionManager& CurrentLatentActionManager = GetLatentActionManager();
diff --git a/Engine/Source/Runtime/Engine/Private/World.cpp b/Engine/Source/Runtime/Engine/Private/World.cpp
--- a/Engine/Source/Runtime/Engine/Private/World.cpp
+++ b/Engine/Source/Runtime/Engine/Private/World.cpp
@@ -3241,6 +3241,22 @@ bool UWorld::IsLevelMakingInvisible(const ULevel* InLevel) const
 	return bIsMakingInvisible;
 }
 
+#if UE_SUPPORT_FOR_ACTOR_TICK_DISABLE
+void UWorld::EnableActorTickAndUserCallbacks(bool bEnable)
+{
+	if (bEnable)
+	{
+		UE_LOG(LogWorld, Warning, TEXT("EnableActorTickAndUserCallbacks: Enabling actor ticking and actor user callbacks."));
+	}
+	else
+	{
+		UE_LOG(LogWorld, Warning, TEXT("EnableActorTickAndUserCallbacks: Disabling actor ticking and actor user callbacks."));
+	}
+
+	bEnableActorTickAndUserCallbacks = bEnable;
+}
+#endif
+
 void UWorld::SetLevelPendingVisibilityState(ULevel* InLevel, ELevelPendingVisibilityState InState)
 {
 	FLevelPendingVisibilityStateAccessor::SetPendingVisibilityState(InLevel, InState);
