Added accessors for NetDriver::GuidCache which will be made protected in order to trap unexpected usage when running iris replication

diff --git a/Engine/Plugins/Online/OnlineSubsystemUtils/Source/OnlineSubsystemUtils/Private/OnlineBeaconClient.cpp b/Engine/Plugins/Online/OnlineSubsystemUtils/Source/OnlineSubsystemUtils/Private/OnlineBeaconClient.cpp
--- a/Engine/Plugins/Online/OnlineSubsystemUtils/Source/OnlineSubsystemUtils/Private/OnlineBeaconClient.cpp
+++ b/Engine/Plugins/Online/OnlineSubsystemUtils/Source/OnlineSubsystemUtils/Private/OnlineBeaconClient.cpp
@@ -483,7 +483,7 @@ void AOnlineBeaconClient::NotifyControlMessage(UNetConnection* Connection, uint8
 				{
 					if (NetGUID.IsValid())
 					{
-						Connection->Driver->GuidCache->RegisterNetGUID_Client(NetGUID, this);
+						Connection->Driver->GetNetGuidCache()->RegisterNetGUID_Client(NetGUID, this);
 
 						FString BeaconType = GetBeaconType();
 						FNetControlMessage<NMT_BeaconNetGUIDAck>::Send(Connection, BeaconType);
diff --git a/Engine/Plugins/Online/OnlineSubsystemUtils/Source/OnlineSubsystemUtils/Private/OnlineBeaconHost.cpp b/Engine/Plugins/Online/OnlineSubsystemUtils/Source/OnlineSubsystemUtils/Private/OnlineBeaconHost.cpp
--- a/Engine/Plugins/Online/OnlineSubsystemUtils/Source/OnlineSubsystemUtils/Private/OnlineBeaconHost.cpp
+++ b/Engine/Plugins/Online/OnlineSubsystemUtils/Source/OnlineSubsystemUtils/Private/OnlineBeaconHost.cpp
@@ -450,7 +450,7 @@ bool AOnlineBeaconHost::HandleControlMessage(UNetConnection* Connection, uint8 M
 			Connection->SetClientLoginState(EClientLoginState::ReceivedJoin);
 			NewClientActor->SetConnectionState(EBeaconConnectionState::Pending);
 
-			FNetworkGUID NetGUID = Connection->Driver->GuidCache->AssignNewNetGUID_Server(NewClientActor);
+			FNetworkGUID NetGUID = Connection->Driver->GetNetGuidCache()->AssignNewNetGUID_Server(NewClientActor);
 			NewClientActor->SetNetConnection(Connection);
 			Connection->PlayerId = UniqueId;
 			NewClientActor->SetRole(ROLE_Authority);
diff --git a/Engine/Plugins/Runtime/ReplicationGraph/Source/Private/ReplicationGraph.cpp b/Engine/Plugins/Runtime/ReplicationGraph/Source/Private/ReplicationGraph.cpp
--- a/Engine/Plugins/Runtime/ReplicationGraph/Source/Private/ReplicationGraph.cpp
+++ b/Engine/Plugins/Runtime/ReplicationGraph/Source/Private/ReplicationGraph.cpp
@@ -2707,7 +2707,7 @@ void UReplicationGraph::SetActorDestructionInfoToIgnoreDistanceCulling(AActor* D
 	}
 
 	check(NetDriver);
-	FNetworkGUID NetGUID = NetDriver->GuidCache->GetNetGUID(DestroyedActor);
+	FNetworkGUID NetGUID = NetDriver->GetNetGuidCache()->GetNetGUID(DestroyedActor);
 			
 	if (NetGUID.IsDefault())
 	{
@@ -2992,7 +2992,7 @@ void UNetReplicationGraphConnection::NotifyAddDestructionInfo(FActorDestructionI
 void UNetReplicationGraphConnection::NotifyAddDormantDestructionInfo(AActor* Actor)
 {
 #if WITH_SERVER_CODE
-	if (Actor && NetConnection && NetConnection->Driver && NetConnection->Driver->GuidCache)
+	if (Actor && NetConnection && NetConnection->Driver && NetConnection->Driver->GetNetGuidCache())
 	{
 		ULevel* Level = Actor->GetLevel();
 
@@ -3007,7 +3007,7 @@ void UNetReplicationGraphConnection::NotifyAddDormantDestructionInfo(AActor* Act
 			}
 		}
 
-		FNetworkGUID NetGUID = NetConnection->Driver->GuidCache->GetNetGUID(Actor);
+		FNetworkGUID NetGUID = NetConnection->Driver->GetNetGuidCache()->GetNetGUID(Actor);
 		if (NetGUID.IsValid() && !NetGUID.IsDefault())
 		{
 			bool bWasAlreadyTracked = false;
diff --git a/Engine/Source/Runtime/Engine/Classes/Engine/NetDriver.h b/Engine/Source/Runtime/Engine/Classes/Engine/NetDriver.h
--- a/Engine/Source/Runtime/Engine/Classes/Engine/NetDriver.h
+++ b/Engine/Source/Runtime/Engine/Classes/Engine/NetDriver.h
@@ -982,7 +982,8 @@ public:
 	TObjectPtr<class UPackage> WorldPackage;
 
 	/** @todo document */
-	TSharedPtr< class FNetGUIDCache > GuidCache;
+	//UE_DEPRECATED(5.6, "Variable will be made private. Use GetNetGuidCache() methods instead.")
+	TSharedPtr<class FNetGUIDCache> GuidCache;
 
 	TSharedPtr< class FClassNetCacheMgr >	NetCache;
 
@@ -1116,6 +1117,34 @@ public:
 	UE_DEPRECATED(5.6, "SetPendingDestruction will be made private. Use RequestNetDriverDestruction instead")
 	ENGINE_API void SetPendingDestruction(bool bDestroy);
 
+	/**
+	 * Get const GuidCache
+	 * Note: GuidCache/FNetworkGUIDs are not used when using Iris replication.
+	*/
+	const TSharedPtr<class FNetGUIDCache>& GetNetGuidCache() const
+	{
+		// Disabled until we have addressed the issues found.
+		//ensureMsgf(!IsUsingIrisReplication(), TEXT("Using NetGuidCache/FNetworkGUIDs is not supported when using Iris replication."));
+	
+		PRAGMA_DISABLE_DEPRECATION_WARNINGS
+		return GuidCache;
+		PRAGMA_ENABLE_DEPRECATION_WARNINGS
+	}
+
+	/**
+	* Get const GuidCache
+	* Note: GuidCache/FNetworkGUIDs are not used when using Iris replication.
+	*/
+	TSharedPtr<class FNetGUIDCache>& GetNetGuidCache()
+	{
+		// Disabled until we have addressed the issues found.
+		//ensureMsgf(!IsUsingIrisReplication(), TEXT("Using NetGuidCache/FNetworkGUIDs is not supported using Iris replication."));
+	
+		PRAGMA_DISABLE_DEPRECATION_WARNINGS
+		return GuidCache;
+		PRAGMA_ENABLE_DEPRECATION_WARNINGS
+	}
+
 private:
 	double						ElapsedTime;
 
diff --git a/Engine/Source/Runtime/Engine/Private/DataChannel.cpp b/Engine/Source/Runtime/Engine/Private/DataChannel.cpp
--- a/Engine/Source/Runtime/Engine/Private/DataChannel.cpp
+++ b/Engine/Source/Runtime/Engine/Private/DataChannel.cpp
@@ -2304,9 +2304,9 @@ void UControlChannel::ReceiveDestructionInfo(FInBunch& Bunch)
 				if (UE::Net::FilterGuidRemapping != 0)
 				{
 					// Remove this actor's NetGUID from the list of unmapped values, it will be added back if it replicates again
-					if (NetGUID.IsValid() && Connection != nullptr && Connection->Driver != nullptr && Connection->Driver->GuidCache.IsValid())
+					if (NetGUID.IsValid() && Connection != nullptr && Connection->Driver != nullptr && Connection->Driver->GetNetGuidCache().IsValid())
 					{
-						Connection->Driver->GuidCache->ImportedNetGuids.Remove(NetGUID);
+						Connection->Driver->GetNetGuidCache()->ImportedNetGuids.Remove(NetGUID);
 					}
 				}
 
@@ -2546,9 +2546,9 @@ void UActorChannel::DestroyActorAndComponents()
 	if (UE::Net::FilterGuidRemapping != 0)
 	{
 		// Remove this actor's NetGUID from the list of unmapped values, it will be added back if it replicates again
-		if (ActorNetGUID.IsValid() && Connection != nullptr && Connection->Driver != nullptr && Connection->Driver->GuidCache.IsValid())
+		if (ActorNetGUID.IsValid() && Connection != nullptr && Connection->Driver != nullptr && Connection->Driver->GetNetGuidCache().IsValid())
 		{
-			Connection->Driver->GuidCache->ImportedNetGuids.Remove(ActorNetGUID);
+			Connection->Driver->GetNetGuidCache()->ImportedNetGuids.Remove(ActorNetGUID);
 		}
 	}
 }
@@ -2625,7 +2625,7 @@ bool UActorChannel::CleanUp(const bool bForDestroy, EChannelCloseReason CloseRea
 				ReplicatedObject->OnLostRelevancy();
 				if (!ReplicatedObject->ShouldClientDestroyWhenNotRelevant())
 				{
-					UE_LOG(LogNetTraffic, VeryVerbose, TEXT("UActorChannel::CleanUp: Closed for ReplicatedObject Relevancy. %s - [%s] FullNetGUIDPath: %s"), *Describe(), *ActorNetGUID.ToString(), *Connection->Driver->GuidCache->FullNetGUIDPath(ActorNetGUID));
+					UE_LOG(LogNetTraffic, VeryVerbose, TEXT("UActorChannel::CleanUp: Closed for ReplicatedObject Relevancy. %s - [%s] FullNetGUIDPath: %s"), *Describe(), *ActorNetGUID.ToString(), *Connection->Driver->GetNetGuidCache()->FullNetGUIDPath(ActorNetGUID));
 					bHandled = true;
 				}
 			}
@@ -2913,14 +2913,15 @@ bool UActorChannel::ProcessQueuedBunches()
 	// TODO: This could take a non-trivial amount of time since both GetObjectFromNetGUID
 	// and IsGUIDBroken may do Map Lookups and GetObjectFromNetGUID will attempt
 	// to resolve weak objects.
+	FNetGUIDCache* GuidCache = Connection->Driver->GetNetGuidCache().Get();
 	for (auto It = PendingGuidResolves.CreateIterator(); It; ++It)
 	{
-		if (Connection->Driver->GuidCache->GetObjectFromNetGUID(*It, true))
+		if (GuidCache->GetObjectFromNetGUID(*It, true))
 		{
 			// This guid is now resolved, we can remove it from the pending guid list
 			It.RemoveCurrent();
 		}
-		else if (Connection->Driver->GuidCache->IsGUIDBroken(*It, true))
+		else if (GuidCache->IsGUIDBroken(*It, true))
 		{
 			// This guid is broken, remove it, and warn
 			UE_LOG(LogNet, Warning, TEXT("UActorChannel::ProcessQueuedBunches: Guid is broken. NetGUID: %s, ChIndex: %i, Actor: %s"), *It->ToString(), ChIndex, *GetPathNameSafe(Actor));
@@ -2987,7 +2988,7 @@ bool UActorChannel::ProcessQueuedBunches()
 			if (UPackageMapClient * PackageMapClient = Cast< UPackageMapClient >(Connection->PackageMap))
 			{
 #if CSV_PROFILER_STATS
-				FNetGUIDCache::FIsOwnerOrPawnHelper Helper(Connection->Driver->GuidCache.Get(), Connection->OwningActor, Actor);
+				FNetGUIDCache::FIsOwnerOrPawnHelper Helper(GuidCache, Connection->OwningActor, Actor);
 #endif
 
 				PackageMapClient->SetHasQueuedBunches(ActorNetGUID, !QueuedBunches.IsEmpty());
@@ -3011,10 +3012,10 @@ bool UActorChannel::ProcessQueuedBunches()
 					{
 						for (const FNetworkGUID& Guid : PendingGuidResolves)
 						{
-							UE_LOG(LogNet, Log, TEXT("  PendingGuidResolve [%s] %s"), *Guid.ToString(), *Connection->Driver->GuidCache->Describe(Guid));
+							UE_LOG(LogNet, Log, TEXT("  PendingGuidResolve [%s] %s"), *Guid.ToString(), *GuidCache->Describe(Guid));
 
-							const FNetworkGUID OuterGuid = Connection->Driver->GuidCache->GetOuterNetGUID(Guid);
-							UE_LOG(LogNet, Log, TEXT("      OuterGuid [%s] %s"), *OuterGuid.ToString(), *Connection->Driver->GuidCache->Describe(OuterGuid));
+							const FNetworkGUID OuterGuid = GuidCache->GetOuterNetGUID(Guid);
+							UE_LOG(LogNet, Log, TEXT("      OuterGuid [%s] %s"), *OuterGuid.ToString(), *GuidCache->Describe(OuterGuid));
 						}
 					}
 				}
@@ -3074,7 +3075,7 @@ void UActorChannel::ReceivedBunch( FInBunch & Bunch )
 			QueuedObjectsToTrack.Reserve(NumMustBeMappedGUIDs);
 			//UE_LOG( LogNetTraffic, Warning, TEXT( "Read must be mapped GUID's. NumMustBeMappedGUIDs: %i" ), NumMustBeMappedGUIDs );
 
-			FNetGUIDCache* GuidCache = Connection->Driver->GuidCache.Get();
+			FNetGUIDCache* GuidCache = Connection->Driver->GetNetGuidCache().Get();
 
 #if CSV_PROFILER_STATS
 			FNetGUIDCache::FIsOwnerOrPawnHelper Helper(GuidCache, Connection->OwningActor, Actor);
@@ -3178,7 +3179,7 @@ void UActorChannel::ReceivedBunch( FInBunch & Bunch )
 				PackageMapClient->SetHasQueuedBunches(ActorNetGUID, true);
 			}
 
-			if (FNetGUIDCache* LocalGuidCache = Connection->Driver->GuidCache.Get())
+			if (FNetGUIDCache* LocalGuidCache = Connection->Driver->GetNetGuidCache().Get())
 			{
 				for (TPair<FNetworkGUID, UObject*>& NetGUIDObjectPair : QueuedObjectsToTrack)
 				{
@@ -3249,7 +3250,7 @@ void UActorChannel::ProcessBunch( FInBunch & Bunch )
 			}
 
 			check( !bSpawnedNewActor );
-			UE_LOG(LogNet, Warning, TEXT("UActorChannel::ProcessBunch: SerializeNewActor failed to find/spawn actor. Actor: %s, Channel: %i, FullNetGUIDPath: %s"), *GetFullNameSafe(NewChannelActor), ChIndex, *Connection->Driver->GuidCache->FullNetGUIDPath(ActorNetGUID));
+			UE_LOG(LogNet, Warning, TEXT("UActorChannel::ProcessBunch: SerializeNewActor failed to find/spawn actor. Actor: %s, Channel: %i, FullNetGUIDPath: %s"), *GetFullNameSafe(NewChannelActor), ChIndex, *Connection->Driver->GetNetGuidCache()->FullNetGUIDPath(ActorNetGUID));
 			Broken = 1;
 
 			if (!Connection->IsInternalAck()
@@ -3356,7 +3357,7 @@ void UActorChannel::ProcessBunch( FInBunch & Bunch )
 		if ( Reader.GetNumBits() == 0 )
 		{
 			// Set the scope name
-			UE_NET_TRACE_SET_SCOPE_OBJECTID(ContentBlockScope, Connection->Driver->GuidCache->GetNetGUID(RepObj));
+			UE_NET_TRACE_SET_SCOPE_OBJECTID(ContentBlockScope, Connection->Driver->GetNetGuidCache()->GetNetGUID(RepObj));
 
 			// Nothing else in this block, continue on (should have been a delete or create block)
 			continue;
@@ -4819,7 +4820,7 @@ UObject* UActorChannel::ReadContentBlockHeader(FInBunch& Bunch, bool& bObjectDel
 			// (ignore though if this is for replays)
 			if (!Connection->IsInternalAck())
 			{
-				UE_LOG(LogNetTraffic, Log, TEXT("ReadContentBlockHeader: Stably named sub-object not found. Its level may have streamed out. Component: %s, Actor: %s"), *Connection->Driver->GuidCache->FullNetGUIDPath(NetGUID), *Actor->GetName());
+				UE_LOG(LogNetTraffic, Log, TEXT("ReadContentBlockHeader: Stably named sub-object not found. Its level may have streamed out. Component: %s, Actor: %s"), *Connection->Driver->GetNetGuidCache()->FullNetGUIDPath(NetGUID), *Actor->GetName());
 			}
 
 			return nullptr;
@@ -4994,7 +4995,7 @@ UObject* UActorChannel::ReadContentBlockHeader(FInBunch& Bunch, bool& bObjectDel
 		checkf(Cast< AActor >(SubObj) == nullptr, TEXT("UActorChannel::ReadContentBlockHeader: Subobject is an Actor. SubObject: %s, Actor: %s"), *SubObj->GetName(), *Actor->GetName());
 
 		// Register the component guid
-		Connection->Driver->GuidCache->RegisterNetGUID_Client(NetGUID, SubObj);
+		Connection->Driver->GetNetGuidCache()->RegisterNetGUID_Client(NetGUID, SubObj);
 
 		// Notify actor that we created a component from replication
 		Actor->OnSubobjectCreatedFromReplication( SubObj );
@@ -5006,7 +5007,7 @@ UObject* UActorChannel::ReadContentBlockHeader(FInBunch& Bunch, bool& bObjectDel
 		if (ensureMsgf(NetGUID.IsValid(), TEXT("Channel tried to add an invalid GUID to the import list: %s"), *Describe()))
 		{
 			LLM_SCOPE_BYTAG(GuidCache);
-			Connection->Driver->GuidCache->ImportedNetGuids.Add( NetGUID );
+			Connection->Driver->GetNetGuidCache()->ImportedNetGuids.Add( NetGUID );
 		}
 	}
 
@@ -5541,9 +5542,9 @@ bool UActorChannel::WriteSubObjectInBunch(UObject* Obj, FOutBunch& Bunch, FRepli
 	// here, by forcing the packagemap to give them a NetGUID.
 	//
 	// Once we can lazily handle unmapped references on the client side, this can be simplified.
-	if ( !Connection->Driver->GuidCache->SupportsObject( Obj, &WeakObj ) )
+	if ( !Connection->Driver->GetNetGuidCache()->SupportsObject( Obj, &WeakObj ) )
 	{
-		Connection->Driver->GuidCache->AssignNewNetGUID_Server( Obj );	//Make sure it gets a NetGUID so that it is now 'supported'
+		Connection->Driver->GetNetGuidCache()->AssignNewNetGUID_Server( Obj );	//Make sure it gets a NetGUID so that it is now 'supported'
 	}
 
 	FReplicationFlags ObjRepFlags = RepFlags;
@@ -5743,7 +5744,7 @@ static void	FindNetGUID( const TArray<FString>& Args, UWorld* InWorld )
 		if (Args.Num() <= 0)
 		{
 			// Display all
-			for (auto It = Driver->GuidCache->History.CreateIterator(); It; ++It)
+			for (auto It = Driver->GetNetGuidCache()->History.CreateIterator(); It; ++It)
 			{
 				FString Str = It.Value();
 				FNetworkGUID NetGUID = It.Key();
@@ -5762,7 +5763,7 @@ static void	FindNetGUID( const TArray<FString>& Args, UWorld* InWorld )
 			FNetworkGUID NetGUID = FNetworkGUID::CreateFromIndex(NetIndex, bStatic);
 
 			// Search
-			FString Str = Driver->GuidCache->History.FindRef(NetGUID);
+			FString Str = Driver->GetNetGuidCache()->History.FindRef(NetGUID);
 
 			if (!Str.IsEmpty())
 			{
diff --git a/Engine/Source/Runtime/Engine/Private/DataReplication.cpp b/Engine/Source/Runtime/Engine/Private/DataReplication.cpp
--- a/Engine/Source/Runtime/Engine/Private/DataReplication.cpp
+++ b/Engine/Source/Runtime/Engine/Private/DataReplication.cpp
@@ -656,7 +656,7 @@ void FObjectReplicator::StartReplicating(class UActorChannel * InActorChannel)
 	UNetDriver* const ConnectionNetDriver = Connection->GetDriver();
 
 	// Cache off netGUID so if this object gets deleted we can close it
-	ObjectNetGUID = ConnectionNetDriver->GuidCache->GetOrAssignNetGUID( Object );
+	ObjectNetGUID = ConnectionNetDriver->GetNetGuidCache()->GetOrAssignNetGUID( Object );
 
 	const bool bIsValidToReplicate = !ObjectNetGUID.IsDefault() && ObjectNetGUID.IsValid();
 	if (UNLIKELY(!bIsValidToReplicate))
diff --git a/Engine/Source/Runtime/Engine/Private/DemoNetDriver.cpp b/Engine/Source/Runtime/Engine/Private/DemoNetDriver.cpp
--- a/Engine/Source/Runtime/Engine/Private/DemoNetDriver.cpp
+++ b/Engine/Source/Runtime/Engine/Private/DemoNetDriver.cpp
@@ -969,15 +969,15 @@ bool UDemoNetDriver::InitConnect(FNetworkNotify* InNotify, const FURL& ConnectUR
 		return false;
 	}
 
-	GuidCache->SetNetworkChecksumMode(FNetGUIDCache::ENetworkChecksumMode::SaveButIgnore);
+	GetNetGuidCache()->SetNetworkChecksumMode(FNetGUIDCache::ENetworkChecksumMode::SaveButIgnore);
 
 	if (CVarForceDisableAsyncPackageMapLoading.GetValueOnGameThread() > 0)
 	{
-		GuidCache->SetAsyncLoadMode(FNetGUIDCache::EAsyncLoadMode::ForceDisable);
+		GetNetGuidCache()->SetAsyncLoadMode(FNetGUIDCache::EAsyncLoadMode::ForceDisable);
 	}
 	else
 	{
-		GuidCache->SetAsyncLoadMode(FNetGUIDCache::EAsyncLoadMode::UseCVar);
+		GetNetGuidCache()->SetAsyncLoadMode(FNetGUIDCache::EAsyncLoadMode::UseCVar);
 	}
 
 	// Playback, local machine is a client, and the demo stream acts "as if" it's the server.
@@ -2793,11 +2793,11 @@ void UDemoNetDriver::TickDemoPlayback(float DeltaSeconds)
 
 	if (CVarForceDisableAsyncPackageMapLoading.GetValueOnGameThread() > 0)
 	{
-		GuidCache->SetAsyncLoadMode(FNetGUIDCache::EAsyncLoadMode::ForceDisable);
+		GetNetGuidCache()->SetAsyncLoadMode(FNetGUIDCache::EAsyncLoadMode::ForceDisable);
 	}
 	else
 	{
-		GuidCache->SetAsyncLoadMode(FNetGUIDCache::EAsyncLoadMode::UseCVar);
+		GetNetGuidCache()->SetAsyncLoadMode(FNetGUIDCache::EAsyncLoadMode::UseCVar);
 	}
 
 	if (CVarGotoTimeInSeconds.GetValueOnGameThread() >= 0.0f)
@@ -3299,7 +3299,7 @@ void UDemoNetDriver::ReplayStreamingReady(const FStartStreamingResult& Result)
 
 FReplayExternalDataArray* UDemoNetDriver::GetExternalDataArrayForObject(UObject* Object)
 {
-	FNetworkGUID NetworkGUID = GuidCache->NetGUIDLookup.FindRef(Object);
+	FNetworkGUID NetworkGUID = GetNetGuidCache()->NetGUIDLookup.FindRef(Object);
 
 	if (!NetworkGUID.IsValid())
 	{
@@ -4044,7 +4044,7 @@ bool UDemoNetDriver::LoadCheckpoint(const FGotoResult& GotoResult)
 		}
 
 		// Clean package map to prepare to restore it to the checkpoint state
-		GuidCache->ResetCacheForDemo();
+		GetNetGuidCache()->ResetCacheForDemo();
 
 		// Since we only count the number of sub-spectators, add one more slot for main spectator
 		// Very small optimization. We do want to clear this so that we don't end up doing during ProcessSeamlessTravel
@@ -4088,7 +4088,7 @@ bool UDemoNetDriver::LoadCheckpoint(const FGotoResult& GotoResult)
 
 	if (SpectatorController && SpectatorController->GetViewTarget())
 	{
-		ViewTargetGUID = GuidCache->NetGUIDLookup.FindRef(SpectatorController->GetViewTarget());
+		ViewTargetGUID = GetNetGuidCache()->NetGUIDLookup.FindRef(SpectatorController->GetViewTarget());
 
 		if (ViewTargetGUID.IsValid())
 		{
@@ -4161,7 +4161,7 @@ bool UDemoNetDriver::LoadCheckpoint(const FGotoResult& GotoResult)
 			// If an non-startup actor that we don't destroy has an entry in the GuidCache, preserve that entry so
 			// that the object will be re-used after loading the checkpoint. Otherwise, a new copy
 			// of the object will be created each time a checkpoint is loaded, causing a leak.
-			const FNetworkGUID FoundGUID = GuidCache->NetGUIDLookup.FindRef(CurrentActor);
+			const FNetworkGUID FoundGUID = GetNetGuidCache()->NetGUIDLookup.FindRef(CurrentActor);
 				
 			if (FoundGUID.IsValid())
 			{
@@ -4270,7 +4270,7 @@ bool UDemoNetDriver::LoadCheckpoint(const FGotoResult& GotoResult)
 	// Catch a rare case where the spectator controller is null, but a valid GUID is
 	// found on the GuidCache. The weak pointers in the NetGUIDLookup map are probably
 	// going null, and we want catch these cases and investigate further.
-	if (!ensure(GuidCache->NetGUIDLookup.FindRef(SpectatorController).IsValid() == (SpectatorController != nullptr)))
+	if (!ensure(GetNetGuidCache()->NetGUIDLookup.FindRef(SpectatorController).IsValid() == (SpectatorController != nullptr)))
 	{
 		UE_LOG(LogDemo, Log, TEXT("LoadCheckpoint: SpectatorController is null and a valid GUID for null was found in the GuidCache. SpectatorController = %s"),
 			*GetFullNameSafe(SpectatorController));
@@ -4278,19 +4278,19 @@ bool UDemoNetDriver::LoadCheckpoint(const FGotoResult& GotoResult)
 
 	// Clean package map to prepare to restore it to the checkpoint state
 	FlushAsyncLoading();
-	GuidCache->ResetCacheForDemo();
+	GetNetGuidCache()->ResetCacheForDemo();
 
 	// Restore preserved packagemap entries
 	for (const FPreservedNetworkGUIDEntry& PreservedEntry : NetGUIDsToPreserve)
 	{
 		check(PreservedEntry.NetGUID.IsValid());
 		
-		FNetGuidCacheObject& CacheObject = GuidCache->ObjectLookup.FindOrAdd(PreservedEntry.NetGUID);
+		FNetGuidCacheObject& CacheObject = GetNetGuidCache()->ObjectLookup.FindOrAdd(PreservedEntry.NetGUID);
 
 		CacheObject.Object = MakeWeakObjectPtr(const_cast<AActor*>(PreservedEntry.Actor));
 		check(CacheObject.Object != nullptr);
 		CacheObject.bNoLoad = true;
-		GuidCache->NetGUIDLookup.Add(CacheObject.Object, PreservedEntry.NetGUID);
+		GetNetGuidCache()->NetGUIDLookup.Add(CacheObject.Object, PreservedEntry.NetGUID);
 	}
 
 	if (GotoCheckpointArchive->TotalSize() == 0 || GotoCheckpointArchive->TotalSize() == INDEX_NONE)
@@ -4437,7 +4437,7 @@ bool UDemoNetDriver::LoadCheckpoint(const FGotoResult& GotoResult)
 			CacheObject.bNoLoad = (Flags & (1 << 0)) ? true : false;
 			CacheObject.bIgnoreWhenMissing = (Flags & (1 << 1)) ? true : false;		
 
-			GuidCache->ObjectLookup.Add(Guid, CacheObject);
+			GetNetGuidCache()->ObjectLookup.Add(Guid, CacheObject);
 
 			if (GotoCheckpointArchive->IsError())
 			{
@@ -4581,7 +4581,7 @@ bool UDemoNetDriver::LoadCheckpoint(const FGotoResult& GotoResult)
 
 	if (SpectatorController && ViewTargetGUID.IsValid())
 	{
-		AActor* ViewTarget = Cast<AActor>(GuidCache->GetObjectFromNetGUID(ViewTargetGUID, false));
+		AActor* ViewTarget = Cast<AActor>(GetNetGuidCache()->GetObjectFromNetGUID(ViewTargetGUID, false));
 
 		if (ViewTarget)
 		{
@@ -5120,7 +5120,7 @@ void UDemoNetDriver::NotifyActorDestroyed(AActor* Actor, bool IsSeamlessTravel)
 
 	if (bActorRewindable && !IsSeamlessTravel && !bIsRecording)
 	{
-		if (bNetStartup || !TrackedRewindActorsByGUID.Contains(GuidCache->NetGUIDLookup.FindRef(Actor)))
+		if (bNetStartup || !TrackedRewindActorsByGUID.Contains(GetNetGuidCache()->NetGUIDLookup.FindRef(Actor)))
 		{
 			// This may happen during playback due to new versions of code playing captures with old versions.
 			// but this should never happen during recording (otherwise it's likely a game code bug). 
@@ -5191,15 +5191,15 @@ void UDemoNetDriver::CleanupOutstandingRewindActors()
 	{
 		for (const FNetworkGUID& NetGUID : TrackedRewindActorsByGUID)
 		{
-			if (FNetGuidCacheObject* CacheObject = GuidCache->ObjectLookup.Find(NetGUID))
+			if (FNetGuidCacheObject* CacheObject = GetNetGuidCache()->ObjectLookup.Find(NetGUID))
 			{
 				if (AActor* Actor = Cast<AActor>(CacheObject->Object))
 				{
 					// Destroy the actor before removing entries from the GuidCache so its entries are still valid in NotifyActorDestroyed.
 					World->DestroyActor(Actor, true);
 
-					ensureMsgf(GuidCache->NetGUIDLookup.Remove(CacheObject->Object) > 0, TEXT("CleanupOutstandingRewindActors: No entry found for %s in NetGUIDLookup"), *NetGUID.ToString());
-					GuidCache->ObjectLookup.Remove(NetGUID);
+					ensureMsgf(GetNetGuidCache()->NetGUIDLookup.Remove(CacheObject->Object) > 0, TEXT("CleanupOutstandingRewindActors: No entry found for %s in NetGUIDLookup"), *NetGUID.ToString());
+					GetNetGuidCache()->ObjectLookup.Remove(NetGUID);
 					CacheObject->bNoLoad = false;
 				}
 				else
diff --git a/Engine/Source/Runtime/Engine/Private/LevelStreaming.cpp b/Engine/Source/Runtime/Engine/Private/LevelStreaming.cpp
--- a/Engine/Source/Runtime/Engine/Private/LevelStreaming.cpp
+++ b/Engine/Source/Runtime/Engine/Private/LevelStreaming.cpp
@@ -184,9 +184,9 @@ static void NetDriverRenameStreamingLevelPackageForPIE(const UWorld* World, FNam
 
 	for (FNamedNetDriver& Driver : WorldContext->ActiveNetDrivers)
 	{
-		if (Driver.NetDriver && Driver.NetDriver->GuidCache.IsValid())
+		if (Driver.NetDriver && Driver.NetDriver->GetNetGuidCache().IsValid())
 		{
-			for (TPair<FNetworkGUID, FNetGuidCacheObject>& GuidPair : Driver.NetDriver->GuidCache->ObjectLookup)
+			for (TPair<FNetworkGUID, FNetGuidCacheObject>& GuidPair : Driver.NetDriver->GetNetGuidCache()->ObjectLookup)
 			{
 				// Only look for packages, which will have a static GUID and an invalid OuterGUID.
 				const bool bIsPackage = GuidPair.Key.IsStatic() && !GuidPair.Value.OuterGUID.IsValid();
diff --git a/Engine/Source/Runtime/Engine/Private/NetConnection.cpp b/Engine/Source/Runtime/Engine/Private/NetConnection.cpp
--- a/Engine/Source/Runtime/Engine/Private/NetConnection.cpp
+++ b/Engine/Source/Runtime/Engine/Private/NetConnection.cpp
@@ -255,7 +255,7 @@ namespace UE::Net::Connection::Private
 					UActorChannel::CanSubObjectReplicateToClient(Connection->PlayerController, SubObjectInfo.NetCondition, SubObjectInfo.Key, ConditionMap, *NetConditionGroupManager))
 				{
 					Connection->FlushDormancyForObject(Actor, SubObject);
-					TrackFlushedSubObject(InOutFlushedGuids, SubObject, Connection->Driver->GuidCache);
+					TrackFlushedSubObject(InOutFlushedGuids, SubObject, Connection->Driver->GetNetGuidCache());
 				}
 			}
 		}
@@ -570,7 +570,10 @@ void UNetConnection::InitBase(UNetDriver* InDriver,class FSocket* InSocket, cons
 
 	if (ensure(PackageMapClient != nullptr))
 	{
+		// We access GuidCache directly here to not trigger ensure when running iris.
+		PRAGMA_DISABLE_DEPRECATION_WARNINGS
 		PackageMapClient->Initialize(this, Driver->GuidCache);
+		PRAGMA_ENABLE_DEPRECATION_WARNINGS
 		PackageMap = PackageMapClient;
 	}
 
@@ -670,7 +673,7 @@ void UNetConnection::InitConnection(UNetDriver* InDriver, EConnectionState InSta
 
 	// Create package map.
 	auto PackageMapClient = NewObject<UPackageMapClient>(this);
-	PackageMapClient->Initialize(this, Driver->GuidCache);
+	PackageMapClient->Initialize(this, Driver->GetNetGuidCache());
 	PackageMap = PackageMapClient;
 
 	UE_NET_TRACE_CONNECTION_CREATED(NetTraceId, GetConnectionHandle().GetParentConnectionId());
@@ -3703,7 +3706,7 @@ void UNetConnection::DispatchPacket( FBitReader& Reader, int32 PacketId, bool& b
 								FNetworkGUID ActorGUID = IgnoringChannels.FindAndRemoveChecked(Bunch.ChIndex);
 								if (ActorGUID.IsStatic())
 								{
-									UObject* FoundObject = Driver->GuidCache->GetObjectFromNetGUID(ActorGUID, false);
+									UObject* FoundObject = Driver->GetNetGuidCache()->GetObjectFromNetGUID(ActorGUID, false);
 									if (AActor* StaticActor = Cast<AActor>(FoundObject))
 									{
 										DestroyIgnoredActor(StaticActor);
@@ -5276,7 +5279,7 @@ void UNetConnection::FlushDormancy(AActor* Actor)
 					UActorChannel::CanSubObjectReplicateToClient(PlayerController, RepComponentInfo.NetCondition, RepComponentInfo.Key, ConditionMap, *NetConditionGroupManager))
 				{
 					FlushDormancyForObject(Actor, ActorComp);
-					TrackFlushedSubObject(FlushedGuids, ActorComp, Driver->GuidCache);
+					TrackFlushedSubObject(FlushedGuids, ActorComp, Driver->GetNetGuidCache());
 
 					if (bEnableFlushDormantSubObjects)
 					{
diff --git a/Engine/Source/Runtime/Engine/Private/NetDriver.cpp b/Engine/Source/Runtime/Engine/Private/NetDriver.cpp
--- a/Engine/Source/Runtime/Engine/Private/NetDriver.cpp
+++ b/Engine/Source/Runtime/Engine/Private/NetDriver.cpp
@@ -677,7 +677,10 @@ void UNetDriver::PostInitProperties()
 		RoleProperty		= FindFieldChecked<FProperty>( AActor::StaticClass(), TEXT("Role"      ) );
 		RemoteRoleProperty	= FindFieldChecked<FProperty>( AActor::StaticClass(), TEXT("RemoteRole") );
 
+		PRAGMA_DISABLE_DEPRECATION_WARNINGS
 		GuidCache			= TSharedPtr< FNetGUIDCache >( new FNetGUIDCache( this ) );
+		PRAGMA_ENABLE_DEPRECATION_WARNINGS
+
 		NetCache			= TSharedPtr< FClassNetCacheMgr >( new FClassNetCacheMgr() );
 
 		ProfileStats		= FParse::Param(FCommandLine::Get(),TEXT("profilestats"));
@@ -1103,11 +1106,11 @@ void UNetDriver::UpdateUnmappedObjects()
 	CSV_SCOPED_TIMING_STAT(Networking, UpdateUnmappedObjects);
 	SCOPE_CYCLE_COUNTER(STAT_NetUpdateUnmappedObjectsTime);
 
-	if (CVarOptimizedRemapping.GetValueOnAnyThread() && GuidCache.IsValid())
+	if (CVarOptimizedRemapping.GetValueOnAnyThread() && GetNetGuidCache().IsValid())
 	{
 		// Go over recently imported network guids, and see if there are any replicators that need to map them
-		TSet<FNetworkGUID>& ImportedNetGuidsRef = GuidCache->ImportedNetGuids;
-		TMap<FNetworkGUID, TSet<FNetworkGUID>>& PendingOuterNetGuidsRef = GuidCache->PendingOuterNetGuids;
+		TSet<FNetworkGUID>& ImportedNetGuidsRef = GetNetGuidCache()->ImportedNetGuids;
+		TMap<FNetworkGUID, TSet<FNetworkGUID>>& PendingOuterNetGuidsRef = GetNetGuidCache()->PendingOuterNetGuids;
 
 		GetMetrics()->SetInt(UE::Net::Metric::ImportedNetGuids, ImportedNetGuidsRef.Num());
 		GetMetrics()->SetInt(UE::Net::Metric::PendingOuterNetGuids, PendingOuterNetGuidsRef.Num());
@@ -1137,7 +1140,7 @@ void UNetDriver::UpdateUnmappedObjects()
 				const FNetworkGUID NetworkGuid = *It;
 				bool bMappedOrBroken = false;
 
-				if (GuidCache->GetObjectFromNetGUID(NetworkGuid, false) != nullptr)
+				if (GetNetGuidCache()->GetObjectFromNetGUID(NetworkGuid, false) != nullptr)
 				{
 					if (UE::Net::Private::bRemapStableSubobjects)
 					{
@@ -1146,7 +1149,7 @@ void UNetDriver::UpdateUnmappedObjects()
 						// Import any unmapped, stably-named guids that are inners of the GUID that just mapped.
 						// These are tracked separately from the normal ImportedNetGuids since they are often
 						// default subobjects created in constructors and there's no other hook to import them.
-						const TArray<FNetworkGUID>* Inners = GuidCache->FindUnmappedStablyNamedGuidsWithOuter(NetworkGuid);
+						const TArray<FNetworkGUID>* Inners = GetNetGuidCache()->FindUnmappedStablyNamedGuidsWithOuter(NetworkGuid);
 
 						if (Inners)
 						{
@@ -1154,7 +1157,7 @@ void UNetDriver::UpdateUnmappedObjects()
 							PendingGuidsRef.Append(*Inners);
 
 							// Now that they're on the pending import list, they will stay there until mapped. Can remove from the outer-to-inner map.
-							GuidCache->RemoveUnmappedStablyNamedGuidsWithOuter(NetworkGuid);
+							GetNetGuidCache()->RemoveUnmappedStablyNamedGuidsWithOuter(NetworkGuid);
 						}
 					}
 
@@ -1162,7 +1165,7 @@ void UNetDriver::UpdateUnmappedObjects()
 					bMappedOrBroken = true;
 				}
 
-				if (GuidCache->IsGUIDBroken(NetworkGuid, false))
+				if (GetNetGuidCache()->IsGUIDBroken(NetworkGuid, false))
 				{
 					bMappedOrBroken = true;
 				}
@@ -1171,12 +1174,12 @@ void UNetDriver::UpdateUnmappedObjects()
 
 				if (!bMappedOrBroken)
 				{
-					const FNetworkGUID OuterGUID = GuidCache->GetOuterNetGUID(NetworkGuid);
+					const FNetworkGUID OuterGUID = GetNetGuidCache()->GetOuterNetGUID(NetworkGuid);
 
 					// we're missing the outer, stop checking until we map it
-					if ((UE::Net::FilterGuidRemapping != 0) && OuterGUID.IsValid() && !OuterGUID.IsDefault() && !GuidCache->IsGUIDLoaded(OuterGUID) && !GuidCache->IsGUIDPending(OuterGUID))
+					if ((UE::Net::FilterGuidRemapping != 0) && OuterGUID.IsValid() && !OuterGUID.IsDefault() && !GetNetGuidCache()->IsGUIDLoaded(OuterGUID) && !GetNetGuidCache()->IsGUIDPending(OuterGUID))
 					{
-						UE_LOG(LogNetPackageMap, Log, TEXT("Missing outer (%s) for unmapped guid (%s), marking pending"), *GuidCache->Describe(OuterGUID), *GuidCache->Describe(NetworkGuid));
+						UE_LOG(LogNetPackageMap, Log, TEXT("Missing outer (%s) for unmapped guid (%s), marking pending"), *GetNetGuidCache()->Describe(OuterGUID), *GetNetGuidCache()->Describe(NetworkGuid));
 
 						TSet<FNetworkGUID>& PendingGuidsRef = PendingOuterNetGuidsRef.FindOrAdd(OuterGUID);
 						PendingGuidsRef.Add(NetworkGuid);
@@ -1203,7 +1206,7 @@ void UNetDriver::UpdateUnmappedObjects()
 			{
 				if (TSet<FNetworkGUID>* DependentGuids = PendingOuterNetGuidsRef.Find(NetGuid))
 				{
-					UE_LOG(LogNetPackageMap, Log, TEXT("Newly mapped outer (%s) removing from pending"), *GuidCache->FullNetGUIDPath(NetGuid));
+					UE_LOG(LogNetPackageMap, Log, TEXT("Newly mapped outer (%s) removing from pending"), *GetNetGuidCache()->FullNetGUIDPath(NetGuid));
 
 					ImportedNetGuidsRef.Append(*DependentGuids);
 					PendingOuterNetGuidsRef.Remove(NetGuid);
@@ -3199,7 +3202,7 @@ void UNetDriver::Serialize( FArchive& Ar )
 		GRANULAR_NETWORK_MEMORY_TRACKING_TRACK("MappedClientConnection", MappedClientConnections.CountBytes(Ar));
 		GRANULAR_NETWORK_MEMORY_TRACKING_TRACK("RecentlyDisconnectedClients", RecentlyDisconnectedClients.CountBytes(Ar));
 		GRANULAR_NETWORK_MEMORY_TRACKING_TRACK("GuidCache",
-			if (FNetGUIDCache const * const LocalGuidCache = GuidCache.Get())
+			if (FNetGUIDCache const * const LocalGuidCache = GetNetGuidCache().Get())
 			{
 				LocalGuidCache->CountBytes(Ar);
 			}
@@ -3321,7 +3324,7 @@ void UNetDriver::FinishDestroy()
 		LowLevelDestroy();
 
 		// Delete the guid cache
-		GuidCache.Reset();
+		GetNetGuidCache().Reset();
 
 		FWorldDelegates::LevelRemovedFromWorld.Remove(OnLevelRemovedFromWorldHandle);
 		FWorldDelegates::LevelAddedToWorld.Remove(OnLevelAddedToWorldHandle);
@@ -3332,7 +3335,7 @@ void UNetDriver::FinishDestroy()
 	{
 		check(ServerConnection==NULL);
 		check(ClientConnections.Num()==0);
-		check(!GuidCache.IsValid());
+		check(!GetNetGuidCache().IsValid());
 	}
 
 	// Make sure we've properly shut down all of the FObjectReplicator's
@@ -3349,9 +3352,9 @@ void UNetDriver::LowLevelDestroy()
 	// Sever the link with UWorld to ensure we don't tick again
 	SetWorld(NULL);
 
-	if(GuidCache.IsValid())
+	if(GetNetGuidCache().IsValid())
 	{
-		GuidCache->ReportSyncLoadedGUIDs();
+		GetNetGuidCache()->ReportSyncLoadedGUIDs();
 	}
 }
 
@@ -3916,7 +3919,7 @@ bool UNetDriver::SendDestructionInfoForLevelUnloadIfDormant(AActor* ThisActor, U
 
 			FActorDestructionInfo DestructInfo;
 			DestructInfo.DestroyedPosition = ThisActor->GetActorLocation();
-			DestructInfo.NetGUID = GuidCache->GetNetGUID(ThisActor);
+			DestructInfo.NetGUID = GetNetGuidCache()->GetNetGUID(ThisActor);
 			DestructInfo.Level = ThisActor->GetLevel();
 			DestructInfo.ObjOuter = ThisActor->GetOuter();
 			DestructInfo.PathName = ThisActor->GetName();
@@ -3939,7 +3942,7 @@ FActorDestructionInfo* UNetDriver::CreateDestructionInfo(AActor* ThisActor, FAct
 		return DestructionInfo;
 	}
 
-	FNetworkGUID NetGUID = GuidCache->GetOrAssignNetGUID( ThisActor );
+	FNetworkGUID NetGUID = GetNetGuidCache()->GetOrAssignNetGUID( ThisActor );
 	if (NetGUID.IsDefault())
 	{
 		UE_LOG(LogNet, Error, TEXT("CreateDestructionInfo got an invalid NetGUID for %s"), *ThisActor->GetName());
@@ -4030,7 +4033,7 @@ void UNetDriver::NotifyActorDestroyed(AActor* ThisActor, bool IsSeamlessTravel)
 			IsStaticReplicator = ReplicatedObject->IsStaticReplicator();
 		}
 
-		const bool bIsActorStatic = !GuidCache->IsDynamicObject( ThisActor ) || IsStaticReplicator;
+		const bool bIsActorStatic = !GetNetGuidCache()->IsDynamicObject( ThisActor ) || IsStaticReplicator;
 		const bool bActorHasRole = ThisActor->GetRemoteRole() != ROLE_None;
 		const bool bShouldCreateDestructionInfo = bIsServer && bIsActorStatic && bActorHasRole && !IsSeamlessTravel && !GIsReconstructingBlueprintInstances && !UE::Net::ShouldIgnoreStaticActorDestruction();
 
@@ -4174,7 +4177,7 @@ void UNetDriver::NotifyActorRenamed(AActor* ThisActor, UObject* PreviousOuter, F
 	LLM_SCOPE_BYTAG(NetDriver);
 
 	const bool bIsServer = IsServer();
-	const bool bIsActorStatic = !GuidCache->IsDynamicObject(ThisActor);
+	const bool bIsActorStatic = !GetNetGuidCache()->IsDynamicObject(ThisActor);
 	const bool bActorHasRole = ThisActor->GetRemoteRole() != ROLE_None;
 
 #if WITH_EDITOR
@@ -4599,9 +4602,9 @@ void UNetDriver::AddReferencedObjects(UObject* InThis, FReferenceCollector& Coll
 		Collector.AddStableReference(&Pair.Value);
 	}
 
-	if (This->GuidCache.IsValid())
+	if (This->GetNetGuidCache().IsValid())
 	{
-		This->GuidCache->CollectReferences(Collector);
+		This->GetNetGuidCache()->CollectReferences(Collector);
 	}
 }
 
@@ -5569,7 +5572,7 @@ int32 UNetDriver::ServerReplicateActors_ProcessPrioritizedActorsRange( UNetConne
 				// we can't create the channel if the client is in a different world than we are
 				// or the package map doesn't support the actor's class/archetype (or the actor itself in the case of serializable actors)
 				// or it's an editor placed actor and the client hasn't initialized the level it's in
-				if ( Channel == NULL && GuidCache->SupportsObject( Actor->GetClass() ) && GuidCache->SupportsObject( Actor->IsNetStartupActor() ? Actor : Actor->GetArchetype() ) )
+				if ( Channel == NULL && GetNetGuidCache()->SupportsObject( Actor->GetClass() ) && GetNetGuidCache()->SupportsObject( Actor->IsNetStartupActor() ? Actor : Actor->GetArchetype() ) )
 				{
 					if ( bLevelInitializedForActor )
 					{
@@ -6668,7 +6671,7 @@ void UNetDriver::CreateReplicatedStaticActorDestructionInfo(ULevel* Level, const
 	}
 #endif // UE_WITH_IRIS
 
-	FNetworkGUID NetGUID = GuidCache->AssignNewNetGUIDFromPath_Server(Info.PathName.ToString(), Info.ObjOuter.Get(), Info.ObjClass);
+	FNetworkGUID NetGUID = GetNetGuidCache()->AssignNewNetGUIDFromPath_Server(Info.PathName.ToString(), Info.ObjOuter.Get(), Info.ObjClass);
 	if (NetGUID.IsDefault())
 	{
 		UE_LOG(LogNet, Error, TEXT("CreateReplicatedStaticActorDestructionInfo got an invalid NetGUID for %s"), *Info.PathName.ToString());
@@ -6919,9 +6922,9 @@ void UNetDriver::ResetGameWorldState()
 
 void UNetDriver::CleanPackageMaps()
 {
-	if ( GuidCache.IsValid() )
+	if ( GetNetGuidCache().IsValid() )
 	{ 
-		GuidCache->CleanReferences();
+		GetNetGuidCache()->CleanReferences();
 	}
 
 	if (GNetResetAckStatePostSeamlessTravel)
@@ -7975,7 +7978,7 @@ void UNetDriver::NotifyActorTornOff(AActor* Actor)
 
 void UNetDriver::ConsumeAsyncLoadDelinquencyAnalytics(FNetAsyncLoadDelinquencyAnalytics& Out)
 {
-	if (FNetGUIDCache* LocalGuidCache = GuidCache.Get())
+	if (FNetGUIDCache* LocalGuidCache = GetNetGuidCache().Get())
 	{
 		LocalGuidCache->ConsumeAsyncLoadDelinquencyAnalytics(Out);
 	}
@@ -7989,7 +7992,7 @@ const FNetAsyncLoadDelinquencyAnalytics& UNetDriver::GetAsyncLoadDelinquencyAnal
 {
 	static const FNetAsyncLoadDelinquencyAnalytics Empty(0);
 
-	if (FNetGUIDCache const * const LocalGuidCache = GuidCache.Get())
+	if (FNetGUIDCache const * const LocalGuidCache = GetNetGuidCache().Get())
 	{
 		return LocalGuidCache->GetAsyncLoadDelinquencyAnalytics();
 	}
@@ -7999,7 +8002,7 @@ const FNetAsyncLoadDelinquencyAnalytics& UNetDriver::GetAsyncLoadDelinquencyAnal
 
 void UNetDriver::ResetAsyncLoadDelinquencyAnalytics()
 {
-	if (FNetGUIDCache* LocalGuidCache = GuidCache.Get())
+	if (FNetGUIDCache* LocalGuidCache = GetNetGuidCache().Get())
 	{
 		LocalGuidCache->ResetAsyncLoadDelinquencyAnalytics();
 	}
@@ -8056,7 +8059,7 @@ void UNetDriver::MoveMappedObjectToUnmapped(const UObject* Object)
 
 	// Find all replicators that are referencing this object, and make sure to mark the references as unmapped
 	// This is so when/if this object is instantiated again (using same network guid), we can re-establish the old references
-	FNetworkGUID NetGuid = GuidCache->NetGUIDLookup.FindRef(const_cast<UObject*>(Object));
+	FNetworkGUID NetGuid = GetNetGuidCache()->NetGUIDLookup.FindRef(const_cast<UObject*>(Object));
 
 	if (NetGuid.IsValid())
 	{
diff --git a/Engine/Source/Runtime/Engine/Private/PackageMapClient.cpp b/Engine/Source/Runtime/Engine/Private/PackageMapClient.cpp
--- a/Engine/Source/Runtime/Engine/Private/PackageMapClient.cpp
+++ b/Engine/Source/Runtime/Engine/Private/PackageMapClient.cpp
@@ -3588,9 +3588,9 @@ PRAGMA_ENABLE_DEPRECATION_WARNINGS
 
 	FLoadPackageAsyncDelegate LoadPackageCompleteDelegate = FLoadPackageAsyncDelegate::CreateWeakLambda(Driver, [NetDriver = Driver](const FName& PackageName, UPackage* Package, EAsyncLoadingResult::Type Result)
 	{
-		if (NetDriver->GuidCache.IsValid())
+		if (NetDriver->GetNetGuidCache().IsValid())
 		{
-			NetDriver->GuidCache->AsyncPackageCallback(PackageName, Package, Result);
+			NetDriver->GetNetGuidCache()->AsyncPackageCallback(PackageName, Package, Result);
 		}
 	});
 
diff --git a/Engine/Source/Runtime/Engine/Private/PhysicsEngine/PhysicsLogUtil.cpp b/Engine/Source/Runtime/Engine/Private/PhysicsEngine/PhysicsLogUtil.cpp
--- a/Engine/Source/Runtime/Engine/Private/PhysicsEngine/PhysicsLogUtil.cpp
+++ b/Engine/Source/Runtime/Engine/Private/PhysicsEngine/PhysicsLogUtil.cpp
@@ -32,7 +32,7 @@ namespace PhysicsLogUtil
 			{
 				if (const UNetDriver* NetDriver = Actor->GetNetDriver())
 				{
-					if (const TSharedPtr<FNetGUIDCache>& NetGUIDCache = NetDriver->GuidCache)
+					if (const TSharedPtr<FNetGUIDCache>& NetGUIDCache = NetDriver->GetNetGuidCache())
 					{
 						// Network data available - return net mode, net ID, actor and component name
 						const ENetMode NetMode = Actor->GetNetMode();
diff --git a/Engine/Source/Runtime/Engine/Private/ReplayHelper.cpp b/Engine/Source/Runtime/Engine/Private/ReplayHelper.cpp
--- a/Engine/Source/Runtime/Engine/Private/ReplayHelper.cpp
+++ b/Engine/Source/Runtime/Engine/Private/ReplayHelper.cpp
@@ -1439,7 +1439,7 @@ void FReplayHelper::SaveExternalData(UNetConnection* Connection, FArchive& Ar)
 					if (!NetworkGUID.IsValid())
 					{
 						// try the lookup again, it may not have been registered when the data was added
-						NetworkGUID = Connection->Driver->GuidCache->NetGUIDLookup.FindRef(Object);
+						NetworkGUID = Connection->Driver->GetNetGuidCache()->NetGUIDLookup.FindRef(Object);
 					}
 
 					if (NetworkGUID.IsValid())
@@ -1504,7 +1504,7 @@ void FReplayHelper::CacheNetGuids(UNetConnection* Connection)
 		CheckpointSaveContext.NextAmortizedItem = 0;
 		CheckpointSaveContext.NumNetGuidsForRecording = 0;
 
-		for (auto It = Connection->Driver->GuidCache->ObjectLookup.CreateIterator(); It; ++It)
+		for (auto It = Connection->Driver->GetNetGuidCache()->ObjectLookup.CreateIterator(); It; ++It)
 		{
 			FNetworkGUID& NetworkGUID = It.Key();
 			FNetGuidCacheObject& CacheObject = It.Value();
@@ -1623,7 +1623,7 @@ bool FReplayHelper::SetExternalDataForObject(UNetConnection* Connection, UObject
 	check(Connection && Connection->Driver);
 
 	// It's fine if we don't find the guid, we may replicate the actor later this frame which will register it
-	FNetworkGUID NetworkGUID = Connection->Driver->GuidCache->NetGUIDLookup.FindRef(OwningObject);
+	FNetworkGUID NetworkGUID = Connection->Driver->GetNetGuidCache()->NetGUIDLookup.FindRef(OwningObject);
 
 	if (!ExternalDataMap.Contains(OwningObject))
 	{
@@ -2329,7 +2329,7 @@ void FReplayHelper::NotifyActorDestroyed(UNetConnection* Connection, AActor* Act
 	{
 		if (bDeltaCheckpoint)
 		{
-			FNetworkGUID NetGUID = Connection->Driver->GuidCache->NetGUIDLookup.FindRef(Actor);
+			FNetworkGUID NetGUID = Connection->Driver->GetNetGuidCache()->NetGUIDLookup.FindRef(Actor);
 			if (NetGUID.IsValid())
 			{
 				RecordingDeltaCheckpointData.DestroyedDynamicActors.Add(NetGUID);
@@ -2356,7 +2356,7 @@ void FReplayHelper::ResetDeltaCheckpointTracking(UNetConnection* Connection)
 			NetworkObjects.ResetReplayDirtyTracking();
 
 			// reset guid cache
-			if (FNetGUIDCache* GuidCache = Connection->Driver->GuidCache.Get())
+			if (FNetGUIDCache* GuidCache = Connection->Driver->GetNetGuidCache().Get())
 			{
 				GuidCache->ResetReplayDirtyTracking();
 			}
diff --git a/Engine/Source/Runtime/Engine/Private/Tests/NetTestHelpers.cpp b/Engine/Source/Runtime/Engine/Private/Tests/NetTestHelpers.cpp
--- a/Engine/Source/Runtime/Engine/Private/Tests/NetTestHelpers.cpp
+++ b/Engine/Source/Runtime/Engine/Private/Tests/NetTestHelpers.cpp
@@ -494,14 +494,14 @@ UObject* FTestWorlds::FindReplicatedObjectOnClient(UObject* ServerObject, uint32
 	}
 	else
 	{
-		const FNetworkGUID NetGUID = ServerContext.NetDriver->GuidCache->GetNetGUID(ServerObject);
+		const FNetworkGUID NetGUID = ServerContext.NetDriver->GetNetGuidCache()->GetNetGUID(ServerObject);
 
 		if (!NetGUID.IsValid())
 		{
 			ensureMsgf(false, TEXT("FTestWorlds::FindReplicatedObjectOnClient ServerObject: %s is not replicated."), *GetFullNameSafe(ServerObject));
 		}
 
-		return ClientContext.NetDriver->GuidCache->GetObjectFromNetGUID(NetGUID, false);
+		return ClientContext.NetDriver->GetNetGuidCache()->GetObjectFromNetGUID(NetGUID, false);
 	}
 }
 
