Remove useless GetBox

#rnx

diff --git a/Engine/Source/Runtime/Engine/Private/InstancedStaticMesh.cpp b/Engine/Source/Runtime/Engine/Private/InstancedStaticMesh.cpp
--- a/Engine/Source/Runtime/Engine/Private/InstancedStaticMesh.cpp
+++ b/Engine/Source/Runtime/Engine/Private/InstancedStaticMesh.cpp
@@ -2754,7 +2754,7 @@ FBoxSphereBounds UInstancedStaticMeshComponent::CalcBounds(const FTransform& InB
 void UInstancedStaticMeshComponent::CalcAndCacheNavigationBounds()
 {
 	TRACE_CPUPROFILER_EVENT_SCOPE(UInstancedStaticMeshComponent::CalcAndCacheNavigationBounds);
-	GetOrCalcBounds(GetComponentTransform(), EBoundsCacheType::Navigation).GetBox();
+	GetOrCalcBounds(GetComponentTransform(), EBoundsCacheType::Navigation);
 }
 
 FBoxSphereBounds UInstancedStaticMeshComponent::GetOrCalcBounds(const FTransform& InBoundTransform, EBoundsCacheType InType) const
