Fixed typos / inconsistent naming of Overlap Component with Orientation nodes.

diff --git a/Engine/Source/Runtime/Engine/Classes/Kismet/KismetSystemLibrary.h b/Engine/Source/Runtime/Engine/Classes/Kismet/KismetSystemLibrary.h
--- a/Engine/Source/Runtime/Engine/Classes/Kismet/KismetSystemLibrary.h
+++ b/Engine/Source/Runtime/Engine/Classes/Kismet/KismetSystemLibrary.h
@@ -1112,7 +1112,7 @@ class UKismetSystemLibrary : public UBlueprintFunctionLibrary
 	 * @param OutActors		Returned array of actors. Unsorted.
 	 * @return				true if there was an overlap that passed the filters, false otherwise.
 	 */
-	UFUNCTION(BlueprintCallable, Category="Collision", meta=(WorldContext="WorldContextObject", AutoCreateRefTerm="ActorsToIgnore", DisplayName="Box Overlap Components With Rotation"))
+	UFUNCTION(BlueprintCallable, Category="Collision", meta=(WorldContext="WorldContextObject", AutoCreateRefTerm="ActorsToIgnore", DisplayName="Box Overlap Components With Orientation"))
 	static ENGINE_API bool BoxOverlapComponentsWithOrientation(const UObject* WorldContextObject, const FVector BoxPos, FVector Extent, FRotator Orientation, const TArray<TEnumAsByte<EObjectTypeQuery> > & ObjectTypes, UClass* ComponentClassFilter, const TArray<AActor*>& ActorsToIgnore, TArray<class UPrimitiveComponent*>& OutComponents);
 
 
@@ -1175,7 +1175,7 @@ class UKismetSystemLibrary : public UBlueprintFunctionLibrary
 	 * @param OutActors		Returned array of actors. Unsorted.
 	 * @return				true if there was an overlap that passed the filters, false otherwise.
 	 */
-	UFUNCTION(BlueprintCallable, Category="Collision", meta=(WorldContext="WorldContextObject", AutoCreateRefTerm="ActorsToIgnore", DisplayName="Capsule Overlap Components With Orietnation") )
+	UFUNCTION(BlueprintCallable, Category="Collision", meta=(WorldContext="WorldContextObject", AutoCreateRefTerm="ActorsToIgnore", DisplayName="Capsule Overlap Components With Orientation") )
 	static ENGINE_API bool CapsuleOverlapComponentsWithOrientation(const UObject* WorldContextObject, const FVector CapsulePos, float Radius, float HalfHeight, FRotator Orientation, const TArray<TEnumAsByte<EObjectTypeQuery> > & ObjectTypes, UClass* ComponentClassFilter, const TArray<AActor*>& ActorsToIgnore, TArray<class UPrimitiveComponent*>& OutComponents);
 
 	/**
