LocalFogVolume - updated actor link to official documentation.

#rnx
#tests Editor PC

diff --git a/Engine/Source/Runtime/Engine/Classes/Engine/LocalFogVolume.h b/Engine/Source/Runtime/Engine/Classes/Engine/LocalFogVolume.h
--- a/Engine/Source/Runtime/Engine/Classes/Engine/LocalFogVolume.h
+++ b/Engine/Source/Runtime/Engine/Classes/Engine/LocalFogVolume.h
@@ -10,7 +10,7 @@ class ULocalFogVolumeComponent;
 
 /**
  *	Actor used to position a local fog volume in the scene.
- *	@see https://docs.unrealengine.com/???
+ *	@see https://dev.epicgames.com/documentation/en-us/unreal-engine/local-fog-volumes-in-unreal-engine
  */
 UCLASS(showcategories = (Movement, Rendering, Transformation, DataLayers, "Input|MouseInput", "Input|TouchInput"), ClassGroup = Fog, hidecategories = (Info, Object, Input), MinimalAPI)
 class ALocalFogVolume : public AInfo
