Fixing transactional async physics queries.

diff --git a/Engine/Source/Runtime/Engine/Classes/Engine/World.h b/Engine/Source/Runtime/Engine/Classes/Engine/World.h
--- a/Engine/Source/Runtime/Engine/Classes/Engine/World.h
+++ b/Engine/Source/Runtime/Engine/Classes/Engine/World.h
@@ -2466,6 +2466,10 @@ public:
 	bool IsTraceHandleValid(const FTraceHandle& Handle, bool bOverlapTrace);
 
 private:
+	friend class FPhysicsAsyncQueryReporter;
+	void AddTraceData(const FTraceHandle& Handle, const TArray<FHitResult>& Results);
+	void AddOverlapData(const FTraceHandle& Handle, const TArray<FOverlapResult>& Results);
+	
 	static void GetCollisionProfileChannelAndResponseParams(FName ProfileName, ECollisionChannel& CollisionChannel, FCollisionResponseParams& ResponseParams);
 
 public:
diff --git a/Engine/Source/Runtime/Engine/Private/Collision/Collision.cpp b/Engine/Source/Runtime/Engine/Private/Collision/Collision.cpp
--- a/Engine/Source/Runtime/Engine/Private/Collision/Collision.cpp
+++ b/Engine/Source/Runtime/Engine/Private/Collision/Collision.cpp
@@ -180,8 +180,7 @@ FORCEINLINE_DEBUGGABLE void FCollisionQueryParams::Internal_AddIgnoredComponent(
 {
 	if (InIgnoreComponent && (!CheckForCollision(InIgnoreComponent) || IsQueryCollisionEnabled(InIgnoreComponent)))
 	{
-		IgnoreComponents.Add(InIgnoreComponent->GetUniqueID());
-		bComponentListUnique = false;
+		Internal_AddIgnoredComponent(InIgnoreComponent->GetUniqueID());
 	}
 }
 
@@ -221,6 +220,12 @@ void FCollisionQueryParams::AddIgnoredComponent_LikelyDuplicatedRoot(const UPrim
 	}
 }
 
+void FCollisionQueryParams::Internal_AddIgnoredComponent(const uint32 ComponentID)
+{
+	IgnoreComponents.Add(ComponentID);
+	bComponentListUnique = false;
+}
+
 const FCollisionQueryParams::IgnoreComponentsArrayType& FCollisionQueryParams::GetIgnoredComponents() const
 {
 	if (!bComponentListUnique)
diff --git a/Engine/Source/Runtime/Engine/Private/Collision/WorldCollisionAsync.cpp b/Engine/Source/Runtime/Engine/Private/Collision/WorldCollisionAsync.cpp
--- a/Engine/Source/Runtime/Engine/Private/Collision/WorldCollisionAsync.cpp
+++ b/Engine/Source/Runtime/Engine/Private/Collision/WorldCollisionAsync.cpp
@@ -14,6 +14,7 @@
 #include "Physics/Experimental/PhysInterface_Chaos.h"
 #include "ProfilingDebugging/CsvProfiler.h"
 #include "ProfilingDebugging/ScopedTimers.h"
+#include "UObject/RemoteObjectTransfer.h"
 
 CSV_DEFINE_CATEGORY(WorldCollision, true);
 DEFINE_LOG_CATEGORY_STATIC(LogWorldCollision, Log, All);
@@ -77,6 +78,10 @@ namespace
 	template <>                   FORCEINLINE int32& GetTraceIndex<FTraceDatum>(AsyncTraceData& DataBuffer) { return DataBuffer.NumQueuedTraceData; }
 	template <>                   FORCEINLINE int32& GetTraceIndex<FOverlapDatum>(AsyncTraceData& DataBuffer) { return DataBuffer.NumQueuedOverlapData; }
 
+	template <typename DatumType> FORCEINLINE FTransactionalAsyncTraceBuffer<DatumType>& GetTransactionalData(AsyncTraceData& DataBuffer);
+	template <> FORCEINLINE FTransactionalAsyncTraceBuffer<FTraceDatum>& GetTransactionalData(AsyncTraceData& DataBuffer) { return DataBuffer.TransactionalTraceData; }
+	template <> FORCEINLINE FTransactionalAsyncTraceBuffer<FOverlapDatum>& GetTransactionalData(AsyncTraceData& DataBuffer) { return DataBuffer.TransactionalOverlapData; }
+
 	/** For referencing a thread data buffer and a datum within it */
 	struct FBufferIndexPair
 	{
@@ -205,6 +210,57 @@ namespace
 		}
 	}
 
+	EPhysicsQueryKind TraceTypeToQueryKind(EAsyncTraceType TraceType)
+	{
+		if (TraceType == EAsyncTraceType::Multi)
+		{
+			return EPhysicsQueryKind::Multi;
+		}
+		else if (TraceType == EAsyncTraceType::Single)
+		{
+			return EPhysicsQueryKind::Single;
+		}
+		else
+		{
+			return EPhysicsQueryKind::Test;
+		}
+	}
+
+	void RunTransactionalTraceTask(UPhysicsQueryHandler& QueryHandler, UWorld* World, FTraceDatum& TraceData, const FTraceHandle& TraceHandle)
+	{
+		TraceData.OutHits.Empty();
+
+		if (TraceData.PhysWorld.IsValid())
+		{
+			EPhysicsQueryKind QueryKind = TraceTypeToQueryKind(TraceData.TraceType);
+			Chaos::EThreadQueryContext ThreadContext = Chaos::EThreadQueryContext::GTData;
+			Chaos::FCommonQueryData CommonData;
+			CommonData.TraceChannel = TraceData.TraceChannel;
+			CommonData.Params = TraceData.CollisionParams.CollisionQueryParam;
+			CommonData.ResponseParams = TraceData.CollisionParams.ResponseParam;
+			CommonData.ObjectParams = TraceData.CollisionParams.ObjectQueryParam;
+
+			if ((TraceData.CollisionParams.CollisionShape.ShapeType == ECollisionShape::Line) || TraceData.CollisionParams.CollisionShape.IsNearlyZero())
+			{
+				Chaos::FRayQueryData RayData;
+				RayData.Start = TraceData.Start;
+				RayData.End = TraceData.End;
+				
+				QueryHandler.QueueAsyncRaycast(TraceHandle, QueryKind, ThreadContext, World, RayData, CommonData);
+			}
+			else
+			{
+				Chaos::FSweepQueryData SweepData;
+				SweepData.Start = TraceData.Start;
+				SweepData.End = TraceData.End;
+				SweepData.GeomRot = TraceData.Rot;
+				SweepData.QueryShape.CollisionShape = TraceData.CollisionParams.CollisionShape;
+
+				QueryHandler.QueueAsyncSweep(TraceHandle, QueryKind, ThreadContext, World, SweepData, CommonData);
+			}
+		}
+	}
+
 	void RunTraceTask(FOverlapDatum* OverlapDataBuffer, int32 TotalCount)
 	{
 		check(OverlapDataBuffer);
@@ -232,6 +288,29 @@ namespace
 		}
 	}
 
+	void RunTransactionalTraceTask(UPhysicsQueryHandler& QueryHandler, UWorld* World, FOverlapDatum& OverlapDatum, const FTraceHandle& TraceHandle)
+	{
+		OverlapDatum.OutOverlaps.Empty();
+
+		if (!OverlapDatum.PhysWorld.IsValid())
+		{
+			return;
+		}
+
+		Chaos::EThreadQueryContext ThreadContext = Chaos::EThreadQueryContext::GTData;
+		Chaos::EQueryInfo QueryInfo = Chaos::EQueryInfo::GatherAll;
+		Chaos::FOverlapQueryData OverlapData;
+		OverlapData.GeomPose = FTransform(OverlapDatum.Rot, OverlapDatum.Pos);
+		OverlapData.QueryShape.CollisionShape = OverlapDatum.CollisionParams.CollisionShape;
+		Chaos::FCommonQueryData CommonData;
+		CommonData.TraceChannel = OverlapDatum.TraceChannel;
+		CommonData.Params = OverlapDatum.CollisionParams.CollisionQueryParam;
+		CommonData.ResponseParams = OverlapDatum.CollisionParams.ResponseParam;
+		CommonData.ObjectParams = OverlapDatum.CollisionParams.ObjectQueryParam;
+
+		QueryHandler.QueueAsyncOverlap(TraceHandle, QueryInfo, ThreadContext, World, OverlapData, CommonData);
+	}
+
 	FAutoConsoleTaskPriority CPrio_FAsyncTraceTask(
 		TEXT("TaskGraph.TaskPriorities.AsyncTraceTask"),
 		TEXT("Task and thread priority for async traces."),
@@ -336,6 +415,37 @@ namespace
 			RunTraceTask(Datum, Next.Index);
 		}
 	}
+	
+	template <typename DatumType>
+	void ExecuteAsyncTransactionalTraceIfAvailable(UWorld* World, FWorldAsyncTraceState& State, bool bExecuteAll)
+	{
+		AsyncTraceData& FrameBuffer = State.GetBufferForCurrentFrame();
+
+		FTransactionalAsyncTraceBuffer<DatumType>& TransactionalData = GetTransactionalData<DatumType>(FrameBuffer);
+
+		// Nothing to run
+		if (TransactionalData.NumQueued == 0)
+		{
+			return;
+		}
+
+		// Can only run if we have a query handler
+		UPhysicsQueryHandler* QueryHandler = World->PhysicsQueryHandler;
+		if (QueryHandler == nullptr)
+		{
+			return;
+		}
+
+		for (int32 Index = 0; Index < TransactionalData.NumQueued; ++Index)
+		{
+			FBufferIndexPair Next(Index);
+			DatumType* DatumBuffer = TransactionalData.Data[Next.Block]->Buffer;
+			DatumType& Datum = DatumBuffer[Next.Index];
+
+			FTraceHandle TraceHandle(State.CurrentFrame, Index, true);
+			RunTransactionalTraceTask(*QueryHandler, World, Datum, TraceHandle);
+		}
+	}
 
 	template <typename DatumType>
 	FTraceHandle StartNewTrace(FWorldAsyncTraceState& State, const DatumType& Val)
@@ -349,27 +459,59 @@ namespace
 		// Check we're allowed to do an async call here
 		check(DataBuffer.bAsyncAllowed);
 
-		auto&  TraceData  = GetTraceContainer<DatumType>(DataBuffer);
-		int32& TraceIndex = GetTraceIndex<DatumType>(DataBuffer);
+		// Handle transactional async queries separately. These need to be thrown into a separate queue 
+		// that is processed outside of a transaction (we cannot process in "batches" like normal async).
+		if (AutoRTFM::IsClosed())
+		{
+			FTransactionalAsyncTraceBuffer<DatumType>& TransactionalData = GetTransactionalData<DatumType>(DataBuffer);
+			TransactionalData.Data.Add(MakeUnique<TTraceThreadData<DatumType>>());
+			
+			TArray<TUniquePtr<TTraceThreadData<DatumType>>>& TraceData = TransactionalData.Data;
+			int32& TraceIndex = TransactionalData.NumQueued;
 
-		// we calculate index as continuous, not as each chunk, but continuous way 
-		int32 LastAvailableIndex = TraceData.Num() * ASYNC_TRACE_BUFFER_SIZE;
+			// Indices are calculated as if the buffer was one contiguous array. Convert the chunked array size to the contiguous size.
+			int32 LastAvailableIndex = TraceData.Num() * ASYNC_TRACE_BUFFER_SIZE;
 
-		// if smaller than next available index
-		if (LastAvailableIndex <= TraceIndex)
-		{
-			// add one more buffer
-			TraceData.Add(MakeUnique<TTraceThreadData<DatumType>>());
+			// if smaller than next available index
+			if (LastAvailableIndex <= TraceIndex)
+			{
+				// add one more buffer
+				TraceData.Add(MakeUnique<TTraceThreadData<DatumType>>());
+				// We just resized to account for the next item. This shouldn't be able to fail.
+				check(TraceData.Num() * ASYNC_TRACE_BUFFER_SIZE > TraceIndex);
+			}
+
+			FTraceHandle Result(State.CurrentFrame, TraceIndex, true);
+			FBufferIndexPair(TraceIndex).DatumLookupChecked(TraceData) = Val;
+
+			++TraceIndex;
+
+			return Result;
 		}
+		else
+		{
+			TArray<TUniquePtr<TTraceThreadData<DatumType>>>& TraceData = GetTraceContainer<DatumType>(DataBuffer);
+			int32& TraceIndex = GetTraceIndex<DatumType>(DataBuffer);
 
-		FTraceHandle Result(State.CurrentFrame, TraceIndex);
-		FBufferIndexPair(TraceIndex).DatumLookupChecked(TraceData) = Val;
+			// we calculate index as continuous, not as each chunk, but continuous way 
+			int32 LastAvailableIndex = TraceData.Num() * ASYNC_TRACE_BUFFER_SIZE;
 
-		ExecuteAsyncTraceIfAvailable<DatumType>(State, false);
+			// if smaller than next available index
+			if (LastAvailableIndex <= TraceIndex)
+			{
+				// add one more buffer
+				TraceData.Add(MakeUnique<TTraceThreadData<DatumType>>());
+			}
 
-		++TraceIndex;
+			FTraceHandle Result(State.CurrentFrame, TraceIndex);
+			FBufferIndexPair(TraceIndex).DatumLookupChecked(TraceData) = Val;
 
-		return Result;
+			ExecuteAsyncTraceIfAvailable<DatumType>(State, false);
+
+			++TraceIndex;
+
+			return Result;
+		}
 	}
 }
 
@@ -454,11 +596,11 @@ bool UWorld::IsTraceHandleValid(const FTraceHandle& Handle, bool bOverlapTrace)
 	FBufferIndexPair Loc(Handle._Data.Index);
 	if (bOverlapTrace)
 	{
-		return !!Loc.DatumLookup(DataBuffer.OverlapData);
+		return !!Loc.DatumLookup(Handle.IsTransactional() ? DataBuffer.TransactionalOverlapData.Data : DataBuffer.OverlapData);
 	}
 	else
 	{
-		return !!Loc.DatumLookup(DataBuffer.TraceData);
+		return !!Loc.DatumLookup(Handle.IsTransactional() ? DataBuffer.TransactionalTraceData.Data : DataBuffer.TraceData);
 	}
 }
 
@@ -475,7 +617,9 @@ bool UWorld::QueryTraceData(const FTraceHandle& Handle, FTraceDatum& OutData)
 	{
 		return false;
 	}
-	if (auto* Data = FBufferIndexPair(Handle._Data.Index).DatumLookup(DataBuffer.TraceData))
+
+	TArray<TUniquePtr<TTraceThreadData<FTraceDatum>>>& TraceData = Handle.IsTransactional() ? DataBuffer.TransactionalTraceData.Data : DataBuffer.TraceData;
+	if (auto* Data = FBufferIndexPair(Handle._Data.Index).DatumLookup(TraceData))
 	{
 		OutData = *Data;
 		return true;
@@ -496,7 +640,8 @@ bool UWorld::QueryOverlapData(const FTraceHandle& Handle, FOverlapDatum& OutData
 	{
 		return false;
 	}
-	if (auto* Data = FBufferIndexPair(Handle._Data.Index).DatumLookup(DataBuffer.OverlapData))
+	TArray<TUniquePtr<TTraceThreadData<FOverlapDatum>>>& OverlapData = Handle.IsTransactional() ? DataBuffer.TransactionalOverlapData.Data : DataBuffer.OverlapData;
+	if (auto* Data = FBufferIndexPair(Handle._Data.Index).DatumLookup(OverlapData))
 	{
 		OutData = *Data;
 		return true;
@@ -505,6 +650,51 @@ bool UWorld::QueryOverlapData(const FTraceHandle& Handle, FOverlapDatum& OutData
 	return false;
 }
 
+AsyncTraceData* GetTraceDataForFrame(FWorldAsyncTraceState& AsyncTraceState, int32 FrameNumber)
+{
+	if (FrameNumber == AsyncTraceState.CurrentFrame)
+	{
+		return &AsyncTraceState.GetBufferForCurrentFrame();
+	}
+	else if (FrameNumber == AsyncTraceState.CurrentFrame - 1)
+	{
+		return &AsyncTraceState.GetBufferForPreviousFrame();
+	}
+	return nullptr;
+}
+
+void UWorld::AddTraceData(const FTraceHandle& Handle, const TArray<FHitResult>& Results)
+{
+	AsyncTraceData* TraceData = GetTraceDataForFrame(AsyncTraceState, Handle._Data.FrameNumber);
+	if (TraceData == nullptr)
+	{
+		return;
+	}
+
+	AsyncTraceData& DataBuffer = *TraceData;
+	FBufferIndexPair Pair(Handle._Data.Index);
+	if (FTraceDatum* Datum = Pair.DatumLookup(DataBuffer.TransactionalTraceData.Data))
+	{
+		Datum->OutHits.Append(Results);
+	}
+}
+
+void UWorld::AddOverlapData(const FTraceHandle& Handle, const TArray<FOverlapResult>& Results)
+{
+	AsyncTraceData* TraceData = GetTraceDataForFrame(AsyncTraceState, Handle._Data.FrameNumber);
+	if (TraceData == nullptr)
+	{
+		return;
+	}
+
+	AsyncTraceData& DataBuffer = *TraceData;
+	FBufferIndexPair Pair(Handle._Data.Index);
+	if (FOverlapDatum* Datum = Pair.DatumLookup(DataBuffer.TransactionalOverlapData.Data))
+	{
+		Datum->OutOverlaps.Append(Results);
+	}
+}
+
 void UWorld::WaitForAllAsyncTraceTasks()
 {
 	const bool bRunAsyncTraceOnWorkerThread = AsyncTraceCVars::IsAsyncTraceOnWorkerThreads();
@@ -547,6 +737,17 @@ void UWorld::ResetAsyncTrace()
 		TraceData.Delegate.ExecuteIfBound(FTraceHandle(TraceData.FrameNumber, Idx), TraceData);
 	}
 
+	for (int32 Idx = 0; Idx != DataBufferExecuted.TransactionalTraceData.NumQueued; ++Idx)
+	{
+		FTraceDatum& TraceData = FBufferIndexPair(Idx).DatumLookupChecked(DataBufferExecuted.TransactionalTraceData.Data);
+		TraceData.Delegate.ExecuteIfBound(FTraceHandle(TraceData.FrameNumber, Idx, true), TraceData);
+	}
+	for (int32 Idx = 0; Idx != DataBufferExecuted.TransactionalOverlapData.NumQueued; ++Idx)
+	{
+		FOverlapDatum& OverlapData = FBufferIndexPair(Idx).DatumLookupChecked(DataBufferExecuted.TransactionalOverlapData.Data);
+		OverlapData.Delegate.ExecuteIfBound(FTraceHandle(OverlapData.FrameNumber, Idx, true), OverlapData);
+	}
+
 #if !UE_BUILD_SHIPPING
 	Timer.Stop();
 
@@ -591,6 +792,26 @@ void UWorld::FinishAsyncTrace()
 	ExecuteAsyncTraceIfAvailable<FTraceDatum>  (AsyncTraceState, true);
 	ExecuteAsyncTraceIfAvailable<FOverlapDatum>(AsyncTraceState, true);
 
+#if UE_WITH_REMOTE_OBJECT_HANDLE
+	auto DoWorkCallback = [this]()
+	{
+		// We need to add all requests inside of a transaction, otherwise when a transaction is 
+		// aborted because other work needs to be processed we would lose all active queries.
+		if (!PhysicsQueryHandler->AreAsyncRequestsAdded())
+		{
+			ExecuteAsyncTransactionalTraceIfAvailable<FTraceDatum>(this, AsyncTraceState, true);
+			ExecuteAsyncTransactionalTraceIfAvailable<FOverlapDatum>(this, AsyncTraceState, true);
+		}
+
+		// Make sure all requests are done, otherwise abort. If this transaction succeeds then it will commit back the results.
+		PhysicsQueryHandler->VerifyAsyncRequestsAreCompletedOrAbort();
+	};
+	if (PhysicsQueryHandler != nullptr)
+	{
+		UE::RemoteObject::Transfer::ExecuteTransactional(DoWorkCallback);
+	}
+#endif
+
 	// this flag only needed to know I can't accept any more new request in current frame
 	AsyncTraceState.GetBufferForCurrentFrame().bAsyncAllowed = false;
 
@@ -602,6 +823,10 @@ void UWorld::FinishAsyncTrace()
 	NewAsyncBuffer.bAsyncAllowed = true;
 	NewAsyncBuffer.NumQueuedTraceData = 0;
 	NewAsyncBuffer.NumQueuedOverlapData = 0;
+
+	NewAsyncBuffer.TransactionalTraceData.NumQueued = 0;
+	NewAsyncBuffer.TransactionalOverlapData.NumQueued = 0;
+
 	NewAsyncBuffer.bAsyncTasksCompleted = false;
 
 }
diff --git a/Engine/Source/Runtime/Engine/Private/PhysicsEngine/PhysicsQueryHandler.cpp b/Engine/Source/Runtime/Engine/Private/PhysicsEngine/PhysicsQueryHandler.cpp
--- a/Engine/Source/Runtime/Engine/Private/PhysicsEngine/PhysicsQueryHandler.cpp
+++ b/Engine/Source/Runtime/Engine/Private/PhysicsEngine/PhysicsQueryHandler.cpp
@@ -32,7 +32,28 @@ bool UPhysicsQueryHandler::SweepSingle(const EThreadQueryContext ThreadContext,
 	return false;
 }
 
-bool UPhysicsQueryHandler::SweepMulti(const EThreadQueryContext ThreadContext, const UWorld* World, const FSweepQueryData& RaySweepData, const FCommonQueryData& CommonData, TArray<FHitResult>& OutHits)
+bool UPhysicsQueryHandler::SweepMulti(const EThreadQueryContext ThreadContext, const UWorld* World, const FSweepQueryData& SweepData, const FCommonQueryData& CommonData, TArray<FHitResult>& OutHits)
 {
 	return false;
 }
+
+void UPhysicsQueryHandler::QueueAsyncOverlap(const FTraceHandle& TraceHandle, const EQueryInfo InfoType, const EThreadQueryContext ThreadContext, const UWorld* World, const FOverlapQueryData& OverlapData, const FCommonQueryData& CommonData)
+{
+}
+
+void UPhysicsQueryHandler::QueueAsyncRaycast(const FTraceHandle& TraceHandle, const EPhysicsQueryKind QueryKind, const EThreadQueryContext ThreadContext, const UWorld* World, const FRayQueryData& RayData, const FCommonQueryData& CommonData)
+{
+}
+
+void UPhysicsQueryHandler::QueueAsyncSweep(const FTraceHandle& TraceHandle, const EPhysicsQueryKind QueryKind, const EThreadQueryContext ThreadContext, const UWorld* World, const FSweepQueryData& SweepData, const FCommonQueryData& CommonData)
+{
+}
+
+bool UPhysicsQueryHandler::AreAsyncRequestsAdded() 
+{
+	return false;
+}
+
+void UPhysicsQueryHandler::VerifyAsyncRequestsAreCompletedOrAbort()
+{
+}
diff --git a/Engine/Source/Runtime/Engine/Public/CollisionQueryParams.h b/Engine/Source/Runtime/Engine/Public/CollisionQueryParams.h
--- a/Engine/Source/Runtime/Engine/Public/CollisionQueryParams.h
+++ b/Engine/Source/Runtime/Engine/Public/CollisionQueryParams.h
@@ -232,6 +232,9 @@ public:
 	 */
 	ENGINE_API void AddIgnoredComponent_LikelyDuplicatedRoot(const UPrimitiveComponent* InIgnoreComponent);
 
+	/** Add a component (by id) for this trace to ignore. Internal method meant for copying data in/out*/
+	UE_INTERNAL ENGINE_API void Internal_AddIgnoredComponent(const uint32 ComponentID);
+
 	FString ToString() const
 	{
 		return FString::Printf(TEXT("[%s:%s] TraceComplex(%d)"), *OwnerTag.ToString(), *TraceTag.ToString(), bTraceComplex );
diff --git a/Engine/Source/Runtime/Engine/Public/Physics/PhysicsQueryHandler.h b/Engine/Source/Runtime/Engine/Public/Physics/PhysicsQueryHandler.h
--- a/Engine/Source/Runtime/Engine/Public/Physics/PhysicsQueryHandler.h
+++ b/Engine/Source/Runtime/Engine/Public/Physics/PhysicsQueryHandler.h
@@ -11,6 +11,18 @@
 
 struct FOverlapResult;
 struct FHitResult;
+struct FTraceHandle;
+
+UENUM()
+enum class EPhysicsQueryKind : uint8
+{
+	/** A physics query that is only a test (no results are returned). */
+	Test = 0,
+	/** A physics query that only returns the first hit. */
+	Single,
+	/** A physics query that returns all hits up to a blocking hit. */
+	Multi,
+};
 
 UCLASS(hidecategories=(Object), MinimalAPI)
 // Allows custom handling of physics queries. Primarily for handling network/server redirection.
@@ -35,5 +47,13 @@ public:
 
 	ENGINE_API virtual bool SweepTest(const EThreadQueryContext ThreadContext, const UWorld* World, const FSweepQueryData& SweepData, const FCommonQueryData& CommonData, FHitResult& OutHits);
 	ENGINE_API virtual bool SweepSingle(const EThreadQueryContext ThreadContext, const UWorld* World, const FSweepQueryData& SweepData, const FCommonQueryData& CommonData, FHitResult& OutHits);
-	ENGINE_API virtual bool SweepMulti(const EThreadQueryContext ThreadContext, const UWorld* World, const FSweepQueryData& RaySweepData, const FCommonQueryData& CommonData, TArray<FHitResult>& OutHits);
+	ENGINE_API virtual bool SweepMulti(const EThreadQueryContext ThreadContext, const UWorld* World, const FSweepQueryData& SweepData, const FCommonQueryData& CommonData, TArray<FHitResult>& OutHits);
+
+	// Async
+	ENGINE_API virtual void QueueAsyncOverlap(const FTraceHandle& TraceHandle, const EQueryInfo InfoType, const EThreadQueryContext ThreadContext, const UWorld* World, const FOverlapQueryData& OverlapData, const FCommonQueryData& CommonData);
+	ENGINE_API virtual void QueueAsyncRaycast(const FTraceHandle& TraceHandle, const EPhysicsQueryKind QueryKind, const EThreadQueryContext ThreadContext, const UWorld* World, const FRayQueryData& RayData, const FCommonQueryData& CommonData);
+	ENGINE_API virtual void QueueAsyncSweep(const FTraceHandle& TraceHandle, const EPhysicsQueryKind QueryKind, const EThreadQueryContext ThreadContext, const UWorld* World, const FSweepQueryData& SweepData, const FCommonQueryData& CommonData);
+
+	ENGINE_API virtual bool AreAsyncRequestsAdded();
+	ENGINE_API virtual void VerifyAsyncRequestsAreCompletedOrAbort();
 };
diff --git a/Engine/Source/Runtime/Engine/Public/WorldCollision.h b/Engine/Source/Runtime/Engine/Public/WorldCollision.h
--- a/Engine/Source/Runtime/Engine/Public/WorldCollision.h
+++ b/Engine/Source/Runtime/Engine/Public/WorldCollision.h
@@ -37,12 +37,15 @@ struct FTraceHandle
 			uint32 Index;
 		} _Data;
 	};
+	/** Denotes that this trace was in an auto-rtfm transactional context */
+	bool bTransactional = false;
 
 	FTraceHandle() : _Handle(0) {};
-	FTraceHandle(uint32 InFrameNumber, uint32 InIndex) 
+	FTraceHandle(uint32 InFrameNumber, uint32 InIndex, bool bInTransactional = false)
 	{
 		_Data.FrameNumber = InFrameNumber;
 		_Data.Index = InIndex;
+		bTransactional = bInTransactional;
 	}
 
 	friend inline uint32 GetTypeHash(const FTraceHandle& Handle)
@@ -70,6 +73,10 @@ struct FTraceHandle
 		_Handle = 0;
 	}
 
+	bool IsTransactional() const
+	{
+		return bTransactional;
+	}
 };
 
 /** 
@@ -226,6 +233,13 @@ struct TTraceThreadData
 	T Buffer[ASYNC_TRACE_BUFFER_SIZE];
 };
 
+template <typename DatumType>
+struct FTransactionalAsyncTraceBuffer
+{
+	TArray<TUniquePtr<TTraceThreadData<DatumType>>> Data;
+	int32 NumQueued;
+};
+
 /** 
  * Contains all Async Trace Result for one frame. 
  * 
@@ -253,6 +267,13 @@ struct AsyncTraceData : FNoncopyable
 	/** Datum entries in OverlapData are persistent for efficiency. This is the number of them that are actually in use (rather than OverlapData.Num()). */
 	int32 NumQueuedOverlapData;
 
+	/**
+	 * Data for Transactional async queries. This is effectively the same as above,
+	 * however it must be stored separately as we cannot run any async queries until a transaction is closed.
+	 */
+	FTransactionalAsyncTraceBuffer<FTraceDatum> TransactionalTraceData;
+	FTransactionalAsyncTraceBuffer<FOverlapDatum> TransactionalOverlapData;
+
 	/**
 	 * if Execution is all done, set this to be true
 	 * 
