Subtitles - expand the API to support blueprint and a new subtitles sequencer track
-) Subtitles API - Add IsSubtitleActive, StopSubtitle, and StopAllSubtitles, and FQueueSubtitleParams to specify a duration to play the subtitle for (used by Sequencer to play subtitles from mid-way after scrubbing)
-) Editor - Add Editor module with a SubtitleFactory so USubtitleAssetUserData can be created and played as standalone subtitle UObjects.  Also contains Sequencer Track + Section editor (SubtitlesTrackEditor, SubtitleSequencerSection)
-) Blueprint - Add SubtitlesblueprintFunctionLibrary as a blueprint wrapper around the SubtitlesAndClosedCaptions delegates
-) Sequencer - Add MovieSceneSubtitlesSystem to evaluate subtitle sequencer tracks + sections

diff --git a/Engine/Source/Runtime/Engine/Private/SoundWave.cpp b/Engine/Source/Runtime/Engine/Private/SoundWave.cpp
--- a/Engine/Source/Runtime/Engine/Private/SoundWave.cpp
+++ b/Engine/Source/Runtime/Engine/Private/SoundWave.cpp
@@ -3377,7 +3377,7 @@ void USoundWave::Parse(FAudioDevice* AudioDevice, const UPTRINT NodeWaveInstance
 			const ISubtitlesAndClosedCaptionsModule& SubtitlesModule = ModularFeatures.GetModularFeature<ISubtitlesAndClosedCaptionsModule>(ISubtitlesAndClosedCaptionsModule::GetModularFeatureName());
 			if (const UAssetUserData * Subtitle{ GetAssetUserDataOfClass(SubtitlesModule.GetAssetUserDataClass()) })
 			{
-				FSubtitlesAndClosedCaptionsDelegates::QueueSubtitle.ExecuteIfBound(*Subtitle);
+				FSubtitlesAndClosedCaptionsDelegates::QueueSubtitle.ExecuteIfBound(FQueueSubtitleParameters{*Subtitle});
 			}
 		}
 	}
diff --git a/Engine/Source/Runtime/Engine/Private/Subtitles/SubtitlesAndClosedCaptionsDelegates.cpp b/Engine/Source/Runtime/Engine/Private/Subtitles/SubtitlesAndClosedCaptionsDelegates.cpp
--- a/Engine/Source/Runtime/Engine/Private/Subtitles/SubtitlesAndClosedCaptionsDelegates.cpp
+++ b/Engine/Source/Runtime/Engine/Private/Subtitles/SubtitlesAndClosedCaptionsDelegates.cpp
@@ -2,5 +2,8 @@
 
 #include "Subtitles/SubtitlesAndClosedCaptionsDelegates.h"
 
-TDelegate<void (const UAssetUserData&)> FSubtitlesAndClosedCaptionsDelegates::QueueSubtitle;
+TDelegate<void(const FQueueSubtitleParameters&)> FSubtitlesAndClosedCaptionsDelegates::QueueSubtitle;
+TDelegate<bool(const UAssetUserData&)> FSubtitlesAndClosedCaptionsDelegates::IsSubtitleActive;
+TDelegate<void(const UAssetUserData&)> FSubtitlesAndClosedCaptionsDelegates::StopSubtitle;
+TDelegate<void()> FSubtitlesAndClosedCaptionsDelegates::StopAllSubtitles;
 TDelegate<void(FCanvas&)> FSubtitlesAndClosedCaptionsDelegates::DisplaySubtitle;
diff --git a/Engine/Source/Runtime/Engine/Public/Subtitles/SubtitlesAndClosedCaptionsDelegates.h b/Engine/Source/Runtime/Engine/Public/Subtitles/SubtitlesAndClosedCaptionsDelegates.h
--- a/Engine/Source/Runtime/Engine/Public/Subtitles/SubtitlesAndClosedCaptionsDelegates.h
+++ b/Engine/Source/Runtime/Engine/Public/Subtitles/SubtitlesAndClosedCaptionsDelegates.h
@@ -6,15 +6,25 @@
 class FCanvas;
 class UAssetUserData;
 
+struct ENGINE_API FQueueSubtitleParameters
+{
+	const UAssetUserData& Subtitle;
+	TOptional<float> Duration;
+};
+
 class ENGINE_API FSubtitlesAndClosedCaptionsDelegates
 {
 public:
 
-	// Have the subtitle subsystem to queue up a subtitle
-	static TDelegate<void(const UAssetUserData&)> QueueSubtitle;
+	// Have the subtitle subsystem to queue a subtitle to be displayed
+	static TDelegate<void(const FQueueSubtitleParameters&)> QueueSubtitle;
+
+	static TDelegate<bool(const UAssetUserData&)> IsSubtitleActive;
+
+	static TDelegate<void(const UAssetUserData&)> StopSubtitle;
+
+	static TDelegate<void()> StopAllSubtitles;
 
 	// Have the subtitle subsystem to display the highest priority subtitle. DisplaySubtitles is expected to be called/drawn per frame.
 	static TDelegate<void(FCanvas&)> DisplaySubtitle;
 };
-
-
