World Partition Runtime HashSet: fix source queries not respecting the target grids inclusions/exclusions for the non-spatially loaded cells.

diff --git a/Engine/Source/Runtime/Engine/Private/WorldPartition/RuntimeHashSet/WorldPartitionRuntimeHashSet.cpp b/Engine/Source/Runtime/Engine/Private/WorldPartition/RuntimeHashSet/WorldPartitionRuntimeHashSet.cpp
--- a/Engine/Source/Runtime/Engine/Private/WorldPartition/RuntimeHashSet/WorldPartitionRuntimeHashSet.cpp
+++ b/Engine/Source/Runtime/Engine/Private/WorldPartition/RuntimeHashSet/WorldPartitionRuntimeHashSet.cpp
@@ -548,11 +548,11 @@ void UWorldPartitionRuntimeHashSet::ForEachStreamingCellsQuery(const FWorldParti
 		return false;
 	};
 
-	auto ForEachSpatiallyLoadedCells = [&ShouldAddCell, QueryCache, &QuerySource, &Func]<typename SpatialIndexType>(SpatialIndexType* InSpatialIndex, int32 InLoadingRange, FName InGridName)
+	auto ForEachSpatiallyLoadedCells = [&ShouldAddCell, QueryCache, &QuerySource, &Func]<typename SpatialIndexType>(SpatialIndexType* InSpatialIndex, int32 InLoadingRange)
 	{
 		if (InSpatialIndex)
 		{
-			QuerySource.ForEachShape(InLoadingRange, InGridName, false, [InSpatialIndex, QueryCache, &ShouldAddCell, &QuerySource, &Func](const FSphericalSector& Shape)
+			QuerySource.ForEachShape(InLoadingRange, false, [InSpatialIndex, QueryCache, &ShouldAddCell, &QuerySource, &Func](const FSphericalSector& Shape)
 			{
 				auto ForEachIntersectingElement = [InSpatialIndex, QueryCache, &ShouldAddCell, &Shape, &QuerySource, &Func]<typename ShapeType>(const ShapeType& InShape)
 				{
@@ -600,9 +600,13 @@ void UWorldPartitionRuntimeHashSet::ForEachStreamingCellsQuery(const FWorldParti
 
 	ForEachStreamingData([&QuerySource, &ForEachSpatiallyLoadedCells, &ForEachNonSpatiallyLoadedCells](const FRuntimePartitionStreamingData& StreamingData)
 	{
-		return ForEachSpatiallyLoadedCells(StreamingData.SpatialIndex.Get(), StreamingData.GetLoadingRange(), StreamingData.Name) && 
-			   ForEachSpatiallyLoadedCells(StreamingData.SpatialIndex2D.Get(), StreamingData.GetLoadingRange(), StreamingData.Name) &&
-			   ForEachNonSpatiallyLoadedCells(StreamingData.NonSpatiallyLoadedCells);
+        if (FStreamingSourceShapeHelper::IsSourceAffectingGrid(QuerySource.TargetGrids, QuerySource.TargetBehavior, StreamingData.Name))
+        {
+			return ForEachSpatiallyLoadedCells(StreamingData.SpatialIndex.Get(), StreamingData.GetLoadingRange()) && 
+				   ForEachSpatiallyLoadedCells(StreamingData.SpatialIndex2D.Get(), StreamingData.GetLoadingRange()) &&
+				   ForEachNonSpatiallyLoadedCells(StreamingData.NonSpatiallyLoadedCells);
+		}
+		return true;
 	});
 }
 
diff --git a/Engine/Source/Runtime/Engine/Public/WorldPartition/WorldPartitionStreamingSource.h b/Engine/Source/Runtime/Engine/Public/WorldPartition/WorldPartitionStreamingSource.h
--- a/Engine/Source/Runtime/Engine/Public/WorldPartition/WorldPartitionStreamingSource.h
+++ b/Engine/Source/Runtime/Engine/Public/WorldPartition/WorldPartitionStreamingSource.h
@@ -304,12 +304,18 @@ struct FWorldPartitionStreamingQuerySource
 	
 	TArray<FStreamingSourceShape> Shapes;
 
+	/** Helper method that iterates over all shapes. If none is provided, it will still pass a sphere shape using Radius or grid's loading range (see bUseGridLoadingRange). */
+	FORCEINLINE void ForEachShape(float InGridLoadingRange, bool bInProjectIn2D, TFunctionRef<void(const FSphericalSector&)> InOperation) const
+	{
+		FStreamingSourceShapeHelper::ForEachShape(InGridLoadingRange, bUseGridLoadingRange ? InGridLoadingRange : Radius, bInProjectIn2D, Location, Rotation, Shapes, InOperation);
+	}
+
 	/** Helper method that iterates over all shapes. If none is provided, it will still pass a sphere shape using Radius or grid's loading range (see bUseGridLoadingRange). */
 	FORCEINLINE void ForEachShape(float InGridLoadingRange, FName InGridName, bool bInProjectIn2D, TFunctionRef<void(const FSphericalSector&)> InOperation) const
 	{
 		if (bSpatialQuery && FStreamingSourceShapeHelper::IsSourceAffectingGrid(TargetGrids, TargetBehavior, InGridName))
 		{
-			FStreamingSourceShapeHelper::ForEachShape(InGridLoadingRange, bUseGridLoadingRange ? InGridLoadingRange : Radius, bInProjectIn2D, Location, Rotation, Shapes, InOperation);
+			ForEachShape(InGridLoadingRange, bInProjectIn2D, InOperation);
 		}
 	}
 };
