Move skin weight profiles back to the render thread due to the staging buffer readback needing the immediate command list.

diff --git a/Engine/Source/Runtime/Engine/Private/Animation/SkinWeightProfile.cpp b/Engine/Source/Runtime/Engine/Private/Animation/SkinWeightProfile.cpp
--- a/Engine/Source/Runtime/Engine/Private/Animation/SkinWeightProfile.cpp
+++ b/Engine/Source/Runtime/Engine/Private/Animation/SkinWeightProfile.cpp
@@ -687,8 +687,7 @@ void FSkinWeightProfilesData::EnqueueDataReadback()
 		ReadbackData.ReadbackData.SetNumZeroed(BaseBuffer->GetVertexDataSize());
 
 		ENQUEUE_RENDER_COMMAND(FSkinWeightProfilesData_EnqueueDataReadback)(
-			UE::RenderCommandPipe::SkeletalMesh,
-				[&Data=ReadbackData](FRHICommandList& RHICmdList)
+			[&Data=ReadbackData](FRHICommandListImmediate& RHICmdList)
 		{
 			FScopeLock Lock(&Data.Mutex);
 		
