Skin Profiles: Change read-backs to operate on the main RHI thread due to the staging buffer readback needing the immediate command list.

#rnx

diff --git a/Engine/Source/Runtime/Engine/Private/Animation/SkinWeightProfile.cpp b/Engine/Source/Runtime/Engine/Private/Animation/SkinWeightProfile.cpp
--- a/Engine/Source/Runtime/Engine/Private/Animation/SkinWeightProfile.cpp
+++ b/Engine/Source/Runtime/Engine/Private/Animation/SkinWeightProfile.cpp
@@ -633,8 +633,7 @@ void FSkinWeightProfilesData::EnqueueGPUReadback()
 		{
 			static const FName ReadbackName("ReadbackSkinWeightBuffer");
 			ENQUEUE_RENDER_COMMAND(FSkinWeightProfilesData_EnqueueGPUReadback)(
-				UE::RenderCommandPipe::SkeletalMesh,
-				[&Data=ReadbackData, SkinWeightBuffer=BaseBuffer](FRHICommandList& RHICmdList)
+				[&Data=ReadbackData, SkinWeightBuffer=BaseBuffer](FRHICommandListImmediate& RHICmdList)
 			{
 				FScopeLock Lock(&Data.Mutex);
 					
@@ -654,8 +653,7 @@ void FSkinWeightProfilesData::EnqueueGPUReadback()
 		{
 			static const FName ReadbackName("ReadbackSkinWeightBuffer");
 			ENQUEUE_RENDER_COMMAND(FSkinWeightProfilesData_EnqueueGPUReadback)(
-			UE::RenderCommandPipe::SkeletalMesh,
-				[&Data=ReadbackData, SkinWeightBuffer=BaseBuffer](FRHICommandList& RHICmdList)
+				[&Data=ReadbackData, SkinWeightBuffer=BaseBuffer](FRHICommandListImmediate& RHICmdList)
 			{
 				FScopeLock Lock(&Data.Mutex);
 			
@@ -687,7 +685,7 @@ void FSkinWeightProfilesData::EnqueueDataReadback()
 		ReadbackData.ReadbackData.SetNumZeroed(BaseBuffer->GetVertexDataSize());
 
 		ENQUEUE_RENDER_COMMAND(FSkinWeightProfilesData_EnqueueDataReadback)(
-			[&Data=ReadbackData](FRHICommandListImmediate& RHICmdList)
+				[&Data=ReadbackData](FRHICommandListImmediate& )
 		{
 			FScopeLock Lock(&Data.Mutex);
 		
