Fixed sky light realtime capture not working with OpenGL for the mobile renderer.

#tests pc editor android preview
#rnx

diff --git a/Engine/Source/Runtime/Engine/Private/Components/SkyLightComponent.cpp b/Engine/Source/Runtime/Engine/Private/Components/SkyLightComponent.cpp
--- a/Engine/Source/Runtime/Engine/Private/Components/SkyLightComponent.cpp
+++ b/Engine/Source/Runtime/Engine/Private/Components/SkyLightComponent.cpp
@@ -89,7 +89,8 @@ FAutoConsoleVariableRef CVarSkylightIntensityMultiplier(
 
 static bool FeatureLevelSupportsRealTimeReflectionCapture(FStaticFeatureLevel FeatureLevel)
 {
-	return FeatureLevel >= ERHIFeatureLevel::ES3_1;
+	return FeatureLevel >= ERHIFeatureLevel::ES3_1
+		&& !IsOpenGLPlatform(GMaxRHIShaderPlatform); // UE-236698 - OpenGL does not support creating UAV on PF_FloatR11G11B10 or SRV on specific cube map mip.
 }
 
 void OnChangeSkylightRealTimeReflectionCapture(IConsoleVariable* Var)
