Suspend streamer before GC full purge when Engine.SuspendAssetStreamerBeforePerfomingFullGC is enabled (disabled by default)

diff --git a/Engine/Source/Runtime/Engine/Classes/Engine/Engine.h b/Engine/Source/Runtime/Engine/Classes/Engine/Engine.h
--- a/Engine/Source/Runtime/Engine/Classes/Engine/Engine.h
+++ b/Engine/Source/Runtime/Engine/Classes/Engine/Engine.h
@@ -1900,6 +1900,9 @@ private:
 	/** Whether a full purge has been triggered, so that the next GarbageCollect will do a full purge no matter what. */
 	bool bFullPurgeTriggered;
 
+	/** Whether a full purge is being performed during GC. */
+	bool bGCPerformingFullPurge = false;
+
 	/** Whether we should delay GC for one frame to finish some pending operation */
 	bool bShouldDelayGarbageCollect; 
 
diff --git a/Engine/Source/Runtime/Engine/Private/Streaming/RenderAssetUpdate.cpp b/Engine/Source/Runtime/Engine/Private/Streaming/RenderAssetUpdate.cpp
--- a/Engine/Source/Runtime/Engine/Private/Streaming/RenderAssetUpdate.cpp
+++ b/Engine/Source/Runtime/Engine/Private/Streaming/RenderAssetUpdate.cpp
@@ -25,18 +25,18 @@ static TAutoConsoleVariable<int32> CVarStreamingStressTestExtraAsyncLatency(
 	ECVF_Cheat);
 
 
-volatile int32 GRenderAssetStreamingSuspension = 0;
+std::atomic<int32> GRenderAssetStreamingSuspension = 0;
 
 bool IsAssetStreamingSuspended()
 {
-	return GRenderAssetStreamingSuspension > 0;
+	return GRenderAssetStreamingSuspension.load(std::memory_order_relaxed) > 0;
 }
 
 void SuspendRenderAssetStreaming()
 {
 	ensure(IsInGameThread());
 
-	if (FPlatformAtomics::InterlockedIncrement(&GRenderAssetStreamingSuspension) == 1)
+	if (GRenderAssetStreamingSuspension.fetch_add(1, std::memory_order_relaxed) == 0)
 	{
 		bool bHasPendingStreamingRequest = false;
 
@@ -63,6 +63,11 @@ void SuspendRenderAssetStreaming()
 		while (LockedAssets.Num() && (TimeLimit > 0 || GStreamingFlushTimeOut <= 0))
 		{
 			FPlatformProcess::Sleep(RENDER_ASSET_STREAMING_SLEEP_DT);
+			ENQUEUE_RENDER_COMMAND(FlushCommand)(
+				[](FRHICommandList& RHICmdList)
+				{
+					GRHICommandList.GetImmediateCommandList().ImmediateFlush(EImmediateFlushType::FlushRHIThread);
+				});
 			FlushRenderingCommands();
 			
 			TimeLimit -= RENDER_ASSET_STREAMING_SLEEP_DT;
@@ -103,6 +108,11 @@ void SuspendRenderAssetStreaming()
 		if (bHasPendingStreamingRequest)
 		{
 			// Ensure any pending render command executes.
+			ENQUEUE_RENDER_COMMAND(FlushCommand)(
+				[](FRHICommandList& RHICmdList)
+				{
+					GRHICommandList.GetImmediateCommandList().ImmediateFlush(EImmediateFlushType::FlushRHIThread);
+				});
 			FlushRenderingCommands();
 		}
 	}
@@ -110,8 +120,7 @@ void SuspendRenderAssetStreaming()
 
 void ResumeRenderAssetStreaming()
 {
-	FPlatformAtomics::InterlockedDecrement(&GRenderAssetStreamingSuspension);
-	ensure(GRenderAssetStreamingSuspension >= 0);
+	GRenderAssetStreamingSuspension.fetch_sub(1, std::memory_order_relaxed);
 }
 
 FRenderAssetUpdate::FRenderAssetUpdate(const UStreamableRenderAsset* InAsset)
diff --git a/Engine/Source/Runtime/Engine/Private/UnrealEngine.cpp b/Engine/Source/Runtime/Engine/Private/UnrealEngine.cpp
--- a/Engine/Source/Runtime/Engine/Private/UnrealEngine.cpp
+++ b/Engine/Source/Runtime/Engine/Private/UnrealEngine.cpp
@@ -504,11 +504,11 @@ static FAutoConsoleVariableRef GDelayTrimMemoryDuringMapLoadModeCVar(
 );
 
 // Flush Rendering and RHI before GC in trim memory
-int32 GFlushRTAndRHIBeforeGCInTrimMemory = 1;
-static FAutoConsoleVariableRef GFlushRTAndRHIBeforeGCInTrimMemoryCVar(
-	TEXT("Engine.FlushRTAndRHIBeforeGCInTrimMemory"),
-	GFlushRTAndRHIBeforeGCInTrimMemory,
-	TEXT("Flush Rendering and RHI before GC in trim memory in case there are some UStreamableRenderAssets with pending RenderAssetUpdate RHI tasks (default: 1)"),
+int32 GSuspendAssetStreamerBeforePerformingFullGC = 0;
+static FAutoConsoleVariableRef GSuspendAssetStreamerBeforePerformingFullGCCVar(
+	TEXT("Engine.SuspendAssetStreamerBeforePerformingFullGC"),
+	GSuspendAssetStreamerBeforePerformingFullGC,
+	TEXT("Suspend asset streamer before performing a full GC (default: 0)"),
 	ECVF_Default
 );
 
@@ -1905,6 +1905,7 @@ void UEngine::ConditionalCollectGarbage()
 
 			if (bFullPurgeTriggered)
 			{
+				TGuardValue<bool> _FullGC(bGCPerformingFullPurge, true);
 				if (TryCollectGarbage(GARBAGE_COLLECTION_KEEPFLAGS, true))
 				{
 					ForEachObjectOfClass(UWorld::StaticClass(),[](UObject* World)
@@ -2017,6 +2018,7 @@ void UEngine::PerformGarbageCollectionAndCleanupActors()
 		}
 
 		// Perform housekeeping.
+		TGuardValue<bool> _FullGC(bGCPerformingFullPurge, bForcePurge);
 		if (TryCollectGarbage(GARBAGE_COLLECTION_KEEPFLAGS, bForcePurge))
 		{
 			ForEachObjectOfClass(UWorld::StaticClass(), [](UObject* World)
@@ -2105,6 +2107,31 @@ void UEngine::Init(IEngineLoop* InEngineLoop)
 
 	FCoreUObjectDelegates::GetPreGarbageCollectDelegate().AddStatic(UEngine::PreGarbageCollect);
 
+	FCoreUObjectDelegates::GetGarbageCollectStartedDelegate().AddLambda([]()
+		{
+			if (GSuspendAssetStreamerBeforePerformingFullGC != 0 && GEngine != nullptr && GEngine->bGCPerformingFullPurge)
+			{
+				// Flush render and RHI commands in case there are some UStreamableRenderAssets with pending RenderAssetUpdate RHI tasks
+				ENQUEUE_RENDER_COMMAND(FlushCommand)(
+					[](FRHICommandList& RHICmdList)
+					{
+						GRHICommandList.GetImmediateCommandList().ImmediateFlush(EImmediateFlushType::FlushRHIThread);
+					});
+				FlushRenderingCommands();
+
+				SuspendRenderAssetStreaming();
+			}
+		});
+
+	FCoreUObjectDelegates::GetPostGarbageCollect().AddLambda([]()
+		{
+			if (GSuspendAssetStreamerBeforePerformingFullGC != 0 && GEngine != nullptr && GEngine->bGCPerformingFullPurge && IsAssetStreamingSuspended())
+			{
+				ResumeRenderAssetStreaming();
+			}
+		});
+
+
 	if (!FApp::IsProjectNameEmpty())
 	{
 		// Initialize the HMDs and motion controllers, if any
@@ -16029,19 +16056,20 @@ bool UEngine::LoadMap( FWorldContext& WorldContext, FURL URL, class UPendingNetG
 
 void UEngine::TrimMemory()
 {
-	if (GFlushRTAndRHIBeforeGCInTrimMemory > 0)
+	// Clean up the previous level out of memory.
 	{
-		// Flush render and RHI commands in case there are some UStreamableRenderAssets with pending RenderAssetUpdate RHI tasks
-		ENQUEUE_RENDER_COMMAND(FlushCommand)(
-			[](FRHICommandList& RHICmdList)
+		if (GEngine != nullptr)
 		{
-			GRHICommandList.GetImmediateCommandList().ImmediateFlush(EImmediateFlushType::FlushRHIThread);
-		});
-		FlushRenderingCommands();
-	}
+			GEngine->bGCPerformingFullPurge = true;
+		}
 
-	// Clean up the previous level out of memory.
-	CollectGarbage(GARBAGE_COLLECTION_KEEPFLAGS, true);
+		CollectGarbage(GARBAGE_COLLECTION_KEEPFLAGS, true);
+
+		if (GEngine != nullptr)
+		{
+			GEngine->bGCPerformingFullPurge = false;
+		}
+	}
 
 	// For platforms which manage GPU memory directly we must Enqueue a flush, and wait for it to be processed
 	// so that any pending frees that depend on the GPU will be processed.  Otherwise a whole map's worth of GPU memory
@@ -17081,7 +17109,10 @@ bool UEngine::CommitMapChange( FWorldContext &Context )
 		}
 
 		// Collect garbage. @todo streaming: make sure that this doesn't stall due to async loading in the background
-		CollectGarbage( GARBAGE_COLLECTION_KEEPFLAGS, true );
+		{
+			TGuardValue<bool> _FullGC(bGCPerformingFullPurge, true);
+			CollectGarbage(GARBAGE_COLLECTION_KEEPFLAGS, true);
+		}
 
 		// The new fake persistent level is first in the LevelsToLoadForPendingMapChange array.
 		FName	FakePersistentLevelName = Context.LevelsToLoadForPendingMapChange[0];
@@ -17196,7 +17227,10 @@ bool UEngine::CommitMapChange( FWorldContext &Context )
 		}
 
 		// Remove all unloaded levels from memory and perform full purge.
-		CollectGarbage( GARBAGE_COLLECTION_KEEPFLAGS, true );
+		{
+			TGuardValue<bool> _FullGC(bGCPerformingFullPurge, true);
+			CollectGarbage(GARBAGE_COLLECTION_KEEPFLAGS, true);
+		}
 
 		// if there are pending streaming changes replicated from the server, apply them immediately
 		if (Context.PendingLevelStreamingStatusUpdates.Num() > 0)
