[NetPhysics] Network optimizations
- Stop sending LocalFrame and InputFrame NetworkPhysicsComponent for states and inputs.
- Only send altered inputs back to owning client.

diff --git a/Engine/Source/Runtime/Engine/Private/PhysicsEngine/NetworkPhysicsComponent.cpp b/Engine/Source/Runtime/Engine/Private/PhysicsEngine/NetworkPhysicsComponent.cpp
--- a/Engine/Source/Runtime/Engine/Private/PhysicsEngine/NetworkPhysicsComponent.cpp
+++ b/Engine/Source/Runtime/Engine/Private/PhysicsEngine/NetworkPhysicsComponent.cpp
@@ -61,7 +61,7 @@ namespace PhysicsReplicationCVars
 		bool bAllowInputExtrapolation = true;
 		static FAutoConsoleVariableRef CVarResimAllowInputExtrapolation(TEXT("np2.Resim.AllowInputExtrapolation"), bAllowInputExtrapolation, TEXT("When true and not locally controlled, allow inputs to be extrapolated from last known and if there is a gap allow interpolation between two known inputs."));
 		bool bValidateDataOnGameThread = false;
-		static FAutoConsoleVariableRef CVarResimValidateDataOnGameThread(TEXT("np2.Resim.ValidateDataOnGameThread"), bValidateDataOnGameThread, TEXT("When true, perform server-side input validation through FNetworkPhysicsData::ValidateData on the Game Thread. If false, perform the call on the Physics Thread."));
+		static FAutoConsoleVariableRef CVarResimValidateDataOnGameThread(TEXT("np2.Resim.ValidateDataOnGameThread"), bValidateDataOnGameThread, TEXT("When true, perform server-side input validation through FNetworkPhysicsData::ValidateData on the Game Thread, note that LocalFrame will be the same as ServerFrame on Game Thread. If false, perform the call on the Physics Thread."));
 		bool bApplySimProxyStateAtRuntime = false;
 		static FAutoConsoleVariableRef CVarResimApplySimProxyStateAtRuntime(TEXT("np2.Resim.ApplySimProxyStateAtRuntime"), bApplySimProxyStateAtRuntime, TEXT("When true, call ApplyData on received states for simulated proxies at runtime."));
 		bool bApplySimProxyInputAtRuntime = true;
@@ -624,6 +624,12 @@ void UNetworkPhysicsComponent::NetworkMarshaledData()
 	if (NetworkPhysicsComponent_Internal)
 	{
 		const bool bIsServer = HasServerWorld();
+		if (!bIsServer && !IsNetworkPhysicsTickOffsetAssigned())
+		{
+			// Don't replicate data to the server until networked physics is setup with a synchronized physics tick offset
+			return;
+		}
+
 		while (Chaos::TSimCallbackOutputHandle<FAsyncNetworkPhysicsComponentOutput> AsyncOutput = NetworkPhysicsComponent_Internal->PopFutureOutputData_External())
 		{
 			if (AsyncOutput->InputsToNetwork_Owner.IsSet())
@@ -649,10 +655,12 @@ void UNetworkPhysicsComponent::NetworkMarshaledData()
 						}
 					}
 
-					// TODO, Only copy/replicate data if InputFrame is INDEX_NONE i.e. the input has been altered by the server
+					// Only replicate data to owning client if bDataAltered is true i.e. the input has been altered by the server
+					ReplicatedInputs.History->ResizeDataHistory(AsyncOutput->InputData->CountAlteredData(/*bIncludeUnimportant*/ true, /*bIncludeImportant*/ bEnableReliableFlow == false), EAllowShrinking::Yes);
+					ReplicatedInputs.History->SetRecordDataIncremental(true); // Only record data that is newer than already cached data
 
 					// Server sends inputs through property replication to owning client
-					if (AsyncOutput->InputData->CopyAllData(*ReplicatedInputs.History, /*bIncludeUnimportant*/ true, /*bIncludeImportant*/ bEnableReliableFlow == false))
+					if (AsyncOutput->InputData->CopyAlteredData(*ReplicatedInputs.History, /*bIncludeUnimportant*/ true, /*bIncludeImportant*/ bEnableReliableFlow == false))
 					{
 						MARK_PROPERTY_DIRTY_FROM_NAME(UNetworkPhysicsComponent, ReplicatedInputs, this);
 					}
@@ -868,6 +876,7 @@ void UNetworkPhysicsComponent::ServerReceiveInputData_Implementation(const FNetw
 
 		AsyncInput->InputData->ResetFast();
 		AsyncInput->InputData->ResizeDataHistory(ClientInputs.History->GetHistorySize(), EAllowShrinking::No);
+		AsyncInput->InputData->SetRecordDataIncremental(true); // Only record data that is newer than already cached data
 		ClientInputs.History->CopyAllData(*AsyncInput->InputData.Get(), /*bIncludeUnimportant*/ true, /*bIncludeImportant*/ true);
 
 		// Send received inputs to remote clients
@@ -1009,6 +1018,15 @@ bool UNetworkPhysicsComponent::IsLocallyControlled() const
 	return false;
 }
 
+bool UNetworkPhysicsComponent::IsNetworkPhysicsTickOffsetAssigned() const
+{
+	if (APlayerController* PlayerController = GetPlayerController())
+	{
+		return PlayerController->GetNetworkPhysicsTickOffsetAssigned();
+	}
+	return false;
+}
+
 void UNetworkPhysicsComponent::SetCompareStateToTriggerRewind(const bool bInCompareStateToTriggerRewind, const bool bInIncludeSimProxies)
 {
 	bCompareStateToTriggerRewind = bInCompareStateToTriggerRewind;
@@ -1387,8 +1405,8 @@ void FAsyncNetworkPhysicsComponent::ConsumeAsyncInput(const int32 PhysicsStep)
 				if (bIsServer && InputHistory->GetLatestFrame() < PhysicsStep && PhysicsData)
 				{
 #if DEBUG_NETWORK_PHYSICS
-					UE_LOG(LogChaos, Log, TEXT("SERVER | PT | Input Buffer Empty, Injecting Received Input at frame %d || LocalFrame = %d || ServerFrame = %d || InputFrame = %d || Data: %s || Actor: %s")
-						, PhysicsStep, PhysicsData->LocalFrame, PhysicsData->ServerFrame, PhysicsData->InputFrame, *PhysicsData->DebugData(), *GetActorName());
+					UE_LOG(LogChaos, Log, TEXT("SERVER | PT | Input Buffer Empty, Injecting Received Input at frame %d || LocalFrame = %d || ServerFrame = %d || bDataAltered = %d || Data: %s || Actor: %s")
+						, PhysicsStep, PhysicsData->LocalFrame, PhysicsData->ServerFrame, PhysicsData->bDataAltered, *PhysicsData->DebugData(), *GetActorName());
 #endif
 					
 					// Record data in InputHistory
@@ -1484,8 +1502,8 @@ void FAsyncNetworkPhysicsComponent::OnPreProcessInputs_Internal(const int32 Phys
 			{
 				PhysicsData->ApplyData(ActorComponent.Get());
 #if DEBUG_NETWORK_PHYSICS
-				UE_LOG(LogChaos, Log, TEXT("			Applying extracted state from history | bExactFrame = %d | LocalFrame = %d | ServerFrame = %d | InputFrame = %d | Data: %s")
-					, bExactFrame, PhysicsData->LocalFrame, PhysicsData->ServerFrame, PhysicsData->InputFrame, *PhysicsData->DebugData());
+				UE_LOG(LogChaos, Log, TEXT("			Applying extracted state from history | bExactFrame = %d | LocalFrame = %d | ServerFrame = %d | bDataAltered = %d | Data: %s")
+					, bExactFrame, PhysicsData->LocalFrame, PhysicsData->ServerFrame, PhysicsData->bDataAltered, *PhysicsData->DebugData());
 #endif
 			}
 #if DEBUG_NETWORK_PHYSICS
@@ -1566,7 +1584,7 @@ void FAsyncNetworkPhysicsComponent::OnPreProcessInputs_Internal(const int32 Phys
 				}
 
 				// If the extracted input data was altered (extrapolated, interpolated) on the server, record it into the history for it to get replicated to clients
-				if (bIsServer && PhysicsData->InputFrame == INDEX_NONE)
+				if (bIsServer && PhysicsData->bDataAltered)
 				{
 					// Check that buffer is not empty
 					if (PhysicsStep <= InputHistory->GetLatestFrame())
@@ -1595,8 +1613,8 @@ void FAsyncNetworkPhysicsComponent::OnPreProcessInputs_Internal(const int32 Phys
 
 #if DEBUG_NETWORK_PHYSICS
 				{
-					UE_LOG(LogChaos, Log, TEXT("			Applying extracted input from history | LocalFrame = %d | ServerFrame = %d | InputFrame = %d | IsResim = %d | IsLocallyControlled = %d | InputDecay = %f | Data: %s")
-						, PhysicsData->LocalFrame, PhysicsData->ServerFrame, PhysicsData->InputFrame, bIsSolverResim, IsLocallyControlled(), GetCurrentInputDecay(PhysicsData), *PhysicsData->DebugData());
+					UE_LOG(LogChaos, Log, TEXT("			Applying extracted input from history | LocalFrame = %d | ServerFrame = %d | bDataAltered = %d | IsResim = %d | IsLocallyControlled = %d | InputDecay = %f | Data: %s")
+						, PhysicsData->LocalFrame, PhysicsData->ServerFrame, PhysicsData->bDataAltered, bIsSolverResim, IsLocallyControlled(), GetCurrentInputDecay(PhysicsData), *PhysicsData->DebugData());
 				}
 #endif
 			}
@@ -1650,8 +1668,8 @@ void FAsyncNetworkPhysicsComponent::OnPostProcessInputs_Internal(const int32 Phy
 
 #if DEBUG_NETWORK_PHYSICS
 			{
-				UE_LOG(LogChaos, Log, TEXT("		Recording input into history | LocalFrame = %d | ServerFrame = %d | InputFrame = %d | Input: %s ")
-					, PhysicsData->LocalFrame, PhysicsData->ServerFrame, PhysicsData->InputFrame, *PhysicsData->DebugData());
+				UE_LOG(LogChaos, Log, TEXT("		Recording input into history | LocalFrame = %d | ServerFrame = %d | bDataAltered = %d | Input: %s ")
+					, PhysicsData->LocalFrame, PhysicsData->ServerFrame, PhysicsData->bDataAltered, *PhysicsData->DebugData());
 			}
 #endif
 		}
@@ -1662,21 +1680,9 @@ void FAsyncNetworkPhysicsComponent::OnPostProcessInputs_Internal(const int32 Phy
 
 		if (StateHistory && StateData && bShouldCacheStateHistory)
 		{
-			// Compute of the local frame coming from the client that was used to generate this state
-			int32 InputFrame = INDEX_NONE;
-			if (InputData)
-			{
-				FNetworkPhysicsData* PhysicsData = InputData.Get();
-				if (InputHistory && InputHistory->ExtractData(PhysicsStep, false, PhysicsData, true))
-				{
-					InputFrame = PhysicsData->InputFrame;
-				}
-			}
-
 			// Prepare to gather state data
 			FNetworkPhysicsData* PhysicsData = StateData.Get();
 			PhysicsData->PrepareFrame(PhysicsStep, bIsServer, GetNetworkPhysicsTickOffset());
-			PhysicsData->InputFrame = InputFrame;
 
 			// Gather input data from implementation
 			PhysicsData->BuildData(ActorComponent.Get());
@@ -1686,8 +1692,8 @@ void FAsyncNetworkPhysicsComponent::OnPostProcessInputs_Internal(const int32 Phy
 
 #if DEBUG_NETWORK_PHYSICS
 			{
-				UE_LOG(LogChaos, Log, TEXT("		Recording state into history | LocalFrame = %d | ServerFrame = %d | InputFrame = %d | State: %s ")
-					, PhysicsData->LocalFrame, PhysicsData->ServerFrame, PhysicsData->InputFrame, *PhysicsData->DebugData());
+				UE_LOG(LogChaos, Log, TEXT("		Recording state into history | LocalFrame = %d | ServerFrame = %d | bDataAltered = %d | State: %s ")
+					, PhysicsData->LocalFrame, PhysicsData->ServerFrame, PhysicsData->bDataAltered, *PhysicsData->DebugData());
 			}
 #endif
 		}
diff --git a/Engine/Source/Runtime/Engine/Public/Physics/NetworkPhysicsComponent.h b/Engine/Source/Runtime/Engine/Public/Physics/NetworkPhysicsComponent.h
--- a/Engine/Source/Runtime/Engine/Public/Physics/NetworkPhysicsComponent.h
+++ b/Engine/Source/Runtime/Engine/Public/Physics/NetworkPhysicsComponent.h
@@ -61,17 +61,17 @@ struct TNetRewindHistory : public Chaos::TDataRewindHistory<DataType>
 		}
 	}
 
-	virtual int32 CountValidData(const uint32 StartFrame, const uint32 EndFrame, const bool bIncludeUnimportant = true, const bool bIncludeImportant = false)
+	virtual int32 CountValidData(const uint32 StartFrame, const uint32 EndFrame, const bool bIncludeUnimportant = true, const bool bIncludeImportant = false) override
 	{
 		// Find how many entries are valid in frame range
 		DataType FrameData;
 		int32 Count = 0;
-		for (uint32 CopyFrame = StartFrame; CopyFrame <= EndFrame; ++CopyFrame)
+		for (uint32 Frame = StartFrame; Frame <= EndFrame; ++Frame)
 		{
-			const int32 CopyIndex = Super::GetFrameIndex(CopyFrame);
-			if (CopyFrame == Super::DataHistory[CopyIndex].LocalFrame)
+			const int32 Index = Super::GetFrameIndex(Frame);
+			if (Frame == Super::DataHistory[Index].LocalFrame)
 			{
-				FrameData = Super::DataHistory[CopyIndex];
+				FrameData = Super::DataHistory[Index];
 
 				// Check if we should include unimportant and/or important data
 				if ((!FrameData.bImportant && bIncludeUnimportant) || (FrameData.bImportant && bIncludeImportant))
@@ -83,6 +83,23 @@ struct TNetRewindHistory : public Chaos::TDataRewindHistory<DataType>
 		return Count;
 	}
 
+	virtual int32 CountAlteredData(const bool bIncludeUnimportant = true, const bool bIncludeImportant = false) override
+	{
+		DataType FrameData;
+		int32 Count = 0;
+		for (int32 Index = 0; Index < Super::NumFrames; ++Index)
+		{
+			FrameData = Super::DataHistory[Index];
+
+			// Check if we should include unimportant and/or important data
+			if (FrameData.bDataAltered && ((!FrameData.bImportant && bIncludeUnimportant) || (FrameData.bImportant && bIncludeImportant)))
+			{
+				Count++;
+			}
+		}
+		return Count;
+	}
+
 	virtual void SetImportant(const bool bImportant, const int32 Frame = INDEX_NONE) override
 	{
 		if (Frame > INDEX_NONE)
@@ -96,9 +113,9 @@ struct TNetRewindHistory : public Chaos::TDataRewindHistory<DataType>
 		else
 		{
 			// Set importance on all frames
-			for (int32 CopyIndex = 0; CopyIndex < Super::NumFrames; ++CopyIndex)
+			for (int32 Index = 0; Index < Super::NumFrames; ++Index)
 			{
-				Super::DataHistory[CopyIndex].bImportant = bImportant;
+				Super::DataHistory[Index].bImportant = bImportant;
 			}
 		}
 
@@ -137,7 +154,26 @@ struct TNetRewindHistory : public Chaos::TDataRewindHistory<DataType>
 			}
 		}
 		return bHasCopiedData;
+	}
 
+	virtual bool CopyAlteredData(Chaos::FBaseRewindHistory& OutHistory, bool bIncludeUnimportant = true, bool bIncludeImportant = false) override
+	{
+		TNetRewindHistory& OutNetHistory = static_cast<TNetRewindHistory&>(OutHistory);
+		bool bHasCopiedData = false;
+
+		DataType FrameData;
+		for (int32 CopyIndex = 0; CopyIndex < Super::NumFrames; ++CopyIndex)
+		{
+			FrameData = Super::DataHistory[CopyIndex];
+
+			// Check if we should include unimportant and/or important data
+			if (FrameData.bDataAltered && ((!FrameData.bImportant && bIncludeUnimportant) || (FrameData.bImportant && bIncludeImportant)))
+			{
+				OutNetHistory.RecordData(FrameData.LocalFrame, &FrameData);
+				bHasCopiedData = true;
+			}
+		}
+		return bHasCopiedData;
 	}
 
 	virtual bool CopyData(Chaos::FBaseRewindHistory& OutHistory, const uint32 StartFrame, const uint32 EndFrame, bool bIncludeUnimportant = true, bool bIncludeImportant = false) override
@@ -245,8 +281,10 @@ struct TNetRewindHistory : public Chaos::TDataRewindHistory<DataType>
 
 	virtual void NetSerialize(FArchive& Ar, UPackageMap* InPackageMap) override
 	{
-		Ar << Super::NumFrames;
-		
+		uint32 NumFramesUnsigned = static_cast<uint32>(Super::NumFrames);
+		Ar.SerializeIntPacked(NumFramesUnsigned);
+		Super::NumFrames = static_cast<int32>(NumFramesUnsigned);
+
 		if (Super::NumFrames > GetMaxArraySize())
 		{
 			UE_LOG(LogTemp, Warning, TEXT("TNetRewindHistory: serialized array of size %d exceeds maximum size %d."), Super::NumFrames, GetMaxArraySize());
@@ -282,7 +320,7 @@ struct TNetRewindHistory : public Chaos::TDataRewindHistory<DataType>
 				FrameData = NetDebugHistory.DataHistory[FrameIndex];
 				LocalFrames[FrameIndex] = FrameData.LocalFrame;
 				ServerFrames[FrameIndex] = FrameData.ServerFrame;
-				InputFrames[FrameIndex] = FrameData.InputFrame;
+				InputFrames[FrameIndex] = FrameData.bDataAltered ? 1 : 0; // For now we show the altered state inside the InputFrames array, since that was the main usecase for it when it was implemented
 			}
 		}
 	}
@@ -297,11 +335,11 @@ struct TNetRewindHistory : public Chaos::TDataRewindHistory<DataType>
 		{
 			for (int32 FrameIndex = 0; FrameIndex < Super::NumFrames; ++FrameIndex)
 			{
-				UE_LOG(LogChaos, Log, TEXT("		Index: %d || LocalFrame = %d || ServerFrame = %d || InputFrame = %d || bReceivedData = %d || bImportant = %d  ||  Data: %s")
+				UE_LOG(LogChaos, Log, TEXT("		Index: %d || LocalFrame = %d || ServerFrame = %d || bDataAltered = %d || bReceivedData = %d || bImportant = %d  ||  Data: %s")
 				, FrameIndex
 				, Super::DataHistory[FrameIndex].LocalFrame
 				, Super::DataHistory[FrameIndex].ServerFrame
-				, Super::DataHistory[FrameIndex].InputFrame
+				, Super::DataHistory[FrameIndex].bDataAltered
 				, Super::DataHistory[FrameIndex].bReceivedData
 				, Super::DataHistory[FrameIndex].bImportant
 				, *Super::DataHistory[FrameIndex].DebugData());
@@ -558,19 +596,32 @@ struct FNetworkPhysicsData
 {
 	GENERATED_USTRUCT_BODY()
 
+	PRAGMA_DISABLE_DEPRECATION_WARNINGS // Needed while we have a deprecated property inside the struct
+	FNetworkPhysicsData() = default;
 	virtual ~FNetworkPhysicsData() = default;
+	FNetworkPhysicsData(const FNetworkPhysicsData&) = default;
+	FNetworkPhysicsData(FNetworkPhysicsData&&) = default;
+	FNetworkPhysicsData& operator=(const FNetworkPhysicsData&) = default;
+	FNetworkPhysicsData& operator=(FNetworkPhysicsData&&) = default;
+	PRAGMA_ENABLE_DEPRECATION_WARNINGS
 
-	// Server frame at which this data has been generated
+	// Server frame at which this data has been generated 
 	UPROPERTY()
 	int32 ServerFrame = INDEX_NONE;
 
-	// Local frame at which this data has been generated
+	// Local frame at which this data has been generated (not replicated, calculated from ServerFrame - FrameOffset )
 	UPROPERTY()
 	int32 LocalFrame = INDEX_NONE;
 
-	// Input frame used to generate the network data
-	UPROPERTY()
-	int32 InputFrame = INDEX_NONE;
+#if WITH_EDITORONLY_DATA
+	// InputFrame is no longer replicated or used. Use bDataAltered to check if an input has been altered
+	UE_DEPRECATED(5.6, "InputFrame is no longer replicated or used. Use bDataAltered to check if an input has been altered.")
+	UPROPERTY(meta = (DeprecatedProperty, DeprecationMessage = "InputFrame is no longer replicated or populated. Use bDataAltered to check if an input has been altered"))
+	int32 InputFrame_DEPRECATED = INDEX_NONE;
+#endif
+
+	// If this data was altered so that it doesn't correspond the produced source data (from merging, interpolating or extrapolating)
+	bool bDataAltered = false;
 
 	// If this data was received over the network or locally predicted
 	bool bReceivedData = false;
@@ -581,16 +632,27 @@ struct FNetworkPhysicsData
 	// Serialize the data into/from the archive
 	void SerializeFrames(FArchive& Ar)
 	{
-		Ar << ServerFrame;
-		Ar << LocalFrame;
-		Ar << InputFrame;
+		uint32 ServerFrameUnsigned = 0;
+
+		if (Ar.IsLoading()) // Deserializing
+		{
+			Ar.SerializeIntPacked(ServerFrameUnsigned);
+			ServerFrame = static_cast<int32>(ServerFrameUnsigned) - 1;
+			LocalFrame = ServerFrame; // Temporarily set LocalFrame to ServerFrame, it will get recalculated later in TRewindHistory::ReceiveNewData
+		}
+		else // Serializing
+		{
+			check((ServerFrame + 1) >= 0);
+			ServerFrameUnsigned = static_cast<uint32>(ServerFrame + 1);
+			Ar.SerializeIntPacked(ServerFrameUnsigned);
+		}
 	}
 
 	void PrepareFrame(int32 CurrentFrame, bool bIsServer, int32 ClientFrameOffset)
 	{
 		LocalFrame = CurrentFrame;
 		ServerFrame = bIsServer ? CurrentFrame : CurrentFrame + ClientFrameOffset;
-		InputFrame = CurrentFrame;
+		bDataAltered = false;
 		bReceivedData = false;
 		bImportant = false;
 	}
@@ -664,7 +726,7 @@ struct FNetworkPhysicsData
 
 	bool operator==(const FNetworkPhysicsData& Other) const
 	{
-		return ServerFrame == Other.ServerFrame && LocalFrame == Other.LocalFrame && InputFrame == Other.InputFrame;
+		return ServerFrame == Other.ServerFrame && LocalFrame == Other.LocalFrame;
 	}
 
 	friend UNetworkPhysicsComponent;
@@ -762,6 +824,9 @@ public:
 	
 	// Check if this is controlled locally through relayed inputs or an existing local player controller
 	ENGINE_API bool IsLocallyControlled() const;
+	
+	// Check if networked physics is setup with a synchronized physics tick offset
+	ENGINE_API bool IsNetworkPhysicsTickOffsetAssigned() const;
 
 	/** Mark this as controlled through locally relayed inputs rather than controlled as a pawn through a player controller.
 	* Set if NetworkPhysicsComponent is implemented on an AActor instead of APawn and it's currently being fed inputs, or if this is controlled by the server. 
diff --git a/Engine/Source/Runtime/Experimental/Chaos/Public/RewindData.h b/Engine/Source/Runtime/Experimental/Chaos/Public/RewindData.h
--- a/Engine/Source/Runtime/Experimental/Chaos/Public/RewindData.h
+++ b/Engine/Source/Runtime/Experimental/Chaos/Public/RewindData.h
@@ -51,6 +51,9 @@ struct FBaseRewindHistory
 	* @param bIncludeUnimportant = If to include unimportant data entries
 	* @param bIncludeImportant = If to include important data entries */
 	FORCEINLINE virtual int32 CountValidData(const uint32 StartFrame, const uint32 EndFrame, const bool bIncludeUnimportant = true, const bool bIncludeImportant = false) { return 0; }
+	
+	/** Find how many entries that have been marked as altered, meaning the server has altered the data so that it doesn't match the received data from the client */
+	FORCEINLINE virtual int32 CountAlteredData(const bool bIncludeUnimportant = true, const bool bIncludeImportant = false) { return 0; }
 
 	/** Mark frame important or unimportant
 	* @param Frame = use INDEX_NONE to set importance on all entries */
@@ -81,6 +84,11 @@ struct FBaseRewindHistory
 	* @param bIncludeImportant = If to copy important data entries */
 	virtual bool CopyAllData(Chaos::FBaseRewindHistory& OutHistory, bool bIncludeUnimportant = true, bool bIncludeImportant = false) { return false; }
 
+	/** Copy altered data from local history into into @param OutHistory
+	* @param bIncludeUnimportant = If to copy unimportant data entries
+	* @param bIncludeImportant = If to copy important data entries */
+	virtual bool CopyAlteredData(Chaos::FBaseRewindHistory& OutHistory, bool bIncludeUnimportant = true, bool bIncludeImportant = false) { return false; }
+
 	/** Copy data from local history into @param OutHistory
 	* @param StartFrame = Included
 	* @param EndFrame = Included 
@@ -157,6 +165,7 @@ struct TDataRewindHistory : public FBaseRewindHistory
 	/** Initialize history */
 	virtual void Initialize()
 	{
+		LatestFrame = INDEX_NONE;
 		// Iterate over current data to find latest frame
 		for (int32 FrameIndex = 0; FrameIndex < NumFrames; ++FrameIndex)
 		{
@@ -240,10 +249,10 @@ public :
 
 				// Don't update LocalFrame here so that it can be used to check if this data was grabbed from the correct frame or not
 				ExtractedData.ServerFrame = DataHistory[MinFrameIndex].ServerFrame + DeltaFrame;
-				ExtractedData.InputFrame = INDEX_NONE; // Clear InputFrame since this history entry is now altered and doesn't correspond to the source entry with the same InputFrame
+				ExtractedData.bDataAltered = true; // This history entry is now altered and doesn't correspond to the source entry
 
 #if DEBUG_NETWORK_PHYSICS
-				UE_LOG(LogChaos, Log, TEXT("		Interpolating data between frame %d and %d - > [%d %d]"), DataHistory[MinFrameIndex].LocalFrame, DataHistory[MaxFrameIndex].LocalFrame, static_cast<DataType*>(HistoryData)->InputFrame, static_cast<DataType*>(HistoryData)->ServerFrame);
+				UE_LOG(LogChaos, Log, TEXT("		Interpolating data between frame %d and %d - > [%d]"), DataHistory[MinFrameIndex].LocalFrame, DataHistory[MaxFrameIndex].LocalFrame, static_cast<DataType*>(HistoryData)->ServerFrame);
 #endif
 				return true;
 			}
@@ -256,7 +265,7 @@ public :
 
 				// Don't update LocalFrame here so that it can be used to check how many frames ahead of the last entry we are currently
 				ExtractedData.ServerFrame = DataHistory[MinFrameIndex].ServerFrame + DeltaFrame;
-				ExtractedData.InputFrame = INDEX_NONE; // Clear InputFrame since this history entry is now altered and doesn't correspond to the source entry with the same InputFrame
+				ExtractedData.bDataAltered = true; // This history entry is now altered and doesn't correspond to the source entry
 
 #if DEBUG_NETWORK_PHYSICS
 				UE_LOG(LogChaos, Log, TEXT("		Setting data to frame %d"), DataHistory[MinFrameIndex].LocalFrame);
@@ -283,7 +292,7 @@ public :
 			if (FromFrame == DataHistory[FromIndex].LocalFrame)
 			{
 				static_cast<DataType*>(ToData)->MergeData(DataHistory[FromIndex]);
-				static_cast<DataType*>(ToData)->InputFrame = INDEX_NONE; // Clear InputFrame since this history entry is now altered and doesn't correspond to the source entry with the same InputFrame
+				static_cast<DataType*>(ToData)->bDataAltered = true; // This history entry is now altered and doesn't correspond to the source entry
 			}
 		}
 	}
