[NetPhysics] Update physics replication island callbacks to the more optimized version.
- Along with smaller cleanup in physics replication files.

diff --git a/Engine/Source/Runtime/Engine/Private/PhysicsEngine/PhysicsReplication.cpp b/Engine/Source/Runtime/Engine/Private/PhysicsEngine/PhysicsReplication.cpp
--- a/Engine/Source/Runtime/Engine/Private/PhysicsEngine/PhysicsReplication.cpp
+++ b/Engine/Source/Runtime/Engine/Private/PhysicsEngine/PhysicsReplication.cpp
@@ -18,6 +18,8 @@
 #include "PBDRigidsSolver.h"
 #include "Chaos/DebugDrawQueue.h"
 #include "Chaos/Particles.h"
+#include "Chaos/Island/IslandManager.h"
+
 
 namespace CharacterMovementCVars
 {
@@ -395,6 +397,12 @@ void FPhysicsReplication::OnTick(float DeltaSeconds, TMap<TWeakObjectPtr<UPrimit
 		return;
 	}
 
+	// Don't tick unless we have data to process
+	if (ComponentsToTargets.Num() == 0 && ReplicatedTargetsQueue.Num() == 0)
+	{
+		return;
+	}
+
 	int32 LocalFrameOffset = 0; // LocalFrame = ServerFrame + LocalFrameOffset;
 	bool LocalFrameOffsetAssigned = false;
 
@@ -1222,7 +1230,9 @@ void FPhysicsReplicationAsync::CacheResimInteractions()
 		return;
 	}
 
-	ParticlesInResimIslands.Empty(FMath::CeilToInt(static_cast<float>(ParticlesInResimIslands.Num()) * 0.9f));
+	ResimIslands.Reset();
+	ResimIslandsParticles.Reset();
+	ParticlesInResimIslands.Reset();
 
 	Chaos::Private::FPBDIslandManager& IslandManager = RigidsSolver->GetEvolution()->GetIslandManager();
 	Chaos::FReadPhysicsObjectInterface_Internal Interface = Chaos::FPhysicsObjectInternalInterface::GetRead();
@@ -1236,7 +1246,9 @@ void FPhysicsReplicationAsync::CacheResimInteractions()
 			if (Chaos::FGeometryParticleHandle* Handle = Interface.GetParticle(POHandle))
 			{
 				// Get a list of particles from the same island as a resim particle is in, i.e. particles interacting with a resim particle
-				for (const Chaos::FGeometryParticleHandle* InteractParticle : IslandManager.FindParticlesInIslands(IslandManager.FindParticleIslands(Handle)))
+				IslandManager.FindParticleIslands(Handle, OUT ResimIslands);
+				IslandManager.FindParticlesInIslands(ResimIslands, OUT ResimIslandsParticles);
+				for (const Chaos::FGeometryParticleHandle* InteractParticle : ResimIslandsParticles)
 				{
 					ParticlesInResimIslands.Add(InteractParticle->GetHandleIdx());
 				}
@@ -1344,19 +1356,19 @@ void FPhysicsReplicationAsync::ApplyPhysicsReplicationLOD(Chaos::FConstPhysicsOb
 		if (Target.RepMode == EPhysicsReplicationMode::PredictiveInterpolation)
 		{
 			const bool bShouldSleep = (Target.TargetState.Flags & ERigidBodyFlags::Sleeping) != 0;
-			int32 TargetSourceClientFrame = (Target.ServerFrame - Target.FrameOffset);
+			int32 TargetClientFrame = (Target.ServerFrame - Target.FrameOffset);
 
 			// If we use Predicitve Interpolation and we should not sleep and the aligned frame from LOD is ahead of the target, perform LOD aligment extrapolation
-			if (!bShouldSleep && LodData->AlignedFrame > TargetSourceClientFrame)
+			if (!bShouldSleep && LodData->AlignedFrame > TargetClientFrame)
 			{
 				// Calculate how far to forward predict and extrapolate target by that amount
-				const int32 FullPredictionFrames = RigidsSolver.GetCurrentFrame() - TargetSourceClientFrame;
+				const int32 FullPredictionFrames = RigidsSolver.GetCurrentFrame() - TargetClientFrame;
 				const float FullPredictionTime = (FullPredictionFrames * GetDeltaTime_Internal());
 				const float AlignedPredictionTime = FullPredictionTime - LodData->AlignedTime;
 				FPhysicsReplicationAsync::ExtrapolateTarget(Target, AlignedPredictionTime);
 
 				// Update tick count based on LOD alignment
-				Target.TickCount = (LodData->AlignedFrame + Target.FrameOffset) - Target.ServerFrame;  
+				Target.TickCount = LodData->AlignedFrame - TargetClientFrame;
 			}
 		}
 	}
diff --git a/Engine/Source/Runtime/Engine/Private/PhysicsEngine/PhysicsReplicationCache.cpp b/Engine/Source/Runtime/Engine/Private/PhysicsEngine/PhysicsReplicationCache.cpp
--- a/Engine/Source/Runtime/Engine/Private/PhysicsEngine/PhysicsReplicationCache.cpp
+++ b/Engine/Source/Runtime/Engine/Private/PhysicsEngine/PhysicsReplicationCache.cpp
@@ -169,6 +169,11 @@ void FPhysicsReplicationCacheAsync::OnPhysicsObjectUnregistered_Internal(Chaos::
 {
 	// Unregister physics object from caching state
 	ReplicationCache_Internal.Remove(PhysicsObject);
+	if (const FPhysicsReplicationCacheAsyncInput* AsyncInput = GetConsumerInput_Internal())
+	{
+		// Remove from AsyncInput->AccessedObjects so that we don't potentially add this back again in ProcessAsyncInputs()
+		AsyncInput->AccessedObjects.Remove(PhysicsObject);
+	}
 	bUpdateAfterRemoval = true; // Ensure we produce an async output after this removal, even if the cache is empty
 }
 
diff --git a/Engine/Source/Runtime/Engine/Private/PhysicsEngine/PhysicsReplicationLOD.cpp b/Engine/Source/Runtime/Engine/Private/PhysicsEngine/PhysicsReplicationLOD.cpp
--- a/Engine/Source/Runtime/Engine/Private/PhysicsEngine/PhysicsReplicationLOD.cpp
+++ b/Engine/Source/Runtime/Engine/Private/PhysicsEngine/PhysicsReplicationLOD.cpp
@@ -7,6 +7,7 @@
 #include <Components/PrimitiveComponent.h>
 #include "Chaos/ParticleDirtyFlags.h"
 #include "Chaos/DebugDrawQueue.h"
+#include "Chaos/Island/IslandManager.h"
 
 
 namespace PhysicsReplicationLODCVars
@@ -137,7 +138,7 @@ IPhysicsReplicationLODAsync* FPhysicsReplicationLOD::GetPhysicsReplicationLOD_In
 
 bool FPhysicsReplicationLODAsync::IsEnabled() const
 {
-	return (bEnabled || PhysicsReplicationLODCVars::bEnabled); 
+	return (bEnabled || PhysicsReplicationLODCVars::bEnabled) && FocalParticles.Num() > 0;
 }
 
 void FPhysicsReplicationLODAsync::OnPostInitialize_Internal()
diff --git a/Engine/Source/Runtime/Engine/Public/Physics/PhysicsReplicationCache.h b/Engine/Source/Runtime/Engine/Public/Physics/PhysicsReplicationCache.h
--- a/Engine/Source/Runtime/Engine/Public/Physics/PhysicsReplicationCache.h
+++ b/Engine/Source/Runtime/Engine/Public/Physics/PhysicsReplicationCache.h
@@ -53,7 +53,7 @@ private:
 /** Async Marshal Input */
 struct FPhysicsReplicationCacheAsyncInput : public Chaos::FSimCallbackInput
 {
-	TArray<Chaos::FConstPhysicsObjectHandle> AccessedObjects;
+	mutable TArray<Chaos::FConstPhysicsObjectHandle> AccessedObjects;
 	TArray<Chaos::FConstPhysicsObjectHandle> UnregisterObjects;
 
 	void Reset()
diff --git a/Engine/Source/Runtime/Engine/Public/PhysicsReplication.h b/Engine/Source/Runtime/Engine/Public/PhysicsReplication.h
--- a/Engine/Source/Runtime/Engine/Public/PhysicsReplication.h
+++ b/Engine/Source/Runtime/Engine/Public/PhysicsReplication.h
@@ -18,6 +18,14 @@
 #include "Physics/PhysicsInterfaceUtils.h"
 #include "Physics/NetworkPhysicsSettingsComponent.h"
 
+namespace Chaos
+{
+	namespace Private
+	{
+		class FPBDIsland;
+	}
+}
+
 namespace CharacterMovementCVars
 {
 	extern ENGINE_API int32 SkipPhysicsReplication;
@@ -222,6 +230,8 @@ private:
 	FNetworkPhysicsSettingsAsync SettingsDefault;
 	TMap<Chaos::FConstPhysicsObjectHandle, FReplicatedPhysicsTargetAsync> ObjectToTarget;
 	TMap<Chaos::FConstPhysicsObjectHandle, FNetworkPhysicsSettingsAsync> ObjectToSettings;
+	TArray<const Chaos::Private::FPBDIsland*> ResimIslands;
+	TArray<const Chaos::FGeometryParticleHandle*> ResimIslandsParticles;
 	TArray<int32> ParticlesInResimIslands;
 	TArray<Chaos::FParticleID> ReplicatedParticleIDs;
 
diff --git a/Engine/Source/Runtime/Engine/Public/PhysicsReplicationLOD.h b/Engine/Source/Runtime/Engine/Public/PhysicsReplicationLOD.h
--- a/Engine/Source/Runtime/Engine/Public/PhysicsReplicationLOD.h
+++ b/Engine/Source/Runtime/Engine/Public/PhysicsReplicationLOD.h
@@ -8,8 +8,14 @@
 #include "Physics/PhysicsInterfaceDeclares.h"
 #include "Physics/PhysicsInterfaceUtils.h"
 #include "PhysicsEngine/PhysicsSettings.h"
-#include "Chaos/Island/IslandManager.h"
 
+namespace Chaos
+{
+	namespace Private
+	{
+		class FPBDIsland;
+	}
+}
 
 enum EPhysicsReplicationLODFlags : uint32
 {
diff --git a/Engine/Source/Runtime/Experimental/Chaos/Private/RewindData.cpp b/Engine/Source/Runtime/Experimental/Chaos/Private/RewindData.cpp
--- a/Engine/Source/Runtime/Experimental/Chaos/Private/RewindData.cpp
+++ b/Engine/Source/Runtime/Experimental/Chaos/Private/RewindData.cpp
@@ -9,6 +9,7 @@
 #include "Misc/Paths.h"
 #include "Misc/FileHelper.h"
 #include "Misc/Paths.h"
+#include "Chaos/Island/IslandManager.h"
 
 namespace Chaos
 {
@@ -1225,20 +1226,20 @@ int32 FRewindData::FindValidResimFrame(const int32 RequestedFrame)
 	Private::FPBDIslandManager& IslandManager = Solver->GetEvolution()->GetIslandManager();
 
 	// Cache all particles in islands that have a resim triggering particle
-	TArray<const FGeometryParticleHandle*> ResimIslandParticles;
 	if (ResimFrameValidation == (int32)EResimFrameValidation::IslandValidation)
 	{
-		TArray<const Private::FPBDIsland*> ResimIslands;
+		IslandValidationIslands.Reset();
+		IslandValidationIslandParticles.Reset();
 		for (FDirtyParticleInfo& DirtyParticleInfo : DirtyParticles)
 		{
 			FGeometryParticleHandle* Handle = DirtyParticleInfo.GetObjectPtr();
 			if (IslandManager.GetParticleResimFrame(Handle) != INDEX_NONE)
 			{
-				ResimIslands.AddUnique(IslandManager.GetParticleIsland(Handle));
+				IslandManager.FindParticleIslands(Handle, OUT IslandValidationIslands);
 			}
 		}
 
-		ResimIslandParticles = IslandManager.FindParticlesInIslands(ResimIslands);
+		IslandManager.FindParticlesInIslands(IslandValidationIslands, OUT IslandValidationIslandParticles);
 	}
 
 	// First frame of the history datas
@@ -1251,13 +1252,13 @@ int32 FRewindData::FindValidResimFrame(const int32 RequestedFrame)
 		bHasTargetHistory = true;
 		
 #if DEBUG_REWIND_DATA
-		UE_LOG(LogChaos, Log, TEXT("CLIENT | PT | FindValidResimFrame | Processing resim particles | History Frame: %d | Total Particle Count: %d | ResimIslands Particle Count: %d | ResimFrameValidation: %d"), ValidFrame, DirtyParticles.Num(), ResimIslandParticles.Num(), ResimFrameValidation);
+		UE_LOG(LogChaos, Log, TEXT("CLIENT | PT | FindValidResimFrame | Processing resim particles | History Frame: %d | Total Particle Count: %d | ResimIslands Particle Count: %d | ResimFrameValidation: %d"), ValidFrame, DirtyParticles.Num(), IslandValidationIslandParticles.Num(), ResimFrameValidation);
 #endif
 
 		if ((EResimFrameValidation)ResimFrameValidation == EResimFrameValidation::IslandValidation)
 		{
 			// Iterate over islands previously found having resim particles in them and check if the particles in the islands have targets
-			for (const FGeometryParticleHandle* IslandParticle : ResimIslandParticles)
+			for (const FGeometryParticleHandle* IslandParticle : IslandValidationIslandParticles)
 			{
 				// Cache particle handles for objects in islands that need resim
 				if (FDirtyParticleInfo* DirtyParticleInfo = FindDirtyObj(*IslandParticle))
diff --git a/Engine/Source/Runtime/Experimental/Chaos/Public/RewindData.h b/Engine/Source/Runtime/Experimental/Chaos/Public/RewindData.h
--- a/Engine/Source/Runtime/Experimental/Chaos/Public/RewindData.h
+++ b/Engine/Source/Runtime/Experimental/Chaos/Public/RewindData.h
@@ -22,6 +22,10 @@
 
 namespace Chaos
 {
+	namespace Private
+	{
+		class FPBDIsland;
+	}
 
 /** Base rewind history used in the rewind data */
 struct FBaseRewindHistory
@@ -1990,6 +1994,10 @@ private:
 	// Used to block rewinding past a physics change we currently don't handle
 	int32 BlockResimFrame = INDEX_NONE;
 
+	// Properties for EResimFrameValidation::IslandValidation logic
+	TArray<const Private::FPBDIsland*> IslandValidationIslands;
+	TArray<const FGeometryParticleHandle*> IslandValidationIslandParticles;
+
 	template <typename TObj>
 	bool IsResimAndInSync(const TObj& Handle) const { return IsResim() && Handle.SyncState() == ESyncState::InSync; }
 
