Improved a couple of assert messages

diff --git a/Engine/Source/Editor/UnrealEd/Private/EditorViewportClient.cpp b/Engine/Source/Editor/UnrealEd/Private/EditorViewportClient.cpp
--- a/Engine/Source/Editor/UnrealEd/Private/EditorViewportClient.cpp
+++ b/Engine/Source/Editor/UnrealEd/Private/EditorViewportClient.cpp
@@ -622,7 +622,7 @@ bool FEditorViewportClient::RemoveRealtimeOverride(FText SystemDisplayName, bool
 			break;
 		}
 	}
-	ensureMsgf(!bCheckMissingOverride || bRemoved, TEXT("No realtime override was found with the given SystemDisplayName"));
+	ensureMsgf(!bCheckMissingOverride || bRemoved, TEXT("No realtime override was found with %s"), *SystemDisplayString);
 	
 	if (bRemoved)
 	{
diff --git a/Engine/Source/Runtime/Engine/Private/World.cpp b/Engine/Source/Runtime/Engine/Private/World.cpp
--- a/Engine/Source/Runtime/Engine/Private/World.cpp
+++ b/Engine/Source/Runtime/Engine/Private/World.cpp
@@ -9382,7 +9382,7 @@ void UWorld::ChangeFeatureLevel(ERHIFeatureLevel::Type InFeatureLevel, bool bSho
 {
 	if (InFeatureLevel != GetFeatureLevel() || bForceUpdate)
 	{
-		UE_LOG(LogWorld, Log, TEXT("Changing Feature Level (Enum) from %i to %i%s"), (int)GetFeatureLevel(), (int)InFeatureLevel, bForceUpdate ? TEXT(" (forced)") : TEXT(""));
+		UE_LOG(LogWorld, Log, TEXT("Changing Feature Level (Enum) from %s to %s%s"), *LexToString(GetFeatureLevel()), *LexToString(InFeatureLevel), bForceUpdate ? TEXT(" (forced)") : TEXT(""));
 		FScopedSlowTask SlowTask(100.f, NSLOCTEXT("Engine", "ChangingPreviewRenderingLevelMessage", "Changing Preview Rendering Level"), bShowSlowProgressDialog);
 		SlowTask.MakeDialog();
 		{
