Implementing UMaterialExpressionMin::Build for the new translator

- The op already exists so simply hook into the existing op.

#rnx

diff --git a/Engine/Source/Runtime/Engine/Private/Materials/MaterialExpressionsIR.cpp b/Engine/Source/Runtime/Engine/Private/Materials/MaterialExpressionsIR.cpp
--- a/Engine/Source/Runtime/Engine/Private/Materials/MaterialExpressionsIR.cpp
+++ b/Engine/Source/Runtime/Engine/Private/Materials/MaterialExpressionsIR.cpp
@@ -498,6 +498,11 @@ void UMaterialExpressionDivide::Build(MIR::FEmitter& Emitter)
 	BuildBinaryArithmeticOperator(Emitter, MIR::BO_Divide, &A, ConstA, &B, ConstB);
 }
 
+void UMaterialExpressionMin::Build(MIR::FEmitter& Emitter)
+{ 
+	BuildBinaryArithmeticOperator(Emitter, MIR::BO_Min, &A, ConstA, &B, ConstB);
+}
+
 void UMaterialExpressionOneMinus::Build(MIR::FEmitter& Emitter)
 {
 	// Default input to zero if not connected, then get it as a primitive.
diff --git a/Engine/Source/Runtime/Engine/Public/Materials/MaterialExpressionMin.h b/Engine/Source/Runtime/Engine/Public/Materials/MaterialExpressionMin.h
--- a/Engine/Source/Runtime/Engine/Public/Materials/MaterialExpressionMin.h
+++ b/Engine/Source/Runtime/Engine/Public/Materials/MaterialExpressionMin.h
@@ -30,6 +30,7 @@ class UMaterialExpressionMin : public UMaterialExpression
 
 	//~ Begin UMaterialExpression Interface
 #if WITH_EDITOR
+	virtual void Build(UE::MIR::FEmitter& Emitter) override;
 	virtual int32 Compile(class FMaterialCompiler* Compiler, int32 OutputIndex) override;
 	virtual void GetCaption(TArray<FString>& OutCaptions) const override;
 #endif
