Minor fixup for EngineTest compiling with UE_WITH_IRIS=0

diff --git a/Engine/Source/Runtime/Engine/Private/Tests/NetTestHelpers.cpp b/Engine/Source/Runtime/Engine/Private/Tests/NetTestHelpers.cpp
--- a/Engine/Source/Runtime/Engine/Private/Tests/NetTestHelpers.cpp
+++ b/Engine/Source/Runtime/Engine/Private/Tests/NetTestHelpers.cpp
@@ -194,8 +194,10 @@ FTestWorldInstance::FContext FTestWorldInstance::GetTestContext() const
 	{ 
 		.World = GetWorld(),
 		.NetDriver = GetNetDriver(),
+#if UE_WITH_IRIS
 		.IrisRepSystem = RepSystem,
 		.IrisRepBridge = RepBridge,
+#endif
 	};
 }
 
@@ -477,10 +479,10 @@ UObject* FTestWorlds::FindReplicatedObjectOnClient(UObject* ServerObject, uint32
 	const FTestWorldInstance::FContext ServerContext = Server.GetTestContext();
 	const FTestWorldInstance::FContext ClientContext = Clients[ClientIndex].GetTestContext();
 
+#if UE_WITH_IRIS
 	if (ServerContext.NetDriver->IsUsingIrisReplication())
 	{
-#if UE_WITH_IRIS
-	const FNetRefHandle NetHandle = ServerContext.IrisRepBridge->GetReplicatedRefHandle(ServerObject);
+		const FNetRefHandle NetHandle = ServerContext.IrisRepBridge->GetReplicatedRefHandle(ServerObject);
 
 		if (!NetHandle.IsValid())
 		{
@@ -489,10 +491,9 @@ UObject* FTestWorlds::FindReplicatedObjectOnClient(UObject* ServerObject, uint32
 		}
 
 		return ClientContext.IrisRepBridge->GetReplicatedObject(NetHandle);
-
-#endif // UE_WITH_IRIS
 	}
 	else
+#endif // UE_WITH_IRIS
 	{
 		const FNetworkGUID NetGUID = ServerContext.NetDriver->GetNetGuidCache()->GetNetGUID(ServerObject);
 
