Added GetReplayTotalTime function to the ReplaySubsystem

diff --git a/Engine/Source/Runtime/Engine/Private/ReplaySubsystem.cpp b/Engine/Source/Runtime/Engine/Private/ReplaySubsystem.cpp
--- a/Engine/Source/Runtime/Engine/Private/ReplaySubsystem.cpp
+++ b/Engine/Source/Runtime/Engine/Private/ReplaySubsystem.cpp
@@ -336,6 +336,23 @@ FString UReplaySubsystem::GetActiveReplayName() const
 	return FString();
 }
 
+float UReplaySubsystem::GetReplayTotalTime() const
+{
+	UWorld* CurrentWorld = GetWorld();
+
+	if (CurrentWorld != nullptr && CurrentWorld->GetDemoNetDriver() != nullptr)
+	{
+		return CurrentWorld->GetDemoNetDriver()->GetDemoTotalTime();
+	}
+
+	if (UReplayNetConnection* Connection = ReplayConnection.Get())
+	{
+		return Connection->GetReplayTotalTime();
+	}
+	
+	return 0.0f;
+}
+
 float UReplaySubsystem::GetReplayCurrentTime() const
 {
 	UWorld* CurrentWorld = GetWorld();
diff --git a/Engine/Source/Runtime/Engine/Public/ReplayNetConnection.h b/Engine/Source/Runtime/Engine/Public/ReplayNetConnection.h
--- a/Engine/Source/Runtime/Engine/Public/ReplayNetConnection.h
+++ b/Engine/Source/Runtime/Engine/Public/ReplayNetConnection.h
@@ -47,6 +47,7 @@ public:
 
 	FString GetActiveReplayName() const { return ReplayHelper.ActiveReplayName; }
 	float GetReplayCurrentTime() const { return ReplayHelper.DemoCurrentTime; }
+	float GetReplayTotalTime() const { return ReplayHelper.DemoTotalTime; }
 
 	ENGINE_API bool IsSavingCheckpoint() const;
 
diff --git a/Engine/Source/Runtime/Engine/Public/ReplaySubsystem.h b/Engine/Source/Runtime/Engine/Public/ReplaySubsystem.h
--- a/Engine/Source/Runtime/Engine/Public/ReplaySubsystem.h
+++ b/Engine/Source/Runtime/Engine/Public/ReplaySubsystem.h
@@ -61,6 +61,14 @@ public:
 	UFUNCTION(BlueprintCallable, Category=Replay)
 	ENGINE_API FString GetActiveReplayName() const;
 
+	/**
+	 * Get current recording/playing total replay time
+	 *
+	* @return float total recording/playing time in seconds
+	*/
+	UFUNCTION(BlueprintCallable, Category = Replay)
+	ENGINE_API float GetReplayTotalTime() const;
+
 	/**
 	 * Get current recording/playing replay time
 	 *
