Fix missing version of UE_VLOG_ALWAYS_UELOG macro from previous change

diff --git a/Engine/Source/Runtime/Engine/Public/VisualLogger/VisualLogger.h b/Engine/Source/Runtime/Engine/Public/VisualLogger/VisualLogger.h
--- a/Engine/Source/Runtime/Engine/Public/VisualLogger/VisualLogger.h
+++ b/Engine/Source/Runtime/Engine/Public/VisualLogger/VisualLogger.h
@@ -128,6 +128,10 @@
 #define UE_CVLOG(Condition, LogOwner, CategoryName, Verbosity, Format, ...)
 // Make sure to UE_LOG when Visual Logging is disabled.
 #define UE_VLOG_UELOG(LogOwner, CategoryName, Verbosity, Format, ...) { UE_LOG(CategoryName, Verbosity, Format, ##__VA_ARGS__); }
+
+// Deprecated.
+#define UE_VLOG_ALWAYS_UELOG(LogOwner, CategoryName, Verbosity, Format, ...) UE_DEPRECATED_MACRO(5.6, "Use UE_VLOG_UELOG instead, it now behaves like the 'always' version.") { UE_LOG(CategoryName, Verbosity, Format, ##__VA_ARGS__); }
+
 #define UE_VLOG_SEGMENT(LogOwner, CategoryName, Verbosity, SegmentStart, SegmentEnd, Color, DescriptionFormat, ...)
 #define UE_CVLOG_SEGMENT(Condition, LogOwner, CategoryName, Verbosity, SegmentStart, SegmentEnd, Color, DescriptionFormat, ...)
 #define UE_VLOG_SEGMENT_THICK(LogOwner, CategoryName, Verbosity, SegmentStart, SegmentEnd, Color, Thickness, DescriptionFormat, ...)
