Added `trunc` material expression to the new translator

- Added a new unary Operator EOperator::UO_Truncate
- Implement UMaterialExpressionTruncate::Build
- FMaterialIRToHLSLTranslation::Lower now handles unary operators and will assert if the wrong op type (unary, binary, ternary) has too many inputs (this should probably be caught eariler?)
- Implement FoldScalarArithmeticOperator for Truncate.  Needed to special case based on type since we don't have a templated trunc function.  If the type is unsupported this will generate runtime code that fails (I can't static assert since this templated function is called with int)
- Adds BuildUnaryArithmeticOperatorFromInput() which doesn't take a default (you must plug in a value), and refactored out common code into BuildUnaryArithmeticOperatorFromValue.
- Refactored BuildUnaryArithmeticOperatorFromInput() to be BuildUnaryArithmeticOperatorFromInputWithDefault().

#rnx

diff --git a/Engine/Source/Runtime/Engine/Private/Materials/MaterialExpressionsIR.cpp b/Engine/Source/Runtime/Engine/Private/Materials/MaterialExpressionsIR.cpp
--- a/Engine/Source/Runtime/Engine/Private/Materials/MaterialExpressionsIR.cpp
+++ b/Engine/Source/Runtime/Engine/Private/Materials/MaterialExpressionsIR.cpp
@@ -425,8 +425,19 @@ void UMaterialExpressionAppendVector::Build(MIR::FEmitter& Emitter)
 }
 
 /* Unary Operators */
+static void BuildUnaryArithmeticOperatorFromValue(
+	MIR::FEmitter& Emitter,
+	MIR::EOperator Op,
+	MIR::FValue* Value)
+{
+	// Emit the unary operator.
+	Value = Emitter.EmitOperator(Op, Value);
 
-static void BuildUnaryArithmeticOperator(
+	// And flow it out of the expression's only output.
+	Emitter.Output(0, Value);
+}
+
+static void BuildUnaryArithmeticOperatorFromInputWithDefault(
 	MIR::FEmitter& Emitter,
 	MIR::EOperator Op,
 	FExpressionInput* Input,
@@ -437,14 +448,29 @@ static void BuildUnaryArithmeticOperator(
 
 	UE_MIR_CHECKPOINT();
 
-	// Finally emit the unary operator.
-	Value = Emitter.EmitOperator(Op, Value);
+	BuildUnaryArithmeticOperatorFromValue(Emitter, Op, Value);
+}
 
-	// And flow it out of the expression's only output.
-	Emitter.Output(0, Value);
+// Special case function which doesn't try to emit a default value and will error if you don't have an Input
+static void BuildUnaryArithmeticOperatorFromInput(
+	MIR::FEmitter& Emitter,
+	MIR::EOperator Op,
+	FExpressionInput* Input)
+{
+	// Get input after checking it has primitive type.
+	MIR::FValue* Value = Emitter.Input(Input);
+
+	UE_MIR_CHECKPOINT();
+
+	BuildUnaryArithmeticOperatorFromValue(Emitter, Op, Value);
 }
 
-// TODO: Add the missing expressions 
+// TODO: Add the missing expressions
+
+void UMaterialExpressionTruncate::Build(UE::MIR::FEmitter& Emitter)
+{
+	BuildUnaryArithmeticOperatorFromInput(Emitter, MIR::UO_Truncate, &Input);
+}
 
 /* Binary Operators */
 
diff --git a/Engine/Source/Runtime/Engine/Private/Materials/MaterialIR.cpp b/Engine/Source/Runtime/Engine/Private/Materials/MaterialIR.cpp
--- a/Engine/Source/Runtime/Engine/Private/Materials/MaterialIR.cpp
+++ b/Engine/Source/Runtime/Engine/Private/Materials/MaterialIR.cpp
@@ -412,6 +412,9 @@ const TCHAR* OperatorToString(EOperator Op)
 	{
 		case O_Invalid: return TEXT("Invalid");
 
+		/* Unary Arithmetic */
+		case UO_Truncate: return TEXT("Truncate");
+
 		/* Binary Comparisons */
 		case BO_GreaterThan: return TEXT("GreaterThan");
 		case BO_GreaterThanOrEquals: return TEXT("GreaterThanOrEquals");
diff --git a/Engine/Source/Runtime/Engine/Private/Materials/MaterialIREmitter.cpp b/Engine/Source/Runtime/Engine/Private/Materials/MaterialIREmitter.cpp
--- a/Engine/Source/Runtime/Engine/Private/Materials/MaterialIREmitter.cpp
+++ b/Engine/Source/Runtime/Engine/Private/Materials/MaterialIREmitter.cpp
@@ -757,6 +757,21 @@ static T FoldScalarArithmeticOperator(EOperator Operator, T A, T B, T C)
 	switch (Operator)
 	{
 		case UO_Negate: return -A;
+		case UO_Truncate:
+		{
+			if constexpr (std::is_same<T, float>::value)
+			{
+				return FMath::TruncToFloat(A);
+			}
+			else if constexpr (std::is_same<T, double>::value)
+			{
+				return FMath::TruncToDouble(A);
+			}
+			else
+			{
+				UE_MIR_UNREACHABLE();
+			}
+		}
 		case BO_Add: return A + B;
 		case BO_Subtract: return A - B;
 		case BO_Multiply: return A * B;
@@ -859,6 +874,7 @@ static FValue* TryFoldOperatorScalar(FEmitter* Emitter, EOperator Op, FValue* A,
 		/* No simplifications for these operations */
 		case UO_Negate:
 		case UO_Not:
+		case UO_Truncate:
 		case BO_Min:
 		case BO_Max:
 			break;
diff --git a/Engine/Source/Runtime/Engine/Private/Materials/MaterialIRToHLSLTranslator.cpp b/Engine/Source/Runtime/Engine/Private/Materials/MaterialIRToHLSLTranslator.cpp
--- a/Engine/Source/Runtime/Engine/Private/Materials/MaterialIRToHLSLTranslator.cpp
+++ b/Engine/Source/Runtime/Engine/Private/Materials/MaterialIRToHLSLTranslator.cpp
@@ -459,6 +459,7 @@ struct FTranslator : FMaterialIRToHLSLTranslation
 			const TCHAR* OpString;
 			switch (Operator->Op)
 			{
+				case MIR::UO_Truncate: OpString = TEXT("trunc"); break;
 				case MIR::BO_Min: OpString = TEXT("min"); break;
 				case MIR::BO_Max: OpString = TEXT("max"); break;
 				case MIR::TO_Lerp: OpString = TEXT("lerp"); break;
@@ -466,9 +467,20 @@ struct FTranslator : FMaterialIRToHLSLTranslation
 				default: UE_MIR_UNREACHABLE();
 			}
 
-			Printer << OpString << "(" << LowerValue(Operator->AArg) << ", " << LowerValue(Operator->BArg);
+			// Unary
+			Printer << OpString << "(" << LowerValue(Operator->AArg);
+
+			// Binary
+			if (Operator->BArg)
+			{
+				check(MIR::IsBinaryOperator(Operator->Op));
+				Printer << ", " << LowerValue(Operator->BArg);
+			}
+
+			// Ternary
 			if (Operator->CArg)
 			{
+				check(MIR::IsTernaryOperator(Operator->Op));
 				Printer << ", " << LowerValue(Operator->CArg);
 			}
 			Printer << ")";
diff --git a/Engine/Source/Runtime/Engine/Public/Materials/MaterialExpressionTruncate.h b/Engine/Source/Runtime/Engine/Public/Materials/MaterialExpressionTruncate.h
--- a/Engine/Source/Runtime/Engine/Public/Materials/MaterialExpressionTruncate.h
+++ b/Engine/Source/Runtime/Engine/Public/Materials/MaterialExpressionTruncate.h
@@ -18,6 +18,7 @@ class UMaterialExpressionTruncate : public UMaterialExpression
 
 	//~ Begin UMaterialExpression Interface
 #if WITH_EDITOR
+	virtual void Build(UE::MIR::FEmitter& Emitter) override;
 	virtual int32 Compile(class FMaterialCompiler* Compiler, int32 OutputIndex) override;
 	virtual uint32 GetOutputType(int32 OutputIndex) override;
 	virtual void GetCaption(TArray<FString>& OutCaptions) const override;
diff --git a/Engine/Source/Runtime/Engine/Public/Materials/MaterialIR.h b/Engine/Source/Runtime/Engine/Public/Materials/MaterialIR.h
--- a/Engine/Source/Runtime/Engine/Public/Materials/MaterialIR.h
+++ b/Engine/Source/Runtime/Engine/Public/Materials/MaterialIR.h
@@ -309,6 +309,7 @@ enum EOperator
 	/* Unary */
 	UO_Negate, 	// Arithmetic negation: negate(5) -> -5
 	UO_Not,		// Logical negation: not(true) -> false
+	UO_Truncate,
 	
 	/* Binary Comparison Operators */
 	BO_GreaterThan,
