Prevent the Physics System and Network Stats from causing migrations when attempting to get Ping from the PlayerState

#rnx

diff --git a/Engine/Source/Runtime/Engine/Private/NetDriver.cpp b/Engine/Source/Runtime/Engine/Private/NetDriver.cpp
--- a/Engine/Source/Runtime/Engine/Private/NetDriver.cpp
+++ b/Engine/Source/Runtime/Engine/Private/NetDriver.cpp
@@ -8436,8 +8436,20 @@ void UNetDriver::UpdateNetworkStats()
 
 						if (Connection != nullptr)
 						{
-							if (Connection->PlayerController != nullptr && Connection->PlayerController->PlayerState != nullptr)
+							if (Connection->PlayerController && Connection->PlayerController->PlayerState)
 							{
+#if UE_WITH_REMOTE_OBJECT_HANDLE
+								if (Connection->PlayerController.IsRemote())
+								{
+									ensureMsgf(false, TEXT("Expecting the Connection's PlayerController to be locally owned by now"));
+									continue;
+								}
+								else if (Connection->PlayerController->PlayerState.IsRemote())
+								{
+									continue;
+								}
+#endif
+
 								// Ping value calculated per client
 								float ConnPing = Connection->PlayerController->PlayerState->ExactPing;
 
diff --git a/Engine/Source/Runtime/Engine/Private/PhysicsEngine/PhysicsReplication.cpp b/Engine/Source/Runtime/Engine/Private/PhysicsEngine/PhysicsReplication.cpp
--- a/Engine/Source/Runtime/Engine/Private/PhysicsEngine/PhysicsReplication.cpp
+++ b/Engine/Source/Runtime/Engine/Private/PhysicsEngine/PhysicsReplication.cpp
@@ -2331,9 +2331,12 @@ float FPhysicsReplication::GetLocalPing() const
 	{
 		if (APlayerController* PlayerController = World->GetFirstPlayerController())
 		{
-			if (APlayerState* PlayerState = PlayerController->PlayerState)
+			if (const TObjectPtr<APlayerState>& PlayerState = PlayerController->PlayerState)
 			{
-				return PlayerState->ExactPing;
+				if (!PlayerState.IsRemote())
+				{
+					return PlayerState->ExactPing;
+				}
 			}
 		}
 
