Add View.NearPlane to UMaterialExpressionViewProperty so it can be accessed in the material

[RN] minor | Rendering - Materials and Shaders

diff --git a/Engine/Source/Runtime/Engine/Private/Materials/HLSLMaterialTranslator.cpp b/Engine/Source/Runtime/Engine/Private/Materials/HLSLMaterialTranslator.cpp
--- a/Engine/Source/Runtime/Engine/Private/Materials/HLSLMaterialTranslator.cpp
+++ b/Engine/Source/Runtime/Engine/Private/Materials/HLSLMaterialTranslator.cpp
@@ -2184,6 +2184,11 @@ void FHLSLMaterialTranslator::TranslateMaterial()
 					DerivativeVariations[Variation].TranslatedCodeChunks[PropertyId],
 					Variation,
 					bReduceAfterReturnValue);
+
+				if (PropertyId == MP_BaseColor)
+				{
+					UE_LOG(LogTemp, Display, TEXT("%d\n%s\n%s"), PropertyId, *DerivativeVariations[Variation].TranslatedCodeChunkDefinitions[PropertyId] , *DerivativeVariations[Variation].TranslatedCodeChunks[PropertyId]);
+				}
 			}
 		}
 
@@ -5342,6 +5347,7 @@ int32 FHLSLMaterialTranslator::ViewProperty(EMaterialExposedViewProperty Propert
 		{MEVP_FirstPersonFieldOfView, MCT_Float2, TEXT("View.<PREV>FirstPersonFieldOfViewWideAngles"), nullptr},
 		{MEVP_FirstPersonTanHalfFieldOfView, MCT_Float2, TEXT("View.<PREV>FirstPersonTanAndInvTanHalfFOV.xy"), TEXT("View.<PREV>FirstPersonTanAndInvTanHalfFOV.zw")},
 		{MEVP_FirstPersonScale, MCT_Float, TEXT("View.<PREV>FirstPersonScale"), nullptr},
+		{MEVP_NearPlane, MCT_Float, TEXT("View.NearPlane"), nullptr},
 	};
 	static_assert((sizeof(ViewPropertyMetaArray) / sizeof(ViewPropertyMetaArray[0])) == MEVP_MAX, "incoherency between EMaterialExposedViewProperty and ViewPropertyMetaArray");
 
diff --git a/Engine/Source/Runtime/Engine/Public/Materials/MaterialExpressionViewProperty.h b/Engine/Source/Runtime/Engine/Public/Materials/MaterialExpressionViewProperty.h
--- a/Engine/Source/Runtime/Engine/Public/Materials/MaterialExpressionViewProperty.h
+++ b/Engine/Source/Runtime/Engine/Public/Materials/MaterialExpressionViewProperty.h
@@ -49,6 +49,8 @@ enum EMaterialExposedViewProperty : int
 	MEVP_FirstPersonTanHalfFieldOfView UMETA(DisplayName = "Tan(0.5 * First Person Field Of View)"),
 	/** Scaling factor applied to first person primitives to keep them from intersecting with the scene. */
 	MEVP_FirstPersonScale UMETA(DisplayName = "First Person Scale"),
+	/** Near plane. */
+	MEVP_NearPlane UMETA(DisplayName = "View Near Plane"),
 
 	MEVP_MAX,
 };
