Back out debug code I didn't intend to check in.

#rnx

diff --git a/Engine/Source/Runtime/Engine/Private/Materials/HLSLMaterialTranslator.cpp b/Engine/Source/Runtime/Engine/Private/Materials/HLSLMaterialTranslator.cpp
--- a/Engine/Source/Runtime/Engine/Private/Materials/HLSLMaterialTranslator.cpp
+++ b/Engine/Source/Runtime/Engine/Private/Materials/HLSLMaterialTranslator.cpp
@@ -2184,11 +2184,6 @@ void FHLSLMaterialTranslator::TranslateMaterial()
 					DerivativeVariations[Variation].TranslatedCodeChunks[PropertyId],
 					Variation,
 					bReduceAfterReturnValue);
-
-				if (PropertyId == MP_BaseColor)
-				{
-					UE_LOG(LogTemp, Display, TEXT("%d\n%s\n%s"), PropertyId, *DerivativeVariations[Variation].TranslatedCodeChunkDefinitions[PropertyId] , *DerivativeVariations[Variation].TranslatedCodeChunks[PropertyId]);
-				}
 			}
 		}
 
