Migrated Actors now obey their LifeSpan (and get Destroyed on the Owning Server).

diff --git a/Engine/Source/Runtime/Engine/Private/Actor.cpp b/Engine/Source/Runtime/Engine/Private/Actor.cpp
--- a/Engine/Source/Runtime/Engine/Private/Actor.cpp
+++ b/Engine/Source/Runtime/Engine/Private/Actor.cpp
@@ -960,10 +960,14 @@ void AActor::Serialize(FArchive& Ar)
 		Ar << CreationTime;
 		Ar << bLocalActorInitialized;
 
+		// TODO: This a workaround for not currently supporting native delegates in object migration
+		// we are treating LifeSpan as a property of Actor rather than solving TimerManager's delegates generically.
+		// Calling SetLifeSpan here is also not ideal because it's virtual which means it can do "anything" while in Serialize.
 		float CurrentLifespan = 0.0f; 
 		if (Ar.IsSaving())
 		{
 			CurrentLifespan = GetLifeSpan();
+			SetLifeSpan(0.0f);
 		}
 		Ar << CurrentLifespan;
 		if (Ar.IsLoading())
@@ -6231,7 +6235,12 @@ float AActor::ActorGetDistanceToCollision(const FVector& Point, ECollisionChanne
 
 void AActor::LifeSpanExpired()
 {
+#if UE_WITH_REMOTE_OBJECT_HANDLE
+	// World Timer Manager is not yet transactional, so let's make this call transactional to make up for it
+	UE::RemoteObject::Transfer::ExecuteTransactional([this] { Destroy(); });
+#else
 	Destroy();
+#endif
 }
 
 void AActor::SetLifeSpan( float InLifespan )
