Exposed DuplicateObject to editor scripting

diff --git a/Engine/Source/Runtime/Engine/Classes/Kismet/KismetSystemLibrary.h b/Engine/Source/Runtime/Engine/Classes/Kismet/KismetSystemLibrary.h
--- a/Engine/Source/Runtime/Engine/Classes/Kismet/KismetSystemLibrary.h
+++ b/Engine/Source/Runtime/Engine/Classes/Kismet/KismetSystemLibrary.h
@@ -174,6 +174,20 @@ class UKismetSystemLibrary : public UBlueprintFunctionLibrary
 	UFUNCTION(BlueprintPure, Category = "Utilities")
 	static ENGINE_API UObject* GetOuterObject(const UObject* Object);
 
+#if WITH_EDITOR
+	/**
+	 * Duplicate an object to a new instance
+	 *
+	 * @param SourceObject	The object being duplicated
+	 * @param Outer			The outer to use for the new instance
+	 * @param Name			The optional name of the new instance
+	 *
+	 * @return The duplicated object or null if it failed for some reason
+	 */
+	UFUNCTION(BlueprintCallable, Category = "Utilities")
+	static ENGINE_API UObject* DuplicateObject(const UObject* SourceObject, UObject* Outer, const FName Name = NAME_None);
+#endif	// WITH_EDITOR
+
 	// Engine build number, for displaying to end users.
 	UFUNCTION(BlueprintPure, Category="Development", meta=(BlueprintThreadSafe))
 	static ENGINE_API FString GetEngineVersion();
diff --git a/Engine/Source/Runtime/Engine/Private/KismetSystemLibrary.cpp b/Engine/Source/Runtime/Engine/Private/KismetSystemLibrary.cpp
--- a/Engine/Source/Runtime/Engine/Private/KismetSystemLibrary.cpp
+++ b/Engine/Source/Runtime/Engine/Private/KismetSystemLibrary.cpp
@@ -216,6 +216,13 @@ UObject* UKismetSystemLibrary::GetOuterObject(const UObject* Object)
 	return Object ? Object->GetOuter() : nullptr;
 }
 
+#if WITH_EDITOR
+UObject* UKismetSystemLibrary::DuplicateObject(const UObject* SourceObject, UObject* Outer, const FName Name)
+{
+	return SourceObject ? ::DuplicateObject(SourceObject, Outer, Name) : nullptr;
+}
+#endif	// WITH_EDITOR
+
 FString UKismetSystemLibrary::GetEngineVersion()
 {
 	return FEngineVersion::Current().ToString();
