Engine: Add min/max clamps for CinematicTimeDilation

The defaults are 0.0001 and 20.0, which are the same values for the existing min/max clamps for GlobalTimeDilation.

[/Script/Engine.WorldSettings]
MinCinematicTimeDilation=0.0001
MaxCinematicTimeDilation=20.0

diff --git a/Engine/Config/BaseGame.ini b/Engine/Config/BaseGame.ini
--- a/Engine/Config/BaseGame.ini
+++ b/Engine/Config/BaseGame.ini
@@ -190,6 +190,8 @@ MinUndilatedFrameTime=0.0005		; 2000 fps
 MaxUndilatedFrameTime=0.4			; 2.5 fps
 MinGlobalTimeDilation=0.0001
 MaxGlobalTimeDilation=20.0
+MinCinematicTimeDilation=0.0001
+MaxCinematicTimeDilation=20.0
 
 [/Script/AIModule.AIPerceptionComponent]
 HearingRange=768
diff --git a/Engine/Source/Runtime/Engine/Classes/GameFramework/WorldSettings.h b/Engine/Source/Runtime/Engine/Classes/GameFramework/WorldSettings.h
--- a/Engine/Source/Runtime/Engine/Classes/GameFramework/WorldSettings.h
+++ b/Engine/Source/Runtime/Engine/Classes/GameFramework/WorldSettings.h
@@ -711,6 +711,14 @@ public:
 	UPROPERTY(config, EditAnywhere, Category = Tick, AdvancedDisplay, meta = (UIMin = "0", ClampMin = "0"))
 	float MaxGlobalTimeDilation;
 
+	/** Lowest acceptable cinematic time dilation. */
+	UPROPERTY(config, EditAnywhere, Category = Tick, AdvancedDisplay, meta = (UIMin = "0", ClampMin = "0"))
+	float MinCinematicTimeDilation;
+
+	/** Highest acceptable cinematic time dilation. */
+	UPROPERTY(config, EditAnywhere, Category = Tick, AdvancedDisplay, meta = (UIMin = "0", ClampMin = "0"))
+	float MaxCinematicTimeDilation;
+
 	/** Smallest possible frametime, not considering dilation. Equiv to 1/FastestFPS. */
 	UPROPERTY(config, EditAnywhere, Category = Tick, AdvancedDisplay, meta = (UIMin = "0", ClampMin = "0"))
 	float MinUndilatedFrameTime;
@@ -805,7 +813,10 @@ public:
 	/** Sets the global time dilation value (subject to clamping). Returns the final value that was set. */
 	ENGINE_API virtual float SetTimeDilation(float NewTimeDilation);
 
-	/** @return configuration for NavigationSystem's creation. Null means 
+	/** Sets the cinematic time dilation value (subject to clamping). Returns the final value that was set. */
+	ENGINE_API virtual float SetCinematicTimeDilation(float NewCinematicTimeDilation);
+
+	/** @return configuration for NavigationSystem's creation. Null means
 	 *	no navigation system will be created*/
 	UNavigationSystemConfig* const GetNavigationSystemConfig() const { return NavigationSystemConfigOverride ? NavigationSystemConfigOverride : NavigationSystemConfig; }
 
diff --git a/Engine/Source/Runtime/Engine/Private/WorldSettings.cpp b/Engine/Source/Runtime/Engine/Private/WorldSettings.cpp
--- a/Engine/Source/Runtime/Engine/Private/WorldSettings.cpp
+++ b/Engine/Source/Runtime/Engine/Private/WorldSettings.cpp
@@ -173,6 +173,16 @@ void AWorldSettings::PostInitProperties()
 		MaxGlobalTimeDilation = 0;
 	}
 
+	if (MinCinematicTimeDilation < 0)
+	{
+		MinCinematicTimeDilation = 0;
+	}
+
+	if (MaxCinematicTimeDilation < 0)
+	{
+		MaxCinematicTimeDilation = 0;
+	}
+
 	if (MinUndilatedFrameTime < 0)
 	{
 		MinUndilatedFrameTime = 0;
@@ -325,6 +335,12 @@ float AWorldSettings::SetTimeDilation(float NewTimeDilation)
 	return TimeDilation;
 }
 
+float AWorldSettings::SetCinematicTimeDilation(float NewCinematicTimeDilation)
+{
+	CinematicTimeDilation = FMath::Clamp(NewCinematicTimeDilation, MinCinematicTimeDilation, MaxCinematicTimeDilation);
+	return CinematicTimeDilation;
+}
+
 void AWorldSettings::NotifyBeginPlay()
 {
 	UWorld* World = GetWorld();
diff --git a/Engine/Source/Runtime/MovieSceneTracks/Private/Systems/MovieSceneSlomoSystem.cpp b/Engine/Source/Runtime/MovieSceneTracks/Private/Systems/MovieSceneSlomoSystem.cpp
--- a/Engine/Source/Runtime/MovieSceneTracks/Private/Systems/MovieSceneSlomoSystem.cpp
+++ b/Engine/Source/Runtime/MovieSceneTracks/Private/Systems/MovieSceneSlomoSystem.cpp
@@ -31,7 +31,7 @@ struct FSlomoUtil
 
 		if (WorldSettings)
 		{
-			WorldSettings->CinematicTimeDilation = TimeDilation;
+			WorldSettings->SetCinematicTimeDilation(TimeDilation);
 			WorldSettings->ForceNetUpdate();
 		}
 	}
