Prevent a config reload of UNetDriver from stomping bNoTimeouts

diff --git a/Engine/Source/Runtime/Engine/Classes/Engine/NetDriver.h b/Engine/Source/Runtime/Engine/Classes/Engine/NetDriver.h
--- a/Engine/Source/Runtime/Engine/Classes/Engine/NetDriver.h
+++ b/Engine/Source/Runtime/Engine/Classes/Engine/NetDriver.h
@@ -2347,6 +2347,9 @@ protected:
 	/** Called from RPC processing code to forward RPC to other NetDrivers if ShouldForwardFunction returns true. */
 	ENGINE_API virtual void ForwardRemoteFunction(UObject* RootObject, UObject* SubObject, UFunction* Function, void* Parms);
 
+	/** Determine the correct value for bNoTimeouts based on config/command line/PIE */
+	void EvaluateNoTimeouts();
+
 private:
 	/** Go over imported network guids and map them to the newly created object. */
 	void UpdateUnmappedObjects();
diff --git a/Engine/Source/Runtime/Engine/Private/NetDriver.cpp b/Engine/Source/Runtime/Engine/Private/NetDriver.cpp
--- a/Engine/Source/Runtime/Engine/Private/NetDriver.cpp
+++ b/Engine/Source/Runtime/Engine/Private/NetDriver.cpp
@@ -685,15 +685,8 @@ void UNetDriver::PostInitProperties()
 
 		ProfileStats		= FParse::Param(FCommandLine::Get(),TEXT("profilestats"));
 
-#if !UE_BUILD_SHIPPING
-		bNoTimeouts = bNoTimeouts || FParse::Param(FCommandLine::Get(), TEXT("NoTimeouts")) ? true : false;
-#endif // !UE_BUILD_SHIPPING
+		EvaluateNoTimeouts();
 
-#if WITH_EDITOR
-		// Do not time out in PIE since the server is local.
-		bNoTimeouts = bNoTimeouts || (GEditor && GEditor->PlayWorld);
-#endif // WITH_EDITOR
-	
 		OnLevelRemovedFromWorldHandle = FWorldDelegates::LevelRemovedFromWorld.AddUObject(this, &UNetDriver::OnLevelRemovedFromWorld);
 		OnLevelAddedToWorldHandle = FWorldDelegates::LevelAddedToWorld.AddUObject(this, &UNetDriver::OnLevelAddedToWorld);
 		PostGarbageCollectHandle = FCoreUObjectDelegates::GetPostGarbageCollect().AddUObject(this, &UNetDriver::PostGarbageCollect);
@@ -705,9 +698,27 @@ void UNetDriver::PostInitProperties()
 	NotifyGameInstanceUpdated();
 }
 
+void UNetDriver::EvaluateNoTimeouts()
+{
+#if !UE_BUILD_SHIPPING
+	bNoTimeouts = bNoTimeouts || FParse::Param(FCommandLine::Get(), TEXT("NoTimeouts")) ? true : false;
+#endif // !UE_BUILD_SHIPPING
+
+#if WITH_EDITOR
+	// Do not time out in PIE since the server is local.
+	bNoTimeouts = bNoTimeouts || (GEditor && GEditor->PlayWorld);
+#endif // WITH_EDITOR
+}
+
 void UNetDriver::PostReloadConfig(FProperty* PropertyToLoad)
 {
 	Super::PostReloadConfig(PropertyToLoad);
+
+	if (!HasAnyFlags(RF_ClassDefaultObject))
+	{
+		EvaluateNoTimeouts();
+	}
+
 	LoadChannelDefinitions();
 }
 
