AudioDebug - adding a digit of precision for volume with au.debug.soundcues.minimal

#rnx

diff --git a/Engine/Source/Runtime/Engine/Private/AudioDebug.cpp b/Engine/Source/Runtime/Engine/Private/AudioDebug.cpp
--- a/Engine/Source/Runtime/Engine/Private/AudioDebug.cpp
+++ b/Engine/Source/Runtime/Engine/Private/AudioDebug.cpp
@@ -1586,7 +1586,7 @@ namespace Audio
 
 						const int32 SoundNameIndex = ActiveSoundCount++;
 						const FString LeadingNumber = FString::Printf(TEXT("%4i. "), SoundNameIndex);
-						const FString Volume = FString::Printf(TEXT("%6.2f "), WaveInstanceInfo.Volume);
+						const FString Volume = FString::Printf(TEXT("%6.3f "), WaveInstanceInfo.Volume);
 						const FString PathAndMuting = FString::Printf(TEXT("Path: %s %s"), *SoundPath, *MuteSoloReason);
 
 						int32 CurrentX = X;
