fix static mesh component guard against accessing render data when no valid data exists

diff --git a/Engine/Source/Runtime/Engine/Private/Components/StaticMeshComponent.cpp b/Engine/Source/Runtime/Engine/Private/Components/StaticMeshComponent.cpp
--- a/Engine/Source/Runtime/Engine/Private/Components/StaticMeshComponent.cpp
+++ b/Engine/Source/Runtime/Engine/Private/Components/StaticMeshComponent.cpp
@@ -2021,13 +2021,17 @@ void UStaticMeshComponent::PostEditUndo()
 	// Debug check command trying to track down undo related uninitialized resource
 	if (GetStaticMesh() != NULL && GetStaticMesh()->GetRenderData() && GetStaticMesh()->GetRenderData()->LODResources.Num() > 0)
 	{
-		FRenderResource* Resource = &GetStaticMesh()->GetRenderData()->LODResources[GetStaticMesh()->GetRenderData()->GetCurrentFirstLODIdx(0)].IndexBuffer;
-		ENQUEUE_RENDER_COMMAND(ResourceCheckCommand)(
-			[Resource](FRHICommandList& RHICmdList)
-			{
-				check( Resource->IsInitialized() );
-			}
-		);
+		const int32 LODResourceIdx = GetStaticMesh()->GetRenderData()->GetCurrentFirstLODIdx(0);
+		if (LODResourceIdx != INDEX_NONE)
+		{
+			FRenderResource* Resource = &GetStaticMesh()->GetRenderData()->LODResources[LODResourceIdx].IndexBuffer;
+			ENQUEUE_RENDER_COMMAND(ResourceCheckCommand)(
+				[Resource](FRHICommandList& RHICmdList)
+				{
+					check(Resource->IsInitialized());
+				}
+				);
+		}
 	}
 	Super::PostEditUndo();
 }
