Iris
* Fix FActorReplicationParams struct being ignored by BeginReplication
* FActorReplicationParams now has a EFilterType enum to better control the intented dynamic filter to be used by the actor

diff --git a/Engine/Source/Runtime/Engine/Private/Net/Iris/ReplicationSystem/EngineReplicationBridge.cpp b/Engine/Source/Runtime/Engine/Private/Net/Iris/ReplicationSystem/EngineReplicationBridge.cpp
--- a/Engine/Source/Runtime/Engine/Private/Net/Iris/ReplicationSystem/EngineReplicationBridge.cpp
+++ b/Engine/Source/Runtime/Engine/Private/Net/Iris/ReplicationSystem/EngineReplicationBridge.cpp
@@ -191,7 +191,7 @@ void UEngineReplicationBridge::Deinitialize()
 	ObjectReferencePackageMap = nullptr;
 }
 
-UE::Net::FNetRefHandle UEngineReplicationBridge::StartReplicatingActor(AActor* Actor, const FActorReplicationParams& Params)
+UE::Net::FNetRefHandle UEngineReplicationBridge::StartReplicatingActor(AActor* Actor, const FActorReplicationParams& ActorParams )
 {
 	using namespace UE::Net;
 
@@ -283,9 +283,33 @@ UE::Net::FNetRefHandle UEngineReplicationBridge::StartReplicatingActor(AActor* A
 		.bNeedsWorldLocationUpdate = true,
 		.bIsDormant = Actor->NetDormancy > DORM_Awake,
 		.StaticPriority = (Actor->bAlwaysRelevant || Actor->bOnlyRelevantToOwner) ? Actor->NetPriority : 0.0f,
-		.PollFrequency = Actor->GetNetUpdateFrequency()
+		.PollFrequency = Actor->GetNetUpdateFrequency(),
 	};
 
+	// Assign a specific filter if requested
+	if (ActorParams.FilterType != FActorReplicationParams::ConfigFilter)
+	{
+		RootObjectParams.bUseExplicitDynamicFilter = true;
+
+		switch(ActorParams.FilterType)
+		{
+			case FActorReplicationParams::AlwaysRelevant:				
+				RootObjectParams.ExplicitDynamicFilterName = NAME_None;
+				break;
+
+			case FActorReplicationParams::DefaultSpatial:
+				RootObjectParams.ExplicitDynamicFilterName = GetDefaultSpatialFilterName();
+				break;
+
+			case FActorReplicationParams::ExplicitFilter:
+				ensureMsgf(!ActorParams.ExplicitDynamicFilterName.IsNone(), TEXT("FActorReplicationParams::ExplicitDynamicFilterName must be defined when FilterType is ExplicitFilter"));
+				RootObjectParams.ExplicitDynamicFilterName = ActorParams.ExplicitDynamicFilterName;
+				break;
+		}
+	}
+	
+
+
 #if !UE_BUILD_SHIPPING
 	ensureMsgf(!(Actor->bAlwaysRelevant || Actor->bOnlyRelevantToOwner) || RootObjectParams.StaticPriority >= 1.0f, TEXT("Very low NetPriority %.02f for always relevant or owner relevant Actor %s. Set it to 1.0f or higher."), Actor->NetPriority, ToCStr(Actor->GetName()));
 #endif
diff --git a/Engine/Source/Runtime/Engine/Public/Net/Iris/ReplicationSystem/EngineReplicationBridge.h b/Engine/Source/Runtime/Engine/Public/Net/Iris/ReplicationSystem/EngineReplicationBridge.h
--- a/Engine/Source/Runtime/Engine/Public/Net/Iris/ReplicationSystem/EngineReplicationBridge.h
+++ b/Engine/Source/Runtime/Engine/Public/Net/Iris/ReplicationSystem/EngineReplicationBridge.h
@@ -31,10 +31,26 @@ ENGINE_API bool ShouldUseIrisReplication(const UObject* Actor);
 
 struct FActorReplicationParams
 {
-	/** When true we ignore the configured dynamic filter for this actor type and use the explicit filter instead */
-	bool bOverrideDynamicFilterConfig = false;
+	enum EFilterType
+	{
+		/** Let the config filter configs assign a filter based on the class type. */
+		ConfigFilter = 0,
+		/** When set don't assign any dynamic filter and default to being always relevant. */
+		AlwaysRelevant,
+		/** When set use the default spatial filter of the bridge. Generally that is the NetObjectGridFilter. */
+		DefaultSpatial,
+		/** When set use filter defined by ExplicitDynamicFilterName. */
+		ExplicitFilter,
+	};
 
-	/** Only used when bOverrideDynamicFilterConfig is true. The name of the dynamic filter to assign to this actor. Can be NAME_None if no filters should be set. */
+	/**
+	 * The default behavior for actors (e.g. ConfigFilter) is that they are automatically assigned a filter based on the class type via the engine config and UObjectReplicationBridgeConfig::FilterConfigs.
+	 * Choosing a different option allows you to ignore the automatic assignment and select a specific filter for the replicated actor.
+	 * @see FObjectReplicationBridgeFilterConfig
+	 */
+	EFilterType FilterType = ConfigFilter;
+
+	/** Only used when ExplicitFilter is the type used. The dynamic filter to assign to this actor. */
 	FName ExplicitDynamicFilterName;
 };
 
