Fix for destroyed actors not being removed from the multi server proxy.

diff --git a/Engine/Source/Runtime/Engine/Private/LevelActor.cpp b/Engine/Source/Runtime/Engine/Private/LevelActor.cpp
--- a/Engine/Source/Runtime/Engine/Private/LevelActor.cpp
+++ b/Engine/Source/Runtime/Engine/Private/LevelActor.cpp
@@ -889,7 +889,11 @@ bool UWorld::DestroyActor( AActor* ThisActor, bool bNetForce, bool bShouldModify
 			return false;
 		}
 
-		if (ThisActor->IsActorBeginningPlay())
+		if (ThisActor->IsActorBeginningPlay()
+#if UE_SUPPORT_FOR_ACTOR_TICK_DISABLE
+			&& IsActorTickAndUserCallbacksEnabled()
+#endif
+		)
 		{
 			FSetActorWantsDestroyDuringBeginPlay SetActorWantsDestroyDuringBeginPlay(ThisActor);
 			return true; // while we didn't actually destroy it now, we are going to, so tell the calling code it succeeded
