Fix for double anim update in Anim Next

diff --git a/Engine/Source/Runtime/Engine/Private/Components/SkeletalMeshComponent.cpp b/Engine/Source/Runtime/Engine/Private/Components/SkeletalMeshComponent.cpp
--- a/Engine/Source/Runtime/Engine/Private/Components/SkeletalMeshComponent.cpp
+++ b/Engine/Source/Runtime/Engine/Private/Components/SkeletalMeshComponent.cpp
@@ -3274,6 +3274,14 @@ void USkeletalMeshComponent::SetAnimInstanceClass(class UClass* NewClass)
 
 void USkeletalMeshComponent::SetEnableAnimation(bool bNewEnableAnimation)
 {
+	// If we're not registered yet, we can set the value straight away.
+	if (IsRegistered() == false)
+	{
+		bEnableAnimation = bNewEnableAnimation;
+		PrimaryComponentTick.bRunOnAnyThread = !bEnableAnimation;
+		return;
+	}
+
 	// Defer setting flag to the end of the frame as other tasks can be in-flight during a tick
 	FTSTicker::GetCoreTicker().AddTicker(ANSI_TO_TCHAR(__FUNCTION__), 0.0f,
 		[bNewEnableAnimation, WeakComponent = TWeakObjectPtr<USkeletalMeshComponent>(this)](float InDeltaTime)
@@ -3282,6 +3290,11 @@ void USkeletalMeshComponent::SetEnableAnimation(bool bNewEnableAnimation)
 			{
 				Component->bEnableAnimation = bNewEnableAnimation;
 				Component->PrimaryComponentTick.bRunOnAnyThread = !bNewEnableAnimation;
+
+				if (bNewEnableAnimation && Component->IsRegistered())
+				{
+					Component->InitAnim(false);
+				}
 			}
 			return false;
 		});
