Do not perform skinning for FSkeletalMeshObjectGPUSkin if proxy is not visible in raytracing scene

diff --git a/Engine/Source/Runtime/Engine/Private/SkeletalRenderGPUSkin.cpp b/Engine/Source/Runtime/Engine/Private/SkeletalRenderGPUSkin.cpp
--- a/Engine/Source/Runtime/Engine/Private/SkeletalRenderGPUSkin.cpp
+++ b/Engine/Source/Runtime/Engine/Private/SkeletalRenderGPUSkin.cpp
@@ -594,13 +594,6 @@ void FSkeletalMeshObjectGPUSkin::UpdateDynamicData_RenderThread(FGPUSkinCache* G
 		!DynamicData->ActiveMorphTargetsEqual(InDynamicData->ActiveMorphTargets, InDynamicData->MorphTargetWeights))
 		: true);
 
-#if RHI_RAYTRACING
-	if (bMorphNeedsUpdate && Scene)
-	{
-		Scene->InvalidatePathTracedOutput();
-	}
-#endif
-
 	if (DynamicData)
 	{
 		FDynamicSkelMeshObjectDataGPUSkin::FreeDynamicSkelMeshObjectDataGPUSkin(DynamicData);
@@ -609,10 +602,20 @@ void FSkeletalMeshObjectGPUSkin::UpdateDynamicData_RenderThread(FGPUSkinCache* G
 	DynamicData = InDynamicData;
 	LastBoneTransformRevisionNumber = RevisionNumber;
 
+#if RHI_RAYTRACING
+	// If set to -1 means that update is not needed for raytracing 
+	const bool bVisibleInRayTracing = (DynamicData->RayTracingLODIndex != -1);
+
+	if (bVisibleInRayTracing && bMorphNeedsUpdate && Scene)
+	{
+		Scene->InvalidatePathTracedOutput();
+	}
+#endif
+
 	ProcessUpdatedDynamicData(EGPUSkinCacheEntryMode::Raster, GPUSkinCache, RHICmdList, FrameNumberToPrepare, RevisionNumber, PreviousRevisionNumber, bMorphNeedsUpdate, DynamicData->LODIndex, bRecreating);
 
 #if RHI_RAYTRACING
-	if (ShouldUseSeparateSkinCacheEntryForRayTracing() && FGPUSkinCache::IsGPUSkinCacheRayTracingSupported() && GPUSkinCache && SkeletalMeshRenderData->bSupportRayTracing)
+	if (bVisibleInRayTracing && ShouldUseSeparateSkinCacheEntryForRayTracing() && FGPUSkinCache::IsGPUSkinCacheRayTracingSupported() && GPUSkinCache && SkeletalMeshRenderData->bSupportRayTracing)
 	{
 		// Morph delta is updated in raster pass above, no need to update again for ray tracing
 		ProcessUpdatedDynamicData(EGPUSkinCacheEntryMode::RayTracing, GPUSkinCache, RHICmdList, FrameNumberToPrepare, RevisionNumber, PreviousRevisionNumber, /*bMorphNeedsUpdate=*/false, DynamicData->RayTracingLODIndex, bRecreating);
@@ -625,7 +628,7 @@ void FSkeletalMeshObjectGPUSkin::UpdateDynamicData_RenderThread(FGPUSkinCache* G
 #endif
 
 #if RHI_RAYTRACING
-	if (DynamicData->GPUSkinTechnique != ESkeletalMeshGPUSkinTechnique::MeshDeformer && FGPUSkinCache::IsGPUSkinCacheRayTracingSupported() && SkeletalMeshRenderData->bSupportRayTracing && !GetSkinCacheEntryForRayTracing())
+	if (bVisibleInRayTracing && DynamicData->GPUSkinTechnique != ESkeletalMeshGPUSkinTechnique::MeshDeformer && FGPUSkinCache::IsGPUSkinCacheRayTracingSupported() && SkeletalMeshRenderData->bSupportRayTracing && !GetSkinCacheEntryForRayTracing())
 	{
 		// When SkinCacheEntry is gone, clear geometry
 		RayTracingGeometry.ReleaseRHI();
@@ -2472,7 +2475,9 @@ void FDynamicSkelMeshObjectDataGPUSkin::InitDynamicSkelMeshObjectDataGPUSkin(
 	}
 
 #if RHI_RAYTRACING
-	RayTracingLODIndex = FMath::Clamp(FMath::Max(LODIndex + GetRayTracingSkeletalMeshGlobalLODBias(), InMeshObject->RayTracingMinLOD), LODIndex, InSkeletalMeshRenderData->LODRenderData.Num() - 1);
+	// If Proxy is not visible in raytracing scene, set RayTracingLODIndex to -1 which means that additional RT update will not be performed for that mesh object, probably RT proxy handles that 
+	const bool bVisibleInRayTracing = SkeletalMeshProxy ? SkeletalMeshProxy->IsVisibleInRayTracing() : true;
+	RayTracingLODIndex = bVisibleInRayTracing ? FMath::Clamp(FMath::Max(LODIndex + GetRayTracingSkeletalMeshGlobalLODBias(), InMeshObject->RayTracingMinLOD), LODIndex, InSkeletalMeshRenderData->LODRenderData.Num() - 1) : -1;
 #endif
 
 #if !(UE_BUILD_SHIPPING || UE_BUILD_TEST)
@@ -2489,7 +2494,7 @@ void FDynamicSkelMeshObjectDataGPUSkin::InitDynamicSkelMeshObjectDataGPUSkin(
 	// update ReferenceToLocal
 	UpdateRefToLocalMatrices( ReferenceToLocal, InDynamicData, InSkinnedAsset, InSkeletalMeshRenderData, LODIndex, ExtraRequiredBoneIndices, LeaderBoneMappedMeshComponentSpaceTransforms );
 #if RHI_RAYTRACING
-	if (RayTracingLODIndex != LODIndex)
+	if (bVisibleInRayTracing && RayTracingLODIndex != LODIndex)
 	{
 		UpdateRefToLocalMatrices(ReferenceToLocalForRayTracing, InDynamicData, InSkinnedAsset, InSkeletalMeshRenderData, RayTracingLODIndex, ExtraRequiredBoneIndices);
 	}
@@ -2504,7 +2509,7 @@ void FDynamicSkelMeshObjectDataGPUSkin::InitDynamicSkelMeshObjectDataGPUSkin(
 	case EPreviousBoneTransformUpdateMode::UpdatePrevious:
 		UpdatePreviousRefToLocalMatrices(PreviousReferenceToLocal, InDynamicData, InSkinnedAsset, InSkeletalMeshRenderData, LODIndex, ExtraRequiredBoneIndices);
 	#if RHI_RAYTRACING
-		if (RayTracingLODIndex != LODIndex)
+		if (bVisibleInRayTracing && RayTracingLODIndex != LODIndex)
 		{
 			UpdatePreviousRefToLocalMatrices(PreviousReferenceToLocalForRayTracing, InDynamicData, InSkinnedAsset, InSkeletalMeshRenderData, RayTracingLODIndex, ExtraRequiredBoneIndices);
 		}
@@ -2513,7 +2518,7 @@ void FDynamicSkelMeshObjectDataGPUSkin::InitDynamicSkelMeshObjectDataGPUSkin(
 	case EPreviousBoneTransformUpdateMode::DuplicateCurrentToPrevious:
 		PreviousReferenceToLocal = ReferenceToLocal;
 	#if RHI_RAYTRACING
-		if (RayTracingLODIndex != LODIndex)
+		if (bVisibleInRayTracing && RayTracingLODIndex != LODIndex)
 		{
 			PreviousReferenceToLocalForRayTracing = ReferenceToLocalForRayTracing;
 		}
@@ -2591,7 +2596,7 @@ void FDynamicSkelMeshObjectDataGPUSkin::InitDynamicSkelMeshObjectDataGPUSkin(
 	}
 
 #if RHI_RAYTRACING
-	if (SkeletalMeshProxy && !SkeletalMeshProxy->IsNaniteMesh())
+	if (bVisibleInRayTracing && SkeletalMeshProxy && !SkeletalMeshProxy->IsNaniteMesh())
 	{
 		// TODO: Nanite-Skinning
 		bAnySegmentUsesWorldPositionOffset = SkeletalMeshProxy->bAnySegmentUsesWorldPositionOffset;
