Add null check to AActor::GetGameInstance.

#rnx

diff --git a/Engine/Source/Runtime/Engine/Private/Actor.cpp b/Engine/Source/Runtime/Engine/Private/Actor.cpp
--- a/Engine/Source/Runtime/Engine/Private/Actor.cpp
+++ b/Engine/Source/Runtime/Engine/Private/Actor.cpp
@@ -627,7 +627,12 @@ FTimerManager& AActor::GetWorldTimerManager() const
 
 UGameInstance* AActor::GetGameInstance() const
 {
-	return GetWorld()->GetGameInstance();
+	if (const UWorld* World = GetWorld())
+	{
+		return World->GetGameInstance();
+	}
+
+	return nullptr;
 }
 
 bool AActor::IsNetStartupActor() const
