Fixed skeletal mesh LOD stream out flicker due to comparison against a game-thread LOD index in the render-thread

After a stream out GetMeshElementsConditionallySelectable may be called with a new higher LOD index before SendRenderDynamicData_Concurrent has a chance to update the GPUSkin mesh object to the new LOD index. Mesh draw command population is rejected for one frame since the game-thread LOD (PendingFirstLODIdx) is higher than the render-thread LOD (CurrentFirstLODIdx), which is checked against the mesh object (lives on the render thread).

To avoid this in GetMeshElementsConditionallySelectable, just compare against CurrentFirstLODIdx, which is safely updated on the render-thread after stream out, and synchronized with the mesh object updates.

diff --git a/Engine/Source/Runtime/Engine/Private/SkeletalMeshSceneProxy.cpp b/Engine/Source/Runtime/Engine/Private/SkeletalMeshSceneProxy.cpp
--- a/Engine/Source/Runtime/Engine/Private/SkeletalMeshSceneProxy.cpp
+++ b/Engine/Source/Runtime/Engine/Private/SkeletalMeshSceneProxy.cpp
@@ -549,7 +549,7 @@ void FSkeletalMeshSceneProxy::GetMeshElementsConditionallySelectable(const TArra
 
 	const FEngineShowFlags& EngineShowFlags = ViewFamily.EngineShowFlags;
 
-	const int32 FirstLODIdx = SkeletalMeshRenderData->GetFirstValidLODIdx(FMath::Max(SkeletalMeshRenderData->PendingFirstLODIdx, SkeletalMeshRenderData->CurrentFirstLODIdx));
+	const int32 FirstLODIdx = SkeletalMeshRenderData->GetFirstValidLODIdx(SkeletalMeshRenderData->CurrentFirstLODIdx);
 	if (FirstLODIdx == INDEX_NONE)
 	{
 #if DO_CHECK
