Reduced lock contention during skeletal mesh gpu skin object allocation / free's

diff --git a/Engine/Source/Runtime/Engine/Private/SkeletalRenderGPUSkin.cpp b/Engine/Source/Runtime/Engine/Private/SkeletalRenderGPUSkin.cpp
--- a/Engine/Source/Runtime/Engine/Private/SkeletalRenderGPUSkin.cpp
+++ b/Engine/Source/Runtime/Engine/Private/SkeletalRenderGPUSkin.cpp
@@ -2360,12 +2360,10 @@ void FDynamicSkelMeshObjectDataGPUSkin::Clear()
 
 #define SKELETON_POOL_GPUSKINS 1
 #if SKELETON_POOL_GPUSKINS
-TArray<FDynamicSkelMeshObjectDataGPUSkin*> FreeGpuSkins;
-FCriticalSection FreeGpuSkinsCriticalSection;
+std::atomic<int64> FreeGpuSkinsCounter{0};
+TLockFreePointerListFIFO<FDynamicSkelMeshObjectDataGPUSkin, 0> FreeGpuSkinsFIFO;
 
 static int32 GPoolGpuSkins = 1;
-static int32 GMinPoolCount = 0;
-static int32 GAllocationCounter = 0;
 static const int32 GAllocationsBeforeCleanup = 1000; // number of allocations we make before we clean up the pool, this number is increased when we have to allocate not from the pool
 static FAutoConsoleVariableRef CVarPoolGpuSkins(
 	TEXT("r.GpuSkin.Pool"),
@@ -2373,6 +2371,19 @@ static FAutoConsoleVariableRef CVarPoolGpuSkins(
 	TEXT("Should we pool gpu skins.\n")
 	TEXT(" 0: Don't pool anything\n")
 	TEXT(" 1: Pool gpu skins bro (default)\n"),
+	FConsoleVariableDelegate::CreateLambda([](IConsoleVariable*)
+	{
+		if (!GPoolGpuSkins)
+		{
+			while (FDynamicSkelMeshObjectDataGPUSkin* Skin = FreeGpuSkinsFIFO.Pop())
+			{
+				// Pooling is disabled, just destroys
+				FDynamicSkelMeshObjectDataGPUSkin::FreeDynamicSkelMeshObjectDataGPUSkin(Skin);
+			}
+
+			FreeGpuSkinsCounter.store(0);
+		}
+	}),
 	ECVF_Default
 );
 #endif
@@ -2381,64 +2392,45 @@ static FAutoConsoleVariableRef CVarPoolGpuSkins(
 FDynamicSkelMeshObjectDataGPUSkin* FDynamicSkelMeshObjectDataGPUSkin::AllocDynamicSkelMeshObjectDataGPUSkin()
 {
 #if SKELETON_POOL_GPUSKINS
-	if (!GPoolGpuSkins)
+	if (GPoolGpuSkins)
 	{
-		return new FDynamicSkelMeshObjectDataGPUSkin;
-	}
-
-	FScopeLock S(&FreeGpuSkinsCriticalSection);
-	++GAllocationCounter;
-	GMinPoolCount = FMath::Min(FreeGpuSkins.Num(), GMinPoolCount);
-	if ( FreeGpuSkins.Num() > 0 )
-	{
-		FDynamicSkelMeshObjectDataGPUSkin *Result = FreeGpuSkins[0];
-		FreeGpuSkins.RemoveAtSwap(0);
-		return Result;
-	}
-	else
-	{
-		return new FDynamicSkelMeshObjectDataGPUSkin;
+		if (FDynamicSkelMeshObjectDataGPUSkin* Skin = FreeGpuSkinsFIFO.Pop())
+		{
+			// Note that the counter isn't fully sync'ed with the free skins, but at most
+			// we risk undershooting it by a few entries.
+			--FreeGpuSkinsCounter;
+			
+			return Skin;
+		}
 	}
-#else
+#endif // SKELETON_POOL_GPUSKINS
+	
 	return new FDynamicSkelMeshObjectDataGPUSkin;
-#endif
 }
 
 void FDynamicSkelMeshObjectDataGPUSkin::FreeDynamicSkelMeshObjectDataGPUSkin(FDynamicSkelMeshObjectDataGPUSkin* Who)
 {
 #if SKELETON_POOL_GPUSKINS
-	if (!GPoolGpuSkins)
+	if (GPoolGpuSkins)
 	{
-		delete Who;
-
-		if ( FreeGpuSkins.Num() > 0 )
+		// If exceeded, destroy the oldest and trade with this one
+		if (++FreeGpuSkinsCounter > GAllocationsBeforeCleanup)
 		{
-			FScopeLock S(&FreeGpuSkinsCriticalSection);
-			for ( FDynamicSkelMeshObjectDataGPUSkin* I : FreeGpuSkins )
+			if (FDynamicSkelMeshObjectDataGPUSkin* Skin = FreeGpuSkinsFIFO.Pop())
 			{
-				delete I;
+				--FreeGpuSkinsCounter;
+				delete Skin;
 			}
-			FreeGpuSkins.Empty();
 		}
-		return;
-	}
 
-	Who->Clear();
-	FScopeLock S(&FreeGpuSkinsCriticalSection);
-	FreeGpuSkins.Add(Who);
-	if ( GAllocationCounter > GAllocationsBeforeCleanup )
-	{
-		GAllocationCounter = 0;
-		for ( int32 I = 0; I < GMinPoolCount; ++I )
-		{
-			delete FreeGpuSkins[0];
-			FreeGpuSkins.RemoveAtSwap(0);
-		}
-		GMinPoolCount = FreeGpuSkins.Num();
+		Who->Clear();
+	
+		FreeGpuSkinsFIFO.Push(Who);
+		return;
 	}
-#else
+#endif // SKELETON_POOL_GPUSKINS
+	
 	delete Who;
-#endif
 }
 
 void FDynamicSkelMeshObjectDataGPUSkin::InitDynamicSkelMeshObjectDataGPUSkin(
