Fix Content Worker crash in FixupHLODSetup when called from PostLoad.

diff --git a/Engine/Source/Runtime/Engine/Private/WorldPartition/RuntimeHashSet/WorldPartitionRuntimeHashSet.cpp b/Engine/Source/Runtime/Engine/Private/WorldPartition/RuntimeHashSet/WorldPartitionRuntimeHashSet.cpp
--- a/Engine/Source/Runtime/Engine/Private/WorldPartition/RuntimeHashSet/WorldPartitionRuntimeHashSet.cpp
+++ b/Engine/Source/Runtime/Engine/Private/WorldPartition/RuntimeHashSet/WorldPartitionRuntimeHashSet.cpp
@@ -955,10 +955,7 @@ void UWorldPartitionRuntimeHashSet::PostEditChangeChainProperty(FPropertyChanged
 		}
 		
 		// Recreate the HLOD partitions as they must use the same class
-		if (RuntimePartitionDesc.MainLayer)
-		{
-			FixupHLODSetup(RuntimePartitionDesc);
-		}
+		FixupHLODSetup(RuntimePartitionDesc);
 	}
 	else if (PropertyName == GET_MEMBER_NAME_CHECKED(FRuntimePartitionDesc, Name))
 	{
@@ -1133,24 +1130,27 @@ void UWorldPartitionRuntimeHashSet::PostLoad()
 
 void UWorldPartitionRuntimeHashSet::FixupHLODSetup(FRuntimePartitionDesc& RuntimePartition)
 {
-	for (int32 HLODSetupIndex = 0; HLODSetupIndex < RuntimePartition.HLODSetups.Num(); HLODSetupIndex++)
+	if (RuntimePartition.MainLayer)
 	{
-		FRuntimePartitionHLODSetup& HLODSetup = RuntimePartition.HLODSetups[HLODSetupIndex];
-		if (HLODSetup.PartitionLayer)
+		for (int32 HLODSetupIndex = 0; HLODSetupIndex < RuntimePartition.HLODSetups.Num(); HLODSetupIndex++)
 		{
-			// If needed, recreate the HLOD partition to use the same class as the main partition
-			if (HLODSetup.PartitionLayer->GetClass() != RuntimePartition.MainLayer->GetClass())
+			FRuntimePartitionHLODSetup& HLODSetup = RuntimePartition.HLODSetups[HLODSetupIndex];
+			if (HLODSetup.PartitionLayer)
 			{
-				if (HLODSetup.bIsSpatiallyLoaded)
+				// If needed, recreate the HLOD partition to use the same class as the main partition
+				if (HLODSetup.PartitionLayer->GetClass() != RuntimePartition.MainLayer->GetClass())
 				{
-					URuntimePartition* ParentRuntimePartition = HLODSetupIndex != 0 ? RuntimePartition.HLODSetups[HLODSetupIndex - 1].PartitionLayer : RuntimePartition.MainLayer;
-					HLODSetup.PartitionLayer = ParentRuntimePartition->CreateHLODRuntimePartition(HLODSetupIndex);
-					HLODSetup.PartitionLayer->Name = HLODSetup.Name;
+					if (HLODSetup.bIsSpatiallyLoaded)
+					{
+						URuntimePartition* ParentRuntimePartition = HLODSetupIndex != 0 ? RuntimePartition.HLODSetups[HLODSetupIndex - 1].PartitionLayer : RuntimePartition.MainLayer;
+						HLODSetup.PartitionLayer = ParentRuntimePartition->CreateHLODRuntimePartition(HLODSetupIndex);
+						HLODSetup.PartitionLayer->Name = HLODSetup.Name;
+					}
 				}
-			}
 
-			// Make sure that HLOD partitions settings are updated to match the main layer settings
-			HLODSetup.PartitionLayer->UpdateHLODRuntimePartitionFrom(RuntimePartition.MainLayer);
+				// Make sure that HLOD partitions settings are updated to match the main layer settings
+				HLODSetup.PartitionLayer->UpdateHLODRuntimePartitionFrom(RuntimePartition.MainLayer);
+			}
 		}
 	}
 }
