Addresses crash when we have a leader pose component set but an empty bone map present.

#rnx

diff --git a/Engine/Source/Runtime/Engine/Private/SkeletalRender.cpp b/Engine/Source/Runtime/Engine/Private/SkeletalRender.cpp
--- a/Engine/Source/Runtime/Engine/Private/SkeletalRender.cpp
+++ b/Engine/Source/Runtime/Engine/Private/SkeletalRender.cpp
@@ -526,7 +526,7 @@ void UpdatePreviousRefToLocalMatrices(TArray<FMatrix44f>& ReferenceToLocal, cons
 
 	const FReferenceSkeleton& RefSkeleton = InSkinnedAsset->GetRefSkeleton();
 	const TArrayView<const int32> LeaderBoneMap = InDynamicData.GetLeaderBoneMap();
-	const bool bIsLeaderCompValid = InDynamicData.HasLeaderPoseComponent();
+	const bool bIsLeaderCompValid = InDynamicData.HasLeaderPoseComponent() && LeaderBoneMap.Num() == InSkinnedAsset->GetRefSkeleton().GetNum();
 	const TArrayView<const FTransform> ComponentTransform = InDynamicData.GetPreviousComponentTransformsArray();
 	const TArrayView<const uint8> BoneVisibilityStates = InDynamicData.GetPreviousBoneVisibilityStates();
 	// Get inv ref pose matrices
diff --git a/Engine/Source/Runtime/Engine/Private/SkinnedMeshSceneProxyDesc.cpp b/Engine/Source/Runtime/Engine/Private/SkinnedMeshSceneProxyDesc.cpp
--- a/Engine/Source/Runtime/Engine/Private/SkinnedMeshSceneProxyDesc.cpp
+++ b/Engine/Source/Runtime/Engine/Private/SkinnedMeshSceneProxyDesc.cpp
@@ -426,7 +426,7 @@ FSkinnedMeshSceneProxyDynamicData::FSkinnedMeshSceneProxyDynamicData(const USkin
 	, bHasMeshDeformerInstance(InSkinnedMeshComponent->GetMeshDeformerInstance() != nullptr)
 	, bRenderStateRecreating(InSkinnedMeshComponent->IsRenderStateRecreating())
 {
-	checkf(InLeaderPoseComponent == nullptr || InSkinnedMeshComponent->GetLeaderBoneMap().Num() == InSkinnedMeshComponent->GetSkinnedAsset()->GetRefSkeleton().GetNum(),
+	ensureMsgf(InLeaderPoseComponent == nullptr || InSkinnedMeshComponent->GetLeaderBoneMap().Num() == InSkinnedMeshComponent->GetSkinnedAsset()->GetRefSkeleton().GetNum(),
 		TEXT("Leader pose component skeleton doesn't match follower. LeaderPoseComponent: %s SkinnedMeshComponent: %s"),
 		InLeaderPoseComponent ? *InLeaderPoseComponent->GetFullName() : *FString("NULL"),
 		InSkinnedMeshComponent ? *InSkinnedMeshComponent->GetFullName() : *FString("NULL"));
