Chaos:: Add scaling support into skeletal mesh for physics. Physics now support SetActorScale3D blueprint node with uniform scales.

diff --git a/Engine/Source/Runtime/Engine/Classes/Components/SkeletalMeshComponent.h b/Engine/Source/Runtime/Engine/Classes/Components/SkeletalMeshComponent.h
--- a/Engine/Source/Runtime/Engine/Classes/Components/SkeletalMeshComponent.h
+++ b/Engine/Source/Runtime/Engine/Classes/Components/SkeletalMeshComponent.h
@@ -1565,6 +1565,9 @@ private:
 	/** previous root bone matrix to compare the difference and decide to do clothing teleport  */
 	FMatrix	PrevRootBoneMatrix;
 
+	/** previous root body scale to compare the difference when the actor scale is updated. */
+	FVector3f PrevRootRigidBodyScale;
+
 	/** 
 	 * Clothing simulation objects.
 	 * ClothingSimulation is responsible for maintaining and siulating clothing actors
@@ -2367,6 +2370,9 @@ public:
 	 */
 	ENGINE_API void UpdateClothTransform(ETeleportType TeleportType);
 
+	/** Update rigid body scaling when the actor is set with a new uniform scale. */
+	void UpdateRigidBodyScaling(ETeleportType TeleportType);
+
 	/** Set the cloth transform update to trigger with no teleport option. */
 	void UpdateClothTransform() { UpdateClothTransform(ETeleportType::None); }
 
diff --git a/Engine/Source/Runtime/Engine/Private/Components/SkeletalMeshComponent.cpp b/Engine/Source/Runtime/Engine/Private/Components/SkeletalMeshComponent.cpp
--- a/Engine/Source/Runtime/Engine/Private/Components/SkeletalMeshComponent.cpp
+++ b/Engine/Source/Runtime/Engine/Private/Components/SkeletalMeshComponent.cpp
@@ -325,6 +325,7 @@ USkeletalMeshComponent::USkeletalMeshComponent(const FObjectInitializer& ObjectI
 
 	ClothTeleportMode = EClothingTeleportMode::None;
 	PrevRootBoneMatrix = GetBoneMatrix(0); // save the root bone transform
+	PrevRootRigidBodyScale = FVector3f(GetRelativeScale3D());
 
 	// pre-compute cloth teleport thresholds for performance
 	ComputeTeleportRotationThresholdInRadians();
diff --git a/Engine/Source/Runtime/Engine/Private/SkeletalMeshComponentPhysics.cpp b/Engine/Source/Runtime/Engine/Private/SkeletalMeshComponentPhysics.cpp
--- a/Engine/Source/Runtime/Engine/Private/SkeletalMeshComponentPhysics.cpp
+++ b/Engine/Source/Runtime/Engine/Private/SkeletalMeshComponentPhysics.cpp
@@ -64,6 +64,8 @@ static TAutoConsoleVariable<float> CVarClothTeleportRotationThreshold(TEXT("p.Cl
 
 static TAutoConsoleVariable<int32> CVarEnableKinematicDeferralPrePhysicsCondition(TEXT("p.EnableKinematicDeferralPrePhysicsCondition"), 1, TEXT("If is 1, and deferral would've been disallowed due to EUpdateTransformFlags, allow if in PrePhysics tick. If 0, condition is unchanged."));
 
+static TAutoConsoleVariable<bool> CVarUpdateRigidBodyScaling(TEXT("p.UpdateRigidBodyScaling"), true, TEXT("Updates the rigid body scaling for simulation when the actor scaling is changed \n Default: true."));
+
 // Used as a default return value for invalid cloth data access
 static const TMap<int32, FClothSimulData> SEmptyClothSimulationData;
 
@@ -753,6 +755,8 @@ void USkeletalMeshComponent::InitArticulated(FPhysScene* PhysScene)
 	// Update Flag
 	PrevRootBoneMatrix = GetBoneMatrix(0); // save the root bone transform
 
+	PrevRootRigidBodyScale = FVector3f(GetRelativeScale3D()); //Update the root rigid body scale. 
+
 	// pre-compute cloth teleport thresholds for performance
 	ComputeTeleportDistanceThresholdInRadians();
 	ComputeTeleportRotationThresholdInRadians();
@@ -1596,6 +1600,12 @@ void USkeletalMeshComponent::OnUpdateTransform(EUpdateTransformFlags UpdateTrans
 
 	// Mark the cloth simulation transform update as pending (the actual update will happen in the cloth thread if the simulation is running)
 	UpdateClothTransform(Teleport);
+
+	if (CVarUpdateRigidBodyScaling.GetValueOnGameThread())
+	{
+		//Updates the scale of the rigid bodies for physics. 
+		UpdateRigidBodyScaling(Teleport);
+	}
 }
 
 bool USkeletalMeshComponent::UpdateOverlapsImpl(const TOverlapArrayView* PendingOverlaps, bool bDoNotifies, const TOverlapArrayView* OverlapsAtEndLocation)
@@ -3259,6 +3269,55 @@ void USkeletalMeshComponent::UpdateClothTransform(ETeleportType TeleportType)
 	PendingTeleportType = ((TeleportType > PendingTeleportType) ? TeleportType : PendingTeleportType); 
 }
 
+void USkeletalMeshComponent::UpdateRigidBodyScaling(ETeleportType TeleportType)
+{
+	if (bEnablePerPolyCollision == false)
+	{
+		FVector3f SceneRelativeScale = FVector3f(GetRelativeScale3D());
+
+		if (!SceneRelativeScale.IsUniform())
+		{
+			UE_LOG(LogSkeletalMesh, Warning, TEXT("USkeletalMeshComponent::UpdateRigidBodyScaling : Non-uniform scale applied to skeletal mesh %s "), *GetPathName());
+			return;
+		}
+
+		if (!(PrevRootRigidBodyScale - SceneRelativeScale).IsNearlyZero())
+		{
+			const FVector RelativeScale3DLocal = FVector(SceneRelativeScale / PrevRootRigidBodyScale);
+
+			for (FBodyInstance* BI : Bodies)
+			{
+				if (BI->IsValidBodyInstance())
+				{
+					//TODO(Chaos): Using current scale to compute for desired scale may lead to drift when scales are continuously varied. 
+					// It may help if the scale can be pulled from somewhere else. 
+					if (BI->UpdateBodyScale(RelativeScale3DLocal * BI->Scale3D, true))
+					{
+						//Reset rotation as before. 
+						BI->SetBodyTransform(GetBoneTransform(BI->InstanceBoneIndex), TeleportType, true);
+					}
+				}
+			}
+
+			//Reset relative transform using the relative scale. 
+			for (int32 i = 0; i < Constraints.Num(); i++)
+			{
+				check(Constraints[i]);
+				if (Chaos::FJointConstraint* JointConstraint = static_cast<Chaos::FJointConstraint*>(Constraints[i]->ConstraintHandle.Constraint))
+				{
+					Chaos::FTransformPair JointTransforms = JointConstraint->GetJointTransforms();
+					JointTransforms[0].SetTranslation(JointTransforms[0].GetTranslation() * RelativeScale3DLocal);
+					JointTransforms[1].SetTranslation(JointTransforms[1].GetTranslation() * RelativeScale3DLocal);
+
+					JointConstraint->SetJointTransforms(JointTransforms);
+				}
+			}
+
+			PrevRootRigidBodyScale = FVector3f(GetRelativeScale3D());
+		}
+	}
+}
+
 void USkeletalMeshComponent::CheckClothTeleport()
 {
 	// Get the root bone transform
