backout 39180841 due to startup lockfree crash

diff --git a/Engine/Source/Runtime/Engine/Private/SkeletalRenderGPUSkin.cpp b/Engine/Source/Runtime/Engine/Private/SkeletalRenderGPUSkin.cpp
--- a/Engine/Source/Runtime/Engine/Private/SkeletalRenderGPUSkin.cpp
+++ b/Engine/Source/Runtime/Engine/Private/SkeletalRenderGPUSkin.cpp
@@ -2360,10 +2360,12 @@ void FDynamicSkelMeshObjectDataGPUSkin::Clear()
 
 #define SKELETON_POOL_GPUSKINS 1
 #if SKELETON_POOL_GPUSKINS
-std::atomic<int64> FreeGpuSkinsCounter{0};
-TLockFreePointerListFIFO<FDynamicSkelMeshObjectDataGPUSkin, 0> FreeGpuSkinsFIFO;
+TArray<FDynamicSkelMeshObjectDataGPUSkin*> FreeGpuSkins;
+FCriticalSection FreeGpuSkinsCriticalSection;
 
 static int32 GPoolGpuSkins = 1;
+static int32 GMinPoolCount = 0;
+static int32 GAllocationCounter = 0;
 static const int32 GAllocationsBeforeCleanup = 1000; // number of allocations we make before we clean up the pool, this number is increased when we have to allocate not from the pool
 static FAutoConsoleVariableRef CVarPoolGpuSkins(
 	TEXT("r.GpuSkin.Pool"),
@@ -2371,19 +2373,6 @@ static FAutoConsoleVariableRef CVarPoolGpuSkins(
 	TEXT("Should we pool gpu skins.\n")
 	TEXT(" 0: Don't pool anything\n")
 	TEXT(" 1: Pool gpu skins bro (default)\n"),
-	FConsoleVariableDelegate::CreateLambda([](IConsoleVariable*)
-	{
-		if (!GPoolGpuSkins)
-		{
-			while (FDynamicSkelMeshObjectDataGPUSkin* Skin = FreeGpuSkinsFIFO.Pop())
-			{
-				// Pooling is disabled, just destroys
-				FDynamicSkelMeshObjectDataGPUSkin::FreeDynamicSkelMeshObjectDataGPUSkin(Skin);
-			}
-
-			FreeGpuSkinsCounter.store(0);
-		}
-	}),
 	ECVF_Default
 );
 #endif
@@ -2392,45 +2381,64 @@ static FAutoConsoleVariableRef CVarPoolGpuSkins(
 FDynamicSkelMeshObjectDataGPUSkin* FDynamicSkelMeshObjectDataGPUSkin::AllocDynamicSkelMeshObjectDataGPUSkin()
 {
 #if SKELETON_POOL_GPUSKINS
-	if (GPoolGpuSkins)
+	if (!GPoolGpuSkins)
 	{
-		if (FDynamicSkelMeshObjectDataGPUSkin* Skin = FreeGpuSkinsFIFO.Pop())
-		{
-			// Note that the counter isn't fully sync'ed with the free skins, but at most
-			// we risk undershooting it by a few entries.
-			--FreeGpuSkinsCounter;
-			
-			return Skin;
-		}
+		return new FDynamicSkelMeshObjectDataGPUSkin;
 	}
-#endif // SKELETON_POOL_GPUSKINS
-	
+
+	FScopeLock S(&FreeGpuSkinsCriticalSection);
+	++GAllocationCounter;
+	GMinPoolCount = FMath::Min(FreeGpuSkins.Num(), GMinPoolCount);
+	if ( FreeGpuSkins.Num() > 0 )
+	{
+		FDynamicSkelMeshObjectDataGPUSkin *Result = FreeGpuSkins[0];
+		FreeGpuSkins.RemoveAtSwap(0);
+		return Result;
+	}
+	else
+	{
+		return new FDynamicSkelMeshObjectDataGPUSkin;
+	}
+#else
 	return new FDynamicSkelMeshObjectDataGPUSkin;
+#endif
 }
 
 void FDynamicSkelMeshObjectDataGPUSkin::FreeDynamicSkelMeshObjectDataGPUSkin(FDynamicSkelMeshObjectDataGPUSkin* Who)
 {
 #if SKELETON_POOL_GPUSKINS
-	if (GPoolGpuSkins)
+	if (!GPoolGpuSkins)
 	{
-		// If exceeded, destroy the oldest and trade with this one
-		if (++FreeGpuSkinsCounter > GAllocationsBeforeCleanup)
+		delete Who;
+
+		if ( FreeGpuSkins.Num() > 0 )
 		{
-			if (FDynamicSkelMeshObjectDataGPUSkin* Skin = FreeGpuSkinsFIFO.Pop())
+			FScopeLock S(&FreeGpuSkinsCriticalSection);
+			for ( FDynamicSkelMeshObjectDataGPUSkin* I : FreeGpuSkins )
 			{
-				--FreeGpuSkinsCounter;
-				delete Skin;
+				delete I;
 			}
+			FreeGpuSkins.Empty();
 		}
-
-		Who->Clear();
-	
-		FreeGpuSkinsFIFO.Push(Who);
 		return;
 	}
-#endif // SKELETON_POOL_GPUSKINS
-	
+
+	Who->Clear();
+	FScopeLock S(&FreeGpuSkinsCriticalSection);
+	FreeGpuSkins.Add(Who);
+	if ( GAllocationCounter > GAllocationsBeforeCleanup )
+	{
+		GAllocationCounter = 0;
+		for ( int32 I = 0; I < GMinPoolCount; ++I )
+		{
+			delete FreeGpuSkins[0];
+			FreeGpuSkins.RemoveAtSwap(0);
+		}
+		GMinPoolCount = FreeGpuSkins.Num();
+	}
+#else
 	delete Who;
+#endif
 }
 
 void FDynamicSkelMeshObjectDataGPUSkin::InitDynamicSkelMeshObjectDataGPUSkin(
