Add a "Get input device for force feedback" function on the input device subsystem

diff --git a/Engine/Plugins/Runtime/Windows/XInputDevice/Source/XInputDevice/Private/XInputInterface.cpp b/Engine/Plugins/Runtime/Windows/XInputDevice/Source/XInputDevice/Private/XInputInterface.cpp
--- a/Engine/Plugins/Runtime/Windows/XInputDevice/Source/XInputDevice/Private/XInputInterface.cpp
+++ b/Engine/Plugins/Runtime/Windows/XInputDevice/Source/XInputDevice/Private/XInputInterface.cpp
@@ -378,7 +378,9 @@ void XInputInterface::SetChannelValue( int32 ControllerId, const FForceFeedbackC
 
 	UInputDeviceSubsystem* DeviceSubsystem = UInputDeviceSubsystem::Get();
 	check(DeviceSubsystem);
-	const FInputDeviceId MostRecentDevice = DeviceSubsystem->GetMostRecentlyUsedInputDeviceId(UserId);
+	
+	// Get the latest gamepad input device
+	const FInputDeviceId MostRecentDevice = DeviceSubsystem->GetLatestDeviceOfType(UserId, EHardwareDevicePrimaryType::Gamepad);
 	
 	if (ControllerId >= 0 && ControllerId < MAX_NUM_XINPUT_CONTROLLERS)
 	{
@@ -421,7 +423,7 @@ void XInputInterface::SetChannelValues( int32 ControllerId, const FForceFeedback
 
 	UInputDeviceSubsystem* DeviceSubsystem = UInputDeviceSubsystem::Get();
 	check(DeviceSubsystem);
-	const FInputDeviceId MostRecentDevice = DeviceSubsystem->GetMostRecentlyUsedInputDeviceId(UserId);
+	const FInputDeviceId MostRecentDevice = DeviceSubsystem->GetLatestDeviceOfType(UserId, EHardwareDevicePrimaryType::Gamepad);
 	
 	if (ControllerId >= 0 && ControllerId < MAX_NUM_XINPUT_CONTROLLERS)
 	{
@@ -432,11 +434,20 @@ void XInputInterface::SetChannelValues( int32 ControllerId, const FForceFeedback
 			const FInputDeviceId CurrentDeviceId = InternalDeviceIdMappings.FindDeviceId(ControllerState.ControllerId);
 
 			if (ControllerState.bIsConnected &&
-				ControllerState.LastUsedValidPlatformUserId == UserId &&
-				CurrentDeviceId.IsValid() && CurrentDeviceId == MostRecentDevice)
+				ControllerState.LastUsedValidPlatformUserId == UserId)
 			{
-				ControllerState.ForceFeedback = Values;
-			}	
+				if (CurrentDeviceId.IsValid() && CurrentDeviceId == MostRecentDevice)
+				{
+					ControllerState.ForceFeedback = Values;
+				}
+				// Ensure that other gamepads mapped to this user are zero'd out if they are not currently active for force feedback
+				// This way you don't get a controller which is stuck in a force feedback loop if you change devices mid-effect if there are
+				// multiple devices mapped to the same user
+				else
+				{
+					ControllerState.ForceFeedback = FForceFeedbackValues();
+				}
+			}			
 		}		
 	}
 }
diff --git a/Engine/Source/Runtime/Engine/Classes/GameFramework/InputDeviceSubsystem.h b/Engine/Source/Runtime/Engine/Classes/GameFramework/InputDeviceSubsystem.h
--- a/Engine/Source/Runtime/Engine/Classes/GameFramework/InputDeviceSubsystem.h
+++ b/Engine/Source/Runtime/Engine/Classes/GameFramework/InputDeviceSubsystem.h
@@ -199,13 +199,21 @@ public:
 	 * Returns the most recently used FInputDeviceId for the given platform user id.
 	 * 
 	 * This will be INPUTDEVICEID_NONE if there is no known device for the given user.
+	 * 
+	 * @param InUserId	The platform user to find the most recent device for
+	 * @param OfType	The type of input device to filter by. If unspecified, then the latest device regardless of its type will be returned.
 	 */
 	UFUNCTION(BlueprintCallable, Category = "Input Devices", meta = (ReturnDisplayName = "Latest Device Id"))
-	ENGINE_API FInputDeviceId GetMostRecentlyUsedInputDeviceId(const FPlatformUserId InUserId) const;
+	ENGINE_API FInputDeviceId GetMostRecentlyUsedInputDeviceId(const FPlatformUserId InUserId, const EHardwareDevicePrimaryType OfType = EHardwareDevicePrimaryType::Unspecified) const;
 	
 	UFUNCTION(BlueprintCallable, Category = "Input Devices")
 	ENGINE_API FHardwareDeviceIdentifier GetInputDeviceHardwareIdentifier(const FInputDeviceId InputDevice) const;
 
+	/**
+	 * Returns the latest used input device for a specific type of hardware
+	 */
+	ENGINE_API FInputDeviceId GetLatestDeviceOfType(const FPlatformUserId InUserId, const EHardwareDevicePrimaryType OfType) const;
+
 	/** A delegate that is fired when a platform user changes what Hardware Input device they are using */
 	UPROPERTY(BlueprintAssignable, Category = "Input Devices")
 	FHardwareInputDeviceChanged OnInputHardwareDeviceChanged;
@@ -257,7 +265,9 @@ protected:
 	UPROPERTY(Transient)
 	TSet<FInputDevicePropertyHandle> PropertiesPendingRemoval;
 	
-	/** A map of an input device to it's most recent hardware device identifier */
+	/**
+	 * A map of an input device to it's most recent hardware device identifier
+	 */
 	TMap<FInputDeviceId, FHardwareDeviceIdentifier> LatestInputDeviceIdentifiers;
 
 	/**
@@ -265,21 +275,21 @@ protected:
 	 */
 	struct FLatestInputDeviceData
 	{
-		/** The input device ID of the latest device used */
+		/**
+		 * The most recently used input device by this user
+		 */
 		FInputDeviceId LatestDeviceId = INPUTDEVICEID_NONE;
-
-		/** The latest hardware device ID for the input device last used by the platform user. */
-		FHardwareDeviceIdentifier HardwareDeviceId = FHardwareDeviceIdentifier::Invalid;
-
-		bool operator==(const FLatestInputDeviceData& Other) const
-		{
-			return LatestDeviceId == Other.LatestDeviceId && HardwareDeviceId == Other.HardwareDeviceId;
-		}
+		
+		/**
+		 * A map of the latest used input device of each type for a single user.
+		 * This allows us to query the most recently used device of a specific type for the player
+		 */
+		TMap<EHardwareDevicePrimaryType, FInputDeviceId> LatestDevices;
 	};
 	
-	/** A map of platform user's to their most recent hardware device identifier */
+	/** A map of platform user's to their most recent hardware device identifiers */
 	TMap<FPlatformUserId, FLatestInputDeviceData> LatestUserDeviceIdentifiers;
 
 	/** An input processor that is used to determine the most recently used hardware device for each user*/
 	TSharedPtr<class FInputDeviceSubsystemProcessor> InputPreprocessor;
-};
+};
\ No newline at end of file
diff --git a/Engine/Source/Runtime/Engine/Private/GameFramework/InputDeviceSubsystem.cpp b/Engine/Source/Runtime/Engine/Private/GameFramework/InputDeviceSubsystem.cpp
--- a/Engine/Source/Runtime/Engine/Private/GameFramework/InputDeviceSubsystem.cpp
+++ b/Engine/Source/Runtime/Engine/Private/GameFramework/InputDeviceSubsystem.cpp
@@ -494,22 +494,57 @@ FHardwareDeviceIdentifier UInputDeviceSubsystem::GetMostRecentlyUsedHardwareDevi
 {
 	if (const FLatestInputDeviceData* FoundDevice = LatestUserDeviceIdentifiers.Find(InUserId))
 	{
-		return FoundDevice->HardwareDeviceId;
+		if (const FHardwareDeviceIdentifier* ExistingDeviceId = LatestInputDeviceIdentifiers.Find(FoundDevice->LatestDeviceId))
+		{
+			return *ExistingDeviceId;
+		}
 	}
 		
 	return FHardwareDeviceIdentifier::Invalid;
 }
 
-FInputDeviceId UInputDeviceSubsystem::GetMostRecentlyUsedInputDeviceId(const FPlatformUserId InUserId) const
+FInputDeviceId UInputDeviceSubsystem::GetMostRecentlyUsedInputDeviceId(const FPlatformUserId InUserId, const EHardwareDevicePrimaryType OfType /*= EHardwareDevicePrimaryType::Unspecified*/) const
 {
 	if (const FLatestInputDeviceData* FoundDevice = LatestUserDeviceIdentifiers.Find(InUserId))
 	{
-		return FoundDevice->LatestDeviceId;
+		// If we didn't specific a hardware device type, then just return the latest one that has been used
+		if (OfType == EHardwareDevicePrimaryType::Unspecified)
+		{
+			return FoundDevice->LatestDeviceId;	
+		}
+
+		// Attempt to find a device of that specific type
+		if (const FInputDeviceId* FoundDeviceOfType = FoundDevice->LatestDevices.Find(OfType))
+		{
+			return *FoundDeviceOfType;
+		}
 	}
 		
 	return INPUTDEVICEID_NONE;
 }
 
+FInputDeviceId UInputDeviceSubsystem::GetLatestDeviceOfType(const FPlatformUserId InUserId, const EHardwareDevicePrimaryType OfType) const
+{
+	// TODO: In the future, we may provide an interface to specify a specify a specific
+	// input device to play force feedback on. We would implement that here. For example, it could
+	// be useful to be able to play different force feedback events on multiple devices for a given user
+	// or otherwise allow a game to specify how a force feedback effect behaves when played.
+	// For example, you may want to play a force feedback event on multiple input devices if you have a racing sim
+	// or HOTAS setup, with multiple HID's connected at once. 
+	
+	// Query what the most recently used devices are for the given user...
+	if (const FLatestInputDeviceData* FoundDevice = LatestUserDeviceIdentifiers.Find(InUserId))
+	{
+		// And then return the most recently used device of this type for them.
+		if (const FInputDeviceId* FoundDeviceOfType = FoundDevice->LatestDevices.Find(OfType))
+		{
+			return *FoundDeviceOfType;
+		}
+	}
+	
+	return INPUTDEVICEID_NONE;
+}
+
 FHardwareDeviceIdentifier UInputDeviceSubsystem::GetInputDeviceHardwareIdentifier(const FInputDeviceId InputDevice) const
 {
 	if (const FHardwareDeviceIdentifier* FoundDevice = LatestInputDeviceIdentifiers.Find(InputDevice))
@@ -523,34 +558,45 @@ FHardwareDeviceIdentifier UInputDeviceSubsystem::GetInputDeviceHardwareIdentifie
 void UInputDeviceSubsystem::SetMostRecentlyUsedHardwareDevice(const FInputDeviceId InDeviceId, const FHardwareDeviceIdentifier& InHardwareId)
 {	
 	const FPlatformUserId OwningUserId = IPlatformInputDeviceMapper::Get().GetUserForInputDevice(InDeviceId);
+	
+	// Keep track of each input device's latest hardware id so that it can be queried at any time
+	LatestInputDeviceIdentifiers.Add(InDeviceId, InHardwareId);
 
 	// If this hardware is the same as what the platform user already has, then there is no need to fire this event
-	if (FLatestInputDeviceData* ExistingDevice = LatestUserDeviceIdentifiers.Find(OwningUserId))
+	if (FLatestInputDeviceData* ExistingData = LatestUserDeviceIdentifiers.Find(OwningUserId))
 	{
-		// We always want to update the device ID
-		ExistingDevice->LatestDeviceId = InDeviceId;
-		
-		// If they are exactly the same, do nothing
-		if (InHardwareId == ExistingDevice->HardwareDeviceId)
-		{
-			return;
-		}
-		// If they are the same device but different input API's, just update the data but don't broadcast the delegate
-		// This can happen if there are multiple input interfaces sending messages for the same device (such as a keyboard)
-		// Which your gameplay code doesn't really care about switching which API the input came from
-		else if (InHardwareId.HardwareDeviceIdentifier == ExistingDevice->HardwareDeviceId.HardwareDeviceIdentifier)
+		// We always want to update the device ID and info about the device
+		ExistingData->LatestDeviceId = InDeviceId;
+		ExistingData->LatestDevices.Add(InHardwareId.PrimaryDeviceType, InDeviceId);
+
+		// Some conditions to exit early to not broadcast the OnInputHardwareDeviceChanged delegate
+		// based on the current state of the known input devices
+		if (FHardwareDeviceIdentifier* ExistingDeviceId = LatestInputDeviceIdentifiers.Find(InDeviceId))
 		{
-			ExistingDevice->HardwareDeviceId = InHardwareId;
-			return;
+			// If the hardware devices are the same, then there is no need to broadcast a "changed" delegate, because
+			// they have not changed.
+			if (InHardwareId == *ExistingDeviceId)
+			{
+				return;
+			}
+			// If they are the same device but different input API's, just update the data but don't broadcast the delegate
+			// This can happen if there are multiple input interfaces sending messages for the same device (such as a keyboard)
+			// Which your gameplay code doesn't really care about switching which API the input came from
+			else if (InHardwareId.HardwareDeviceIdentifier == ExistingDeviceId->HardwareDeviceIdentifier)
+			{
+				*ExistingDeviceId = InHardwareId;
+				return;
+			}	
 		}
 	}
 
-	// Keep track of each input device's latest hardware id
-	LatestInputDeviceIdentifiers.Add(InDeviceId, InHardwareId);
-
 	// Keep a map to platform users so that we can easily get their most recent hardware
-	LatestUserDeviceIdentifiers.Add(OwningUserId, { .LatestDeviceId = InDeviceId, .HardwareDeviceId = InHardwareId });
+	FLatestInputDeviceData& LatestData = LatestUserDeviceIdentifiers.Add(OwningUserId);
 
+	// Keep track of the latest hardware device and the latest of this type for the given user
+	LatestData.LatestDevices.FindOrAdd(InHardwareId.PrimaryDeviceType, InDeviceId);
+	LatestData.LatestDeviceId = InDeviceId;
+	
 	if (OnInputHardwareDeviceChanged.IsBound())
 	{
 		OnInputHardwareDeviceChanged.Broadcast(OwningUserId, InDeviceId);
