Chaos:: Disable scale warning for Autotest Cook. Fix CIS 744389

diff --git a/Engine/Source/Runtime/Engine/Private/SkeletalMeshComponentPhysics.cpp b/Engine/Source/Runtime/Engine/Private/SkeletalMeshComponentPhysics.cpp
--- a/Engine/Source/Runtime/Engine/Private/SkeletalMeshComponentPhysics.cpp
+++ b/Engine/Source/Runtime/Engine/Private/SkeletalMeshComponentPhysics.cpp
@@ -3275,9 +3275,9 @@ void USkeletalMeshComponent::UpdateRigidBodyScaling(ETeleportType TeleportType)
 	{
 		FVector3f SceneRelativeScale = FVector3f(GetRelativeScale3D());
 
+		//Physics does not support nonuniform scaling of the bodies.
 		if (!SceneRelativeScale.IsUniform())
 		{
-			UE_LOG(LogSkeletalMesh, Warning, TEXT("USkeletalMeshComponent::UpdateRigidBodyScaling : Non-uniform scale applied to skeletal mesh %s "), *GetPathName());
 			return;
 		}
 
