Components - Made the OnPreEndOfFrameSync loop safe for when new components are added to the pre-end of frame sync list while the loop is being iterated.

[RN] Components - Made the OnPreEndOfFrameSync loop safe for when new components are added to the pre-end of frame sync list while the loop is being iterated.

diff --git a/Engine/Source/Runtime/Engine/Private/LevelTick.cpp b/Engine/Source/Runtime/Engine/Private/LevelTick.cpp
--- a/Engine/Source/Runtime/Engine/Private/LevelTick.cpp
+++ b/Engine/Source/Runtime/Engine/Private/LevelTick.cpp
@@ -994,7 +994,14 @@ void UWorld::SendAllEndOfFrameUpdates()
 
 	// Wait for tasks that are generating data for the render proxies, but are not awaited in any TickFunctions 
 	// E.g., see cloth USkeletalMeshComponent::UpdateClothStateAndSimulate
-	for (UActorComponent* Component : ComponentsThatNeedPreEndOfFrameSync)
+	static TArray<TObjectPtr<UActorComponent>> LocalComponentsThatNeedPreEndOfFrameSync;
+	{
+		QUICK_SCOPE_CYCLE_COUNTER(STAT_PostTickComponentUpdate_PreEndOfFrameSync_Gather);
+		check(IsInGameThread() && !LocalComponentsThatNeedPreEndOfFrameSync.Num());
+		LocalComponentsThatNeedPreEndOfFrameSync.Append(ComponentsThatNeedPreEndOfFrameSync.Array());
+	}
+
+	for (UActorComponent* Component : LocalComponentsThatNeedPreEndOfFrameSync)
 	{
 		if (Component)
 		{
@@ -1006,6 +1013,7 @@ void UWorld::SendAllEndOfFrameUpdates()
 			FMarkComponentEndOfFrameUpdateState::ClearMarkedForPreEndOfFrameSync(Component);
 		}
 	}
+	LocalComponentsThatNeedPreEndOfFrameSync.Reset();
 	ComponentsThatNeedPreEndOfFrameSync.Reset();
 
 	//If we call SendAllEndOfFrameUpdates during a tick, we must ensure that all marked objects have completed any async work etc before doing the updates.
