Added Metal SM6 to checks for VSM

diff --git a/Engine/Source/Runtime/Engine/Private/RendererSettings.cpp b/Engine/Source/Runtime/Engine/Private/RendererSettings.cpp
--- a/Engine/Source/Runtime/Engine/Private/RendererSettings.cpp
+++ b/Engine/Source/Runtime/Engine/Private/RendererSettings.cpp
@@ -20,7 +20,7 @@
 #include "Misc/MessageDialog.h"
 #include "HAL/PlatformFileManager.h"
 
-#if PLATFORM_WINDOWS || PLATFORM_LINUX
+#if PLATFORM_WINDOWS || PLATFORM_LINUX || PLATFORM_MAC
 #include "Framework/Docking/TabManager.h"
 #include "Framework/Notifications/NotificationManager.h"
 #include "Widgets/Notifications/SNotificationList.h"
@@ -504,7 +504,7 @@ bool URendererSettings::CanEditChange(const FProperty* InProperty) const
 void URendererSettings::CheckForMissingShaderModels()
 {
 	// Don't show the SM6 toasts on non-Windows/Linux platforms to avoid confusion around platform requirements.
-#if PLATFORM_WINDOWS || PLATFORM_LINUX
+#if PLATFORM_WINDOWS || PLATFORM_LINUX || PLATFORM_MAC
 	if (GIsEditor && (ShadowMapMethod == EShadowMapMethod::VirtualShadowMaps || bEnableRayTracing))
 	{
 		auto CopySM6Format = [](const TCHAR* ShaderFormatName, const TArray<FString>& SrcArray, TArray<FString>& DstArray)
@@ -521,6 +521,8 @@ void URendererSettings::CheckForMissingShaderModels()
 		TArray<FString> WindowsTargetedRHIs;
 		TArray<FString> LinuxVulkanTargetedShaderFormats;
 		TArray<FString> LinuxTargetedRHIs;
+		TArray<FString> MacTargetedShaderFormats;
+		TArray<FString> MacTargetedRHIs;
 
 #if PLATFORM_WINDOWS
 		// Gather all Windows shader format settings
@@ -547,6 +549,8 @@ void URendererSettings::CheckForMissingShaderModels()
 		CopySM6Format(TEXT("SF_VULKAN_SM6"), LinuxTargetedRHIs, LinuxVulkanTargetedShaderFormats);
 #elif PLATFORM_MAC
 		// TODO: Gather all Mac shader format settings
+		GConfig->GetArray(TEXT("/Script/MacTargetPlatform.MacTargetSettings"), TEXT("TargetedRHIs"), MacTargetedRHIs, GEngineIni);
+		CopySM6Format(TEXT("SF_METAL_SM6"), MacTargetedRHIs, MacTargetedShaderFormats);
 #endif
 
 		const bool bProjectUsesD3D = (D3D11TargetedShaderFormats.Num() + D3D12TargetedShaderFormats.Num()) > 0;
@@ -557,8 +561,11 @@ void URendererSettings::CheckForMissingShaderModels()
 
 		const bool bProjectUsesLinuxVulkan = (LinuxTargetedRHIs.Num() > 0) || (LinuxVulkanTargetedShaderFormats.Num() > 0);
 		const bool bProjectMissingLinuxVulkanSM6 = (bProjectUsesLinuxVulkan && !LinuxVulkanTargetedShaderFormats.Contains(TEXT("SF_VULKAN_SM6")));
+		
+		const bool bProjectUsesMetal = (MacTargetedRHIs.Num() > 0) || (MacTargetedShaderFormats.Num() > 0);
+		const bool bProjectMissingMacSM6 = (bProjectUsesMetal && !MacTargetedShaderFormats.Contains(TEXT("SF_METAL_SM6")));
 
-		if (bProjectMissingD3DSM6 || bProjectMissingWindowsVulkanSM6 || bProjectMissingLinuxVulkanSM6)
+		if (bProjectMissingD3DSM6 || bProjectMissingWindowsVulkanSM6 || bProjectMissingLinuxVulkanSM6 || bProjectMissingMacSM6)
 		{
 			auto DismissNotification = [this]()
 			{
@@ -645,6 +652,21 @@ void URendererSettings::CheckForMissingShaderModels()
 					FeatureWontWorkLinuxVulkan[FeatureNeedsSM6Index]
 				);
 			}
+			else if(bProjectMissingMacSM6)
+			{
+				const FText FeatureWontWorkMac[] =
+				{
+					LOCTEXT("VirtualShadowMapsWontWorkMac", "Virtual shadow maps will not work on Mac until this is enabled.."),
+					LOCTEXT("RayTracingWontWorkMac", "Ray tracing will not work on Mac until this is enabled.."),
+					LOCTEXT("VirtualShadowMapsAndRayTracingWontWorkMac", "Virtual shadow maps and ray tracing will not work in Mac until this is enabled."),
+				};
+
+				Info.SubText = FText::Format(
+					LOCTEXT("VirtualShadowMapsNeedsMacSM6Setting", "{0} Please enable this in:\n  Project Settings -> Platforms -> Mac -> Targeted RHIs\n{1}"),
+					FeaturesNeedSM6[FeatureNeedsSM6Index],
+					FeatureWontWorkMac[FeatureNeedsSM6Index]
+				);
+			}
 
 			ShaderModelNotificationPtr = FSlateNotificationManager::Get().AddNotification(Info);
 		}
