Improve Previewing Texture Streaming Pool in Editor
- Get the correct PoolSize value from the DP cvar in Texture Streaming Stats
- Exclude Editor Resident Memory from the Runtime Resident Memory in Texture Streaming Stats

diff --git a/Engine/Source/Editor/MaterialEditor/Private/MaterialEditor.cpp b/Engine/Source/Editor/MaterialEditor/Private/MaterialEditor.cpp
--- a/Engine/Source/Editor/MaterialEditor/Private/MaterialEditor.cpp
+++ b/Engine/Source/Editor/MaterialEditor/Private/MaterialEditor.cpp
@@ -1074,6 +1074,11 @@ FMaterialEditor::~FMaterialEditor()
 {
 	bDestructing = true;
 
+	if(Material != nullptr)
+	{
+		Material->SetMarkTextureAsEditorStreamingPool(false);
+	}
+
 	if (OnWorldCleanupDelegateHandle.IsValid())
 	{
 		FWorldDelegates::OnWorldCleanup.Remove(OnWorldCleanupDelegateHandle);
diff --git a/Engine/Source/Editor/MaterialEditor/Private/MaterialEditorModule.cpp b/Engine/Source/Editor/MaterialEditor/Private/MaterialEditorModule.cpp
--- a/Engine/Source/Editor/MaterialEditor/Private/MaterialEditorModule.cpp
+++ b/Engine/Source/Editor/MaterialEditor/Private/MaterialEditorModule.cpp
@@ -100,6 +100,7 @@ public:
 	virtual TSharedRef<IMaterialEditor> CreateMaterialEditor(const EToolkitMode::Type Mode, const TSharedPtr< IToolkitHost >& InitToolkitHost, UMaterial* Material) override
 	{
 		TSharedRef<FMaterialEditor> NewMaterialEditor(new FMaterialEditor());
+		Material->SetMarkTextureAsEditorStreamingPool(true);
 		NewMaterialEditor->InitEditorForMaterial(Material);
 		OnMaterialEditorOpened().Broadcast(NewMaterialEditor);
 		NewMaterialEditor->InitMaterialEditor(Mode, InitToolkitHost, Material);
@@ -118,6 +119,7 @@ public:
 	virtual TSharedRef<IMaterialEditor> CreateMaterialInstanceEditor(const EToolkitMode::Type Mode, const TSharedPtr< IToolkitHost >& InitToolkitHost, UMaterialInstance* MaterialInstance) override
 	{
 		TSharedRef<FMaterialInstanceEditor> NewMaterialInstanceEditor(new FMaterialInstanceEditor());
+		MaterialInstance->SetMarkTextureAsEditorStreamingPool(true);
 		NewMaterialInstanceEditor->InitEditorForMaterial(MaterialInstance);
 		OnMaterialInstanceEditorOpened().Broadcast(NewMaterialInstanceEditor);
 		NewMaterialInstanceEditor->InitMaterialInstanceEditor(Mode, InitToolkitHost, MaterialInstance);
diff --git a/Engine/Source/Editor/MaterialEditor/Private/MaterialInstanceEditor.cpp b/Engine/Source/Editor/MaterialEditor/Private/MaterialInstanceEditor.cpp
--- a/Engine/Source/Editor/MaterialEditor/Private/MaterialInstanceEditor.cpp
+++ b/Engine/Source/Editor/MaterialEditor/Private/MaterialInstanceEditor.cpp
@@ -768,6 +768,8 @@ FMaterialInstanceEditor::~FMaterialInstanceEditor()
 	// The streaming data will be null if there were any edits
 	if (MaterialEditorInstance && MaterialEditorInstance->SourceInstance && !MaterialEditorInstance->SourceInstance->HasTextureStreamingData())
 	{
+		MaterialEditorInstance->SourceInstance->SetMarkTextureAsEditorStreamingPool(false);
+
 		UPackage* Package = MaterialEditorInstance->SourceInstance->GetOutermost();
 		if (Package && Package->IsDirty() && Package != GetTransientPackage())
 		{
diff --git a/Engine/Source/Editor/TextureEditor/Private/TextureEditorModule.cpp b/Engine/Source/Editor/TextureEditor/Private/TextureEditorModule.cpp
--- a/Engine/Source/Editor/TextureEditor/Private/TextureEditorModule.cpp
+++ b/Engine/Source/Editor/TextureEditor/Private/TextureEditorModule.cpp
@@ -35,6 +35,7 @@ public:
 	virtual TSharedRef<ITextureEditorToolkit> CreateTextureEditor( const EToolkitMode::Type Mode, const TSharedPtr< IToolkitHost >& InitToolkitHost, UTexture* Texture ) override
 	{
 		TSharedRef<FTextureEditorToolkit> NewTextureEditor(new FTextureEditorToolkit());
+		Texture->bMarkAsEditorStreamingPool = true;
 		NewTextureEditor->InitTextureEditor(Mode, InitToolkitHost, Texture);
 
 		return NewTextureEditor;
diff --git a/Engine/Source/Editor/TextureEditor/Private/TextureEditorToolkit.cpp b/Engine/Source/Editor/TextureEditor/Private/TextureEditorToolkit.cpp
--- a/Engine/Source/Editor/TextureEditor/Private/TextureEditorToolkit.cpp
+++ b/Engine/Source/Editor/TextureEditor/Private/TextureEditorToolkit.cpp
@@ -177,6 +177,7 @@ FTextureEditorToolkit::~FTextureEditorToolkit( )
 	
 	// we are leaving the texture editor
 	// restore any temporary encoding settings we may have changed
+	Texture->bMarkAsEditorStreamingPool = false;
 
 	if ( Texture->DeferCompression || 
 		Texture->OverrideRunningPlatformName != NAME_None ||
diff --git a/Engine/Source/Runtime/Engine/Classes/Components/MeshComponent.h b/Engine/Source/Runtime/Engine/Classes/Components/MeshComponent.h
--- a/Engine/Source/Runtime/Engine/Classes/Components/MeshComponent.h
+++ b/Engine/Source/Runtime/Engine/Classes/Components/MeshComponent.h
@@ -103,6 +103,14 @@ public:
 	 */
 	ENGINE_API virtual void SetTextureForceResidentFlag( bool bForceMiplevelsToBeResident );
 
+#if WITH_EDITOR
+	/**
+	 *	Set the bMarkAsEditorStreamingPool on all textures used by this component.
+	 *	@param bInMarkAsEditorStreamingPool		Whether textures should only be counted for the Editor pool stats.
+	 */
+	ENGINE_API void SetMarkTextureAsEditorStreamingPool(bool bInMarkAsEditorStreamingPool);
+#endif
+
 	/**
 	 *	Tell the streaming system to start loading all textures with all mip-levels.
 	 *	@param Seconds							Number of seconds to force all mip-levels to be resident
diff --git a/Engine/Source/Runtime/Engine/Classes/Engine/StreamableRenderAsset.h b/Engine/Source/Runtime/Engine/Classes/Engine/StreamableRenderAsset.h
--- a/Engine/Source/Runtime/Engine/Classes/Engine/StreamableRenderAsset.h
+++ b/Engine/Source/Runtime/Engine/Classes/Engine/StreamableRenderAsset.h
@@ -311,6 +311,10 @@ public:
 	UPROPERTY(transient)
 	uint8 bIgnoreStreamingMipBias : 1;
 
+	/** Object marked as being part of Editors Pool */
+	UPROPERTY(transient)
+	uint8 bMarkAsEditorStreamingPool : 1;
+	
 protected:
 	/** Whether to use the extra cinematic quality mip-levels, when we're forcing mip-levels to be resident. */
 	UPROPERTY(transient)
diff --git a/Engine/Source/Runtime/Engine/Private/Components/MeshComponent.cpp b/Engine/Source/Runtime/Engine/Private/Components/MeshComponent.cpp
--- a/Engine/Source/Runtime/Engine/Private/Components/MeshComponent.cpp
+++ b/Engine/Source/Runtime/Engine/Private/Components/MeshComponent.cpp
@@ -380,6 +380,20 @@ void UMeshComponent::SetTextureForceResidentFlag( bool bForceMiplevelsToBeReside
 	}
 }
 
+#if WITH_EDITOR
+void UMeshComponent::SetMarkTextureAsEditorStreamingPool(bool bInMarkAsEditorStreamingPool)
+{
+	TArray<UTexture*> Textures;
+	GetUsedTextures(/*out*/ Textures, EMaterialQualityLevel::Num);
+	FTextureCompilingManager::Get().FinishCompilation(Textures);
+
+	for (UTexture* Texture : Textures)
+	{
+		Texture->bMarkAsEditorStreamingPool = bInMarkAsEditorStreamingPool;
+	}
+}
+#endif
+
 TArray<class UMaterialInterface*> UMeshComponent::GetMaterials() const
 {
 	TArray<class UMaterialInterface*> OutMaterials;
diff --git a/Engine/Source/Runtime/Engine/Private/Materials/MaterialInterface.cpp b/Engine/Source/Runtime/Engine/Private/Materials/MaterialInterface.cpp
--- a/Engine/Source/Runtime/Engine/Private/Materials/MaterialInterface.cpp
+++ b/Engine/Source/Runtime/Engine/Private/Materials/MaterialInterface.cpp
@@ -2154,4 +2154,17 @@ void UMaterialInterfaceEditorOnlyData::Serialize(FArchive& Ar)
 	}
 }
 
+#if WITH_EDITOR
+void UMaterialInterface::SetMarkTextureAsEditorStreamingPool(bool bInMarkAsEditorStreamingPool)
+{
+	TArray<UTexture*> Textures;
+	GetUsedTextures(Textures, EMaterialQualityLevel::High, true, ERHIFeatureLevel::SM5, true);
+
+	for (UTexture* Texture : Textures)
+	{
+		Texture->bMarkAsEditorStreamingPool = bInMarkAsEditorStreamingPool;
+	}
+}
+#endif
+
 #undef LOCTEXT_NAMESPACE 
diff --git a/Engine/Source/Runtime/Engine/Private/PreviewScene.cpp b/Engine/Source/Runtime/Engine/Private/PreviewScene.cpp
--- a/Engine/Source/Runtime/Engine/Private/PreviewScene.cpp
+++ b/Engine/Source/Runtime/Engine/Private/PreviewScene.cpp
@@ -130,6 +130,13 @@ void FPreviewScene::Uninitialize()
 			}
 		}
 
+#if WITH_EDITOR
+		if (UMeshComponent* pMesh = Cast<UMeshComponent>(Component))
+		{
+			pMesh->SetMarkTextureAsEditorStreamingPool(false);
+		}
+#endif
+
 		Component->UnregisterComponent();
 	}
 
@@ -178,6 +185,13 @@ void FPreviewScene::AddComponent(UActorComponent* Component,const FTransform& Lo
 		}
 	}
 
+#if WITH_EDITOR
+	if (UMeshComponent* pMesh = Cast<UMeshComponent>(Component))
+	{
+		pMesh->SetMarkTextureAsEditorStreamingPool(true);
+	}
+#endif
+
 	{
 		UStaticMeshComponent* pStaticMesh = Cast<UStaticMeshComponent>(Component);
 		if(pStaticMesh != nullptr)
diff --git a/Engine/Source/Runtime/Engine/Private/Streaming/AsyncTextureStreaming.cpp b/Engine/Source/Runtime/Engine/Private/Streaming/AsyncTextureStreaming.cpp
--- a/Engine/Source/Runtime/Engine/Private/Streaming/AsyncTextureStreaming.cpp
+++ b/Engine/Source/Runtime/Engine/Private/Streaming/AsyncTextureStreaming.cpp
@@ -10,6 +10,11 @@ AsyncTextureStreaming.cpp: Definitions of classes used for texture streaming asy
 #include "Engine/Level.h"
 #include "ProfilingDebugging/CsvProfiler.h"
 
+#if ALLOW_OTHER_PLATFORM_CONFIG && WITH_EDITOR
+#include "DeviceProfiles/DeviceProfileManager.h"
+#include "DeviceProfiles/DeviceProfile.h"
+#endif
+
 CSV_DECLARE_CATEGORY_EXTERN(TextureStreaming);
 
 void FAsyncRenderAssetStreamingData::Init(
@@ -963,10 +968,28 @@ void FRenderAssetStreamingMipCalcTask::UpdateStats_Async()
 
 	Stats.SafetyPool = MemoryMargin; 
 	Stats.TemporaryPool = TempMemoryBudget;
+
+#if ALLOW_OTHER_PLATFORM_CONFIG && WITH_EDITOR
+	UDeviceProfileManager& DeviceProfileManager = UDeviceProfileManager::Get();
+	UDeviceProfile* PreviewDeviceProfile = DeviceProfileManager.GetPreviewDeviceProfile();
+	if (PreviewDeviceProfile)
+	{
+		const FString* PoolSizeString = PreviewDeviceProfile->GetAllExpandedCVars().Find(TEXT("r.Streaming.PoolSize"));
+		if (PoolSizeString != nullptr)
+		{
+			int64 PoolSizePreview = FCString::Atoi(**PoolSizeString);
+			MemoryBudget = PoolSizePreview * 1024 * 1024;
+		}
+	}
+#endif
 	Stats.StreamingPool = MemoryBudget;
+
 	Stats.NonStreamingMips = AllocatedMemory;
 
 	Stats.RequiredPool = 0;
+#if WITH_EDITOR
+	Stats.RequiredPool_Editor = 0;
+#endif
 	Stats.VisibleMips = 0;
 	Stats.HiddenMips = 0;
 
@@ -1002,6 +1025,18 @@ void FRenderAssetStreamingMipCalcTask::UpdateStats_Async()
 		const int64 MaxSize = StreamingRenderAsset.GetSize(StreamingRenderAsset.MaxAllowedMips);
 		const int64 VisibleWantedSize = StreamingRenderAsset.GetSize(StreamingRenderAsset.VisibleWantedMips);
 
+#if WITH_EDITOR
+		if (StreamingRenderAsset.IsTexture() && StreamingRenderAsset.RenderAsset->bMarkAsEditorStreamingPool)
+		{
+			// Remove from the non streaming budget what is actually taken by streaming.
+			Stats.NonStreamingMips -= RequiredSize;
+
+			// Editor resources should be forced resident
+			Stats.RequiredPool_Editor += RequiredSize;
+			continue;
+		}
+#endif
+
 		// How much the streamer would use if there was no limit.
 		Stats.RequiredPool += RequiredSize;
 
diff --git a/Engine/Source/Runtime/Engine/Private/Streaming/TextureStreamingHelpers.cpp b/Engine/Source/Runtime/Engine/Private/Streaming/TextureStreamingHelpers.cpp
--- a/Engine/Source/Runtime/Engine/Private/Streaming/TextureStreamingHelpers.cpp
+++ b/Engine/Source/Runtime/Engine/Private/Streaming/TextureStreamingHelpers.cpp
@@ -28,6 +28,9 @@ DECLARE_MEMORY_STAT_POOL(TEXT("Cached Mips"), STAT_Streaming11_CachedMips, STATG
 DECLARE_MEMORY_STAT_POOL(TEXT("Wanted Mips"), STAT_Streaming12_WantedMips, STATGROUP_Streaming, FPlatformMemory::MCR_UsedStreamingPool);
 DECLARE_MEMORY_STAT_POOL(TEXT("Inflight Requests"), STAT_Streaming13_InflightRequests, STATGROUP_Streaming, FPlatformMemory::MCR_UsedStreamingPool);
 DECLARE_MEMORY_STAT_POOL(TEXT("IO Bandwidth"), STAT_Streaming14_MipIOBandwidth, STATGROUP_Streaming, FPlatformMemory::MCR_UsedStreamingPool);
+#if WITH_EDITOR
+DECLARE_MEMORY_STAT_POOL(TEXT("Editor Required Pool"), STAT_Streaming15_EditorRequiredPool, STATGROUP_Streaming, FPlatformMemory::MCR_StreamingPool);
+#endif
 
 DECLARE_CYCLE_STAT(TEXT("Setup Async Task"), STAT_Streaming01_SetupAsyncTask, STATGROUP_Streaming);
 DECLARE_CYCLE_STAT(TEXT("Update Streaming Data"), STAT_Streaming02_UpdateStreamingData, STATGROUP_Streaming);
@@ -423,7 +426,9 @@ void FRenderAssetStreamingStats::Apply()
 	SET_MEMORY_STAT(STAT_Streaming12_WantedMips, WantedMips);
 	SET_MEMORY_STAT(STAT_Streaming13_InflightRequests, PendingRequests);	
 	SET_MEMORY_STAT(STAT_Streaming14_MipIOBandwidth, MipIOBandwidth);
-
+#if WITH_EDITOR
+	SET_MEMORY_STAT(STAT_Streaming15_EditorRequiredPool, RequiredPool_Editor);
+#endif
 	SET_CYCLE_COUNTER(STAT_Streaming01_SetupAsyncTask, SetupAsyncTaskCycles);
 	SET_CYCLE_COUNTER(STAT_Streaming02_UpdateStreamingData, UpdateStreamingDataCycles);
 	SET_CYCLE_COUNTER(STAT_Streaming03_StreamRenderAssets, StreamRenderAssetsCycles);
diff --git a/Engine/Source/Runtime/Engine/Private/Streaming/TextureStreamingHelpers.h b/Engine/Source/Runtime/Engine/Private/Streaming/TextureStreamingHelpers.h
--- a/Engine/Source/Runtime/Engine/Private/Streaming/TextureStreamingHelpers.h
+++ b/Engine/Source/Runtime/Engine/Private/Streaming/TextureStreamingHelpers.h
@@ -170,6 +170,9 @@ struct FRenderAssetStreamingStats
 	int64 NonStreamingMips;		// Estimated memory in bytes the streamer deemed non-streamable
 
 	int64 RequiredPool;			// Estimated memory in bytes the streamer would use if there was no limit
+#if WITH_EDITOR
+	int64 RequiredPool_Editor;	// Estimated memory in bytes used by the Editor
+#endif
 	int64 VisibleMips;			// Estimated memory in bytes used by visible assets
 	int64 HiddenMips;			// Estimated memory in bytes used by hidden (e.g. culled by view frustum or occlusion) assets
 	int64 ForcedMips;			// Estimated memory in bytes used by forced-fully-load assets
diff --git a/Engine/Source/Runtime/Engine/Public/Materials/MaterialInterface.h b/Engine/Source/Runtime/Engine/Public/Materials/MaterialInterface.h
--- a/Engine/Source/Runtime/Engine/Public/Materials/MaterialInterface.h
+++ b/Engine/Source/Runtime/Engine/Public/Materials/MaterialInterface.h
@@ -1279,6 +1279,13 @@ private:
 	static UEnum* SamplerTypeEnum;
 
 #if WITH_EDITOR
+public:
+	/**
+	 *	Set the bMarkAsEditorStreamingPool on all textures used by this MaterialInterface.
+	 *	@param bInMarkAsEditorStreamingPool		Whether textures should only be counted for the Editor pool stats.
+	 */
+	ENGINE_API void SetMarkTextureAsEditorStreamingPool(bool bInMarkAsEditorStreamingPool);
+
 protected:
 	mutable TOptional<FStaticParameterSet> CachedStaticParameterValues;
 	mutable uint8 AllowCachingStaticParameterValuesCounter = 0;
