[MaterialIR] Add trigonometric intrinsics as operators to IR generator of new material translator.

#rnx

diff --git a/Engine/Source/Runtime/Engine/Private/Materials/MaterialExpressionsIR.cpp b/Engine/Source/Runtime/Engine/Private/Materials/MaterialExpressionsIR.cpp
--- a/Engine/Source/Runtime/Engine/Private/Materials/MaterialExpressionsIR.cpp
+++ b/Engine/Source/Runtime/Engine/Private/Materials/MaterialExpressionsIR.cpp
@@ -472,6 +472,36 @@ void UMaterialExpressionTruncate::Build(UE::MIR::FEmitter& Emitter)
 	BuildUnaryArithmeticOperatorFromInput(Emitter, MIR::UO_Truncate, &Input);
 }
 
+void UMaterialExpressionArccosine::Build(UE::MIR::FEmitter& Emitter)
+{
+	BuildUnaryArithmeticOperatorFromInput(Emitter, MIR::UO_ACos, &Input);
+}
+
+void UMaterialExpressionArcsine::Build(UE::MIR::FEmitter& Emitter)
+{
+	BuildUnaryArithmeticOperatorFromInput(Emitter, MIR::UO_ASin, &Input);
+}
+
+void UMaterialExpressionArctangent::Build(UE::MIR::FEmitter& Emitter)
+{
+	BuildUnaryArithmeticOperatorFromInput(Emitter, MIR::UO_ATan, &Input);
+}
+
+void UMaterialExpressionCosine::Build(UE::MIR::FEmitter& Emitter)
+{
+	BuildUnaryArithmeticOperatorFromInput(Emitter, MIR::UO_Cos, &Input);
+}
+
+void UMaterialExpressionSine::Build(UE::MIR::FEmitter& Emitter)
+{
+	BuildUnaryArithmeticOperatorFromInput(Emitter, MIR::UO_Sin, &Input);
+}
+
+void UMaterialExpressionTangent::Build(UE::MIR::FEmitter& Emitter)
+{
+	BuildUnaryArithmeticOperatorFromInput(Emitter, MIR::UO_Tan, &Input);
+}
+
 /* Binary Operators */
 
 static void BuildBinaryArithmeticOperator(
diff --git a/Engine/Source/Runtime/Engine/Private/Materials/MaterialIR.cpp b/Engine/Source/Runtime/Engine/Private/Materials/MaterialIR.cpp
--- a/Engine/Source/Runtime/Engine/Private/Materials/MaterialIR.cpp
+++ b/Engine/Source/Runtime/Engine/Private/Materials/MaterialIR.cpp
@@ -393,12 +393,12 @@ bool IsBitOperator(EOperator Op)
 
 bool IsBinaryOperator(EOperator Op)
 {
-	return Op >= BO_GreaterThan && Op <= BO_Max;
+	return Op >= BO_FirstBinaryOperator && Op < TO_FirstTernaryOperator;
 }
 
 bool IsTernaryOperator(EOperator Op)
 {
-	return Op >= TO_Lerp;
+	return Op >= TO_FirstTernaryOperator;
 }
 
 bool IsBinaryInfixOperator(EOperator Op)
@@ -413,6 +413,12 @@ const TCHAR* OperatorToString(EOperator Op)
 		case O_Invalid: return TEXT("Invalid");
 
 		/* Unary Arithmetic */
+		case UO_ACos: return TEXT("Arccosine");
+		case UO_ASin: return TEXT("Arcsine");
+		case UO_ATan: return TEXT("Arctangent");
+		case UO_Cos: return TEXT("Cosine");
+		case UO_Sin: return TEXT("Sine");
+		case UO_Tan: return TEXT("Tangent");
 		case UO_Truncate: return TEXT("Truncate");
 
 		/* Binary Comparisons */
diff --git a/Engine/Source/Runtime/Engine/Private/Materials/MaterialIREmitter.cpp b/Engine/Source/Runtime/Engine/Private/Materials/MaterialIREmitter.cpp
--- a/Engine/Source/Runtime/Engine/Private/Materials/MaterialIREmitter.cpp
+++ b/Engine/Source/Runtime/Engine/Private/Materials/MaterialIREmitter.cpp
@@ -751,12 +751,43 @@ static bool FoldComparisonOperatorScalar(EOperator Operator, T A, T B)
 	}
 }
 
+template <typename T>
+static T FoldFloatScalarArithmeticOperator(EOperator Operator, T A, T B, T C)
+{
+	if constexpr (std::is_same<T, float>::value || std::is_same<T, double>::value)
+	{
+		switch (Operator)
+		{
+			case UO_ACos: return FMath::Acos(A);
+			case UO_ASin: return FMath::Asin(A);
+			case UO_ATan: return FMath::Atan(A);
+			case UO_Cos: return FMath::Cos(A);
+			case UO_Sin: return FMath::Sin(A);
+			case UO_Tan: return FMath::Tan(A);
+			default: UE_MIR_UNREACHABLE();
+		}
+	}
+	else
+	{
+		UE_MIR_UNREACHABLE();
+	}
+}
+
 template <typename T>
 static T FoldScalarArithmeticOperator(EOperator Operator, T A, T B, T C)
 {
 	switch (Operator)
 	{
 		case UO_Negate: return -A;
+		case UO_ACos: [[fallthrough]];
+		case UO_ASin: [[fallthrough]];
+		case UO_ATan: [[fallthrough]];
+		case UO_Cos: [[fallthrough]];
+		case UO_Sin: [[fallthrough]];
+		case UO_Tan:
+		{
+			return FoldFloatScalarArithmeticOperator<T>(Operator, A, B, C);
+		}
 		case UO_Truncate:
 		{
 			if constexpr (std::is_same<T, float>::value)
@@ -872,8 +903,14 @@ static FValue* TryFoldOperatorScalar(FEmitter* Emitter, EOperator Op, FValue* A,
 			break; // todo
 
 		/* No simplifications for these operations */
+		case UO_ACos:
+		case UO_ASin:
+		case UO_ATan:
+		case UO_Cos:
 		case UO_Negate:
 		case UO_Not:
+		case UO_Sin:
+		case UO_Tan:
 		case UO_Truncate:
 		case BO_Min:
 		case BO_Max:
@@ -937,7 +974,7 @@ static FValue* TryFoldOperatorScalar(FEmitter* Emitter, EOperator Op, FValue* A,
 		{
 			case ScalarKind_Int:
 			{
-				int Result = FoldScalarArithmeticOperator<TInteger>(Op, AConstant->Integer, BConstant ? BConstant->Integer : 0, CConstant ? CConstant->Integer : 0);
+				TInteger Result = FoldScalarArithmeticOperator<TInteger>(Op, AConstant->Integer, BConstant ? BConstant->Integer : 0, CConstant ? CConstant->Integer : 0);
 				return Emitter->EmitConstantIntScalar(Result);
 			}
 
diff --git a/Engine/Source/Runtime/Engine/Private/Materials/MaterialIRToHLSLTranslator.cpp b/Engine/Source/Runtime/Engine/Private/Materials/MaterialIRToHLSLTranslator.cpp
--- a/Engine/Source/Runtime/Engine/Private/Materials/MaterialIRToHLSLTranslator.cpp
+++ b/Engine/Source/Runtime/Engine/Private/Materials/MaterialIRToHLSLTranslator.cpp
@@ -459,6 +459,12 @@ struct FTranslator : FMaterialIRToHLSLTranslation
 			const TCHAR* OpString;
 			switch (Operator->Op)
 			{
+				case MIR::UO_ACos: OpString = TEXT("acos"); break;
+				case MIR::UO_ASin: OpString = TEXT("asin"); break;
+				case MIR::UO_ATan: OpString = TEXT("atan"); break;
+				case MIR::UO_Cos: OpString = TEXT("cos"); break;
+				case MIR::UO_Sin: OpString = TEXT("sin"); break;
+				case MIR::UO_Tan: OpString = TEXT("tan"); break;
 				case MIR::UO_Truncate: OpString = TEXT("trunc"); break;
 				case MIR::BO_Min: OpString = TEXT("min"); break;
 				case MIR::BO_Max: OpString = TEXT("max"); break;
diff --git a/Engine/Source/Runtime/Engine/Public/Materials/MaterialExpressionArccosine.h b/Engine/Source/Runtime/Engine/Public/Materials/MaterialExpressionArccosine.h
--- a/Engine/Source/Runtime/Engine/Public/Materials/MaterialExpressionArccosine.h
+++ b/Engine/Source/Runtime/Engine/Public/Materials/MaterialExpressionArccosine.h
@@ -18,6 +18,7 @@ class UMaterialExpressionArccosine : public UMaterialExpression
 
 	//~ Begin UMaterialExpression Interface
 #if WITH_EDITOR
+	virtual void Build(UE::MIR::FEmitter& Emitter) override;
 	virtual int32 Compile(class FMaterialCompiler* Compiler, int32 OutputIndex) override;
 	virtual void GetCaption(TArray<FString>& OutCaptions) const override;
 	virtual void GetExpressionToolTip(TArray<FString>& OutToolTip) override;
diff --git a/Engine/Source/Runtime/Engine/Public/Materials/MaterialExpressionArcsine.h b/Engine/Source/Runtime/Engine/Public/Materials/MaterialExpressionArcsine.h
--- a/Engine/Source/Runtime/Engine/Public/Materials/MaterialExpressionArcsine.h
+++ b/Engine/Source/Runtime/Engine/Public/Materials/MaterialExpressionArcsine.h
@@ -18,6 +18,7 @@ class UMaterialExpressionArcsine : public UMaterialExpression
 
 	//~ Begin UMaterialExpression Interface
 #if WITH_EDITOR
+	virtual void Build(UE::MIR::FEmitter& Emitter) override;
 	virtual int32 Compile(class FMaterialCompiler* Compiler, int32 OutputIndex) override;
 	virtual void GetCaption(TArray<FString>& OutCaptions) const override;
 	virtual void GetExpressionToolTip(TArray<FString>& OutToolTip) override;
diff --git a/Engine/Source/Runtime/Engine/Public/Materials/MaterialExpressionArctangent.h b/Engine/Source/Runtime/Engine/Public/Materials/MaterialExpressionArctangent.h
--- a/Engine/Source/Runtime/Engine/Public/Materials/MaterialExpressionArctangent.h
+++ b/Engine/Source/Runtime/Engine/Public/Materials/MaterialExpressionArctangent.h
@@ -18,6 +18,7 @@ class UMaterialExpressionArctangent : public UMaterialExpression
 
 	//~ Begin UMaterialExpression Interface
 #if WITH_EDITOR
+	virtual void Build(UE::MIR::FEmitter& Emitter) override;
 	virtual int32 Compile(class FMaterialCompiler* Compiler, int32 OutputIndex) override;
 	virtual void GetCaption(TArray<FString>& OutCaptions) const override;
 	virtual void GetExpressionToolTip(TArray<FString>& OutToolTip) override;
diff --git a/Engine/Source/Runtime/Engine/Public/Materials/MaterialExpressionCosine.h b/Engine/Source/Runtime/Engine/Public/Materials/MaterialExpressionCosine.h
--- a/Engine/Source/Runtime/Engine/Public/Materials/MaterialExpressionCosine.h
+++ b/Engine/Source/Runtime/Engine/Public/Materials/MaterialExpressionCosine.h
@@ -23,6 +23,7 @@ class UMaterialExpressionCosine : public UMaterialExpression
 
 	//~ Begin UMaterialExpression Interface
 #if WITH_EDITOR
+	virtual void Build(UE::MIR::FEmitter& Emitter) override;
 	virtual int32 Compile(class FMaterialCompiler* Compiler, int32 OutputIndex) override;
 	virtual void GetCaption(TArray<FString>& OutCaptions) const override;
 	virtual FText GetKeywords() const override {return FText::FromString(TEXT("cos"));}
diff --git a/Engine/Source/Runtime/Engine/Public/Materials/MaterialExpressionSine.h b/Engine/Source/Runtime/Engine/Public/Materials/MaterialExpressionSine.h
--- a/Engine/Source/Runtime/Engine/Public/Materials/MaterialExpressionSine.h
+++ b/Engine/Source/Runtime/Engine/Public/Materials/MaterialExpressionSine.h
@@ -23,6 +23,7 @@ class UMaterialExpressionSine : public UMaterialExpression
 
 	//~ Begin UMaterialExpression Interface
 #if WITH_EDITOR
+	virtual void Build(UE::MIR::FEmitter& Emitter) override;
 	virtual int32 Compile(class FMaterialCompiler* Compiler, int32 OutputIndex) override;
 	virtual void GetCaption(TArray<FString>& OutCaptions) const override;
 	virtual FText GetKeywords() const override {return FText::FromString(TEXT("sin"));}
diff --git a/Engine/Source/Runtime/Engine/Public/Materials/MaterialExpressionTangent.h b/Engine/Source/Runtime/Engine/Public/Materials/MaterialExpressionTangent.h
--- a/Engine/Source/Runtime/Engine/Public/Materials/MaterialExpressionTangent.h
+++ b/Engine/Source/Runtime/Engine/Public/Materials/MaterialExpressionTangent.h
@@ -21,6 +21,7 @@ class UMaterialExpressionTangent : public UMaterialExpression
 
 	//~ Begin UMaterialExpression Interface
 #if WITH_EDITOR
+	virtual void Build(UE::MIR::FEmitter& Emitter) override;
 	virtual int32 Compile(class FMaterialCompiler* Compiler, int32 OutputIndex) override;
 	virtual void GetCaption(TArray<FString>& OutCaptions) const override;
 	virtual FText GetKeywords() const override {return FText::FromString(TEXT("tan"));}
diff --git a/Engine/Source/Runtime/Engine/Public/Materials/MaterialIR.h b/Engine/Source/Runtime/Engine/Public/Materials/MaterialIR.h
--- a/Engine/Source/Runtime/Engine/Public/Materials/MaterialIR.h
+++ b/Engine/Source/Runtime/Engine/Public/Materials/MaterialIR.h
@@ -307,12 +307,24 @@ enum EOperator
 	O_Invalid,
 
 	/* Unary */
-	UO_Negate, 	// Arithmetic negation: negate(5) -> -5
-	UO_Not,		// Logical negation: not(true) -> false
-	UO_Truncate,
+	UO_FirstUnaryOperator,
+
+	UO_Negate = UO_FirstUnaryOperator, 	// Arithmetic negation: negate(5) -> -5
+	UO_Not,							// Logical negation: not(true) -> false
 	
+	/* Unary Intrinsic Operators */
+	UO_ACos,
+	UO_ASin,
+	UO_ATan,
+	UO_Cos,
+	UO_Sin,
+	UO_Tan,
+	UO_Truncate,
+
 	/* Binary Comparison Operators */
-	BO_GreaterThan,
+	BO_FirstBinaryOperator,
+
+	BO_GreaterThan = BO_FirstBinaryOperator,
 	BO_GreaterThanOrEquals,
 	BO_LowerThan,
 	BO_LowerThanOrEquals,
@@ -327,11 +339,15 @@ enum EOperator
 	BO_Or,
 	BO_And,
 	BO_Xor,
-	BO_Min,
+
+	/* Binary Intrinsic Operators */
 	BO_Max,
+	BO_Min,
 
 	/* Ternary Arithmetic Operators */
-	TO_Lerp,
+	TO_FirstTernaryOperator,
+
+	TO_Lerp = TO_FirstTernaryOperator,
 	TO_Smoothstep,
 };
 
