Fix string matching logic in AudioDebug

#rnx

diff --git a/Engine/Source/Runtime/Engine/Public/Audio/AudioDebug.h b/Engine/Source/Runtime/Engine/Public/Audio/AudioDebug.h
--- a/Engine/Source/Runtime/Engine/Public/Audio/AudioDebug.h
+++ b/Engine/Source/Runtime/Engine/Public/Audio/AudioDebug.h
@@ -25,7 +25,7 @@ namespace Audio
 #else
 		extern ENGINE_API FString ExtraTracingSoundWaveFilterCVar;
 		return !ExtraTracingSoundWaveFilterCVar.IsEmpty()	// Trivially reject if we have no filter.
-			&& FCString::Strifind(InName, *ExtraTracingSoundWaveFilterCVar,false);
+		 && FCString::Stristr(InName, *ExtraTracingSoundWaveFilterCVar);
 #endif // UE_BUILD_SHIPPING
 	}
 	
