Add exp, exp2, log, log2, and log10 to the new translator.

- Adds a new code path to validate the input to the operation when performing constant folding.
- This allows us to catch <=0.0 values and error translation if you plug a constant into log, log2, or log10

#rnx

diff --git a/Engine/Source/Runtime/Engine/Private/Materials/MaterialExpressionsIR.cpp b/Engine/Source/Runtime/Engine/Private/Materials/MaterialExpressionsIR.cpp
--- a/Engine/Source/Runtime/Engine/Private/Materials/MaterialExpressionsIR.cpp
+++ b/Engine/Source/Runtime/Engine/Private/Materials/MaterialExpressionsIR.cpp
@@ -467,6 +467,31 @@ static void BuildUnaryArithmeticOperatorFromInput(
 
 // TODO: Add the missing expressions
 
+void UMaterialExpressionExponential::Build(UE::MIR::FEmitter& Emitter)
+{
+	BuildUnaryArithmeticOperatorFromInput(Emitter, MIR::UO_Exponential, &Input);
+}
+
+void UMaterialExpressionExponential2::Build(UE::MIR::FEmitter& Emitter)
+{
+	BuildUnaryArithmeticOperatorFromInput(Emitter, MIR::UO_Exponential2, &Input);
+}
+
+void UMaterialExpressionLogarithm::Build(UE::MIR::FEmitter& Emitter)
+{
+	BuildUnaryArithmeticOperatorFromInput(Emitter, MIR::UO_Logarithm, &Input);
+}
+
+void UMaterialExpressionLogarithm2::Build(UE::MIR::FEmitter& Emitter)
+{
+	BuildUnaryArithmeticOperatorFromInput(Emitter, MIR::UO_Logarithm2, &X);
+}
+
+void UMaterialExpressionLogarithm10::Build(UE::MIR::FEmitter& Emitter)
+{
+	BuildUnaryArithmeticOperatorFromInput(Emitter, MIR::UO_Logarithm10, &X);
+}
+
 void UMaterialExpressionTruncate::Build(UE::MIR::FEmitter& Emitter)
 {
 	BuildUnaryArithmeticOperatorFromInput(Emitter, MIR::UO_Truncate, &Input);
diff --git a/Engine/Source/Runtime/Engine/Private/Materials/MaterialIR.cpp b/Engine/Source/Runtime/Engine/Private/Materials/MaterialIR.cpp
--- a/Engine/Source/Runtime/Engine/Private/Materials/MaterialIR.cpp
+++ b/Engine/Source/Runtime/Engine/Private/Materials/MaterialIR.cpp
@@ -417,6 +417,11 @@ const TCHAR* OperatorToString(EOperator Op)
 		case UO_ASin: return TEXT("Arcsine");
 		case UO_ATan: return TEXT("Arctangent");
 		case UO_Cos: return TEXT("Cosine");
+		case UO_Exponential: return TEXT("Exponential");
+		case UO_Exponential2: return TEXT("Exponential2");
+		case UO_Logarithm: return TEXT("Logarithm");
+		case UO_Logarithm2: return TEXT("Logarithm2");
+		case UO_Logarithm10: return TEXT("Logarithm10");
 		case UO_Sin: return TEXT("Sine");
 		case UO_Tan: return TEXT("Tangent");
 		case UO_Truncate: return TEXT("Truncate");
diff --git a/Engine/Source/Runtime/Engine/Private/Materials/MaterialIREmitter.cpp b/Engine/Source/Runtime/Engine/Private/Materials/MaterialIREmitter.cpp
--- a/Engine/Source/Runtime/Engine/Private/Materials/MaterialIREmitter.cpp
+++ b/Engine/Source/Runtime/Engine/Private/Materials/MaterialIREmitter.cpp
@@ -751,6 +751,26 @@ static bool FoldComparisonOperatorScalar(EOperator Operator, T A, T B)
 	}
 }
 
+template <typename T>
+static bool ValidateFoldScalarArithmeticOperator(FEmitter* Emitter, EOperator Operator, T A, T B, T C)
+{
+	switch (Operator)
+	{
+		case UO_Logarithm:
+		case UO_Logarithm2:
+		case UO_Logarithm10:
+		{
+			if (A <= 0)
+			{
+				Emitter->Errorf(TEXT("Logarithm of 0 or a negative number is not valid."));
+				return false;
+			}
+		}
+	}
+
+	return true;
+}
+
 template <typename T>
 static T FoldFloatScalarArithmeticOperator(EOperator Operator, T A, T B, T C)
 {
@@ -788,6 +808,11 @@ static T FoldScalarArithmeticOperator(EOperator Operator, T A, T B, T C)
 		{
 			return FoldFloatScalarArithmeticOperator<T>(Operator, A, B, C);
 		}
+		case UO_Exponential: return FGenericPlatformMath::Pow(UE_EULERS_NUMBER, A);
+		case UO_Exponential2: return FGenericPlatformMath::Pow(2.0f, A);
+		case UO_Logarithm: return FGenericPlatformMath::LogX(UE_EULERS_NUMBER, A);
+		case UO_Logarithm2: return FGenericPlatformMath::LogX(2.0f, A);
+		case UO_Logarithm10: return FGenericPlatformMath::LogX(10.0f, A);
 		case UO_Truncate:
 		{
 			if constexpr (std::is_same<T, float>::value)
@@ -912,6 +937,11 @@ static FValue* TryFoldOperatorScalar(FEmitter* Emitter, EOperator Op, FValue* A,
 		case UO_Sin:
 		case UO_Tan:
 		case UO_Truncate:
+		case UO_Exponential:
+		case UO_Exponential2:
+		case UO_Logarithm:
+		case UO_Logarithm2:
+		case UO_Logarithm10:
 		case BO_Min:
 		case BO_Max:
 			break;
@@ -974,12 +1004,22 @@ static FValue* TryFoldOperatorScalar(FEmitter* Emitter, EOperator Op, FValue* A,
 		{
 			case ScalarKind_Int:
 			{
+				if (!ValidateFoldScalarArithmeticOperator<TInteger>(Emitter, Op, AConstant->Integer, BConstant ? BConstant->Integer : 0, CConstant ? CConstant->Integer : 0))
+				{
+					return nullptr;
+				}
+
 				TInteger Result = FoldScalarArithmeticOperator<TInteger>(Op, AConstant->Integer, BConstant ? BConstant->Integer : 0, CConstant ? CConstant->Integer : 0);
 				return Emitter->EmitConstantIntScalar(Result);
 			}
 
 			case ScalarKind_Float:
 			{
+				if (!ValidateFoldScalarArithmeticOperator<TFloat>(Emitter, Op, AConstant->Float, BConstant ? BConstant->Float : 0, CConstant ? CConstant->Float : 0))
+				{
+					return nullptr;
+				}
+
 				TFloat Result = FoldScalarArithmeticOperator<TFloat>(Op, AConstant->Float, BConstant ? BConstant->Float : 0, CConstant ? CConstant->Float : 0);
 				return Emitter->EmitConstantFloatScalar(Result);
 			}
diff --git a/Engine/Source/Runtime/Engine/Private/Materials/MaterialIRToHLSLTranslator.cpp b/Engine/Source/Runtime/Engine/Private/Materials/MaterialIRToHLSLTranslator.cpp
--- a/Engine/Source/Runtime/Engine/Private/Materials/MaterialIRToHLSLTranslator.cpp
+++ b/Engine/Source/Runtime/Engine/Private/Materials/MaterialIRToHLSLTranslator.cpp
@@ -466,6 +466,11 @@ struct FTranslator : FMaterialIRToHLSLTranslation
 				case MIR::UO_Sin: OpString = TEXT("sin"); break;
 				case MIR::UO_Tan: OpString = TEXT("tan"); break;
 				case MIR::UO_Truncate: OpString = TEXT("trunc"); break;
+				case MIR::UO_Exponential: OpString = TEXT("exp"); break;
+				case MIR::UO_Exponential2: OpString = TEXT("exp2"); break;
+				case MIR::UO_Logarithm: OpString = TEXT("log"); break;
+				case MIR::UO_Logarithm2: OpString = TEXT("log2"); break;
+				case MIR::UO_Logarithm10: OpString = TEXT("log10"); break;
 				case MIR::BO_Min: OpString = TEXT("min"); break;
 				case MIR::BO_Max: OpString = TEXT("max"); break;
 				case MIR::TO_Lerp: OpString = TEXT("lerp"); break;
diff --git a/Engine/Source/Runtime/Engine/Public/Materials/MaterialExpressionExponential.h b/Engine/Source/Runtime/Engine/Public/Materials/MaterialExpressionExponential.h
--- a/Engine/Source/Runtime/Engine/Public/Materials/MaterialExpressionExponential.h
+++ b/Engine/Source/Runtime/Engine/Public/Materials/MaterialExpressionExponential.h
@@ -16,6 +16,7 @@ class UMaterialExpressionExponential : public UMaterialExpression
 
 	//~ Begin UMaterialExpression Interface
 #if WITH_EDITOR
+	virtual void Build(UE::MIR::FEmitter& Emitter) override;
 	virtual int32 Compile(class FMaterialCompiler* Compiler, int32 OutputIndex) override;
 	virtual void GetCaption(TArray<FString>& OutCaptions) const override;
 	virtual void GetExpressionToolTip(TArray<FString>& OutToolTip) override;
diff --git a/Engine/Source/Runtime/Engine/Public/Materials/MaterialExpressionExponential2.h b/Engine/Source/Runtime/Engine/Public/Materials/MaterialExpressionExponential2.h
--- a/Engine/Source/Runtime/Engine/Public/Materials/MaterialExpressionExponential2.h
+++ b/Engine/Source/Runtime/Engine/Public/Materials/MaterialExpressionExponential2.h
@@ -16,6 +16,7 @@ class UMaterialExpressionExponential2 : public UMaterialExpression
 
 	//~ Begin UMaterialExpression Interface
 #if WITH_EDITOR
+	virtual void Build(UE::MIR::FEmitter& Emitter) override;
 	virtual int32 Compile(class FMaterialCompiler* Compiler, int32 OutputIndex) override;
 	virtual void GetCaption(TArray<FString>& OutCaptions) const override;
 	virtual void GetExpressionToolTip(TArray<FString>& OutToolTip) override;
diff --git a/Engine/Source/Runtime/Engine/Public/Materials/MaterialExpressionLogarithm.h b/Engine/Source/Runtime/Engine/Public/Materials/MaterialExpressionLogarithm.h
--- a/Engine/Source/Runtime/Engine/Public/Materials/MaterialExpressionLogarithm.h
+++ b/Engine/Source/Runtime/Engine/Public/Materials/MaterialExpressionLogarithm.h
@@ -16,6 +16,7 @@ class UMaterialExpressionLogarithm : public UMaterialExpression
 
 	//~ Begin UMaterialExpression Interface
 #if WITH_EDITOR
+	virtual void Build(UE::MIR::FEmitter& Emitter) override;
 	virtual int32 Compile(class FMaterialCompiler* Compiler, int32 OutputIndex) override;
 	virtual void GetCaption(TArray<FString>& OutCaptions) const override;
 	virtual void GetExpressionToolTip(TArray<FString>& OutToolTip) override;
diff --git a/Engine/Source/Runtime/Engine/Public/Materials/MaterialExpressionLogarithm10.h b/Engine/Source/Runtime/Engine/Public/Materials/MaterialExpressionLogarithm10.h
--- a/Engine/Source/Runtime/Engine/Public/Materials/MaterialExpressionLogarithm10.h
+++ b/Engine/Source/Runtime/Engine/Public/Materials/MaterialExpressionLogarithm10.h
@@ -19,6 +19,7 @@ class UMaterialExpressionLogarithm10 : public UMaterialExpression
 
 	//~ Begin UMaterialExpression Interface
 #if WITH_EDITOR
+	virtual void Build(UE::MIR::FEmitter& Emitter) override;
 	virtual int32 Compile(class FMaterialCompiler* Compiler, int32 OutputIndex) override;
 	virtual void GetCaption(TArray<FString>& OutCaptions) const override;
 	virtual void GetExpressionToolTip(TArray<FString>& OutToolTip) override;
diff --git a/Engine/Source/Runtime/Engine/Public/Materials/MaterialExpressionLogarithm2.h b/Engine/Source/Runtime/Engine/Public/Materials/MaterialExpressionLogarithm2.h
--- a/Engine/Source/Runtime/Engine/Public/Materials/MaterialExpressionLogarithm2.h
+++ b/Engine/Source/Runtime/Engine/Public/Materials/MaterialExpressionLogarithm2.h
@@ -19,6 +19,7 @@ class UMaterialExpressionLogarithm2 : public UMaterialExpression
 
 	//~ Begin UMaterialExpression Interface
 #if WITH_EDITOR
+	virtual void Build(UE::MIR::FEmitter& Emitter) override;
 	virtual int32 Compile(class FMaterialCompiler* Compiler, int32 OutputIndex) override;
 	virtual void GetCaption(TArray<FString>& OutCaptions) const override;
 	virtual void GetExpressionToolTip(TArray<FString>& OutToolTip) override;
diff --git a/Engine/Source/Runtime/Engine/Public/Materials/MaterialIR.h b/Engine/Source/Runtime/Engine/Public/Materials/MaterialIR.h
--- a/Engine/Source/Runtime/Engine/Public/Materials/MaterialIR.h
+++ b/Engine/Source/Runtime/Engine/Public/Materials/MaterialIR.h
@@ -317,6 +317,11 @@ enum EOperator
 	UO_ASin,
 	UO_ATan,
 	UO_Cos,
+	UO_Exponential,
+	UO_Exponential2,
+	UO_Logarithm,
+	UO_Logarithm2,
+	UO_Logarithm10,
 	UO_Sin,
 	UO_Tan,
 	UO_Truncate,
