Fix InitWorld to always update the default level collection even if it already exists. This fixes bugs where a world reinit (like from building geometry) would break ticking until the map was reloaded

diff --git a/Engine/Source/Runtime/Engine/Private/World.cpp b/Engine/Source/Runtime/Engine/Private/World.cpp
--- a/Engine/Source/Runtime/Engine/Private/World.cpp
+++ b/Engine/Source/Runtime/Engine/Private/World.cpp
@@ -2397,22 +2397,19 @@ void UWorld::ConditionallyCreateDefaultLevelCollections()
 		LevelCollections.Reserve(1);
 	}
 
-	// Create main level collection. The persistent level will always be considered dynamic.
-	if (!FindCollectionByType(ELevelCollectionType::DynamicSourceLevels))
-	{
-		// Default to the dynamic source collection
-		ActiveLevelCollectionIndex = FindOrAddCollectionByType_Index(ELevelCollectionType::DynamicSourceLevels);
-		LevelCollections[ActiveLevelCollectionIndex].SetPersistentLevel(PersistentLevel);
+	// Always create and update the main level collection. The persistent level will always be considered dynamic.
+	ActiveLevelCollectionIndex = FindOrAddCollectionByType_Index(ELevelCollectionType::DynamicSourceLevels);
+	LevelCollections[ActiveLevelCollectionIndex].SetPersistentLevel(PersistentLevel);
 		
-		// Don't add the persistent level if it is already a member of another collection.
-		// This may be the case if, for example, this world is the outer of a streaming level,
-		// in which case the persistent level may be in one of the collections in the streaming level's OwningWorld.
-		if (PersistentLevel->GetCachedLevelCollection() == nullptr)
-		{
-			LevelCollections[ActiveLevelCollectionIndex].AddLevel(PersistentLevel);
-		}
+	// Don't add the persistent level if it is already a member of another collection.
+	// This may be the case if, for example, this world is the outer of a streaming level,
+	// in which case the persistent level may be in one of the collections in the streaming level's OwningWorld.
+	if (PersistentLevel->GetCachedLevelCollection() == nullptr)
+	{
+		LevelCollections[ActiveLevelCollectionIndex].AddLevel(PersistentLevel);
 	}
 
+	// Optionally create the static level collection, this is not needed for normal operation
 	if (bCreateStaticLevelCollection && !FindCollectionByType(ELevelCollectionType::StaticLevels))
 	{
 		FLevelCollection& StaticCollection = FindOrAddCollectionByType(ELevelCollectionType::StaticLevels);
