Iris
* Expose RepSystemUtil::ForEachReplicationSystem

diff --git a/Engine/Source/Runtime/Engine/Private/Net/Iris/ReplicationSystem/ReplicationSystemUtil.cpp b/Engine/Source/Runtime/Engine/Private/Net/Iris/ReplicationSystem/ReplicationSystemUtil.cpp
--- a/Engine/Source/Runtime/Engine/Private/Net/Iris/ReplicationSystem/ReplicationSystemUtil.cpp
+++ b/Engine/Source/Runtime/Engine/Private/Net/Iris/ReplicationSystem/ReplicationSystemUtil.cpp
@@ -24,9 +24,11 @@
 #include "GameFramework/PlayerController.h"
 #include "Templates/Function.h"
 
-namespace UE::Net::ReplicationSystemUtil
+
+namespace UE::Net
 {
-	static void ForEachReplicationSystem(const UEngine* Engine, const UWorld* World, TFunctionRef<void(UReplicationSystem*)> Function)
+
+void FReplicationSystemUtil::ForEachReplicationSystem(const UEngine* Engine, const UWorld* World, TFunctionRef<void(UReplicationSystem*)> Function)
 	{
 		if (Engine == nullptr || World == nullptr)
 		{
@@ -48,7 +50,7 @@ namespace UE::Net::ReplicationSystemUtil
 		}
 	}
 
-	static void ForEachReplicationSystem(TFunctionRef<void(UReplicationSystem*)> Function)
+void FReplicationSystemUtil::ForEachReplicationSystem(TFunctionRef<void(UReplicationSystem*)> Function)
 	{
 		for (UReplicationSystem* ReplicationSystem : FReplicationSystemFactory::GetAllReplicationSystems())
 		{
@@ -58,10 +60,6 @@ namespace UE::Net::ReplicationSystemUtil
 			}
 		}
 	}
-}
-
-namespace UE::Net
-{
 
 UReplicationSystem* FReplicationSystemUtil::GetReplicationSystem(const AActor* Actor)
 {
@@ -121,7 +119,7 @@ void FReplicationSystemUtil::BeginReplication(AActor* Actor, const FActorReplica
 {
 	if (const UWorld* World = Actor->GetWorld())
 	{
-		ReplicationSystemUtil::ForEachReplicationSystem(GEngine, World, [Actor, &Params](UReplicationSystem* ReplicationSystem)
+		ForEachReplicationSystem(GEngine, World, [Actor, &Params](UReplicationSystem* ReplicationSystem)
 		{
 			if (ReplicationSystem->IsServer())
 			{
@@ -144,7 +142,7 @@ void FReplicationSystemUtil::EndReplication(AActor* Actor, EEndPlayReason::Type
 {
 	// If the call is coming from for example destroying an actor then a formerly associated NetRefHandle will no longer be valid.
 	// The bridge itself will verify that the actor is replicated by it so there's no reason to check that either.
-	ReplicationSystemUtil::ForEachReplicationSystem([Actor, EndPlayReason](UReplicationSystem* ReplicationSystem)
+	ForEachReplicationSystem([Actor, EndPlayReason](UReplicationSystem* ReplicationSystem)
 	{
 		if (UEngineReplicationBridge* Bridge = Cast<UEngineReplicationBridge>(ReplicationSystem->GetReplicationBridge()))
 		{
@@ -172,7 +170,7 @@ void FReplicationSystemUtil::BeginReplicationForActorComponent(FNetHandle ActorH
 	{
 		if (const UWorld* World = Actor->GetWorld())
 		{
-			ReplicationSystemUtil::ForEachReplicationSystem(GEngine, World, [ActorHandle, ActorComp](UReplicationSystem* ReplicationSystem)
+			ForEachReplicationSystem(GEngine, World, [ActorHandle, ActorComp](UReplicationSystem* ReplicationSystem)
 			{
 				if (ReplicationSystem->IsServer())
 				{
@@ -201,7 +199,7 @@ void FReplicationSystemUtil::BeginReplicationForActorComponent(const AActor* Act
 
 	if (const UWorld* World = Actor->GetWorld())
 	{
-		ReplicationSystemUtil::ForEachReplicationSystem(GEngine, World, [ActorHandle, ActorComp](UReplicationSystem* ReplicationSystem)
+		ForEachReplicationSystem(GEngine, World, [ActorHandle, ActorComp](UReplicationSystem* ReplicationSystem)
 		{
 			if (ReplicationSystem->IsServer())
 			{
@@ -234,7 +232,7 @@ void FReplicationSystemUtil::BeginReplicationForActorSubObject(const AActor* Act
 
 	if (const UWorld* World = Actor->GetWorld())
 	{
-		ReplicationSystemUtil::ForEachReplicationSystem(GEngine, World, [ActorHandle, ActorSubObject, NetCondition](UReplicationSystem* ReplicationSystem)
+		ForEachReplicationSystem(GEngine, World, [ActorHandle, ActorSubObject, NetCondition](UReplicationSystem* ReplicationSystem)
 		{
 			if (ReplicationSystem->IsServer())
 			{
@@ -269,7 +267,7 @@ void FReplicationSystemUtil::BeginReplicationForActorComponentSubObject(UActorCo
 
 		if (const UWorld* World = Actor->GetWorld())
 		{
-			ReplicationSystemUtil::ForEachReplicationSystem(GEngine, World, [ActorHandle, ActorComponent, SubObject, NetCondition](UReplicationSystem* ReplicationSystem)
+			ForEachReplicationSystem(GEngine, World, [ActorHandle, ActorComponent, SubObject, NetCondition](UReplicationSystem* ReplicationSystem)
 			{
 				if (ReplicationSystem->IsServer())
 				{
@@ -300,7 +298,7 @@ void FReplicationSystemUtil::BeginReplicationForActorComponentSubObject(UActorCo
 
 void FReplicationSystemUtil::EndReplicationForActorComponent(UActorComponent* ActorComponent)
 {
-	ReplicationSystemUtil::ForEachReplicationSystem([ActorComponent](UReplicationSystem* ReplicationSystem)
+	ForEachReplicationSystem([ActorComponent](UReplicationSystem* ReplicationSystem)
 	{
 		if (UEngineReplicationBridge* Bridge = Cast<UEngineReplicationBridge>(ReplicationSystem->GetReplicationBridge()))
 		{
@@ -312,7 +310,7 @@ void FReplicationSystemUtil::EndReplicationForActorComponent(UActorComponent* Ac
 
 void FReplicationSystemUtil::EndReplicationForActorSubObject(const AActor* Actor, UObject* SubObject)
 {
-	ReplicationSystemUtil::ForEachReplicationSystem([SubObject](UReplicationSystem* ReplicationSystem)
+	ForEachReplicationSystem([SubObject](UReplicationSystem* ReplicationSystem)
 	{
 		if (UEngineReplicationBridge* Bridge = Cast<UEngineReplicationBridge>(ReplicationSystem->GetReplicationBridge()))
 		{				
@@ -344,7 +342,7 @@ void FReplicationSystemUtil::AddDependentActor(const AActor* Parent, AActor* Chi
 		return;
 	}
 
-	ReplicationSystemUtil::ForEachReplicationSystem(GEngine, World, [ParentHandle, Child, SchedulingHint](UReplicationSystem* ReplicationSystem)
+	ForEachReplicationSystem(GEngine, World, [ParentHandle, Child, SchedulingHint](UReplicationSystem* ReplicationSystem)
 	{
 		if (ReplicationSystem->IsServer())
 		{
@@ -390,7 +388,7 @@ void FReplicationSystemUtil::RemoveDependentActor(const AActor* Parent, AActor*
 		return;
 	}
 
-	ReplicationSystemUtil::ForEachReplicationSystem(GEngine, World, [ParentHandle, ChildHandle](UReplicationSystem* ReplicationSystem)
+	ForEachReplicationSystem(GEngine, World, [ParentHandle, ChildHandle](UReplicationSystem* ReplicationSystem)
 	{
 		if (ReplicationSystem->IsServer())
 		{
@@ -423,7 +421,7 @@ void FReplicationSystemUtil::SetActorComponentNetCondition(const UActorComponent
 
 	if (const UWorld* World = Actor->GetWorld())
 	{
-		ReplicationSystemUtil::ForEachReplicationSystem(GEngine, World, [ActorComponent, NetCondition](UReplicationSystem* ReplicationSystem)
+		ForEachReplicationSystem(GEngine, World, [ActorComponent, NetCondition](UReplicationSystem* ReplicationSystem)
 		{
 			if (ReplicationSystem->IsServer())
 			{
@@ -528,7 +526,7 @@ void FReplicationSystemUtil::UpdateSubObjectGroupMemberships(const UObject* SubO
 	if (const UNetworkSubsystem* NetSubsystem = World ? World->GetSubsystem<UNetworkSubsystem>() : nullptr)
 	{
 		FObjectKey SubObjectKey(SubObject);
-		ReplicationSystemUtil::ForEachReplicationSystem(GEngine, World, [NetSubsystem, SubObject, SubObjectKey](UReplicationSystem* ReplicationSystem)
+		ForEachReplicationSystem(GEngine, World, [NetSubsystem, SubObject, SubObjectKey](UReplicationSystem* ReplicationSystem)
 		{
 			if (ReplicationSystem->IsServer())
 			{
@@ -588,7 +586,7 @@ void FReplicationSystemUtil::UpdateSubObjectGroupMemberships(const APlayerContro
 
 void FReplicationSystemUtil::SetReplicationCondition(FNetHandle NetHandle, EReplicationCondition Condition, bool bEnableCondition)
 {
-	ReplicationSystemUtil::ForEachReplicationSystem([NetHandle, Condition, bEnableCondition](UReplicationSystem* ReplicationSystem)
+	ForEachReplicationSystem([NetHandle, Condition, bEnableCondition](UReplicationSystem* ReplicationSystem)
 	{
 		if (ReplicationSystem->IsServer())
 		{
@@ -612,7 +610,7 @@ void FReplicationSystemUtil::SetStaticPriority(const AActor* Actor, float Priori
 		return;
 	}
 	
-	ReplicationSystemUtil::ForEachReplicationSystem([ActorHandle, Priority](UReplicationSystem* ReplicationSystem)
+	ForEachReplicationSystem([ActorHandle, Priority](UReplicationSystem* ReplicationSystem)
 	{
 		if (ReplicationSystem->IsServer())
 		{
@@ -633,7 +631,7 @@ void FReplicationSystemUtil::SetCullDistanceOverride(const AActor* Actor, float
 		return;
 	}
 	
-	ReplicationSystemUtil::ForEachReplicationSystem([ActorHandle, CullDist](UReplicationSystem* ReplicationSystem)
+	ForEachReplicationSystem([ActorHandle, CullDist](UReplicationSystem* ReplicationSystem)
 	{
 		if (ReplicationSystem->IsServer())
 		{
@@ -657,7 +655,7 @@ void FReplicationSystemUtil::ClearCullDistanceOverride(const AActor* Actor)
 		return;
 	}
 	
-	ReplicationSystemUtil::ForEachReplicationSystem([ActorHandle](UReplicationSystem* ReplicationSystem)
+	ForEachReplicationSystem([ActorHandle](UReplicationSystem* ReplicationSystem)
 	{
 		if (ReplicationSystem->IsServer())
 		{
@@ -681,7 +679,7 @@ void FReplicationSystemUtil::SetPollFrequency(const UObject* Object, float PollF
 		return;
 	}
 	
-	ReplicationSystemUtil::ForEachReplicationSystem([NetHandle, PollFrequency](UReplicationSystem* ReplicationSystem)
+	ForEachReplicationSystem([NetHandle, PollFrequency](UReplicationSystem* ReplicationSystem)
 	{
 		if (ReplicationSystem->IsServer())
 		{
diff --git a/Engine/Source/Runtime/Engine/Public/Net/Iris/ReplicationSystem/ReplicationSystemUtil.h b/Engine/Source/Runtime/Engine/Public/Net/Iris/ReplicationSystem/ReplicationSystemUtil.h
--- a/Engine/Source/Runtime/Engine/Public/Net/Iris/ReplicationSystem/ReplicationSystemUtil.h
+++ b/Engine/Source/Runtime/Engine/Public/Net/Iris/ReplicationSystem/ReplicationSystemUtil.h
@@ -12,14 +12,16 @@
 
 class AActor;
 class UActorComponent;
-struct FActorReplicationParams;
+class UEngine;
 class UEngineReplicationBridge;
 class UObject;
-class FRepChangedPropertyTracker;
 class UReplicationSystem;
 class UWorld;
 class UNetConnection;
 class UNetDriver;
+class FRepChangedPropertyTracker;
+
+struct FActorReplicationParams;
 
 namespace UE::Net
 {
@@ -32,6 +34,12 @@ namespace UE::Net
 // Helper methods to interact with ReplicationSystem from engine code
 struct FReplicationSystemUtil
 {
+	/** Call a function on all ReplicationSystems of a specific world. */
+	ENGINE_API static void ForEachReplicationSystem(const UEngine* Engine, const UWorld* World, TFunctionRef<void(UReplicationSystem*)> Function);
+
+	/** Call a function on all existing ReplicationSystems. This could include systems from other clients (e.g. PIE worlds) */
+	ENGINE_API static void ForEachReplicationSystem(TFunctionRef<void(UReplicationSystem*)> Function);
+
 	/** Returns the UReplicationSystem for the main NetDriver assigned to the Actor. Note that an Actor may be replicated by multiple ReplicationSystems. May return null. */
 	ENGINE_API static UReplicationSystem* GetReplicationSystem(const AActor* Actor);
 
