- Re-added some code that got dropped in the refactor with -DPApplyMode

diff --git a/Engine/Source/Runtime/Engine/Private/DeviceProfiles/DeviceProfileManager.cpp b/Engine/Source/Runtime/Engine/Private/DeviceProfiles/DeviceProfileManager.cpp
--- a/Engine/Source/Runtime/Engine/Private/DeviceProfiles/DeviceProfileManager.cpp
+++ b/Engine/Source/Runtime/Engine/Private/DeviceProfiles/DeviceProfileManager.cpp
@@ -217,10 +217,10 @@ void UpdateCVarsAfterDynamicConfig(const FConfigModificationTracker* ChangeTrack
 		}
 	}
 
-	if (CVarSetMode() > 0)
+	const TSet<FString>* DeviceProfileSections = ChangeTracker->ModifiedSectionsPerBranch.Find(*GDeviceProfilesIni);
+	if (DeviceProfileSections != nullptr)
 	{
-		const TSet<FString>* DeviceProfileSections = ChangeTracker->ModifiedSectionsPerBranch.Find(*GDeviceProfilesIni);
-		if (DeviceProfileSections != nullptr)
+		if (CVarSetMode() > 0)
 		{
 			UDeviceProfileManager& DPManager = UDeviceProfileManager::Get();
 			IConsoleManager& ConsoleManager = IConsoleManager::Get();
@@ -247,6 +247,20 @@ void UpdateCVarsAfterDynamicConfig(const FConfigModificationTracker* ChangeTrack
 				}
 			}
 		}
+		else
+		{
+			TSet<FString> DeviceProfilesToQuery;
+			int32 DeviceProfilePatternLength = UDeviceProfile::StaticClass()->GetName().Len() + 1;
+			for (const FString& Section : *DeviceProfileSections)
+			{
+				DeviceProfilesToQuery.Add(Section.LeftChop(DeviceProfilePatternLength));
+			}
+
+			if (UDeviceProfileManager::Get().DoActiveProfilesReference(DeviceProfilesToQuery))
+			{
+				UDeviceProfileManager::Get().ReapplyDeviceProfile();
+			}
+		}
 	}
 
 	// if we modify a scalability group section (ie EffectsQuality@2) that matches current value, re-set the group to apply the modified group
