Fix broken device profile fragment processing with dpapplymode=1

#rnx

diff --git a/Engine/Source/Runtime/Engine/Private/DeviceProfiles/DeviceProfileManager.cpp b/Engine/Source/Runtime/Engine/Private/DeviceProfiles/DeviceProfileManager.cpp
--- a/Engine/Source/Runtime/Engine/Private/DeviceProfiles/DeviceProfileManager.cpp
+++ b/Engine/Source/Runtime/Engine/Private/DeviceProfiles/DeviceProfileManager.cpp
@@ -230,19 +230,23 @@ void UpdateCVarsAfterDynamicConfig(const FConfigModificationTracker* ChangeTrack
 
 			for (const auto& DP : DPChain)
 			{
-				// mode 1 we wipe everything then start over instead of updating in place
-				if (CVarSetMode() == 1)
-				{
-					ConsoleManager.UnsetAllConsoleVariablesWithTag(DP.Get<0>());
-				}
-
 				// if we updated this DP, then refresh it
 				const FName& DPName = DP.Get<0>();
 				const FString& DPSectionName = DP.Get<1>();
 				if (DeviceProfileSections->Contains(DPSectionName))
 				{
+					// mode 1 we wipe everything then start over instead of updating in place
+					if (CVarSetMode() == 1)
+					{
+						ConsoleManager.UnsetAllConsoleVariablesWithTag(DP.Get<0>());
+					}
+
+					const bool bIsFragment = DPSectionName.EndsWith(UDeviceProfileFragment::StaticClass()->GetName());
+					const UDeviceProfileManager::EDeviceProfileMode DeviceProfileMode =
+						bIsFragment ? UDeviceProfileManager::EDeviceProfileMode::DPM_SetCVarsForSingleFragment : UDeviceProfileManager::EDeviceProfileMode::DPM_SetCVarsForSingleDP;
+
 					// get just this DP's cvars, and then update in place
-					TMap<FName, FString> DeviceProfileCVars = DPManager.GatherDeviceProfileCVars(DPName.ToString(), UDeviceProfileManager::EDeviceProfileMode::DPM_SetCVarsForSingleDP);
+					TMap<FName, FString> DeviceProfileCVars = DPManager.GatherDeviceProfileCVars(DPName.ToString(), DeviceProfileMode);
 					ConsoleManager.BatchUpdateTag(DPName, DeviceProfileCVars);
 				}
 			}
