Ensure that static permutation resources of a Material Instance are refreshed when bForceStaticPermutationUpdate is set to true to ensure redundant shader maps do not cause a crash when a parent is swapped or removed.

diff --git a/Engine/Source/Runtime/Engine/Private/Materials/MaterialInstance.cpp b/Engine/Source/Runtime/Engine/Private/Materials/MaterialInstance.cpp
--- a/Engine/Source/Runtime/Engine/Private/Materials/MaterialInstance.cpp
+++ b/Engine/Source/Runtime/Engine/Private/Materials/MaterialInstance.cpp
@@ -4302,7 +4302,7 @@ void UMaterialInstance::UpdateStaticPermutation(const FStaticParameterSet& NewPa
 		// This will flush the rendering thread which is necessary before changing bHasStaticPermutationResource, since the RT is reading from that directly
 		FlushRenderingCommands();
 
-		bHasStaticPermutationResource = bWantsStaticPermutationResource;
+		bHasStaticPermutationResource = bWantsStaticPermutationResource || bForceStaticPermutationUpdate;
 		StaticParametersRuntime = CompareParameters.GetRuntime();
 		EditorOnly->StaticParameters = CompareParameters.EditorOnly;
 
