AActor::IncrementalRegisterComponents now handles the time limit internally (if any) instead of exiting after processing s.LevelStreamingComponentsRegistrationGranularity components.

diff --git a/Engine/Source/Runtime/Engine/Private/Actor.cpp b/Engine/Source/Runtime/Engine/Private/Actor.cpp
--- a/Engine/Source/Runtime/Engine/Private/Actor.cpp
+++ b/Engine/Source/Runtime/Engine/Private/Actor.cpp
@@ -5911,9 +5911,24 @@ bool AActor::IncrementalRegisterComponents(int32 NumComponentsToRegister, FRegis
 	int32 NumRegisteredComponentsThisRun = 0;
 	TInlineComponentArray<UActorComponent*> Components;
 	GetComponents(Components);
+	const int32 ComponentCount = Components.Num();
 	TSet<UActorComponent*> RegisteredParents;
-	
-	for (int32 CompIdx = 0; CompIdx < Components.Num() && NumRegisteredComponentsThisRun < NumComponentsToRegister; CompIdx++)
+
+	const double StartTime = FPlatformTime::Seconds();
+	const double TimeLimitMs = Context ? Context->GetIncrementalUpdateTimeLimitMs() : 0;
+	const bool bConsiderTimeLimit = Context && TimeLimitMs >= 0;
+
+	ON_SCOPE_EXIT
+	{
+		if (bConsiderTimeLimit)
+		{
+			const double ElapsedTimeMs = (FPlatformTime::Seconds() - StartTime) * 1000;
+			const double IncrementalUpdateTimeLimitMs = bConsiderTimeLimit ? FMath::Max<double>(0.0, TimeLimitMs - ElapsedTimeMs) : -1;
+			Context->SetIncrementalUpdateTimeLimitMs(IncrementalUpdateTimeLimitMs);
+		}
+	};
+
+	for (int32 CompIdx = 0; CompIdx < ComponentCount; CompIdx++)
 	{
 		UActorComponent* Component = Components[CompIdx];
 		if (!Component->IsRegistered() && Component->bAutoRegister && IsValidChecked(Component))
@@ -5948,10 +5963,23 @@ bool AActor::IncrementalRegisterComponents(int32 NumComponentsToRegister, FRegis
 		}
 
 		NumTotalRegisteredComponents++;
+
+		if ((CompIdx + 1 < ComponentCount) && (NumRegisteredComponentsThisRun >= NumComponentsToRegister))
+		{
+			NumRegisteredComponentsThisRun = 0;
+			if (!bConsiderTimeLimit)
+			{
+				break;
+			}
+			else if (const double ElapseTimeMs = (FPlatformTime::Seconds() - StartTime) * 1000; ElapseTimeMs > TimeLimitMs)
+			{
+				break;
+			}
+		}
 	}
 
 	// See whether we are done
-	if (Components.Num() == NumTotalRegisteredComponents)
+	if (ComponentCount == NumTotalRegisteredComponents)
 	{
 #if PERF_TRACK_DETAILED_ASYNC_STATS
 		QUICK_SCOPE_CYCLE_COUNTER(STAT_AActor_IncrementalRegisterComponents_PostRegisterAllComponents);
