Change new material translator error reporting for missing input to match that of the old translator

- It now tells you what node failed.
- Like the previous code it grabs the material expression name and chops off the "MaterialExpression" part.

#rnx

diff --git a/Engine/Source/Runtime/Engine/Private/Materials/MaterialIREmitter.cpp b/Engine/Source/Runtime/Engine/Private/Materials/MaterialIREmitter.cpp
--- a/Engine/Source/Runtime/Engine/Private/Materials/MaterialIREmitter.cpp
+++ b/Engine/Source/Runtime/Engine/Private/Materials/MaterialIREmitter.cpp
@@ -207,7 +207,11 @@ FValue* FEmitter::Input(const FExpressionInput* InInput)
 	FValue* Value = TryInput(InInput);
 	if (!Value)
 	{
-		Errorf(TEXT("Input '%s' has no input value."), *SlowFindInputName(Expression, InInput).ToString());
+		// Add the node type to the error message
+		const int32 ChopCount = FCString::Strlen(TEXT("MaterialExpression"));
+		const FString ErrorClassName = Expression->GetClass()->GetName();
+
+		Errorf(TEXT("(Node %s) Missing '%s' input value."), *ErrorClassName.Right(ErrorClassName.Len() - ChopCount), *SlowFindInputName(Expression, InInput).ToString());
 		return MIR::FPoison::Get();
 	}
 	return Value;
