[Shaders] Fix crash when printing diagnostics of hung distributed shader jobs.

#rnx

diff --git a/Engine/Source/Runtime/Engine/Private/ShaderCompiler/ShaderCompilerDistributed.cpp b/Engine/Source/Runtime/Engine/Private/ShaderCompiler/ShaderCompilerDistributed.cpp
--- a/Engine/Source/Runtime/Engine/Private/ShaderCompiler/ShaderCompilerDistributed.cpp
+++ b/Engine/Source/Runtime/Engine/Private/ShaderCompiler/ShaderCompilerDistributed.cpp
@@ -293,9 +293,11 @@ static void LogShaderCompileWorkerDistributedDiagnostics(const TSparseArray<FDis
 
 	FString JobDiagnostics;
 
-	for (int32 TaskIndex = 0; TaskIndex < InDispatchedTasks.Num(); ++TaskIndex)
+	int32 TaskIndex = 0;
+	for (auto Iter = InDispatchedTasks.CreateConstIterator(); Iter; ++Iter)
 	{
-		FDistributedShaderCompilerTask* Task = InDispatchedTasks[TaskIndex];
+		FDistributedShaderCompilerTask* Task = *Iter;
+		checkf(Task != nullptr, TEXT("Task entries in the sparse array of dispatched distributed shader jobs (FDistributedShaderCompilerTask) must not be null"));
 		if (!Task->ShaderJobs.IsEmpty())
 		{
 			JobDiagnostics.Empty();
@@ -310,6 +312,7 @@ static void LogShaderCompileWorkerDistributedDiagnostics(const TSparseArray<FDis
 				TaskIndex + 1, InDispatchedTasks.Num(), *JobDiagnostics
 			);
 		}
+		++TaskIndex;
 	}
 }
 
