Remap the NetConnection pointer in ADstmPlayerController when migrating between servers.

diff --git a/Engine/Source/Runtime/Engine/Classes/Engine/NetDriver.h b/Engine/Source/Runtime/Engine/Classes/Engine/NetDriver.h
--- a/Engine/Source/Runtime/Engine/Classes/Engine/NetDriver.h
+++ b/Engine/Source/Runtime/Engine/Classes/Engine/NetDriver.h
@@ -1986,6 +1986,15 @@ public:
 	/** Adds (fully initialized, ready to go) client connection to the ClientConnections list + any other game related setup */
 	ENGINE_API void	AddClientConnection(UNetConnection * NewConnection);
 
+	/** 
+	 * Enumerate through each client connection (including any child connections).
+	 * 
+	 * Func is a function that is called with the current instance of UNetConnection in the enumeration, and
+	 * should return true if the enumeration must stop (i.e. no more calls to Func) or false for the enumeration
+	 * to continue.
+	 */
+	ENGINE_API void ForEachClientConnection(TFunction<bool(TObjectPtr<UNetConnection>)> Func);
+
 	//~ This method should only be called by internal networking systems.
 	ENGINE_API void NotifyActorFullyDormantForConnection(AActor* Actor, UNetConnection* Connection);
 
diff --git a/Engine/Source/Runtime/Engine/Private/NetDriver.cpp b/Engine/Source/Runtime/Engine/Private/NetDriver.cpp
--- a/Engine/Source/Runtime/Engine/Private/NetDriver.cpp
+++ b/Engine/Source/Runtime/Engine/Private/NetDriver.cpp
@@ -6626,6 +6626,25 @@ void UNetDriver::AddClientConnection(UNetConnection* NewConnection)
 	UpdateCrashContext();
 }
 
+void UNetDriver::ForEachClientConnection(TFunction<bool(TObjectPtr<UNetConnection>)> Func)
+{
+	for (TObjectPtr<UNetConnection> ParentConnection : ClientConnections)
+	{
+		if (Func(ParentConnection))
+		{
+			return;
+		}
+		
+		for (TObjectPtr<UNetConnection> ChildConnection : ParentConnection->Children)
+		{
+			if (Func(ChildConnection))
+			{
+				return;
+			}
+		}
+	}
+}
+
 void UNetDriver::UpdateCrashContext(ECrashContextUpdate UpdateType)
 {
 	if (NetDriverDefinition != NAME_GameNetDriver)
