[AutoRTFM] Fix unnecessary Open block in FTraceFilter::SetObjectIsTraceable.

This was added at CL 34637107, which predates a lot of our RTFM building blocks like the FTransactionallySafeCriticalSection. The locks in this code had already been updated to use a safe critical section in a separate CL that landed shortly after the transactionally-safe critical section was developed.

When reviewing the effects of this change, I found that the low level object trace logic would fail in the Closed since it takes a RWLock to manage a set of "already traced" object IDs. This set is only used to manage logging/tracing and not game behavior, so it doesn't seem useful to handle this in a transactional way. Instead, I've moved the object ID tracking logic and locking to a dedicated function which runs in the open. This set is never touched in the closed so this should be safe.

diff --git a/Engine/Source/Runtime/Engine/Private/ObjectTrace.cpp b/Engine/Source/Runtime/Engine/Private/ObjectTrace.cpp
--- a/Engine/Source/Runtime/Engine/Private/ObjectTrace.cpp
+++ b/Engine/Source/Runtime/Engine/Private/ObjectTrace.cpp
@@ -117,6 +117,9 @@ int32 GetTypeHash(const FObjectIdAnnotation& Annotation)
 
 // Object annotations used for tracing
 // FUObjectAnnotationSparse<FTracedObjectAnnotation, true> GObjectTracedAnnotations;
+
+// GObjectTracedSet must only be accessed in the open, and must be protected by the GObjectTracedSetLock.
+// Changing GObjectTracedSet in the closed is likely to cause a validation error.
 FRWLock GObjectTracedSetLock;
 TSet<uint64> GObjectTracedSet;
 FUObjectAnnotationSparseSearchable<FObjectIdAnnotation, true> GObjectIdAnnotations;
@@ -124,7 +127,7 @@ FUObjectAnnotationSparseSearchable<FObjectIdAnnotation, true> GObjectIdAnnotatio
 // Handle used to hook to world tick
 static FDelegateHandle WorldTickStartHandle;
 
-void TickObjectTraceWorldSubsystem(UWorld* InWorld, ELevelTick InTickType, float InDeltaSeconds)
+static void TickObjectTraceWorldSubsystem(UWorld* InWorld, ELevelTick InTickType, float InDeltaSeconds)
 {
 	if(InTickType == LEVELTICK_All)
 	{
@@ -146,6 +149,26 @@ void TickObjectTraceWorldSubsystem(UWorld* InWorld, ELevelTick InTickType, float
 	}
 }
 
+// Returns true if the object ID was added to the traced set.
+// Returns false if the object ID already existed in the set.
+// If the object ID was already in the set, we don't need to log it again.
+UE_AUTORTFM_ALWAYS_OPEN static bool AddObjectToTracedSet(uint64 ObjectId)
+{
+	{
+		FRWScopeLock ScopeLock(GObjectTracedSetLock, SLT_ReadOnly);
+		if (GObjectTracedSet.Contains(ObjectId))
+		{
+			// We've already traced this object ID, so it can be skipped.
+			return false;
+		}
+	}
+
+	// We haven't traced this object yet. Add its ID to the set.
+	FRWScopeLock ScopeLock(GObjectTracedSetLock, SLT_Write);
+	GObjectTracedSet.Add(ObjectId);
+	return true;
+}
+
 void FObjectTrace::Init()
 {
 	WorldTickStartHandle = FWorldDelegates::OnWorldTickStart.AddStatic(&TickObjectTraceWorldSubsystem);
@@ -160,10 +183,11 @@ void FObjectTrace::Reset()
 {
 	GObjectIdAnnotations.RemoveAllAnnotations();
 	
+	UE_AUTORTFM_OPEN
 	{
-		FRWScopeLock ScopeLock(GObjectTracedSetLock,SLT_Write);
+		FRWScopeLock ScopeLock(GObjectTracedSetLock, SLT_Write);
 		GObjectTracedSet.Empty();
-	}
+	};
 }
 
 uint64 FObjectTrace::GetObjectId(const UObject* InObject)
@@ -304,18 +328,9 @@ void FObjectTrace::OutputClass(const UClass* InClass)
 		return;
 	}
 	uint64 ObjectId = GetObjectId(InClass);
-
-	{
-		FRWScopeLock ScopeLock(GObjectTracedSetLock, SLT_ReadOnly);
-		if (GObjectTracedSet.Contains(ObjectId))
-		{
-			// Already traced, so skip
-			return;
-		}
-	}
+	if (!AddObjectToTracedSet(ObjectId))
 	{
-		FRWScopeLock ScopeLock(GObjectTracedSetLock, SLT_Write);
-		GObjectTracedSet.Add(ObjectId);
+		return;
 	}
 
 	OutputClass(InClass->GetSuperClass());
@@ -380,19 +395,11 @@ void FObjectTrace::OutputObject(const UObject* InObject)
 	{
 		return;
 	}
-	uint64 ObjectId = GetObjectId(InObject);
 	
+	uint64 ObjectId = GetObjectId(InObject);
+	if (!AddObjectToTracedSet(ObjectId))
 	{
-		FRWScopeLock ScopeLock(GObjectTracedSetLock, SLT_ReadOnly);
-		if (GObjectTracedSet.Contains(ObjectId))
-		{
-			// Already traced, so skip
-			return;
-		}
-	}
-	{
-		FRWScopeLock ScopeLock(GObjectTracedSetLock, SLT_Write);
-		GObjectTracedSet.Add(ObjectId);
+		return;
 	}
 
 	OutputObject(InObject->GetOuter());
@@ -520,9 +527,8 @@ void FObjectTrace::OutputWorld(const UWorld* InWorld)
 		<< World.NetMode((uint8)InWorld->GetNetMode())
 		<< World.IsSimulating(bIsSimulating);
 
-	// Trace object AFTER world info so we dont risk world info not being present in the trace
+	// Trace object AFTER world info so we don't risk world info not being present in the trace
 	TRACE_OBJECT(InWorld);
 }
 
 #endif
-
diff --git a/Engine/Source/Runtime/Engine/Private/TraceFilter.cpp b/Engine/Source/Runtime/Engine/Private/TraceFilter.cpp
--- a/Engine/Source/Runtime/Engine/Private/TraceFilter.cpp
+++ b/Engine/Source/Runtime/Engine/Private/TraceFilter.cpp
@@ -33,7 +33,6 @@ DEFINE_LOG_CATEGORY_STATIC(TraceFiltering, Display, Display);
 class FTraceUObjectAnnotation : public FUObjectArray::FUObjectDeleteListener
 {
 public:
-
 	virtual void NotifyUObjectDeleted(const UObjectBase* Object, int32 Index) override
 	{
 		RemoveAnnotation(Object);
@@ -45,7 +44,7 @@ public:
 		GUObjectArray.RemoveUObjectDeleteListener(this);
 	}
 
-	FTraceUObjectAnnotation() : AnnotationCacheKey(nullptr)
+	FTraceUObjectAnnotation()
 	{
 		// default constructor is required to be default annotation
 		check(AnnotationCacheValue.IsDefault());
@@ -110,7 +109,7 @@ public:
 	{
 		UE::TScopeLock ScopeLock(AnnotationMapCritical);
 
-		AnnotationCacheKey = NULL;
+		AnnotationCacheKey = nullptr;
 		AnnotationCacheValue = FTraceFilterObjectAnnotation();
 		const bool bHadElements = (AnnotationMap.Num() > 0);
 		AnnotationMap.Empty();
@@ -174,7 +173,7 @@ private:
 
 	TOptional<UE::TScopeLock<FTransactionallySafeCriticalSection>> AnnotationMapScopeLock;
 
-	const UObjectBase* AnnotationCacheKey;
+	const UObjectBase* AnnotationCacheKey = nullptr;
 	FTraceFilterObjectAnnotation AnnotationCacheValue;
 };
 
@@ -279,16 +278,13 @@ FAutoConsoleCommand FlushFilterStateCommand(TEXT("TraceFilter.FlushState"), TEXT
 template<bool bForceThreadSafe>
 bool FTraceFilter::IsObjectTraceable(const UObject* InObject)
 {
-	return AutoRTFM::Open([&]
+	if constexpr (!bForceThreadSafe)
 	{
-		if constexpr (!bForceThreadSafe)
-		{
-			check(GObjectFilterAnnotations.IsLocked());
-		}
+		check(GObjectFilterAnnotations.IsLocked());
+	}
 
-		// Object not found in the AnnotationMap means that it is at the default value, which is bIsTraceable == true
-		return GObjectFilterAnnotations.GetAnnotationMap().Find(InObject) == nullptr;
-	});
+	// Object not found in the AnnotationMap means that it is at the default value, which is bIsTraceable == true
+	return GObjectFilterAnnotations.GetAnnotationMap().Find(InObject) == nullptr;
 }
 
 template bool FTraceFilter::IsObjectTraceable</*bForceThreadSafe = */ false>(const UObject* InObject);
@@ -297,41 +293,33 @@ template bool FTraceFilter::IsObjectTraceable</*bForceThreadSafe = */ true>(cons
 template<bool bForceThreadSafe>
 void FTraceFilter::SetObjectIsTraceable(const UObject* InObject, bool bIsTraceable)
 {
-	AutoRTFM::Open([&]
+	ensure(InObject);
+
+	if constexpr (bForceThreadSafe)
 	{
-		ensure(InObject);
+		FTraceFilterObjectAnnotation Annotation;
+		Annotation.bIsTraceable = bIsTraceable;
+		GObjectFilterAnnotations.AddAnnotation(InObject, Annotation);
 
-		if constexpr (bForceThreadSafe)
+		if (bIsTraceable)
 		{
-			FTraceFilterObjectAnnotation Annotation;
-			Annotation.bIsTraceable = bIsTraceable;
-			GObjectFilterAnnotations.AddAnnotation(InObject, Annotation);
-
-			if (bIsTraceable)
-			{
-				TRACE_OBJECT(InObject);
-			}
+			TRACE_OBJECT(InObject);
+		}
+	}
+	else
+	{
+		check(GObjectFilterAnnotations.IsLocked());
+		TMap<const UObjectBase*, FTraceFilterObjectAnnotation>& AnnotationMap = GObjectFilterAnnotations.GetAnnotationMap();
+		if (!bIsTraceable)
+		{
+			AnnotationMap.FindOrAdd(InObject).bIsTraceable = false;
 		}
 		else
 		{
-			check(GObjectFilterAnnotations.IsLocked());
-			TMap<const UObjectBase*, FTraceFilterObjectAnnotation>& AnnotationMap = GObjectFilterAnnotations.GetAnnotationMap();
-			if (!bIsTraceable)
-			{
-				AnnotationMap.FindOrAdd(InObject).bIsTraceable = false;
-			}
-			else
-			{
-				AnnotationMap.Remove(InObject);
-				TRACE_OBJECT(InObject);
-			}
+			AnnotationMap.Remove(InObject);
+			TRACE_OBJECT(InObject);
 		}
-	});
-
-	AutoRTFM::OnAbort([InObject, bIsTraceable]
-	{
-		SetObjectIsTraceable(InObject, !bIsTraceable);
-	});
+	}
 }
 
 template void FTraceFilter::SetObjectIsTraceable</*bForceThreadSafe = */ true>(const UObject* InObject, bool bIsTraceable);
