Persona: Add the ability to keep absolute location of socket when reparenting.

[RN] Users can now hold down the Shift key when moving sockets in the skeletal editor to keep them in the same absolute location

diff --git a/Engine/Source/Editor/Persona/Public/SocketDragDropOp.h b/Engine/Source/Editor/Persona/Public/SocketDragDropOp.h
--- a/Engine/Source/Editor/Persona/Public/SocketDragDropOp.h
+++ b/Engine/Source/Editor/Persona/Public/SocketDragDropOp.h
@@ -66,15 +66,18 @@ public:
 	FSelectedSocketInfo& GetSocketInfo() { return SocketInfo; }
 
 	/** Is this an alt-drag operation? */
-	bool IsAltDrag() { return bIsAltDrag; }
+	bool IsAltDrag() const { return ModifierKeysState.IsAltDown(); }
 
+	/** Is this an shift-drag operation? */
+	bool IsShiftDrag() const { return ModifierKeysState.IsShiftDown(); }
+	
 	/* Use this function to create a new one of me */
-	static TSharedRef<FSocketDragDropOp> New( FSelectedSocketInfo InSocketInfo, bool bInIsAltDrag )
+	static TSharedRef<FSocketDragDropOp> New( FSelectedSocketInfo InSocketInfo, FModifierKeysState InModifierKeysState )
 	{
 		check( InSocketInfo.Socket );
 		TSharedRef<FSocketDragDropOp> Operation = MakeShareable(new FSocketDragDropOp);
 		Operation->SocketInfo = InSocketInfo;
-		Operation->bIsAltDrag = bInIsAltDrag;
+		Operation->ModifierKeysState = InModifierKeysState;
 		Operation->SetIcon( FAppStyle::GetBrush(TEXT("Graph.ConnectorFeedback.Error")) );
 		Operation->Construct();
 		return Operation;
@@ -82,11 +85,11 @@ public:
 
 private:
 	/** The icon to display before the text */
-	const FSlateBrush* CurrentIconBrush;
+	const FSlateBrush* CurrentIconBrush = nullptr;
 
 	/** The socket that we're dragging */
 	FSelectedSocketInfo SocketInfo;
 
-	/** Is this an alt-drag? */
-	bool bIsAltDrag;
+	/** What was the modifier state at the start of the drag? */
+	FModifierKeysState ModifierKeysState;
 };
diff --git a/Engine/Source/Editor/SkeletonEditor/Private/EditableSkeleton.cpp b/Engine/Source/Editor/SkeletonEditor/Private/EditableSkeleton.cpp
--- a/Engine/Source/Editor/SkeletonEditor/Private/EditableSkeleton.cpp
+++ b/Engine/Source/Editor/SkeletonEditor/Private/EditableSkeleton.cpp
@@ -451,9 +451,21 @@ void FEditableSkeleton::RenameSocket(const FName OldSocketName, const FName NewS
 	}
 }
 
-void FEditableSkeleton::SetSocketParent(const FName& SocketName, const FName& NewParentName, USkeletalMesh* InSkeletalMesh)
+void FEditableSkeleton::SetSocketParent(const FName& SocketName, const FName& NewParentName, USkeletalMesh* InSkeletalMesh, bool bKeepAbsoluteLocation)
 {
-	USkeletalMeshSocket* Socket = FindSocket(SocketName, InSkeletalMesh, Skeleton);
+	USkeletalMeshSocket* Socket = nullptr;
+	const FReferenceSkeleton* ReferenceSkeleton = nullptr;
+
+	if (InSkeletalMesh)
+	{
+		Socket = InSkeletalMesh->FindSocket(SocketName);
+		ReferenceSkeleton = &InSkeletalMesh->GetRefSkeleton();
+	}
+	if (!Socket && Skeleton)
+	{
+		Socket = Skeleton->FindSocket(SocketName);
+		ReferenceSkeleton = &Skeleton->GetReferenceSkeleton();
+	}
 	if (Socket)
 	{
 		// Create an undo transaction, re-parent the socket and rebuild the skeleton tree view
@@ -462,6 +474,22 @@ void FEditableSkeleton::SetSocketParent(const FName& SocketName, const FName& Ne
 		Socket->SetFlags(RF_Transactional);	// Undo doesn't work without this!
 		Socket->Modify();
 
+		if (bKeepAbsoluteLocation)
+		{
+			const int32 OldBoneIndex = ReferenceSkeleton->FindBoneIndex(Socket->BoneName);
+			const int32 NewBoneIndex = ReferenceSkeleton->FindBoneIndex(NewParentName);
+
+			if (OldBoneIndex != NewBoneIndex && OldBoneIndex != INDEX_NONE && NewBoneIndex != INDEX_NONE)
+			{
+				// Transform the relative positioning of the socket from the space of the old parent bone to the space of the
+				// new parent bone.
+				const FTransform OldRefTransform = ReferenceSkeleton->GetBoneAbsoluteTransform(OldBoneIndex);
+				const FTransform NewRefTransform = ReferenceSkeleton->GetBoneAbsoluteTransform(NewBoneIndex);
+
+				Socket->SetSocketLocalTransform(Socket->GetSocketLocalTransform() * OldRefTransform * NewRefTransform.Inverse());
+			}
+		}
+		
 		Socket->BoneName = NewParentName;
 
 		for (TWeakPtr<SSkeletonTree> SkeletonTree : SkeletonTrees)
diff --git a/Engine/Source/Editor/SkeletonEditor/Private/EditableSkeleton.h b/Engine/Source/Editor/SkeletonEditor/Private/EditableSkeleton.h
--- a/Engine/Source/Editor/SkeletonEditor/Private/EditableSkeleton.h
+++ b/Engine/Source/Editor/SkeletonEditor/Private/EditableSkeleton.h
@@ -48,7 +48,7 @@ public:
 	virtual void BroadcastNotifyChanged() override;
 	virtual void GetCompatibleAnimSequences(TArray<struct FAssetData>& OutAssets) override;
 	virtual void RenameSocket(const FName OldSocketName, const FName NewSocketName, USkeletalMesh* InSkeletalMesh) override;
-	virtual void SetSocketParent(const FName& SocketName, const FName& NewParentName, USkeletalMesh* InSkeletalMesh) override;
+	virtual void SetSocketParent(const FName& SocketName, const FName& NewParentName, USkeletalMesh* InSkeletalMesh, bool bKeepAbsoluteLocation) override;
 	virtual bool DoesSocketAlreadyExist(const class USkeletalMeshSocket* InSocket, const FText& InSocketName, ESocketParentType SocketParentType, USkeletalMesh* InSkeletalMesh) const override;
 	virtual bool DoesVirtualBoneAlreadyExist(const FString& InVBName) const override;
 	virtual void RenameVirtualBone(const FName OriginalName, const FName InVBName) override;
diff --git a/Engine/Source/Editor/SkeletonEditor/Private/SkeletonTreeBoneItem.cpp b/Engine/Source/Editor/SkeletonEditor/Private/SkeletonTreeBoneItem.cpp
--- a/Engine/Source/Editor/SkeletonEditor/Private/SkeletonTreeBoneItem.cpp
+++ b/Engine/Source/Editor/SkeletonEditor/Private/SkeletonTreeBoneItem.cpp
@@ -499,7 +499,8 @@ FReply FSkeletonTreeBoneItem::HandleDrop(const FDragDropEvent& DragDropEvent)
 		{
 			// The socket can be dropped here if we're a bone and NOT the socket's existing parent
 			USkeletalMesh* SkeletalMesh = GetSkeletonTree()->GetPreviewScene().IsValid() ? ToRawPtr(GetSkeletonTree()->GetPreviewScene()->GetPreviewMeshComponent()->GetSkeletalMeshAsset()) : nullptr;
-			GetEditableSkeleton()->SetSocketParent(SocketInfo.Socket->SocketName, BoneName, SkeletalMesh);
+			const bool bKeepAbsoluteLocation = DragConnectionOp->IsShiftDrag();
+			GetEditableSkeleton()->SetSocketParent(SocketInfo.Socket->SocketName, BoneName, SkeletalMesh, bKeepAbsoluteLocation);
 
 			return FReply::Handled();
 		}
diff --git a/Engine/Source/Editor/SkeletonEditor/Private/SkeletonTreeSocketItem.cpp b/Engine/Source/Editor/SkeletonEditor/Private/SkeletonTreeSocketItem.cpp
--- a/Engine/Source/Editor/SkeletonEditor/Private/SkeletonTreeSocketItem.cpp
+++ b/Engine/Source/Editor/SkeletonEditor/Private/SkeletonTreeSocketItem.cpp
@@ -65,7 +65,8 @@ void FSkeletonTreeSocketItem::GenerateWidgetForNameColumn( TSharedPtr< SHorizont
 
 	const FSlateFontInfo TextFont = FCoreStyle::GetDefaultFontStyle("Regular", 10);
 
-	FText ToolTip = GetSocketToolTip();
+	FText ToolTip = FText::Format(LOCTEXT("SocketMainTreeItemTooltip", "{0}\n\nHold Alt to duplicate the socket.\nHold Shift to keep the same absolute location on the new parent bone."),
+		GetSocketToolTip());
 
 	TAttribute<FText> SocketNameAttr = TAttribute<FText>::Create(TAttribute<FText>::FGetter::CreateSP(this, &FSkeletonTreeSocketItem::GetSocketNameAsText));
 	TSharedPtr< SInlineEditableTextBlock > InlineWidget;
@@ -233,13 +234,13 @@ FReply FSkeletonTreeSocketItem::OnDragDetected(const FGeometry& MyGeometry, cons
 	{
 		FSelectedSocketInfo SocketInfo( Socket, ParentType == ESocketParentType::Skeleton );
 
-		return FReply::Handled().BeginDragDrop( FSocketDragDropOp::New( SocketInfo, MouseEvent.IsAltDown() ) );
+		return FReply::Handled().BeginDragDrop( FSocketDragDropOp::New( SocketInfo, MouseEvent.GetModifierKeys() ) );
 	}
 
 	return FReply::Unhandled();
 }
 
-FText FSkeletonTreeSocketItem::GetSocketToolTip()
+FText FSkeletonTreeSocketItem::GetSocketToolTip() const
 {
 	FText ToolTip;
 
@@ -253,7 +254,7 @@ FText FSkeletonTreeSocketItem::GetSocketToolTip()
 	}
 	else if ( ParentType == ESocketParentType::Skeleton )
 	{
-		ToolTip = LOCTEXT( "SocketToolTipSkeleton", "This socket is on the skeleton (shared with all meshes that use the skeleton), and the current mesh has duplciated version of it" );
+		ToolTip = LOCTEXT( "SocketToolTipSkeleton", "This socket is on the skeleton (shared with all meshes that use the skeleton), and the current mesh has duplicated version of it" );
 	}
 	else
 	{
diff --git a/Engine/Source/Editor/SkeletonEditor/Private/SkeletonTreeSocketItem.h b/Engine/Source/Editor/SkeletonEditor/Private/SkeletonTreeSocketItem.h
--- a/Engine/Source/Editor/SkeletonEditor/Private/SkeletonTreeSocketItem.h
+++ b/Engine/Source/Editor/SkeletonEditor/Private/SkeletonTreeSocketItem.h
@@ -58,7 +58,7 @@ private:
 	void OnCommitSocketName(const FText& InText, ETextCommit::Type CommitInfo);
 
 	/** Function that returns the current tooltip for this socket */
-	FText GetSocketToolTip();
+	FText GetSocketToolTip() const;
 
 	/** Get the color to display for the socket name */
 	FSlateColor GetTextColor() const;	
diff --git a/Engine/Source/Editor/SkeletonEditor/Public/IEditableSkeleton.h b/Engine/Source/Editor/SkeletonEditor/Public/IEditableSkeleton.h
--- a/Engine/Source/Editor/SkeletonEditor/Public/IEditableSkeleton.h
+++ b/Engine/Source/Editor/SkeletonEditor/Public/IEditableSkeleton.h
@@ -61,7 +61,7 @@ public:
 	virtual void RenameSocket(const FName OldSocketName, const FName NewSocketName, USkeletalMesh* InSkeletalMesh) = 0;
 
 	/** Set the parent of a socket */
-	virtual void SetSocketParent(const FName& SocketName, const FName& NewParentName, USkeletalMesh* InSkeletalMesh) = 0;
+	virtual void SetSocketParent(const FName& SocketName, const FName& NewParentName, USkeletalMesh* InSkeletalMesh, bool bKeepAbsoluteLocation = false) = 0;
 
 	/** Function to tell you if a socket name already exists. SocketParentType determines where we will look to see if the socket exists, i.e. mesh or skeleton */
 	virtual bool DoesSocketAlreadyExist(const class USkeletalMeshSocket* InSocket, const FText& InSocketName, ESocketParentType SocketParentType, USkeletalMesh* InSkeletalMesh) const = 0;
diff --git a/Engine/Source/Runtime/Engine/Private/Animation/ReferenceSkeleton.cpp b/Engine/Source/Runtime/Engine/Private/Animation/ReferenceSkeleton.cpp
--- a/Engine/Source/Runtime/Engine/Private/Animation/ReferenceSkeleton.cpp
+++ b/Engine/Source/Runtime/Engine/Private/Animation/ReferenceSkeleton.cpp
@@ -653,6 +653,83 @@ SIZE_T FReferenceSkeleton::GetDataSize() const
 	return ResourceSize;
 }
 
+namespace
+{
+void ComposeAbsoluteTransforms(
+	const FReferenceSkeleton& InRefSkeleton,
+	TConstArrayView<FTransform> InBoneSpaceTransforms,
+	TArray<FTransform>& OutAbsoluteTransforms
+	)
+{
+	OutAbsoluteTransforms.SetNumUninitialized(InBoneSpaceTransforms.Num());
+
+	for (int32 BoneIndex = 0; BoneIndex < InBoneSpaceTransforms.Num(); BoneIndex++)
+	{
+		const int32 ParentIndex = InRefSkeleton.GetParentIndex(BoneIndex);
+		if (ParentIndex >= 0)
+		{
+			OutAbsoluteTransforms[BoneIndex] = InBoneSpaceTransforms[BoneIndex] * OutAbsoluteTransforms[ParentIndex];
+		}
+		else
+		{
+			OutAbsoluteTransforms[BoneIndex] = InBoneSpaceTransforms[BoneIndex];
+		}
+	}
+}
+
+FTransform ComposeSingleAbsoluteTransform(
+	const int32 InBoneIndex,
+	const FReferenceSkeleton& InRefSkeleton,
+	TConstArrayView<FTransform> InBoneSpaceTransforms
+	)
+{
+	FTransform AbsolutePose = InBoneSpaceTransforms[InBoneIndex];
+
+	int32 ParentIndex = InRefSkeleton.GetParentIndex(InBoneIndex);
+
+	while(ParentIndex != INDEX_NONE)
+	{
+		AbsolutePose = AbsolutePose * InBoneSpaceTransforms[ParentIndex];
+		ParentIndex = InRefSkeleton.GetParentIndex(ParentIndex);
+	}
+
+	return AbsolutePose;
+}
+
+
+}
+
+
+void FReferenceSkeleton::GetRawBoneAbsoluteTransforms(TArray<FTransform>& OutAbsoluteTransforms) const
+{
+	ComposeAbsoluteTransforms(*this, RawRefBonePose, OutAbsoluteTransforms);
+}
+
+void FReferenceSkeleton::GetBoneAbsoluteTransforms(TArray<FTransform>& OutAbsoluteTransforms) const
+{
+	ComposeAbsoluteTransforms(*this, FinalRefBonePose, OutAbsoluteTransforms);
+}
+
+FTransform FReferenceSkeleton::GetRawBoneAbsoluteTransform(const int32 InRawBoneIndex) const
+{
+	if (ensure(IsValidRawIndex(InRawBoneIndex)))
+	{
+		return ComposeSingleAbsoluteTransform(InRawBoneIndex, *this, RawRefBonePose);
+	}
+
+	return FTransform::Identity;
+}
+
+FTransform FReferenceSkeleton::GetBoneAbsoluteTransform(const int32 InBoneIndex) const
+{
+	if (ensure(IsValidIndex(InBoneIndex)))
+	{
+		return ComposeSingleAbsoluteTransform(InBoneIndex, *this, FinalRefBonePose);
+	}
+
+	return FTransform::Identity;
+}
+
 void FReferenceSkeleton::GetRawChildrenIndicesCached(const int32 BoneIndex, TArray<int32>& OutChildren) const
 {
 	const int32 LastBranchIndex = GetCachedEndOfBranchIndex(BoneIndex);
diff --git a/Engine/Source/Runtime/Engine/Public/ReferenceSkeleton.h b/Engine/Source/Runtime/Engine/Public/ReferenceSkeleton.h
--- a/Engine/Source/Runtime/Engine/Public/ReferenceSkeleton.h
+++ b/Engine/Source/Runtime/Engine/Public/ReferenceSkeleton.h
@@ -450,6 +450,20 @@ public:
 
 	SIZE_T GetDataSize() const;
 
+	/** Return all the non-virtual bone poses as absolute transforms (i.e. relative to the origin, rather than the parent bone).
+	  * This function will resize the array to store the transforms, with no re-allocation if shrinking */
+	ENGINE_API void GetRawBoneAbsoluteTransforms(TArray<FTransform>& OutAbsoluteTransforms) const;
+	
+	/** Return all the bone poses as absolute transforms (i.e. relative to the origin, rather than the parent bone).
+	  * This function will resize the array to store the transforms, with no re-allocation if shrinking */
+	ENGINE_API void GetBoneAbsoluteTransforms(TArray<FTransform>& OutAbsoluteTransforms) const;
+
+	/** Return a single non-virtual bone's pose as absolute transform (i.e. relative to the origin, rather than the parent bone). */
+	ENGINE_API FTransform GetRawBoneAbsoluteTransform(const int32 InRawBoneIndex) const;
+	
+	/** Return a single bone's pose as absolute transform (i.e. relative to the origin, rather than the parent bone). */
+	ENGINE_API FTransform GetBoneAbsoluteTransform(const int32 InBoneIndex) const;
+
 	// lazy-cached for fast access to children and recursive children
 	void GetRawChildrenIndicesCached(const int32 BoneIndex, TArray<int32>& OutChildren) const;
 	void GetRawChildrenIndicesRecursiveCached(const int32 BoneIndex, TArray<int32>& OutChildren) const;
