[shaders] in the rare case where a distributed task fails to complete, we return compile jobs to the job queue without taking the lock. this can mess up the pending jobs list (and can likely lead to hangs)

diff --git a/Engine/Source/Runtime/Engine/Private/ShaderCompiler/ShaderCompilerDistributed.cpp b/Engine/Source/Runtime/Engine/Private/ShaderCompiler/ShaderCompilerDistributed.cpp
--- a/Engine/Source/Runtime/Engine/Private/ShaderCompiler/ShaderCompilerDistributed.cpp
+++ b/Engine/Source/Runtime/Engine/Private/ShaderCompiler/ShaderCompilerDistributed.cpp
@@ -523,6 +523,8 @@ int32 FShaderCompileDistributedThreadRunnable_Interface::CompilingLoop()
 		else
 		{
 			// The compile job was canceled. Return the jobs to the manager's compile queue.
+			UE_LOG(LogShaderCompilers, Display, TEXT("Distributed build task did not complete; returning %d jobs to the compile queue"), Task->ShaderJobs.Num());
+			FScopeLock Lock(&Manager->CompileQueueSection);
 			Manager->AllJobs.SubmitJobs(Task->ShaderJobs);
 		}
 
