NetConnection
* Deprecated IsNetReady(bool Saturate) and added a const version that doesn't have the Saturate parameter.
* The expectation of calling IsNetReady(Saturate=true) was wildly inconsistent with the results. In reality it would "unsaturate" the buffer immediately and the next serialized bits would pass-throught and then make it saturated.

diff --git a/Engine/Plugins/NetcodeUnitTest/NetcodeUnitTest/Source/NetcodeUnitTest/Private/MinimalClient.cpp b/Engine/Plugins/NetcodeUnitTest/NetcodeUnitTest/Source/NetcodeUnitTest/Private/MinimalClient.cpp
--- a/Engine/Plugins/NetcodeUnitTest/NetcodeUnitTest/Source/NetcodeUnitTest/Private/MinimalClient.cpp
+++ b/Engine/Plugins/NetcodeUnitTest/NetcodeUnitTest/Source/NetcodeUnitTest/Private/MinimalClient.cpp
@@ -193,7 +193,7 @@ FOutBunch* UMinimalClient::CreateChannelBunchByName(const FName& ChName, int32 C
 		// Hack to force IsNetReady
 		UnitConn->QueuedBits = -(MAX_PACKET_SIZE * 8);
 
-		if (ControlChan != nullptr && ControlChan->IsNetReady(false))
+		if (ControlChan != nullptr && ControlChan->IsNetReady())
 		{
 			int32 BunchSequence = ++UnitConn->OutReliable[ChIndex];
 
@@ -391,7 +391,7 @@ bool UMinimalClient::SendRPCChecked(UObject* Target, const TCHAR* FunctionName,
 #if 0
 			if (UnitConn != nullptr)
 			{
-				if (UnitConn->IsNetReady(false))
+				if (UnitConn->IsNetReady())
 				{
 					Target->ProcessEvent(TargetFunc, Parms);
 				}
@@ -545,7 +545,7 @@ bool UMinimalClient::PostSendRPC(FString RPCName, UObject* Target/*=nullptr*/)
 			{
 				LogAppend += TEXT(", GetNetConnection() returned nullptr");
 			}
-			else if (!TargetConn->IsNetReady(false))
+			else if (!TargetConn->IsNetReady())
 			{
 				LogAppend += TEXT(", IsNetReady() returned FALSE (NOTE: This should not block reliable packets!)");
 			}
diff --git a/Engine/Source/Runtime/Engine/Classes/Engine/Channel.h b/Engine/Source/Runtime/Engine/Classes/Engine/Channel.h
--- a/Engine/Source/Runtime/Engine/Classes/Engine/Channel.h
+++ b/Engine/Source/Runtime/Engine/Classes/Engine/Channel.h
@@ -158,8 +158,12 @@ public:
 	ENGINE_API virtual FPacketIdRange SendBunch(FOutBunch* Bunch, bool Merge);
 	
 	/** Return whether this channel is ready for sending. */
+	UE_DEPRECATED(5.6, "Replaced with a version that doesn't have a saturate option.")
 	ENGINE_API int32 IsNetReady( bool Saturate );
 
+	/** Returns true if this channel is ready to send and has not exceeded the bandwidth limit yet. */
+	ENGINE_API bool IsNetReady() const;
+
 	/** Make sure the incoming buffer is in sequence and there are no duplicates. */
 	ENGINE_API void AssertInSequenced();
 
diff --git a/Engine/Source/Runtime/Engine/Classes/Engine/ChildConnection.h b/Engine/Source/Runtime/Engine/Classes/Engine/ChildConnection.h
--- a/Engine/Source/Runtime/Engine/Classes/Engine/ChildConnection.h
+++ b/Engine/Source/Runtime/Engine/Classes/Engine/ChildConnection.h
@@ -60,9 +60,17 @@ public:
 		Parent->FlushNet(bIgnoreSimulation);
 	}
 
+	UE_DEPRECATED(5.6, "Replaced with a version that doesn't have a saturate option.")
 	virtual int32 IsNetReady(bool Saturate) override
 	{
+		PRAGMA_DISABLE_DEPRECATION_WARNINGS
 		return Parent->IsNetReady(Saturate);
+		PRAGMA_ENABLE_DEPRECATION_WARNINGS
+	}
+
+	virtual bool IsNetReady() const override
+	{
+		return Parent->IsNetReady();
 	}
 
 	virtual bool IsEncryptionEnabled() const override
diff --git a/Engine/Source/Runtime/Engine/Classes/Engine/DemoNetConnection.h b/Engine/Source/Runtime/Engine/Classes/Engine/DemoNetConnection.h
--- a/Engine/Source/Runtime/Engine/Classes/Engine/DemoNetConnection.h
+++ b/Engine/Source/Runtime/Engine/Classes/Engine/DemoNetConnection.h
@@ -29,7 +29,15 @@ public:
 	ENGINE_API virtual FString LowLevelGetRemoteAddress( bool bAppendPort = false ) override;
 	ENGINE_API virtual FString LowLevelDescribe() override;
 	ENGINE_API virtual void LowLevelSend(void* Data, int32 CountBits, FOutPacketTraits& Traits) override;
+
+	UE_DEPRECATED(5.6, "Replaced with a version that doesn't have a saturate option.")
 	ENGINE_API virtual int32 IsNetReady( bool Saturate ) override;
+	
+	ENGINE_API virtual bool IsNetReady() const override
+	{
+		return true;
+	}
+
 	ENGINE_API virtual void FlushNet( bool bIgnoreSimulation = false ) override;
 	ENGINE_API virtual void HandleClientPlayer( APlayerController* PC, class UNetConnection* NetConnection ) override;
 	ENGINE_API virtual TSharedPtr<const FInternetAddr> GetRemoteAddr() override;
diff --git a/Engine/Source/Runtime/Engine/Classes/Engine/NetConnection.h b/Engine/Source/Runtime/Engine/Classes/Engine/NetConnection.h
--- a/Engine/Source/Runtime/Engine/Classes/Engine/NetConnection.h
+++ b/Engine/Source/Runtime/Engine/Classes/Engine/NetConnection.h
@@ -455,9 +455,19 @@ public:
 	double			LastGoodPacketRealtime;	// Last real time a packet was considered valid
 	double			LastSendTime;			// Last time a packet was sent, for keepalives.
 	double			LastTickTime;			// Last time of polling.
-	int32			QueuedBits;			// Bits assumed to be queued up.
-	int32			TickCount;				// Count of ticks.
-	uint32			LastProcessedFrame;   // The last frame where we gathered and processed actors for this connection
+
+	/** 
+	 * The amount of bits allowed to be sent in a given tick to respect the bandwidth per second net speed.
+	 * Starts negative and increases with every serialized bits. When >0 it means we sent more bits then allowed.
+	 * Set at the end of every Tick based on the data sent that tick.
+	 */
+	int32			QueuedBits;
+
+	/** Count of ticks. */
+	int32			TickCount;
+
+	/** The last frame where we gathered and processed actors for this connection */
+	uint32			LastProcessedFrame;
 
 	double GetLastRecvAckTime() const { return LastRecvAckTimestamp; }
 	double GetConnectTime() const { return ConnectTimestamp; }
@@ -1025,7 +1035,11 @@ public:
 	ENGINE_API virtual void Tick(float DeltaSeconds);
 
 	/** Return whether this connection is ready for sending. */
-	ENGINE_API virtual int32 IsNetReady( bool Saturate );
+	UE_DEPRECATED(5.6, "Replaced with a version that doesn't have a saturate option. Call Saturate() directly if passing true before." )
+	ENGINE_API virtual int32 IsNetReady( bool bSaturate);
+
+	/** Returns true if the connection is ready to send and has not exceeded the bandwidth limit yet. */
+	ENGINE_API virtual bool IsNetReady() const;
 
 	/**
 	 * Return whether this connection can send packets without exhausting the packet sequence history window, as it could cause packets to be NAKed even when they've been received by the remote peer. 
diff --git a/Engine/Source/Runtime/Engine/Private/DataChannel.cpp b/Engine/Source/Runtime/Engine/Private/DataChannel.cpp
--- a/Engine/Source/Runtime/Engine/Private/DataChannel.cpp
+++ b/Engine/Source/Runtime/Engine/Private/DataChannel.cpp
@@ -1617,15 +1617,27 @@ FString UChannel::Describe()
 	return FString::Printf(TEXT("[UChannel] ChIndex: %d, Closing: %d %s"), ChIndex, (int32)Closing, Connection ? *Connection->Describe() : TEXT( "NULL CONNECTION" ));
 }
 
-
 int32 UChannel::IsNetReady( bool Saturate )
+{
+	// If saturation allowed, ignore queued byte count.
+	if (NumOutRec >= RELIABLE_BUFFER - 1)
+	{
+		return false;
+	}
+	
+	PRAGMA_DISABLE_DEPRECATION_WARNINGS
+	return Connection->IsNetReady(Saturate);
+	PRAGMA_ENABLE_DEPRECATION_WARNINGS
+}
+
+bool UChannel::IsNetReady() const
 {
 	// If saturation allowed, ignore queued byte count.
 	if( NumOutRec>=RELIABLE_BUFFER-1 )
 	{
-		return 0;
+		return false;
 	}
-	return Connection->IsNetReady( Saturate );
+	return Connection->IsNetReady();
 }
 
 void UChannel::ReceivedAck( int32 AckPacketId )
diff --git a/Engine/Source/Runtime/Engine/Private/Net/Experimental/Iris/DataStreamChannel.cpp b/Engine/Source/Runtime/Engine/Private/Net/Experimental/Iris/DataStreamChannel.cpp
--- a/Engine/Source/Runtime/Engine/Private/Net/Experimental/Iris/DataStreamChannel.cpp
+++ b/Engine/Source/Runtime/Engine/Private/Net/Experimental/Iris/DataStreamChannel.cpp
@@ -34,13 +34,6 @@
 
 namespace UE::Net::Private
 {
-	static bool bIrisSaturateBandwidth = false;
-	static FAutoConsoleVariableRef CVarIrisSaturateBandwidth(
-		TEXT("net.Iris.SaturateBandwidth"),
-		bIrisSaturateBandwidth,
-		TEXT("Whether to saturate the bandwidth or not. Default is false."
-		));
-
 	static int32 IrisPacketSequenceSafetyMargin = 4;
 	static FAutoConsoleVariableRef CVarIrisPacketSequenceSafetyMargin(
 		TEXT("net.Iris.PacketSequenceSafetyMargin"),
@@ -208,7 +201,7 @@ void UDataStreamChannel::PostTickDispatch()
 	}
 
 	// We probably want separate bandwidth management for iris as we are not pre-filling sendbuffer before call to NetReady.
-	if (!IsNetReady(UE::Net::Private::bIrisSaturateBandwidth))
+	if (!IsNetReady())
 	{
 		return;
 	}
@@ -241,7 +234,7 @@ void UDataStreamChannel::Tick()
 	}
 
 	// We probably want separate bandwidth management for iris as we are not pre-filling sendbuffer before call to NetReady.
-	if (!IsNetReady(UE::Net::Private::bIrisSaturateBandwidth))
+	if (!IsNetReady())
 	{
 		UE_CLOG(bHandshakeComplete, LogIris, Log, TEXT("Disallowed to write first packet in batch, with Iris this is not good!"))
 		return;
@@ -295,7 +288,7 @@ void UDataStreamChannel::WriteData(UE::Net::EDataStreamWriteMode WriteMode)
 	// Currently we want to use a full bunch so we flush if we have to
 	bool bNeedsPreSendFlush = Connection->SendBuffer.GetNumBits() > MAX_PACKET_HEADER_BITS;
 
-	auto WriteDataFunction = [this, &bNeedsPreSendFlush, &BeginWriteParams, MaxBunchBits, MaxBytes]()
+	auto WriteDataFunction = [this, &bNeedsPreSendFlush, &BeginWriteParams, MaxBunchBits, MaxBytes]() -> UDataStream::EWriteResult
 	{
 		if (bNeedsPreSendFlush)
 		{
@@ -324,7 +317,8 @@ void UDataStreamChannel::WriteData(UE::Net::EDataStreamWriteMode WriteMode)
 #endif
 
 		const FDataStreamRecord* Record = nullptr;
-		UDataStream::EWriteResult WriteResult = DataStreamManager->WriteData(SerializationContext, Record);
+		const UDataStream::EWriteResult WriteResult = DataStreamManager->WriteData(SerializationContext, Record);
+
 		if (WriteResult == UDataStream::EWriteResult::NoData || SerializationContext.HasError())
 		{
 			IRIS_PROFILER_SCOPE(UDataStreamChannel_NoDataSent);
@@ -370,8 +364,8 @@ void UDataStreamChannel::WriteData(UE::Net::EDataStreamWriteMode WriteMode)
 
 		if (PacketIds.First == INDEX_NONE)
 		{
-			check(PacketIds.First != INDEX_NONE);
 			// Something went wrong.
+			check(PacketIds.First != INDEX_NONE);
 			return UDataStream::EWriteResult::NoData;
 		}
 
@@ -393,11 +387,25 @@ void UDataStreamChannel::WriteData(UE::Net::EDataStreamWriteMode WriteMode)
 		return;
 	}
 
+
+	bool bContinueWriting = false;
 	do 
 	{
+		bContinueWriting = false;
+
 		// Write data until we are not allowed to write more
+		const UDataStream::EWriteResult Result = WriteDataFunction();
+
+		if (Result == UDataStream::EWriteResult::HasMoreData)
+		{
+			if (IsNetReady() && !IsPacketWindowFull())
+			{
+				bContinueWriting = true;
+			}
+		}
+
 	}
-	while ((WriteDataFunction() == UDataStream::EWriteResult::HasMoreData) && IsNetReady(UE::Net::Private::bIrisSaturateBandwidth) && (!IsPacketWindowFull()));
+	while (bContinueWriting);
 
 	// call end write to cleanup data initialized in BeginWrite	
 	DataStreamManager->EndWrite();
diff --git a/Engine/Source/Runtime/Engine/Private/Net/TrafficControl.cpp b/Engine/Source/Runtime/Engine/Private/Net/TrafficControl.cpp
--- a/Engine/Source/Runtime/Engine/Private/Net/TrafficControl.cpp
+++ b/Engine/Source/Runtime/Engine/Private/Net/TrafficControl.cpp
@@ -36,7 +36,7 @@ void FNetworkCongestionControl::OnAck(const FAckSample& AckSample)
 	UpdateMinRTT(AckSample.Timestamp);
 }
 
-bool FNetworkCongestionControl::IsReadyToSend(double Timestamp)
+bool FNetworkCongestionControl::IsReadyToSend(double Timestamp) const
 {
 	uint32 BytesInFlight = Analyzer.GetBytesInFlight();
 	uint32 PacketsInFlight = Analyzer.GetPacketsInFlight();
diff --git a/Engine/Source/Runtime/Engine/Private/NetConnection.cpp b/Engine/Source/Runtime/Engine/Private/NetConnection.cpp
--- a/Engine/Source/Runtime/Engine/Private/NetConnection.cpp
+++ b/Engine/Source/Runtime/Engine/Private/NetConnection.cpp
@@ -2543,23 +2543,30 @@ bool UNetConnection::ShouldDropOutgoingPacketForLossSimulation(int64 NumBits) co
 }
 #endif
 
-int32 UNetConnection::IsNetReady(bool Saturate)
+int32 UNetConnection::IsNetReady(bool bSaturate)
 {
-	if (IsReplay())
+	if (!IsReplay() && bSaturate)
 	{
-		return 1;
+		// Note that this doesn't immediately saturate, but instead forces the buffer to be on the threshold of being saturated.
+		// So IsNetReady will be true but the next serialized bit will make it go over the limit and then IsNetReady will return false
+		// Due to this confusion between name and effect it was removed as an option.
+		QueuedBits = -SendBuffer.GetNumBits();
 	}
 
-	// Return whether we can send more data without saturation the connection.
-	if (Saturate)
+	return (int32)IsNetReady();
+}
+
+bool UNetConnection::IsNetReady() const
+{
+	if (IsReplay())
 	{
-		QueuedBits = -SendBuffer.GetNumBits();
+		return true;
 	}
 
 #if !(UE_BUILD_SHIPPING || UE_BUILD_TEST)
 	if (CVarDisableBandwithThrottling.GetValueOnAnyThread() > 0)
 	{
-		return 1;
+		return true;
 	}
 #endif
 
@@ -4810,7 +4817,7 @@ void UNetConnection::Tick(float DeltaSeconds)
 	// Update queued byte count.
 	// this should be at the end so that the cap is applied *after* sending (and adjusting QueuedBytes for) any remaining data for this tick
 
-	SaturationAnalytics.TrackFrame(!IsNetReady(false));
+	SaturationAnalytics.TrackFrame(!IsNetReady());
 
 	if (!IsReplay())
 	{
diff --git a/Engine/Source/Runtime/Engine/Private/NetDriver.cpp b/Engine/Source/Runtime/Engine/Private/NetDriver.cpp
--- a/Engine/Source/Runtime/Engine/Private/NetDriver.cpp
+++ b/Engine/Source/Runtime/Engine/Private/NetDriver.cpp
@@ -2657,7 +2657,7 @@ void UNetDriver::InternalProcessRemoteFunctionPrivate(
 	}
 
 	// If saturated and function is unimportant, skip it. Note unreliable multicasts are queued at the actor channel level so they are not gated here.
-	if (!(Function->FunctionFlags & FUNC_NetReliable) && (!(Function->FunctionFlags & FUNC_NetMulticast)) && (!Connection->IsNetReady(0)))
+	if (!(Function->FunctionFlags & FUNC_NetReliable) && (!(Function->FunctionFlags & FUNC_NetMulticast)) && (!Connection->IsNetReady()))
 	{
 		DEBUG_REMOTEFUNCTION(TEXT("Network saturated, not calling %s::%s"), *GetNameSafe(Actor), *GetNameSafe(Function));
 		return;
@@ -5490,7 +5490,7 @@ int32 UNetDriver::ServerReplicateActors_ProcessPrioritizedActorsRange( UNetConne
 	int32 ActorUpdatesThisConnectionSent	= 0;
 	int32 FinalRelevantCount				= 0;
 
-	if (!Connection->IsNetReady( 0 ) && !bIgnoreSaturation)
+	if (!Connection->IsNetReady() && !bIgnoreSaturation)
 	{
 		GNumSaturatedConnections++;
 		// Connection saturated, don't process any actors
@@ -5610,7 +5610,7 @@ int32 UNetDriver::ServerReplicateActors_ProcessPrioritizedActorsRange( UNetConne
 						Channel->RelevantTime = ElapsedTime + 0.5 * UpdateDelayRandomStream.FRand();
 					}
 					// if the channel isn't saturated
-					if ( Channel->IsNetReady( 0 ) || bIgnoreSaturation)
+					if ( Channel->IsNetReady() || bIgnoreSaturation)
 					{
 						// replicate the actor
 						UE_LOG( LogNetTraffic, Log, TEXT( "- Replicate %s. %d" ), *Actor->GetName(), PriorityActors[j]->Priority );
@@ -5662,7 +5662,7 @@ int32 UNetDriver::ServerReplicateActors_ProcessPrioritizedActorsRange( UNetConne
 						Actor->ForceNetUpdate();
 					}
 					// second check for channel saturation
-					if (!Connection->IsNetReady( 0 ) && !bIgnoreSaturation)
+					if (!Connection->IsNetReady() && !bIgnoreSaturation)
 					{
 						// We can bail out now since this connection is saturated, we'll return how far we got though
 						GNumSaturatedConnections++;
@@ -8332,7 +8332,7 @@ void UNetDriver::UpdateNetworkStats()
 #if STATS
 					ServerConnection->PackageMap->GetNetGUIDStats(AckCount, UnAckCount, PendingCount);
 #endif
-					NetSaturated = ServerConnection->IsNetReady(false) ? 0 : 1;
+					NetSaturated = ServerConnection->IsNetReady() ? 0 : 1;
 				}
 				else
 				{
@@ -8350,7 +8350,7 @@ void UNetDriver::UpdateNetworkStats()
 						UnAckCount += ClientUnAckCount;
 						PendingCount += ClientPendingCount;
 #endif
-						NetSaturated |= ClientConnections[i]->IsNetReady(false) ? 0 : 1;
+						NetSaturated |= ClientConnections[i]->IsNetReady() ? 0 : 1;
 					}
 					// Subtract out the duplicate counted dormant actors
 					NumDormantActors -= FMath::Max(0, ClientConnections.Num() - 1) * SharedDormantActors;
diff --git a/Engine/Source/Runtime/Engine/Private/VoiceChannel.cpp b/Engine/Source/Runtime/Engine/Private/VoiceChannel.cpp
--- a/Engine/Source/Runtime/Engine/Private/VoiceChannel.cpp
+++ b/Engine/Source/Runtime/Engine/Private/VoiceChannel.cpp
@@ -71,7 +71,7 @@ void UVoiceChannel::Tick()
 		int32 Index;
 		for (Index = 0; Index < VoicePackets.Num(); Index++)
 		{
-			if (!Connection->IsNetReady(0))
+			if (!Connection->IsNetReady())
 			{
 				// If the network is saturated bail early
 				UE_LOG(LogNet, Warning, TEXT("Network saturated"));
diff --git a/Engine/Source/Runtime/Engine/Public/Net/TrafficControl.h b/Engine/Source/Runtime/Engine/Public/Net/TrafficControl.h
--- a/Engine/Source/Runtime/Engine/Public/Net/TrafficControl.h
+++ b/Engine/Source/Runtime/Engine/Public/Net/TrafficControl.h
@@ -81,7 +81,7 @@ public:
 	ENGINE_API FNetworkCongestionControl(double ConfiguredNetSpeed, uint32 MaxPackets);
 
 	ENGINE_API void OnAck(const FAckSample& AckSample);
-	ENGINE_API bool IsReadyToSend(double Timestamp);
+	ENGINE_API bool IsReadyToSend(double Timestamp) const;
 	ENGINE_API void OnSend(const FSeqSample& SeqSample);
 
 private:
diff --git a/Engine/Source/Runtime/Engine/Public/ReplayNetConnection.h b/Engine/Source/Runtime/Engine/Public/ReplayNetConnection.h
--- a/Engine/Source/Runtime/Engine/Public/ReplayNetConnection.h
+++ b/Engine/Source/Runtime/Engine/Public/ReplayNetConnection.h
@@ -23,7 +23,15 @@ public:
 	ENGINE_API virtual FString LowLevelGetRemoteAddress(bool bAppendPort = false) override;
 	ENGINE_API virtual FString LowLevelDescribe() override;
 	ENGINE_API virtual void LowLevelSend(void* Data, int32 CountBits, FOutPacketTraits& Traits) override;
-	ENGINE_API virtual int32 IsNetReady(bool Saturate) override;
+	
+	UE_DEPRECATED(5.6, "Replaced with a version that doesn't have a saturate option.")
+	ENGINE_API virtual int32 IsNetReady(bool bSaturate) override;
+
+	virtual bool IsNetReady() const override
+	{
+		return true;
+	}
+	
 	ENGINE_API virtual void HandleClientPlayer(APlayerController* PC, UNetConnection* NetConnection) override;
 	ENGINE_API virtual TSharedPtr<const FInternetAddr> GetRemoteAddr() override;
 	ENGINE_API virtual bool ClientHasInitializedLevel(const ULevel* TestLevel) const override;
