[SceneComponent] removed problematic navigation update when attaching component. The updates from PropagateTransformUpdate should cover all required scenarios.

#tests Attachments in EngineTest, compared NavigationOctree elements in complex maps

diff --git a/Engine/Source/Runtime/Engine/Private/Components/SceneComponent.cpp b/Engine/Source/Runtime/Engine/Private/Components/SceneComponent.cpp
--- a/Engine/Source/Runtime/Engine/Private/Components/SceneComponent.cpp
+++ b/Engine/Source/Runtime/Engine/Private/Components/SceneComponent.cpp
@@ -2386,15 +2386,6 @@ bool USceneComponent::AttachToComponent(USceneComponent* Parent, const FAttachme
 		{
 			// Update overlaps, in case location changed or overlap state depends on attachment.
 			UpdateOverlaps();
-
-			// Update our owner actor in the navigation system since its associated bounds have changed
-			if (IsNavigationRelevant())
-			{
-				if (AActor* Owner = GetOwner())
-				{
-					FNavigationSystem::UpdateActorAndComponentData(*Owner);
-				}
-			}
 		}
 
 		return true;
