Fix Thickness usage on FDebugRenderSceneProxy primitives

diff --git a/Engine/Source/Developer/LogVisualizer/Private/VisualLoggerRenderingActorBase.cpp b/Engine/Source/Developer/LogVisualizer/Private/VisualLoggerRenderingActorBase.cpp
--- a/Engine/Source/Developer/LogVisualizer/Private/VisualLoggerRenderingActorBase.cpp
+++ b/Engine/Source/Developer/LogVisualizer/Private/VisualLoggerRenderingActorBase.cpp
@@ -352,7 +352,7 @@ void AVisualLoggerRenderingActorBase::GetDebugShapes(const FVisualLogEntry& InEn
 			for (int32 Index = 0; Index + 1 < NumPoints; Index += 2, BoxExtent += 2)
 			{
 				const FBox Box = FBox(*BoxExtent, *(BoxExtent + 1));
-				DebugShapes.Boxes.Add(FDebugRenderSceneProxy::FDebugBox(Box, Color, Transform, DrawTypeOverride));
+				DebugShapes.Boxes.Add(FDebugRenderSceneProxy::FDebugBox(Box, Color, Transform, DrawTypeOverride, Thickness));
 
 				if (bDrawLabel)
 				{
@@ -376,7 +376,8 @@ void AVisualLoggerRenderingActorBase::GetDebugShapes(const FVisualLogEntry& InEn
 
 				FVector YAxis, ZAxis;
 				Direction.FindBestAxisVectors(YAxis, ZAxis);
-				DebugShapes.Cones.Add(FDebugRenderSceneProxy::FCone(FScaleMatrix(FVector(Length)) * FMatrix(Direction, YAxis, ZAxis, Origin), static_cast<float>(Angles.Y), static_cast<float>(Angles.Z), Color, DrawTypeOverride));
+				DebugShapes.Cones.Add(FDebugRenderSceneProxy::FCone(FScaleMatrix(FVector(Length)) * FMatrix(Direction, YAxis, ZAxis, Origin),
+					static_cast<float>(Angles.Y), static_cast<float>(Angles.Z), Color, DrawTypeOverride, Thickness));
 
 				if (bDrawLabel)
 				{
@@ -400,7 +401,7 @@ void AVisualLoggerRenderingActorBase::GetDebugShapes(const FVisualLogEntry& InEn
 				DebugShapes.Cylinders.Add(FDebugRenderSceneProxy::FWireCylinder(Center
 					, (End - Start).GetSafeNormal()
 					, static_cast<float>(OtherData.X)
-					, HalfHeight, Color, DrawTypeOverride)); // Base parameter is the center of the cylinder
+					, HalfHeight, Color, DrawTypeOverride, Thickness)); // Base parameter is the center of the cylinder
 				if (bDrawLabel)
 				{
 					DebugShapes.Texts.Add(FDebugRenderSceneProxy::FText3d(ElementToDraw->Description, Center, Color.WithAlpha(255)));
@@ -427,7 +428,7 @@ void AVisualLoggerRenderingActorBase::GetDebugShapes(const FVisualLogEntry& InEn
 				const FVector YAxis = Axes.GetScaledAxis(EAxis::Y);
 				const FVector ZAxis = Axes.GetScaledAxis(EAxis::Z);
 
-				DebugShapes.Capsules.Add(FDebugRenderSceneProxy::FCapsule(Base, Radius, XAxis, YAxis, ZAxis, HalfHeight, Color, DrawTypeOverride));
+				DebugShapes.Capsules.Add(FDebugRenderSceneProxy::FCapsule(Base, Radius, XAxis, YAxis, ZAxis, HalfHeight, Color, DrawTypeOverride, Thickness));
 				if (bDrawLabel)
 				{
 					DebugShapes.Texts.Add(FDebugRenderSceneProxy::FText3d(ElementToDraw->Description, Base + HalfHeight * FVector::UpVector, Color.WithAlpha(255)));
@@ -545,7 +546,7 @@ void AVisualLoggerRenderingActorBase::GetDebugShapes(const FVisualLogEntry& InEn
 				const FVector UpAxis = ElementToDraw->Points[Index + 1];
 				const double Radius = ElementToDraw->Points[Index + 2].X;
 
-				DebugShapes.Circles.Add(FDebugRenderSceneProxy::FCircle(Center, UpAxis.GetSafeNormal(), Radius, Color, DrawTypeOverride)); 
+				DebugShapes.Circles.Add(FDebugRenderSceneProxy::FCircle(Center, UpAxis.GetSafeNormal(), Radius, Color, DrawTypeOverride, Thickness)); 
 				
 				if (bDrawLabel)
 				{
diff --git a/Engine/Source/Runtime/Engine/Private/DebugRenderSceneProxy.cpp b/Engine/Source/Runtime/Engine/Private/DebugRenderSceneProxy.cpp
--- a/Engine/Source/Runtime/Engine/Private/DebugRenderSceneProxy.cpp
+++ b/Engine/Source/Runtime/Engine/Private/DebugRenderSceneProxy.cpp
@@ -440,7 +440,7 @@ void FDebugRenderSceneProxy::FDebugBox::Draw(FPrimitiveDrawInterface* PDI, EDraw
 {
 	if (InDrawType == SolidAndWireMeshes || InDrawType == WireMesh)
 	{
-		DrawWireBox(PDI, Transform.ToMatrixWithScale(), Box, Color, SDPG_World, InDrawType == SolidAndWireMeshes ? 2 : 0, 0, true);
+		DrawWireBox(PDI, Transform.ToMatrixWithScale(), Box, Color, SDPG_World, InDrawType == SolidAndWireMeshes ? 2 : Thickness, 0, true);
 	}
 	if (InDrawType == SolidAndWireMeshes || InDrawType == SolidMesh)
 	{
@@ -454,7 +454,7 @@ void FDebugRenderSceneProxy::FCircle::Draw(FPrimitiveDrawInterface* PDI, EDrawTy
 	Axis.FindBestAxisVectors(XAxis, YAxis);
 	if (InDrawType == SolidAndWireMeshes || InDrawType == WireMesh)
 	{
-		DrawCircle(PDI, Center, XAxis, YAxis, Color, Radius, (InDrawType == SolidAndWireMeshes) ? 9 : 12, SDPG_World, InDrawType == SolidAndWireMeshes ? 2 : 0, 0, true);
+		DrawCircle(PDI, Center, XAxis, YAxis, Color, Radius, (InDrawType == SolidAndWireMeshes) ? 9 : 12, SDPG_World, InDrawType == SolidAndWireMeshes ? 2 : Thickness, 0, true);
 	}
 
 	if (InDrawType == SolidAndWireMeshes || InDrawType == SolidMesh)
diff --git a/Engine/Source/Runtime/Engine/Public/DebugRenderSceneProxy.h b/Engine/Source/Runtime/Engine/Public/DebugRenderSceneProxy.h
--- a/Engine/Source/Runtime/Engine/Public/DebugRenderSceneProxy.h
+++ b/Engine/Source/Runtime/Engine/Public/DebugRenderSceneProxy.h
@@ -123,7 +123,7 @@ public:
 	/** Struct to hold info about lines to render. */
 	struct FDebugLine
 	{
-		FDebugLine(const FVector &InStart, const FVector &InEnd, const FColor &InColor, float InThickness = 0) 
+		FDebugLine(const FVector &InStart, const FVector &InEnd, const FColor &InColor, const float InThickness = 0.f) 
 			: Start(InStart)
 			, End(InEnd)
 			, Color(InColor)
@@ -141,11 +141,13 @@ public:
 	/** Struct to hold info about circles to render. */
 	struct FCircle
 	{
-		FCircle(const FVector& InCenter, const FVector& InAxis, const float InRadius, const FColor& InColor, EDrawType InDrawTypeOverride = EDrawType::Invalid)
+		FCircle(const FVector& InCenter, const FVector& InAxis, const float InRadius, const FColor& InColor, 
+			EDrawType InDrawTypeOverride = EDrawType::Invalid, const float InThickness = 0.f)
 			: Center(InCenter)
 			, Axis(InAxis)
 			, Radius(InRadius)
 			, Color(InColor)
+			, Thickness(InThickness)
 			, DrawTypeOverride(InDrawTypeOverride)
 		{}
 
@@ -155,21 +157,25 @@ public:
 		FVector Axis;
 		float Radius;
 		FColor Color;
+		float Thickness;
 		EDrawType DrawTypeOverride = EDrawType::Invalid;
 	};
 
 	/** Struct to hold info about boxes to render. */
 	struct FDebugBox
 	{
-		FDebugBox(const FBox& InBox, const FColor& InColor, EDrawType InDrawTypeOverride = EDrawType::Invalid)
+		FDebugBox(const FBox& InBox, const FColor& InColor, EDrawType InDrawTypeOverride = EDrawType::Invalid, const float InThickness = 0.f)
 			: Box(InBox)
 			, Color(InColor)
+			, Thickness(InThickness)
 			, DrawTypeOverride(InDrawTypeOverride)
 		{}
 
-		FDebugBox(const FBox& InBox, const FColor& InColor, const FTransform& InTransform, EDrawType InDrawTypeOverride = EDrawType::Invalid)
+		FDebugBox(const FBox& InBox, const FColor& InColor, const FTransform& InTransform, EDrawType InDrawTypeOverride = EDrawType::Invalid, 
+			const float InThickness = 0.f)
 			: Box(InBox)
 			, Color(InColor)
+			, Thickness(InThickness)
 			, Transform(InTransform)
 			, DrawTypeOverride(InDrawTypeOverride)
 		{}
@@ -178,6 +184,7 @@ public:
 
 		FBox Box;
 		FColor Color;
+		float Thickness;
 		FTransform Transform;
 		EDrawType DrawTypeOverride = EDrawType::Invalid;
 	};
@@ -185,12 +192,14 @@ public:
 	/** Struct to hold info about cylinders to render. */
 	struct FWireCylinder
 	{
-		FWireCylinder(const FVector& InBase, const FVector& InDirection, const float InRadius, const float InHalfHeight, const FColor &InColor, EDrawType InDrawTypeOverride = EDrawType::Invalid)
+		FWireCylinder(const FVector& InBase, const FVector& InDirection, const float InRadius, const float InHalfHeight, const FColor &InColor,
+			EDrawType InDrawTypeOverride = EDrawType::Invalid, const float InThickness = 0.f)
 			: Base(InBase)
 			, Direction(InDirection)
 			, Radius(InRadius)
 			, HalfHeight(InHalfHeight)
 			, Color(InColor) 
+			, Thickness(InThickness)
 			, DrawTypeOverride(InDrawTypeOverride)
 		{}
 
@@ -201,6 +210,7 @@ public:
 		float Radius;
 		float HalfHeight;
 		FColor Color;
+		float Thickness;
 		EDrawType DrawTypeOverride = EDrawType::Invalid;
 	};
 
@@ -296,11 +306,13 @@ public:
 	struct FCone
 	{
 		FCone() {}
-		FCone(const FMatrix& InConeToWorld, const float InAngle1, const float InAngle2, const FLinearColor& InColor, EDrawType InDrawTypeOverride = EDrawType::Invalid)
+		FCone(const FMatrix& InConeToWorld, const float InAngle1, const float InAngle2, const FLinearColor& InColor,
+			EDrawType InDrawTypeOverride = EDrawType::Invalid, const float InThickness = 0.f)
 			: ConeToWorld(InConeToWorld)
 			, Angle1(InAngle1)
 			, Angle2(InAngle2)
-			, Color(InColor.ToFColor(true)) 
+			, Color(InColor.ToFColor(true))
+			, Thickness(InThickness)
 			, DrawTypeOverride(InDrawTypeOverride)
 		{}
 
@@ -310,6 +322,7 @@ public:
 		float Angle1;
 		float Angle2;
 		FColor Color;
+		float Thickness;
 		EDrawType DrawTypeOverride = EDrawType::Invalid;
 	};
 
@@ -324,10 +337,12 @@ public:
 	struct FCapsule
 	{
 		FCapsule() {}
-		FCapsule(const FVector& InBase, const float& InRadius, const FVector& x, const FVector& y, const FVector &z, const float& InHalfHeight, const FLinearColor& InColor, EDrawType InDrawTypeOverride = EDrawType::Invalid)
+		FCapsule(const FVector& InBase, const float& InRadius, const FVector& x, const FVector& y, const FVector &z, const float& InHalfHeight,
+			const FLinearColor& InColor, EDrawType InDrawTypeOverride = EDrawType::Invalid, const float InThickness = 0.f)
 			: Radius(InRadius)
 			, Base(InBase)
 			, Color(InColor.ToFColor(true))
+			, Thickness(InThickness)
 			, HalfHeight(InHalfHeight)
 			, X(x)
 			, Y(y)
@@ -340,6 +355,7 @@ public:
 		float Radius;
 		FVector Base; //Center point of the base of the cylinder.
 		FColor Color;
+		float Thickness;
 		float HalfHeight;
 		FVector X, Y, Z; //X, Y, and Z alignment axes to draw along.
 		EDrawType DrawTypeOverride = EDrawType::Invalid;
