Make the lock in SkyLightComponent transactionally safe.

diff --git a/Engine/Source/Runtime/Engine/Classes/Components/SkyLightComponent.h b/Engine/Source/Runtime/Engine/Classes/Components/SkyLightComponent.h
--- a/Engine/Source/Runtime/Engine/Classes/Components/SkyLightComponent.h
+++ b/Engine/Source/Runtime/Engine/Classes/Components/SkyLightComponent.h
@@ -10,6 +10,7 @@
 #include "RenderingThread.h"
 #include "Components/LightComponentBase.h"
 #include "Math/SHMath.h"
+#include "Misc/TransactionallySafeCriticalSection.h"
 #include "RenderDeferredCleanup.h"
 #include "SkyLightComponent.generated.h"
 
@@ -354,7 +355,7 @@ protected:
 	 */
 	static ENGINE_API TArray<USkyLightComponent*> SkyCapturesToUpdate;
 	static ENGINE_API TArray<USkyLightComponent*> SkyCapturesToUpdateBlendDestinations;
-	static ENGINE_API FCriticalSection SkyCapturesToUpdateLock;
+	static ENGINE_API FTransactionallySafeCriticalSection SkyCapturesToUpdateLock;
 
 	//~ Begin UActorComponent Interface
 	ENGINE_API virtual void CreateRenderState_Concurrent(FRegisterComponentContext* Context) override;
diff --git a/Engine/Source/Runtime/Engine/Private/Components/SkyLightComponent.cpp b/Engine/Source/Runtime/Engine/Private/Components/SkyLightComponent.cpp
--- a/Engine/Source/Runtime/Engine/Private/Components/SkyLightComponent.cpp
+++ b/Engine/Source/Runtime/Engine/Private/Components/SkyLightComponent.cpp
@@ -19,6 +19,7 @@
 #include "Misc/UObjectToken.h"
 #include "Net/UnrealNetwork.h"
 #include "Misc/MapErrors.h"
+#include "Misc/ScopeLock.h"
 #include "SceneInterface.h"
 #include "SceneProxies/SkyLightSceneProxy.h"
 #include "ShaderCompiler.h"
@@ -351,7 +352,7 @@ void USkyLightComponent::SetCaptureIsDirty()
 { 
 	if (GetVisibleFlag() && bAffectsWorld && !SkipStaticSkyLightCapture(*this))
 	{
-		FScopeLock Lock(&SkyCapturesToUpdateLock);
+		UE::TScopeLock Lock(SkyCapturesToUpdateLock);
 
 #if WITH_EDITOR
 		this->CaptureStatus = ESkyLightCaptureStatus::SLCS_Uninitialized;
@@ -428,7 +429,7 @@ void USkyLightComponent::SetBlendDestinationCaptureIsDirty()
 
 TArray<USkyLightComponent*> USkyLightComponent::SkyCapturesToUpdate;
 TArray<USkyLightComponent*> USkyLightComponent::SkyCapturesToUpdateBlendDestinations;
-FCriticalSection USkyLightComponent::SkyCapturesToUpdateLock;
+FTransactionallySafeCriticalSection USkyLightComponent::SkyCapturesToUpdateLock;
 
 void USkyLightComponent::CreateRenderState_Concurrent(FRegisterComponentContext* Context)
 {
@@ -456,7 +457,7 @@ void USkyLightComponent::PostInitProperties()
 	{
 		// Enqueue an update by default, so that newly placed components will get an update
 		// PostLoad will undo this for components loaded from disk
-		FScopeLock Lock(&SkyCapturesToUpdateLock);
+		UE::TScopeLock Lock(SkyCapturesToUpdateLock);
 		SkyCapturesToUpdate.AddUnique(this);
 	}
 
@@ -474,7 +475,7 @@ void USkyLightComponent::PostLoad()
 		// All components are queued for update on creation by default. But we do not want this top happen in some cases.
 		if (!GetVisibleFlag() || HasAnyFlags(RF_ClassDefaultObject | RF_ArchetypeObject) || SkipStaticSkyLightCapture(*this))
 		{
-			FScopeLock Lock(&SkyCapturesToUpdateLock);
+			UE::TScopeLock Lock(SkyCapturesToUpdateLock);
 			SkyCapturesToUpdate.Remove(this);
 		}
 	}
@@ -676,7 +677,7 @@ void USkyLightComponent::BeginDestroy()
 {
 	// Deregister the component from the update queue
 	{
-		FScopeLock Lock(&SkyCapturesToUpdateLock); 
+		UE::TScopeLock Lock(SkyCapturesToUpdateLock); 
 		SkyCapturesToUpdate.Remove(this);
 	}
 	
@@ -726,7 +727,7 @@ void USkyLightComponent::ApplyComponentInstanceData(FPrecomputedSkyLightInstance
 	if (ProcessedSkyTexture && bSavedConstructionScriptValuesValid)
 	{
 		// We have valid capture state, remove the queued update
-		FScopeLock Lock(&SkyCapturesToUpdateLock);
+		UE::TScopeLock Lock(SkyCapturesToUpdateLock);
 		SkyCapturesToUpdate.Remove(this);
 	}
 
@@ -928,7 +929,7 @@ void USkyLightComponent::UpdateSkyCaptureContents(UWorld* WorldToUpdate)
 
 		if (SkyCapturesToUpdate.Num() > 0)
 		{
-			FScopeLock Lock(&SkyCapturesToUpdateLock);
+			UE::TScopeLock Lock(SkyCapturesToUpdateLock);
 			// Remove the sky captures if real time capture is enabled. 
 			SkyCapturesToUpdate.RemoveAll([WorldToUpdate](const USkyLightComponent* CaptureComponent) { return CaptureComponent->GetWorld() == WorldToUpdate && CaptureComponent->IsRealTimeCaptureEnabled(); });
 			UpdateSkyCaptureContentsArray(WorldToUpdate, SkyCapturesToUpdate, true);
