Enable GPUScene for skinned meshes on mobile platforms

diff --git a/Engine/Source/Runtime/Engine/Private/GPUSkinVertexFactory.cpp b/Engine/Source/Runtime/Engine/Private/GPUSkinVertexFactory.cpp
--- a/Engine/Source/Runtime/Engine/Private/GPUSkinVertexFactory.cpp
+++ b/Engine/Source/Runtime/Engine/Private/GPUSkinVertexFactory.cpp
@@ -771,8 +771,7 @@ void TGPUSkinVertexFactory<BoneInfluenceType>::ModifyCompilationEnvironment(cons
 	FVertexFactory::ModifyCompilationEnvironment(Parameters, OutEnvironment);
 
 	const FStaticFeatureLevel MaxSupportedFeatureLevel = GetMaxSupportedFeatureLevel(Parameters.Platform);
-	// TODO: support GPUScene on mobile
-	const bool bUseGPUScene = UseGPUScene(Parameters.Platform, MaxSupportedFeatureLevel) && (MaxSupportedFeatureLevel > ERHIFeatureLevel::ES3_1);
+	const bool bUseGPUScene = UseGPUScene(Parameters.Platform, MaxSupportedFeatureLevel);
 	const bool bSupportsPrimitiveIdStream = Parameters.VertexFactoryType->SupportsPrimitiveIdStream();
 	{
 		const bool bLimit2BoneInfluences = (CVarGPUSkinLimit2BoneInfluences.GetValueOnAnyThread() != 0);
@@ -934,8 +933,8 @@ void TGPUSkinVertexFactory<BoneInfluenceType>::GetVertexElements(ERHIFeatureLeve
 	int32 MorphDeltaStreamIndex;
 	GetVertexElements(FeatureLevel, InputStreamType, GPUSkinData, OutElements, VertexStreams, MorphDeltaStreamIndex);
 
-	if (UseGPUScene(GMaxRHIShaderPlatform, GMaxRHIFeatureLevel) && 
-		FeatureLevel > ERHIFeatureLevel::ES3_1) // Skin VF does not use GPUScene on mobile
+	if (UseGPUScene(GMaxRHIShaderPlatform, GMaxRHIFeatureLevel)
+		&& !PlatformGPUSceneUsesUniformBufferView(GMaxRHIShaderPlatform))
 	{
 		OutElements.Add(FVertexElement(VertexStreams.Num(), 0, VET_UInt, 16, 0, true));
 	}
@@ -952,7 +951,7 @@ void TGPUSkinVertexFactory<BoneInfluenceType>::AddVertexElements(FVertexDeclarat
 	check(Data.IsValid());
 	GetVertexElements(GetFeatureLevel(), EVertexInputStreamType::Default, *Data, OutElements, Streams, MorphDeltaStreamIndex);
 
-	AddPrimitiveIdStreamElement(EVertexInputStreamType::Default, OutElements, 16, 0xff);
+	AddPrimitiveIdStreamElement(EVertexInputStreamType::Default, OutElements, 16, 16);
 }
 
 /**
diff --git a/Engine/Source/Runtime/Engine/Private/SkeletalMeshSceneProxy.cpp b/Engine/Source/Runtime/Engine/Private/SkeletalMeshSceneProxy.cpp
--- a/Engine/Source/Runtime/Engine/Private/SkeletalMeshSceneProxy.cpp
+++ b/Engine/Source/Runtime/Engine/Private/SkeletalMeshSceneProxy.cpp
@@ -273,11 +273,7 @@ FSkeletalMeshSceneProxy::FSkeletalMeshSceneProxy(const FSkinnedMeshSceneProxyDes
 		ShadowCapsuleBoneIndices.Sort();
 	}
 
-	if (FeatureLevel > ERHIFeatureLevel::ES3_1) // TODO: support gpuscene for a mobile SkinVF
-	{
-		EnableGPUSceneSupportFlags();
-
-	}
+	EnableGPUSceneSupportFlags();
 
 	if (IsAllowingApproximateOcclusionQueries())
 	{		
