Refactored World Line Batchers into an array, so these can be processed in for loops
Added a flush function to the World class, so users of the line batchers does not have to copy paste the flush code
Deprecated the public line batchers.
Added persistent foreground debug line drawer to World, enabling foreground debug to have either persistence or lifetime in seconds (this enables being able to see the debug data when pausung PIE)
Updated various files accessing directly the line batchers from the World

diff --git a/Engine/Plugins/Cameras/GameplayCameras/Source/GameplayCameras/Private/Debug/CameraDebugRenderer.cpp b/Engine/Plugins/Cameras/GameplayCameras/Source/GameplayCameras/Private/Debug/CameraDebugRenderer.cpp
--- a/Engine/Plugins/Cameras/GameplayCameras/Source/GameplayCameras/Private/Debug/CameraDebugRenderer.cpp
+++ b/Engine/Plugins/Cameras/GameplayCameras/Source/GameplayCameras/Private/Debug/CameraDebugRenderer.cpp
@@ -542,7 +542,7 @@ void FCameraDebugRenderer::DrawText(const FVector3d& WorldPosition, const FVecto
 
 ULineBatchComponent* FCameraDebugRenderer::GetDebugLineBatcher() const
 {
-	return World ? World->ForegroundLineBatcher : nullptr;
+	return World ? World->GetLineBatcher(UWorld::ELineBatcherType::Foreground) : nullptr;
 }
 
 void FCameraDebugRenderer::SkipAttachedBlocks()
diff --git a/Engine/Plugins/Experimental/Avalanche/Source/Avalanche/Private/Framework/AvaGameViewportClient.cpp b/Engine/Plugins/Experimental/Avalanche/Source/Avalanche/Private/Framework/AvaGameViewportClient.cpp
--- a/Engine/Plugins/Experimental/Avalanche/Source/Avalanche/Private/Framework/AvaGameViewportClient.cpp
+++ b/Engine/Plugins/Experimental/Avalanche/Source/Avalanche/Private/Framework/AvaGameViewportClient.cpp
@@ -295,14 +295,8 @@ void UAvaGameViewportClient::Draw(FViewport* InViewport, FCanvas* InCanvas)
 	GEngine->EmitDynamicResolutionEvent(EDynamicResolutionStateEvent::EndDynamicResolutionRendering);
 
 	// Remove temporary debug lines.
-	if (World->LineBatcher)
-	{
-		World->LineBatcher->Flush();
-	}
-	if (World->ForegroundLineBatcher)
-	{
-		World->ForegroundLineBatcher->Flush();
-	}
+	constexpr const UWorld::ELineBatcherType LineBatchersToFlush[] = { UWorld::ELineBatcherType::World, UWorld::ELineBatcherType::Foreground };
+	World->FlushLineBatchers(LineBatchersToFlush);
 
 	// Render Stats HUD in the main canvas so that it gets captured
 	// and is displayed in the broadcast channel's outputs.
diff --git a/Engine/Plugins/FX/Niagara/Source/Niagara/Private/NiagaraDebugHud.cpp b/Engine/Plugins/FX/Niagara/Source/Niagara/Private/NiagaraDebugHud.cpp
--- a/Engine/Plugins/FX/Niagara/Source/Niagara/Private/NiagaraDebugHud.cpp
+++ b/Engine/Plugins/FX/Niagara/Source/Niagara/Private/NiagaraDebugHud.cpp
@@ -614,7 +614,7 @@ namespace NiagaraDebugLocal
 
 	void DrawWireframeBox(UWorld* World, const FVector& Location, const FVector& Extents, const FLinearColor& Color, float Thickness = 3.0f)
 	{
-		if (ULineBatchComponent* LineBatcher = World->LineBatcher)
+		if (ULineBatchComponent* LineBatcher = World->GetLineBatcher(UWorld::ELineBatcherType::World))
 		{
 			LineBatcher->DrawLine(Location + FVector(Extents.X, Extents.Y, Extents.Z), Location + FVector(Extents.X, -Extents.Y, Extents.Z), Color, 0, Thickness);
 			LineBatcher->DrawLine(Location + FVector(Extents.X, -Extents.Y, Extents.Z), Location + FVector(-Extents.X, -Extents.Y, Extents.Z), Color, 0, Thickness);
@@ -635,7 +635,7 @@ namespace NiagaraDebugLocal
 
 	void DrawBox(UWorld* World, const FVector& Location, const FVector& Extents, const FLinearColor& Color)
 	{
-		if (ULineBatchComponent* LineBatcher = World->LineBatcher)
+		if (ULineBatchComponent* LineBatcher = World->GetLineBatcher(UWorld::ELineBatcherType::World))
 		{
 			const FBox BoundsBox(-Extents, Extents);
 			FColor BoxColor = Color.ToFColor(false);
diff --git a/Engine/Plugins/FX/Niagara/Source/Niagara/Private/NiagaraRendererDecals.cpp b/Engine/Plugins/FX/Niagara/Source/Niagara/Private/NiagaraRendererDecals.cpp
--- a/Engine/Plugins/FX/Niagara/Source/Niagara/Private/NiagaraRendererDecals.cpp
+++ b/Engine/Plugins/FX/Niagara/Source/Niagara/Private/NiagaraRendererDecals.cpp
@@ -47,29 +47,28 @@ namespace NiagaraRendererDecalsLocal
 			ENamedThreads::GameThread,
 			[WeakWorld=MakeWeakObjectPtr(World), AllUpdateParams_GT=TArray<FDeferredDecalUpdateParams>(AllUpdateParams), WorldTime]
 			{
-				UWorld* World = WeakWorld.Get();
-				if (World == nullptr || World->LineBatcher == nullptr)
+				if (UWorld* World = WeakWorld.Get())
 				{
-					return;
-				}
-				ULineBatchComponent* LineBatcher = World->LineBatcher;
-				
-				for (const FDeferredDecalUpdateParams& UpdateParams : AllUpdateParams_GT)
-				{
-					if ( UpdateParams.OperationType == FDeferredDecalUpdateParams::EOperationType::RemoveFromSceneAndDelete )
+					if (ULineBatchComponent* LineBatcher = World->GetLineBatcher(UWorld::ELineBatcherType::World))
 					{
-						continue;
+						for (const FDeferredDecalUpdateParams& UpdateParams : AllUpdateParams_GT)
+						{
+							if ( UpdateParams.OperationType == FDeferredDecalUpdateParams::EOperationType::RemoveFromSceneAndDelete )
+							{
+								continue;
+							}
+
+							const float Fade = (WorldTime - UpdateParams.FadeStartDelay) / UpdateParams.FadeDuration;
+							const FBox BoundsBox(-FVector(1), FVector(1));
+							//const uint8 IntFade = FMath::Clamp(uint8(255.0f * Fade), 0, 255);
+							const FColor DecalColor = UpdateParams.DecalColor.ToFColor(false);
+							LineBatcher->DrawSolidBox(BoundsBox, UpdateParams.Transform, DecalColor, 0, 0.0f);
+
+							LineBatcher->DrawCircle(UpdateParams.Bounds.Origin, FVector::XAxisVector, FVector::YAxisVector, FColor::Red, UpdateParams.Bounds.SphereRadius, 16, 0);
+							LineBatcher->DrawCircle(UpdateParams.Bounds.Origin, FVector::XAxisVector, FVector::ZAxisVector, FColor::Red, UpdateParams.Bounds.SphereRadius, 16, 0);
+							LineBatcher->DrawCircle(UpdateParams.Bounds.Origin, FVector::YAxisVector, FVector::ZAxisVector, FColor::Red, UpdateParams.Bounds.SphereRadius, 16, 0);
+						}
 					}
-
-					const float Fade = (WorldTime - UpdateParams.FadeStartDelay) / UpdateParams.FadeDuration;
-					const FBox BoundsBox(-FVector(1), FVector(1));
-					//const uint8 IntFade = FMath::Clamp(uint8(255.0f * Fade), 0, 255);
-					const FColor DecalColor = UpdateParams.DecalColor.ToFColor(false);
-					LineBatcher->DrawSolidBox(BoundsBox, UpdateParams.Transform, DecalColor, 0, 0.0f);
-
-					LineBatcher->DrawCircle(UpdateParams.Bounds.Origin, FVector::XAxisVector, FVector::YAxisVector, FColor::Red, UpdateParams.Bounds.SphereRadius, 16, 0);
-					LineBatcher->DrawCircle(UpdateParams.Bounds.Origin, FVector::XAxisVector, FVector::ZAxisVector, FColor::Red, UpdateParams.Bounds.SphereRadius, 16, 0);
-					LineBatcher->DrawCircle(UpdateParams.Bounds.Origin, FVector::YAxisVector, FVector::ZAxisVector, FColor::Red, UpdateParams.Bounds.SphereRadius, 16, 0);
 				}
 			}
 		);
diff --git a/Engine/Plugins/FX/Niagara/Source/Niagara/Private/Stateless/NiagaraStatelessDrawDebugContext.cpp b/Engine/Plugins/FX/Niagara/Source/Niagara/Private/Stateless/NiagaraStatelessDrawDebugContext.cpp
--- a/Engine/Plugins/FX/Niagara/Source/Niagara/Private/Stateless/NiagaraStatelessDrawDebugContext.cpp
+++ b/Engine/Plugins/FX/Niagara/Source/Niagara/Private/Stateless/NiagaraStatelessDrawDebugContext.cpp
@@ -65,7 +65,7 @@ void FNiagaraStatelessDrawDebugContext::DrawCylinder(float CylinderHeight, float
 
 void FNiagaraStatelessDrawDebugContext::DrawCircle(const FVector& InCenter, const float Radius, const FColor& Color) const
 {
-	if (ULineBatchComponent* LineBatcher = World->LineBatcher)
+	if (ULineBatchComponent* LineBatcher = World->GetLineBatcher(UWorld::ELineBatcherType::World))
 	{
 		int32 Segments = 16;
 		const float AngleStep = 2.f * UE_PI / float(Segments);
diff --git a/Engine/Plugins/Runtime/nDisplay/Source/DisplayCluster/Private/Game/EngineClasses/Basics/DisplayClusterViewportClient.cpp b/Engine/Plugins/Runtime/nDisplay/Source/DisplayCluster/Private/Game/EngineClasses/Basics/DisplayClusterViewportClient.cpp
--- a/Engine/Plugins/Runtime/nDisplay/Source/DisplayCluster/Private/Game/EngineClasses/Basics/DisplayClusterViewportClient.cpp
+++ b/Engine/Plugins/Runtime/nDisplay/Source/DisplayCluster/Private/Game/EngineClasses/Basics/DisplayClusterViewportClient.cpp
@@ -807,15 +807,8 @@ void UDisplayClusterViewportClient::Draw(FViewport* InViewport, FCanvas* SceneCa
 	MyWorld->UpdateLevelStreaming();
 
 	// Remove temporary debug lines.
-	if (MyWorld->LineBatcher != nullptr)
-	{
-		MyWorld->LineBatcher->Flush();
-	}
-
-	if (MyWorld->ForegroundLineBatcher != nullptr)
-	{
-		MyWorld->ForegroundLineBatcher->Flush();
-	}
+	constexpr const UWorld::ELineBatcherType LineBatchersToFlush[] = { UWorld::ELineBatcherType::World, UWorld::ELineBatcherType::Foreground };
+	World->FlushLineBatchers(LineBatchersToFlush);
 
 	// Draw FX debug information.
 	if (MyWorld->FXSystem)
diff --git a/Engine/Plugins/Runtime/nDisplay/Source/DisplayClusterConfigurator/Private/Views/Viewport/DisplayClusterConfiguratorSCSEditorViewportClient.cpp b/Engine/Plugins/Runtime/nDisplay/Source/DisplayClusterConfigurator/Private/Views/Viewport/DisplayClusterConfiguratorSCSEditorViewportClient.cpp
--- a/Engine/Plugins/Runtime/nDisplay/Source/DisplayClusterConfigurator/Private/Views/Viewport/DisplayClusterConfiguratorSCSEditorViewportClient.cpp
+++ b/Engine/Plugins/Runtime/nDisplay/Source/DisplayClusterConfigurator/Private/Views/Viewport/DisplayClusterConfiguratorSCSEditorViewportClient.cpp
@@ -293,11 +293,9 @@ void FDisplayClusterConfiguratorSCSEditorViewportClient::Draw(const FSceneView*
 	FEditorViewportClient::Draw(View, PDI);
 
 	// mimic the behaviour of the FEditorViewportClient call used in editor viewport
-	if (PreviewScene->GetWorld()->LineBatcher != NULL && (PreviewScene->GetWorld()->LineBatcher->BatchedLines.Num() || PreviewScene->GetWorld()->LineBatcher->BatchedPoints.Num() || PreviewScene->GetWorld()->LineBatcher->BatchedMeshes.Num()))
-	{
-		PreviewScene->GetWorld()->LineBatcher->Flush();
-	}
-	
+	constexpr const UWorld::ELineBatcherType LineBatchersToFlush[] = { UWorld::ELineBatcherType::World, UWorld::ELineBatcherType::Foreground };
+	PreviewScene->GetWorld()->FlushLineBatchers(LineBatchersToFlush);
+
 	bool bHitTesting = PDI->IsHitTesting();
 	AActor* PreviewActor = GetPreviewActor();
 	if (PreviewActor)
diff --git a/Engine/Plugins/Runtime/nDisplay/Source/DisplayClusterLightCardEditor/Private/Viewport/DisplayClusterLightCardEditorViewportClient.cpp b/Engine/Plugins/Runtime/nDisplay/Source/DisplayClusterLightCardEditor/Private/Viewport/DisplayClusterLightCardEditorViewportClient.cpp
--- a/Engine/Plugins/Runtime/nDisplay/Source/DisplayClusterLightCardEditor/Private/Viewport/DisplayClusterLightCardEditorViewportClient.cpp
+++ b/Engine/Plugins/Runtime/nDisplay/Source/DisplayClusterLightCardEditor/Private/Viewport/DisplayClusterLightCardEditorViewportClient.cpp
@@ -357,15 +357,8 @@ void FDisplayClusterLightCardEditorViewportClient::Draw(FViewport* InViewport, F
 
 	// Remove temporary debug lines.
 	// Possibly a hack. Lines may get added without the scene being rendered etc.
-	if (World->LineBatcher != NULL && (World->LineBatcher->BatchedLines.Num() || World->LineBatcher->BatchedPoints.Num() || World->LineBatcher->BatchedMeshes.Num() ) )
-	{
-		World->LineBatcher->Flush();
-	}
-
-	if (World->ForegroundLineBatcher != NULL && (World->ForegroundLineBatcher->BatchedLines.Num() || World->ForegroundLineBatcher->BatchedPoints.Num() || World->ForegroundLineBatcher->BatchedMeshes.Num() ) )
-	{
-		World->ForegroundLineBatcher->Flush();
-	}
+	constexpr const UWorld::ELineBatcherType LineBatchersToFlush[] = { UWorld::ELineBatcherType::World, UWorld::ELineBatcherType::Foreground };
+	World->FlushLineBatchers(LineBatchersToFlush);
 
 	// Draw the widget.
 	/*if (Widget && bShowWidget)
diff --git a/Engine/Source/Developer/DrawPrimitiveDebugger/Private/SDrawPrimitiveDebugger.cpp b/Engine/Source/Developer/DrawPrimitiveDebugger/Private/SDrawPrimitiveDebugger.cpp
--- a/Engine/Source/Developer/DrawPrimitiveDebugger/Private/SDrawPrimitiveDebugger.cpp
+++ b/Engine/Source/Developer/DrawPrimitiveDebugger/Private/SDrawPrimitiveDebugger.cpp
@@ -1550,7 +1550,7 @@ void SDrawPrimitiveDebugger::SetShowDebugBoundsForEntry(FPrimitiveComponentId En
 		FPrimitiveDebuggerEntry* Entry = Entries.Find(EntryId);
 		if (Entry && bShowDebugBounds != Entry->bShowingDebugBounds)
 		{
-			if (ULineBatchComponent* const LineBatcher = ActiveWorld->PersistentLineBatcher)
+			if (ULineBatchComponent* const LineBatcher = ActiveWorld->GetLineBatcher(UWorld::ELineBatcherType::WorldPersistent))
 			{
 				if (bShowDebugBounds && Entry->Data.IsValid() && Entry->Data->IsPrimitiveValid())
 				{
@@ -1589,7 +1589,7 @@ void SDrawPrimitiveDebugger::RedrawAllDebugBounds() const
 #if PRIMITIVE_DEBUGGER_SUPPORT_DEBUG_VISUALIZATIONS
 	if (ActiveWorld.IsValid())
 	{
-		if (ULineBatchComponent* const LineBatcher = ActiveWorld->PersistentLineBatcher)
+		if (ULineBatchComponent* const LineBatcher = ActiveWorld->GetLineBatcher(UWorld::ELineBatcherType::WorldPersistent))
 		{
 			for (const FPrimitiveComponentId EntryId : EntriesShowingDebugBounds)
 			{
@@ -1614,7 +1614,7 @@ void SDrawPrimitiveDebugger::FlushAllDebugBounds()
 #if PRIMITIVE_DEBUGGER_SUPPORT_DEBUG_VISUALIZATIONS
 	if (ActiveWorld.IsValid())
 	{
-		if (ULineBatchComponent* const LineBatcher = ActiveWorld->PersistentLineBatcher)
+		if (ULineBatchComponent* const LineBatcher = ActiveWorld->GetLineBatcher(UWorld::ELineBatcherType::WorldPersistent))
 		{
 			for (const FPrimitiveComponentId Entry : EntriesShowingDebugBounds)
 			{
@@ -1740,7 +1740,7 @@ void SDrawPrimitiveDebugger::ResetDebuggerChanges()
 		{
 			if (ActiveWorld.IsValid())
 			{
-				if (ULineBatchComponent* const LineBatcher = ActiveWorld->PersistentLineBatcher)
+				if (ULineBatchComponent* const LineBatcher = ActiveWorld->GetLineBatcher(UWorld::ELineBatcherType::WorldPersistent))
 				{
 					LineBatcher->ClearBatch(PrimitiveId.PrimIDValue);
 				}
diff --git a/Engine/Source/Editor/UnrealEd/Private/EditorViewportClient.cpp b/Engine/Source/Editor/UnrealEd/Private/EditorViewportClient.cpp
--- a/Engine/Source/Editor/UnrealEd/Private/EditorViewportClient.cpp
+++ b/Engine/Source/Editor/UnrealEd/Private/EditorViewportClient.cpp
@@ -4465,17 +4465,12 @@ void FEditorViewportClient::Draw(FViewport* InViewport, FCanvas* Canvas)
 
 	// Remove temporary debug lines.
 	// Possibly a hack. Lines may get added without the scene being rendered etc.
-	if (World && World->LineBatcher != NULL && (World->LineBatcher->BatchedLines.Num() || World->LineBatcher->BatchedPoints.Num() || World->LineBatcher->BatchedMeshes.Num() ) )
+	if (World)
 	{
-		World->LineBatcher->Flush();
+		constexpr const UWorld::ELineBatcherType LineBatchersToFlush[] = { UWorld::ELineBatcherType::World, UWorld::ELineBatcherType::Foreground };
+		World->FlushLineBatchers(LineBatchersToFlush);
 	}
 
-	if (World && World->ForegroundLineBatcher != NULL && (World->ForegroundLineBatcher->BatchedLines.Num() || World->ForegroundLineBatcher->BatchedPoints.Num() || World->ForegroundLineBatcher->BatchedMeshes.Num() ) )
-	{
-		World->ForegroundLineBatcher->Flush();
-	}
-
-
 	// Draw the widget.
 	if (Widget && bShowWidget)
 	{
diff --git a/Engine/Source/Runtime/Engine/Classes/Engine/World.h b/Engine/Source/Runtime/Engine/Classes/Engine/World.h
--- a/Engine/Source/Runtime/Engine/Classes/Engine/World.h
+++ b/Engine/Source/Runtime/Engine/Classes/Engine/World.h
@@ -899,16 +899,19 @@ class ENGINE_API UWorld final : public UObject, public FNetworkNotify
 	TObjectPtr<class UNetDriver>							NetDriver;
 
 	/** Line Batchers. All lines to be drawn in the world. */
+	UE_DEPRECATED(5.6, "Please use GetLineBatcher(UWorld::ELineBatcherType::World")
 	UPROPERTY(Transient)
-	TObjectPtr<class ULineBatchComponent>					LineBatcher;
+	TObjectPtr<class ULineBatchComponent>					LineBatcher_DEPRECATED;
 
 	/** Persistent Line Batchers. They don't get flushed every frame.  */
+	UE_DEPRECATED(5.6, "Please use GetLineBatcher(UWorld::ELineBatcherType::WorldPersistent")
 	UPROPERTY(Transient)
-	TObjectPtr<class ULineBatchComponent>					PersistentLineBatcher;
+	TObjectPtr<class ULineBatchComponent>					PersistentLineBatcher_DEPRECATED;
 
 	/** Foreground Line Batchers. This can't be Persistent.  */
+	UE_DEPRECATED(5.6, "Please use GetLineBatcher(UWorld::ELineBatcherType::Foreground")
 	UPROPERTY(Transient)
-	TObjectPtr<class ULineBatchComponent>					ForegroundLineBatcher;
+	TObjectPtr<class ULineBatchComponent>					ForegroundLineBatcher_DEPRECATED;
 
 	/** Instance of this world's game-specific networking management */
 	UPROPERTY(Transient)
@@ -947,6 +950,22 @@ private:
 	TObjectPtr<AServerStreamingLevelsVisibility> ServerStreamingLevelsVisibility;
 
 public:
+	enum class ELineBatcherType : uint8
+	{
+		World = 0,
+		WorldPersistent,
+		Foreground,
+		ForegroundPersistent,
+		// --- ---
+		NUM,
+		Invalid = NUM
+	};
+
+	/** Returns a pointer to the specified line batcher type */
+	ULineBatchComponent* GetLineBatcher(ELineBatcherType Type) const;
+	
+	/** Flush the specified Line Batcher types */
+	void FlushLineBatchers(const TArrayView<const UWorld::ELineBatcherType>& TypesToFlush);
 
 	/** Returns whether the world supports for a client to use "making visible" transaction requests to the server. */
 	bool SupportsMakingVisibleTransactionRequests() const;
@@ -1058,6 +1077,9 @@ public:
 #endif
 
 private:
+	/** Line Batchers. Persistent types don't get flushed every frame. */
+	UPROPERTY(Transient)
+	TObjectPtr<class ULineBatchComponent> LineBatchers[(int32)ELineBatcherType::NUM];
 
 	/** Sets level's PendingVisibilityState */
 	void SetLevelPendingVisibilityState(ULevel* InLevel, ELevelPendingVisibilityState InState);
diff --git a/Engine/Source/Runtime/Engine/Private/DrawDebugHelpers.cpp b/Engine/Source/Runtime/Engine/Private/DrawDebugHelpers.cpp
--- a/Engine/Source/Runtime/Engine/Private/DrawDebugHelpers.cpp
+++ b/Engine/Source/Runtime/Engine/Private/DrawDebugHelpers.cpp
@@ -93,9 +93,16 @@ void FlushPersistentDebugLines( const UWorld* InWorld )
 {
 	if (GEngine->GetNetMode(InWorld) != NM_DedicatedServer)
 	{
-		if (InWorld && InWorld->PersistentLineBatcher)
+		if (InWorld)
 		{
-			InWorld->PersistentLineBatcher->Flush();
+			if (ULineBatchComponent* WorldPersistentLineBatcher = InWorld->GetLineBatcher(UWorld::ELineBatcherType::WorldPersistent))
+			{
+				WorldPersistentLineBatcher->Flush();
+			}
+			if (ULineBatchComponent* ForegroundPersistentLineBatcher = InWorld->GetLineBatcher(UWorld::ELineBatcherType::ForegroundPersistent))
+			{
+				ForegroundPersistentLineBatcher->Flush();
+			}
 		}
 	}
 #if WITH_EDITOR
@@ -117,7 +124,19 @@ void FlushPersistentDebugLines( const UWorld* InWorld )
 
 ULineBatchComponent* GetDebugLineBatcher( const UWorld* InWorld, bool bPersistentLines, float LifeTime, bool bDepthIsForeground )
 {
-	return (InWorld ? (bDepthIsForeground ? InWorld->ForegroundLineBatcher : (( bPersistentLines || (LifeTime > 0.f) ) ? InWorld->PersistentLineBatcher : InWorld->LineBatcher)) : nullptr);
+	if (InWorld)
+	{
+		if (bPersistentLines || LifeTime > 0.f)
+		{
+			return bDepthIsForeground ? InWorld->GetLineBatcher(UWorld::ELineBatcherType::ForegroundPersistent) : InWorld->GetLineBatcher(UWorld::ELineBatcherType::WorldPersistent);
+		}
+		else
+		{
+			return bDepthIsForeground ? InWorld->GetLineBatcher(UWorld::ELineBatcherType::Foreground) : InWorld->GetLineBatcher(UWorld::ELineBatcherType::World);
+		}
+	}
+
+	return nullptr;
 }
 
 static float GetDebugLineLifeTime(ULineBatchComponent* LineBatcher, float LifeTime, bool bPersistent)
diff --git a/Engine/Source/Runtime/Engine/Private/GameViewportClient.cpp b/Engine/Source/Runtime/Engine/Private/GameViewportClient.cpp
--- a/Engine/Source/Runtime/Engine/Private/GameViewportClient.cpp
+++ b/Engine/Source/Runtime/Engine/Private/GameViewportClient.cpp
@@ -1909,15 +1909,8 @@ void UGameViewportClient::Draw(FViewport* InViewport, FCanvas* SceneCanvas)
 	}
 
 	// Remove temporary debug lines.
-	if (MyWorld->LineBatcher != nullptr)
-	{
-		MyWorld->LineBatcher->Flush();
-	}
-
-	if (MyWorld->ForegroundLineBatcher != nullptr)
-	{
-		MyWorld->ForegroundLineBatcher->Flush();
-	}
+	constexpr const UWorld::ELineBatcherType LineBatchersToFlush[] = { UWorld::ELineBatcherType::World, UWorld::ELineBatcherType::Foreground };
+	MyWorld->FlushLineBatchers(LineBatchersToFlush);
 
 	// Draw FX debug information.
 	if (MyWorld->FXSystem)
diff --git a/Engine/Source/Runtime/Engine/Private/HUD.cpp b/Engine/Source/Runtime/Engine/Private/HUD.cpp
--- a/Engine/Source/Runtime/Engine/Private/HUD.cpp
+++ b/Engine/Source/Runtime/Engine/Private/HUD.cpp
@@ -87,7 +87,10 @@ void AHUD::SetCanvas(class UCanvas* InCanvas, class UCanvas* InDebugCanvas)
 
 void AHUD::Draw3DLine(FVector Start, FVector End, FColor LineColor)
 {
-	GetWorld()->LineBatcher->DrawLine(Start, End, LineColor, SDPG_World);
+	if (ULineBatchComponent* LineBatcher = GetWorld()->GetLineBatcher(UWorld::ELineBatcherType::World))
+	{
+		LineBatcher->DrawLine(Start, End, LineColor, SDPG_World);
+	}
 }
 
 void AHUD::Draw2DLine(int32 X1, int32 Y1, int32 X2, int32 Y2, FColor LineColor)
diff --git a/Engine/Source/Runtime/Engine/Private/PhysicsEngine/PhysAnim.cpp b/Engine/Source/Runtime/Engine/Private/PhysicsEngine/PhysAnim.cpp
--- a/Engine/Source/Runtime/Engine/Private/PhysicsEngine/PhysAnim.cpp
+++ b/Engine/Source/Runtime/Engine/Private/PhysicsEngine/PhysAnim.cpp
@@ -594,14 +594,17 @@ void USkeletalMeshComponent::UpdateKinematicBonesToAnim(const TArray<FTransform>
 	// If desired, draw the skeleton at the point where we pass it to the physics.
 	if (bShowPrePhysBones && GetSkeletalMeshAsset() && InSpaceBases.Num() == GetSkeletalMeshAsset()->GetRefSkeleton().GetNum())
 	{
-		for (int32 i = 1; i<InSpaceBases.Num(); i++)
+		if (ULineBatchComponent* LineBatcher = World->GetLineBatcher(UWorld::ELineBatcherType::World))
 		{
-			FVector ThisPos = CurrentLocalToWorld.TransformPosition(InSpaceBases[i].GetLocation());
+			for (int32 i = 1; i<InSpaceBases.Num(); i++)
+			{
+				FVector ThisPos = CurrentLocalToWorld.TransformPosition(InSpaceBases[i].GetLocation());
 
-			int32 ParentIndex = GetSkeletalMeshAsset()->GetRefSkeleton().GetParentIndex(i);
-			FVector ParentPos = CurrentLocalToWorld.TransformPosition(InSpaceBases[ParentIndex].GetLocation());
+				int32 ParentIndex = GetSkeletalMeshAsset()->GetRefSkeleton().GetParentIndex(i);
+				FVector ParentPos = CurrentLocalToWorld.TransformPosition(InSpaceBases[ParentIndex].GetLocation());
 
-			World->LineBatcher->DrawLine(ThisPos, ParentPos, AnimSkelDrawColor, SDPG_Foreground);
+				LineBatcher->DrawLine(ThisPos, ParentPos, AnimSkelDrawColor, SDPG_Foreground);
+			}
 		}
 	}
 #endif
diff --git a/Engine/Source/Runtime/Engine/Private/World.cpp b/Engine/Source/Runtime/Engine/Private/World.cpp
--- a/Engine/Source/Runtime/Engine/Private/World.cpp
+++ b/Engine/Source/Runtime/Engine/Private/World.cpp
@@ -789,9 +789,11 @@ void UWorld::Serialize( FArchive& Ar )
 
 		Ar << NetDriver;
 		
-		Ar << LineBatcher;
-		Ar << PersistentLineBatcher;
-		Ar << ForegroundLineBatcher;  
+		for (TObjectPtr<ULineBatchComponent>& LineBatcher : LineBatchers)
+		{
+			Ar << LineBatcher;
+		}
+
 		Ar << PhysicsField;
 
 		Ar << MyParticleEventManager;
@@ -860,9 +862,10 @@ void UWorld::AddReferencedObjects(UObject* InThis, FReferenceCollector& Collecto
 		Collector.AddReferencedObject( This->CurrentLevel, This );
 		Collector.AddReferencedObject( This->NetDriver, This );
 		Collector.AddReferencedObject( This->DemoNetDriver, This );
-		Collector.AddReferencedObject( This->LineBatcher, This );
-		Collector.AddReferencedObject( This->PersistentLineBatcher, This );
-		Collector.AddReferencedObject( This->ForegroundLineBatcher, This );
+		for (TObjectPtr<ULineBatchComponent>& LineBatcher : This->LineBatchers)
+		{
+			Collector.AddReferencedObject(LineBatcher, This);
+		}
 		Collector.AddReferencedObject( This->PhysicsField, This);
 		Collector.AddReferencedObject( This->MyParticleEventManager, This );
 		Collector.AddReferencedObject( This->GameState, This );
@@ -2663,19 +2666,12 @@ void UWorld::ClearWorldComponents()
 		Level->ClearLevelComponents();
 	}
 
-	if(LineBatcher && LineBatcher->IsRegistered())
-	{
-		LineBatcher->UnregisterComponent();
-	}
-
-	if(PersistentLineBatcher && PersistentLineBatcher->IsRegistered())
-	{
-		PersistentLineBatcher->UnregisterComponent();
-	}
-
-	if(ForegroundLineBatcher && ForegroundLineBatcher->IsRegistered())
+	for (TObjectPtr<ULineBatchComponent>& LineBatcher : LineBatchers)
 	{
-		ForegroundLineBatcher->UnregisterComponent();
+		if (LineBatcher && LineBatcher->IsRegistered())
+		{
+			LineBatcher->UnregisterComponent();
+		}
 	}
 
 	if (PhysicsField && PhysicsField->IsRegistered())
@@ -2695,37 +2691,21 @@ void UWorld::UpdateWorldComponents(bool bRerunConstructionScripts, bool bCurrent
 
 	if ( !IsRunningDedicatedServer() )
 	{
-		if(!LineBatcher)
+		for (TObjectPtr<ULineBatchComponent>& LineBatcher : LineBatchers)
 		{
-			LineBatcher = NewObject<ULineBatchComponent>();
-			LineBatcher->bCalculateAccurateBounds = false;
-		}
-
-		if(!LineBatcher->IsRegistered())
-		{	
-			LineBatcher->RegisterComponentWithWorld(this, Context);
-		}
-
-		if(!PersistentLineBatcher)
-		{
-			PersistentLineBatcher = NewObject<ULineBatchComponent>();
-			PersistentLineBatcher->bCalculateAccurateBounds = false;
-		}
-
-		if(!PersistentLineBatcher->IsRegistered())	
-		{
-			PersistentLineBatcher->RegisterComponentWithWorld(this, Context);
-		}
-
-		if(!ForegroundLineBatcher)
-		{
-			ForegroundLineBatcher = NewObject<ULineBatchComponent>();
-			ForegroundLineBatcher->bCalculateAccurateBounds = false;
+			if (!LineBatcher)
+			{
+				LineBatcher = NewObject<ULineBatchComponent>();
+				LineBatcher->bCalculateAccurateBounds = false;
+			}
 		}
 
-		if(!ForegroundLineBatcher->IsRegistered())	
+		for (TObjectPtr<ULineBatchComponent>& LineBatcher : LineBatchers)
 		{
-			ForegroundLineBatcher->RegisterComponentWithWorld(this, Context);
+			if (!LineBatcher->IsRegistered())
+			{
+				LineBatcher->RegisterComponentWithWorld(this, Context);
+			}
 		}
 
 		static IConsoleVariable* PhysicsFieldEnableClipmapCVar = IConsoleManager::Get().FindConsoleVariable(TEXT("r.PhysicsField.EnableField"));
@@ -3190,6 +3170,32 @@ static void ResetLevelFlagsOnLevelAddedToWorld(ULevel* Level)
 	Level->ResetRouteActorInitializationState();
 }
 
+ULineBatchComponent* UWorld::GetLineBatcher(ELineBatcherType Type) const
+{
+	ULineBatchComponent* LineBatcher = nullptr;
+
+	if (ensure(Type < ELineBatcherType::NUM))
+	{
+		LineBatcher = LineBatchers[(int32)Type].Get();
+	}
+
+	return LineBatcher;
+}
+
+void UWorld::FlushLineBatchers(const TArrayView<const UWorld::ELineBatcherType>& TypesToFlush)
+{
+	for (const UWorld::ELineBatcherType& BatcherType : TypesToFlush)
+	{
+		if (ULineBatchComponent* LineBatcher = GetLineBatcher(BatcherType))
+		{
+			if (LineBatcher->BatchedLines.Num() || LineBatcher->BatchedPoints.Num() || LineBatcher->BatchedMeshes.Num())
+			{
+				LineBatcher->Flush();
+			}
+		}
+	}
+}
+
 bool UWorld::SupportsMakingVisibleTransactionRequests() const
 {
 	if (GetWorld()->IsGameWorld())
@@ -5231,7 +5237,8 @@ bool UWorld::HandleTraceTagCommand( const TCHAR* Cmd, FOutputDevice& Ar )
 
 bool UWorld::HandleFlushPersistentDebugLinesCommand( const TCHAR* Cmd, FOutputDevice& Ar )
 {
-	PersistentLineBatcher->Flush();
+	GetLineBatcher(ELineBatcherType::WorldPersistent)->Flush();
+	GetLineBatcher(ELineBatcherType::ForegroundPersistent)->Flush();
 	return true;
 }
 
@@ -7427,19 +7434,12 @@ bool UWorld::SetNewWorldOrigin(FIntVector InNewOriginLocation)
 		}
 	}
 			
-	if (LineBatcher)
+	for (TObjectPtr<ULineBatchComponent>& LineBatcher : LineBatchers)
 	{
-		LineBatcher->ApplyWorldOffset(Offset, true);
-	}
-	
-	if (PersistentLineBatcher)
-	{
-		PersistentLineBatcher->ApplyWorldOffset(Offset, true);
-	}
-	
-	if (ForegroundLineBatcher)
-	{
-		ForegroundLineBatcher->ApplyWorldOffset(Offset, true);
+		if(LineBatcher)
+		{
+			LineBatcher->ApplyWorldOffset(Offset, true);
+		}
 	}
 
 	if (PhysicsField)
diff --git a/Engine/Source/Runtime/Engine/Private/WorldPartition/WorldPartitionRuntimeSpatialHash.cpp b/Engine/Source/Runtime/Engine/Private/WorldPartition/WorldPartitionRuntimeSpatialHash.cpp
--- a/Engine/Source/Runtime/Engine/Private/WorldPartition/WorldPartitionRuntimeSpatialHash.cpp
+++ b/Engine/Source/Runtime/Engine/Private/WorldPartition/WorldPartitionRuntimeSpatialHash.cpp
@@ -953,8 +953,7 @@ void FSpatialHashStreamingGrid::DrawStreamingSource3D(UWorld* World, const FSphe
 	}
 	else
 	{
-		ULineBatchComponent* const LineBatcher = World->LineBatcher;
-		if (LineBatcher)
+		if (ULineBatchComponent* const LineBatcher = World->GetLineBatcher(UWorld::ELineBatcherType::World))
 		{
 			FSphericalSector Shape = InShape;
 			Shape.SetAxis(InTransform.TransformVector(Shape.GetAxis()));
diff --git a/Engine/Source/Runtime/UMG/Private/Components/Viewport.cpp b/Engine/Source/Runtime/UMG/Private/Components/Viewport.cpp
--- a/Engine/Source/Runtime/UMG/Private/Components/Viewport.cpp
+++ b/Engine/Source/Runtime/UMG/Private/Components/Viewport.cpp
@@ -196,14 +196,10 @@ void FUMGViewportClient::Draw(FViewport* InViewport, FCanvas* Canvas)
 
 	// Remove temporary debug lines.
 	// Possibly a hack. Lines may get added without the scene being rendered etc.
-	if ( World->LineBatcher != NULL && ( World->LineBatcher->BatchedLines.Num() || World->LineBatcher->BatchedPoints.Num() ) )
+	if (World)
 	{
-		World->LineBatcher->Flush();
-	}
-
-	if ( World->ForegroundLineBatcher != NULL && ( World->ForegroundLineBatcher->BatchedLines.Num() || World->ForegroundLineBatcher->BatchedPoints.Num() ) )
-	{
-		World->ForegroundLineBatcher->Flush();
+		constexpr const UWorld::ELineBatcherType LineBatchersToFlush[] = { UWorld::ELineBatcherType::World, UWorld::ELineBatcherType::Foreground };
+		World->FlushLineBatchers(LineBatchersToFlush);
 	}
 	
 	Viewport = ViewportBackup;
