Added a debug tool in the new material translator to dump the translated material module IR to a text file. This dump is an abstract representation of the instructions that set the various material outputs, that will be further translated to HLSL in a later stage.

Enable with r.Material.Translator.DumpIR=1. Upon translation, a dump of the IR will be written to file MaterialIRDump.txt on the logs folder.

#rnx

diff --git a/Engine/Source/Runtime/Engine/Private/Materials/MaterialIRDebug.cpp b/Engine/Source/Runtime/Engine/Private/Materials/MaterialIRDebug.cpp
--- a/Engine/Source/Runtime/Engine/Private/Materials/MaterialIRDebug.cpp
+++ b/Engine/Source/Runtime/Engine/Private/Materials/MaterialIRDebug.cpp
@@ -4,46 +4,40 @@
 #include "Materials/MaterialIRModule.h"
 #include "Materials/MaterialIR.h"
 #include "Materials/MaterialIRTypes.h"
+#include "Materials/MaterialIRInternal.h"
 #include "Materials/MaterialAttributeDefinitionMap.h"
 #include "Misc/Paths.h"
 #include "Misc/FileHelper.h"
+#include "Engine/Texture.h"
 
 #if WITH_EDITOR
 
-static TAutoConsoleVariable<bool> CVarDumpMaterialIRUseGraph(
-	TEXT("r.Material.Translator.DumpUseGraph"),
-	true,
-	TEXT("Whether the material translator should emit the Module IR 'Uses' graph in Graphviz Dot syntax (to 'MaterialIRDumpGraph.dot')."),
-	ECVF_RenderThreadSafe);
-
-
 static TAutoConsoleVariable<bool> CVarDumpMaterialIRUseGraph_EnableNext(
 	TEXT("r.Material.Translator.DumpUseGraphOpts.EnableSuccessors"),
 	false,
 	TEXT("Whether the Material Module IR 'Uses' graph should also display 'Instruction Next' edges."),
 	ECVF_RenderThreadSafe);
 
-
 namespace UE::MIR
 {
 
-static void DumpValueInfo(const MIR::FValue* Value, FString& Out)
+static void DumpValueInfo(const FValue* Value, FString& Out)
 {
-	if (auto Constant = Value->As<MIR::FConstant>())
+	if (const FConstant* Constant = Value->As<FConstant>())
 	{
 		switch (Constant->Type->AsPrimitive()->ScalarKind)
 		{
-			case MIR::ScalarKind_Bool:	Out.Append(Constant->Boolean ? TEXT("true") : TEXT("false")); break;
-			case MIR::ScalarKind_Int:	Out.Appendf(TEXT("%lld"), Constant->Integer); break;
-			case MIR::ScalarKind_Float: Out.Appendf(TEXT("%f"), Constant->Float); break;
+			case ScalarKind_Bool:	Out.Append(Constant->Boolean ? TEXT("true") : TEXT("false")); break;
+			case ScalarKind_Int:	Out.Appendf(TEXT("%lld"), Constant->Integer); break;
+			case ScalarKind_Float: Out.Appendf(TEXT("%f"), Constant->Float); break;
 			default: UE_MIR_UNREACHABLE();
 		}
 	}
-	else if (auto ExternalInput = Value->As<MIR::FExternalInput>())
+	else if (const FExternalInput* ExternalInput = Value->As<FExternalInput>())
 	{
-		Out.Append(MIR::ExternalInputToString(ExternalInput->Id));
+		Out.Append(ExternalInputToString(ExternalInput->Id));
 	}
-	else if (auto SetMaterailOutput = Value->As<MIR::FSetMaterialOutput>())
+	else if (const FSetMaterialOutput* SetMaterailOutput = Value->As<FSetMaterialOutput>())
 	{
 		const FString& PropertyName = (SetMaterailOutput->Property == MP_SubsurfaceColor)
 			? TEXT("Subsurface")
@@ -51,7 +45,7 @@ static void DumpValueInfo(const MIR::FValue* Value, FString& Out)
 		
 		Out.Append(PropertyName);
 	}
-	else if (auto Subscript = Value->As<MIR::FSubscript>())
+	else if (const FSubscript* Subscript = Value->As<FSubscript>())
 	{
 		if (Subscript->Arg->Type->AsVector())
 		{
@@ -64,15 +58,15 @@ static void DumpValueInfo(const MIR::FValue* Value, FString& Out)
 			Out.Appendf(TEXT("Index: %d"), Subscript->Index);
 		}
 	}
-	else if (auto Operator = Value->As<MIR::FOperator>())
+	else if (const FOperator* Operator = Value->As<FOperator>())
 	{
-		Out.Append(MIR::OperatorToString(Operator->Op));
+		Out.Append(OperatorToString(Operator->Op));
 	}
 }				
 
-static void DumpUseInfo(const MIR::FValue* Value, const MIR::FValue* Use, int UseIndex, FString& Out)
+static void DumpUseInfo(const FValue* Value, const FValue* Use, int32 UseIndex, FString& Out)
 {
-	if (auto Dimensional = Value->As<MIR::FDimensional>())
+	if (const FDimensional* Dimensional = Value->As<FDimensional>())
 	{
 		if (Dimensional->Type->AsPrimitive()->IsVector())
 		{
@@ -84,12 +78,12 @@ static void DumpUseInfo(const MIR::FValue* Value, const MIR::FValue* Use, int Us
 			Out.AppendInt(UseIndex);
 		}
 	}
-	else if (auto If = Value->As<MIR::FBranch>())
+	else if (const FBranch* If = Value->As<FBranch>())
 	{
 		static const TCHAR* Uses[] = { TEXT("condition"), TEXT("true"), TEXT("false") };
 		Out.Append(Uses[UseIndex]);
 	}
-	else if (auto Operator = Value->As<MIR::FOperator>())
+	else if (const FOperator* Operator = Value->As<FOperator>())
 	{
 		static const TCHAR* Uses[] = { TEXT("a"), TEXT("b"), TEXT("c") };
 		Out.Append(Uses[UseIndex]);
@@ -98,14 +92,9 @@ static void DumpUseInfo(const MIR::FValue* Value, const MIR::FValue* Use, int Us
 
 void DebugDumpIRUseGraph(const FMaterialIRModule& Module)
 {
-	if (!CVarDumpMaterialIRUseGraph.GetValueOnAnyThread())
-	{
-		return;
-	}
-
 	FString Content;
-	TSet<const MIR::FValue*> Crawled;
-	TArray<const MIR::FValue*> ValueStack;
+	TSet<const FValue*> Crawled;
+	TArray<const FValue*> ValueStack;
 
 	Content.Appendf(TEXT(
 		"digraph G {\n\n"
@@ -114,7 +103,7 @@ void DebugDumpIRUseGraph(const FMaterialIRModule& Module)
 		"edge [fontname=\"Consolas\"]\n\n"
 	));
 
-	for (const UE::MIR::FSetMaterialOutput* Output : Module.GetOutputs())
+	for (const FSetMaterialOutput* Output : Module.GetOutputs())
 	{
 		ValueStack.Push(Output);
 	}
@@ -123,12 +112,12 @@ void DebugDumpIRUseGraph(const FMaterialIRModule& Module)
 
 	while (!ValueStack.IsEmpty())
 	{
-		const MIR::FValue* Value = ValueStack.Pop();
+		const FValue* Value = ValueStack.Pop();
 
 		// Begin the node declaration
 		Content.Appendf(TEXT("\"%p\" [label=< <b>%s</b>  (%s) <br/> "),
 						Value,
-						MIR::ValueKindToString(Value->Kind),
+						ValueKindToString(Value->Kind),
 						Value->Type ? Value->Type->GetSpelling().GetData() : TEXT("???"));
 
 		DumpValueInfo(Value, Content);
@@ -136,14 +125,14 @@ void DebugDumpIRUseGraph(const FMaterialIRModule& Module)
 		// End the node declaration
 		Content.Append(TEXT(">]\n"));
 
-		const MIR::FInstruction* Instr = Value->AsInstruction();
+		const FInstruction* Instr = Value->AsInstruction();
 		if (DumpInstructionSequence && Instr && Instr->Next)
 		{
 			Content.Appendf(TEXT("\"%p\" -> \"%p\" [color=\"red\"]\n"), Instr, Instr->Next);
 		}
 
-		int UseIndex = -1;
-		for (const MIR::FValue* Use : Value->GetUses())
+		int32 UseIndex = -1;
+		for (const FValue* Use : Value->GetUses())
 		{
 			++UseIndex;
 
@@ -166,7 +155,7 @@ void DebugDumpIRUseGraph(const FMaterialIRModule& Module)
 
 			if (DumpInstructionSequence && Instr)
 			{
-				const MIR::FInstruction* UseInstr = Use->AsInstruction();
+				const FInstruction* UseInstr = Use->AsInstruction();
 				if (UseInstr && UseInstr->Block != Instr->Block)
 				{
 					Content.Appendf(TEXT("\"%p\" -> \"%p\" [color=\"red\", style=\"dashed\"]\n"), UseInstr, Instr);
@@ -177,10 +166,264 @@ void DebugDumpIRUseGraph(const FMaterialIRModule& Module)
 
 	Content.Appendf(TEXT("\n}\n"));
 	
-	FString FilePath = FPaths::Combine(FPaths::ProjectLogDir(), TEXT("MaterialIRDumpGraph.dot"));
+	FString FilePath = FPaths::Combine(FPaths::ProjectSavedDir(), "Materials", TEXT("IRDumpUseGraph.dot"));
 	FFileHelper::SaveStringToFile(Content, *FilePath);
 }
 
+/* Module IR to textual representation dumping */
+
+/// Helper struct to wrap the state used during IR to text dumping.
+struct FDebugDumpIRState
+{
+	/// Output string containing the generated result.
+	FString Out; 
+
+	/// String used for temporary operations. Clear before use.
+	FString Temp;
+
+	/// Maps values to an incrementing id. Used to give values a "name" for future referencing (e.g. "%6")
+	TMap<const FValue*, uint32> ValueToIdMap;
+
+	/// Counter used to assign an id to encountered values.
+	uint32 InstrIdCounter = 0;
+
+	/// Array of encountered parameters. Used later on to generate the a recap of all referenced parameters.
+	TArray<TPair<uint32, const FValue*>> ReferencedParameters;
+
+	/// Prints a block of instructions to `Out`. Indentation indicates how many levels of indentation to put
+	/// to the left of printed instruction.
+	void AppendBlock(const FBlock& Block, int32 Indentation)
+	{
+		for (FInstruction* Instr = Block.Instructions; Instr; Instr = Instr->Next)
+		{
+			// Format the left column (e.g. "%4 = ") string if this instruction is referenceable.
+			Temp.Empty();
+			if (Instr->Kind != VK_SetMaterialOutput)
+			{
+				Temp.Appendf(TEXT("%%%llu = "), ReferenceInstruction(Instr));
+			}
+
+			// Print indentation, then Temp aligned to the right.
+			AppendLeftColumn(Indentation, Temp);
+
+			// Print the kind of the instruction (the opcode, e.g. "Operator")
+			Out.Append(ValueKindToString(Instr->Kind));
+
+			// Begin printing the arguments (used values)
+			Out.Append(TEXT(" ("));
+
+			bool bAddComma = false;
+			TArrayView<FValue*> Uses = Instr->GetUses();
+			for (int32 UseIndex = 0; UseIndex < Uses.Num(); ++UseIndex)
+			{
+				if (bAddComma)
+				{
+					Out.Append(TEXT(", "));
+				}
+				bAddComma = true;
+				
+				const FValue* Use = Uses[UseIndex];
+				if (!Use)
+				{
+					Out.Append(TEXT("null"));
+					continue;
+				}
+
+				// First the type...
+				Out.Appendf(TEXT("%s "), Use->Type->GetSpelling().GetData());
+
+				// If this use is in a block different from current's, dump the block in "{}" first.
+				const FBlock* UseBlock = Instr->GetDesiredBlockForUse(UseIndex);
+				if (UseBlock != Instr->Block && UseBlock->Instructions)
+				{
+					Out.Append(TEXT("{\n"));
+					AppendBlock(*UseBlock, Indentation + 1);
+
+					AppendLeftColumn(Indentation, TEXT(""));
+					Out.Append(TEXT("} "));
+				}
+				
+				// Finally, reference the used value (this will print "%x" if it's an
+				// instruction, or inline its information otherwise, like in constants).
+				AppendValueReference(Use);
+			}
+
+			Out.Append(TEXT(")"));
+
+			// Dump the instruction properties.
+			AppendInstructionProperties(Instr);
+
+			Out.Append(TEXT("\n"));
+		}
+	}
+
+	/// Appends extra information regarding the instruction.
+	void AppendInstructionProperties(const FInstruction* Instr)
+	{
+		if (const FSetMaterialOutput* SetMaterialOutput = Instr->As<FSetMaterialOutput>())
+		{
+			Out.Appendf(TEXT(" \"%s\""), *FMaterialAttributeDefinitionMap::GetAttributeName(SetMaterialOutput->Property));
+		}
+		else if (const FOperator* Operator = Instr->As<FOperator>())
+		{
+			Out.Appendf(TEXT(" \"%s\""), OperatorToString(Operator->Op));
+		}
+		else if (const FSubscript* Subscript = Instr->As<FSubscript>())
+		{
+			if (Subscript->Arg->Type->AsVector())
+			{
+				check(Subscript->Index < 4);
+				Out.Appendf(TEXT(" .%c"), TEXT("XYZW")[Subscript->Index]);
+			}
+			else
+			{
+				Out.Appendf(TEXT(" Index=%d"), Subscript->Index);
+			}
+		}
+		else if (const FTextureSample* TextureSample = Instr->As<FTextureSample>())
+		{
+			Out.Appendf(TEXT(" SamplerSourceMode=\"%s\""), *StaticEnum<ESamplerSourceMode>()->GetDisplayNameTextByValue(TextureSample->SamplerSourceMode).ToString());
+			Out.Appendf(TEXT(" MipValueMode=\"%s\""), *StaticEnum<ETextureMipValueMode>()->GetDisplayNameTextByValue(TextureSample->MipValueMode).ToString());
+			Out.Appendf(TEXT(" SamplerType=\"%s\""), *StaticEnum<EMaterialSamplerType>()->GetDisplayNameTextByValue(TextureSample->SamplerType).ToString());
+		}
+	}
+
+	/// Appends a reference to the specified value. This will look like "%x" if Value is an
+	/// instruction, otherwise it will inline information regarding the value.
+	void AppendValueReference(const FValue* Value)
+	{
+		if (ValueToIdMap.Contains(Value))
+		{
+			Out.Appendf(TEXT("%%%llu"), ValueToIdMap[Value]);
+			return;
+		}
+
+		if (Value->As<FPoison>())
+		{
+			Out.Append(TEXT("Poison"));
+		}
+		else if (const FConstant* Constant = Value->As<FConstant>())
+		{
+			if (Constant->IsBool())
+			{
+				Out.Append(Constant->Boolean ? TEXT("true") : TEXT("false"));
+			}
+			else if (Constant->IsInteger())
+			{
+				Out.Appendf(TEXT("%lld"), Constant->Integer);
+			}
+			else if (Constant->IsFloat())
+			{
+				Out.Appendf(TEXT("%.5ff"), Constant->Float);
+			}
+			else
+			{
+				UE_MIR_UNREACHABLE();
+			}
+		}
+		else if (const FExternalInput* ExternalInput = Value->As<FExternalInput>())
+		{
+			Out.Appendf(TEXT("[ExternalInput \"%s\"]"), ExternalInputToString(Value->As<FExternalInput>()->Id));
+		}
+		else if (const FTextureObject* TextureObject = Value->As<FTextureObject>())
+		{
+			Out.Appendf(TEXT("[TextureObject #%d SamplerType=\"%s\"]"), TextureObject->Analysis_UniformParameterIndex, *StaticEnum<EMaterialSamplerType>()->GetDisplayNameTextByValue(TextureObject->SamplerType).ToString());
+			ReferencedParameters.AddUnique({ TextureObject->Analysis_UniformParameterIndex, TextureObject });
+		}
+		else if (const FUniformParameter* UniformParameter = Value->As<FUniformParameter>())
+		{
+			Out.Appendf(TEXT("[Parameter #%d \"%s\"]"), UniformParameter->Analysis_UniformParameterIndex, *UniformParameter->Info.Name.ToString());
+			ReferencedParameters.AddUnique({ UniformParameter->Analysis_UniformParameterIndex, UniformParameter });
+		}
+		else
+		{
+			Out.Appendf(TEXT("[%s]"), ValueKindToString(Value->Kind));
+		}
+	}
+
+	/// Gets the instruction reference.
+	uint32 ReferenceInstruction(const FInstruction* Instr)
+	{
+		uint32 Id;
+		if (!Internal::Find(ValueToIdMap, static_cast<const FValue*>(Instr), Id))
+		{
+			Id = InstrIdCounter++;
+			ValueToIdMap.Add(Instr, Id);
+		}
+		return Id;
+	}
+	
+	void AppendLeftColumn(int32 Indentation, FStringView LeftColumn)
+	{
+		for (int32 i = 0; i < Indentation; ++i)
+		{
+			Out.Append("        ");
+		}
+		
+		// Put some padding so that all '=' are aligned to the right.
+		for (int32 i = 0, Spaces = 8 - LeftColumn.Len(); i < Spaces; ++i)
+		{
+			Out.Append(TEXT(" "));
+		}
+
+		Out += LeftColumn;
+	}
+
+	/// Prints the parameter recap section.
+	void DumpReferencedParameters()
+	{
+		if (ReferencedParameters.IsEmpty())
+		{
+			return;
+		}
+
+		// Dump the list of referenced parameters.
+		ReferencedParameters.Sort([](const auto& Param1, const auto& Param2)
+		{
+			return Param1.Key < Param2.Key;
+		});
+
+		Out.Append(TEXT("\n; Referenced material parameters\n"));
+		for (const auto& ParamPair : ReferencedParameters)
+		{
+			const FValue* Value = ParamPair.Value;
+			if (const FUniformParameter* Param = Value->As<FUniformParameter>())
+			{
+				Out.Appendf(TEXT("#%d = Name=\"%s\" Type=\"%s\"\n"),
+					Param->Analysis_UniformParameterIndex,
+					*Param->Info.Name.ToString(),
+					*MaterialParameterTypeToString(Param->Metadata.Value.Type));
+			}
+			else if (const FTextureObject* TextureObject = Value->As<FTextureObject>())
+			{
+				Out.Appendf(TEXT("#%d = Name=\"%s\" Type=\"%s\"\n"),
+					TextureObject->Analysis_UniformParameterIndex,
+					*TextureObject->Texture->GetName(),
+					TEXT("Texture"));
+			}
+		}
+	}
+};
+
+void DebugDumpIR(FStringView MaterialName, const FMaterialIRModule& Module)
+{
+	FDebugDumpIRState State{};
+	State.Out.Append(TEXT("; Material IR module dump.\n"));
+	State.Out.Appendf(TEXT(";    Material: %s\n"), MaterialName.GetData());
+	State.Out.Append(TEXT("\n"));
+
+	// Dump the IR instructions in the root block.
+	static const TCHAR* Header = TEXT("; Material \"%s\" translated IR dump\n\n");
+	State.AppendBlock(Module.GetRootBlock(), 0);
+
+	// Print referenced material parameters recap if any
+	State.DumpReferencedParameters();
+
+	// Save the dump to file
+	FString FilePath = FPaths::Combine(FPaths::ProjectSavedDir(), "Materials", TEXT("IRDump.txt"));
+	FFileHelper::SaveStringToFile(State.Out, *FilePath);
+}
+
 } // namespace UE::MIR
 
 #endif // #if WITH_EDITOR
diff --git a/Engine/Source/Runtime/Engine/Private/Materials/MaterialIRModuleBuilder.cpp b/Engine/Source/Runtime/Engine/Private/Materials/MaterialIRModuleBuilder.cpp
--- a/Engine/Source/Runtime/Engine/Private/Materials/MaterialIRModuleBuilder.cpp
+++ b/Engine/Source/Runtime/Engine/Private/Materials/MaterialIRModuleBuilder.cpp
@@ -24,6 +24,15 @@
 
 namespace MIR = UE::MIR;
 
+static TAutoConsoleVariable<int32> CVarMaterialIRDebugDumpLevel(
+	TEXT("r.Material.Translator.DebugDump"),
+	0,
+	TEXT("Whether the material translator should dump debug information about the translated module IR.\n")
+	TEXT("0 (Default): No debug dump generated.\n")
+	TEXT("1: Dump the material IR instructions to readable a human readable textual representation (to '{SavedDir}/Materials/IRDump.txt').\n")
+	TEXT("2: Everything above, plus dump the 'Uses' graph in Graphviz Dot syntax (to '{SavedDir}/Materials/IRDumpUseGraph.dot').\n"),
+	ECVF_RenderThreadSafe);
+
 struct FAnalysisContext
 {
 	UMaterialExpressionMaterialFunctionCall* Call{};
@@ -532,7 +541,12 @@ bool FMaterialIRModuleBuilder::Build(FMaterialIRModule* TargetModule)
 	Impl.Step_LinkInstructions();
 	Impl.Step_Finalize();
 
-	UE::MIR::DebugDumpIRUseGraph(*TargetModule);
+	// Dump debugging information if requested 
+	switch (CVarMaterialIRDebugDumpLevel.GetValueOnGameThread())
+	{
+		case 2: UE::MIR::DebugDumpIRUseGraph(*TargetModule); // fallthrough
+		case 1: UE::MIR::DebugDumpIR(Material->GetFullName(), *TargetModule); // fallthrough
+	}
 
 	return true;
 }
diff --git a/Engine/Source/Runtime/Engine/Public/Materials/MaterialIRDebug.h b/Engine/Source/Runtime/Engine/Public/Materials/MaterialIRDebug.h
--- a/Engine/Source/Runtime/Engine/Public/Materials/MaterialIRDebug.h
+++ b/Engine/Source/Runtime/Engine/Public/Materials/MaterialIRDebug.h
@@ -8,9 +8,15 @@
 namespace UE::MIR
 {
 
-//
+// Dumps a Graphviz Dot file to "{ProjectLogDir}/MaterialIRDumpUseGraph.dot" containing the
+// full graph of IR nodes (with a string representation of value parameters inside each node),
+// connected by the "Use" relationship. A value node will look connected to another value node
+// iff it uses that other value as argument (e.g. an "Add" instruction will link to its operand values).
 void DebugDumpIRUseGraph(const FMaterialIRModule& Module);
 
+// Dumps the module to a human readable file in "{ProjectLogDir}/MaterialIRDump.txt"
+void DebugDumpIR(FStringView MaterialName, const FMaterialIRModule& Module);
+
 } // namespace UE::MIR
 
 #endif // #if WITH_EDITOR
diff --git a/Engine/Source/Runtime/Engine/Public/Materials/MaterialIRTypes.h b/Engine/Source/Runtime/Engine/Public/Materials/MaterialIRTypes.h
--- a/Engine/Source/Runtime/Engine/Public/Materials/MaterialIRTypes.h
+++ b/Engine/Source/Runtime/Engine/Public/Materials/MaterialIRTypes.h
@@ -64,22 +64,22 @@ struct FType
 	// Returns whether this type is a texture.
 	bool IsTexture() const;
 
-	// Returns this type upcast this type to ArithmeticType if it is one. Otherwise it returns nullptr.
+	// Returns this type upcast this type to PrimitiveType if it is one, otherwise nullptr.
 	FPrimitiveTypePtr AsPrimitive() const;
 
-	// Returns this type upcast this type to ArithmeticType if it's a scalar. Otherwise it returns nullptr.
+	// Returns this type upcast this type to PrimitiveType if it's a scalar, otherwise nullptr.
 	FPrimitiveTypePtr AsScalar() const;
 
-	// Returns this type upcast this type to ArithmeticType if it's a vector. Otherwise it returns nullptr.
+	// Returns this type upcast this type to PrimitiveType if it's a vector, otherwise nullptr.
 	FPrimitiveTypePtr AsVector() const;
 
-	// Returns this type upcast this type to ArithmeticType if it's a matrix. Otherwise it returns nullptr.
+	// Returns this type upcast this type to PrimitiveType if it's a matrix, otherwise nullptr.
 	FPrimitiveTypePtr AsMatrix() const;
 
 	// Returns the this type name spelling (e.g. float4x4).
 	FStringView GetSpelling() const;
 	
-	//
+	// Returns this type cast as a FObjectPtr if it is one, otherwise nullptr.
 	FObjectTypePtr AsObject() const;
 
 	// Converts this type to a UE::Shader::EValueType.
