Optimization: Added AActor::CanEverAffectNavigation (returns true by default) and overriden by AWorldPartitionHLOD (return false).
Adapted UActorComponent::CanEverAffectNavigation to take into consideration its actor.

diff --git a/Engine/Source/Runtime/Engine/Classes/Components/ActorComponent.h b/Engine/Source/Runtime/Engine/Classes/Components/ActorComponent.h
--- a/Engine/Source/Runtime/Engine/Classes/Components/ActorComponent.h
+++ b/Engine/Source/Runtime/Engine/Classes/Components/ActorComponent.h
@@ -1423,11 +1423,6 @@ protected:
 //////////////////////////////////////////////////////////////////////////
 // UActorComponent inlines
 
-FORCEINLINE bool UActorComponent::CanEverAffectNavigation() const
-{
-	return bCanEverAffectNavigation;
-}
-
 FORCEINLINE_DEBUGGABLE bool UActorComponent::IsNetSimulating() const
 {
 	return GetIsReplicated() && GetOwnerRole() != ROLE_Authority;
diff --git a/Engine/Source/Runtime/Engine/Classes/GameFramework/Actor.h b/Engine/Source/Runtime/Engine/Classes/GameFramework/Actor.h
--- a/Engine/Source/Runtime/Engine/Classes/GameFramework/Actor.h
+++ b/Engine/Source/Runtime/Engine/Classes/GameFramework/Actor.h
@@ -4330,6 +4330,15 @@ public:
 	 */
 	virtual bool IsComponentRelevantForNavigation(UActorComponent* Component) const { return true; }
 
+	/** 
+	 * Can this actor potentially influence navigation.
+	 * This function should always return the same value.
+	 */
+	ENGINE_API virtual bool CanActorTypeEverAffectNavigation() const
+	{
+		return true;
+	}
+
 private:
 	static ENGINE_API FMakeNoiseDelegate MakeNoiseDelegate;
 
diff --git a/Engine/Source/Runtime/Engine/Private/AI/NavigationSystemBase.cpp b/Engine/Source/Runtime/Engine/Private/AI/NavigationSystemBase.cpp
--- a/Engine/Source/Runtime/Engine/Private/AI/NavigationSystemBase.cpp
+++ b/Engine/Source/Runtime/Engine/Private/AI/NavigationSystemBase.cpp
@@ -26,6 +26,21 @@ DEFINE_LOG_CATEGORY(LogAStar);
 
 namespace FNavigationSystem
 {
+	FORCEINLINE bool IsActorRelevantForNav(AActor& InActor)
+	{
+		return InActor.CanActorTypeEverAffectNavigation();
+	}
+
+	FORCEINLINE bool IsComponentRelevantForNav(UActorComponent& InComponent)
+	{
+		const AActor* Actor = InComponent.GetOwner();
+		if ((Actor && !Actor->CanActorTypeEverAffectNavigation()) || !InComponent.IsNavigationRelevant())
+		{
+			return false;
+		}
+		return true;
+	}
+
 	void DiscardNavigationDataChunks(UWorld& InWorld)
 	{
 		const auto& Levels = InWorld.GetLevels();
@@ -77,7 +92,7 @@ namespace FNavigationSystem
 
 	void AddNavigationSystemToWorld(UWorld& WorldOwner, const FNavigationSystemRunMode RunMode, UNavigationSystemConfig* NavigationSystemConfig, const bool bInitializeForWorld, const bool bOverridePreviousNavSys)
 	{
-		TRACE_CPUPROFILER_EVENT_SCOPE(FNavigationSystem::AddNavigationSystemToWorld);
+		TRACE_CPUPROFILER_EVENT_SCOPE(FNavigationSystemPerf::AddNavigationSystemToWorld);
 		UE_LOG(LogNavigation, Verbose, TEXT("%hs bOverridePreviousNavSys=%i (WorldOwner: %s)"), __FUNCTION__, bOverridePreviousNavSys, *WorldOwner.GetOuter()->GetName());
 
 		const FNavigationSystemRunMode ResolvedRunMode = (RunMode == FNavigationSystemRunMode::InferFromWorldMode) ? FindRunModeFromWorldType(WorldOwner) : RunMode;
@@ -262,21 +277,126 @@ PRAGMA_ENABLE_DEPRECATION_WARNINGS
 	void UnregisterNavRelevantObject(UObject& Object) { Delegates.UnregisterNavRelevantObject.Execute(Object); }
 	void OnObjectBoundsChanged(UObject& Object, const FBox& NewBounds, const TConstArrayView<FBox> DirtyAreas) { Delegates.OnObjectBoundsChanged.Execute(Object, NewBounds, DirtyAreas); }
 
-	void UpdateActorData(AActor& Actor) { Delegates.UpdateActorData.Execute(Actor); }
-	void UpdateComponentData(UActorComponent& Comp) { Delegates.UpdateComponentData.Execute(Comp); }
-	void UpdateActorAndComponentData(AActor& Actor, bool bUpdateAttachedActors) { Delegates.UpdateActorAndComponentData.Execute(Actor, bUpdateAttachedActors); }
-	void UpdateComponentDataAfterMove(USceneComponent& Comp) { Delegates.UpdateComponentDataAfterMove.Execute(Comp); }
-	void OnActorBoundsChanged(AActor& Actor) { Delegates.OnActorBoundsChanged.Execute(Actor); }
-	void OnPostEditActorMove(AActor& Actor) { Delegates.OnPostEditActorMove.Execute(Actor); }
-	void OnComponentBoundsChanged(UActorComponent& Comp, const FBox& NewBounds, const FBox& DirtyArea) { OnObjectBoundsChanged(Comp, NewBounds, { DirtyArea }); }
-	void OnComponentTransformChanged(USceneComponent& Comp) { Delegates.OnComponentTransformChanged.Execute(Comp); }
-	void OnActorRegistered(AActor& Actor) { Delegates.OnActorRegistered.Execute(Actor); }
-	void OnActorUnregistered(AActor& Actor) { Delegates.OnActorUnregistered.Execute(Actor); }
-	void OnComponentRegistered(UActorComponent& Comp) { Delegates.OnComponentRegistered.Execute(Comp); }
-	void OnComponentUnregistered(UActorComponent& Comp) { Delegates.OnComponentUnregistered.Execute(Comp); }
-	void RegisterComponent(UActorComponent& Comp) { Delegates.RegisterComponent.Execute(Comp); }
-	void UnregisterComponent(UActorComponent& Comp) { Delegates.UnregisterComponent.Execute(Comp); }
-	void RemoveActorData(AActor& Actor) { Delegates.RemoveActorData.Execute(Actor); }
+	void UpdateActorData(AActor& Actor)
+	{
+		if (IsActorRelevantForNav(Actor))
+		{
+			Delegates.UpdateActorData.Execute(Actor);
+		}
+	}
+	
+	void UpdateComponentData(UActorComponent& Comp)
+	{
+		if (IsComponentRelevantForNav(Comp))
+		{
+			Delegates.UpdateComponentData.Execute(Comp);
+		}
+	}
+	
+	void UpdateActorAndComponentData(AActor& Actor, bool bUpdateAttachedActors)
+	{
+		if (IsActorRelevantForNav(Actor))
+		{
+			Delegates.UpdateActorAndComponentData.Execute(Actor, bUpdateAttachedActors);
+		}
+	}
+	
+	void UpdateComponentDataAfterMove(USceneComponent& Comp)
+	{
+		if (IsComponentRelevantForNav(Comp))
+		{
+			Delegates.UpdateComponentDataAfterMove.Execute(Comp);
+		}
+	}
+	
+	void OnActorBoundsChanged(AActor& Actor)
+	{
+		if (IsActorRelevantForNav(Actor))
+		{
+			Delegates.OnActorBoundsChanged.Execute(Actor);
+		}
+	}
+	
+	void OnPostEditActorMove(AActor& Actor)
+	{
+		if (IsActorRelevantForNav(Actor))
+		{
+			Delegates.OnPostEditActorMove.Execute(Actor);
+		}
+	}
+	
+	void OnComponentBoundsChanged(UActorComponent& Comp, const FBox& NewBounds, const FBox& DirtyArea)
+	{
+		if (IsComponentRelevantForNav(Comp))
+		{
+			OnObjectBoundsChanged(Comp, NewBounds, { DirtyArea });
+		}
+	}
+	
+	void OnComponentTransformChanged(USceneComponent& Comp)
+	{
+		if (IsComponentRelevantForNav(Comp))
+		{
+			Delegates.OnComponentTransformChanged.Execute(Comp);
+		}
+	}
+	
+	void OnActorRegistered(AActor& Actor)
+	{
+		if (IsActorRelevantForNav(Actor))
+		{
+			Delegates.OnActorRegistered.Execute(Actor);
+		}
+	}
+	
+	void OnActorUnregistered(AActor& Actor)
+	{
+		if (IsActorRelevantForNav(Actor))
+		{
+			Delegates.OnActorUnregistered.Execute(Actor);
+		}
+	}
+	
+	void OnComponentRegistered(UActorComponent& Comp)
+	{
+		if (IsComponentRelevantForNav(Comp))
+		{
+			Delegates.OnComponentRegistered.Execute(Comp);
+		}
+	}
+	
+	void OnComponentUnregistered(UActorComponent& Comp)
+	{
+		if (IsComponentRelevantForNav(Comp))
+		{
+			Delegates.OnComponentUnregistered.Execute(Comp);
+		}
+	}
+	
+	void RegisterComponent(UActorComponent& Comp)
+	{
+		if (IsComponentRelevantForNav(Comp))
+		{
+			Delegates.RegisterComponent.Execute(Comp);
+		}
+	}
+	
+	void UnregisterComponent(UActorComponent& Comp)
+	{
+		if (IsComponentRelevantForNav(Comp))
+		{
+			Delegates.UnregisterComponent.Execute(Comp);
+		}
+	}
+	
+	void RemoveActorData(AActor& Actor)
+	{
+		if (IsActorRelevantForNav(Actor))
+		{
+			Delegates.RemoveActorData.Execute(Actor);
+		}
+	}
+	
 	bool HasComponentData(UActorComponent& Comp) { return Delegates.HasComponentData.Execute(Comp);	}
 	const FNavDataConfig& GetDefaultSupportedAgent() { return Delegates.GetDefaultSupportedAgent.Execute(); }
 	const FNavDataConfig& GetBiggestSupportedAgent(const UWorld* World) { return Delegates.GetBiggestSupportedAgent.Execute(World); }
diff --git a/Engine/Source/Runtime/Engine/Private/Components/ActorComponent.cpp b/Engine/Source/Runtime/Engine/Private/Components/ActorComponent.cpp
--- a/Engine/Source/Runtime/Engine/Private/Components/ActorComponent.cpp
+++ b/Engine/Source/Runtime/Engine/Private/Components/ActorComponent.cpp
@@ -2932,6 +2932,13 @@ void UActorComponent::AddReferencedObjects(UObject* InThis, FReferenceCollector&
 	}
 }
 
+bool UActorComponent::CanEverAffectNavigation() const
+{
+	const AActor* Actor = GetOwner();
+	const bool bOwnerCanEverAffectNavigation = Actor ? Actor->CanActorTypeEverAffectNavigation() : true;
+	return bCanEverAffectNavigation && bOwnerCanEverAffectNavigation;
+}
+
 void UActorComponent::SetCanEverAffectNavigation(bool bRelevant)
 {
 	if (bCanEverAffectNavigation != bRelevant)
@@ -2947,7 +2954,7 @@ void UActorComponent::HandleCanEverAffectNavigationChange(bool bForceUpdate)
 	// update octree if already registered
 	if (bRegistered || bForceUpdate)
 	{
-		if (bCanEverAffectNavigation)
+		if (CanEverAffectNavigation())
 		{
 			// Update cached value
 			bNavigationRelevant = IsNavigationRelevant();
diff --git a/Engine/Source/Runtime/Engine/Private/Components/PrimitiveComponent.cpp b/Engine/Source/Runtime/Engine/Private/Components/PrimitiveComponent.cpp
--- a/Engine/Source/Runtime/Engine/Private/Components/PrimitiveComponent.cpp
+++ b/Engine/Source/Runtime/Engine/Private/Components/PrimitiveComponent.cpp
@@ -671,7 +671,7 @@ void UPrimitiveComponent::OnRegister()
 {
 	Super::OnRegister();
 	
-	if (bCanEverAffectNavigation)
+	if (CanEverAffectNavigation())
 	{
 		const bool bNavRelevant = bNavigationRelevant = IsNavigationRelevant();
 		if (bNavRelevant)
@@ -717,7 +717,7 @@ void UPrimitiveComponent::OnUnregister()
 		IStreamingManager::Get().NotifyPrimitiveDetached(this);
 	}
 
-	if (bCanEverAffectNavigation)
+	if (CanEverAffectNavigation())
 	{
 		FNavigationSystem::OnComponentUnregistered(*this);
 	}
diff --git a/Engine/Source/Runtime/Engine/Private/Components/SceneComponent.cpp b/Engine/Source/Runtime/Engine/Private/Components/SceneComponent.cpp
--- a/Engine/Source/Runtime/Engine/Private/Components/SceneComponent.cpp
+++ b/Engine/Source/Runtime/Engine/Private/Components/SceneComponent.cpp
@@ -3090,7 +3090,7 @@ bool USceneComponent::InternalSetWorldLocationAndRotation(FVector NewLocation, c
 		UpdateComponentToWorldWithParent(GetAttachParent(),GetAttachSocketName(), SkipPhysicsToEnum(bNoPhysics), RelativeRotationCache.GetCachedQuat(), Teleport);
 
 		// we need to call this even if this component itself is not navigation relevant
-		if (IsRegistered() && bCanEverAffectNavigation)
+		if (IsRegistered() && CanEverAffectNavigation())
 		{
 			PostUpdateNavigationData();
 		}
diff --git a/Engine/Source/Runtime/Engine/Private/InstancedStaticMesh.cpp b/Engine/Source/Runtime/Engine/Private/InstancedStaticMesh.cpp
--- a/Engine/Source/Runtime/Engine/Private/InstancedStaticMesh.cpp
+++ b/Engine/Source/Runtime/Engine/Private/InstancedStaticMesh.cpp
@@ -2859,7 +2859,7 @@ void UInstancedStaticMeshComponent::UpdateBounds()
 		Super::UpdateBounds();
 	}
 
-	if (bCanEverAffectNavigation && (IsRegistered() || IsPreRegistering()))
+	if (CanEverAffectNavigation() && (IsRegistered() || IsPreRegistering()))
 	{
 		CalcAndCacheNavigationBounds();
 	}
diff --git a/Engine/Source/Runtime/Engine/Public/WorldPartition/HLOD/HLODActor.h b/Engine/Source/Runtime/Engine/Public/WorldPartition/HLOD/HLODActor.h
--- a/Engine/Source/Runtime/Engine/Public/WorldPartition/HLOD/HLODActor.h
+++ b/Engine/Source/Runtime/Engine/Public/WorldPartition/HLOD/HLODActor.h
@@ -102,6 +102,10 @@ protected:
 	{
 		return false;
 	}
+	ENGINE_API virtual bool CanActorTypeEverAffectNavigation() const override
+	{
+		return false;
+	}
 	virtual bool IsRuntimeOnly() const override { return true; }
 #if WITH_EDITOR
 	ENGINE_API virtual TUniquePtr<class FWorldPartitionActorDesc> CreateClassActorDesc() const override;
