AnimNext: Fix for unnecessary InitAnim when EnableAnimation value  doesn't change

diff --git a/Engine/Source/Runtime/Engine/Private/Components/SkeletalMeshComponent.cpp b/Engine/Source/Runtime/Engine/Private/Components/SkeletalMeshComponent.cpp
--- a/Engine/Source/Runtime/Engine/Private/Components/SkeletalMeshComponent.cpp
+++ b/Engine/Source/Runtime/Engine/Private/Components/SkeletalMeshComponent.cpp
@@ -3289,12 +3289,15 @@ void USkeletalMeshComponent::SetEnableAnimation(bool bNewEnableAnimation)
 		{
 			if(USkeletalMeshComponent* Component = WeakComponent.Get())
 			{
-				Component->bEnableAnimation = bNewEnableAnimation;
-				Component->PrimaryComponentTick.bRunOnAnyThread = !bNewEnableAnimation;
-
-				if (bNewEnableAnimation && Component->IsRegistered())
+				if (bNewEnableAnimation != Component->bEnableAnimation)
 				{
-					Component->InitAnim(false);
+					Component->bEnableAnimation = bNewEnableAnimation;
+					Component->PrimaryComponentTick.bRunOnAnyThread = !bNewEnableAnimation;
+
+					if (bNewEnableAnimation && Component->IsRegistered())
+					{
+						Component->InitAnim(false);
+					}
 				}
 			}
 			return false;
