Do not pre-cache materials with an incomplete vertex declaration on non-MVF platforms
#rnx

diff --git a/Engine/Source/Runtime/Engine/Private/Materials/Material.cpp b/Engine/Source/Runtime/Engine/Private/Materials/Material.cpp
--- a/Engine/Source/Runtime/Engine/Private/Materials/Material.cpp
+++ b/Engine/Source/Runtime/Engine/Private/Materials/Material.cpp
@@ -4608,10 +4608,13 @@ void UMaterial::PostLoad()
 	}
 	else if (IsDeferredDecal() || IsUIMaterial() || (MaterialDomain == MD_LightFunction) || IsPostProcessMaterial())
 	{
-		FPSOPrecacheParams PSOPrecacheParams;
-		UMaterialInterface::PrecachePSOs(&FLocalVertexFactory::StaticType, PSOPrecacheParams);
+		// TODO: need to pass a correct vertex declaration for non-MVF platforms
+		if (RHISupportsManualVertexFetch(GMaxRHIShaderPlatform)) 
+		{
+			FPSOPrecacheParams PSOPrecacheParams;
+			UMaterialInterface::PrecachePSOs(&FLocalVertexFactory::StaticType, PSOPrecacheParams);
+		}
 	}
-
 #if !WITH_EDITOR
 	CacheTexturesSamplingInfo();
 #endif
diff --git a/Engine/Source/Runtime/Engine/Private/Materials/MaterialInstance.cpp b/Engine/Source/Runtime/Engine/Private/Materials/MaterialInstance.cpp
--- a/Engine/Source/Runtime/Engine/Private/Materials/MaterialInstance.cpp
+++ b/Engine/Source/Runtime/Engine/Private/Materials/MaterialInstance.cpp
@@ -3471,8 +3471,12 @@ void UMaterialInstance::PostLoad()
 	}
 	else if (IsDeferredDecal() || IsUIMaterial() || IsPostProcessMaterial())
 	{
-		FPSOPrecacheParams PSOPrecacheParams;
-		UMaterialInterface::PrecachePSOs(&FLocalVertexFactory::StaticType, PSOPrecacheParams);
+		// TODO: need to pass a correct vertex declaration for non-MVF platforms
+		if (RHISupportsManualVertexFetch(GMaxRHIShaderPlatform))
+		{
+			FPSOPrecacheParams PSOPrecacheParams;
+			UMaterialInterface::PrecachePSOs(&FLocalVertexFactory::StaticType, PSOPrecacheParams);
+		}
 	}
 	//DumpDebugInfo(*GLog);
 
diff --git a/Engine/Source/Runtime/Engine/Private/Particles/ParticleGpuSimulation.cpp b/Engine/Source/Runtime/Engine/Private/Particles/ParticleGpuSimulation.cpp
--- a/Engine/Source/Runtime/Engine/Private/Particles/ParticleGpuSimulation.cpp
+++ b/Engine/Source/Runtime/Engine/Private/Particles/ParticleGpuSimulation.cpp
@@ -5499,11 +5499,17 @@ void ParticleGPUSimulationGlobalPSOCollector(const FSceneTexturesConfig& SceneTe
 	TShaderMapRef<TParticleSimulationPS<PCM_None> > NoCollisionPixelShader(GlobalShaderMap);
 	AddPSOInitializer(NoCollisionPixelShader.GetPixelShader());
 
-	TShaderMapRef<TParticleSimulationPS<PCM_DepthBuffer> > DepthCollisionPixelShader(GlobalShaderMap);
-	AddPSOInitializer(DepthCollisionPixelShader.GetPixelShader());
-
-	TShaderMapRef<TParticleSimulationPS<PCM_DistanceField> > DistanceFieldCollisionPixelShader(GlobalShaderMap);
-	AddPSOInitializer(DistanceFieldCollisionPixelShader.GetPixelShader());
+	if (IsParticleCollisionModeSupported(ShaderPlatform, PCM_DepthBuffer))
+	{
+		TShaderMapRef<TParticleSimulationPS<PCM_DepthBuffer> > DepthCollisionPixelShader(GlobalShaderMap);
+		AddPSOInitializer(DepthCollisionPixelShader.GetPixelShader());
+	}
+	
+	if (IsParticleCollisionModeSupported(ShaderPlatform, PCM_DistanceField))
+	{
+		TShaderMapRef<TParticleSimulationPS<PCM_DistanceField> > DistanceFieldCollisionPixelShader(GlobalShaderMap);
+		AddPSOInitializer(DistanceFieldCollisionPixelShader.GetPixelShader());
+	}
 }
 
 FRegisterGlobalPSOCollectorFunction RegisterParticleGPUSimulationGlobalPSOCollector(&ParticleGPUSimulationGlobalPSOCollector, TEXT("ParticleGPUSimulationGlobalPSOCollector"));
