Implement Abs, Ceil, Floor and Round in the new material translator

#rnx

diff --git a/Engine/Source/Runtime/Engine/Private/Materials/MaterialExpressionsIR.cpp b/Engine/Source/Runtime/Engine/Private/Materials/MaterialExpressionsIR.cpp
--- a/Engine/Source/Runtime/Engine/Private/Materials/MaterialExpressionsIR.cpp
+++ b/Engine/Source/Runtime/Engine/Private/Materials/MaterialExpressionsIR.cpp
@@ -467,6 +467,26 @@ static void BuildUnaryArithmeticOperatorFromInput(
 
 // TODO: Add the missing expressions
 
+void UMaterialExpressionAbs::Build(UE::MIR::FEmitter& Emitter)
+{
+	BuildUnaryArithmeticOperatorFromInput(Emitter, MIR::UO_Abs, &Input);
+}
+
+void UMaterialExpressionCeil::Build(UE::MIR::FEmitter& Emitter)
+{
+	BuildUnaryArithmeticOperatorFromInput(Emitter, MIR::UO_Ceil, &Input);
+}
+
+void UMaterialExpressionFloor::Build(UE::MIR::FEmitter& Emitter)
+{
+	BuildUnaryArithmeticOperatorFromInput(Emitter, MIR::UO_Floor, &Input);
+}
+
+void UMaterialExpressionRound::Build(UE::MIR::FEmitter& Emitter)
+{
+	BuildUnaryArithmeticOperatorFromInput(Emitter, MIR::UO_Round, &Input);
+}
+
 void UMaterialExpressionExponential::Build(UE::MIR::FEmitter& Emitter)
 {
 	BuildUnaryArithmeticOperatorFromInput(Emitter, MIR::UO_Exponential, &Input);
@@ -529,20 +549,12 @@ void UMaterialExpressionTangent::Build(UE::MIR::FEmitter& Emitter)
 
 /* Binary Operators */
 
-static void BuildBinaryArithmeticOperator(
+static void BuildBinaryOperatorFromValue(
 	MIR::FEmitter& Emitter,
 	MIR::EOperator Op,
-	FExpressionInput* LhsInput,
-	float LhsConst,
-	FExpressionInput* RhsInput,
-	float RhsConst)
+	MIR::FValue* LhsValue,
+	MIR::FValue* RhsValue)
 {
-	// Default inputs to their relative constants if disconnected, then get each input after checking it has primitive type.
-	MIR::FValue* LhsValue = Emitter.DefaultInputToFloatScalar(LhsInput, LhsConst).InputPrimitive(LhsInput);
-	MIR::FValue* RhsValue = Emitter.DefaultInputToFloatScalar(RhsInput, RhsConst).InputPrimitive(RhsInput);
-
-	UE_MIR_CHECKPOINT();
-
 	// Determine operation input/output type by looking at the first connected input and picking float1 otherwise.
 	MIR::FTypePtr ResultType = Emitter.GetCommonPrimitiveType(LhsValue->Type->AsPrimitive(), RhsValue->Type->AsPrimitive());
 
@@ -559,29 +571,66 @@ static void BuildBinaryArithmeticOperator(
 	Emitter.Output(0, Value);
 }
 
+static void BuildBinaryOperatorFromInputWithDefault(
+	MIR::FEmitter& Emitter,
+	MIR::EOperator Op,
+	FExpressionInput* LhsInput,
+	float LhsConst,
+	FExpressionInput* RhsInput,
+	float RhsConst)
+{
+	// Default inputs to their relative constants if disconnected, then get each input after checking it has primitive type.
+	MIR::FValue* LhsValue = Emitter.DefaultInputToFloatScalar(LhsInput, LhsConst).InputPrimitive(LhsInput);
+	MIR::FValue* RhsValue = Emitter.DefaultInputToFloatScalar(RhsInput, RhsConst).InputPrimitive(RhsInput);
+
+	UE_MIR_CHECKPOINT();
+
+	BuildBinaryOperatorFromValue(Emitter, Op, LhsValue, RhsValue);
+}
+
+static void BuildBinaryOperatorFromInput(
+	MIR::FEmitter& Emitter,
+	MIR::EOperator Op,
+	FExpressionInput* LhsInput,
+	FExpressionInput* RhsInput)
+{
+	// Default inputs to their relative constants if disconnected, then get each input after checking it has primitive type.
+	MIR::FValue* LhsValue = Emitter.Input(LhsInput);
+	MIR::FValue* RhsValue = Emitter.Input(RhsInput);
+
+	UE_MIR_CHECKPOINT();
+
+	BuildBinaryOperatorFromValue(Emitter, Op, LhsValue, RhsValue);
+}
+
 void UMaterialExpressionAdd::Build(MIR::FEmitter& Emitter)
 { 
-	BuildBinaryArithmeticOperator(Emitter, MIR::BO_Add, &A, ConstA, &B, ConstB);
+	BuildBinaryOperatorFromInputWithDefault(Emitter, MIR::BO_Add, &A, ConstA, &B, ConstB);
+}
+
+void UMaterialExpressionFmod::Build(MIR::FEmitter& Emitter)
+{ 
+	BuildBinaryOperatorFromInput(Emitter, MIR::BO_Fmod, &A, &B);
 }
 
 void UMaterialExpressionSubtract::Build(MIR::FEmitter& Emitter)
 { 
-	BuildBinaryArithmeticOperator(Emitter, MIR::BO_Subtract, &A, ConstA, &B, ConstB);
+	BuildBinaryOperatorFromInputWithDefault(Emitter, MIR::BO_Subtract, &A, ConstA, &B, ConstB);
 }
 
 void UMaterialExpressionMultiply::Build(MIR::FEmitter& Emitter)
 { 
-	BuildBinaryArithmeticOperator(Emitter, MIR::BO_Multiply, &A, ConstA, &B, ConstB);
+	BuildBinaryOperatorFromInputWithDefault(Emitter, MIR::BO_Multiply, &A, ConstA, &B, ConstB);
 }
 
 void UMaterialExpressionDivide::Build(MIR::FEmitter& Emitter)
 { 
-	BuildBinaryArithmeticOperator(Emitter, MIR::BO_Divide, &A, ConstA, &B, ConstB);
+	BuildBinaryOperatorFromInputWithDefault(Emitter, MIR::BO_Divide, &A, ConstA, &B, ConstB);
 }
 
 void UMaterialExpressionMin::Build(MIR::FEmitter& Emitter)
 { 
-	BuildBinaryArithmeticOperator(Emitter, MIR::BO_Min, &A, ConstA, &B, ConstB);
+	BuildBinaryOperatorFromInputWithDefault(Emitter, MIR::BO_Min, &A, ConstA, &B, ConstB);
 }
 
 void UMaterialExpressionOneMinus::Build(MIR::FEmitter& Emitter)
diff --git a/Engine/Source/Runtime/Engine/Private/Materials/MaterialIR.cpp b/Engine/Source/Runtime/Engine/Private/Materials/MaterialIR.cpp
--- a/Engine/Source/Runtime/Engine/Private/Materials/MaterialIR.cpp
+++ b/Engine/Source/Runtime/Engine/Private/Materials/MaterialIR.cpp
@@ -413,15 +413,19 @@ const TCHAR* OperatorToString(EOperator Op)
 		case O_Invalid: return TEXT("Invalid");
 
 		/* Unary Arithmetic */
+		case UO_Abs: return TEXT("Abs");
 		case UO_ACos: return TEXT("Arccosine");
 		case UO_ASin: return TEXT("Arcsine");
 		case UO_ATan: return TEXT("Arctangent");
+		case UO_Ceil: return TEXT("Ceil");
 		case UO_Cos: return TEXT("Cosine");
 		case UO_Exponential: return TEXT("Exponential");
 		case UO_Exponential2: return TEXT("Exponential2");
+		case UO_Floor: return TEXT("Floor");
 		case UO_Logarithm: return TEXT("Logarithm");
 		case UO_Logarithm2: return TEXT("Logarithm2");
 		case UO_Logarithm10: return TEXT("Logarithm10");
+		case UO_Round: return TEXT("Round");
 		case UO_Sin: return TEXT("Sine");
 		case UO_Tan: return TEXT("Tangent");
 		case UO_Truncate: return TEXT("Truncate");
@@ -436,6 +440,7 @@ const TCHAR* OperatorToString(EOperator Op)
 	
 		/* Binary Arithmetic */
 		case BO_Add: return TEXT("Add");
+		case BO_Fmod: return TEXT("Fmod");
 		case BO_Subtract: return TEXT("Subtract");
 		case BO_Multiply: return TEXT("Multiply");
 		case BO_Divide: return TEXT("Divide");
diff --git a/Engine/Source/Runtime/Engine/Private/Materials/MaterialIREmitter.cpp b/Engine/Source/Runtime/Engine/Private/Materials/MaterialIREmitter.cpp
--- a/Engine/Source/Runtime/Engine/Private/Materials/MaterialIREmitter.cpp
+++ b/Engine/Source/Runtime/Engine/Private/Materials/MaterialIREmitter.cpp
@@ -802,6 +802,29 @@ static T FoldScalarArithmeticOperator(EOperator Operator, T A, T B, T C)
 {
 	switch (Operator)
 	{
+		case UO_Abs: return FGenericPlatformMath::Abs<T>(A);
+		case UO_Ceil:
+		{
+			if constexpr (std::is_same<T, float>::value || std::is_same<T, double>::value)
+			{
+				return FMath::CeilToFloat(A);
+			}
+			else if constexpr (std::is_same<T, TInteger>::value)
+			{
+				return A;
+			}
+		}
+		case UO_Floor:
+		{
+			if constexpr (std::is_same<T, float>::value || std::is_same<T, double>::value)
+			{
+				return FGenericPlatformMath::FloorToFloat(A);
+			}
+			else if constexpr (std::is_same<T, TInteger>::value)
+			{
+				return A;
+			}
+		}
 		case UO_Negate: return -A;
 		case UO_ACos: [[fallthrough]];
 		case UO_ASin: [[fallthrough]];
@@ -817,6 +840,17 @@ static T FoldScalarArithmeticOperator(EOperator Operator, T A, T B, T C)
 		case UO_Logarithm: return FGenericPlatformMath::LogX(UE_EULERS_NUMBER, A);
 		case UO_Logarithm2: return FGenericPlatformMath::LogX(2.0f, A);
 		case UO_Logarithm10: return FGenericPlatformMath::LogX(10.0f, A);
+		case UO_Round:
+		{
+			if constexpr (std::is_same<T, float>::value || std::is_same<T, double>::value)
+			{
+				return FGenericPlatformMath::RoundToFloat(A);
+			}
+			else if constexpr (std::is_same<T, TInteger>::value)
+			{
+				return A;
+			}
+		}
 		case UO_Truncate:
 		{
 			if constexpr (std::is_same<T, float>::value)
@@ -832,7 +866,18 @@ static T FoldScalarArithmeticOperator(EOperator Operator, T A, T B, T C)
 				UE_MIR_UNREACHABLE();
 			}
 		}
-		case BO_Add: return A + B;
+		case BO_Add: return A + B;\
+		case BO_Fmod:
+		{
+			if constexpr (std::is_same<T, float>::value || std::is_same<T, double>::value)
+			{
+				return FGenericPlatformMath::Fmod(A, B);
+			}
+			else if constexpr (std::is_same<T, TInteger>::value)
+			{
+				return 0.0f;
+			}
+		}
 		case BO_Subtract: return A - B;
 		case BO_Multiply: return A * B;
 		case BO_Divide: return A / B;
@@ -946,8 +991,13 @@ static FValue* TryFoldOperatorScalar(FEmitter* Emitter, EOperator Op, FValue* A,
 		case UO_Logarithm:
 		case UO_Logarithm2:
 		case UO_Logarithm10:
+		case BO_Fmod:
 		case BO_Min:
 		case BO_Max:
+		case UO_Abs:
+		case UO_Ceil:
+		case UO_Floor:
+		case UO_Round:
 			break;
 
 		default:
diff --git a/Engine/Source/Runtime/Engine/Private/Materials/MaterialIRToHLSLTranslator.cpp b/Engine/Source/Runtime/Engine/Private/Materials/MaterialIRToHLSLTranslator.cpp
--- a/Engine/Source/Runtime/Engine/Private/Materials/MaterialIRToHLSLTranslator.cpp
+++ b/Engine/Source/Runtime/Engine/Private/Materials/MaterialIRToHLSLTranslator.cpp
@@ -459,10 +459,13 @@ struct FTranslator : FMaterialIRToHLSLTranslation
 			const TCHAR* OpString;
 			switch (Operator->Op)
 			{
+				case MIR::UO_Abs: OpString = TEXT("abs"); break;
 				case MIR::UO_ACos: OpString = TEXT("acos"); break;
 				case MIR::UO_ASin: OpString = TEXT("asin"); break;
 				case MIR::UO_ATan: OpString = TEXT("atan"); break;
+				case MIR::UO_Ceil: OpString = TEXT("ceil"); break;
 				case MIR::UO_Cos: OpString = TEXT("cos"); break;
+				case MIR::UO_Floor: OpString = TEXT("floor"); break;
 				case MIR::UO_Sin: OpString = TEXT("sin"); break;
 				case MIR::UO_Tan: OpString = TEXT("tan"); break;
 				case MIR::UO_Truncate: OpString = TEXT("trunc"); break;
@@ -471,6 +474,8 @@ struct FTranslator : FMaterialIRToHLSLTranslation
 				case MIR::UO_Logarithm: OpString = TEXT("log"); break;
 				case MIR::UO_Logarithm2: OpString = TEXT("log2"); break;
 				case MIR::UO_Logarithm10: OpString = TEXT("log10"); break;
+				case MIR::UO_Round: OpString = TEXT("round"); break;
+				case MIR::BO_Fmod: OpString = TEXT("fmod"); break;
 				case MIR::BO_Min: OpString = TEXT("min"); break;
 				case MIR::BO_Max: OpString = TEXT("max"); break;
 				case MIR::TO_Lerp: OpString = TEXT("lerp"); break;
diff --git a/Engine/Source/Runtime/Engine/Public/Materials/MaterialExpressionAbs.h b/Engine/Source/Runtime/Engine/Public/Materials/MaterialExpressionAbs.h
--- a/Engine/Source/Runtime/Engine/Public/Materials/MaterialExpressionAbs.h
+++ b/Engine/Source/Runtime/Engine/Public/Materials/MaterialExpressionAbs.h
@@ -25,6 +25,7 @@ class UMaterialExpressionAbs : public UMaterialExpression
 
 	//~ Begin UMaterialExpression Interface
 #if WITH_EDITOR
+	virtual void Build(UE::MIR::FEmitter& Emitter) override;
 	virtual int32 Compile(class FMaterialCompiler* Compiler, int32 OutputIndex) override;
 	virtual void GetCaption(TArray<FString>& OutCaptions) const override;
 #endif
diff --git a/Engine/Source/Runtime/Engine/Public/Materials/MaterialExpressionCeil.h b/Engine/Source/Runtime/Engine/Public/Materials/MaterialExpressionCeil.h
--- a/Engine/Source/Runtime/Engine/Public/Materials/MaterialExpressionCeil.h
+++ b/Engine/Source/Runtime/Engine/Public/Materials/MaterialExpressionCeil.h
@@ -20,6 +20,7 @@ class UMaterialExpressionCeil : public UMaterialExpression
 
 	//~ Begin UMaterialExpression Interface
 #if WITH_EDITOR
+	virtual void Build(UE::MIR::FEmitter& Emitter) override;
 	virtual int32 Compile(class FMaterialCompiler* Compiler, int32 OutputIndex) override;
 	virtual uint32 GetOutputType(int32 OutputIndex) override;
 	virtual void GetCaption(TArray<FString>& OutCaptions) const override;
diff --git a/Engine/Source/Runtime/Engine/Public/Materials/MaterialExpressionFloor.h b/Engine/Source/Runtime/Engine/Public/Materials/MaterialExpressionFloor.h
--- a/Engine/Source/Runtime/Engine/Public/Materials/MaterialExpressionFloor.h
+++ b/Engine/Source/Runtime/Engine/Public/Materials/MaterialExpressionFloor.h
@@ -20,6 +20,7 @@ class UMaterialExpressionFloor : public UMaterialExpression
 
 	//~ Begin UMaterialExpression Interface
 #if WITH_EDITOR
+	virtual void Build(UE::MIR::FEmitter& Emitter) override;
 	virtual int32 Compile(class FMaterialCompiler* Compiler, int32 OutputIndex) override;
 	virtual uint32 GetOutputType(int32 OutputIndex) override;
 	virtual void GetCaption(TArray<FString>& OutCaptions) const override;
diff --git a/Engine/Source/Runtime/Engine/Public/Materials/MaterialExpressionFmod.h b/Engine/Source/Runtime/Engine/Public/Materials/MaterialExpressionFmod.h
--- a/Engine/Source/Runtime/Engine/Public/Materials/MaterialExpressionFmod.h
+++ b/Engine/Source/Runtime/Engine/Public/Materials/MaterialExpressionFmod.h
@@ -23,6 +23,7 @@ class UMaterialExpressionFmod : public UMaterialExpression
 
 	//~ Begin UMaterialExpression Interface
 #if WITH_EDITOR
+	virtual void Build(UE::MIR::FEmitter& Emitter) override;
 	virtual int32 Compile(class FMaterialCompiler* Compiler, int32 OutputIndex) override;
 	virtual void GetCaption(TArray<FString>& OutCaptions) const override;
 	virtual FText GetKeywords() const override {return FText::FromString(TEXT("%"));}
diff --git a/Engine/Source/Runtime/Engine/Public/Materials/MaterialExpressionRound.h b/Engine/Source/Runtime/Engine/Public/Materials/MaterialExpressionRound.h
--- a/Engine/Source/Runtime/Engine/Public/Materials/MaterialExpressionRound.h
+++ b/Engine/Source/Runtime/Engine/Public/Materials/MaterialExpressionRound.h
@@ -18,6 +18,7 @@ class UMaterialExpressionRound : public UMaterialExpression
 
 	//~ Begin UMaterialExpression Interface
 #if WITH_EDITOR
+	virtual void Build(UE::MIR::FEmitter& Emitter) override;
 	virtual int32 Compile(class FMaterialCompiler* Compiler, int32 OutputIndex) override;
 	virtual uint32 GetOutputType(int32 OutputIndex) override;
 	virtual void GetCaption(TArray<FString>& OutCaptions) const override;
diff --git a/Engine/Source/Runtime/Engine/Public/Materials/MaterialIR.h b/Engine/Source/Runtime/Engine/Public/Materials/MaterialIR.h
--- a/Engine/Source/Runtime/Engine/Public/Materials/MaterialIR.h
+++ b/Engine/Source/Runtime/Engine/Public/Materials/MaterialIR.h
@@ -313,15 +313,19 @@ enum EOperator
 	UO_Not,							// Logical negation: not(true) -> false
 	
 	/* Unary Intrinsic Operators */
+	UO_Abs,
 	UO_ACos,
 	UO_ASin,
 	UO_ATan,
+	UO_Ceil,
 	UO_Cos,
 	UO_Exponential,
 	UO_Exponential2,
+	UO_Floor,
 	UO_Logarithm,
 	UO_Logarithm2,
 	UO_Logarithm10,
+	UO_Round,
 	UO_Sin,
 	UO_Tan,
 	UO_Truncate,
@@ -346,6 +350,7 @@ enum EOperator
 	BO_Xor,
 
 	/* Binary Intrinsic Operators */
+	BO_Fmod,
 	BO_Max,
 	BO_Min,
 
