Fix UWorld::IsDataValid to validate all loaded actors regardless of level streaming load/visibility status.

diff --git a/Engine/Source/Runtime/Engine/Private/World.cpp b/Engine/Source/Runtime/Engine/Private/World.cpp
--- a/Engine/Source/Runtime/Engine/Private/World.cpp
+++ b/Engine/Source/Runtime/Engine/Private/World.cpp
@@ -887,7 +887,8 @@ void UWorld::AddReferencedObjects(UObject* InThis, FReferenceCollector& Collecto
 EDataValidationResult UWorld::IsDataValid(FDataValidationContext& Context) const 
 {
 	EDataValidationResult Result = EDataValidationResult::NotValidated;
-	for (FActorIterator It(this); It; ++It)
+	// Pass flags explicitly so we iterate non-active levels for data validation but continue to skip destroyed/deleted actors.
+	for (FActorIterator It(this, EActorIteratorFlags::SkipPendingKill); It; ++It)
 	{
 		AActor* Actor = *It;
 		if (Actor)
