Added profiler uobject stats for syncbodies. Default disabled and compiled out of shipping.

p.Chaos.EmitSyncBodiesUObjectStats=1 to enable.

diff --git a/Engine/Source/Runtime/Engine/Private/PhysicsEngine/Experimental/PhysScene_Chaos.cpp b/Engine/Source/Runtime/Engine/Private/PhysicsEngine/Experimental/PhysScene_Chaos.cpp
--- a/Engine/Source/Runtime/Engine/Private/PhysicsEngine/Experimental/PhysScene_Chaos.cpp
+++ b/Engine/Source/Runtime/Engine/Private/PhysicsEngine/Experimental/PhysScene_Chaos.cpp
@@ -53,6 +53,18 @@ TAutoConsoleVariable<int32> CVar_ChaosUpdateKinematicsOnDeferredSkelMeshes(TEXT(
 
 #endif
 
+#if UE_BUILD_SHIPPING
+inline constexpr bool GEmitSyncBodiesUObjectStats = false;
+#else
+
+static bool GEmitSyncBodiesUObjectStats = false;
+static FAutoConsoleVariableRef CVar_EmitSyncBodiesUObjectStats(
+	TEXT("p.Chaos.EmitSyncBodiesUObjectStats"),
+	GEmitSyncBodiesUObjectStats,
+	TEXT("Emits profiler scopes for uobjects updating in SyncBodies"));
+
+#endif
+
 int32 GEnableKinematicDeferralStartPhysicsCondition = 1;
 FAutoConsoleVariableRef CVar_EnableKinematicDeferralStartPhysicsCondition(TEXT("p.EnableKinematicDeferralStartPhysicsCondition"), GEnableKinematicDeferralStartPhysicsCondition, TEXT("If is 1, allow kinematics to be deferred in start physics (probably only called from replication tick). If 0, no deferral in startphysics."));
 
@@ -2321,40 +2333,39 @@ void FPhysScene_Chaos::OnSyncBodies(Chaos::FPhysicsSolverBase* Solver)
 
 			if (FBodyInstance* BodyInstance = ChaosInterface::GetUserData(*(Proxy->GetParticle_LowLevel())))
 			{
-				if (BodyInstance->OwnerComponent.IsValid())
+				SCOPE_CYCLE_UOBJECT(Component, GEmitSyncBodiesUObjectStats ? BodyInstance->OwnerComponent.Get() : nullptr);
+
+				// Skip kinematics, unless they've been flagged to update following simulation
+				if (!BodyInstance->bUpdateKinematicFromSimulation && BodyInstance->IsInstanceSimulatingPhysics() == false)
 				{
-					// Skip kinematics, unless they've been flagged to update following simulation
-					if (!BodyInstance->bUpdateKinematicFromSimulation && BodyInstance->IsInstanceSimulatingPhysics() == false)
-					{
-						return;
-					}
-					UPrimitiveComponent* OwnerComponent = BodyInstance->OwnerComponent.Get();
-					if (OwnerComponent != nullptr)
-					{
-						bool bPendingMove = false;
-						FRigidTransform3 NewTransform(DirtyParticle->X(), DirtyParticle->R());
+					return;
+				}
 
-						if (!NewTransform.EqualsNoScale(OwnerComponent->GetComponentTransform()))
-						{
-							if (BodyInstance->InstanceBodyIndex == INDEX_NONE)
-							{
+				if (UPrimitiveComponent* OwnerComponent = BodyInstance->OwnerComponent.Get())
+				{
+					bool bPendingMove = false;
+					FRigidTransform3 NewTransform(DirtyParticle->X(), DirtyParticle->R());
 
-								bPendingMove = true;
-								const FVector MoveBy = NewTransform.GetLocation() - OwnerComponent->GetComponentTransform().GetLocation();
-								const FQuat NewRotation = NewTransform.GetRotation();
-								PendingTransforms.Add(FPhysScenePendingComponentTransform_Chaos(OwnerComponent, MoveBy, NewRotation, Proxy->GetWakeEvent()));
+					if (!NewTransform.EqualsNoScale(OwnerComponent->GetComponentTransform()))
+					{
+						if (BodyInstance->InstanceBodyIndex == INDEX_NONE)
+						{
 
-							}
+							bPendingMove = true;
+							const FVector MoveBy = NewTransform.GetLocation() - OwnerComponent->GetComponentTransform().GetLocation();
+							const FQuat NewRotation = NewTransform.GetRotation();
+							PendingTransforms.Add(FPhysScenePendingComponentTransform_Chaos(OwnerComponent, MoveBy, NewRotation, Proxy->GetWakeEvent()));
 
-							PhysScene->UpdateActorInAccelerationStructure(BodyInstance->ActorHandle);
 						}
 
-						if (Proxy->GetWakeEvent() != Chaos::EWakeEventEntry::None && !bPendingMove)
-						{
-							PendingTransforms.Add(FPhysScenePendingComponentTransform_Chaos(OwnerComponent, Proxy->GetWakeEvent()));
-						}
-						Proxy->ClearEvents();
+						PhysScene->UpdateActorInAccelerationStructure(BodyInstance->ActorHandle);
+					}
+
+					if (Proxy->GetWakeEvent() != Chaos::EWakeEventEntry::None && !bPendingMove)
+					{
+						PendingTransforms.Add(FPhysScenePendingComponentTransform_Chaos(OwnerComponent, Proxy->GetWakeEvent()));
 					}
+					Proxy->ClearEvents();
 				}
 			}
 		}
