Disable the use of ClusterUnions when using Remote Object Handles.  The reasoning is that it's currently implemented as a dynamic cast and we do not want a migration to occur during remote physics queries.

#rnx

diff --git a/Engine/Source/Runtime/Engine/Private/Collision/SceneQuery.cpp b/Engine/Source/Runtime/Engine/Private/Collision/SceneQuery.cpp
--- a/Engine/Source/Runtime/Engine/Private/Collision/SceneQuery.cpp
+++ b/Engine/Source/Runtime/Engine/Private/Collision/SceneQuery.cpp
@@ -586,6 +586,8 @@ bool TSceneCastCommonImpWithRetryRequest(const UWorld* World, typename Traits::O
 				auto DoClusterUnionTrace = [&GeomInputs, &Start, &End, TraceChannel, &Params, &ResponseParams, &ObjectParams](const FHitResult& OriginalHit, auto& NewHit)
 				{
 					FClusterUnionHit Result;
+
+#if !UE_WITH_REMOTE_OBJECT_HANDLE
 					if (UClusterUnionComponent* ClusterUnion = Cast<UClusterUnionComponent>(OriginalHit.GetComponent()))
 					{
 						Result.bIsClusterUnion = true;
@@ -598,6 +600,7 @@ bool TSceneCastCommonImpWithRetryRequest(const UWorld* World, typename Traits::O
 							Result.bHit = ClusterUnion->SweepComponent(NewHit, Start, End, *GeomInputs.GetGeometryOrientation(), *GeomInputs.GetGeometry(), TraceChannel, Params, ResponseParams, ObjectParams);
 						}
 					}
+#endif
 
 					return Result;
 				};
@@ -1155,12 +1158,15 @@ bool GeomOverlapMultiImp(const UWorld* World, const FPhysicsGeometry& Geom, cons
 						auto DoClusterUnionOverlap = [&GeomPose, &Geom, TraceChannel, &Params, &ResponseParams, &ObjectParams](const FOverlapResult& OriginalOverlap, TArray<FOverlapResult>& NewOverlaps)
 						{
 							FClusterUnionHit Result;
+
+#if !UE_WITH_REMOTE_OBJECT_HANDLE
 							if (UClusterUnionComponent* ClusterUnion = Cast<UClusterUnionComponent>(OriginalOverlap.GetComponent()))
 							{
 								Result.bIsClusterUnion = true;
 								ClusterUnion->OverlapComponentWithResult(GeomPose.GetTranslation(), GeomPose.GetRotation(), Geom, TraceChannel, Params, ResponseParams, ObjectParams, NewOverlaps);
 								Result.bHit = !NewOverlaps.IsEmpty();
 							}
+#endif
 
 							return Result;
 						};
