Added a bEmitChangeNotify to EditTextSourceString/EditTextPropertySourceString to control whether they emit the Pre/Post change events in the editor

diff --git a/Engine/Source/Runtime/Engine/Classes/Kismet/KismetTextLibrary.h b/Engine/Source/Runtime/Engine/Classes/Kismet/KismetTextLibrary.h
--- a/Engine/Source/Runtime/Engine/Classes/Kismet/KismetTextLibrary.h
+++ b/Engine/Source/Runtime/Engine/Classes/Kismet/KismetTextLibrary.h
@@ -365,14 +365,15 @@ class UKismetTextLibrary : public UBlueprintFunctionLibrary
 	 * 
 	 * @note This is an ADVANCED function that is ONLY safe to be used in environments where the modified text property will be gathered for localization (eg, in the editor, or a game mode that collects text properties to be localized).
 	 * 
-	 * @param TextOwner		The object that owns the given Text to be edited.
-	 * @param Text			The text property to edit. This must be a property that exists on TextOwner.
-	 * @param SourceString	The source string that the edited text property should use.
+	 * @param TextOwner			The object that owns the given Text to be edited.
+	 * @param Text				The text property to edit. This must be a property that exists on TextOwner.
+	 * @param SourceString		The source string that the edited text property should use.
+	 * @param bEmitChangeNotify	True to emit a pre/post edit change notify (in the editor). You should set this to false if editing text within a construction script.
 	 * 
 	 * @return True if edit was possible, or false if not.
 	 */
 	UFUNCTION(BlueprintCallable, CustomThunk, Category="Utilities|Text", meta=(DefaultToSelf="TextOwner"))
-	static ENGINE_API bool EditTextSourceString(UObject* TextOwner, UPARAM(ref) FText& Text, const FString& SourceString);
+	static ENGINE_API bool EditTextSourceString(UObject* TextOwner, UPARAM(ref) FText& Text, const FString& SourceString, const bool bEmitChangeNotify = true);
 	DECLARE_FUNCTION(execEditTextSourceString);
 
 	/**
@@ -381,15 +382,16 @@ class UKismetTextLibrary : public UBlueprintFunctionLibrary
 	 * 
 	 * @note This is an ADVANCED function that is ONLY safe to be used in environments where the modified text property will be gathered for localization (eg, in the editor, or a game mode that collects text properties to be localized).
 	 * 
-	 * @param TextOwner		The object that owns the given Text to be edited.
-	 * @param PropertyName	The name of the text property to edit. This must be a property that exists on TextOwner.
-	 * @param SourceString	The source string that the edited text property should use.
+	 * @param TextOwner			The object that owns the given Text to be edited.
+	 * @param PropertyName		The name of the text property to edit. This must be a property that exists on TextOwner.
+	 * @param SourceString		The source string that the edited text property should use.
+	 * @param bEmitChangeNotify	True to emit a pre/post edit change notify (in the editor). You should set this to false if editing text within a construction script.
 	 * 
 	 * @return True if edit was possible, or false if not.
 	 */
 	UFUNCTION(BlueprintCallable, Category="Utilities|Text", meta=(DefaultToSelf="TextOwner"))
-	static ENGINE_API bool EditTextPropertySourceString(UObject* TextOwner, const FName PropertyName, const FString& SourceString);
-	static bool Generic_EditTextPropertySourceString(UObject* TextOwner, const FTextProperty* TextProperty, const FString& SourceString);
+	static ENGINE_API bool EditTextPropertySourceString(UObject* TextOwner, const FName PropertyName, const FString& SourceString, const bool bEmitChangeNotify = true);
+	static bool Generic_EditTextPropertySourceString(UObject* TextOwner, const FTextProperty* TextProperty, const FString& SourceString, const bool bEmitChangeNotify);
 
 	/** Converts a generic value to localized formatted text using the user culture. */
 	UFUNCTION(BlueprintPure, CustomThunk, DisplayName = "To Text", Category = "Utilities|Text", meta = (CustomStructureParam = "Value", BlueprintInternalUseOnly = "true"))
diff --git a/Engine/Source/Runtime/Engine/Private/KismetTextLibrary.cpp b/Engine/Source/Runtime/Engine/Private/KismetTextLibrary.cpp
--- a/Engine/Source/Runtime/Engine/Private/KismetTextLibrary.cpp
+++ b/Engine/Source/Runtime/Engine/Private/KismetTextLibrary.cpp
@@ -368,7 +368,7 @@ FText UKismetTextLibrary::PolyglotDataToText(const FPolyglotTextData& PolyglotDa
 	return PolyglotData.GetText();
 }
 
-bool UKismetTextLibrary::EditTextSourceString(UObject* TextOwner, FText& Text, const FString& SourceString)
+bool UKismetTextLibrary::EditTextSourceString(UObject* TextOwner, FText& Text, const FString& SourceString, const bool bEmitChangeNotify)
 {
 	// We should never hit this! Stubbed to avoid NoExport on the class.
 	check(0);
@@ -384,6 +384,8 @@ DEFINE_FUNCTION(UKismetTextLibrary::execEditTextSourceString)
 
 	P_GET_PROPERTY_REF(FStrProperty, SourceString);
 
+	P_GET_PROPERTY(FBoolProperty, bEmitChangeNotify);
+
 	P_FINISH;
 
 	P_NATIVE_BEGIN;
@@ -402,7 +404,7 @@ DEFINE_FUNCTION(UKismetTextLibrary::execEditTextSourceString)
 			return;
 		}
 
-		if (Generic_EditTextPropertySourceString(TextOwner, TextProperty, SourceString))
+		if (Generic_EditTextPropertySourceString(TextOwner, TextProperty, SourceString, bEmitChangeNotify))
 		{
 			*(bool*)RESULT_PARAM = true;
 			Text = TextProperty->GetPropertyValue_InContainer(TextOwner);
@@ -411,7 +413,7 @@ DEFINE_FUNCTION(UKismetTextLibrary::execEditTextSourceString)
 	P_NATIVE_END;
 }
 
-bool UKismetTextLibrary::EditTextPropertySourceString(UObject* TextOwner, const FName PropertyName, const FString& SourceString)
+bool UKismetTextLibrary::EditTextPropertySourceString(UObject* TextOwner, const FName PropertyName, const FString& SourceString, const bool bEmitChangeNotify)
 {
 	if (!TextOwner)
 	{
@@ -426,10 +428,10 @@ bool UKismetTextLibrary::EditTextPropertySourceString(UObject* TextOwner, const
 		return false;
 	}
 
-	return Generic_EditTextPropertySourceString(TextOwner, TextProperty, SourceString);
+	return Generic_EditTextPropertySourceString(TextOwner, TextProperty, SourceString, bEmitChangeNotify);
 }
 
-bool UKismetTextLibrary::Generic_EditTextPropertySourceString(UObject* TextOwner, const FTextProperty* TextProperty, const FString& SourceString)
+bool UKismetTextLibrary::Generic_EditTextPropertySourceString(UObject* TextOwner, const FTextProperty* TextProperty, const FString& SourceString, const bool bEmitChangeNotify)
 {
 	checkf(TextOwner, TEXT("TextOwner is null!"));
 	checkf(TextProperty, TEXT("TextProperty is null!"));
@@ -476,8 +478,11 @@ bool UKismetTextLibrary::Generic_EditTextPropertySourceString(UObject* TextOwner
 			PropertyAccessUtil::GetArchetypeInstancesInheritingPropertyValue(TextProperty, TextOwner, InheritedInstances);
 		}
 
-		ChangeNotify = PropertyAccessUtil::BuildBasicChangeNotify(TextProperty, TextOwner, EPropertyAccessChangeNotifyMode::Always, EPropertyChangeType::ValueSet);
-		PropertyAccessUtil::EmitPreChangeNotify(ChangeNotify.Get(), false);
+		if (bEmitChangeNotify)
+		{
+			ChangeNotify = PropertyAccessUtil::BuildBasicChangeNotify(TextProperty, TextOwner, EPropertyAccessChangeNotifyMode::Always, EPropertyChangeType::ValueSet);
+			PropertyAccessUtil::EmitPreChangeNotify(ChangeNotify.Get(), false);
+		}
 	}
 
 	auto DeterministicTextKeyGenerator = [TextOwner, TextProperty, bIsEditingEditorObject]()
@@ -498,7 +503,10 @@ bool UKismetTextLibrary::Generic_EditTextPropertySourceString(UObject* TextOwner
 			}
 		}
 
-		PropertyAccessUtil::EmitPostChangeNotify(ChangeNotify.Get(), false);
+		if (bEmitChangeNotify)
+		{
+			PropertyAccessUtil::EmitPostChangeNotify(ChangeNotify.Get(), false);
+		}
 	}
 
 	return bEditedValue;
