Reduce number of uniform buffers created by ISM by deduplicating FInstancedStaticMeshVFLooseUniformShaderParameters

diff --git a/Engine/Source/Runtime/Engine/Private/InstancedStaticMesh.cpp b/Engine/Source/Runtime/Engine/Private/InstancedStaticMesh.cpp
--- a/Engine/Source/Runtime/Engine/Private/InstancedStaticMesh.cpp
+++ b/Engine/Source/Runtime/Engine/Private/InstancedStaticMesh.cpp
@@ -1151,6 +1151,69 @@ void FInstancedStaticMeshSceneProxy::GetDynamicMeshElements(const TArray<const F
 	}
 }
 
+static uint32 GetTypeHash(const FInstancedStaticMeshVFLooseUniformShaderParameters& Params)
+{
+	uint32 HashParameters = 0;
+	HashParameters = HashCombine(HashParameters, GetTypeHash(Params.InstancingFadeOutParams));
+	HashParameters = HashCombine(HashParameters, GetTypeHash(Params.InstancingTranslatedWorldViewOriginOne));
+	HashParameters = HashCombine(HashParameters, GetTypeHash(Params.InstancingTranslatedWorldViewOriginZero));
+	HashParameters = HashCombine(HashParameters, GetTypeHash(Params.InstancingViewZCompareOne));
+	HashParameters = HashCombine(HashParameters, GetTypeHash(Params.InstancingViewZCompareZero));
+	HashParameters = HashCombine(HashParameters, GetTypeHash(Params.InstancingViewZConstant));
+	return HashParameters;
+}
+
+static bool operator== (const FInstancedStaticMeshVFLooseUniformShaderParameters& A, const FInstancedStaticMeshVFLooseUniformShaderParameters& B)
+{
+	static_assert(sizeof(FInstancedStaticMeshVFLooseUniformShaderParameters) == 96, "FInstancedStaticMeshVFLooseUniformShaderParameters size has changed, make sure operator== is updated");
+	
+	return A.InstancingFadeOutParams == B.InstancingFadeOutParams
+		&& A.InstancingTranslatedWorldViewOriginOne == B.InstancingTranslatedWorldViewOriginOne
+		&& A.InstancingTranslatedWorldViewOriginZero == B.InstancingTranslatedWorldViewOriginZero
+		&& A.InstancingViewZCompareOne == B.InstancingViewZCompareOne
+		&& A.InstancingViewZCompareZero == B.InstancingViewZCompareZero
+		&& A.InstancingViewZConstant == B.InstancingViewZConstant;
+}
+
+// For deduplicating FInstancedStaticMeshVFLooseUniformShaderParameters UBs
+class FInstancedStaticMeshVFUniformsCache : public FRenderResource
+{
+	const static int32 MAX_CACHE_SIZE = 10;
+	TMap<FInstancedStaticMeshVFLooseUniformShaderParameters, FInstancedStaticMeshVFLooseUniformShaderParametersRef> CachedParams;
+
+public:
+	FInstancedStaticMeshVFUniformsCache()
+		: FRenderResource()
+	{
+	}
+	
+	virtual void ReleaseRHI() override
+	{
+		CachedParams.Empty();
+	}
+
+	FInstancedStaticMeshVFLooseUniformShaderParametersRef CreateUniformBuffer(const FInstancedStaticMeshVFLooseUniformShaderParameters& Params, EUniformBufferUsage UniformBufferUsage)
+	{
+		check(IsInRenderingThread());
+
+		FInstancedStaticMeshVFLooseUniformShaderParametersRef* CachedUB = CachedParams.Find(Params);
+		if (CachedUB)
+		{
+			return *CachedUB;
+		}
+
+		FInstancedStaticMeshVFLooseUniformShaderParametersRef UB = CreateUniformBufferImmediate(Params, UniformBufferUsage);
+		if (CachedParams.Num() < MAX_CACHE_SIZE && UniformBufferUsage == UniformBuffer_MultiFrame)
+		{
+			CachedParams.Add(Params, UB);
+		}
+
+		return UB;
+	}
+};
+
+static TGlobalResource<FInstancedStaticMeshVFUniformsCache> GInstancedStaticMeshVFUniformsCache;
+
 FInstancedStaticMeshVFLooseUniformShaderParametersRef FInstancedStaticMeshSceneProxy::CreateLooseUniformBuffer(const FSceneView* View, const FInstancingUserData* InstancingUserData, uint32 InstancedLODRange, uint32 InstancedLODIndex, EUniformBufferUsage UniformBufferUsage) const
 {
 	FInstancedStaticMeshVFLooseUniformShaderParameters LooseParameters;
@@ -1299,7 +1362,14 @@ FInstancedStaticMeshVFLooseUniformShaderParametersRef FInstancedStaticMeshSceneP
 
 	}
 
-	return CreateUniformBufferImmediate(LooseParameters, UniformBufferUsage);
+	if (IsInRenderingThread())
+	{
+		return GInstancedStaticMeshVFUniformsCache.CreateUniformBuffer(LooseParameters, UniformBufferUsage);
+	}
+	else
+	{
+		return CreateUniformBufferImmediate(LooseParameters, UniformBufferUsage);
+	}
 }
 
 FInstancedStaticMeshSceneProxyDesc::FInstancedStaticMeshSceneProxyDesc(const UInstancedStaticMeshComponent* InComponent)
