New material translator: refactored scalar folding function for consistency, matching the switch-case order to the one in which enums are declared. Also moved "(Node %s)" to error message prefix to the FEmitter::Error function, so that all error messages show which expression originated the error.

diff --git a/Engine/Source/Runtime/Engine/Private/Materials/MaterialExpressionsIR.cpp b/Engine/Source/Runtime/Engine/Private/Materials/MaterialExpressionsIR.cpp
--- a/Engine/Source/Runtime/Engine/Private/Materials/MaterialExpressionsIR.cpp
+++ b/Engine/Source/Runtime/Engine/Private/Materials/MaterialExpressionsIR.cpp
@@ -824,7 +824,7 @@ void UMaterialExpressionTextureCoordinate::Build(UE::MIR::FEmitter& Emitter)
 {
 	if (UnMirrorU || UnMirrorV)
 	{
-		Emitter.Error("unmirroring unsupported");
+		Emitter.Error(TEXT("Unmirroring unsupported"));
 	}
 
 	MIR::FValue* OutputValue = Emitter.EmitExternalInput(MIR::TexCoordIndexToExternalInput(CoordinateIndex));
diff --git a/Engine/Source/Runtime/Engine/Private/Materials/MaterialIREmitter.cpp b/Engine/Source/Runtime/Engine/Private/Materials/MaterialIREmitter.cpp
--- a/Engine/Source/Runtime/Engine/Private/Materials/MaterialIREmitter.cpp
+++ b/Engine/Source/Runtime/Engine/Private/Materials/MaterialIREmitter.cpp
@@ -207,11 +207,7 @@ FValue* FEmitter::Input(const FExpressionInput* InInput)
 	FValue* Value = TryInput(InInput);
 	if (!Value)
 	{
-		// Add the node type to the error message
-		const int32 ChopCount = FCString::Strlen(TEXT("MaterialExpression"));
-		const FString ErrorClassName = Expression->GetClass()->GetName();
-
-		Errorf(TEXT("(Node %s) Missing '%s' input value."), *ErrorClassName.Right(ErrorClassName.Len() - ChopCount), *SlowFindInputName(Expression, InInput).ToString());
+		Errorf(TEXT("Missing '%s' input value."), *SlowFindInputName(Expression, InInput).ToString());
 		return MIR::FPoison::Get();
 	}
 	return Value;
@@ -755,28 +751,20 @@ static bool FoldComparisonOperatorScalar(EOperator Operator, T A, T B)
 	}
 }
 
-template <typename T>
-static bool ValidateFoldScalarArithmeticOperator(FEmitter* Emitter, EOperator Operator, T A, T B, T C)
+static TInteger FoldScalarBitOperator(EOperator Operator, TInteger A, TInteger B)
 {
 	switch (Operator)
 	{
-		case UO_Logarithm:
-		case UO_Logarithm2:
-		case UO_Logarithm10:
-		{
-			if (A <= 0)
-			{
-				Emitter->Errorf(TEXT("Logarithm of 0 or a negative number is not valid."));
-				return false;
-			}
-		}
+		case UO_Not: return !A;
+		case BO_And: return A & B;
+		case BO_Or: return A | B;
+		case BO_Xor: return A ^ B;
+		default: UE_MIR_UNREACHABLE();
 	}
-
-	return true;
 }
 
 template <typename T>
-static T FoldFloatScalarArithmeticOperator(EOperator Operator, T A, T B, T C)
+static TOptional<T> FoldScalarArithmeticOperator(FEmitter* Emitter, EOperator Operator, T A, T B, T C)
 {
 	if constexpr (std::is_same<T, float>::value || std::is_same<T, double>::value)
 	{
@@ -785,118 +773,77 @@ static T FoldFloatScalarArithmeticOperator(EOperator Operator, T A, T B, T C)
 			case UO_ACos: return FMath::Acos(A);
 			case UO_ASin: return FMath::Asin(A);
 			case UO_ATan: return FMath::Atan(A);
+			case UO_Ceil:
+				if constexpr (std::is_same<T, float>::value)
+				{
+					return FMath::CeilToFloat(A);
+				}
+				else
+				{
+					return FMath::CeilToDouble(A);
+				}
 			case UO_Cos: return FMath::Cos(A);
+			case UO_Exponential: return FGenericPlatformMath::Pow(UE_EULERS_NUMBER, A);
+			case UO_Exponential2: return FGenericPlatformMath::Pow(2.0f, A);
+			case UO_Floor:
+				if constexpr (std::is_same<T, float>::value)
+				{
+					return FMath::FloorToFloat(A);
+				}
+				else
+				{
+					return FMath::FloorToDouble(A);
+				}
+			case UO_Logarithm:
+			case UO_Logarithm2:
+			case UO_Logarithm10:
+			{
+				if (A <= 0)
+				{
+					Emitter->Errorf(TEXT("Logarithm of 0 or a negative number is not valid."));
+					return NullOpt;
+				}
+				const T Base = Operator == UO_Logarithm ? UE_EULERS_NUMBER
+					: Operator == UO_Logarithm2 ? T(2)
+					: T(10);
+				return FGenericPlatformMath::LogX(Base, A);
+			}
+			case UO_Round:
+				if constexpr (std::is_same<T, float>::value)
+				{
+					return FMath::RoundToFloat(A);
+				}
+				else
+				{
+					return FMath::RoundToDouble(A);
+				}
 			case UO_Sin: return FMath::Sin(A);
 			case UO_Tan: return FMath::Tan(A);
-			default: UE_MIR_UNREACHABLE();
+			case UO_Truncate:
+				if constexpr (std::is_same<T, float>::value)
+				{
+					return FMath::TruncToFloat(A);
+				}
+				else
+				{
+					return FMath::TruncToDouble(A);
+				}
+			case BO_Fmod: return FGenericPlatformMath::Fmod(A, B);
+			case TO_Lerp: return FMath::Lerp<T>(A, B, C);
+			case TO_Smoothstep: return FMath::SmoothStep<T>(A, B, C);
 		}
 	}
-	else
-	{
-		UE_MIR_UNREACHABLE();
-	}
-}
 
-template <typename T>
-static T FoldScalarArithmeticOperator(EOperator Operator, T A, T B, T C)
-{
 	switch (Operator)
 	{
-		case UO_Abs: return FGenericPlatformMath::Abs<T>(A);
-		case UO_Ceil:
-		{
-			if constexpr (std::is_same<T, float>::value || std::is_same<T, double>::value)
-			{
-				return FMath::CeilToFloat(A);
-			}
-			else if constexpr (std::is_same<T, TInteger>::value)
-			{
-				return A;
-			}
-		}
-		case UO_Floor:
-		{
-			if constexpr (std::is_same<T, float>::value || std::is_same<T, double>::value)
-			{
-				return FGenericPlatformMath::FloorToFloat(A);
-			}
-			else if constexpr (std::is_same<T, TInteger>::value)
-			{
-				return A;
-			}
-		}
 		case UO_Negate: return -A;
-		case UO_ACos: [[fallthrough]];
-		case UO_ASin: [[fallthrough]];
-		case UO_ATan: [[fallthrough]];
-		case UO_Cos: [[fallthrough]];
-		case UO_Sin: [[fallthrough]];
-		case UO_Tan:
-		{
-			return FoldFloatScalarArithmeticOperator<T>(Operator, A, B, C);
-		}
-		case UO_Exponential: return FGenericPlatformMath::Pow(UE_EULERS_NUMBER, A);
-		case UO_Exponential2: return FGenericPlatformMath::Pow(2.0f, A);
-		case UO_Logarithm: return FGenericPlatformMath::LogX(UE_EULERS_NUMBER, A);
-		case UO_Logarithm2: return FGenericPlatformMath::LogX(2.0f, A);
-		case UO_Logarithm10: return FGenericPlatformMath::LogX(10.0f, A);
-		case UO_Round:
-		{
-			if constexpr (std::is_same<T, float>::value || std::is_same<T, double>::value)
-			{
-				return FGenericPlatformMath::RoundToFloat(A);
-			}
-			else if constexpr (std::is_same<T, TInteger>::value)
-			{
-				return A;
-			}
-		}
-		case UO_Truncate:
-		{
-			if constexpr (std::is_same<T, float>::value)
-			{
-				return FMath::TruncToFloat(A);
-			}
-			else if constexpr (std::is_same<T, double>::value)
-			{
-				return FMath::TruncToDouble(A);
-			}
-			else
-			{
-				UE_MIR_UNREACHABLE();
-			}
-		}
-		case BO_Add: return A + B;\
-		case BO_Fmod:
-		{
-			if constexpr (std::is_same<T, float>::value || std::is_same<T, double>::value)
-			{
-				return FGenericPlatformMath::Fmod(A, B);
-			}
-			else if constexpr (std::is_same<T, TInteger>::value)
-			{
-				return 0.0f;
-			}
-		}
+		case BO_Add: return A + B;
 		case BO_Subtract: return A - B;
 		case BO_Multiply: return A * B;
 		case BO_Divide: return A / B;
 		case BO_Min: return FMath::Min<T>(A, B);
 		case BO_Max: return FMath::Max<T>(A, B);
-		case TO_Lerp: return FMath::Lerp<T>(A, B, C);
-		case TO_Smoothstep: return FMath::SmoothStep<T>(A, B, C);
-		default: UE_MIR_UNREACHABLE();
-	}
-}
-
-static TInteger FoldScalarBitOperator(EOperator Operator, TInteger A, TInteger B)
-{
-	switch (Operator)
-	{
-		case UO_Not: return !A;
-		case BO_And: return A & B;
-		case BO_Or: return A | B;
-		case BO_Xor: return A ^ B;
+		case UO_Abs: return FGenericPlatformMath::Abs<T>(A);
 		default: UE_MIR_UNREACHABLE();
 	}
 }
@@ -977,27 +924,27 @@ static FValue* TryFoldOperatorScalar(FEmitter* Emitter, EOperator Op, FValue* A,
 			break; // todo
 
 		/* No simplifications for these operations */
+		case BO_Fmod:
+		case BO_Max:
+		case BO_Min:
+		case UO_Abs:
 		case UO_ACos:
 		case UO_ASin:
 		case UO_ATan:
+		case UO_Ceil:
 		case UO_Cos:
-		case UO_Negate:
-		case UO_Not:
-		case UO_Sin:
-		case UO_Tan:
-		case UO_Truncate:
 		case UO_Exponential:
 		case UO_Exponential2:
+		case UO_Floor:
 		case UO_Logarithm:
-		case UO_Logarithm2:
 		case UO_Logarithm10:
-		case BO_Fmod:
-		case BO_Min:
-		case BO_Max:
-		case UO_Abs:
-		case UO_Ceil:
-		case UO_Floor:
+		case UO_Logarithm2:
+		case UO_Negate:
+		case UO_Not:
 		case UO_Round:
+		case UO_Sin:
+		case UO_Tan:
+		case UO_Truncate:
 			break;
 
 		default:
@@ -1058,24 +1005,14 @@ static FValue* TryFoldOperatorScalar(FEmitter* Emitter, EOperator Op, FValue* A,
 		{
 			case ScalarKind_Int:
 			{
-				if (!ValidateFoldScalarArithmeticOperator<TInteger>(Emitter, Op, AConstant->Integer, BConstant ? BConstant->Integer : 0, CConstant ? CConstant->Integer : 0))
-				{
-					return nullptr;
-				}
-
-				TInteger Result = FoldScalarArithmeticOperator<TInteger>(Op, AConstant->Integer, BConstant ? BConstant->Integer : 0, CConstant ? CConstant->Integer : 0);
-				return Emitter->EmitConstantIntScalar(Result);
+				TOptional<TInteger> Result = FoldScalarArithmeticOperator<TInteger>(Emitter, Op, AConstant->Integer, BConstant ? BConstant->Integer : 0, CConstant ? CConstant->Integer : 0);
+				return Result ? Emitter->EmitConstantIntScalar(*Result) : nullptr;
 			}
 
 			case ScalarKind_Float:
 			{
-				if (!ValidateFoldScalarArithmeticOperator<TFloat>(Emitter, Op, AConstant->Float, BConstant ? BConstant->Float : 0, CConstant ? CConstant->Float : 0))
-				{
-					return nullptr;
-				}
-
-				TFloat Result = FoldScalarArithmeticOperator<TFloat>(Op, AConstant->Float, BConstant ? BConstant->Float : 0, CConstant ? CConstant->Float : 0);
-				return Emitter->EmitConstantFloatScalar(Result);
+				TOptional<TFloat> Result = FoldScalarArithmeticOperator<TFloat>(Emitter, Op, AConstant->Float, BConstant ? BConstant->Float : 0, CConstant ? CConstant->Float : 0);
+				return Result ? Emitter->EmitConstantFloatScalar(*Result) : nullptr;
 			}
 
 			default:
@@ -1103,18 +1040,18 @@ FValue* FEmitter::EmitOperator(EOperator Op, FValue* A, FValue* B, FValue* C)
 		return MIR::FPoison::Get();
 	}
 
-	// Get the operands primitive type.
+	// Get the arguments primitive type.
 	FPrimitiveTypePtr PrimitiveType = A->Type->AsPrimitive();
 	if (!PrimitiveType)
 	{
-		Errorf(TEXT("Operator on operands of non-primitive type."));
+		Errorf(TEXT("Operator on arguments of non-primitive type."));
 		return MIR::FPoison::Get();
 	}
 	
 	// If operator is logical, arguments must be booleans.
 	if (IsBitOperator(Op) && !A->Type->IsBoolean() && !A->Type->IsInteger())
 	{
-		Errorf(TEXT("Logical operator on non-boolean operands."));
+		Errorf(TEXT("Bit operator on non-boolean, non-integer arguments."));
 		return MIR::FPoison::Get();
 	}
 
@@ -1124,7 +1061,6 @@ FValue* FEmitter::EmitOperator(EOperator Op, FValue* A, FValue* B, FValue* C)
 		? FPrimitiveType::Get(ScalarKind_Bool, PrimitiveType->NumRows, PrimitiveType->NumColumns)
 		: PrimitiveType;
 
-
 	// Check that at least one component of the resulting dimensional value would folded.
 	// If all components of resulting dimensional value are not folded, then instead of emitting
 	// an individual operator instruction for each component, simply emit a single binary operator
@@ -1241,7 +1177,7 @@ FValue* FEmitter::EmitBranch(FValue* Condition, FValue* True, FValue* False)
 	// Condition must be of type bool
 	check(Condition->Type->IsBoolScalar());
 
-	// If the condition is not static, the types of the true and false operands must match.
+	// If the condition is not static, the types of the true and false arguments must match.
 	// The resulting type will be that of the branch instruction created.
 	check(True->Type == False->Type);
 
@@ -1473,11 +1409,17 @@ FPrimitiveTypePtr FEmitter::GetCommonPrimitiveType(FPrimitiveTypePtr A, FPrimiti
 	return FPrimitiveType::Get(ScalarKind, NumRows, 1);
 }
 
-void FEmitter::Error(FString Message)
+void FEmitter::Error(FStringView Message)
 {
 	FMaterialIRModule::FError Error;
 	Error.Expression = Expression;
-	Error.Message = MoveTemp(Message);
+
+	// Add the node type to the error message
+	const int32 ChopCount = FCString::Strlen(TEXT("MaterialExpression"));
+	const FString& ErrorClassName = Expression->GetClass()->GetName();
+	
+	Error.Message.Appendf(TEXT("(Node %s) %s"), *ErrorClassName + ChopCount, Message.GetData());
+	
 	Module->Errors.Push(Error);
 	bExpressionHasError = true;
 }
diff --git a/Engine/Source/Runtime/Engine/Public/Materials/MaterialIREmitter.h b/Engine/Source/Runtime/Engine/Public/Materials/MaterialIREmitter.h
--- a/Engine/Source/Runtime/Engine/Public/Materials/MaterialIREmitter.h
+++ b/Engine/Source/Runtime/Engine/Public/Materials/MaterialIREmitter.h
@@ -176,7 +176,7 @@ public:
 		Error(FString::Printf(Format, Args...));
 	}
 
-	void Error(FString Message);
+	void Error(FStringView Message);
 	
 	struct FPrivate;
 
