Fix an issue with garbage reference to child actor component during cook

The issue happens during FWorldPartitionLevelHelper::MoveExternalActorsToLevel. When we try to Rename an Actor that has a ChildActorComponent, all of its components are unregistered first, which leads to the ChildActorComponent attempting to destroy the ChildActor. However, in UChildActorComponent::DestroyChildActor one of the conditions that can prevent the ChildActor from being destroyed is IsBeingRemovedFromLevel and that's what happens in this case. IsBeingRemovedFromLevel returns true, because the ChildActor is not found in MyLevel->Actors array. Since that ChildActor is not destroyed, later in MoveExternalActorsToLevel we move it to the new Level and that's what leads to that reachable garbage object warning.

When opening that level in the editor, MyLevel->Actors array is properly populated, so that ChildActor gets properly destroyed. However, that seems not to be the case when loading that Level during cook, since the loading stack is not the same.

This fix effectively reverts the change from CL 21726304, which added the extra condition verifying that the ChildActor can be found in MyLevel->Actors array.

diff --git a/Engine/Source/Runtime/Engine/Classes/Components/ChildActorComponent.h b/Engine/Source/Runtime/Engine/Classes/Components/ChildActorComponent.h
--- a/Engine/Source/Runtime/Engine/Classes/Components/ChildActorComponent.h
+++ b/Engine/Source/Runtime/Engine/Classes/Components/ChildActorComponent.h
@@ -228,7 +228,7 @@ public:
 private:
 	ENGINE_API bool IsChildActorReplicated() const;
 
-	ENGINE_API bool IsBeingRemovedFromLevel() const;
+	ENGINE_API bool IsLevelBeingRemoved() const;
 
 	UFUNCTION()
 	void OnRep_ChildActor();
diff --git a/Engine/Source/Runtime/Engine/Private/Components/ChildActorComponent.cpp b/Engine/Source/Runtime/Engine/Private/Components/ChildActorComponent.cpp
--- a/Engine/Source/Runtime/Engine/Private/Components/ChildActorComponent.cpp
+++ b/Engine/Source/Runtime/Engine/Private/Components/ChildActorComponent.cpp
@@ -650,25 +650,13 @@ bool UChildActorComponent::IsChildActorReplicated() const
 	return (bChildActorClassReplicated || bChildActorTemplateReplicated);
 }
 
-bool UChildActorComponent::IsBeingRemovedFromLevel() const
+bool UChildActorComponent::IsLevelBeingRemoved() const
 {
 	if (AActor* MyOwner = GetOwner())
 	{
 		if (ULevel* MyLevel = MyOwner->GetLevel())
 		{
-			if (MyLevel->bIsBeingRemoved)
-			{
-				return true;
-			}
-
-#if WITH_EDITOR
-			// In the editor, the child actor can already be removed from the world if the component owner was removed first
-			UWorld* MyWorld = MyLevel->GetWorld();
-			if (MyWorld && !MyWorld->IsGameWorld() && !MyLevel->Actors.Contains(ChildActor))
-			{
-				return true;
-			}
-#endif
+			return MyLevel->bIsBeingRemoved;
 		}
 	}
 
@@ -695,7 +683,7 @@ void UChildActorComponent::OnChildActorDestroyed(AActor* Actor)
 		return;
 	}
 
-	if (Actor && (Actor->HasAuthority() || !IsChildActorReplicated()) && !IsBeingRemovedFromLevel())
+	if (Actor && (Actor->HasAuthority() || !IsChildActorReplicated()) && !IsLevelBeingRemoved())
 	{
 		UWorld* World = Actor->GetWorld();
 		// World may be nullptr during shutdown
@@ -918,7 +906,7 @@ void UChildActorComponent::DestroyChildActor()
 	// If the level is being removed then don't destroy the child actor so re-adding it doesn't
 	// need to create a new actor
 
-	if (ChildActor && (ChildActor->HasAuthority() || !IsChildActorReplicated()) && !IsBeingRemovedFromLevel())
+	if (ChildActor && (ChildActor->HasAuthority() || !IsChildActorReplicated()) && !IsLevelBeingRemoved())
 	{
 		QUICK_SCOPE_CYCLE_COUNTER(STAT_DestroyChildActor);
 
