Switch AssetManager/GFP to use AssetRegistry.IsGathering to better defer its searches

diff --git a/Engine/Plugins/Runtime/GameFeatures/Source/GameFeatures/Private/GameFeaturesSubsystem.cpp b/Engine/Plugins/Runtime/GameFeatures/Source/GameFeatures/Private/GameFeaturesSubsystem.cpp
--- a/Engine/Plugins/Runtime/GameFeatures/Source/GameFeatures/Private/GameFeaturesSubsystem.cpp
+++ b/Engine/Plugins/Runtime/GameFeatures/Source/GameFeatures/Private/GameFeaturesSubsystem.cpp
@@ -559,7 +559,7 @@ void UGameFeaturesSubsystem::AddGameFeatureToAssetManager(const UGameFeatureData
 	// We just mounted the folder that contains these primary assets and the editor background scan may not
 	// not be finished by the time this is called, but a rescan will happen later in OnAssetRegistryFilesLoaded 
 	// as long as LocalAssetRegistry.IsLoadingAssets() is true.
-	const bool bForceSynchronousScan = !LocalAssetRegistry.IsLoadingAssets();
+	const bool bForceSynchronousScan = !LocalAssetRegistry.IsGathering();
 
 	for (FPrimaryAssetTypeInfo TypeInfo : GameFeatureToAdd->GetPrimaryAssetTypesToScan())
 	{
diff --git a/Engine/Source/Runtime/Engine/Private/AssetManager.cpp b/Engine/Source/Runtime/Engine/Private/AssetManager.cpp
--- a/Engine/Source/Runtime/Engine/Private/AssetManager.cpp
+++ b/Engine/Source/Runtime/Engine/Private/AssetManager.cpp
@@ -409,7 +409,7 @@ void UAssetManager::PostInitProperties()
 		if (bIsGlobalAsyncScanEnvironment)
 		{
 			// Listen for when the asset registry has finished discovering files
-			AssetRegistry.OnFilesLoaded().AddUObject(this, &UAssetManager::OnAssetRegistryFilesLoaded);
+			AssetRegistry.OnKnownGathersComplete().AddUObject(this, &UAssetManager::OnAssetRegistryFilesLoaded);
 			AssetRegistry.OnInMemoryAssetCreated().AddUObject(this, &UAssetManager::OnInMemoryAssetCreated);
 			AssetRegistry.OnInMemoryAssetDeleted().AddUObject(this, &UAssetManager::OnInMemoryAssetDeleted);
 			AssetRegistry.OnAssetRenamed().AddUObject(this, &UAssetManager::OnAssetRenamed);
@@ -939,7 +939,7 @@ int32 UAssetManager::ScanPathsForPrimaryAssets(FPrimaryAssetType PrimaryAssetTyp
 #if WITH_EDITOR
 	// Cooked data has the asset data already set up
 	const bool bShouldDoSynchronousScan = !bIsGlobalAsyncScanEnvironment || bForceSynchronousScan;
-	if (!bShouldDoSynchronousScan && GetAssetRegistry().IsLoadingAssets())
+	if (!bShouldDoSynchronousScan && GetAssetRegistry().IsGathering())
 	{
 		// Keep track of the paths we asked for so once assets are discovered we will refresh the list
 		for (const FString& Path : Paths)
@@ -3964,7 +3964,7 @@ void UAssetManager::UpdateManagementDatabase(bool bForceRefresh)
 	// Construct the asset management map and pass it to the asset registry
 	IAssetRegistry& AssetRegistry = GetAssetRegistry();
 
-	if (AssetRegistry.IsLoadingAssets())
+	if (AssetRegistry.IsGathering())
 	{
 		bUpdateManagementDatabaseAfterScan = true;
 		return;
