Compressed animation thread-safety and determinism improvements:
- Deprecated public access and editor-time usage of UAnimSequence::CompressedData
- At editor-time all compressed data is now stored as part of DataByPlatformKeyHash TMap (current platform previously was the only one that wasn't)
- At non-editor-time, CompressedData will contain the runtime platform-specific compressed data
- Public access to compressed data is const-only going forward, and only available through a helper struct FScopedCompressedAnimSequence to ensure thread-safetyy
- Deprecated any APIs which allowed non-const non-threadsafe access
- Introduced RW lock around compressed animation data and compression tasks to ensure thread-safety creation, deletion and evaluation
- This relies on FCompressedAnimationDataReadScope and FCompressedAnimationDataWriteScope used through macros (for compile-time exclusion in non-editor builds)
- Introduced options on FAnimExtractContext to enforce compressed animation data (with runtime checks()) according to a specified platform-data or its equivalent compressed data hash
- Introduced Residency API, which allows its user to ensure compressed animation to stay intra-frame resident until its released. This allows for assurance of only evaluating compressed data when desired
- PoseSearchDatabase indexing now uses Residency API and FAnimExtractContext options to enforce compressed animation data sampling, ensuring deterministic cook-behaviour

diff --git a/Engine/Plugins/Animation/ACLPlugin/Source/ACLPlugin/Private/ACLPluginModule.cpp b/Engine/Plugins/Animation/ACLPlugin/Source/ACLPlugin/Private/ACLPluginModule.cpp
--- a/Engine/Plugins/Animation/ACLPlugin/Source/ACLPlugin/Private/ACLPluginModule.cpp
+++ b/Engine/Plugins/Animation/ACLPlugin/Source/ACLPlugin/Private/ACLPluginModule.cpp
@@ -134,7 +134,8 @@ void FACLPlugin::ListCodecs(const TArray<FString>& Args)
 		SIZE_T UsedSize = 0;
 		for (const UAnimSequence* AnimSeq : AnimSequences)
 		{
-			const SIZE_T Size = GetCompressedBoneSize(AnimSeq->CompressedData);
+			UAnimSequence::FScopedCompressedAnimSequence CompressedAnimSequence = AnimSeq->GetCompressedData();
+			const SIZE_T Size = GetCompressedBoneSize(CompressedAnimSequence.Get());
 			if (AnimSeq->BoneCompressionSettings == Settings)
 			{
 				NumReferences++;
@@ -156,8 +157,9 @@ void FACLPlugin::ListCodecs(const TArray<FString>& Args)
 		SIZE_T UsedSize = 0;
 		for (const UAnimSequence* AnimSeq : AnimSequences)
 		{
-			const SIZE_T Size = GetCompressedBoneSize(AnimSeq->CompressedData);
-			if (AnimSeq->CompressedData.BoneCompressionCodec == Codec)
+			UAnimSequence::FScopedCompressedAnimSequence CompressedAnimSequence = AnimSeq->GetCompressedData();
+			const SIZE_T Size = GetCompressedBoneSize(CompressedAnimSequence.Get());
+			if (CompressedAnimSequence.Get().BoneCompressionCodec == Codec)
 			{
 				NumReferences++;
 				UsedSize += Size;
@@ -186,7 +188,8 @@ void FACLPlugin::ListCodecs(const TArray<FString>& Args)
 		SIZE_T UsedSize = 0;
 		for (const UAnimSequence* AnimSeq : AnimSequences)
 		{
-			const SIZE_T Size = GetCompressedCurveSize(AnimSeq->CompressedData);
+			UAnimSequence::FScopedCompressedAnimSequence CompressedAnimSequence = AnimSeq->GetCompressedData();
+			const SIZE_T Size = GetCompressedCurveSize(CompressedAnimSequence.Get());
 			if (AnimSeq->CurveCompressionSettings == Settings)
 			{
 				NumReferences++;
@@ -208,8 +211,9 @@ void FACLPlugin::ListCodecs(const TArray<FString>& Args)
 		SIZE_T UsedSize = 0;
 		for (const UAnimSequence* AnimSeq : AnimSequences)
 		{
-			const SIZE_T Size = GetCompressedCurveSize(AnimSeq->CompressedData);
-			if (AnimSeq->CompressedData.CurveCompressionCodec == Codec)
+			UAnimSequence::FScopedCompressedAnimSequence CompressedAnimSequence = AnimSeq->GetCompressedData();
+			const SIZE_T Size = GetCompressedCurveSize(CompressedAnimSequence.Get());
+			if (CompressedAnimSequence.Get().CurveCompressionCodec == Codec)
 			{
 				NumReferences++;
 				UsedSize += Size;
@@ -228,7 +232,8 @@ void FACLPlugin::ListCodecs(const TArray<FString>& Args)
 		int32 NumReferences = 0;
 		for (const UAnimSequence* AnimSeq : AnimSequences)
 		{
-			UAnimBoneCompressionCodec_ACLDatabase* Codec = Cast<UAnimBoneCompressionCodec_ACLDatabase>(AnimSeq->CompressedData.BoneCompressionCodec);
+			UAnimSequence::FScopedCompressedAnimSequence CompressedAnimSequence = AnimSeq->GetCompressedData();
+			UAnimBoneCompressionCodec_ACLDatabase* Codec = Cast<UAnimBoneCompressionCodec_ACLDatabase>(CompressedAnimSequence.Get().BoneCompressionCodec);
 			if (Codec != nullptr && Codec->DatabaseAsset == Database)
 			{
 				NumReferences++;
@@ -272,28 +277,30 @@ void FACLPlugin::ListAnimSequences(const TArray<FString>& Args)
 	{
 		UE_LOG(LogAnimationCompression, Log, TEXT("%s ..."), *AnimSeq->GetPathName());
 
-		if (AnimSeq->CompressedData.BoneCompressionCodec->Description.IsEmpty())
+		UAnimSequence::FScopedCompressedAnimSequence CompressedAnimSequence = AnimSeq->GetCompressedData();
+		const FCompressedAnimSequence& CompressedData = CompressedAnimSequence.Get();
+		if (CompressedData.BoneCompressionCodec->Description.IsEmpty())
 		{
-			UE_LOG(LogAnimationCompression, Log, TEXT("    uses bone codec %s"), *AnimSeq->CompressedData.BoneCompressionCodec->GetPathName());
+			UE_LOG(LogAnimationCompression, Log, TEXT("    uses bone codec %s"), *CompressedData.BoneCompressionCodec->GetPathName());
 		}
 		else
 		{
-			UE_LOG(LogAnimationCompression, Log, TEXT("    uses bone codec %s (%s)"), *AnimSeq->CompressedData.BoneCompressionCodec->GetPathName(), *AnimSeq->CompressedData.BoneCompressionCodec->Description);
+			UE_LOG(LogAnimationCompression, Log, TEXT("    uses bone codec %s (%s)"), *CompressedData.BoneCompressionCodec->GetPathName(), *CompressedData.BoneCompressionCodec->Description);
 		}
 
-		const SIZE_T BoneDataSize = GetCompressedBoneSize(AnimSeq->CompressedData);
+		const SIZE_T BoneDataSize = GetCompressedBoneSize(CompressedData);
 		UE_LOG(LogAnimationCompression, Log, TEXT("    has %.2f KB of bone data"), BytesToKB(BoneDataSize));
 
 #if WITH_EDITORONLY_DATA
-		if (AnimSeq->CompressedData.CompressedDataStructure)
+		if (CompressedData.CompressedDataStructure)
 		{
-			UE_LOG(LogAnimationCompression, Log, TEXT("    has a bone error of %.4f cm"), AnimSeq->CompressedData.CompressedDataStructure->BoneCompressionErrorStats.MaxError);
+			UE_LOG(LogAnimationCompression, Log, TEXT("    has a bone error of %.4f cm"), CompressedData.CompressedDataStructure->BoneCompressionErrorStats.MaxError);
 		}
 #endif
 
-		UE_LOG(LogAnimationCompression, Log, TEXT("    uses curve codec %s"), *AnimSeq->CompressedData.CurveCompressionCodec->GetPathName());
+		UE_LOG(LogAnimationCompression, Log, TEXT("    uses curve codec %s"), *CompressedData.CurveCompressionCodec->GetPathName());
 
-		const SIZE_T CurveDataSize = GetCompressedCurveSize(AnimSeq->CompressedData);
+		const SIZE_T CurveDataSize = GetCompressedCurveSize(CompressedData);
 		UE_LOG(LogAnimationCompression, Log, TEXT("    has %.2f KB of curve data"), BytesToKB(CurveDataSize));
 
 		BoneDataTotalSize += BoneDataSize;
diff --git a/Engine/Plugins/Animation/ACLPlugin/Source/ACLPlugin/Private/AnimationCompressionLibraryDatabase.cpp b/Engine/Plugins/Animation/ACLPlugin/Source/ACLPlugin/Private/AnimationCompressionLibraryDatabase.cpp
--- a/Engine/Plugins/Animation/ACLPlugin/Source/ACLPlugin/Private/AnimationCompressionLibraryDatabase.cpp
+++ b/Engine/Plugins/Animation/ACLPlugin/Source/ACLPlugin/Private/AnimationCompressionLibraryDatabase.cpp
@@ -8,6 +8,8 @@
 
 #include "LatentActions.h"
 #include "Containers/Ticker.h"
+#include "Interfaces/ITargetPlatformManagerModule.h"
+#include "Animation/IAnimationSequenceCompiler.h"
 
 #include UE_INLINE_GENERATED_CPP_BY_NAME(AnimationCompressionLibraryDatabase)
 
@@ -177,25 +179,43 @@ void UAnimationCompressionLibraryDatabase::BuildDatabase(TArray<uint8>& OutCompr
 	TSet<uint32> SequenceHashes;
 	SequenceHashes.Empty(AnimSequences.Num());
 
+	const ITargetPlatform* CurrentPlatform = GetTargetPlatformManagerRef().GetRunningTargetPlatform();
+	const uint32 LibraryHash = GetTypeHash(this);
+		
+	// Update our compressed data to match the current settings
+	// We might need to revert to non-frame stripped data or the codec might have changed
+	// forcing us to recompress
+	// In practice we'll load directly from the DDC and it should be fast
+	TGuardValue<int32> FrameStrippingGuard(GPerformFrameStripping, 0);
+	
 	for (UAnimSequence* AnimSeq : AnimSequences)
 	{
-		UAnimBoneCompressionCodec_ACLDatabase* DatabaseCodec = Cast<UAnimBoneCompressionCodec_ACLDatabase>(AnimSeq->CompressedData.BoneCompressionCodec);
-		if (DatabaseCodec == nullptr || DatabaseCodec->DatabaseAsset != this)
+		AnimSeq->RequestResidency(CurrentPlatform, LibraryHash);
+	}
+
+	UE::Anim::IAnimSequenceCompilingManager::FinishCompilation(AnimSequences);
+
+	auto ReleaseResidency = [this, CurrentPlatform, &LibraryHash]()
+	{
+		for (UAnimSequence* AnimSeq : AnimSequences)
 		{
-			UE_LOG(LogAnimationCompression, Warning, TEXT("ACL Database mapping is stale. '%s' no longer references it. [%s]"), *AnimSeq->GetPathName(), *GetPathName());
-			continue;
+			AnimSeq->ReleaseResidency(CurrentPlatform, LibraryHash);
 		}
+	};	
 
-		// Update our compressed data to match the current settings
-		// We might need to revert to non-frame stripped data or the codec might have changed
-		// forcing us to recompress
-		// In practice we'll load directly from the DDC and it should be fast
+	for (UAnimSequence* AnimSeq : AnimSequences)
+	{
+		UAnimSequence::FScopedCompressedAnimSequence CompressedAnimSequence = AnimSeq->GetCompressedData();
+		check(CompressedAnimSequence.Get().IsValid(AnimSeq));
+
+		UAnimBoneCompressionCodec_ACLDatabase* DatabaseCodec = Cast<UAnimBoneCompressionCodec_ACLDatabase>(CompressedAnimSequence.Get().BoneCompressionCodec);
+		if (DatabaseCodec == nullptr || DatabaseCodec->DatabaseAsset != this)
 		{
-			TGuardValue<int32> FrameStrippingGuard(GPerformFrameStripping, 0);
-			AnimSeq->CacheDerivedDataForCurrentPlatform();
+			UE_LOG(LogAnimationCompression, Warning, TEXT("ACL Database mapping is stale. '%s' no longer references it. [%s]"), *AnimSeq->GetPathName(), *GetPathName());
+			continue;
 		}
 
-		const FACLDatabaseCompressedAnimData& AnimData = static_cast<const FACLDatabaseCompressedAnimData&>(*AnimSeq->CompressedData.CompressedDataStructure);
+		const FACLDatabaseCompressedAnimData& AnimData = static_cast<const FACLDatabaseCompressedAnimData&>(*CompressedAnimSequence.Get().CompressedDataStructure);
 		if (!AnimData.IsValid())
 		{
 			UE_LOG(LogAnimationCompression, Warning, TEXT("Cannot include a invalid sequence '%s' in the ACL database. [%s]"), *AnimSeq->GetPathName(), *GetPathName());
@@ -223,7 +243,8 @@ void UAnimationCompressionLibraryDatabase::BuildDatabase(TArray<uint8>& OutCompr
 	TArray<const acl::compressed_tracks*> ACLCompressedTracks;
 	for (const UAnimSequence* AnimSeq : CookedSequences)
 	{
-		const FACLDatabaseCompressedAnimData& AnimData = static_cast<const FACLDatabaseCompressedAnimData&>(*AnimSeq->CompressedData.CompressedDataStructure);
+		UAnimSequence::FScopedCompressedAnimSequence CompressedAnimSequence = AnimSeq->GetCompressedData();
+		const FACLDatabaseCompressedAnimData& AnimData = static_cast<const FACLDatabaseCompressedAnimData&>(*CompressedAnimSequence.Get().CompressedDataStructure);
 		ACLCompressedTracks.Add(AnimData.GetCompressedTracks());
 	}
 
@@ -283,6 +304,8 @@ void UAnimationCompressionLibraryDatabase::BuildDatabase(TArray<uint8>& OutCompr
 		}
 
 		UE_LOG(LogAnimationCompression, Error, TEXT("ACL failed to split database: %s [%s]"), ANSI_TO_TCHAR(SplitResult.c_str()), *GetPathName());
+
+		ReleaseResidency();
 		return;
 	}
 
@@ -337,8 +360,10 @@ void UAnimationCompressionLibraryDatabase::BuildDatabase(TArray<uint8>& OutCompr
 	for (int32 MappingIndex = 0; MappingIndex < NumSequences; ++MappingIndex)
 	{
 		UAnimSequence* AnimSeq = CookedSequences[MappingIndex];
+
+		UAnimSequence::FScopedCompressedAnimSequence CompressedAnimSequence = const_cast<const UAnimSequence*>(AnimSeq)->GetCompressedData();
 		const acl::compressed_tracks* CompressedTracks = ACLDBCompressedTracks[MappingIndex];
-		FACLDatabaseCompressedAnimData& AnimData = static_cast<FACLDatabaseCompressedAnimData&>(*AnimSeq->CompressedData.CompressedDataStructure);
+		FACLDatabaseCompressedAnimData& AnimData = static_cast<FACLDatabaseCompressedAnimData&>(*CompressedAnimSequence.Get().CompressedDataStructure);
 
 		// Align our sequence to 16 bytes
 		CompressedSequenceOffset = acl::align_to(CompressedSequenceOffset, 16);
@@ -423,6 +448,8 @@ void UAnimationCompressionLibraryDatabase::BuildDatabase(TArray<uint8>& OutCompr
 	{
 		ACLAllocatorImpl.deallocate(CompressedTracks, CompressedTracks->get_size());
 	}
+	
+	ReleaseResidency();
 }
 
 void UAnimationCompressionLibraryDatabase::UpdatePreviewState(bool bBuildDatabase)
diff --git a/Engine/Plugins/Animation/ACLPlugin/Source/ACLPluginEditor/Private/ACLStatsDumpCommandlet.cpp b/Engine/Plugins/Animation/ACLPlugin/Source/ACLPluginEditor/Private/ACLStatsDumpCommandlet.cpp
--- a/Engine/Plugins/Animation/ACLPlugin/Source/ACLPluginEditor/Private/ACLStatsDumpCommandlet.cpp
+++ b/Engine/Plugins/Animation/ACLPlugin/Source/ACLPluginEditor/Private/ACLStatsDumpCommandlet.cpp
@@ -237,7 +237,8 @@ static int32 GetAnimationTrackIndex(const int32 BoneIndex, const UAnimSequence*
 		return INDEX_NONE;
 	}
 
-	const TArray<FTrackToSkeletonMap>& TrackToSkelMap = AnimSeq->CompressedData.CompressedTrackToSkeletonMapTable;
+	UAnimSequence::FScopedCompressedAnimSequence CompressedAnimSequence = AnimSeq->GetCompressedData();
+	const TArray<FTrackToSkeletonMap>& TrackToSkelMap = CompressedAnimSequence.Get().CompressedTrackToSkeletonMapTable;
 	for (int32 TrackIndex = 0; TrackIndex < TrackToSkelMap.Num(); ++TrackIndex)
 	{
 		const FTrackToSkeletonMap& TrackToSkeleton = TrackToSkelMap[TrackIndex];
@@ -349,10 +350,11 @@ struct SimpleTransformWriter final : public acl::track_writer
 static void CalculateClipError(const acl::track_array_qvvf& Tracks, const UAnimSequence* UEClip, USkeleton* UESkeleton, uint32& OutWorstBone, float& OutMaxError, float& OutWorstSampleTime)
 {
 	// Use the ACL code if we can to calculate the error instead of approximating it with UE.
-	UAnimBoneCompressionCodec_ACLBase* ACLCodec = Cast<UAnimBoneCompressionCodec_ACLBase>(UEClip->CompressedData.BoneCompressionCodec);
+	UAnimSequence::FScopedCompressedAnimSequence CompressedAnimSequence = UEClip->GetCompressedData();
+	UAnimBoneCompressionCodec_ACLBase* ACLCodec = Cast<UAnimBoneCompressionCodec_ACLBase>(CompressedAnimSequence.Get().BoneCompressionCodec);
 	if (ACLCodec != nullptr)
 	{
-		const acl::compressed_tracks* CompressedClipData = acl::make_compressed_tracks(UEClip->CompressedData.CompressedByteStream.GetData());
+		const acl::compressed_tracks* CompressedClipData = acl::make_compressed_tracks(CompressedAnimSequence.Get().CompressedByteStream.GetData());
 
 		const acl::qvvf_transform_error_metric ErrorMetric;
 
@@ -466,7 +468,7 @@ static void CalculateClipError(const acl::track_array_qvvf& Tracks, const UAnimS
 	OutWorstSampleTime = WorstSampleTime;
 }
 
-static void DumpClipDetailedError(const acl::track_array_qvvf& Tracks, UAnimSequence* UEClip, USkeleton* UESkeleton, sjson::ObjectWriter& Writer)
+static void DumpClipDetailedError(const acl::track_array_qvvf& Tracks, const UAnimSequence* UEClip, USkeleton* UESkeleton, sjson::ObjectWriter& Writer)
 {
 	const uint32 NumBones = Tracks.get_num_tracks();
 	const float ClipDuration = Tracks.get_duration();
@@ -511,8 +513,9 @@ static void DumpClipDetailedError(const acl::track_array_qvvf& Tracks, UAnimSequ
 	acl::itransform_error_metric::local_to_object_space_args local_to_object_space_args_lossy = local_to_object_space_args_raw;
 	local_to_object_space_args_lossy.local_transforms = LossyLocalPoseTransforms.GetData();
 
+	UAnimSequence::FScopedCompressedAnimSequence CompressedAnimSequence = UEClip->GetCompressedData();
 	// Use the ACL code if we can to calculate the error instead of approximating it with UE.
-	UAnimBoneCompressionCodec_ACLBase* ACLCodec = Cast<UAnimBoneCompressionCodec_ACLBase>(UEClip->CompressedData.BoneCompressionCodec);
+	const UAnimBoneCompressionCodec_ACLBase* ACLCodec = Cast<const UAnimBoneCompressionCodec_ACLBase>(CompressedAnimSequence.Get().BoneCompressionCodec);
 	if (ACLCodec != nullptr)
 	{
 		uint32 NumOutputBones = 0;
@@ -523,7 +526,7 @@ static void DumpClipDetailedError(const acl::track_array_qvvf& Tracks, UAnimSequ
 
 		local_to_object_space_args_lossy.local_transforms = LossyRemappedLocalPoseTransforms.GetData();
 
-		const acl::compressed_tracks* CompressedClipData = acl::make_compressed_tracks(UEClip->CompressedData.CompressedByteStream.GetData());
+		const acl::compressed_tracks* CompressedClipData = acl::make_compressed_tracks(CompressedAnimSequence.Get().CompressedByteStream.GetData());
 
 		acl::decompression_context<acl::debug_transform_decompression_settings> Context;
 		Context.initialize(*CompressedClipData);
@@ -695,9 +698,11 @@ static void CompressWithUEAuto(FCompressionContext& Context, bool PerformExhaust
 	const uint64 UEElapsedCycles = UEEndTimeCycles - UEStartTimeCycles;
 	const double UEElapsedTimeSec = FPlatformTime::ToSeconds64(UEElapsedCycles);
 
-	if (Context.UEClip->IsCompressedDataValid())
+	if (Context.UEClip->IsBoneCompressedDataValid())
 	{
-		const bool bHasClipData = Context.UEClip->CompressedData.CompressedDataStructure != nullptr;
+		const UAnimSequence* ConstClip = Context.UEClip;
+		UAnimSequence::FScopedCompressedAnimSequence CompressedAnimSequence = ConstClip->GetCompressedData();
+		const bool bHasClipData = CompressedAnimSequence.Get().CompressedDataStructure != nullptr;
 
 		AnimationErrorStats UEErrorStats;
 		uint32 WorstBone = INDEX_NONE;
@@ -706,7 +711,7 @@ static void CompressWithUEAuto(FCompressionContext& Context, bool PerformExhaust
 
 		if (bHasClipData)
 		{
-			UEErrorStats = Context.UEClip->CompressedData.CompressedDataStructure->BoneCompressionErrorStats;
+			UEErrorStats = CompressedAnimSequence.Get().CompressedDataStructure->BoneCompressionErrorStats;
 
 			CalculateClipError(Context.ACLTracks, Context.UEClip, Context.UESkeleton, WorstBone, MaxError, WorstSampleTime);
 		}
@@ -718,7 +723,7 @@ static void CompressWithUEAuto(FCompressionContext& Context, bool PerformExhaust
 		Writer["ue4_auto"] = [&](sjson::ObjectWriter& Writer)	//-V1047
 		{
 			Writer["algorithm_name"] = TCHAR_TO_ANSI(*Context.UEClip->BoneCompressionSettings->GetClass()->GetName());
-			Writer["codec_name"] = TCHAR_TO_ANSI(*GetCodecName(Context.UEClip->CompressedData.BoneCompressionCodec));
+			Writer["codec_name"] = TCHAR_TO_ANSI(*GetCodecName(CompressedAnimSequence.Get().BoneCompressionCodec));
 			Writer["compressed_size"] = CompressedSize;
 			Writer["ue4_compression_ratio"] = UECompressionRatio;
 			Writer["acl_compression_ratio"] = ACLCompressionRatio;
@@ -731,11 +736,11 @@ static void CompressWithUEAuto(FCompressionContext& Context, bool PerformExhaust
 			Writer["acl_worst_bone"] = WorstBone;
 			Writer["acl_worst_time"] = WorstSampleTime;
 
-			if (Context.UEClip->CompressedData.BoneCompressionCodec != nullptr
-				&& Context.UEClip->CompressedData.BoneCompressionCodec->IsA<UAnimCompress>()
+			if (CompressedAnimSequence.Get().BoneCompressionCodec != nullptr
+				&& CompressedAnimSequence.Get().BoneCompressionCodec->IsA<UAnimCompress>()
 				&& bHasClipData)
 			{
-				const FUECompressedAnimData& AnimData = static_cast<FUECompressedAnimData&>(*Context.UEClip->CompressedData.CompressedDataStructure);
+				const FUECompressedAnimData& AnimData = static_cast<FUECompressedAnimData&>(*CompressedAnimSequence.Get().CompressedDataStructure);
 				Writer["rotation_format"] = TCHAR_TO_ANSI(*FAnimationUtils::GetAnimationCompressionFormatString(AnimData.RotationCompressionFormat));
 				Writer["translation_format"] = TCHAR_TO_ANSI(*FAnimationUtils::GetAnimationCompressionFormatString(AnimData.TranslationCompressionFormat));
 				Writer["scale_format"] = TCHAR_TO_ANSI(*FAnimationUtils::GetAnimationCompressionFormatString(AnimData.ScaleCompressionFormat));
@@ -769,9 +774,11 @@ static void CompressWithACL(FCompressionContext& Context, bool PerformExhaustive
 	const uint64 ACLElapsedCycles = ACLEndTimeCycles - ACLStartTimeCycles;
 	const double ACLElapsedTimeSec = FPlatformTime::ToSeconds64(ACLElapsedCycles);
 
-	if (Context.UEClip->IsCompressedDataValid())
+	if (Context.UEClip->IsBoneCompressedDataValid())
 	{
-		const bool bHasClipData = Context.UEClip->CompressedData.CompressedDataStructure != nullptr;
+		const UAnimSequence* ConstClip = Context.UEClip;
+		UAnimSequence::FScopedCompressedAnimSequence CompressedAnimSequence = ConstClip->GetCompressedData();
+		const bool bHasClipData = CompressedAnimSequence.Get().CompressedDataStructure != nullptr;
 
 		AnimationErrorStats UEErrorStats;
 		uint32 WorstBone = INDEX_NONE;
@@ -780,7 +787,7 @@ static void CompressWithACL(FCompressionContext& Context, bool PerformExhaustive
 
 		if (bHasClipData)
 		{
-			UEErrorStats = Context.UEClip->CompressedData.CompressedDataStructure->BoneCompressionErrorStats;
+			UEErrorStats = CompressedAnimSequence.Get().CompressedDataStructure->BoneCompressionErrorStats;
 
 			CalculateClipError(Context.ACLTracks, Context.UEClip, Context.UESkeleton, WorstBone, MaxError, WorstSampleTime);
 		}
@@ -792,7 +799,7 @@ static void CompressWithACL(FCompressionContext& Context, bool PerformExhaustive
 		Writer["ue4_acl"] = [&](sjson::ObjectWriter& Writer)	//-V1047
 		{
 			Writer["algorithm_name"] = TCHAR_TO_ANSI(*Context.UEClip->BoneCompressionSettings->GetClass()->GetName());
-			Writer["codec_name"] = TCHAR_TO_ANSI(*GetCodecName(Context.UEClip->CompressedData.BoneCompressionCodec));
+			Writer["codec_name"] = TCHAR_TO_ANSI(*GetCodecName(CompressedAnimSequence.Get().BoneCompressionCodec));
 			Writer["compressed_size"] = CompressedSize;
 			Writer["ue4_compression_ratio"] = UECompressionRatio;
 			Writer["acl_compression_ratio"] = ACLCompressionRatio;
@@ -870,9 +877,11 @@ static void CompressWithUEKeyReduction(FCompressionContext& Context, bool Perfor
 	const uint64 UEElapsedCycles = UEEndTimeCycles - UEStartTimeCycles;
 	const double UEElapsedTimeSec = FPlatformTime::ToSeconds64(UEElapsedCycles);
 
-	if (Context.UEClip->IsCompressedDataValid())
+	if (Context.UEClip->IsBoneCompressedDataValid())
 	{
-		const bool bHasClipData = Context.UEClip->CompressedData.CompressedDataStructure != nullptr;
+		const UAnimSequence* ConstClip = Context.UEClip;
+		UAnimSequence::FScopedCompressedAnimSequence CompressedAnimSequence = ConstClip->GetCompressedData();
+		const bool bHasClipData = CompressedAnimSequence.Get().CompressedDataStructure != nullptr;
 
 		AnimationErrorStats UEErrorStats;
 		uint32 WorstBone = INDEX_NONE;
@@ -881,7 +890,7 @@ static void CompressWithUEKeyReduction(FCompressionContext& Context, bool Perfor
 
 		if (bHasClipData)
 		{
-			UEErrorStats = Context.UEClip->CompressedData.CompressedDataStructure->BoneCompressionErrorStats;
+			UEErrorStats = CompressedAnimSequence.Get().CompressedDataStructure->BoneCompressionErrorStats;
 
 			CalculateClipError(Context.ACLTracks, Context.UEClip, Context.UESkeleton, WorstBone, MaxError, WorstSampleTime);
 		}
@@ -893,7 +902,7 @@ static void CompressWithUEKeyReduction(FCompressionContext& Context, bool Perfor
 		Writer["ue4_keyreduction"] = [&](sjson::ObjectWriter& Writer)	//-V1047
 		{
 			Writer["algorithm_name"] = TCHAR_TO_ANSI(*Context.UEClip->BoneCompressionSettings->GetClass()->GetName());
-			Writer["codec_name"] = TCHAR_TO_ANSI(*GetCodecName(Context.UEClip->CompressedData.BoneCompressionCodec));
+			Writer["codec_name"] = TCHAR_TO_ANSI(*GetCodecName(CompressedAnimSequence.Get().BoneCompressionCodec));
 			Writer["compressed_size"] = CompressedSize;
 			Writer["ue4_compression_ratio"] = UECompressionRatio;
 			Writer["acl_compression_ratio"] = ACLCompressionRatio;
@@ -927,7 +936,7 @@ static void CompressWithUEKeyReduction(FCompressionContext& Context, bool Perfor
 					return;	// No data, nothing to append
 				}
 
-				const FUECompressedAnimData& AnimData = static_cast<FUECompressedAnimData&>(*Context.UEClip->CompressedData.CompressedDataStructure);
+				const FUECompressedAnimData& AnimData = static_cast<FUECompressedAnimData&>(*CompressedAnimSequence.Get().CompressedDataStructure);
 
 				const AnimDataModelType* ClipData = Context.UEClip->GetDataModel();
 				const int32 NumTracks = ClipData->GetNumBoneTracks();
@@ -1005,7 +1014,7 @@ static void CompressWithUEKeyReduction(FCompressionContext& Context, bool Perfor
 					return;	// No data, nothing to append
 				}
 
-				const FUECompressedAnimData& AnimData = static_cast<FUECompressedAnimData&>(*Context.UEClip->CompressedData.CompressedDataStructure);
+				const FUECompressedAnimData& AnimData = static_cast<FUECompressedAnimData&>(*CompressedAnimSequence.Get().CompressedDataStructure);
 
 				const AnimDataModelType* ClipData = Context.UEClip->GetDataModel();
 				const int32 NumTracks = ClipData->GetNumBoneTracks();
@@ -1304,8 +1313,7 @@ UACLStatsDumpCommandlet::UACLStatsDumpCommandlet(const FObjectInitializer& Objec
 
 static void ClearCompressedData(UAnimSequence* UEClip)
 {
-	UEClip->CompressedData.ClearCompressedBoneData();
-	UEClip->CompressedData.ClearCompressedCurveData();
+	UEClip->ClearAllCachedCookedPlatformData();
 }
 
 int32 UACLStatsDumpCommandlet::Main(const FString& Params)
diff --git a/Engine/Plugins/Animation/PoseSearch/Source/Runtime/Private/PoseSearchAssetSampler.cpp b/Engine/Plugins/Animation/PoseSearch/Source/Runtime/Private/PoseSearchAssetSampler.cpp
--- a/Engine/Plugins/Animation/PoseSearch/Source/Runtime/Private/PoseSearchAssetSampler.cpp
+++ b/Engine/Plugins/Animation/PoseSearch/Source/Runtime/Private/PoseSearchAssetSampler.cpp
@@ -1,9 +1,11 @@
 // Copyright Epic Games, Inc. All Rights Reserved.
 
 #include "PoseSearch/PoseSearchAssetSampler.h"
+#include "Animation/AnimComposite.h"
 #include "Animation/AnimMontage.h"
 #include "Animation/AnimSequence.h"
 #include "Animation/BlendSpace.h"
+#include "Interfaces/ITargetPlatformManagerModule.h"
 #include "PoseSearch/PoseSearchAnimNotifies.h"
 #include "PoseSearch/PoseSearchDefines.h"
 
@@ -57,6 +59,12 @@ static FTransform ExtractRootTransformInternal(const UAnimMontage* AnimMontage,
 		return FTransform::Identity;
 	}
 
+	FAnimExtractContext ExtractionCtx(0.0, true, FDeltaTimeRecord(StartTime), true);
+#if WITH_EDITOR
+	ExtractionCtx.bEnforceCompressedDataSampling = true;
+	ExtractionCtx.TargetPlatform = GetTargetPlatformManagerRef().GetRunningTargetPlatform();
+#endif // WITH_EDITOR
+
 	const FAnimTrack& RootMotionAnimTrack = AnimMontage->SlotAnimTracks[0].AnimTrack;
 	TArray<FRootMotionExtractionStep> RootMotionExtractionSteps;
 	RootMotionAnimTrack.GetRootMotionExtractionStepsForTrackRange(RootMotionExtractionSteps, StartTime, EndTime);
@@ -65,7 +73,7 @@ static FTransform ExtractRootTransformInternal(const UAnimMontage* AnimMontage,
 	{
 		if (CurStep.AnimSequence)
 		{
-			AccumulatedRootMotionParams.Accumulate(CurStep.AnimSequence->ExtractRootMotionFromRange(CurStep.StartPosition, CurStep.EndPosition, FAnimExtractContext()));
+			AccumulatedRootMotionParams.Accumulate(CurStep.AnimSequence->ExtractRootMotionFromRange(CurStep.StartPosition, CurStep.EndPosition, ExtractionCtx));
 		}
 	}
 	return AccumulatedRootMotionParams.GetRootMotionTransform();
@@ -175,9 +183,14 @@ static void ProcessRootTransform(const UBlendSpace* BlendSpace, const FVector& B
 			const float SampleTime = static_cast<float>(SampleIdx) / (NumRootSamples - 1);
 			BlendSpace->ResetBlendSamples(BlendSamplesData, SampleTime, bIsLoopable, true);
 
+			FAnimExtractContext ExtractionCtx(0.0, true, {}, bIsLoopable);
+#if WITH_EDITOR
+			ExtractionCtx.bEnforceCompressedDataSampling = true;
+			ExtractionCtx.TargetPlatform = GetTargetPlatformManagerRef().GetRunningTargetPlatform();
+#endif // WITH_EDITOR
+			
 			// Accumulate root motion after samples have been updated.
 			FRootMotionMovementParams RootMotionMovementParams;
-			FAnimExtractContext ExtractionCtx(0.0, true, {}, bIsLoopable);	
 			for (int32 BlendSampleIndex = 0; BlendSampleIndex < BlendSamplesData.Num(); BlendSampleIndex++)
 			{
 				FBlendSampleData& BlendSample = BlendSamplesData[BlendSampleIndex];
@@ -238,6 +251,15 @@ void FAnimationAssetSampler::Init(const UAnimationAsset* InAnimationAsset, const
 	RootTransformSamplingRate = InRootTransformSamplingRate;
 	CachedPlayLength = GetPlayLength(AnimationAssetPtr.Get(), BlendParameters);
 
+#if WITH_EDITOR
+	ITargetPlatform* TargetPlatform = GetTargetPlatformManagerRef().GetRunningTargetPlatform();		
+	PlatformHash = FIoHash::Zero;
+	if (UAnimSequence* AnimSequence = Cast<UAnimSequence>(const_cast<UAnimationAsset*>(InAnimationAsset)))
+	{		
+		PlatformHash = AnimSequence->GetDerivedDataKeyHash(TargetPlatform);
+	}
+#endif
+
 	if (bPreProcessRootTransform)
 	{
 		Process();
@@ -363,6 +385,9 @@ FTransform FAnimationAssetSampler::GetTotalRootTransform() const
 
 void FAnimationAssetSampler::ExtractPose(const FAnimExtractContext& ExtractionCtx, FAnimationPoseData& OutAnimPoseData) const
 {
+#if WITH_EDITOR
+	check(ExtractionCtx.TargetPlatformHash == FIoHash::Zero);
+#endif // WITH_EDITOR
 	if (const UBlendSpace* BlendSpace = Cast<UBlendSpace>(AnimationAssetPtr.Get()))
 	{
 		TArray<FBlendSampleData> BlendSamples;
@@ -393,7 +418,12 @@ void FAnimationAssetSampler::ExtractPose(const FAnimExtractContext& ExtractionCt
 	}
 	else if (const UAnimSequenceBase* SequenceBase = Cast<UAnimSequenceBase>(AnimationAssetPtr.Get()))
 	{
-		SequenceBase->GetAnimationPose(OutAnimPoseData, ExtractionCtx);
+		FAnimExtractContext SequenceContext = ExtractionCtx;
+#if WITH_EDITOR
+		check(PlatformHash != FIoHash::Zero);
+		SequenceContext.TargetPlatformHash = PlatformHash;
+#endif // WITH_EDITOR
+		SequenceBase->GetAnimationPose(OutAnimPoseData, SequenceContext);
 	}
 	else
 	{
@@ -414,6 +444,8 @@ void FAnimationAssetSampler::ExtractPose(float Time, FCompactPose& OutPose) cons
 
 #if WITH_EDITOR
 	ExtractionCtx.bExtractWithRootMotionProvider = false;
+	ExtractionCtx.bEnforceCompressedDataSampling = true;
+	ExtractionCtx.TargetPlatform = GetTargetPlatformManagerRef().GetRunningTargetPlatform();
 #endif // WITH_EDITOR
 
 	ExtractPose(ExtractionCtx, AnimPoseData);
@@ -431,6 +463,8 @@ void FAnimationAssetSampler::ExtractPose(float Time, FCompactPose& OutPose, FBle
 
 #if WITH_EDITOR
 	ExtractionCtx.bExtractWithRootMotionProvider = false;
+	ExtractionCtx.bEnforceCompressedDataSampling = true;
+	ExtractionCtx.TargetPlatform = GetTargetPlatformManagerRef().GetRunningTargetPlatform();
 #endif // WITH_EDITOR
 
 	ExtractPose(ExtractionCtx, AnimPoseData);
@@ -528,6 +562,11 @@ FTransform FAnimationAssetSampler::ExtractRootTransform(float Time) const
 	else if (const UAnimSequenceBase* SequenceBase = Cast<UAnimSequenceBase>(AnimationAssetPtr.Get()))
 	{
 		FAnimExtractContext ExtractionCtx(0.0, true, FDeltaTimeRecord(Time), true);
+#if WITH_EDITOR
+		ExtractionCtx.bEnforceCompressedDataSampling = true;
+		ExtractionCtx.TargetPlatformHash = PlatformHash;
+#endif // WITH_EDITOR
+		
 		if (IsLoopable())
 		{
 			RootTransform = SequenceBase->ExtractRootMotion(ExtractionCtx);
diff --git a/Engine/Plugins/Animation/PoseSearch/Source/Runtime/Private/PoseSearchDatabase.cpp b/Engine/Plugins/Animation/PoseSearch/Source/Runtime/Private/PoseSearchDatabase.cpp
--- a/Engine/Plugins/Animation/PoseSearch/Source/Runtime/Private/PoseSearchDatabase.cpp
+++ b/Engine/Plugins/Animation/PoseSearch/Source/Runtime/Private/PoseSearchDatabase.cpp
@@ -752,7 +752,11 @@ void UPoseSearchDatabase::PostLoad()
 	}
 #endif // WITH_ENGINE
 
-	FAsyncPoseSearchDatabasesManagement::RequestAsyncBuildIndex(this, Flag);
+	// Delay indexing until BeginCacheForCookedPlatformData when running a CookCommandlet
+	if (!IsRunningCookCommandlet())
+	{
+		FAsyncPoseSearchDatabasesManagement::RequestAsyncBuildIndex(this, Flag);
+	}
 #endif
 
 	Super::PostLoad();
@@ -790,6 +794,12 @@ void UPoseSearchDatabase::SynchronizeWithExternalDependencies()
 	TArray<FAssetIdentifier> Referencers;
 	AssetRegistry.GetReferencers(GetPackage()->GetFName(), Referencers);
 
+	// Sort AssetRegistry results
+	Algo::Sort(Referencers, [](FAssetIdentifier One, FAssetIdentifier Two) -> bool
+	{
+		return One.PackageName.Compare(Two.PackageName) < 0;
+	});
+
 	TArray<UAnimSequenceBase*> SequencesBase;
 	for (const FAssetIdentifier& Referencer : Referencers)
 	{
@@ -1038,7 +1048,7 @@ bool UPoseSearchDatabase::IsCachedCookedPlatformDataLoaded(const ITargetPlatform
 {
 	using namespace UE::PoseSearch;
 	check(IsInGameThread());
-	return EAsyncBuildIndexResult::InProgress != FAsyncPoseSearchDatabasesManagement::RequestAsyncBuildIndex(this, ERequestAsyncBuildFlag::ContinueRequest);
+	return EAsyncBuildIndexResult::Success == FAsyncPoseSearchDatabasesManagement::RequestAsyncBuildIndex(this, ERequestAsyncBuildFlag::ContinueRequest);
 }
 #endif // WITH_EDITOR
 
diff --git a/Engine/Plugins/Animation/PoseSearch/Source/Runtime/Private/PoseSearchDerivedData.cpp b/Engine/Plugins/Animation/PoseSearch/Source/Runtime/Private/PoseSearchDerivedData.cpp
--- a/Engine/Plugins/Animation/PoseSearch/Source/Runtime/Private/PoseSearchDerivedData.cpp
+++ b/Engine/Plugins/Animation/PoseSearch/Source/Runtime/Private/PoseSearchDerivedData.cpp
@@ -31,6 +31,7 @@
 #include "Serialization/BulkDataRegistry.h"
 #include "UObject/NoExportTypes.h"
 #include "UObject/PackageReload.h"
+#include "Interfaces/ITargetPlatformManagerModule.h"
 
 #define LOCTEXT_NAMESPACE "PoseSearchDerivedData"
 
@@ -1347,6 +1348,7 @@ struct FPoseSearchDatabaseAsyncCacheTask
 	};
 
 	FPoseSearchDatabaseAsyncCacheTask(UPoseSearchDatabase* InDatabase, bool bPerformConditionalPostLoadIfRequired, FPartialKeyHashes& PartialKeyHashes);
+	void ClearAnimSequenceResidency();
 	void StartNewRequestIfNeeded(bool bPerformConditionalPostLoadIfRequired, FPartialKeyHashes& PartialKeyHashes);
 	void Update(FCriticalSection& OuterMutex, FPartialKeyHashes& PartialKeyHashes);
 	void Wait(FCriticalSection& OuterMutex);
@@ -1388,6 +1390,7 @@ private:
 		
 	FThreadSafeCounter ThreadSafeState = int32(EState::Notstarted);
 	bool bBroadcastOnDerivedDataRebuild = false;
+	bool bResidencyCleared = true;
 };
 
 class FPoseSearchDatabaseAsyncCacheTasks : public TArray<TUniquePtr<FPoseSearchDatabaseAsyncCacheTask>> {};
@@ -1408,6 +1411,31 @@ FPoseSearchDatabaseAsyncCacheTask::FPoseSearchDatabaseAsyncCacheTask(UPoseSearch
 	}
 }
 
+void FPoseSearchDatabaseAsyncCacheTask::ClearAnimSequenceResidency()
+{
+	if (bResidencyCleared)
+	{
+		return;
+	}
+	
+	const ITargetPlatform* TargetPlatform = GetTargetPlatformManager()->GetRunningTargetPlatform();
+	for (TWeakObjectPtr<const UObject>& Dependency : DatabaseDependencies)
+	{
+		if (Dependency.IsValid())
+		{
+			if (UAnimSequence* AnimSequence = Cast<UAnimSequence>(const_cast<UObject*>(Dependency.Get())))
+			{
+				if (AnimSequence->HasResidency(TargetPlatform))
+				{
+					AnimSequence->ReleaseResidency(TargetPlatform, GetTypeHash(Database.Get()));
+				}
+			}
+		}
+	}
+
+	bResidencyCleared = true;
+}
+
 FPoseSearchDatabaseAsyncCacheTask::~FPoseSearchDatabaseAsyncCacheTask()
 {
 	// Owner.Cancel must be performed before SearchIndex.Reset() in case any task is flying (launched by Owner.LaunchTask)
@@ -1465,12 +1493,12 @@ void FPoseSearchDatabaseAsyncCacheTask::StartNewRequestIfNeeded(bool bPerformCon
 
 	// composing the key
 	const FKeyBuilder KeyBuilder(Database.Get(), true, bPerformConditionalPostLoadIfRequired, &PartialKeyHashes, DebugPartialKeyHashesMode);
-	if (KeyBuilder.AnyAssetNotReady())
+	if (KeyBuilder.AnyAssetNotFullyLoaded())
 	{
 		DerivedDataKey = FIoHash::Zero;
 		SetState(EState::Notstarted);
 	
-		UE_LOG(LogPoseSearch, Log, TEXT("Delaying DDC until dependents post load  - %s"), *Database->GetName());
+		UE_LOG(LogPoseSearch, Log, TEXT("Delaying DDC until dependents are fully loaded - %s"), *Database->GetName());
 	}
 	else
 	{
@@ -1478,28 +1506,42 @@ void FPoseSearchDatabaseAsyncCacheTask::StartNewRequestIfNeeded(bool bPerformCon
 		const bool bHasKeyChanged = NewDerivedDataKey != DerivedDataKey;
 		if (bHasKeyChanged)
 		{
-			DerivedDataKey = NewDerivedDataKey;
-
-			DatabaseDependencies.Reset();
-			DatabaseDependencies.Reserve(KeyBuilder.GetDependencies().Num());
-			for (const UObject* Dependency : KeyBuilder.GetDependencies())
+			if (!KeyBuilder.AnyAssetNotReady())
 			{
-				DatabaseDependencies.Add(Dependency);
-			}
+				DerivedDataKey = NewDerivedDataKey;
 
-			SetState(EState::Prestarted);
+				DatabaseDependencies.Reset();
+				DatabaseDependencies.Reserve(KeyBuilder.GetDependencies().Num());
+				for (const UObject* Dependency : KeyBuilder.GetDependencies())
+				{
+					DatabaseDependencies.Add(Dependency);
+				}
 
-			UE_LOG(LogPoseSearch, Log, TEXT("%s - %s BeginCache"), *LexToString(DerivedDataKey), *Database->GetName());
+				SetState(EState::Prestarted);
+				
+				UE_LOG(LogPoseSearch, Log, TEXT("%s - %s BeginCache"), *LexToString(DerivedDataKey), *Database->GetName());
 
-			TArray<FCacheGetRequest> CacheRequests;
-			const FCacheKey CacheKey{ Bucket, DerivedDataKey };
-			CacheRequests.Add({ { Database->GetPathName() }, CacheKey, ECachePolicy::Default });
+				int32 GSkipDDC = 0;
+				const ECachePolicy Policy = GSkipDDC ? ECachePolicy::None : ECachePolicy::Default;
+			
+				TArray<FCacheGetRequest> CacheRequests;
+				const FCacheKey CacheKey{ Bucket, DerivedDataKey };
+				CacheRequests.Add({ { Database->GetPathName() }, CacheKey, Policy });
 
-			Owner = FRequestOwner(EPriority::Normal);
-			GetCache().Get(CacheRequests, Owner, [this](FCacheGetResponse&& Response)
-				{
-					OnGetComplete(MoveTemp(Response));
-				});
+				Owner = FRequestOwner(EPriority::Normal);
+				GetCache().Get(CacheRequests, Owner, [this](FCacheGetResponse&& Response)
+					{
+						OnGetComplete(MoveTemp(Response));
+					});
+			}
+			else
+			{
+				DerivedDataKey = FIoHash::Zero;
+				SetState(EState::Notstarted);
+				UE_LOG(LogPoseSearch, Log, TEXT("Delaying DDC until dependents are ready - %s"), *Database->GetName());
+			}
+			
+			bResidencyCleared = false;
 		}
 	}
 }
@@ -1541,6 +1583,7 @@ void FPoseSearchDatabaseAsyncCacheTask::Cancel()
 
 	DerivedDataKey = FIoHash::Zero;
 	SetState(EState::Cancelled);
+	ClearAnimSequenceResidency();
 }
 
 void FPoseSearchDatabaseAsyncCacheTask::Update(FCriticalSection& OuterMutex, FPartialKeyHashes& PartialKeyHashes)
@@ -1586,11 +1629,13 @@ void FPoseSearchDatabaseAsyncCacheTask::Wait(FCriticalSection& OuterMutex)
 		check(Database->Schema && !SearchIndex.IsEmpty() && SearchIndex.GetNumDimensions() == Database->Schema->SchemaCardinality);
 
 		SetState(EState::Ended);
+		ClearAnimSequenceResidency();
 		bBroadcastOnDerivedDataRebuild = true;
 	}
 	else
 	{
 		check(!bBroadcastOnDerivedDataRebuild);
+		ClearAnimSequenceResidency();
 		SetState(EState::Failed);
 	}
 	ResetSearchIndex();
diff --git a/Engine/Plugins/Animation/PoseSearch/Source/Runtime/Private/PoseSearchDerivedDataKey.cpp b/Engine/Plugins/Animation/PoseSearch/Source/Runtime/Private/PoseSearchDerivedDataKey.cpp
--- a/Engine/Plugins/Animation/PoseSearch/Source/Runtime/Private/PoseSearchDerivedDataKey.cpp
+++ b/Engine/Plugins/Animation/PoseSearch/Source/Runtime/Private/PoseSearchDerivedDataKey.cpp
@@ -15,6 +15,7 @@
 #include "Components/SkeletalMeshComponent.h"
 #include "PoseSearch/PoseSearchDefines.h"
 #include "UObject/DevObjectVersion.h"
+#include "Animation/IAnimationSequenceCompiler.h"
 
 namespace UE::PoseSearch
 {
@@ -45,6 +46,8 @@ FKeyBuilder::FKeyBuilder(const UObject* Object, bool bUseDataVer, bool bPerformC
 {
 	check(Object);
 
+	KeyOwner = Object;
+
 	// preallocating a reasonable amount of memory to avoid multiple reallocations
 	ObjectsToSerialize.Reserve(256);
 	ObjectBeingSerializedDependencies.Reserve(256);
@@ -58,19 +61,19 @@ FKeyBuilder::FKeyBuilder(const UObject* Object, bool bUseDataVer, bool bPerformC
 	UObject* NonConstObject = const_cast<UObject*>(Object);
 	*this << NonConstObject;
 
-	while (!ObjectsToSerialize.IsEmpty() && !bAnyAssetNotReady)
+	while (!ObjectsToSerialize.IsEmpty() && !bAnyAssetNotFullyLoaded)
 	{
 		SerializeObjectInternal(ObjectsToSerialize.Pop(EAllowShrinking::No));
 	}
 
-	if (bUseDataVer && !bAnyAssetNotReady)
+	if (bUseDataVer && !bAnyAssetNotFullyLoaded)
 	{
 		FLocalPartialKeyHash& LocalCachedHash = LocalPartialKeyHashes.AddDefaulted_GetRef();
 
 		Hasher.Reset();
 
 		// used to invalidate the key without having to change POSESEARCHDB_DERIVEDDATA_VER all the times
-		int32 POSESEARCHDB_DERIVEDDATA_VER_SMALL = 2659;
+		int32 POSESEARCHDB_DERIVEDDATA_VER_SMALL = 2660;
 		FGuid VersionGuid = FDevSystemGuids::GetSystemGuid(FDevSystemGuids::Get().POSESEARCHDB_DERIVEDDATA_VER);
 		FString AnimationCompressionVersionString = UE::Anim::Compression::AnimationCompressionVersionString;
 
@@ -175,16 +178,19 @@ FArchive& FKeyBuilder::TryAddDependency(UObject* Object, bool bAddToPartialKeyHa
 		}
 		else
 		{
-			bAnyAssetNotReady = true;
+			bAnyAssetNotFullyLoaded = true;
 			return *this;
 		}
 	}
 
-	// @todo: should we WaitOnExistingCompression?
-	//if (UAnimSequence* AnimSequence = Cast<UAnimSequence>(Object))
-	//{
-	//	AnimSequence->WaitOnExistingCompression();
-	//}
+	if (UAnimSequence* AnimSequence = Cast<UAnimSequence>(Object))
+	{
+		if (!AnimSequence->CanBeCompressed())
+		{
+			bAnyAssetNotFullyLoaded = true;
+			return *this;
+		}
+	}
 
 	// collecting ALL the dependencies of the object being serialized, so we can then cache it in PartialKeyHashes
 	ObjectBeingSerializedDependencies.Add(Object);
@@ -235,7 +241,7 @@ void FKeyBuilder::SerializeObjectInternal(UObject* Object)
 {
 	Hasher.Reset();
 
-	check(!bAnyAssetNotReady && Object);
+	check(!bAnyAssetNotFullyLoaded && Object);
 
 	// test to validate the PartialKeyHashes
 	bool bIsValidTestPartialKeyHash = false;
@@ -304,7 +310,7 @@ void FKeyBuilder::SerializeObjectInternal(UObject* Object)
 	check(ObjectBeingSerialized == Object);
 	ObjectBeingSerialized = nullptr;
 
-	if (!bAnyAssetNotReady)
+	if (!bAnyAssetNotFullyLoaded)
 	{
 		LocalCachedHash.Object = Object;
 		LocalCachedHash.Hash = Hasher.Finalize();
@@ -326,15 +332,46 @@ FString FKeyBuilder::GetArchiveName() const
 {
 	return TEXT("FDerivedDataKeyBuilder");
 }
-	
+
+bool FKeyBuilder::AnyAssetNotFullyLoaded() const
+{
+	return bAnyAssetNotFullyLoaded;
+}
+
 bool FKeyBuilder::AnyAssetNotReady() const
 {
+	bool bAnyAssetNotReady = false;
+	TArray<UAnimSequence*, TInlineAllocator<64>> SequencesToWaitFor;	
+	const ITargetPlatform* TargetPlatform = GetTargetPlatformManager()->GetRunningTargetPlatform();
+	for (const UObject* Dependency : Dependencies)
+	{
+		if (UAnimSequence* AnimSequence = Cast<UAnimSequence>(const_cast<UObject*>(Dependency)))
+		{
+			constexpr bool bForceResidency = true;
+			if (bForceResidency)
+			{
+				AnimSequence->RequestResidency(TargetPlatform, GetTypeHash(KeyOwner));
+				SequencesToWaitFor.Add(AnimSequence);
+			}
+		}
+	}
+
+	Anim::IAnimSequenceCompilingManager::FinishCompilation(MakeArrayView(SequencesToWaitFor));
+
+	for (UAnimSequence* AnimSequence : SequencesToWaitFor)
+	{
+		if(!AnimSequence->HasCompressedDataForPlatform(TargetPlatform))
+		{						
+			bAnyAssetNotReady = true;
+		}
+	}
+	
 	return bAnyAssetNotReady;
 }
 
 FIoHash FKeyBuilder::Finalize() const
 {
-	check(!bAnyAssetNotReady); // otherwise key can be non deterministic
+	check(!bAnyAssetNotFullyLoaded); // otherwise key can be non deterministic
 
 	HashBuilderType FinalizeHasher;
 
@@ -376,7 +413,7 @@ bool FKeyBuilder::IsAddNameOnlyType(class UObject* Object)
 
 bool FKeyBuilder::ValidateAgainst(const FKeyBuilder& Other) const
 {
-	if (bAnyAssetNotReady != Other.bAnyAssetNotReady)
+	if (bAnyAssetNotFullyLoaded != Other.bAnyAssetNotFullyLoaded)
 	{
 		return false;
 	}
diff --git a/Engine/Plugins/Animation/PoseSearch/Source/Runtime/Public/PoseSearch/PoseSearchAssetSampler.h b/Engine/Plugins/Animation/PoseSearch/Source/Runtime/Public/PoseSearch/PoseSearchAssetSampler.h
--- a/Engine/Plugins/Animation/PoseSearch/Source/Runtime/Public/PoseSearch/PoseSearchAssetSampler.h
+++ b/Engine/Plugins/Animation/PoseSearch/Source/Runtime/Public/PoseSearch/PoseSearchAssetSampler.h
@@ -3,6 +3,9 @@
 #pragma once
 
 #include "BonePose.h"
+#if WITH_EDITOR
+#include "IO/IoHash.h"
+#endif // WITH_EDITOR
 
 struct FAnimationPoseData;
 struct FAnimExtractContext;
@@ -68,6 +71,10 @@ protected:
 
 	const float ExtrapolationSampleTime = 1.f / 30.f;
 	const float ExtractionInterval = 1.0f / 120.0f;
+
+#if WITH_EDITOR
+	FIoHash PlatformHash;
+#endif // WITH_EDITOR
 };
 
 } // namespace UE::PoseSearch
diff --git a/Engine/Plugins/Animation/PoseSearch/Source/Runtime/Public/PoseSearch/PoseSearchDerivedDataKey.h b/Engine/Plugins/Animation/PoseSearch/Source/Runtime/Public/PoseSearch/PoseSearchDerivedDataKey.h
--- a/Engine/Plugins/Animation/PoseSearch/Source/Runtime/Public/PoseSearch/PoseSearchDerivedDataKey.h
+++ b/Engine/Plugins/Animation/PoseSearch/Source/Runtime/Public/PoseSearch/PoseSearchDerivedDataKey.h
@@ -54,6 +54,7 @@ public:
 	virtual FString GetArchiveName() const override;
 	// End FArchive Interface
 	
+	bool AnyAssetNotFullyLoaded() const;
 	bool AnyAssetNotReady() const;
 	FIoHash Finalize() const;
 	const TSet<const UObject*>& GetDependencies() const;
@@ -67,14 +68,17 @@ protected:
 
 	HashBuilderType Hasher;
 
+	// UPoseSearchDatabase instance "owner" of the key generation
+	const UObject* KeyOwner = nullptr;
+
 	// Set of objects that have already been serialized
 	TSet<const UObject*> Dependencies;
 
 	// Object currently being serialized
 	UObject* ObjectBeingSerialized = nullptr;
 
-	// true if some dependent assets are not ready (fully loaded)
-	bool bAnyAssetNotReady = false;
+	// true if some dependent assets are not fully loaded
+	bool bAnyAssetNotFullyLoaded = false;
 
 	// if true ConditionalPostLoad will be performed on the dependant assets requiring it
 	bool bPerformConditionalPostLoad = false;
diff --git a/Engine/Plugins/Experimental/AnimNext/Source/AnimNext/Private/DecompressionTools.cpp b/Engine/Plugins/Experimental/AnimNext/Source/AnimNext/Private/DecompressionTools.cpp
--- a/Engine/Plugins/Experimental/AnimNext/Source/AnimNext/Private/DecompressionTools.cpp
+++ b/Engine/Plugins/Experimental/AnimNext/Source/AnimNext/Private/DecompressionTools.cpp
@@ -63,7 +63,7 @@ bool FDecompressionTools::UseRawDataForPoseExtraction(const UAnimSequence* AnimS
 {
 	return CanEvaluateRawAnimationData(AnimSequence) && (
 #if WITH_EDITOR
-		AnimSequence->bUseRawDataOnly ||
+		!AnimSequence->IsCompressedDataValid() ||
 #endif
 		(AnimSequence->GetSkeletonVirtualBoneGuid() != AnimSequence->GetSkeleton()->GetVirtualBoneGuid()) || AnimationPoseData.GetDisableRetargeting() || AnimationPoseData.ShouldUseRawData() ||
 #if WITH_EDITOR
@@ -178,11 +178,15 @@ void FDecompressionTools::GetBonePose(const UAnimSequence* AnimSequence, const F
 		}
 	}
 
+	UAnimSequence::FScopedCompressedAnimSequence PlatformCompressedData = AnimSequence->GetCompressedData(ExtractionContext);
+	int32 NumTracks = 0;
+	{
 #if WITH_EDITOR
-	const int32 NumTracks = bForceUseRawData ? AnimSequence->GetDataModelInterface()->GetNumBoneTracks() : AnimSequence->CompressedData.CompressedTrackToSkeletonMapTable.Num();
+		NumTracks = bForceUseRawData ? AnimSequence->GetDataModelInterface()->GetNumBoneTracks() : PlatformCompressedData.Get().CompressedTrackToSkeletonMapTable.Num();
 #else
-	const int32 NumTracks = AnimSequence->CompressedData.CompressedTrackToSkeletonMapTable.Num();
-#endif 
+		NumTracks = PlatformCompressedData.Get().CompressedTrackToSkeletonMapTable.Num();
+#endif
+	}
 	const bool bTreatAnimAsAdditive = (AnimSequence->IsValidAdditive() && !bForceUseRawData); // Raw data is never additive
 	const FRootMotionReset RootMotionReset(AnimSequence->bEnableRootMotion, AnimSequence->RootMotionRootLock, AnimSequence->bForceRootLock, AnimSequence->ExtractRootTrackTransform(FAnimExtractContext(0.0), /*&RequiredBones*/nullptr), bTreatAnimAsAdditive);
 
@@ -217,14 +221,14 @@ void FDecompressionTools::GetBonePose(const UAnimSequence* AnimSequence, const F
 				, AnimSequence->GetSamplingFrameRate().AsFrameTime(AnimSequence->GetPlayLength()).RoundToFrame().Value
 				, AnimSequence->Interpolation
 				, AnimSequence->GetRetargetTransformsSourceName()
-				, *AnimSequence->CompressedData.CompressedDataStructure
+				, *PlatformCompressedData.Get().CompressedDataStructure
 				, AnimSequence->GetSkeleton()->GetRefLocalPoses()
-				, AnimSequence->CompressedData.CompressedTrackToSkeletonMapTable
+				, PlatformCompressedData.Get().CompressedTrackToSkeletonMapTable
 				, AnimSequence->GetSkeleton()
 				, AnimSequence->IsValidAdditive()
 				, AnimSequence->GetAdditiveAnimType());
 
-			DecompressPose(OutAnimationPoseData, AnimSequence->CompressedData, ExtractionContext, DecompContext, AnimSequence->GetRetargetTransforms(), RootMotionReset);
+			DecompressPose(OutAnimationPoseData, PlatformCompressedData.Get(), ExtractionContext, DecompContext, AnimSequence->GetRetargetTransforms(), RootMotionReset);
 		}
 	}
 }
diff --git a/Engine/Plugins/Interchange/Runtime/Source/Import/Private/Animation/InterchangeAnimSequenceFactory.cpp b/Engine/Plugins/Interchange/Runtime/Source/Import/Private/Animation/InterchangeAnimSequenceFactory.cpp
--- a/Engine/Plugins/Interchange/Runtime/Source/Import/Private/Animation/InterchangeAnimSequenceFactory.cpp
+++ b/Engine/Plugins/Interchange/Runtime/Source/Import/Private/Animation/InterchangeAnimSequenceFactory.cpp
@@ -1493,13 +1493,10 @@ void UInterchangeAnimSequenceFactory::BuildObject_GameThread(const FSetupObjectP
 		check(AnimSequence == CastChecked<UAnimSequence>(Arguments.ImportedObject));
 		if (AnimSequence)
 		{
-			const bool bNeedPostProcess = !AnimSequence->IsCompressedDataValid();
-
 			// @Todo fix me: This is temporary fix to make sure they always have compressed data
-			if (AnimSequence->IsDataModelValid() && bNeedPostProcess)
+			if (AnimSequence->IsDataModelValid() && AnimSequence->IsCompressedDataOutOfDate())
 			{
-				AnimSequence->ClearAllCachedCookedPlatformData();
-				AnimSequence->BeginCacheDerivedDataForCurrentPlatform();
+				AnimSequence->CacheDerivedDataForCurrentPlatform();
 			}
 		}
 	}
diff --git a/Engine/Source/Editor/AnimationModifiers/Private/AnimationModifier.cpp b/Engine/Source/Editor/AnimationModifiers/Private/AnimationModifier.cpp
--- a/Engine/Source/Editor/AnimationModifiers/Private/AnimationModifier.cpp
+++ b/Engine/Source/Editor/AnimationModifiers/Private/AnimationModifier.cpp
@@ -156,7 +156,7 @@ void UAnimationModifier::ApplyToAnimationSequence(class UAnimSequence* AnimSeque
 
 void UAnimationModifier::UpdateCompressedAnimationData()
 {
-	if (CurrentAnimSequence->DoesNeedRecompress())
+	if (CurrentAnimSequence->IsCompressedDataOutOfDate())
 	{
 		CurrentAnimSequence->CacheDerivedDataForCurrentPlatform();
 	}
diff --git a/Engine/Source/Editor/Persona/Private/AnimationEditorViewportClient.cpp b/Engine/Source/Editor/Persona/Private/AnimationEditorViewportClient.cpp
--- a/Engine/Source/Editor/Persona/Private/AnimationEditorViewportClient.cpp
+++ b/Engine/Source/Editor/Persona/Private/AnimationEditorViewportClient.cpp
@@ -1190,10 +1190,9 @@ FText FAnimationViewportClient::GetDisplayInfo(bool bDisplayAllInfo) const
 	if( PreviewInstance )
 	{
 		// see if you have anim sequence that has transform curves
-		UAnimSequence* Sequence = Cast<UAnimSequence>(PreviewInstance->GetCurrentAsset());
-		if (Sequence)
+		if (UAnimSequence* Sequence = Cast<UAnimSequence>(PreviewInstance->GetCurrentAsset()))
 		{
-			if (Sequence->DoesNeedRecompress())
+			if (Sequence->IsCompressedDataOutOfDate())
 			{
 				TextValue = ConcatenateLine(TextValue, LOCTEXT("ApplyToCompressedDataWarning", "<AnimViewport.WarningText>Animation is being edited. To apply to compressed data (and recalculate baked additives), click \"Apply\"</>"));
 			}
diff --git a/Engine/Source/Editor/UnrealEd/Private/Commandlets/PackageUtilities.cpp b/Engine/Source/Editor/UnrealEd/Private/Commandlets/PackageUtilities.cpp
--- a/Engine/Source/Editor/UnrealEd/Private/Commandlets/PackageUtilities.cpp
+++ b/Engine/Source/Editor/UnrealEd/Private/Commandlets/PackageUtilities.cpp
@@ -51,6 +51,7 @@
 #include "Animation/AnimationSettings.h"
 
 #include "EngineUtils.h"
+#include "Animation/IAnimationSequenceCompiler.h"
 #include "Materials/Material.h"
 #include "Serialization/ArchiveStackTrace.h"
 #include "Misc/OutputDeviceHelper.h"
@@ -1968,8 +1969,10 @@ struct CompressAnimationsFunctor
 				*PackageFileName);
 
 			UE_LOG(LogPackageUtilities, Warning, TEXT("%s (%s) Resetting with to default compression settings."), *AnimSeq->GetName(), *AnimSeq->GetFullName());
-			AnimSeq->BoneCompressionSettings = nullptr;
-			AnimSeq->CurveCompressionSettings = nullptr;
+
+			UE::Anim::IAnimSequenceCompilingManager::FinishCompilation({ AnimSeq });
+
+			AnimSeq->ResetCompressionSettings();
 			AnimSeq->CacheDerivedDataForCurrentPlatform();
 
 			// Automatic compression should have picked a suitable compressor
@@ -1978,7 +1981,7 @@ struct CompressAnimationsFunctor
 				// Update CompressCommandletVersion in that case, and create a proper DDC entry
 				// (with actual compressor)
 				AnimSeq->CompressCommandletVersion = CompressCommandletVersion;
-				AnimSeq->BeginCacheDerivedDataForCurrentPlatform();
+				AnimSeq->CacheDerivedDataForCurrentPlatform();
 				bDirtyPackage = true;
 			}
 
diff --git a/Engine/Source/Runtime/Core/Public/UObject/FortniteMainBranchObjectVersions.inl b/Engine/Source/Runtime/Core/Public/UObject/FortniteMainBranchObjectVersions.inl
--- a/Engine/Source/Runtime/Core/Public/UObject/FortniteMainBranchObjectVersions.inl
+++ b/Engine/Source/Runtime/Core/Public/UObject/FortniteMainBranchObjectVersions.inl
@@ -559,5 +559,8 @@ FORTNITE_MAIN_VERSION(DynamicMeshSerializeSculptLayers, 181)
 // Fix reachable garbage object warnings from some legacy ASpatialHashRuntimeGridInfo actors
 FORTNITE_MAIN_VERSION(SpatialHashRuntimeGridInfoSpriteFixup, 182)
 
+// Removed UAnimSequence::bUseRawDataOnly flag alongside compression refactor
+FORTNITE_MAIN_VERSION(AnimSequenceRawDataOnlyFlagRemoval, 183)
+
 // -----<new versions can be added above this line>-------------------------------------------------
 #undef FORTNITE_MAIN_VERSION
\ No newline at end of file
diff --git a/Engine/Source/Runtime/Engine/Classes/Animation/AnimSequence.h b/Engine/Source/Runtime/Engine/Classes/Animation/AnimSequence.h
--- a/Engine/Source/Runtime/Engine/Classes/Animation/AnimSequence.h
+++ b/Engine/Source/Runtime/Engine/Classes/Animation/AnimSequence.h
@@ -33,6 +33,7 @@ typedef TArray<FTransform> FTransformArrayA2;
 
 class USkeletalMesh;
 class FQueuedThreadPool;
+class ITargetPlatform;
 enum class EQueuedWorkPriority : uint8;
 struct FAnimCompressContext;
 struct FAnimSequenceDecompressionContext;
@@ -271,6 +272,8 @@ public:
 	UPROPERTY(Category = Compression, EditAnywhere, meta = (ForceShowEngineContent))
 	TObjectPtr<class UAnimCurveCompressionSettings> CurveCompressionSettings;
 
+	// CompressedData is only valid in cook/non-editor runtime, see DataByPlatformKeyHash for editor-runtime data
+	UE_DEPRECATED(5.6, "CompressedData public access will be removed")
 	FCompressedAnimSequence CompressedData;
 
 	UPROPERTY(Category = Compression, EditAnywhere, meta = (ForceShowEngineContent))
@@ -406,6 +409,7 @@ public:
 	virtual FFrameRate GetSamplingFrameRate() const override { return PlatformTargetFrameRate.Default; }
 	ENGINE_API virtual void EvaluateCurveData(FBlendedCurve& OutCurve, const FAnimExtractContext& AnimExtractContext, bool bForceUseRawData = false) const override;
 	ENGINE_API virtual float EvaluateCurveData(FName CurveName, const FAnimExtractContext& AnimExtractContext, bool bForceUseRawData = false) const override;
+	// This will only check the current platform its compressed data (if valid)
 	ENGINE_API virtual bool HasCurveData(FName CurveName, bool bForceUseRawData) const override;
 
 	//~ End UAnimSequenceBase Interface
@@ -433,11 +437,23 @@ public:
 	*
 	* @param	OutAnimationPoseData  [out] Animation Pose related data to populate
 	* @param	ExtractionContext	  Extraction Context (position, looping, root motion, etc.)
+	* @param	bForceUseRawData	  Whether or not to forcefully try to sample the animation data model	
 	*/
 	ENGINE_API void GetBonePose(struct FAnimationPoseData& OutAnimationPoseData, const FAnimExtractContext& ExtractionContext, bool bForceUseRawData = false) const;
 
-	const TArray<FTrackToSkeletonMap>& GetCompressedTrackToSkeletonMapTable() const { return CompressedData.CompressedTrackToSkeletonMapTable; }
-	const TArray<FAnimCompressedCurveIndexedName>& GetCompressedCurveIndexedNames() const { return CompressedData.IndexedCurveNames; }
+	UE_DEPRECATED(5.6, "GetCompressedTrackToSkeletonMapTable has been deprecated. Use GetCompressedData().CompressedTrackToSkeletonMapTable instead")
+	const TArray<FTrackToSkeletonMap>& GetCompressedTrackToSkeletonMapTable() const
+	{
+		static TArray<FTrackToSkeletonMap> TempArray;
+		return TempArray;
+	}
+	
+	UE_DEPRECATED(5.6, "GetCompressedCurveIndexedNames has been deprecated. Use GetCompressedData().IndexedCurveNames instead")
+	const TArray<FAnimCompressedCurveIndexedName>& GetCompressedCurveIndexedNames() const
+	{
+		static TArray<FAnimCompressedCurveIndexedName> TempArray;
+		return TempArray;
+	}
 
 #if WITH_EDITORONLY_DATA
 protected:
@@ -509,19 +525,20 @@ public:
 	 *
 	 * @param	OutAtom					[out] Output bone transform.
 	 * @param	BoneIndex				Index of bone to evaluate.
-	 * @param	FAnimExtractContext		Extraction context containing time and interpolation info.
+	 * @param	ExtractionContext		Extraction context containing time and interpolation info.
 	 * @param	bUseRawData		If true, use raw animation data instead of compressed data.
 	 */
 	ENGINE_API void GetBoneTransform(FTransform& OutAtom, FSkeletonPoseBoneIndex BoneIndex, const FAnimExtractContext& ExtractionContext, bool bUseRawData) const;
 
 	/**
-	 * Get Bone Transform of the Time given, relative to Parent for the Track Given
+	 * Get Bone Transform of the Time given, relative to Parent for the Track Given. This will sample the compressed data for the current platform only. (if available)
 	 *
 	 * @param	OutAtom			[out] Output bone transform.
 	 * @param	BoneIndex		Index of bone to evaluate.
 	 * @param	DecompContext	Decompression context to use.
 	 * @param	bUseRawData		If true, use raw animation data instead of compressed data.
 	 */
+	UE_DEPRECATED(5.6, "Please use GetBoneTransform with FAnimExtractContext")
 	ENGINE_API void GetBoneTransform(FTransform& OutAtom, FSkeletonPoseBoneIndex BoneIndex, FAnimSequenceDecompressionContext& DecompContext, bool bUseRawData) const;
 	// End Transform related functions 
 
@@ -557,37 +574,57 @@ public:
 	 * @return		The approximate size of compressed animation data.
 	 */
 	ENGINE_API int32 GetApproxCompressedSize() const;
+protected:
+	/**
+	 * @return		The approximate size of compressed animation data for only bones.
+	 */
+	ENGINE_API int32 GetApproxBoneCompressedSize_Lockless() const;
+	
+	/**
+	 * @return		The approximate size of compressed animation data.
+	 */
+	ENGINE_API int32 GetApproxCompressedSize_Lockless() const;
+
+	void EvaluateCurveData_Lockless(FBlendedCurve& OutCurve, const FAnimExtractContext& AnimExtractContext, bool bForceUseRawData = false) const;	
+	float EvaluateCurveData_Lockless(FName CurveName, const FAnimExtractContext& AnimExtractContext, bool bForceUseRawData = false) const;	
+	FTransform ExtractRootTrackTransform_Lockless(const FAnimExtractContext& ExtractionContext, const FBoneContainer* RequiredBones) const;
+
+	ENGINE_API void GetBoneTransform_Lockless(FTransform& OutAtom, FSkeletonPoseBoneIndex BoneIndex, const FAnimExtractContext& ExtractionContext, bool bForceUseRawData) const;	
 
-	// Get compressed data for this UAnimSequence. May be built directly or pulled from DDC
 #if WITH_EDITOR
-	ENGINE_API void WaitOnExistingCompression(const bool bWantResults=true);
-protected:
 	ENGINE_API bool ShouldPerformStripping(const bool bPerformFrameStripping, const bool bPerformStrippingOnOddFramedAnims) const;
-#endif
+#endif // WITH_EDITOR
+
 
 protected:
-    ENGINE_API void ClearCompressedBoneData();
-    ENGINE_API void ClearCompressedCurveData();
+	UE_DEPRECATED(5.6, "ClearCompressedBoneData will be removed")
+    ENGINE_API void ClearCompressedBoneData() {}
+	UE_DEPRECATED(5.6, "ClearCompressedCurveData will be removed")
+    ENGINE_API void ClearCompressedCurveData() {}
     // Write the compressed data to the supplied FArchive
     ENGINE_API void SerializeCompressedData(FArchive& Ar, bool bDDCData);
+	ENGINE_API void SerializeCompressedData(FArchive& Ar, FCompressedAnimSequence& CompressedData);
 #if WITH_EDITOR
 	ENGINE_API virtual void OnAnimModelLoaded() override;
 #endif
 public:
 	ENGINE_API bool IsCompressedDataValid() const;
+	ENGINE_API bool IsBoneCompressedDataValid() const;	
 	ENGINE_API bool IsCurveCompressedDataValid() const;
 	// End Memory related functions
-	
 #if WITH_EDITOR
 	/**
 	 * Add Key to Transform Curves
 	 */
 	ENGINE_API void AddKeyToSequence(float Time, const FName& BoneName, const FTransform& AdditiveTransform);
 
+	UE_DEPRECATED(5.6, "DoesNeedRecompress has been renamed to IsCompressedDataOutOfDate")
+	ENGINE_API bool DoesNeedRecompress() const { return IsCompressedDataOutOfDate(); }
+
 	/**
-	* Return true if compressed data is out of date / missing and so animation needs to use raw data
+	* Return true if compressed data is invalid or if it is not in sync with the skeleton
 	*/
-	ENGINE_API bool DoesNeedRecompress() const;
+	ENGINE_API bool IsCompressedDataOutOfDate() const;
 
 	/**
 	 * Create Animation Sequence from Reference Pose of the Mesh
@@ -699,11 +736,10 @@ public:
 
 	/** Take a set of marker positions and validates them against a requested start position, updating them as desired */
 	ENGINE_API void ValidateCurrentPosition(const FMarkerSyncAnimPosition& Position, bool bPlayingForwards, bool bLooping, float&CurrentTime, FMarkerPair& PreviousMarker, FMarkerPair& NextMarker, const UMirrorDataTable* MirrorTable = nullptr) const;
+
+	UE_DEPRECATED(5.6, "Public access to UseRawDataForPoseExtraction has been deprecated")
 	ENGINE_API bool UseRawDataForPoseExtraction(const FBoneContainer& RequiredBones) const;
 
-protected:
-	// Should we be always using our raw data (i.e is our compressed data stale)
-	TAtomic<bool> bUseRawDataOnly;
 public:
 	/** Authored Sync markers */
 	UPROPERTY()
@@ -718,10 +754,7 @@ protected:
 #if WITH_EDITOR
 	ENGINE_API void SynchronousAnimatedBoneAttributesCompression();
 	ENGINE_API void MoveAttributesToModel();
-#endif // WITH_EDITOR
 
-protected:
-#if WITH_EDITOR
 	// Begin UAnimSequenceBase virtual overrides
 	ENGINE_API virtual void OnModelModified(const EAnimDataModelNotifyType& NotifyType, IAnimationDataModel* Model, const FAnimDataModelNotifPayload& Payload) override;
 	ENGINE_API virtual void PopulateModel() override;
@@ -733,6 +766,8 @@ protected:
 	ENGINE_API void DeleteDeprecatedRawAnimationData();
 public:
 	ENGINE_API void DeleteNotifyTrackData();
+	// Resets the bone and curve compression settings to the project default 
+	ENGINE_API void ResetCompressionSettings();
 #endif // WITH_EDITOR
 
 protected:
@@ -763,35 +798,156 @@ protected:
 	TMap<FAnimationAttributeIdentifier, FAttributeCurve> AttributeCurves;
 
 #if WITH_EDITOR
-	ENGINE_API FIoHash CreateDerivedDataKeyHash(const ITargetPlatform* TargetPlatform);
-private:
-	FString CreateDerivedDataKeyString(const ITargetPlatform* TargetPlatform);
-protected:
+	ENGINE_API FIoHash CreateDerivedDataKeyHash(const ITargetPlatform* TargetPlatform) const;
+	FString CreateDerivedDataKeyString(const ITargetPlatform* TargetPlatform) const;
+
 	ENGINE_API FIoHash BeginCacheDerivedData(const ITargetPlatform* TargetPlatform);
 	ENGINE_API bool PollCacheDerivedData(const FIoHash& KeyHash) const;
 	ENGINE_API void EndCacheDerivedData(const FIoHash& KeyHash);
 
-	FIoHash DataKeyHash;
 	TMap<FIoHash, TUniquePtr<FCompressedAnimSequence>> DataByPlatformKeyHash;
 	TMap<FIoHash, TPimplPtr<UE::Anim::FAnimationSequenceAsyncCacheTask>> CacheTasksByKeyHash;
+	mutable FRWLock CompressedDataLock;
+
+	// Cached FIoHash keys stored according to ITargetPlatform type-hash
+	mutable TMap<uint32, FIoHash> PlatformHashToKeyHash;
+	mutable FRWLock HashCacheLock;
+	// Whether or not compressed data should be cleared in case of residency releasing (cook-time behaviour only)
+	TAtomic<bool> bShouldClearCompressedData = false;
+	
+	// Multi-map keeping track of required (target-platform specific) compressed data, stored by hash and provided identifier from API usage	
+	TMultiMap<FIoHash, uint32> ResidencyReferencerHashes;
+	mutable FRWLock ResidencyLock;
 
-protected:
 	ENGINE_API bool TryCancelAsyncTasks();
 	ENGINE_API bool WaitForAsyncTasks(float TimeLimitSeconds);
 	ENGINE_API void FinishAsyncTasks();
 	ENGINE_API void Reschedule(FQueuedThreadPool* InThreadPool, EQueuedWorkPriority InPriority);
 	ENGINE_API bool IsAsyncTaskComplete() const;
 	ENGINE_API bool IsCompiling() const;
+	ENGINE_API bool RequiresResidency(const FIoHash& InKeyHash) const;
+
+	void ValidateCompressionSettings();
+
+	const FCompressedAnimSequence& GetPlatformCompressedData(const ITargetPlatform* InTargetPlatform) const;
+#endif // WITH_EDITOR
+
+	// Whether or not RawData should be sampled during pose extraction, either because compressed data is not available or through user directive
+	ENGINE_API bool ShouldUseRawDataForPoseExtraction(const FBoneContainer& RequiredBones, const FAnimExtractContext& ExtractionContext) const;
+	ENGINE_API bool ShouldUseRawDataForPoseExtraction_Lockless(const FBoneContainer& RequiredBones, const FAnimExtractContext& ExtractionContext) const;
+	ENGINE_API bool ShouldUseRawDataForPoseExtraction_Lockless(const FAnimExtractContext& ExtractionContext) const;
+
+	// Clears all currently cached compressed data (all platform data)
+	ENGINE_API void ClearAllCompressionData();
+	// Clears compressed data, if currently cached, for specified hash
+	ENGINE_API void ClearCompressionData(const FIoHash& InKeyHash);
+	// Whether or not compressed data, for specified hash, is currently cached
+	ENGINE_API bool HasCompressedDataForHash(const FIoHash& InKeyHash) const;
+
+private:	
+	// Returns compressed animation data for specified Hash or TargetPlatform (not thread-safe see FCompressedAnimationDataReadScope and FCompressedAnimationResidencyScope)
+	ENGINE_API const FCompressedAnimSequence& GetPlatformCompressedData(const FAnimExtractContext& AnimExtractContext) const;
+	ENGINE_API FCompressedAnimSequence& GetPlatformCompressedData(const FAnimExtractContext& AnimExtractContext);
+	ENGINE_API FCompressedAnimSequence& GetPlatformCompressedData(const ITargetPlatform* InTargetPlatform);
+	
+	// Returns compressed animation data for the current platform (not thread-safe in editor-builds see FCompressedAnimationDataReadScope and FCompressedAnimationResidencyScope)
+	ENGINE_API const FCompressedAnimSequence& GetCompressedData_Internal() const;
+	// Returns compressed animation data for the current platform (not thread-safe in editor-builds see FCompressedAnimationDataReadScope and FCompressedAnimationResidencyScope)
+	ENGINE_API FCompressedAnimSequence& GetCompressedData_Internal();
+
 public:	
 	ENGINE_API void BeginCacheDerivedDataForCurrentPlatform();
 	ENGINE_API void CacheDerivedDataForCurrentPlatform();
+#if WITH_EDITOR
+	ENGINE_API void WaitOnExistingCompression(const bool bWantResults=true);
+	
+	// Returns hash identifying compressed animation data with the specified target platform its settings (these can overlap between different platforms due to the settings matching)
+	ENGINE_API FIoHash GetDerivedDataKeyHash(const ITargetPlatform* TargetPlatform) const;
+	
+	// Whether or not compressed data, for specified platform, is currently cached
+	ENGINE_API bool HasCompressedDataForPlatform(const ITargetPlatform* InPlatform) const;	
+	// Synchronous caching of compressed animation data for provided target platform
+	ENGINE_API void CacheDerivedDataForPlatform(const ITargetPlatform* TargetPlatform);
+
+	// Whether or not animation sequence can be compressed (dependent on loading state)
+	ENGINE_API bool CanBeCompressed() const;
+	
+	// Whether or not compressed animation data has been requested to stay resident until released 
+	ENGINE_API bool HasResidency(const ITargetPlatform* InPlatform) const;
+	// Fetches and keeps resident the compressed animation data for the provided target platform (until released)
+	ENGINE_API void RequestResidency(const ITargetPlatform* InPlatform, uint32 InReferencerHash);
+	// Releases previously requested residency of compressed animation data
+	ENGINE_API void ReleaseResidency(const ITargetPlatform* InPlatform, uint32 InReferencerHash);
 	
 	// Synchronous caching of compressed animation data for provided target platform
+	UE_DEPRECATED(5.6, "CacheDerivedData returning FCompressedAnimSequence reference has been deprecated, use CacheDerivedDataForPlatform")
 	ENGINE_API FCompressedAnimSequence& CacheDerivedData(const ITargetPlatform* TargetPlatform);
 
 	ENGINE_API FFrameRate GetTargetSamplingFrameRate(const ITargetPlatform* InPlatform) const;
 #endif
 
+public:
+	// Provides access to an instance of compressed animation data within the lifetime of FScopedCompressedAnimSequence, should _only_ be used on the stack due to risk of deadlocking
+	struct ENGINE_API FScopedCompressedAnimSequence
+	{
+		FScopedCompressedAnimSequence() = delete;
+		FScopedCompressedAnimSequence(const FScopedCompressedAnimSequence&) = delete;
+		FScopedCompressedAnimSequence& operator=(const FScopedCompressedAnimSequence&) = delete;		
+		FScopedCompressedAnimSequence(FScopedCompressedAnimSequence&&) = delete;
+		FScopedCompressedAnimSequence& operator=(FScopedCompressedAnimSequence&&) = delete;
+		
+		const FCompressedAnimSequence& operator->() const { return *CompressedData; }		
+		const FCompressedAnimSequence& Get() const { return *CompressedData; }
+#if WITH_EDITOR
+		~FScopedCompressedAnimSequence()
+		{
+			ReadScopeLock.Reset();
+		}
+#endif // WITH_EDITOR
+		
+	protected:
+		FScopedCompressedAnimSequence(const UAnimSequence* InAnimSequence, const FCompressedAnimSequence* InCompressedData) :
+			CompressedData(InCompressedData)
+#if WITH_EDITOR
+			, ReadScopeLock(MakeUnique<FReadScopeLock>(InAnimSequence->CompressedDataLock))
+#endif // WITH_EDITOR
+		{
+			check(CompressedData && InAnimSequence);
+		}
+
+		const FCompressedAnimSequence* CompressedData = nullptr;
+#if WITH_EDITOR
+		TUniquePtr<FReadScopeLock> ReadScopeLock;
+#endif // WITH_EDITOR
+		friend class UAnimSequence;
+	};
+	
+	// Returns scoped version of the compressed AnimSequence, only valid to read when FScopedConstCompressedAnimSequence is in scope (this internally locks CompressedDataLock)
+	ENGINE_API FScopedCompressedAnimSequence GetCompressedData() const;	
+	ENGINE_API FScopedCompressedAnimSequence GetCompressedData(const FAnimExtractContext& AnimExtractContext) const;
+	
+#if WITH_EDITOR	
+	ENGINE_API FScopedCompressedAnimSequence GetCompressedData(const ITargetPlatform* InTargetPlatform) const;	
+protected:
+	/** Helper struct to enter/leave compressed data lock in read-mode */
+	struct FCompressedAnimationDataReadScope : public FReadScopeLock
+	{
+		FCompressedAnimationDataReadScope(const UAnimSequence* InAnimSequence) : FReadScopeLock(InAnimSequence->CompressedDataLock)
+		{
+			check(InAnimSequence);
+		}		
+	};
+	
+	/** Helper struct to enter/leave compressed data lock in write-mode */
+	struct FCompressedAnimationDataWriteScope : public FWriteScopeLock
+	{
+		FCompressedAnimationDataWriteScope(const UAnimSequence* InAnimSequence) : FWriteScopeLock(InAnimSequence->CompressedDataLock)
+		{
+			check(InAnimSequence);
+		}
+	};
+#endif // WITH_EDITOR
+
 public:
 	friend class UAnimationAsset;
 	friend struct FScopedAnimSequenceRawDataCache;
@@ -807,4 +963,24 @@ public:
 	friend class FAnimSequenceDetails;
 	friend class UAnimBank;
 	friend class FAnimBankBuildAsyncCacheTask;
+	friend struct FAnimCurveBufferAccess;
+	friend class UE::Anim::FAnimationSequenceAsyncCacheTask;
 };
+
+#ifndef UE_COMPRESSED_DATA_WRITE_SCOPE
+#if WITH_EDITOR
+	#define UE_COMPRESSED_DATA_WRITE_SCOPE(sequence) \
+	UAnimSequence::FCompressedAnimationDataWriteScope Scope(sequence);
+#else
+	#define UE_COMPRESSED_DATA_WRITE_SCOPE(sequence)
+#endif
+#endif
+
+#ifndef UE_COMPRESSED_DATA_READ_SCOPE
+#if WITH_EDITOR
+	#define UE_COMPRESSED_DATA_READ_SCOPE(sequence) \
+	UAnimSequence::FCompressedAnimationDataReadScope Scope(sequence);
+#else
+	#define UE_COMPRESSED_DATA_READ_SCOPE(sequence)
+#endif
+#endif
\ No newline at end of file
diff --git a/Engine/Source/Runtime/Engine/Classes/Animation/AnimationAsset.h b/Engine/Source/Runtime/Engine/Classes/Animation/AnimationAsset.h
--- a/Engine/Source/Runtime/Engine/Classes/Animation/AnimationAsset.h
+++ b/Engine/Source/Runtime/Engine/Classes/Animation/AnimationAsset.h
@@ -17,6 +17,9 @@
 #include "AnimInterpFilter.h"
 #include "AnimEnums.h"
 #include "Interfaces/Interface_PreviewMeshProvider.h"
+#if WITH_EDITOR
+#include "IO/IoHash.h"
+#endif //  WITH_EDITOR
 #include "AnimationAsset.generated.h"
 
 class UAnimMetaData;
@@ -29,6 +32,7 @@ class UBlendSpace;
 class UPoseAsset;
 class UMirrorDataTable;
 class USkeletalMesh;
+class ITargetPlatform;
 struct FAnimationUpdateContext;
 namespace SmartName
 {
@@ -298,6 +302,13 @@ struct FAnimExtractContext
 
 	// Experimental
 	bool bExtractWithRootMotionProvider;
+
+	// Flag when set will enforce evaluation/sampling of, specific, compressed animation data, will assert if it is missing (user is responsible for ensuring the data was previously requested/resident)
+	bool bEnforceCompressedDataSampling = false;
+	// Used in combination with bForceSampleCompressedData, forcing a specific target platform hash with corresponding compressed data to sample (retrieved using UAnimSequence::GetDerivedDataKeyHash) 
+	FIoHash TargetPlatformHash = FIoHash::Zero;
+	// Used in combination with bForceSampleCompressedData, forcing a specific target platform its compressed data to sample
+	const ITargetPlatform* TargetPlatform = nullptr;
 #endif 
 	
 	UE_DEPRECATED(5.1, "FAnimExtractContext construct with float-based time value is deprecated, use other signature")
diff --git a/Engine/Source/Runtime/Engine/Private/Animation/AnimBank.cpp b/Engine/Source/Runtime/Engine/Private/Animation/AnimBank.cpp
--- a/Engine/Source/Runtime/Engine/Private/Animation/AnimBank.cpp
+++ b/Engine/Source/Runtime/Engine/Private/Animation/AnimBank.cpp
@@ -523,34 +523,22 @@ bool FAnimBankBuildAsyncCacheTask::BuildData(const UE::FSharedString& Name, cons
 
 		Sequence->FinishAsyncTasks();
 
-		const FCompressedAnimSequence* CompressedData = nullptr;
-
-		const FIoHash SequenceKeyHash = Sequence->CreateDerivedDataKeyHash(TargetPlatform);
-		if (SequenceKeyHash == Sequence->DataKeyHash)
-		{
-			// Key hash is from running platform, use compressed data on the sequence object
-			CompressedData = &Sequence->CompressedData;
-		}
-		else
-		{
-			// Fetch compressed data for the appropriate target platform
-			CompressedData = Sequence->DataByPlatformKeyHash[SequenceKeyHash].Get();
-		}
-
+		UAnimSequence::FScopedCompressedAnimSequence CompressedAnimSequence = Sequence->GetCompressedData(TargetPlatform);
+		const FCompressedAnimSequence& PlatformCompressedData = CompressedAnimSequence.Get();
 		// We should always have compressed data at this point
-		if (CompressedData == nullptr || CompressedData->CompressedDataStructure == nullptr || CompressedData->CompressedTrackToSkeletonMapTable.Num() == 0)
+		if (PlatformCompressedData.IsBoneDataValid(Sequence))
 		{
 			UE_LOG(LogAnimBank, Error, TEXT("Animation bank referenced sequence is missing compressed data!"));
 			return false;
 		}
 
 		// Set up mapping tables for the decompressor to map internal tracks to the pose array (which is in bone order).
-		const int32 NumTracks = CompressedData->CompressedTrackToSkeletonMapTable.Num();
+		const int32 NumTracks = PlatformCompressedData.CompressedTrackToSkeletonMapTable.Num();
 
 		TrackToBoneIndexMap.Reset(NumTracks);
 		for (int32 TrackIndex = 0; TrackIndex < NumTracks; ++TrackIndex)
 		{
-			const int32 BoneIndex = CompressedData->GetSkeletonIndexFromTrackIndex(TrackIndex);
+			const int32 BoneIndex = PlatformCompressedData.GetSkeletonIndexFromTrackIndex(TrackIndex);
 
 			// We only care about raw bones.
 			if (BoneIndex < NumAnimBones)
@@ -565,9 +553,9 @@ bool FAnimBankBuildAsyncCacheTask::BuildData(const UE::FSharedString& Name, cons
 			Sequence->GetSamplingFrameRate().AsFrameTime(Sequence->GetPlayLength()).RoundToFrame().Value,
 			EAnimInterpolationType::Linear,
 			Sequence->GetRetargetTransformsSourceName(),
-			*CompressedData->CompressedDataStructure,
+			*PlatformCompressedData.CompressedDataStructure,
 			AnimRefSkeleton.GetRefBonePose(),
-			CompressedData->CompressedTrackToSkeletonMapTable,
+			PlatformCompressedData.CompressedTrackToSkeletonMapTable,
 			nullptr,
 			false,
 			AAT_None
@@ -613,7 +601,7 @@ bool FAnimBankBuildAsyncCacheTask::BuildData(const UE::FSharedString& Name, cons
 			DecompressionContext.Seek(SeekTime);
 
 			TArrayView<FTransform> SampledLocalPoseView(SampledLocalAnimPose);
-			CompressedData->BoneCompressionCodec->DecompressPose(DecompressionContext, TrackToBoneIndexMap, TrackToBoneIndexMap, TrackToBoneIndexMap, SampledLocalPoseView);
+			PlatformCompressedData.BoneCompressionCodec->DecompressPose(DecompressionContext, TrackToBoneIndexMap, TrackToBoneIndexMap, TrackToBoneIndexMap, SampledLocalPoseView);
 
 			// Retarget from the anim skeleton to the mesh skeleton.
 			for (int32 MeshBoneIndex = 0; MeshBoneIndex < NumMeshBones; ++MeshBoneIndex)
@@ -1049,7 +1037,7 @@ FIoHash UAnimBank::CreateDerivedDataKeyHash(const ITargetPlatform* TargetPlatfor
 			// Make sure all our required dependencies are loaded, we need them to compute the KeyHash
 			FAnimationUtils::EnsureAnimSequenceLoaded(*BankSequence.Sequence);
 
-			FIoHash SequenceHash = BankSequence.Sequence->CreateDerivedDataKeyHash(TargetPlatform);
+			FIoHash SequenceHash = BankSequence.Sequence->GetDerivedDataKeyHash(TargetPlatform);
 			Writer << SequenceHash;
 		}
 
diff --git a/Engine/Source/Runtime/Engine/Private/Animation/AnimCurveCompressionCodec_UniformIndexable.cpp b/Engine/Source/Runtime/Engine/Private/Animation/AnimCurveCompressionCodec_UniformIndexable.cpp
--- a/Engine/Source/Runtime/Engine/Private/Animation/AnimCurveCompressionCodec_UniformIndexable.cpp
+++ b/Engine/Source/Runtime/Engine/Private/Animation/AnimCurveCompressionCodec_UniformIndexable.cpp
@@ -31,10 +31,10 @@ FAnimCurveBufferAccess::FAnimCurveBufferAccess(const UAnimSequenceBase* InSequen
 		{
 			if (InSequence->IsCurveCompressedDataValid())
 			{
-				const FCompressedAnimSequence& CompressedAnimSequence = InSequence->CompressedData;
-				if (UAnimCurveCompressionCodec_UniformIndexable* IndexableCompressedCurve = Cast<UAnimCurveCompressionCodec_UniformIndexable>(CompressedAnimSequence.CurveCompressionCodec))
+				UAnimSequence::FScopedCompressedAnimSequence CompressedAnimSequence = InSequence->GetCompressedData();
+				if (UAnimCurveCompressionCodec_UniformIndexable* IndexableCompressedCurve = Cast<UAnimCurveCompressionCodec_UniformIndexable>(CompressedAnimSequence.Get().CurveCompressionCodec))
 				{
-					bUseCompressedData = IndexableCompressedCurve->GetCurveBufferAndSamples(CompressedAnimSequence, InName, CompressedBuffer, NumSamples, SampleRate);
+					bUseCompressedData = IndexableCompressedCurve->GetCurveBufferAndSamples(CompressedAnimSequence.Get(), InName, CompressedBuffer, NumSamples, SampleRate);
 				}
 				else
 				{
diff --git a/Engine/Source/Runtime/Engine/Private/Animation/AnimSequence.cpp b/Engine/Source/Runtime/Engine/Private/Animation/AnimSequence.cpp
--- a/Engine/Source/Runtime/Engine/Private/Animation/AnimSequence.cpp
+++ b/Engine/Source/Runtime/Engine/Private/Animation/AnimSequence.cpp
@@ -49,6 +49,7 @@
 #include "Animation/SkeletonRemappingRegistry.h"
 #include "Animation/SkeletonRemapping.h"
 #include "Animation/Skeleton.h"
+#include "UObject/FortniteMainBranchObjectVersion.h"
 
 LLM_DEFINE_TAG(SequenceData);
 
@@ -385,7 +386,6 @@ UAnimSequence::UAnimSequence(const FObjectInitializer& ObjectInitializer)
 	, RootMotionRootLock(ERootMotionRootLock::RefPose)
 	, bUseNormalizedRootMotionScale(true)
 	, bRootMotionSettingsCopiedFromMontage(false)
-	, bUseRawDataOnly(!FPlatformProperties::RequiresCookedData())
 #if WITH_EDITOR
 	, bBlockCompressionRequests(false)
 #endif
@@ -469,8 +469,10 @@ void UAnimSequence::AddReferencedObjects(UObject* This, FReferenceCollector& Col
 	Super::AddReferencedObjects(This, Collector);
 
 	UAnimSequence* AnimSeq = CastChecked<UAnimSequence>(This);
-	Collector.AddReferencedObject(AnimSeq->CompressedData.BoneCompressionCodec);
-	Collector.AddReferencedObject(AnimSeq->CompressedData.CurveCompressionCodec);
+	UE_COMPRESSED_DATA_READ_SCOPE(AnimSeq);
+	FCompressedAnimSequence& ThisCompressedData = AnimSeq->GetCompressedData_Internal();
+	Collector.AddReferencedObject(ThisCompressedData.BoneCompressionCodec);
+	Collector.AddReferencedObject(ThisCompressedData.CurveCompressionCodec);
 }
 
 #if WITH_EDITOR
@@ -485,31 +487,19 @@ static FAutoConsoleVariableRef CVarAllowClearingCompressedDataDuringCook(
 
 void UAnimSequence::WillNeverCacheCookedPlatformDataAgain()
 {
-	if (GAllowClearingCompressedDataDuringCook)
+	UE::Anim::FAnimSequenceCompilingManager::Get().FinishCompilation({this});
+	bShouldClearCompressedData = GAllowClearingCompressedDataDuringCook == 1;
+	if (bShouldClearCompressedData || !IsRunningCookCommandlet())
 	{
-		bUseRawDataOnly = true;
-
-		UE::Anim::FAnimSequenceCompilingManager::Get().FinishCompilation({this});
-		// Clear out current platform, and any target platform data
-		CompressedData.Reset();
-		
-		CacheTasksByKeyHash.Empty();
-		DataByPlatformKeyHash.Empty();
+		ClearAllCompressionData();
 	}
 }
 
 void UAnimSequence::ClearAllCachedCookedPlatformData()
 {
-	if (GAllowClearingCompressedDataDuringCook)
+	if (GAllowClearingCompressedDataDuringCook || !IsRunningCookCommandlet())
 	{
-		bUseRawDataOnly = true;
-		
-		// Delete any cache tasks first because the destructor will cancel the cache and build tasks,
-		// and drop their pointers to the data.
-		CacheTasksByKeyHash.Empty();
-		DataByPlatformKeyHash.Empty();
-		CompressedData.Reset();
-		DataKeyHash = FIoHash::Zero;
+		ClearAllCompressionData();
 	}
 }
 
@@ -575,16 +565,34 @@ int64 UAnimSequence::GetApproxCurveRawSize() const
 
 int32 UAnimSequence::GetApproxBoneCompressedSize() const
 {
-	return CompressedData.CompressedDataStructure != nullptr ? CompressedData.CompressedDataStructure->GetApproxCompressedSize() : 0;
+	UE_COMPRESSED_DATA_READ_SCOPE(this);
+	return GetApproxBoneCompressedSize_Lockless();
 }
 
 int32 UAnimSequence::GetApproxCompressedSize() const
 {
-	int32 BoneTotal = GetApproxBoneCompressedSize();
-	int32 CurveTotal = CompressedData.CompressedCurveByteStream.Num();
+	UE_COMPRESSED_DATA_READ_SCOPE(this);
+	return GetApproxCompressedSize_Lockless();
+}
+
+int32 UAnimSequence::GetApproxBoneCompressedSize_Lockless() const
+{
+	return GetCompressedData_Internal().CompressedDataStructure != nullptr ? GetCompressedData_Internal().CompressedDataStructure->GetApproxCompressedSize(): 0;
+}
+
+int32 UAnimSequence::GetApproxCompressedSize_Lockless() const
+{
+	const int32 BoneTotal = GetApproxBoneCompressedSize_Lockless();
+	int32 CurveTotal = GetCompressedData_Internal().CompressedCurveByteStream.Num();
 	return BoneTotal + CurveTotal;
 }
 
+UAnimSequence::FScopedCompressedAnimSequence UAnimSequence::GetCompressedData() const
+{
+	const FCompressedAnimSequence& PlatformCompressedData = GetCompressedData_Internal();
+	return FScopedCompressedAnimSequence(this, &PlatformCompressedData);
+}
+
 /**
  * Deserializes old compressed track formats from the specified archive.
  */
@@ -613,6 +621,7 @@ void UAnimSequence::Serialize(FArchive& Ar)
 	
 	Ar.UsingCustomVersion(FFrameworkObjectVersion::GUID);
 	Ar.UsingCustomVersion(FUE5MainStreamObjectVersion::GUID);
+	Ar.UsingCustomVersion(FFortniteMainBranchObjectVersion::GUID);	
 
 	FRawCurveTracks RawCurveCache;
 	
@@ -741,13 +750,84 @@ void UAnimSequence::Serialize(FArchive& Ar)
 
 		if (bSerializeCompressedData)
 		{
-			SerializeCompressedData(Ar,false);
-			if (!bIsTransacting)
+			// Backwards compatibility, just serialize out compressed data
+			if (Ar.IsLoading() && Ar.CustomVer(FFortniteMainBranchObjectVersion::GUID) < FFortniteMainBranchObjectVersion::AnimSequenceRawDataOnlyFlagRemoval)
 			{
-				bool bTemp = bUseRawDataOnly;
-				Ar << bTemp;
-				bUseRawDataOnly = bTemp;
+			    SerializeCompressedData(Ar,false);
+			    if (!bIsTransacting)
+			    {
+					bool bTemp = false;
+				    Ar << bTemp;
+			    }
+			}
+			else
+			{
+#if WITH_EDITOR
+				const bool bIsLoading = Ar.IsLoading();
+				const bool bIsSaving = Ar.IsSaving();
+
+				if (bIsLoading)
+				{
+					CompressedDataLock.WriteLock();
+				}
+				else
+				{
+					CompressedDataLock.ReadLock();
+				}
+				
+				if (bCookingTargetNeedsCompressedData)
+				{
+					const ITargetPlatform* Platform = Ar.CookingTarget();
+					checkf(Platform, TEXT("Invalid cooking target platform"));
+					// During cook, serialize out singular (target platform) compressed data 
+					FCompressedAnimSequence& PlatformCompressedData = GetPlatformCompressedData(Platform);
+					check(PlatformCompressedData.IsValid(this, true));
+					SerializeCompressedData(Ar, PlatformCompressedData);
+				}
+				else if(bIsDuplicating || bIsTransacting || bIsCountingMemory)
+				{
+					// During duplication/transacting etc. serialize all contained compressed data
+					if (bIsSaving)
+					{
+						uint32 NumPlatforms = DataByPlatformKeyHash.Num();
+						Ar << NumPlatforms;
+						for (TPair<FIoHash, TUniquePtr<FCompressedAnimSequence>>& CompressedDataPair : DataByPlatformKeyHash)
+						{
+							Ar << CompressedDataPair.Key;
+							SerializeCompressedData(Ar, *CompressedDataPair.Value.Get());
+						}
+					}
+					else
+					{
+						uint32 NumPlatforms = 0;
+						Ar << NumPlatforms;
+						for (uint32 Index = 0; Index < NumPlatforms; ++Index)
+						{
+							FIoHash PlatformHash;
+							Ar << PlatformHash;
+							FCompressedAnimSequence& TargetData = *DataByPlatformKeyHash.Emplace(PlatformHash, MakeUnique<FCompressedAnimSequence>()).Get();
+							SerializeCompressedData(Ar, TargetData);
+						}
+					}					
+				}
+				
+				if (bIsLoading)
+				{
+					CompressedDataLock.WriteUnlock();
+				}
+				else
+				{
+					CompressedDataLock.ReadUnlock();
+				}
+#else
+				{
+					// During cooked runtime serialize in singular compressed data
+					SerializeCompressedData(Ar, GetCompressedData_Internal());
+				}
+#endif
 			}
+			
+			
 		}
 	}
 
@@ -991,7 +1071,7 @@ void UAnimSequence::PreSave(FObjectPreSaveContext ObjectSaveContext)
 		// In case compressed data was cleared between ::IsCachedCookedPlatformDataLoaded and ::PreSave being called, synchronously run compression
 		if(!IsCachedCookedPlatformDataLoaded(ObjectSaveContext.GetTargetPlatform()))
 		{
-			CacheDerivedData(ObjectSaveContext.GetTargetPlatform());
+			CacheDerivedDataForPlatform(ObjectSaveContext.GetTargetPlatform());
 			check(IsCachedCookedPlatformDataLoaded(ObjectSaveContext.GetTargetPlatform()));
 		}
 	}
@@ -1012,6 +1092,10 @@ void UAnimSequence::PreSave(FObjectPreSaveContext ObjectSaveContext)
 
 void UAnimSequence::PostLoad()
 {
+#if WITH_EDITOR
+	ValidateCompressionSettings();
+#endif // WITH_EDITOR
+	
 	//Parent PostLoad will ensure that skeleton is fully loaded
 	//before we do anything further in PostLoad
 	Super::PostLoad();
@@ -1130,9 +1214,9 @@ void UAnimSequence::BeginDestroy()
 
 	Super::BeginDestroy();
 
-	PRAGMA_DISABLE_DEPRECATION_WARNINGS
-	CompressedData.Reset();
-	PRAGMA_ENABLE_DEPRECATION_WARNINGS
+#if !WITH_EDITOR
+	GetCompressedData_Internal().Reset();
+#endif
 }
 
 bool UAnimSequence::IsReadyForFinishDestroy()
@@ -1226,7 +1310,7 @@ void UAnimSequence::PostEditChangeProperty(FPropertyChangedEvent& PropertyChange
 	ValidateModel();
 	if (IsDataModelValid() && bNeedPostProcess)
 	{		
-		ClearAllCachedCookedPlatformData();
+		ClearAllCompressionData();
 		BeginCacheDerivedDataForCurrentPlatform();
 	}
 }
@@ -1238,9 +1322,36 @@ void UAnimSequence::BeginCacheForCookedPlatformData(const ITargetPlatform* Targe
 #endif // WITH_EDITOR
 
 #if WITH_EDITOR
+
+bool UAnimSequence::HasCompressedDataForHash(const FIoHash& InKeyHash) const
+{
+	UE_COMPRESSED_DATA_READ_SCOPE(this)
+	
+	if (InKeyHash.IsZero())
+	{
+		UE_LOG(LogAnimation, Warning, TEXT("HasCompressedDataForHash: Zero key hash %s "), *GetName());
+		return false;
+	}
+		
+	if(const TUniquePtr<FCompressedAnimSequence>* CompressedDataPtr = DataByPlatformKeyHash.Find(InKeyHash))
+	{
+		if ((*CompressedDataPtr)->IsValid(this))
+		{
+			if (CacheTasksByKeyHash.Contains(InKeyHash))
+			{
+				return false;
+			}
+
+			return true;
+		}
+	}
+
+	return false;
+}
+
 bool UAnimSequence::IsCachedCookedPlatformDataLoaded(const ITargetPlatform* TargetPlatform)
 {
-	const FIoHash KeyHash = CreateDerivedDataKeyHash(TargetPlatform);
+	const FIoHash KeyHash = GetDerivedDataKeyHash(TargetPlatform);
 	if (KeyHash.IsZero())
 	{
 		UE_LOG(LogAnimation, Warning, TEXT("Zero key hash compressed animation data for %s requested platform %s"), *GetName(), *TargetPlatform->PlatformName());
@@ -1252,30 +1363,9 @@ bool UAnimSequence::IsCachedCookedPlatformDataLoaded(const ITargetPlatform* Targ
 		EndCacheDerivedData(KeyHash);
 	}
 
-	auto RecompressForPlatform = [this, TargetPlatform, KeyHash]()
-	{
-		if (IsInGameThread())
-		{
-			if (!HasAnyFlags(EObjectFlags::RF_NeedPostLoad) && GetPackage()->GetHasBeenEndLoaded())
-			{
-				BeginCacheForCookedPlatformData(TargetPlatform);
-			}
-		}
-	};
-
-	if (KeyHash == DataKeyHash)
-	{
-		if (!CacheTasksByKeyHash.Contains(KeyHash) && CompressedData.IsValid(this) && !bUseRawDataOnly)
-		{
-			return true;
-		}
-		else if (!CacheTasksByKeyHash.Contains(KeyHash))
-		{
-			RecompressForPlatform();
-		}
-	}
-	else
+	bool bShouldRecompress = false;
 	{
+		UE_COMPRESSED_DATA_READ_SCOPE(this);
 		if(const TUniquePtr<FCompressedAnimSequence>* CompressedDataPtr = DataByPlatformKeyHash.Find(KeyHash))
 		{
 			if (!CacheTasksByKeyHash.Contains(KeyHash) && (*CompressedDataPtr)->IsValid(this))
@@ -1284,15 +1374,22 @@ bool UAnimSequence::IsCachedCookedPlatformDataLoaded(const ITargetPlatform* Targ
 			}
 			else if (!CacheTasksByKeyHash.Contains(KeyHash))
 			{
-				RecompressForPlatform();
+				bShouldRecompress = true;
 			}
 		}
 		else if (!CacheTasksByKeyHash.Contains(KeyHash))
 		{
-			RecompressForPlatform();
+			bShouldRecompress = true;
 		}
 	}
 
+	if (bShouldRecompress)
+	{
+		if (!HasAnyFlags(EObjectFlags::RF_NeedPostLoad) && CanBeCompressed())
+		{
+			BeginCacheForCookedPlatformData(TargetPlatform);
+		}
+	}
 
 	return false;
 }
@@ -1330,24 +1427,30 @@ static FAutoConsoleVariableRef CVarForceRawData(
 
 void UAnimSequence::GetBoneTransform(FTransform& OutAtom, FSkeletonPoseBoneIndex BoneIndex, const FAnimExtractContext& ExtractionContext, bool bUseRawData) const
 {
-	// If the caller didn't request that raw animation data be used . . .
-	const bool bEvaluateCompressedData =
-#if WITH_EDITOR 
-		!GForceRawData && 
-#endif 
-		(!bUseRawData || !CanEvaluateRawAnimationData());
+	UE_COMPRESSED_DATA_READ_SCOPE(this);
+	GetBoneTransform_Lockless(OutAtom, BoneIndex, ExtractionContext, bUseRawData);	
+}
+
+void UAnimSequence::GetBoneTransform_Lockless(FTransform& OutAtom, FSkeletonPoseBoneIndex BoneIndex, const FAnimExtractContext& ExtractionContext, bool bForceUseRawData) const
+{	
+	const FCompressedAnimSequence& PlatformCompressedData = GetPlatformCompressedData(ExtractionContext);
+	const bool bUseRawDataForPoseExtraction = (CanEvaluateRawAnimationData() && bForceUseRawData) || ShouldUseRawDataForPoseExtraction_Lockless(ExtractionContext);
+
+#if WITH_EDITOR
+	ensure(!ExtractionContext.bEnforceCompressedDataSampling || (PlatformCompressedData.IsValid(this) && !bUseRawDataForPoseExtraction));
+#endif // WITH_EDITOR
 
 	const EAnimInterpolationType InterpolationType = ExtractionContext.InterpolationOverride.Get(Interpolation);
-	if ( bEvaluateCompressedData && IsCompressedDataValid())
+	if (!bUseRawDataForPoseExtraction && PlatformCompressedData.IsValid(this))
 	{
-		FAnimSequenceDecompressionContext DecompContext(PlatformTargetFrameRate.Default, PlatformTargetFrameRate.Default.AsFrameTime(GetPlayLength()).RoundToFrame().Value, InterpolationType, GetRetargetTransformsSourceName(), *CompressedData.CompressedDataStructure, GetSkeleton()->GetRefLocalPoses(), CompressedData.CompressedTrackToSkeletonMapTable, GetSkeleton(), IsValidAdditive(), AdditiveAnimType);
+		FAnimSequenceDecompressionContext DecompContext(PlatformTargetFrameRate.Default, PlatformTargetFrameRate.Default.AsFrameTime(GetPlayLength()).RoundToFrame().Value, InterpolationType, GetRetargetTransformsSourceName(), *PlatformCompressedData.CompressedDataStructure, GetSkeleton()->GetRefLocalPoses(), PlatformCompressedData.CompressedTrackToSkeletonMapTable, GetSkeleton(), IsValidAdditive(), AdditiveAnimType);
 		DecompContext.Seek(ExtractionContext.CurrentTime);
-		if (CompressedData.BoneCompressionCodec)
+		if (PlatformCompressedData.BoneCompressionCodec)
 		{
-			const int32 TrackIndex = CompressedData.GetTrackIndexFromSkeletonIndex(BoneIndex.GetInt());
+			const int32 TrackIndex = PlatformCompressedData.GetTrackIndexFromSkeletonIndex(BoneIndex.GetInt());
 			if (TrackIndex != INDEX_NONE)
 			{				
-				CompressedData.BoneCompressionCodec->DecompressBone(DecompContext, TrackIndex, OutAtom);
+				PlatformCompressedData.BoneCompressionCodec->DecompressBone(DecompContext, TrackIndex, OutAtom);
 			}
 		}
 	}
@@ -1373,34 +1476,9 @@ void UAnimSequence::GetBoneTransform(FTransform& OutAtom, FSkeletonPoseBoneIndex
 
 void UAnimSequence::GetBoneTransform(FTransform& OutAtom, FSkeletonPoseBoneIndex BoneIndex, FAnimSequenceDecompressionContext& DecompContext, bool bUseRawData) const
 {
-	// If the caller didn't request that raw animation data be used . . .
-	const bool bEvaluateCompressedData = (!bUseRawData || !CanEvaluateRawAnimationData());
-	if ( bEvaluateCompressedData && IsCompressedDataValid())
-	{
-		if (CompressedData.BoneCompressionCodec)
-		{
-			CompressedData.BoneCompressionCodec->DecompressBone(DecompContext, CompressedData.GetTrackIndexFromSkeletonIndex(BoneIndex.GetInt()), OutAtom);
-			return;
-		}
-	}
-	else if (CanEvaluateRawAnimationData())
-	{
-#if WITH_EDITOR
-		ValidateModel();
-		const FName BoneName = GetSkeleton()->GetReferenceSkeleton().GetBoneName(BoneIndex.GetInt());
-		OutAtom = DataModelInterface->EvaluateBoneTrackTransform(BoneName, DataModelInterface->GetFrameRate().AsFrameTime(DecompContext.GetEvaluationTime()), Interpolation);
-
-		const FAnimationCurveIdentifier TransformCurveId(BoneName, ERawCurveTrackTypes::RCT_Transform);
-		if (const FTransformCurve* TransformCurvePtr = DataModelInterface->FindTransformCurve(TransformCurveId))
-		{
-			const FTransform AdditiveTransform = TransformCurvePtr->Evaluate(DecompContext.GetEvaluationTime(), 1.f);
-			const FTransform LocalTransform = OutAtom;
-			OutAtom.SetRotation(LocalTransform.GetRotation() * AdditiveTransform.GetRotation());
-			OutAtom.SetTranslation(LocalTransform.TransformPosition(AdditiveTransform.GetTranslation()));
-			OutAtom.SetScale3D(LocalTransform.GetScale3D() * AdditiveTransform.GetScale3D());	
-		}
-#endif
-	}
+	UE_COMPRESSED_DATA_READ_SCOPE(this);
+	FAnimExtractContext ExtractContext(DecompContext.GetEvaluationTime());
+	GetBoneTransform(OutAtom, BoneIndex, ExtractContext, bUseRawData);
 }
 
 void UAnimSequence::HandleAssetPlayerTickedInternal(FAnimAssetTickContext &Context, const float PreviousTime, const float MoveDelta, const FAnimTickRecord &Instance, struct FAnimNotifyQueue& NotifyQueue) const
@@ -1409,6 +1487,7 @@ void UAnimSequence::HandleAssetPlayerTickedInternal(FAnimAssetTickContext &Conte
 
 	if (bEnableRootMotion)
 	{
+		UE_COMPRESSED_DATA_READ_SCOPE(this);
 		const FAnimExtractContext ExtractionContext(static_cast<double>(PreviousTime), true, FDeltaTimeRecord(MoveDelta), Instance.bLooping);
 		Context.RootMotionMovementParams.Accumulate(ExtractRootMotion(ExtractionContext));
 	}
@@ -1416,65 +1495,8 @@ void UAnimSequence::HandleAssetPlayerTickedInternal(FAnimAssetTickContext &Conte
 
 FTransform UAnimSequence::ExtractRootTrackTransform(const FAnimExtractContext& ExtractionContext, const FBoneContainer* RequiredBones) const
 {
-	const int32 RootBoneIndex = 0;
-	const bool bContainsRootBoneTrack = [this, RootBoneIndex]() 
-	{
-#if WITH_EDITOR
-		if (CanEvaluateRawAnimationData() && bUseRawDataOnly)
-		{
-			ValidateModel();
-
-			const FName RootBoneName = GetSkeleton()->GetReferenceSkeleton().GetBoneName(RootBoneIndex);
-			if (DataModelInterface->IsValidBoneTrackName(RootBoneName))
-			{
-				return true;
-			}
-		}
-		else
-#endif // WITH_EDITOR
-		{
-			const TArray<FTrackToSkeletonMap>& TrackToSkeletonMap = CompressedData.CompressedTrackToSkeletonMapTable;
-			// we assume root is in first data if available = SkeletonIndex == 0 && BoneTreeIndex == 0)
-			if ((TrackToSkeletonMap.Num() > 0) && (TrackToSkeletonMap[0].BoneTreeIndex == RootBoneIndex))
-			{
-				return true;
-			}
-		}
-
-		return false;
-	}();
-
-	if(bContainsRootBoneTrack)
-	{
-		// if we do have root data, then return root data
-		FTransform RootTransform;
-		GetBoneTransform(RootTransform, FSkeletonPoseBoneIndex(RootBoneIndex), ExtractionContext, bUseRawDataOnly);
-		return RootTransform;
-	}
-
-	// Fallback to root bone from reference skeleton.
-	if( RequiredBones )
-	{
-		const FReferenceSkeleton& RefSkeleton = RequiredBones->GetReferenceSkeleton();
-		if( RefSkeleton.GetNum() > 0 )
-		{
-			return RefSkeleton.GetRefBonePose()[RootBoneIndex];
-		}
-	}
-
-	USkeleton * MySkeleton = GetSkeleton();
-	// If we don't have a RequiredBones array, get root bone from default skeleton.
-	if( !RequiredBones &&  MySkeleton )
-	{
-		const FReferenceSkeleton& RefSkeleton = MySkeleton->GetReferenceSkeleton();
-		if( RefSkeleton.GetNum() > 0 )
-		{
-			return RefSkeleton.GetRefBonePose()[RootBoneIndex];
-		}
-	}
-
-	// Otherwise, use identity.
-	return FTransform::Identity;
+	UE_COMPRESSED_DATA_READ_SCOPE(this);
+	return ExtractRootTrackTransform_Lockless(ExtractionContext, RequiredBones);
 }
 
 FTransform UAnimSequence::ExtractRootMotion(const FAnimExtractContext& ExtractionContext) const
@@ -1533,11 +1555,12 @@ FTransform UAnimSequence::ExtractRootMotionFromRange(double StartTime, double En
 		}
 	}
 
+	UE_COMPRESSED_DATA_READ_SCOPE(this);
 	FAnimExtractContext Context = ExtractionContext;
 	Context.CurrentTime = StartTime;
-	FTransform StartTransform = ExtractRootTrackTransform(Context, nullptr);
+	FTransform StartTransform = ExtractRootTrackTransform_Lockless(Context, nullptr);
 	Context.CurrentTime = EndTime;
-	FTransform EndTransform = ExtractRootTrackTransform(Context, nullptr);
+	FTransform EndTransform = ExtractRootTrackTransform_Lockless(Context, nullptr);
 
 	// Use old calculation if needed.
 	if (bUseNormalizedRootMotionScale)
@@ -1593,7 +1616,7 @@ void UAnimSequence::GetAnimationPose(FAnimationPoseData& OutAnimationPoseData, c
 	const FCompactPose& OutPose = OutAnimationPoseData.GetPose();
 
 	// @todo anim: if compressed and baked in the future, we don't have to do this 
-	if (UseRawDataForPoseExtraction(OutPose.GetBoneContainer()) && IsValidAdditive())
+	if (IsValidAdditive() && ShouldUseRawDataForPoseExtraction(OutPose.GetBoneContainer(), ExtractionContext))
 	{
 		if (AdditiveAnimType == AAT_LocalSpaceBase)
 		{
@@ -1668,11 +1691,13 @@ void UAnimSequence::GetBonePose(FAnimationPoseData& OutAnimationPoseData, const
 
 	FBoneContainer& RequiredBones = OutPose.GetBoneContainer();
 	check(!bForceUseRawData || CanEvaluateRawAnimationData());
-	const bool bUseRawDataForPoseExtraction = (CanEvaluateRawAnimationData() && bForceUseRawData) ||
+
+	UE_COMPRESSED_DATA_READ_SCOPE(this);
+	const bool bUseRawDataForPoseExtraction = (CanEvaluateRawAnimationData() && bForceUseRawData) || ShouldUseRawDataForPoseExtraction_Lockless(RequiredBones, ExtractionContext);
+	const FCompressedAnimSequence& PlatformCompressedData = GetPlatformCompressedData(ExtractionContext);
 #if WITH_EDITOR
- 		GForceRawData == 1 || 
+	check(!ExtractionContext.bEnforceCompressedDataSampling || (PlatformCompressedData.IsValid(this) && !bUseRawDataForPoseExtraction));
 #endif // WITH_EDITOR
-		UseRawDataForPoseExtraction(RequiredBones);
 
 	const bool bIsBakedAdditive = !bUseRawDataForPoseExtraction && IsValidAdditive();
 
@@ -1725,10 +1750,11 @@ void UAnimSequence::GetBonePose(FAnimationPoseData& OutAnimationPoseData, const
 	}
 
 #if WITH_EDITOR
-	const int32 NumTracks = bUseRawDataForPoseExtraction ? DataModelInterface->GetNumBoneTracks() : CompressedData.CompressedTrackToSkeletonMapTable.Num();
+	const int32 NumTracks =	bUseRawDataForPoseExtraction ? DataModelInterface->GetNumBoneTracks() : PlatformCompressedData.CompressedTrackToSkeletonMapTable.Num();
 #else
-	const int32 NumTracks = CompressedData.CompressedTrackToSkeletonMapTable.Num();
-#endif 
+	const int32 NumTracks =	PlatformCompressedData.CompressedTrackToSkeletonMapTable.Num();
+#endif
+
 	const bool bTreatAnimAsAdditive = (IsValidAdditive() && !bUseRawDataForPoseExtraction); // Raw data is never additive
 	FAnimExtractContext RootMotionContext = ExtractionContext;
 	RootMotionContext.CurrentTime = 0.0;
@@ -1737,7 +1763,7 @@ void UAnimSequence::GetBonePose(FAnimationPoseData& OutAnimationPoseData, const
 		!ExtractionContext.bIgnoreRootLock &&
 #endif // WITH_EDITOR
 		bForceRootLock,
-		ExtractRootTrackTransform(RootMotionContext, &RequiredBones), bTreatAnimAsAdditive);
+		ExtractRootTrackTransform_Lockless(RootMotionContext, &RequiredBones), bTreatAnimAsAdditive);
 
 #if WITH_EDITOR
 	// Evaluate raw (source) curve and bone data
@@ -1763,8 +1789,8 @@ void UAnimSequence::GetBonePose(FAnimationPoseData& OutAnimationPoseData, const
 	{
 		// Evaluate compressed bone data
 		const EAnimInterpolationType InterpolationType = ExtractionContext.InterpolationOverride.Get(Interpolation);
-		FAnimSequenceDecompressionContext DecompContext(PlatformTargetFrameRate.Default, PlatformTargetFrameRate.Default.AsFrameTime(GetPlayLength()).RoundToFrame().Value, InterpolationType, GetRetargetTransformsSourceName(), *CompressedData.CompressedDataStructure, GetSkeleton()->GetRefLocalPoses(), CompressedData.CompressedTrackToSkeletonMapTable, GetSkeleton(), IsValidAdditive(), AdditiveAnimType);
-		UE::Anim::Decompression::DecompressPose(OutPose, CompressedData, ExtractionContext, DecompContext, GetRetargetTransforms(), RootMotionReset);
+		FAnimSequenceDecompressionContext DecompContext(PlatformTargetFrameRate.Default, PlatformTargetFrameRate.Default.AsFrameTime(GetPlayLength()).RoundToFrame().Value, InterpolationType, GetRetargetTransformsSourceName(), *PlatformCompressedData.CompressedDataStructure, GetSkeleton()->GetRefLocalPoses(), PlatformCompressedData.CompressedTrackToSkeletonMapTable, GetSkeleton(), IsValidAdditive(), AdditiveAnimType);
+		UE::Anim::Decompression::DecompressPose(OutPose, PlatformCompressedData, ExtractionContext, DecompContext, GetRetargetTransforms(), RootMotionReset);
 	}
 
 	// (Always) evaluate compressed curve data
@@ -1774,7 +1800,7 @@ void UAnimSequence::GetBonePose(FAnimationPoseData& OutAnimationPoseData, const
 		if (!bUseRawDataForPoseExtraction)
 #endif // WITH_EDITOR
 		{
-			EvaluateCurveData(OutAnimationPoseData.GetCurve(), ExtractionContext, bUseRawDataForPoseExtraction);
+			EvaluateCurveData_Lockless(OutAnimationPoseData.GetCurve(), ExtractionContext, bUseRawDataForPoseExtraction);
 		}
 	}
 
@@ -1911,21 +1937,269 @@ bool UAnimSequence::CanEvaluateRawAnimationData() const
 #endif
 }
 
+const FCompressedAnimSequence& UAnimSequence::GetCompressedData_Internal() const
+{
+#if WITH_EDITOR
+	static thread_local FCompressedAnimSequence DefaultCompressedAnimSequence;
+	const ITargetPlatform* RunningPlatform = GetTargetPlatformManagerRef().GetRunningTargetPlatform();
+
+	if(DataByPlatformKeyHash.Num())
+	{
+		const FIoHash CurrentPlatformHash = GetDerivedDataKeyHash(RunningPlatform);
+		if (const TUniquePtr<FCompressedAnimSequence>* PlatformCompressedData = DataByPlatformKeyHash.Find(CurrentPlatformHash))
+		{
+			return *PlatformCompressedData->Get();
+		}
+	}
+
+	return DefaultCompressedAnimSequence;
+#else	
+	PRAGMA_DISABLE_DEPRECATION_WARNINGS
+	return CompressedData;
+	PRAGMA_ENABLE_DEPRECATION_WARNINGS
+#endif // WITH_EDITOR
+}
+
+
+const FCompressedAnimSequence& UAnimSequence::GetPlatformCompressedData(const FAnimExtractContext& AnimExtractContext) const
+{
+#if WITH_EDITOR
+	FIoHash Hash = FIoHash::Zero;	
+	static thread_local FCompressedAnimSequence DefaultCompressedAnimSequence;	
+	if (AnimExtractContext.bEnforceCompressedDataSampling)
+	{
+		if(AnimExtractContext.TargetPlatformHash.IsZero())
+		{
+			// Not set, check platform instead
+			if (AnimExtractContext.TargetPlatform)
+			{
+				Hash = GetDerivedDataKeyHash(AnimExtractContext.TargetPlatform);
+			}
+			else
+			{
+				const ITargetPlatform* RunningPlatform = GetTargetPlatformManagerRef().GetRunningTargetPlatform();
+				Hash = GetDerivedDataKeyHash(RunningPlatform);
+			}
+		}
+		else
+		{
+			Hash = AnimExtractContext.TargetPlatformHash;
+		}
+	}
+	else
+	{
+		const ITargetPlatform* RunningPlatform = GetTargetPlatformManagerRef().GetRunningTargetPlatform();
+		Hash = GetDerivedDataKeyHash(RunningPlatform);
+	}
+	
+	if (const TUniquePtr<FCompressedAnimSequence>* PlatformData = DataByPlatformKeyHash.Find(Hash))
+	{
+		return *PlatformData->Get();
+	}
+
+	return DefaultCompressedAnimSequence;
+#else
+	PRAGMA_DISABLE_DEPRECATION_WARNINGS
+	return CompressedData;
+	PRAGMA_ENABLE_DEPRECATION_WARNINGS
+#endif // WITH_EDITOR
+}
+
+FCompressedAnimSequence& UAnimSequence::GetPlatformCompressedData(const FAnimExtractContext& AnimExtractContext)
+{	
+#if WITH_EDITOR
+	FIoHash Hash = FIoHash::Zero;	
+	static thread_local FCompressedAnimSequence DefaultCompressedAnimSequence;	
+	if (AnimExtractContext.bEnforceCompressedDataSampling)
+	{		
+		if(AnimExtractContext.TargetPlatformHash.IsZero())
+		{
+			// Not set, check platform instead
+			if (AnimExtractContext.TargetPlatform)
+			{
+				Hash = GetDerivedDataKeyHash(AnimExtractContext.TargetPlatform);
+			}
+			else
+			{
+				const ITargetPlatform* RunningPlatform = GetTargetPlatformManagerRef().GetRunningTargetPlatform();
+				Hash = GetDerivedDataKeyHash(RunningPlatform);
+			}
+		}
+		else
+		{
+			Hash = AnimExtractContext.TargetPlatformHash;
+		}
+	}
+	else
+	{
+		const ITargetPlatform* RunningPlatform = GetTargetPlatformManagerRef().GetRunningTargetPlatform();
+		Hash = GetDerivedDataKeyHash(RunningPlatform);
+	}
+	
+	if (TUniquePtr<FCompressedAnimSequence>* PlatformData = DataByPlatformKeyHash.Find(Hash))
+	{
+		return *PlatformData->Get();
+	}
+
+	return DefaultCompressedAnimSequence;
+#else
+	PRAGMA_DISABLE_DEPRECATION_WARNINGS
+	return CompressedData;
+	PRAGMA_ENABLE_DEPRECATION_WARNINGS
+#endif // WITH_EDITOR
+}
+
+
+FCompressedAnimSequence& UAnimSequence::GetPlatformCompressedData(const ITargetPlatform* InTargetPlatform)
+{	
+#if WITH_EDITOR
+	thread_local FCompressedAnimSequence DefaultCompressedAnimSequence;
+	
+	check(InTargetPlatform);
+	
+	const FIoHash Hash = GetDerivedDataKeyHash(InTargetPlatform);
+	check(!Hash.IsZero());
+
+	if (const TUniquePtr<FCompressedAnimSequence>* PlatformData = DataByPlatformKeyHash.Find(Hash))
+	{
+		return *PlatformData->Get();
+	}
+
+	return DefaultCompressedAnimSequence;
+
+#else
+	return GetCompressedData_Internal();
+#endif // WITH_EDITOR
+}
+
+#if WITH_EDITOR
+const FCompressedAnimSequence& UAnimSequence::GetPlatformCompressedData(const ITargetPlatform* InTargetPlatform) const
+{
+	thread_local FCompressedAnimSequence DefaultCompressedAnimSequence;
+	
+	check(InTargetPlatform);
+	
+	const FIoHash Hash = GetDerivedDataKeyHash(InTargetPlatform);
+	check(!Hash.IsZero());
+
+	if (const TUniquePtr<FCompressedAnimSequence>* PlatformData = DataByPlatformKeyHash.Find(Hash))
+	{
+		return *PlatformData->Get();
+	}
+
+	return DefaultCompressedAnimSequence;
+}
+#endif // WITH_EDITOR
+
+
 #if WITH_EDITORONLY_DATA
 void UAnimSequence::UpdateCompressedCurveName(const FName& OldCurveName, const FName& NewCurveName)
 {
-	for (FAnimCompressedCurveIndexedName& IndexedCurveName : CompressedData.IndexedCurveNames)
+	UE_COMPRESSED_DATA_WRITE_SCOPE(this);
+	for (TPair<FIoHash, TUniquePtr<FCompressedAnimSequence>>& CompressedDataPair : DataByPlatformKeyHash)
 	{
-		if (IndexedCurveName.CurveName == OldCurveName)
+		FCompressedAnimSequence& PlatformCompressedData = *CompressedDataPair.Value.Get();
+
+		for (FAnimCompressedCurveIndexedName& IndexedCurveName : PlatformCompressedData.IndexedCurveNames)
 		{
-			IndexedCurveName.CurveName = NewCurveName;
-			break;
+			if (IndexedCurveName.CurveName == OldCurveName)
+			{
+				IndexedCurveName.CurveName = NewCurveName;
+				break;
+			}
 		}
+		PlatformCompressedData.RebuildCurveIndexTable();		
 	}
+}
+#endif // WITH_EDITORONLY_DATA
+
+FTransform UAnimSequence::ExtractRootTrackTransform_Lockless(const FAnimExtractContext& ExtractionContext, const FBoneContainer* RequiredBones) const
+{	
+	const int32 RootBoneIndex = 0;
+
+	bool bValidCompressedData = true;
+	
+	bool bContainsRootBoneTrack = [this, RequiredBones, RootBoneIndex, &ExtractionContext, &bValidCompressedData]()
+	{
+		const FCompressedAnimSequence& PlatformCompressedData = GetPlatformCompressedData(ExtractionContext);
+
+		bool bUseRawDataForPoseExtraction = false;
+#if WITH_EDITOR
+		if (RequiredBones)
+		{
+			bUseRawDataForPoseExtraction = RequiredBones->ShouldUseRawData();
+		}
+		bValidCompressedData = PlatformCompressedData.IsValid(this);
+		check(!ExtractionContext.bEnforceCompressedDataSampling || (bValidCompressedData && !bUseRawDataForPoseExtraction));	
+#endif // WITH_EDITOR
+		
+#if WITH_EDITOR
+		if (bUseRawDataForPoseExtraction)
+		{
+			if(CanEvaluateRawAnimationData())
+			{
+				ValidateModel();
+
+				const FName RootBoneName = GetSkeleton()->GetReferenceSkeleton().GetBoneName(RootBoneIndex);
+				if (DataModelInterface->IsValidBoneTrackName(RootBoneName))
+				{
+					return true;
+				}
+			}
+		}
+		else
+#endif // WITH_EDITOR
+		{
+			const TArray<FTrackToSkeletonMap>& TrackToSkeletonMap = PlatformCompressedData.CompressedTrackToSkeletonMapTable;
+			// we assume root is in first data if available = SkeletonIndex == 0 && BoneTreeIndex == 0)
+			if ((TrackToSkeletonMap.Num() > 0) && (TrackToSkeletonMap[0].BoneTreeIndex == RootBoneIndex))
+			{
+				return true;
+			}
+		}
+
+		return false;
+	}();
+
+	if(bContainsRootBoneTrack)
+	{
+		// if we do have root data, then return root data
+		FTransform RootTransform;
+		GetBoneTransform_Lockless(RootTransform, FSkeletonPoseBoneIndex(RootBoneIndex), ExtractionContext, !bValidCompressedData);
+		return RootTransform;
+	}
+
+	// Fallback to root bone from reference skeleton.
+	if( RequiredBones )
+	{
+		const FReferenceSkeleton& RefSkeleton = RequiredBones->GetReferenceSkeleton();
+		if( RefSkeleton.GetNum() > 0 )
+		{
+			return RefSkeleton.GetRefBonePose()[RootBoneIndex];
+		}
+	}
+
+	USkeleton * MySkeleton = GetSkeleton();
+	// If we don't have a RequiredBones array, get root bone from default skeleton.
+	if( !RequiredBones &&  MySkeleton )
+	{
+		const FReferenceSkeleton& RefSkeleton = MySkeleton->GetReferenceSkeleton();
+		if( RefSkeleton.GetNum() > 0 )
+		{
+			return RefSkeleton.GetRefBonePose()[RootBoneIndex];
+		}
+	}
+
+	// Otherwise, use identity.
+	return FTransform::Identity;
+}
 
-	CompressedData.RebuildCurveIndexTable();
+FCompressedAnimSequence& UAnimSequence::GetCompressedData_Internal()
+{
+	PRAGMA_DISABLE_DEPRECATION_WARNINGS
+	return CompressedData;
+	PRAGMA_ENABLE_DEPRECATION_WARNINGS
 }
-#endif // WITH_EDITOR
 
 #if WITH_EDITOR
 void UAnimSequence::SetRetargetSourceAsset(USkeletalMesh* InRetargetSourceAsset)
@@ -2014,7 +2288,6 @@ void UAnimSequence::RetargetBoneTransform(FTransform& BoneTransform, const int32
 	FAnimationRuntime::RetargetBoneTransform(GetSkeleton(), GetRetargetTransformsSourceName(), GetRetargetTransforms(), BoneTransform, SkeletonBoneIndex, BoneIndex, RequiredBones, bIsBakedAdditive);
 }
 
-
 #if WITH_EDITOR
 bool UAnimSequence::ShouldPerformStripping(const bool bPerformFrameStripping, const bool bPerformStrippingOnOddFramedAnims) const
 {
@@ -2047,31 +2320,35 @@ void UAnimSequence::SerializeCompressedData(FArchive& Ar, bool bDDCData)
 	
 	if (!HasAnyFlags(RF_ClassDefaultObject))
 	{
-		FCompressedAnimSequence& CompressedDataToSerialize = [this, &Ar]() -> FCompressedAnimSequence&
-		{
 #if WITH_EDITOR
-			if (Ar.IsSaving() && Ar.IsCooking())
-			{
-				if (const ITargetPlatform* Platform = Ar.CookingTarget())
-				{
-					const FIoHash Hash = CreateDerivedDataKeyHash(Platform);
-					if (!Hash.IsZero())
-					{
-						if (const TUniquePtr<FCompressedAnimSequence>* DataPtr = DataByPlatformKeyHash.Find(Hash))
-						{
-							return **DataPtr;
-						}
+		const ITargetPlatform* Platform = Ar.IsSaving() && Ar.IsCooking() ? Ar.CookingTarget() : GetTargetPlatformManagerRef().GetRunningTargetPlatform();
+#else
+		const ITargetPlatform* Platform = nullptr;
+#endif // WITH_EDITOR
+		
+		FCompressedAnimSequence& CompressedDataToSerialize = GetPlatformCompressedData(Platform);
+		if (Ar.IsLoading())
+		{
+			UE_COMPRESSED_DATA_WRITE_SCOPE(this);
+			CompressedDataToSerialize.SerializeCompressedData(Ar, bDDCData, this, this->GetSkeleton(), BoneCompressionSettings, CurveCompressionSettings);
+		}
+		else
+		{
+			UE_COMPRESSED_DATA_READ_SCOPE(this);
+			check(!Ar.IsCooking() || CompressedDataToSerialize.IsValid(this, true));
+			CompressedDataToSerialize.SerializeCompressedData(Ar, bDDCData, this, this->GetSkeleton(), BoneCompressionSettings, CurveCompressionSettings);
+		}
+	}
+}
 
-						// If there is no entry in the map for this platform key, it will be stored in CompressedData - so verify against the stored KeyHash for it
-						ensure(DataKeyHash == Hash);
-					}
-				}
-			}
-#endif
-			return CompressedData;
-		}();
-	
-		CompressedDataToSerialize.SerializeCompressedData(Ar, bDDCData, this, this->GetSkeleton(), BoneCompressionSettings, CurveCompressionSettings);
+void UAnimSequence::SerializeCompressedData(FArchive& Ar, FCompressedAnimSequence& CompressedDataToSerialize)
+{
+	LLM_SCOPE_BYNAME(TEXT("Animation/SequenceData"));
+
+	if (!HasAnyFlags(RF_ClassDefaultObject))
+	{
+		check(!Ar.IsCooking() || CompressedDataToSerialize.IsValid(this, true));
+		CompressedDataToSerialize.SerializeCompressedData(Ar, false, this, this->GetSkeleton(), BoneCompressionSettings, CurveCompressionSettings);
 	}
 }
 
@@ -2157,9 +2434,9 @@ void UAnimSequence::FlagDependentAnimationsAsRawDataOnly() const
 	for (TObjectIterator<UAnimSequence> Iter; Iter; ++Iter)
 	{
 		UAnimSequence* Seq = *Iter;
-		if (Seq->RefPoseSeq == this)
+		if (Seq->HasAllFlags(RF_LoadCompleted) && Seq->RefPoseSeq == this)
 		{
-			Seq->bUseRawDataOnly = true;
+			Seq->ClearAllCompressionData();
 		}
 	}
 }
@@ -2217,6 +2494,15 @@ void UAnimSequence::DeleteNotifyTrackData()
 	InitializeNotifyTrack();
 }
 
+void UAnimSequence::ResetCompressionSettings()
+{	
+	UE_COMPRESSED_DATA_WRITE_SCOPE(this);
+	
+	BoneCompressionSettings = nullptr;
+	CurveCompressionSettings = nullptr;
+	ValidateCompressionSettings();
+}
+
 void UAnimSequence::DeleteBoneAnimationData()
 {
 	ValidateModel();
@@ -2246,8 +2532,13 @@ bool UAnimSequence::IsValidAdditive() const
 	return false;
 }
 
-#if WITH_EDITOR
+UAnimSequence::FScopedCompressedAnimSequence UAnimSequence::GetCompressedData(const FAnimExtractContext& AnimExtractContext) const
+{
+	const FCompressedAnimSequence& PlatformCompressedData = GetPlatformCompressedData(AnimExtractContext);
+	return FScopedCompressedAnimSequence(this, &PlatformCompressedData);
+}
 
+#if WITH_EDITOR
 int32 FindMeshBoneIndexFromBoneName(USkeleton * Skeleton, const FName &BoneName)
 {
 	USkeletalMesh * PreviewMesh = Skeleton->GetPreviewMesh();
@@ -2551,6 +2842,14 @@ bool UAnimSequence::GetAllAnimationSequencesReferred(TArray<UAnimationAsset*>& A
 	return AnimationAssets.Num() > 0;
 }
 
+#if WITH_EDITOR
+UAnimSequence::FScopedCompressedAnimSequence UAnimSequence::GetCompressedData(const ITargetPlatform* InTargetPlatform) const
+{
+	const FCompressedAnimSequence& PlatformCompressedData = GetPlatformCompressedData(InTargetPlatform);
+	return FScopedCompressedAnimSequence(this, &PlatformCompressedData);
+}
+#endif // WITH_EDITOR
+
 void UAnimSequence::ReplaceReferredAnimations(const TMap<UAnimationAsset*, UAnimationAsset*>& ReplacementMap)
 {
 	Super::ReplaceReferredAnimations(ReplacementMap);
@@ -2666,9 +2965,13 @@ uint8* UAnimSequence::FindSyncMarkerPropertyData(int32 SyncMarkerIndex, FArrayPr
 	return nullptr;
 }
 
-bool UAnimSequence::DoesNeedRecompress() const
+bool UAnimSequence::IsCompressedDataOutOfDate() const
 {
-	return GetSkeleton() && (bUseRawDataOnly || (GetSkeletonVirtualBoneGuid() != GetSkeleton()->GetVirtualBoneGuid()));
+	return GetSkeleton() &&
+	(
+		!IsCompressedDataValid() ||
+		GetSkeletonVirtualBoneGuid() != GetSkeleton()->GetVirtualBoneGuid()
+	);
 }
 
 bool UAnimSequence::CreateAnimation(USkeletalMesh* Mesh)
@@ -2795,16 +3098,22 @@ int32 UAnimSequence::GetNumberOfSampledKeys() const
 #if WITH_EDITOR
 	return NumberOfSampledKeys;
 #else
-	return CompressedData.CompressedDataStructure ? CompressedData.CompressedDataStructure->CompressedNumberOfKeys :  PlatformTargetFrameRate.Default.AsFrameTime(GetPlayLength()).RoundToFrame().Value + 1;
+	return GetCompressedData_Internal().CompressedDataStructure ? GetCompressedData_Internal().CompressedDataStructure->CompressedNumberOfKeys : PlatformTargetFrameRate.Default.AsFrameTime(GetPlayLength()).RoundToFrame().Value + 1;
 #endif
 }
 
 void UAnimSequence::EvaluateCurveData(FBlendedCurve& OutCurve, const FAnimExtractContext& AnimExtractContext, bool bForceUseRawData) const
+{
+	UE_COMPRESSED_DATA_READ_SCOPE(this);
+	EvaluateCurveData_Lockless(OutCurve, AnimExtractContext, bForceUseRawData);
+}
+
+void UAnimSequence::EvaluateCurveData_Lockless(FBlendedCurve& OutCurve, const FAnimExtractContext& AnimExtractContext, bool bForceUseRawData) const
 {
 	SCOPE_CYCLE_COUNTER(STAT_AnimSeq_EvalCurveData);
 	
-	const bool bEvaluateRawData = bUseRawDataOnly
-	    || bForceUseRawData || !IsCurveCompressedDataValid();
+	const FCompressedAnimSequence& PlatformCompressedData = GetPlatformCompressedData(AnimExtractContext);
+	const bool bEvaluateRawData = bForceUseRawData || !PlatformCompressedData.IsCurveDataValid(this);
 	check(!bForceUseRawData || CanEvaluateRawAnimationData());
 	if (CanEvaluateRawAnimationData() && bEvaluateRawData)
 	{
@@ -2814,27 +3123,40 @@ void UAnimSequence::EvaluateCurveData(FBlendedCurve& OutCurve, const FAnimExtrac
 		Super::EvaluateCurveData(OutCurve, AnimExtractContext, bForceUseRawData);
 #endif
 	}
-	else if(IsCurveCompressedDataValid() && CompressedData.CurveCompressionCodec)
+	else if(PlatformCompressedData.IsCurveDataValid(this))
 	{
-		CSV_SCOPED_TIMING_STAT(Animation, EvaluateCurveData);
-		CompressedData.CurveCompressionCodec->DecompressCurves(CompressedData, OutCurve, AnimExtractContext.CurrentTime);
+		if (PlatformCompressedData.CurveCompressionCodec)
+		{
+			CSV_SCOPED_TIMING_STAT(Animation, EvaluateCurveData);
+			PlatformCompressedData.CurveCompressionCodec->DecompressCurves(PlatformCompressedData, OutCurve, AnimExtractContext.CurrentTime);
+		}
 	}
 }
 
 float UAnimSequence::EvaluateCurveData(FName CurveName, const FAnimExtractContext& AnimExtractContext, bool bForceUseRawData) const
+{
+	UE_COMPRESSED_DATA_READ_SCOPE(this);
+	return EvaluateCurveData_Lockless(CurveName, AnimExtractContext, bForceUseRawData);
+}
+
+float UAnimSequence::EvaluateCurveData_Lockless(FName CurveName, const FAnimExtractContext& AnimExtractContext, bool bForceUseRawData) const
 {
 	QUICK_SCOPE_CYCLE_COUNTER(EvaluateCurveDataByName);
 
-	const bool bEvaluateRawData = bUseRawDataOnly
-	  || bForceUseRawData || !IsCurveCompressedDataValid();
+	const FCompressedAnimSequence& PlatformCompressedData = GetPlatformCompressedData(AnimExtractContext);
+	const bool bEvaluateRawData = bForceUseRawData || !PlatformCompressedData.IsCurveDataValid(this);
 	check(!bForceUseRawData || CanEvaluateRawAnimationData());
 	if (CanEvaluateRawAnimationData() && bEvaluateRawData)
 	{
 		return Super::EvaluateCurveData(CurveName, AnimExtractContext, bForceUseRawData);
 	}
-	else if(IsCurveCompressedDataValid() && CompressedData.CurveCompressionCodec)
+
+	if(PlatformCompressedData.IsCurveDataValid(this))
 	{
-		return CompressedData.CurveCompressionCodec->DecompressCurve(CompressedData, CurveName, AnimExtractContext.CurrentTime);
+		if(PlatformCompressedData.CurveCompressionCodec)
+		{
+			return PlatformCompressedData.CurveCompressionCodec->DecompressCurve(PlatformCompressedData, CurveName, AnimExtractContext.CurrentTime);
+		}
 	}
 
 	return 0.f;
@@ -2842,22 +3164,25 @@ float UAnimSequence::EvaluateCurveData(FName CurveName, const FAnimExtractContex
 
 bool UAnimSequence::HasCurveData(FName CurveName, bool bForceUseRawData) const
 {
-	const bool bEvaluateRawData = bUseRawDataOnly
-		|| bForceUseRawData || !IsCurveCompressedDataValid();
+	const bool bEvaluateRawData = bForceUseRawData || !IsCurveCompressedDataValid();
 	check(!bForceUseRawData || CanEvaluateRawAnimationData());
 	if (CanEvaluateRawAnimationData() && bEvaluateRawData)
 	{
-		PRAGMA_DISABLE_DEPRECATION_WARNINGS
 		return Super::HasCurveData(CurveName, bForceUseRawData);
-		PRAGMA_ENABLE_DEPRECATION_WARNINGS
 	}
-	else if(IsCurveCompressedDataValid() && CompressedData.CurveCompressionCodec)
+
+	if(IsCurveCompressedDataValid())
 	{
-		for (const FAnimCompressedCurveIndexedName& IndexedCurveName : CompressedData.IndexedCurveNames)
+		UE_COMPRESSED_DATA_READ_SCOPE(this);
+		const FCompressedAnimSequence& PlatformCompressedData = GetCompressedData_Internal();
+		if(PlatformCompressedData.CurveCompressionCodec)
 		{
-			if (IndexedCurveName.CurveName == CurveName)
+			for (const FAnimCompressedCurveIndexedName& IndexedCurveName : PlatformCompressedData.IndexedCurveNames)
 			{
-				return true;
+				if (IndexedCurveName.CurveName == CurveName)
+				{
+					return true;
+				}
 			}
 		}
 	}
@@ -3226,11 +3551,45 @@ void UAnimSequence::ValidateCurrentPosition(const FMarkerSyncAnimPosition& Posit
 
 bool UAnimSequence::UseRawDataForPoseExtraction(const FBoneContainer& RequiredBones) const
 {
-	return CanEvaluateRawAnimationData() && (bUseRawDataOnly || (GetSkeletonVirtualBoneGuid() != GetSkeleton()->GetVirtualBoneGuid()) || RequiredBones.GetDisableRetargeting() || RequiredBones.ShouldUseRawData() ||
- #if WITH_EDITOR
-  	GForceRawData == 1 ||
- #endif // WITH_EDITOR
-	RequiredBones.ShouldUseSourceData());
+	UE_COMPRESSED_DATA_READ_SCOPE(this);
+	return CanEvaluateRawAnimationData() &&
+	(
+#if WITH_EDITOR
+		GForceRawData == 1 ||
+#endif // WITH_EDITOR
+		!GetCompressedData_Internal().IsValid(this) ||
+		(GetSkeletonVirtualBoneGuid() != GetSkeleton()->GetVirtualBoneGuid()) ||
+		RequiredBones.GetDisableRetargeting() ||
+		RequiredBones.ShouldUseRawData() ||
+		RequiredBones.ShouldUseSourceData()
+	);
+}
+
+bool UAnimSequence::ShouldUseRawDataForPoseExtraction(const FBoneContainer& RequiredBones, const FAnimExtractContext& ExtractionContext) const
+{
+	UE_COMPRESSED_DATA_READ_SCOPE(this);
+	return ShouldUseRawDataForPoseExtraction_Lockless(RequiredBones, ExtractionContext);
+}
+
+bool UAnimSequence::ShouldUseRawDataForPoseExtraction_Lockless(const FBoneContainer& RequiredBones,
+	const FAnimExtractContext& ExtractionContext) const
+{
+	return ShouldUseRawDataForPoseExtraction_Lockless(ExtractionContext) ||
+	RequiredBones.GetDisableRetargeting() ||
+	RequiredBones.ShouldUseRawData() ||
+	RequiredBones.ShouldUseSourceData();
+}
+
+bool UAnimSequence::ShouldUseRawDataForPoseExtraction_Lockless(const FAnimExtractContext& ExtractionContext) const
+{
+	return CanEvaluateRawAnimationData() &&
+	(
+#if WITH_EDITOR
+		GForceRawData == 1 ||
+		!GetPlatformCompressedData(ExtractionContext).IsValid(this) ||
+#endif // WITH_EDITOR
+		(GetSkeletonVirtualBoneGuid() != GetSkeleton()->GetVirtualBoneGuid())
+	);
 }
 
 void UAnimSequence::EvaluateAttributes(FAnimationPoseData& OutAnimationPoseData, const FAnimExtractContext& ExtractionContext, bool bUseRawData) const
@@ -4051,43 +4410,20 @@ bool UAnimSequence::GetEnableRootMotionSettingFromMontage() const
 
 bool UAnimSequence::IsCompressedDataValid() const
 {
-	// For bone compressed data, we don't check if we have a codec. It is valid to have no compressed data
-	// if we have no raw data. This can happen with sequences that only has curves.
-#if WITH_EDITOR
-	PRAGMA_DISABLE_DEPRECATION_WARNINGS
-	if (IsDataModelValid() && DataModelInterface->GetNumBoneTracks() == 0 && RawAnimationData.Num() == 0)
-	PRAGMA_ENABLE_DEPRECATION_WARNINGS
-	{
-		return true;
-	}
-#endif
+	UE_COMPRESSED_DATA_READ_SCOPE(this);
+	return GetCompressedData_Internal().IsValid(this);
+}
 
-	return CompressedData.CompressedDataStructure != nullptr || CompressedData.CompressedTrackToSkeletonMapTable.Num() == 0;
+bool UAnimSequence::IsBoneCompressedDataValid() const
+{
+	UE_COMPRESSED_DATA_READ_SCOPE(this);
+	return GetCompressedData_Internal().IsBoneDataValid(this);
 }
 
 bool UAnimSequence::IsCurveCompressedDataValid() const
 {
-#if WITH_EDITOR
-	// For curve compressed data, we don't check if we have a codec. It is valid to have no compressed data
-	// if we have no raw data. This can happen with sequences that only has bones.
-
-	const bool bNoCompressedCurveData = CompressedData.CompressedCurveByteStream.Num() == 0;
-	const bool bHasRawCurveData = IsDataModelValid() && DataModelInterface->GetNumberOfFloatCurves() != 0;
-
-	if ( bNoCompressedCurveData && bHasRawCurveData )
-	{
-		// No compressed data but we have raw data
-		if (!IsValidAdditive())
-		{
-			return false;
-		}
-
-		// Additive sequences can have raw curves that all end up being 0.0 (e.g. they 100% match the base sequence curves)
-		// in which case there will be no compressed curve data.
-	}
-#endif
-
-	return true;
+	UE_COMPRESSED_DATA_READ_SCOPE(this);
+	return GetCompressedData_Internal().IsCurveDataValid(this);
 }
 
 #if WITH_EDITOR
@@ -4216,7 +4552,6 @@ void UAnimSequence::OnModelModified(const EAnimDataModelNotifyType& NotifyType,
 		PRAGMA_DISABLE_DEPRECATION_WARNINGS
 		RawDataGuid = bForceNewRawDataGuid ? FGuid::NewGuid() : Model->GenerateGuid();
 		PRAGMA_ENABLE_DEPRECATION_WARNINGS
-        bUseRawDataOnly = true;
 
 		FlagDependentAnimationsAsRawDataOnly();
 		UpdateDependentStreamingAnimations();
@@ -4228,7 +4563,7 @@ void UAnimSequence::OnModelModified(const EAnimDataModelNotifyType& NotifyType,
 		{
 			CalculateNumberOfSampledKeys();
 			UpdateRawDataGuid(RegenerateGUID);
-			ClearAllCachedCookedPlatformData();
+			ClearAllCompressionData();
 			if(!bWasReset)
 			{
 				BeginCacheDerivedDataForCurrentPlatform();
@@ -4360,7 +4695,7 @@ void UAnimSequence::OnModelModified(const EAnimDataModelNotifyType& NotifyType,
 				{
 					const bool bWasModelReset = NotifyCollector.Contains(EAnimDataModelNotifyType::Reset);
 					UpdateRawDataGuid(bWasModelReset ? GenerateNewGUID : RegenerateGUID);
-					CompressedData.ClearCompressedCurveData();
+					ClearAllCompressionData();
 					HandleTrackDataChanged(bWasModelReset);
 				}
 			}
@@ -4379,7 +4714,7 @@ void UAnimSequence::OnModelModified(const EAnimDataModelNotifyType& NotifyType,
 			{
 				CalculateNumberOfSampledKeys();
 				UpdateRawDataGuid(GenerateNewGUID);
-				ClearAllCachedCookedPlatformData();
+				ClearAllCompressionData();
 				BeginCacheDerivedDataForCurrentPlatform();
 			}
 
@@ -4400,11 +4735,10 @@ void UAnimSequence::OnModelModified(const EAnimDataModelNotifyType& NotifyType,
 		case EAnimDataModelNotifyType::CurveFlagsChanged:
 		case EAnimDataModelNotifyType::CurveScaled:
 		{
-			CompressedData.ClearCompressedCurveData();
-				
+			ClearAllCompressionData();
+
 			if (NotifyCollector.IsNotWithinBracket())
 			{						
-				ClearAllCachedCookedPlatformData();
 				UpdateRawDataGuid(RegenerateGUID);
 				BeginCacheDerivedDataForCurrentPlatform();
 			}
@@ -4484,8 +4818,32 @@ void UAnimSequence::CalculateNumberOfSampledKeys()
 	}
 }
 
-FIoHash UAnimSequence::CreateDerivedDataKeyHash(const ITargetPlatform* TargetPlatform)
+FIoHash UAnimSequence::GetDerivedDataKeyHash(const ITargetPlatform* TargetPlatform) const
+{
+	thread_local uint32 CachedCurrentPlatformHash = INDEX_NONE;
+	const ITargetPlatform* RunningPlatform = GetTargetPlatformManagerRef().GetRunningTargetPlatform();
+	if (CachedCurrentPlatformHash == INDEX_NONE)
+	{
+		CachedCurrentPlatformHash = GetTypeHash(RunningPlatform);
+	}
+	
+	const uint32 PlatformHash = TargetPlatform ? GetTypeHash(TargetPlatform) : CachedCurrentPlatformHash;
+
+	const FIoHash* HashPtr = nullptr;
+	{
+		FReadScopeLock ScopeLock(HashCacheLock);
+		HashPtr = PlatformHashToKeyHash.Find(PlatformHash);
+	}
+		
+	return HashPtr ? *HashPtr : CreateDerivedDataKeyHash(TargetPlatform ? TargetPlatform : RunningPlatform);
+}
+
+FIoHash UAnimSequence::CreateDerivedDataKeyHash(const ITargetPlatform* TargetPlatform) const
 {
+	check(TargetPlatform);
+	
+	const uint32 PlatformHash = GetTypeHash(TargetPlatform);
+	
 	FString Ret = CreateDerivedDataKeyString(TargetPlatform);
 	if (Ret.IsEmpty())
 	{
@@ -4495,11 +4853,27 @@ FIoHash UAnimSequence::CreateDerivedDataKeyHash(const ITargetPlatform* TargetPla
 	// New animation DDC key format; use just the hash of the complete dependencies string.
 	FMemoryHasherBlake3 Writer;
 	Writer << Ret;
-	return Writer.Finalize();
-}
+	const FIoHash Hash = Writer.Finalize();
+
+	{
+		FWriteScopeLock ScopeLock(HashCacheLock);
+		// Verifying key (changing)
+		if (FIoHash* StoredHash = PlatformHashToKeyHash.Find(PlatformHash))
+		{
+			if (*StoredHash != Hash)
+			{
+				checkf(!IsRunningCookCommandlet(), TEXT("Hash for %s has changed from %s to %s"), *GetPathName(), *LexToString(*StoredHash), *LexToString(Hash));
+			}
+		}
+		
+		// Store platform-to-hash pair 
+		PlatformHashToKeyHash.Add(PlatformHash, Hash);
+	}
 
+	return Hash;
+}
 
-FString UAnimSequence::CreateDerivedDataKeyString(const ITargetPlatform* TargetPlatform)
+FString UAnimSequence::CreateDerivedDataKeyString(const ITargetPlatform* TargetPlatform) const
 {
 	const USkeleton* CurrentSkeleton = GetSkeleton();
 	if (CurrentSkeleton == nullptr)
@@ -4508,12 +4882,17 @@ FString UAnimSequence::CreateDerivedDataKeyString(const ITargetPlatform* TargetP
 		return FString();
 	}
 
+	if (BoneCompressionSettings == nullptr || CurveCompressionSettings == nullptr)
+	{
+		UE_LOG(LogAnimation, Warning, TEXT("Animation Compression request for %s failed, BoneCompressionSettings == nullptr || CurveCompressionSettings == nullptr."), *GetName());
+		return FString();
+	}
 
 	const bool bIsValidAdditive = IsValidAdditive();
 	const char AdditiveType = bIsValidAdditive ? NibbleToTChar(AdditiveAnimType) : '0';
 	const char RefType = bIsValidAdditive ? NibbleToTChar(RefPoseType) : '0';
 
-	bool bPerformFrameStripping  = false;
+	bool bPerformFrameStripping = false;
 	bool bPerformFrameStrippingOnOddNumberedFrames = false;
 	{
 		if (UDeviceProfile* DeviceProfile = UDeviceProfileManager::Get().FindProfile(TargetPlatform->IniPlatformName()))
@@ -4536,7 +4915,8 @@ FString UAnimSequence::CreateDerivedDataKeyString(const ITargetPlatform* TargetP
 
 	FArcToHexString ArcToHexString;
 
-	ArcToHexString.Ar << CompressionErrorThresholdScale;
+	float TempThreshold = CompressionErrorThresholdScale;
+	ArcToHexString.Ar << TempThreshold;
 	ArcToHexString.Ar << bPerformFrameStripping;
 	BoneCompressionSettings->PopulateDDCKey(UE::Anim::Compression::FAnimDDCKeyArgs(*this, TargetPlatform), ArcToHexString.Ar);
 	CurveCompressionSettings->PopulateDDCKey(ArcToHexString.Ar);
@@ -4583,10 +4963,32 @@ FString UAnimSequence::CreateDerivedDataKeyString(const ITargetPlatform* TargetP
 	return Ret;
 }
 
+void UAnimSequence::ValidateCompressionSettings()
+{
+	// Ensure that there are valid compression settings
+	if (BoneCompressionSettings == nullptr || !BoneCompressionSettings->AreSettingsValid())
+	{
+		BoneCompressionSettings = FAnimationUtils::GetDefaultAnimationBoneCompressionSettings();
+	}	
+	if (CurveCompressionSettings == nullptr || !CurveCompressionSettings->AreSettingsValid())
+	{
+		CurveCompressionSettings = FAnimationUtils::GetDefaultAnimationCurveCompressionSettings();
+	}
+	if (VariableFrameStrippingSettings == nullptr)
+	{
+		VariableFrameStrippingSettings = FAnimationUtils::GetDefaultVariableFrameStrippingSettings();
+	}
+}
+
+bool UAnimSequence::CanBeCompressed() const
+{
+	const UPackage* Package = GetPackage();
+	return Package && (Package->GetHasBeenEndLoaded() || Package->GetLoadedPath().IsEmpty());
+}
+
 FIoHash UAnimSequence::BeginCacheDerivedData(const ITargetPlatform* TargetPlatform)
 {
 	QUICK_SCOPE_CYCLE_COUNTER(STAT_AnimSequence_BeginCacheDerivedData);
-	COOK_STAT(auto Timer = UE::Anim::AnimSequenceCookStats::UsageStats.TimeSyncWork());
 
 	check(TargetPlatform != nullptr);
 	check(IsInGameThread());
@@ -4604,48 +5006,38 @@ FIoHash UAnimSequence::BeginCacheDerivedData(const ITargetPlatform* TargetPlatfo
 	// Once a compression request finishes, it might modify the compression settings below and the key hash as well
 	UE::Anim::FAnimSequenceCompilingManager::Get().FinishCompilation({ this, RefPoseSeq });
 
-	// Ensure that there are valid compression settings
-	if (BoneCompressionSettings == nullptr || !BoneCompressionSettings->AreSettingsValid())
-	{
-		BoneCompressionSettings = FAnimationUtils::GetDefaultAnimationBoneCompressionSettings();
-	}
-	
-	if (CurveCompressionSettings == nullptr || !CurveCompressionSettings->AreSettingsValid())
-	{
-		CurveCompressionSettings = FAnimationUtils::GetDefaultAnimationCurveCompressionSettings();
-	}
-	if (VariableFrameStrippingSettings == nullptr)
+	if (!CanBeCompressed())
 	{
-		VariableFrameStrippingSettings = FAnimationUtils::GetDefaultVariableFrameStrippingSettings();
+		return FIoHash::Zero;
 	}
 
+	ValidateCompressionSettings();
+
 	// Make sure all our required dependencies are loaded, we need them to compute the KeyHash
 	FAnimationUtils::EnsureAnimSequenceLoaded(*this);
 		
 	const FIoHash KeyHash = CreateDerivedDataKeyHash(TargetPlatform);
-	// Early out if not valid, equal to running platform hash or any previously cached hashes
-	if (KeyHash.IsZero() || (DataKeyHash == KeyHash && (CompressedData.IsValid(this) || CacheTasksByKeyHash.Contains(KeyHash))) || DataByPlatformKeyHash.Contains(KeyHash))
-	{
-		return KeyHash;
-	}
 
+	UE_COMPRESSED_DATA_WRITE_SCOPE(this);
 	FCompressedAnimSequence* TargetData = nullptr;
-	if (TargetPlatform->IsRunningPlatform())
 	{
-		DataKeyHash = KeyHash;
-		TargetData = &CompressedData;
-	}
-	else
-    {
-    	TargetData = DataByPlatformKeyHash.Emplace(KeyHash, MakeUnique<FCompressedAnimSequence>()).Get();
-    }
+		checkf(!CacheTasksByKeyHash.Contains(KeyHash), TEXT("[%s] Hash %s still/already has task in-flight"), *GetPathName(), *LexToString(KeyHash));
+		// Early out if not valid, has already been cached, or has an inflight task running 
+		if (KeyHash.IsZero() || CacheTasksByKeyHash.Contains(KeyHash) || DataByPlatformKeyHash.Contains(KeyHash))
+		{
+			return KeyHash;
+		}
 
-	// Mark sequence to only sample raw animation data when recompressing for running platform (or equivalent hash)
-	if (KeyHash == DataKeyHash)
-	{
-		bUseRawDataOnly = true;
+		if (RequiresResidency(KeyHash))
+		{
+			UE_LOG(LogAnimation, Warning, TEXT("Animation Compression request for %s was blocked, Residency."), *GetName());
+			return FIoHash::Zero;
+		}
+
+		TargetData = DataByPlatformKeyHash.Emplace(KeyHash, MakeUnique<FCompressedAnimSequence>()).Get();
 	}
-	
+
+	// Reset the target compressed data, to mark it invalid
 	check(TargetData);
 	TargetData->Reset();
 
@@ -4671,16 +5063,24 @@ FIoHash UAnimSequence::BeginCacheDerivedData(const ITargetPlatform* TargetPlatfo
 	}
 
 	// Always keep AnimatedBoneAttributes in lock-step when requesting compression
-	{		
+	{
 		SynchronousAnimatedBoneAttributesCompression();
 	}
-	
-	{		
+
+	if (!CacheTasksByKeyHash.Contains(KeyHash))
+	{
 		// Data does not exist, need to build it.
 		const FCompressibleAnimRef CompressibleData = MakeShared<FCompressibleAnimData, ESPMode::ThreadSafe>(this, bPerformFrameStripping, TargetPlatform);
+		
+		COOK_STAT(auto Timer = UE::Anim::AnimSequenceCookStats::UsageStats.TimeSyncWork());
 		check(CacheTasksByKeyHash.Find(KeyHash) == nullptr);
 		CacheTasksByKeyHash.Emplace(KeyHash, MakePimpl<UE::Anim::FAnimationSequenceAsyncCacheTask>(KeyHash, CompressibleData, TargetData, *this, TargetPlatform));		
+		COOK_STAT(Timer.TrackCyclesOnly());
 	}
+	else
+    {
+		check(false);
+    }
 
 	// The compiling manager provides throttling, notification manager, etc... for the asset being built.
 	UE::Anim::FAnimSequenceCompilingManager::Get().AddAnimSequences({this});
@@ -4695,6 +5095,7 @@ bool UAnimSequence::PollCacheDerivedData(const FIoHash& KeyHash) const
 		return true;
 	}
 
+	UE_COMPRESSED_DATA_READ_SCOPE(this);
 	if (const TPimplPtr<UE::Anim::FAnimationSequenceAsyncCacheTask>* Task = CacheTasksByKeyHash.Find(KeyHash))
 	{
 		return (*Task)->Poll();
@@ -4717,7 +5118,16 @@ FCompressedAnimSequence& UAnimSequence::CacheDerivedData(const ITargetPlatform*
 {
 	const FIoHash KeyHash = BeginCacheDerivedData(TargetPlatform);
 	EndCacheDerivedData(KeyHash);
-	return DataKeyHash == KeyHash ? CompressedData : *DataByPlatformKeyHash[KeyHash];
+
+	UE_COMPRESSED_DATA_READ_SCOPE(this);
+	return *DataByPlatformKeyHash[KeyHash];
+}
+
+void UAnimSequence::CacheDerivedDataForPlatform(const ITargetPlatform* TargetPlatform)
+{
+	const FIoHash KeyHash = BeginCacheDerivedData(TargetPlatform);
+	EndCacheDerivedData(KeyHash);
+	check(HasCompressedDataForHash(KeyHash));
 }
 
 FFrameRate UAnimSequence::GetTargetSamplingFrameRate(const ITargetPlatform* InPlatform) const
@@ -4737,17 +5147,125 @@ void UAnimSequence::CacheDerivedDataForCurrentPlatform()
 {
 	if (const ITargetPlatform* RunningPlatform = GetTargetPlatformManagerRef().GetRunningTargetPlatform())
 	{
-		CacheDerivedData(RunningPlatform);
+		CacheDerivedDataForPlatform(RunningPlatform);
 	}
 }
 				
 bool UAnimSequence::IsCompiling() const
 {
+	UE_COMPRESSED_DATA_READ_SCOPE(this);	
 	return CacheTasksByKeyHash.Num() > 0;
 }
-					
+
+void UAnimSequence::ClearAllCompressionData()
+{
+	UE_COMPRESSED_DATA_WRITE_SCOPE(this);
+	{
+		// Delete any cache tasks first because the destructor will cancel the cache and build tasks,
+		// and drop their pointers to the data.
+		
+		for (auto It = CacheTasksByKeyHash.CreateIterator(); It; ++It)
+		{	
+			if(!RequiresResidency(It->Key))
+			{
+				It.RemoveCurrent();
+			}
+		}
+		
+		for (auto It = DataByPlatformKeyHash.CreateIterator(); It; ++It)
+		{
+			if(!RequiresResidency(It->Key))
+			{
+				It.RemoveCurrent();
+			}
+		}
+	}
+}
+
+void UAnimSequence::ClearCompressionData(const FIoHash& InKeyHash)
+{
+	check(!InKeyHash.IsZero());
+	
+	// Check if there are any references left, otherwise clear out
+	if(!RequiresResidency(InKeyHash))
+	{
+		UE_COMPRESSED_DATA_WRITE_SCOPE(this);
+		check(DataByPlatformKeyHash.Contains(InKeyHash));
+		DataByPlatformKeyHash.Remove(InKeyHash);
+	}
+}
+
+bool UAnimSequence::RequiresResidency(const FIoHash& InKeyHash) const
+{
+	FReadScopeLock ScopeLock(ResidencyLock);
+	return ResidencyReferencerHashes.Contains(InKeyHash);
+}
+
+bool UAnimSequence::HasResidency(const ITargetPlatform* InPlatform) const
+{
+	const FIoHash KeyHash = GetDerivedDataKeyHash(InPlatform);
+	FReadScopeLock ScopeLock(ResidencyLock);
+	return ResidencyReferencerHashes.Contains(KeyHash);
+}
+
+void UAnimSequence::RequestResidency(const ITargetPlatform* InPlatform, uint32 InReferencerHash)
+{
+	check(CanBeCompressed() && !HasAnyFlags(RF_NeedPostLoad));
+	const FIoHash KeyHash = GetDerivedDataKeyHash(InPlatform);
+
+	TArray<uint32> ReferencerHashes;
+	{
+		FReadScopeLock ScopeLock(ResidencyLock);
+		ResidencyReferencerHashes.MultiFind(KeyHash, ReferencerHashes);
+	}
+	
+	// First request for residency
+	if (ReferencerHashes.Num() == 0)
+	{
+		UE_CLOG(false, LogAnimationCompression, Display, TEXT("RequestResidency platform data [KeyHash]: %s %s %x %s"), *InPlatform->DisplayName().ToString(), *LexToString(KeyHash), InReferencerHash, *GetPathName());
+		BeginCacheForCookedPlatformData(InPlatform);
+	}
+
+	if (!ReferencerHashes.Contains(InReferencerHash))
+	{
+		FWriteScopeLock ScopeLock(ResidencyLock);
+		ResidencyReferencerHashes.Add(KeyHash, InReferencerHash);
+	}
+}
+
+void UAnimSequence::ReleaseResidency(const ITargetPlatform* InPlatform, uint32 InReferencerHash)
+{
+	const FIoHash KeyHash = GetDerivedDataKeyHash(InPlatform);
+	TArray<uint32> ReferencerHashes;
+	{
+		FWriteScopeLock ScopeLock(ResidencyLock);
+	
+		check(ResidencyReferencerHashes.Remove(KeyHash, InReferencerHash) == 1);
+
+		// Check if there are any references left
+		ResidencyReferencerHashes.MultiFind(KeyHash, ReferencerHashes);
+	
+		UE_CLOG(false, LogAnimationCompression, Display, TEXT("ReleaseResidency platform data [KeyHash]: %s %s %x %s"), *InPlatform->DisplayName().ToString(), *LexToString(KeyHash), InReferencerHash, *GetPathName());
+	}
+	
+	if (ReferencerHashes.Num() == 0)
+	{
+		// If previously this anim sequence was marked as never to be cooked again clear out any compressed data without residency references
+		if(bShouldClearCompressedData)
+		{
+			ClearCompressionData(KeyHash);
+		}
+	}
+}
+bool UAnimSequence::HasCompressedDataForPlatform(const ITargetPlatform* InPlatform) const
+{
+	const FIoHash KeyHash = GetDerivedDataKeyHash(InPlatform);
+	return HasCompressedDataForHash(KeyHash);
+}
+
 bool UAnimSequence::TryCancelAsyncTasks()
 {
+	UE_COMPRESSED_DATA_READ_SCOPE(this);
 	for (auto It = CacheTasksByKeyHash.CreateIterator(); It; ++It)
 	{
 		It->Value->Cancel();
@@ -4758,6 +5276,7 @@ bool UAnimSequence::TryCancelAsyncTasks()
 
 bool UAnimSequence::WaitForAsyncTasks(float TimeLimitSeconds)
 {
+	UE_COMPRESSED_DATA_READ_SCOPE(this);
 	double StartTimeSeconds = FPlatformTime::Seconds();
 	for (auto& Pair : CacheTasksByKeyHash)
 	{
@@ -4773,6 +5292,7 @@ bool UAnimSequence::WaitForAsyncTasks(float TimeLimitSeconds)
 
 bool UAnimSequence::IsAsyncTaskComplete() const
 {
+	UE_COMPRESSED_DATA_READ_SCOPE(this);
 	bool bAllFinished = true;
 	for (auto& Pair : CacheTasksByKeyHash)
 	{
@@ -4783,7 +5303,11 @@ bool UAnimSequence::IsAsyncTaskComplete() const
 }
 
 void UAnimSequence::FinishAsyncTasks()
-{	
+{
+	const ITargetPlatform* RunningPlatform = GetTargetPlatformManagerRef().GetRunningTargetPlatform();
+	const FIoHash CurrentPlatformHash = CreateDerivedDataKeyHash(RunningPlatform);
+
+	UE_COMPRESSED_DATA_WRITE_SCOPE(this);
 	const bool bHasInflightTasks = CacheTasksByKeyHash.Num() > 0;
 	if (bHasInflightTasks)
 	{
@@ -4795,17 +5319,11 @@ void UAnimSequence::FinishAsyncTasks()
 			It->Value->Wait();
 
 			FCompressedAnimSequence* TaskData = It->Value->GetTargetCompressedData();
-			const bool bCompressedRunningPlatform = (TaskData == &CompressedData);
-			auto ResetData = [this, TaskData, It, bCompressedRunningPlatform]()
+			const bool bCompressedRunningPlatform = It->Key == CurrentPlatformHash;
+			auto ResetData = [this, TaskData, It]()
 			{
 				TaskData->Reset();
 				DataByPlatformKeyHash.Remove(It->Key);
-
-				// Reset running platform hash (if it got cancelled)
-				if (bCompressedRunningPlatform)
-				{
-					DataKeyHash = FIoHash::Zero;
-				}
 			};
 
 			if (It->Value.Get()->WasCancelled())
@@ -4827,9 +5345,6 @@ void UAnimSequence::FinishAsyncTasks()
 					{
 						FAssetRegistryModule& AssetRegistryModule = FModuleManager::LoadModuleChecked<FAssetRegistryModule>(TEXT("AssetRegistry"));
 						AssetRegistryModule.Get().AssetTagsFinalized(*this);
-					
-						check(bUseRawDataOnly);
-						bUseRawDataOnly = false;
 					}
 				
 					check(TaskData->IsValid(this, true));
@@ -4837,7 +5352,7 @@ void UAnimSequence::FinishAsyncTasks()
 				else
 				{
 					// Failed to compress
-					UE_LOG(LogAnimation, Display, TEXT("Failed to finish async Animation Compression task for %s, as the generated data is not valid."), *GetName());					
+					UE_LOG(LogAnimationCompression, Display, TEXT("Failed to finish async Animation Compression task for %s, as the generated data is not valid."), *GetName());
 					ResetData();
 				}
 			}
@@ -4849,6 +5364,7 @@ void UAnimSequence::FinishAsyncTasks()
 
 void UAnimSequence::Reschedule(FQueuedThreadPool* InThreadPool, EQueuedWorkPriority InPriority)
 {
+	UE_COMPRESSED_DATA_READ_SCOPE(this);
 	for (auto& Pair : CacheTasksByKeyHash)
 	{
 		Pair.Value->Reschedule(InThreadPool, InPriority);
@@ -4888,9 +5404,11 @@ void GatherAnimSequenceStats(FOutputDevice& Ar)
 
 	for( TObjectIterator<UAnimSequence> It; It; ++It )
 	{
-		UAnimSequence* Seq = *It;
+		const UAnimSequence* Seq = *It;
 
-		if (Seq->CompressedData.CompressedDataStructure == nullptr || !Seq->CompressedData.BoneCompressionCodec->IsA<UAnimCompress>())
+		UAnimSequence::FScopedCompressedAnimSequence ScopedCompressedData = Seq->GetCompressedData();
+		const FCompressedAnimSequence& CompressedAnimSequence = ScopedCompressedData.Get();
+		if (CompressedAnimSequence.CompressedDataStructure == nullptr || !CompressedAnimSequence.BoneCompressionCodec->IsA<UAnimCompress>())
 		{
 			continue;	// Custom codec we know nothing about, skip it
 		}
@@ -4909,7 +5427,7 @@ void GatherAnimSequenceStats(FOutputDevice& Ar)
 		int32 NumRotTracksWithOneKey = 0;
 		int32 NumScaleTracksWithOneKey = 0;
 
-		const FUECompressedAnimData& AnimData = static_cast<FUECompressedAnimData&>(*Seq->CompressedData.CompressedDataStructure);
+		const FUECompressedAnimData& AnimData = static_cast<const FUECompressedAnimData&>(*CompressedAnimSequence.CompressedDataStructure);
 
 		AnimationFormat_GetStats(
 			AnimData,
@@ -4948,7 +5466,7 @@ void GatherAnimSequenceStats(FOutputDevice& Ar)
 			NumTransTracksWithOneKey, NumRotTracksWithOneKey, NumScaleTracksWithOneKey,
 			TotalNumTransKeys, TotalNumRotKeys, TotalNumScaleKeys,
 			*FAnimationUtils::GetAnimationKeyFormatString(AnimData.KeyEncodingFormat),
-			(int32)Seq->GetResourceSizeBytes(EResourceSizeMode::EstimatedTotal) );
+			(int32)It->GetResourceSizeBytes(EResourceSizeMode::EstimatedTotal) );
 	}
 	Ar.Logf( TEXT("======================================================================") );
 	Ar.Logf( TEXT("Total Num Tracks: %i trans, %i rot, %i scale, %i trans1, %i rot1, %i scale1"), GlobalNumTransTracks, GlobalNumRotTracks, GlobalNumScaleTracks, GlobalNumTransTracksWithOneKey, GlobalNumRotTracksWithOneKey, GlobalNumScaleTracksWithOneKey  );
diff --git a/Engine/Source/Runtime/Engine/Private/Animation/AnimationCompressionDerivedData.cpp b/Engine/Source/Runtime/Engine/Private/Animation/AnimationCompressionDerivedData.cpp
--- a/Engine/Source/Runtime/Engine/Private/Animation/AnimationCompressionDerivedData.cpp
+++ b/Engine/Source/Runtime/Engine/Private/Animation/AnimationCompressionDerivedData.cpp
@@ -136,6 +136,7 @@ void FAnimationSequenceAsyncCacheTask::EndCache(UE::DerivedData::FCacheGetValueR
 				{
 					COOK_STAT(auto Timer = AnimSequenceCookStats::UsageStats.TimeSyncWork());
 
+					UE_COMPRESSED_DATA_WRITE_SCOPE(AnimSequence);
 					const FSharedBuffer RecordData = Value.GetData().Decompress();
 					FMemoryReaderView Ar(RecordData, /*bIsPersistent*/ true);
 					CompressedData->SerializeCompressedData(Ar, true, AnimSequence, AnimSequence->GetSkeleton(), CompressibleAnimPtr->BoneCompressionSettings, CompressibleAnimPtr->CurveCompressionSettings);
@@ -266,9 +267,10 @@ void FAnimationSequenceAsyncCacheTask::LaunchCompressionTask(const UE::FSharedSt
 				}
 				else
 				{
+					UE_COMPRESSED_DATA_WRITE_SCOPE(AnimSequence);					
 					TArray64<uint8> RecordData;
 					FMemoryWriter64 Ar(RecordData, /*bIsPersistent*/ true);
-					CompressedData->SerializeCompressedData(Ar, true, nullptr, nullptr, CompressibleAnimPtr->BoneCompressionSettings, CompressibleAnimPtr->CurveCompressionSettings);
+					CompressedData->SerializeCompressedData(Ar, true, const_cast<UAnimSequence*>(AnimSequence), nullptr, CompressibleAnimPtr->BoneCompressionSettings, CompressibleAnimPtr->CurveCompressionSettings);
 					UE_LOG(LogAnimationCompression, Display, TEXT("Storing compressed animation data for %s, at %s/%s"), *Name, *FString(Key.Bucket.ToString()), *LexToString(Key.Hash));
 					UE::DerivedData::GetCache().PutValue({ {Name, Key, UE::DerivedData::FValue::Compress(MakeSharedBufferFromArray(MoveTemp(RecordData)))} }, Owner);
 
diff --git a/Engine/Source/Runtime/Engine/Private/Animation/AnimationUtils.cpp b/Engine/Source/Runtime/Engine/Private/Animation/AnimationUtils.cpp
--- a/Engine/Source/Runtime/Engine/Private/Animation/AnimationUtils.cpp
+++ b/Engine/Source/Runtime/Engine/Private/Animation/AnimationUtils.cpp
@@ -857,11 +857,13 @@ static void EnsureDependenciesAreLoaded(UObject* Object)
 
 	if (Object->HasAnyFlags(RF_NeedLoad))
 	{
+		UE_LOG(LogAnimationCompression, Warning, TEXT("EnsureDependenciesAreLoaded: [Preload] %s"), *Object->GetPathName());
 		Object->GetLinker()->Preload(Object);
 	}
 
 	if (Object->HasAnyFlags(RF_NeedPostLoad))
 	{
+		UE_LOG(LogAnimationCompression, Warning, TEXT("EnsureDependenciesAreLoaded: [ConditionalPostLoad] %s"), *Object->GetPathName());
 		Object->ConditionalPostLoad();
 	}
 
@@ -872,11 +874,13 @@ static void EnsureDependenciesAreLoaded(UObject* Object)
 	{
 		if (Dependency->HasAnyFlags(RF_NeedLoad))
 		{
+			UE_LOG(LogAnimationCompression, Warning, TEXT("EnsureDependenciesAreLoaded: [Preload Dependency] %s"), *Dependency->GetPathName());
 			Dependency->GetLinker()->Preload(Dependency);
 		}
 
 		if (Dependency->HasAnyFlags(RF_NeedPostLoad))
 		{
+			UE_LOG(LogAnimationCompression, Warning, TEXT("EnsureDependenciesAreLoaded: [ConditionalPostLoad Dependency] %s"), *Dependency->GetPathName());
 			Dependency->ConditionalPostLoad();
 		}
 	}
diff --git a/Engine/Source/Runtime/Engine/Private/Animation/Skeleton.cpp b/Engine/Source/Runtime/Engine/Private/Animation/Skeleton.cpp
--- a/Engine/Source/Runtime/Engine/Private/Animation/Skeleton.cpp
+++ b/Engine/Source/Runtime/Engine/Private/Animation/Skeleton.cpp
@@ -1238,7 +1238,9 @@ int32 USkeleton::GetRawAnimationTrackIndex(const int32 InSkeletonBoneIndex, cons
 		
 		return BoneTrackNames.IndexOfByKey(BoneName);
 #else
-		return InAnimSeq->GetCompressedTrackToSkeletonMapTable().IndexOfByPredicate([InSkeletonBoneIndex](const FTrackToSkeletonMap& Mapping)
+		UAnimSequence::FScopedCompressedAnimSequence CompressedAnimSequence = InAnimSeq->GetCompressedData();
+		const TArray<FTrackToSkeletonMap>& BoneMappings = CompressedAnimSequence.Get().CompressedTrackToSkeletonMapTable;
+		return BoneMappings.IndexOfByPredicate([InSkeletonBoneIndex](const FTrackToSkeletonMap& Mapping)
 			{
 				return Mapping.BoneTreeIndex == InSkeletonBoneIndex;
 			});
diff --git a/Engine/Source/Runtime/Engine/Private/UnrealEngine.cpp b/Engine/Source/Runtime/Engine/Private/UnrealEngine.cpp
--- a/Engine/Source/Runtime/Engine/Private/UnrealEngine.cpp
+++ b/Engine/Source/Runtime/Engine/Private/UnrealEngine.cpp
@@ -7376,16 +7376,17 @@ bool UEngine::HandleListAnimsCommand(const TCHAR* Cmd, FOutputDevice& Ar)
 		int32 NumCurves = 0;
 		FString DebugString;
 
-		if (UAnimSequence* AnimSeq = Cast<UAnimSequence>(AnimAsset))
+		if (const UAnimSequence* AnimSeq = Cast<UAnimSequence>(AnimAsset))
 		{
 			NumKeys = AnimSeq->GetNumberOfSampledKeys();
 			SequenceLength = AnimSeq->GetPlayLength();
 			RateScale = AnimSeq->RateScale;
 			NumCurves = AnimSeq->GetCurveData().FloatCurves.Num();
 
-			if (AnimSeq->CompressedData.CompressedDataStructure != nullptr)
+			UAnimSequence::FScopedCompressedAnimSequence CompressedAnimSequence = AnimSeq->GetCompressedData();
+			if (CompressedAnimSequence.Get().CompressedDataStructure != nullptr)
 			{
-				DebugString = AnimSeq->CompressedData.CompressedDataStructure->GetDebugString();
+				DebugString = CompressedAnimSequence.Get().CompressedDataStructure->GetDebugString();
 			}
 			else
 			{
diff --git a/Engine/Source/Runtime/MovieSceneTracks/Private/Sections/MovieSceneSkeletalAnimationSection.cpp b/Engine/Source/Runtime/MovieSceneTracks/Private/Sections/MovieSceneSkeletalAnimationSection.cpp
--- a/Engine/Source/Runtime/MovieSceneTracks/Private/Sections/MovieSceneSkeletalAnimationSection.cpp
+++ b/Engine/Source/Runtime/MovieSceneTracks/Private/Sections/MovieSceneSkeletalAnimationSection.cpp
@@ -635,7 +635,7 @@ int32 UMovieSceneSkeletalAnimationSection::SetBoneIndexForRootMotionCalculations
 		TempRootBoneIndex.Reset();
 		return 0;
 	}
-	else if (UAnimSequence* AnimSequence = Cast<UAnimSequence>(Params.Animation))
+	else if (const UAnimSequence* AnimSequence = Cast<UAnimSequence>(Params.Animation))
 	{
 		if (TempRootBoneIndex.IsSet() == false || TempRootBoneIndex.GetValue() == INDEX_NONE)
 		{
@@ -674,7 +674,8 @@ int32 UMovieSceneSkeletalAnimationSection::SetBoneIndexForRootMotionCalculations
 				}
 			}
 #else
-			const TArray<FTrackToSkeletonMap>& BoneMappings = AnimSequence->GetCompressedTrackToSkeletonMapTable();
+			UAnimSequence::FScopedCompressedAnimSequence CompressedData = AnimSequence->GetCompressedData();
+			const TArray<FTrackToSkeletonMap>& BoneMappings = CompressedData.Get().CompressedTrackToSkeletonMapTable;
 			const int32 NumTracks = BoneMappings.Num();
 			for (int32 TrackIndex = 0; TrackIndex < NumTracks; ++TrackIndex)
 			{
