Iris - Avoid tempororary TSharedPtr creation.

diff --git a/Engine/Source/Runtime/Engine/Private/GameFramework/UniqueNetIdReplNetSerializer.cpp b/Engine/Source/Runtime/Engine/Private/GameFramework/UniqueNetIdReplNetSerializer.cpp
--- a/Engine/Source/Runtime/Engine/Private/GameFramework/UniqueNetIdReplNetSerializer.cpp
+++ b/Engine/Source/Runtime/Engine/Private/GameFramework/UniqueNetIdReplNetSerializer.cpp
@@ -223,7 +223,7 @@ void FUniqueNetIdReplNetSerializer::Quantize(FNetSerializationContext& Context,
 		return;
 	}
 
-	const FString& StringID = SourceValue->ToString();
+	const FString& StringID = SourceValue.ToString();
 	// This is compatible with FUniqueNetIdRepl::MakeReplicationData() which treats all empty strings this way, regardless of the ID type.
 	if (StringID.IsEmpty())
 	{
