Packed level actors that reference a level with streaming enabled can be edited in-place regardless of the IsEditInPlaceStreamingEnabled editor setting disabled.

diff --git a/Engine/Source/Runtime/Engine/Private/LevelInstance/LevelInstanceSubsystem.cpp b/Engine/Source/Runtime/Engine/Private/LevelInstance/LevelInstanceSubsystem.cpp
--- a/Engine/Source/Runtime/Engine/Private/LevelInstance/LevelInstanceSubsystem.cpp
+++ b/Engine/Source/Runtime/Engine/Private/LevelInstance/LevelInstanceSubsystem.cpp
@@ -2262,42 +2262,46 @@ bool ULevelInstanceSubsystem::CanEditLevelInstance(const ILevelInstanceInterface
 	
 	if (LevelInstance->IsWorldAssetValid())
 	{
-		if (GetWorld()->PersistentLevel->GetPackage()->GetName() == LevelInstance->GetWorldAssetPackage())
+		const FString WorldAssetPackage = LevelInstance->GetWorldAssetPackage();
+
+		if (GetWorld()->PersistentLevel->GetPackage()->GetName() == WorldAssetPackage)
 		{
 			if (OutReason)
 			{
-				*OutReason = FText::Format(LOCTEXT("CanEditLevelInstancePersistentLevel", "The Persistent level and the Level Instance are the same ({0})."), FText::FromString(LevelInstance->GetWorldAssetPackage()));
+				*OutReason = FText::Format(LOCTEXT("CanEditLevelInstancePersistentLevel", "The Persistent level and the Level Instance are the same ({0})."), FText::FromString(WorldAssetPackage));
 			}
 			return false;
 		}
 
-		if (FLevelUtils::FindStreamingLevel(GetWorld(), *LevelInstance->GetWorldAssetPackage()))
+		if (FLevelUtils::FindStreamingLevel(GetWorld(), *WorldAssetPackage))
 		{
 			if (OutReason)
 			{
-				*OutReason = FText::Format(LOCTEXT("CanEditLevelInstanceAlreadyExists", "The same level was added to world outside of Level Instances ({0})."), FText::FromString(LevelInstance->GetWorldAssetPackage()));
+				*OutReason = FText::Format(LOCTEXT("CanEditLevelInstanceAlreadyExists", "The same level was added to world outside of Level Instances ({0})."), FText::FromString(WorldAssetPackage));
 			}
 			return false;
 		}
 
 		FPackagePath WorldAssetPath;
-		if (!FPackagePath::TryFromPackageName(LevelInstance->GetWorldAssetPackage(), WorldAssetPath) || !FPackageName::DoesPackageExist(WorldAssetPath))
+		if (!FPackagePath::TryFromPackageName(WorldAssetPackage, WorldAssetPath) || !FPackageName::DoesPackageExist(WorldAssetPath))
 		{
 			if (OutReason)
 			{
-				*OutReason = FText::Format(LOCTEXT("CanEditLevelInstanceInvalidAsset", "Level Instance asset is invalid ({0})."), FText::FromString(LevelInstance->GetWorldAssetPackage()));
+				*OutReason = FText::Format(LOCTEXT("CanEditLevelInstanceInvalidAsset", "Level Instance asset is invalid ({0})."), FText::FromString(WorldAssetPackage));
 			}
 			return false;
 		}
 
-		if (ULevel::GetIsLevelPartitionedFromPackage(*LevelInstance->GetWorldAssetPackage()) && !ULevel::GetIsStreamingDisabledFromPackage(*LevelInstance->GetWorldAssetPackage()))
+		if (!Cast<APackedLevelActor>(LevelInstance) && 
+		ULevel::GetIsLevelPartitionedFromPackage(*WorldAssetPackage) &&
+		!ULevel::GetIsStreamingDisabledFromPackage(*WorldAssetPackage))
 		{
 			ILevelInstanceEditorModule& EditorModule = FModuleManager::GetModuleChecked<ILevelInstanceEditorModule>("LevelInstanceEditor");
 			if (!EditorModule.IsEditInPlaceStreamingEnabled())
 			{
 				if (OutReason)
 				{
-					*OutReason = FText::Format(LOCTEXT("CanEditLevelInstanceWithStreamingEnabled", "Level Instance can't be edited in place because it has streaming enabled ({0})"), FText::FromString(LevelInstance->GetWorldAssetPackage()));
+					*OutReason = FText::Format(LOCTEXT("CanEditLevelInstanceWithStreamingEnabled", "Level Instance can't be edited in place because it has streaming enabled ({0})"), FText::FromString(WorldAssetPackage));
 				}
 				return false;
 			}
