Don't register virtual texture callbacks for material render proxies that are created for a single frame.
These transient materials will not need to respond to virtual texture producer recreation/deletion events.
This allows us to avoid a race condition where transient materials are deleted with the scene renderer and try to unregister their virtual texture callbacks.
If the scene renderer deletion overlaps with the virtual texture update of the next scene renderer, then this causes us to hit the bUpdating check in RemoveAllProducerDestroyedCallbacks().
Fixing that race condition means that we no longer need to wait at the start of the virtual texture update. So we can set r.VT.WaitBeforeUpdate 0

diff --git a/Engine/Source/Runtime/Engine/Private/Materials/MaterialRenderProxy.cpp b/Engine/Source/Runtime/Engine/Private/Materials/MaterialRenderProxy.cpp
--- a/Engine/Source/Runtime/Engine/Private/Materials/MaterialRenderProxy.cpp
+++ b/Engine/Source/Runtime/Engine/Private/Materials/MaterialRenderProxy.cpp
@@ -304,8 +304,11 @@ IAllocatedVirtualTexture* FMaterialRenderProxy::GetPreallocatedVTStack(const FMa
 		return nullptr;
 	}
 
-	GetRendererModule().AddVirtualTextureProducerDestroyedCallback(Texture->GetProducerHandle(), &OnVirtualTextureDestroyedCB, const_cast<FMaterialRenderProxy*>(this));
-	HasVirtualTextureCallbacks = -1;
+	if (!bMarkedTransient)
+	{
+		GetRendererModule().AddVirtualTextureProducerDestroyedCallback(Texture->GetProducerHandle(), &OnVirtualTextureDestroyedCB, const_cast<FMaterialRenderProxy*>(this));
+		bHasVirtualTextureCallbacks = true;
+	}
 
 	return Texture->GetAllocatedVirtualTexture();
 }
@@ -373,8 +376,13 @@ IAllocatedVirtualTexture* FMaterialRenderProxy::AllocateVTStack(FRHICommandListB
 					VTDesc.TileBorderSize = VirtualTextureResourceForLayer->GetBorderSize();
 					VTDesc.ProducerHandle[LayerIndex] = ProducerHandle;
 					VTDesc.ProducerLayerIndex[LayerIndex] = 0u;
-					GetRendererModule().AddVirtualTextureProducerDestroyedCallback(ProducerHandle, &OnVirtualTextureDestroyedCB, const_cast<FMaterialRenderProxy*>(this));
 					bFoundValidLayer = true;
+
+					if (!bMarkedTransient)
+					{
+						GetRendererModule().AddVirtualTextureProducerDestroyedCallback(ProducerHandle, &OnVirtualTextureDestroyedCB, const_cast<FMaterialRenderProxy*>(this));
+						bHasVirtualTextureCallbacks = true;
+					}
 				}
 			}
 		}
@@ -382,7 +390,6 @@ IAllocatedVirtualTexture* FMaterialRenderProxy::AllocateVTStack(FRHICommandListB
 
 	if (bFoundValidLayer)
 	{
-		HasVirtualTextureCallbacks = -1;
 		return GetRendererModule().AllocateVirtualTexture(RHICmdList, VTDesc);
 	}
 	return nullptr;
@@ -408,10 +415,10 @@ void FMaterialRenderProxy::EvaluateUniformExpressions(FRHICommandListBase& RHICm
 	OutUniformExpressionCache.AllocatedVTs.Empty(UniformExpressionSet.VTStacks.Num());
 	OutUniformExpressionCache.OwnedAllocatedVTs.Empty(UniformExpressionSet.VTStacks.Num());
 
-	if (HasVirtualTextureCallbacks)
+	if (bHasVirtualTextureCallbacks)
 	{
 		GetRendererModule().RemoveAllVirtualTextureProducerDestroyedCallbacks(this);
-		HasVirtualTextureCallbacks = false;
+		bHasVirtualTextureCallbacks = false;
 	}
 
 	for (int32 i = 0; i < UniformExpressionSet.VTStacks.Num(); ++i)
@@ -560,10 +567,10 @@ void FMaterialRenderProxy::InvalidateUniformExpressionCache(bool bRecreateUnifor
 
 	UE::TScopeLock Lock(Mutex);
 
-	if (HasVirtualTextureCallbacks)
+	if (bHasVirtualTextureCallbacks)
 	{
 		GetRendererModule().RemoveAllVirtualTextureProducerDestroyedCallbacks(this);
-		HasVirtualTextureCallbacks = false;
+		bHasVirtualTextureCallbacks = false;
 	}
 
 	++UniformExpressionCacheSerialNumber;
@@ -611,7 +618,8 @@ FMaterialRenderProxy::FMaterialRenderProxy(FString InMaterialName)
 	, MaterialName(MoveTemp(InMaterialName))
 	, MarkedForGarbageCollection(0)
 	, DeletedFlag(0)
-	, HasVirtualTextureCallbacks(0)
+	, bMarkedTransient(false)
+	, bHasVirtualTextureCallbacks(false)
 {
 }
 
@@ -628,10 +636,10 @@ FMaterialRenderProxy::~FMaterialRenderProxy()
 		ReleaseResource();
 	}
 
-	if (HasVirtualTextureCallbacks)
+	if (bHasVirtualTextureCallbacks)
 	{
 		GetRendererModule().RemoveAllVirtualTextureProducerDestroyedCallbacks(this);
-		HasVirtualTextureCallbacks = false;
+		bHasVirtualTextureCallbacks = false;
 	}
 
 	DeletedFlag = -1;
@@ -685,10 +693,10 @@ void FMaterialRenderProxy::ReleaseResource()
 {
 	ReleaseResourceFlag = -1;
 	FRenderResource::ReleaseResource();
-	if (HasVirtualTextureCallbacks)
+	if (bHasVirtualTextureCallbacks)
 	{
 		GetRendererModule().RemoveAllVirtualTextureProducerDestroyedCallbacks(this);
-		HasVirtualTextureCallbacks = false;
+		bHasVirtualTextureCallbacks = false;
 	}
 }
 
diff --git a/Engine/Source/Runtime/Engine/Private/SceneManagement.cpp b/Engine/Source/Runtime/Engine/Private/SceneManagement.cpp
--- a/Engine/Source/Runtime/Engine/Private/SceneManagement.cpp
+++ b/Engine/Source/Runtime/Engine/Private/SceneManagement.cpp
@@ -465,6 +465,13 @@ FMeshElementCollector::~FMeshElementCollector()
 	MaterialProxiesToDelete.Empty();
 }
 
+void FMeshElementCollector::RegisterOneFrameMaterialProxy(FMaterialRenderProxy* Proxy)
+{
+	check(Proxy);
+	Proxy->MarkTransient();
+	MaterialProxiesToDelete.Add(Proxy);
+}
+
 FPrimitiveDrawInterface* FMeshElementCollector::GetPDI(int32 ViewIndex)
 {
 	return SimpleElementCollectors[ViewIndex];
diff --git a/Engine/Source/Runtime/Engine/Public/Materials/MaterialRenderProxy.h b/Engine/Source/Runtime/Engine/Public/Materials/MaterialRenderProxy.h
--- a/Engine/Source/Runtime/Engine/Public/Materials/MaterialRenderProxy.h
+++ b/Engine/Source/Runtime/Engine/Public/Materials/MaterialRenderProxy.h
@@ -195,6 +195,11 @@ public:
 		return MarkedForGarbageCollection != 0;
 	}
 
+	void MarkTransient()
+	{
+		bMarkedTransient = 1;
+	}
+
 	// FRenderResource interface.
 	ENGINE_API virtual void InitRHI(FRHICommandListBase& RHICmdList) override;
 	ENGINE_API virtual void ReleaseRHI() override;
@@ -262,8 +267,10 @@ private:
 	mutable uint8 MarkedForGarbageCollection : 1;
 	mutable uint8 DeletedFlag : 1;
 	mutable uint8 ReleaseResourceFlag : 1;
+	/** Set true if this is a transient, one frame object. */
+	uint8 bMarkedTransient : 1;
 	/** If any VT producer destroyed callbacks have been registered */
-	mutable uint8 HasVirtualTextureCallbacks : 1;
+	mutable uint8 bHasVirtualTextureCallbacks : 1;
 
 	/** Mutex for locking uniform expression invalidation / evaluation for this material. */
 	mutable UE::FMutex Mutex;
diff --git a/Engine/Source/Runtime/Engine/Public/MeshElementCollector.h b/Engine/Source/Runtime/Engine/Public/MeshElementCollector.h
--- a/Engine/Source/Runtime/Engine/Public/MeshElementCollector.h
+++ b/Engine/Source/Runtime/Engine/Public/MeshElementCollector.h
@@ -87,11 +87,7 @@ public:
 	ENGINE_API void AddMesh(int32 ViewIndex, FMeshBatch& MeshBatch);
 
 	/** Add a material render proxy that will be cleaned up automatically */
-	void RegisterOneFrameMaterialProxy(FMaterialRenderProxy* Proxy)
-	{
-		check(Proxy);
-		MaterialProxiesToDelete.Add(Proxy);
-	}
+	ENGINE_API void RegisterOneFrameMaterialProxy(FMaterialRenderProxy* Proxy);
 
 	/** Adds a request to force caching of uniform expressions for a material render proxy. */
 	ENGINE_API void CacheUniformExpressions(FMaterialRenderProxy* Proxy, bool bRecreateUniformBuffer);
diff --git a/Engine/Source/Runtime/Renderer/Private/VT/VirtualTextureSystem.cpp b/Engine/Source/Runtime/Renderer/Private/VT/VirtualTextureSystem.cpp
--- a/Engine/Source/Runtime/Renderer/Private/VT/VirtualTextureSystem.cpp
+++ b/Engine/Source/Runtime/Renderer/Private/VT/VirtualTextureSystem.cpp
@@ -167,7 +167,7 @@ static TAutoConsoleVariable<int32> CVarVTAsyncPageRequestTask(
 );
 static TAutoConsoleVariable<int32> CVarVTWaitBeforeUpdate(
 	TEXT("r.VT.WaitBeforeUpdate"),
-	1,
+	0,
 	TEXT("Wait on the SceneRenderBuilder before kicking VT async update."),
 	ECVF_RenderThreadSafe
 );
