Added LogNetLifecycle log category to track actor lifecycle in the replication system.

diff --git a/Engine/Source/Runtime/Engine/Private/DataChannel.cpp b/Engine/Source/Runtime/Engine/Private/DataChannel.cpp
--- a/Engine/Source/Runtime/Engine/Private/DataChannel.cpp
+++ b/Engine/Source/Runtime/Engine/Private/DataChannel.cpp
@@ -53,6 +53,7 @@ DEFINE_LOG_CATEGORY(LogNetSubObject);
 DEFINE_LOG_CATEGORY(LogRep);
 DEFINE_LOG_CATEGORY(LogNetPlayerMovement);
 DEFINE_LOG_CATEGORY(LogNetTraffic);
+DEFINE_LOG_CATEGORY(LogNetLifecycle);
 DEFINE_LOG_CATEGORY(LogRepTraffic);
 DEFINE_LOG_CATEGORY(LogNetDormancy);
 DEFINE_LOG_CATEGORY(LogSecurity);
@@ -2344,6 +2345,8 @@ void UActorChannel::Init( UNetConnection* InConnection, int32 InChannelIndex, EC
 	bForceCompareProperties	= false;
 	bActorIsPendingKill		= false;
 	CustomTimeDilation		= 1.0f;
+
+	UE_LOG(LogNetLifecycle, Log, TEXT("UActorChannel::Init: ChIndex: %d Actor: %s"), ChIndex, *GetFullNameSafe(Actor));
 }
 
 void UActorChannel::SetClosingFlag()
@@ -2371,6 +2374,8 @@ void UActorChannel::BreakAndReleaseReferences()
 
 int64 UActorChannel::Close(EChannelCloseReason Reason)
 {
+	UE_LOG(LogNetLifecycle, Log, TEXT("UActorChannel::Close: ChIndex: %d Actor: %s"), ChIndex, *GetFullNameSafe(Actor));
+
 	FScopedRepContext RepContext(Connection, Actor);
 
 	UE_LOG(LogNetTraffic, Log, TEXT("UActorChannel::Close: ChIndex: %d, Actor: %s, Reason: %s"), ChIndex, *GetFullNameSafe(Actor), LexToString(Reason));
@@ -2567,6 +2572,8 @@ void UActorChannel::DestroyActorAndComponents()
 
 bool UActorChannel::CleanUp(const bool bForDestroy, EChannelCloseReason CloseReason)
 {
+	UE_LOG(LogNetLifecycle, Log, TEXT("UActorChannel::CleanUp: ChIndex: %d Actor: %s"), ChIndex, *GetFullNameSafe(Actor));
+
 	SCOPE_CYCLE_COUNTER(Stat_ActorChanCleanUp);
 
 	checkf(Connection != nullptr, TEXT("UActorChannel::CleanUp: Connection is null!"));
@@ -3277,6 +3284,8 @@ void UActorChannel::ProcessBunch( FInBunch & Bunch )
 		}
 		else
 		{
+			UE_LOG(LogNetLifecycle, Log, TEXT("UActorChannel::ProcessBunch: Creating ChIndex: %d Actor: %s"), ChIndex, *NewChannelActor->GetName());
+
 			if (UE::Net::bDiscardTornOffActorRPCs && NewChannelActor->GetTearOff())
 			{
 				UE_LOG(LogNet, Warning, TEXT("UActorChannel::ProcessBunch: SerializeNewActor received an open bunch for a torn off actor. Actor: %s, Channel: %i"), *GetFullNameSafe(NewChannelActor), ChIndex);
diff --git a/Engine/Source/Runtime/Engine/Private/NetDriver.cpp b/Engine/Source/Runtime/Engine/Private/NetDriver.cpp
--- a/Engine/Source/Runtime/Engine/Private/NetDriver.cpp
+++ b/Engine/Source/Runtime/Engine/Private/NetDriver.cpp
@@ -802,6 +802,8 @@ void UNetDriver::AddNetworkActor(AActor* Actor)
 	LLM_SCOPE(ELLMTag::Networking);
 	ensureMsgf(Actor == nullptr || !(Actor->HasAnyFlags(RF_ClassDefaultObject | RF_ArchetypeObject)), TEXT("%s is a CDO or Archetype and should not be replicated."), *GetFullNameSafe(Actor));
 
+	UE_LOG(LogNetLifecycle, Log, TEXT("Registering for replication: Actor: %s"), *GetFullNameSafe(Actor));
+
 #if UE_WITH_IRIS
 	// We currently only want to do this on the server as we want to replicate the actor to the client
 	// but we rely on the old system for actor instantiation
@@ -4106,6 +4108,8 @@ void UNetDriver::NotifySubObjectDestroyed(UObject* SubObject)
 
 void UNetDriver::RemoveNetworkActor(AActor* Actor)
 {
+	UE_LOG(LogNetLifecycle, Log, TEXT("Deregistering for replication: Actor: %s"), *GetFullNameSafe(Actor));
+
 	// This is a hack: Sometimes world streaming will stream out levels even though some players can still see them
 	// Usually this is not a problem since the clients will remove the actors when it's told the level is gone. 
 	// But when SerializeNewActorOverrideLevel is off, the client thinks the actors are in the main level, so they don't get removed.
diff --git a/Engine/Source/Runtime/Engine/Public/EngineLogs.h b/Engine/Source/Runtime/Engine/Public/EngineLogs.h
--- a/Engine/Source/Runtime/Engine/Public/EngineLogs.h
+++ b/Engine/Source/Runtime/Engine/Public/EngineLogs.h
@@ -19,6 +19,7 @@ ENGINE_API DECLARE_LOG_CATEGORY_EXTERN(LogLevel, Log, All);
 ENGINE_API DECLARE_LOG_CATEGORY_EXTERN(LogSkeletalMesh, Log, All);
 ENGINE_API DECLARE_LOG_CATEGORY_EXTERN(LogStaticMesh, Log, All);
 ENGINE_API DECLARE_LOG_CATEGORY_EXTERN(LogNet, Log, All);
+ENGINE_API DECLARE_LOG_CATEGORY_EXTERN(LogNetLifecycle, Warning, All);
 ENGINE_API DECLARE_LOG_CATEGORY_EXTERN(LogNetSubObject, Log, All);
 ENGINE_API DECLARE_LOG_CATEGORY_EXTERN(LogRep, Log, All);
 ENGINE_API DECLARE_LOG_CATEGORY_EXTERN(LogNetPlayerMovement, Warning, All);
