Make the MenuCategories on UMaterialExpressions data driven via ini files.

- Add `config` to `MenuCategories.
- Add `Config = "Editor"` to `UMaterialExpression so `UPROPERTY` can be defined in an ini file.
- As a proof of concept move MenuCategories initialization from the UMaterialExpressionDesaturation to ini file
- ini file definitions can go in BaseEditor.ini

[RN]

diff --git a/Engine/Config/BaseEditor.ini b/Engine/Config/BaseEditor.ini
--- a/Engine/Config/BaseEditor.ini
+++ b/Engine/Config/BaseEditor.ini
@@ -707,4 +707,8 @@ bUseWidgetTemplateSelector=True
 +Filters=(FilterString="SkyLight.",DisplayName="Skylight")
 +Filters=(FilterString="TSR.",DisplayName="Temporal Super Resolution")
 +Filters=(FilterString="TexCoordRHIBuffer",DisplayName="Tex Coord Buffers")
-+Filters=(FilterString="FColorVertexBuffer",DisplayName="Color Vertex Buffers")
\ No newline at end of file
++Filters=(FilterString="FColorVertexBuffer",DisplayName="Color Vertex Buffers")
+
+[/Script/Engine.MaterialExpressionDesaturation]
++MenuCategories=NSLOCTEXT("MaterialExpressions", "Color", "Color")
++MenuCategories=NSLOCTEXT("MaterialExpressions", "Utility", "Utility")
\ No newline at end of file
diff --git a/Engine/Source/Runtime/Engine/Private/Materials/MaterialExpressions.cpp b/Engine/Source/Runtime/Engine/Private/Materials/MaterialExpressions.cpp
--- a/Engine/Source/Runtime/Engine/Private/Materials/MaterialExpressions.cpp
+++ b/Engine/Source/Runtime/Engine/Private/Materials/MaterialExpressions.cpp
@@ -9542,19 +9542,6 @@ void UMaterialExpressionFrac::GetCaption(TArray<FString>& OutCaptions) const
 UMaterialExpressionDesaturation::UMaterialExpressionDesaturation(const FObjectInitializer& ObjectInitializer)
 	: Super(ObjectInitializer)
 {
-	// Structure to hold one-time initialization
-	struct FConstructorStatics
-	{
-		FText NAME_Color;
-		FText NAME_Utility;
-		FConstructorStatics()
-			: NAME_Color(LOCTEXT( "Color", "Color" ))
-			, NAME_Utility(LOCTEXT( "Utility", "Utility" ))
-		{
-		}
-	};
-	static FConstructorStatics ConstructorStatics;
-
 	static const auto CVar = IConsoleManager::Get().FindConsoleVariable(TEXT("r.LegacyLuminanceFactors"));
 	if(CVar && CVar->GetInt() != 0)
 	{
@@ -9564,11 +9551,6 @@ UMaterialExpressionDesaturation::UMaterialExpressionDesaturation(const FObjectIn
 	{
 		LuminanceFactors = UE::Color::FColorSpace::GetWorking().GetLuminanceFactors();
 	}
-
-#if WITH_EDITORONLY_DATA
-	MenuCategories.Add(ConstructorStatics.NAME_Color);
-	MenuCategories.Add(ConstructorStatics.NAME_Utility);
-#endif
 }
 
 #if WITH_EDITOR
diff --git a/Engine/Source/Runtime/Engine/Public/Materials/MaterialExpression.h b/Engine/Source/Runtime/Engine/Public/Materials/MaterialExpression.h
--- a/Engine/Source/Runtime/Engine/Public/Materials/MaterialExpression.h
+++ b/Engine/Source/Runtime/Engine/Public/Materials/MaterialExpression.h
@@ -139,7 +139,7 @@ struct FMaterialExpressionCollection
 	TObjectPtr<UMaterialExpression> ExpressionExecEnd = nullptr;
 };
 
-UCLASS(abstract, Optional, BlueprintType, hidecategories=Object, MinimalAPI)
+UCLASS(abstract, Optional, BlueprintType, Config = "Editor", hidecategories=Object, MinimalAPI)
 class UMaterialExpression : public UObject
 {
 	GENERATED_UCLASS_BODY()
@@ -233,7 +233,7 @@ class UMaterialExpression : public UObject
 	uint32 bShowOutputs:1;
 
 	/** Localized categories to sort this expression into... */
-	UPROPERTY()
+	UPROPERTY(config)
 	TArray<FText> MenuCategories;
 
 	/** The expression's outputs, which are set in default properties by derived classes. */
