Merging RHI CreateBuffer refactor to Main.

diff --git a/Engine/Source/Runtime/D3D12RHI/Private/D3D12RHIPrivate.h b/Engine/Source/Runtime/D3D12RHI/Private/D3D12RHIPrivate.h
--- a/Engine/Source/Runtime/D3D12RHI/Private/D3D12RHIPrivate.h
+++ b/Engine/Source/Runtime/D3D12RHI/Private/D3D12RHIPrivate.h
@@ -337,7 +337,9 @@ public:
 	virtual void RHIReleaseTransition(FRHITransition* Transition) final override;
 	virtual FUniformBufferRHIRef RHICreateUniformBuffer(const void* Contents, const FRHIUniformBufferLayout* Layout, EUniformBufferUsage Usage, EUniformBufferValidation Validation) final override;
 	virtual void RHIUpdateUniformBuffer(FRHICommandListBase& RHICmdList, FRHIUniformBuffer* UniformBufferRHI, const void* Contents) final override;
-	virtual FBufferRHIRef RHICreateBuffer(FRHICommandListBase& RHICmdList, FRHIBufferDesc const& Desc, ERHIAccess ResourceState, FRHIResourceCreateInfo& CreateInfo) final override;
+
+	[[nodiscard]] virtual FRHIBufferInitializer RHICreateBufferInitializer(FRHICommandListBase& RHICmdList, const FRHIBufferCreateDesc& CreateDesc) override;
+
 	virtual void RHIReplaceResources(FRHICommandListBase& RHICmdList, TArray<FRHIResourceReplaceInfo>&& ReplaceInfos) final override;
 	virtual void* RHILockBuffer(FRHICommandListBase& RHICmdList, FRHIBuffer* Buffer, uint32 Offset, uint32 Size, EResourceLockMode LockMode) final override;
 	virtual void* RHILockBufferMGPU(FRHICommandListBase& RHICmdList, FRHIBuffer* Buffer, uint32 GPUIndex, uint32 Offset, uint32 Size, EResourceLockMode LockMode) final override;
@@ -478,8 +480,6 @@ public:
 	bool IsQuadBufferStereoEnabled() const;
 	void DisableQuadBufferStereo();
 
-	FBufferRHIRef CreateBuffer(FRHICommandListBase& RHICmdList, FRHIBufferDesc const& BufferDesc, ERHIAccess ResourceState, FRHIResourceCreateInfo& CreateInfo);
-
 	void* LockBuffer(FRHICommandListBase& RHICmdList, FD3D12Buffer* Buffer, uint32 BufferSize, EBufferUsageFlags BufferUsage, uint32 Offset, uint32 Size, EResourceLockMode LockMode);
 	void UnlockBuffer(FRHICommandListBase& RHICmdList, FD3D12Buffer* Buffer, EBufferUsageFlags BufferUsage);
 
@@ -599,10 +599,26 @@ public:
 		return FinalizeCreateTexture(RHICmdList, CreateDesc, CreateResult);
 	}
 
-	FD3D12Buffer* CreateD3D12Buffer(class FRHICommandListBase* RHICmdList, FRHIBufferDesc const& BufferDesc, ERHIAccess ResourceState, FRHIResourceCreateInfo& CreateInfo, ID3D12ResourceAllocator* ResourceAllocator = nullptr, bool bForceKeepUnlocked = false);
 	virtual FD3D12Texture* CreateNewD3D12Texture(const FRHITextureCreateDesc& CreateDesc, class FD3D12Device* Device);
 
-	FRHIBuffer* CreateBuffer(const FRHIBufferCreateInfo& CreateInfo, const TCHAR* DebugName, ERHIAccess InitialState, ID3D12ResourceAllocator* ResourceAllocator);
+	struct FCreateBufferInternalResult
+	{
+		FD3D12Buffer* Buffer{};
+		D3D12_RESOURCE_STATES DesiredState{};
+	};
+
+	FCreateBufferInternalResult CreateBufferInternal(const FRHIBufferCreateDesc& CreateDesc, bool bHasInitialData, ID3D12ResourceAllocator* ResourceAllocator);
+
+	FD3D12Buffer* CreateD3D12Buffer(FRHICommandListBase* RHICmdList, const FRHIBufferCreateDesc& CreateDesc, ID3D12ResourceAllocator* ResourceAllocator = nullptr)
+	{
+		FCreateBufferInternalResult CreateResult = CreateBufferInternal(CreateDesc, false, ResourceAllocator);
+		return CreateResult.Buffer;
+	}
+
+	FD3D12Buffer* CreateD3D12Buffer(const FRHIBufferCreateDesc& CreateDesc, ID3D12ResourceAllocator* ResourceAllocator)
+	{
+		return CreateD3D12Buffer(nullptr, CreateDesc, ResourceAllocator);
+	}
 
 	bool SetupDisplayHDRMetaData();
 
diff --git a/Engine/Source/Runtime/D3D12RHI/Private/D3D12RayTracing.cpp b/Engine/Source/Runtime/D3D12RHI/Private/D3D12RayTracing.cpp
--- a/Engine/Source/Runtime/D3D12RHI/Private/D3D12RayTracing.cpp
+++ b/Engine/Source/Runtime/D3D12RHI/Private/D3D12RayTracing.cpp
@@ -24,6 +24,7 @@
 #include "D3D12ResourceCollection.h"
 #include "RHIShaderBindingLayout.h"
 #include "RHIUniformBufferUtilities.h"
+#include "RHIResourceUtils.h"
 #include "D3D12TextureReference.h"
 
 extern int32 GD3D12ExplicitViewDescriptorHeapSize;
@@ -742,13 +743,22 @@ static FD3D12ShaderIdentifier GetShaderIdentifier(ID3D12StateObject* StateObject
 FD3D12RayTracingCompactionRequestHandler::FD3D12RayTracingCompactionRequestHandler(FD3D12Device* Device)
 	: FD3D12DeviceChild(Device)
 {
-	D3D12_RESOURCE_DESC PostBuildInfoBufferDesc = CD3DX12_RESOURCE_DESC::Buffer(GD3D12RayTracingMaxBatchedCompaction * sizeof(uint64), D3D12_RESOURCE_FLAG_ALLOW_UNORDERED_ACCESS);
+	const size_t BufferSize = GD3D12RayTracingMaxBatchedCompaction * sizeof(uint64);
 
-	FRHIGPUMask GPUMask = FRHIGPUMask::FromIndex(GetParentDevice()->GetGPUIndex());
-	ID3D12ResourceAllocator* ResourceAllocator = nullptr;
-	bool bHasInitialData = false;
-	PostBuildInfoBuffer = GetParentDevice()->GetParentAdapter()->CreateRHIBuffer(PostBuildInfoBufferDesc, 8,
-		FRHIBufferDesc(PostBuildInfoBufferDesc.Width, 0, BUF_UnorderedAccess | BUF_SourceCopy), ED3D12ResourceStateMode::MultiState, D3D12_RESOURCE_STATE_UNORDERED_ACCESS, bHasInitialData, GPUMask, ResourceAllocator, TEXT("PostBuildInfoBuffer"), ERHIAccess::UAVMask);
+	const D3D12_RESOURCE_DESC ResourceDesc = CD3DX12_RESOURCE_DESC::Buffer(BufferSize, D3D12_RESOURCE_FLAG_ALLOW_UNORDERED_ACCESS);
+
+	const FRHIBufferCreateDesc CreateDesc =
+		FRHIBufferCreateDesc::Create(TEXT("PostBuildInfoBuffer"), BufferSize, 0, BUF_UnorderedAccess | BUF_SourceCopy)
+		.SetInitialState(ERHIAccess::UAVMask)
+		.SetGPUMask(FRHIGPUMask::FromIndex(GetParentDevice()->GetGPUIndex()));
+
+	PostBuildInfoBuffer = GetParentDevice()->GetParentAdapter()->CreateRHIBuffer(
+		ResourceDesc,
+		8,
+		CreateDesc,
+		ED3D12ResourceStateMode::MultiState,
+		D3D12_RESOURCE_STATE_UNORDERED_ACCESS
+	);
 	SetName(PostBuildInfoBuffer->GetResource(), TEXT("PostBuildInfoBuffer"));
 
 	PostBuildInfoStagingBuffer = RHICreateStagingBuffer();
@@ -1641,15 +1651,25 @@ public:
 			WorkerData[WorkerIndex].ShaderRecordCache.Empty();
 		}
 
+		const uint32 BufferSize = Data.GetResourceDataSize();
+
 		FD3D12Adapter* Adapter = Device->GetParentAdapter();
-		D3D12_RESOURCE_DESC BufferDesc = CD3DX12_RESOURCE_DESC::Buffer(Data.GetResourceDataSize(), D3D12_RESOURCE_FLAG_NONE, D3D12_RAYTRACING_SHADER_TABLE_BYTE_ALIGNMENT);
-		FRHIGPUMask GPUMask = FRHIGPUMask::FromIndex(Device->GetGPUIndex());
-		bool bHasInitialData = true;
+		D3D12_RESOURCE_DESC BufferDesc = CD3DX12_RESOURCE_DESC::Buffer(BufferSize, D3D12_RESOURCE_FLAG_NONE, D3D12_RAYTRACING_SHADER_TABLE_BYTE_ALIGNMENT);
+
+		const FRHIBufferCreateDesc CreateDesc =
+			FRHIBufferCreateDesc::Create(TEXT("ShaderBindingTable"), BufferSize, 0, BUF_Static)
+			.SetInitialState(ERHIAccess::CopyDest)
+			.SetGPUMask(FRHIGPUMask::FromIndex(Device->GetGPUIndex()));
 
 		ID3D12ResourceAllocator* ResourceAllocator = nullptr;
 		Buffer = Adapter->CreateRHIBuffer(
-			BufferDesc, BufferDesc.Alignment, FRHIBufferDesc(BufferDesc.Width, 0, BUF_Static), ED3D12ResourceStateMode::MultiState,
-			D3D12_RESOURCE_STATE_COPY_DEST, bHasInitialData, GPUMask, ResourceAllocator, TEXT("Shader binding table"), ERHIAccess::CopyDest);
+			BufferDesc,
+			BufferDesc.Alignment,
+			CreateDesc,
+			ED3D12ResourceStateMode::MultiState,
+			D3D12_RESOURCE_STATE_COPY_DEST,
+			/*bHasInitialData*/ true
+		);
 
 		// Use copy queue for uploading the data
 		Context.BatchedSyncPoints.ToWait.Emplace(Buffer->UploadResourceDataViaCopyQueue(Context, &Data));
@@ -3456,23 +3476,37 @@ static TRefCountPtr<FD3D12Buffer> CreateRayTracingBuffer(FD3D12Adapter* Adapter,
 
 	if (Type == ERayTracingBufferType::AccelerationStructure)
 	{
-		D3D12_RESOURCE_DESC BufferDesc = CD3DX12_RESOURCE_DESC::Buffer(Size, D3D12_RESOURCE_FLAG_NONE);
+		const D3D12_RESOURCE_DESC ResourceDesc = CD3DX12_RESOURCE_DESC::Buffer(Size, D3D12_RESOURCE_FLAG_NONE);
+		const FRHIBufferCreateDesc CreateDesc =
+			FRHIBufferCreateDesc::Create(*DebugNameString, Size, 0, BUF_AccelerationStructure)
+			.SetInitialState(ERHIAccess::BVHWrite)
+			.SetGPUMask(GPUMask);
+
 		Result = Adapter->CreateRHIBuffer(
-			BufferDesc, D3D12_RAYTRACING_ACCELERATION_STRUCTURE_BYTE_ALIGNMENT,
-			FRHIBufferDesc(BufferDesc.Width, 0, BUF_AccelerationStructure),
-			ED3D12ResourceStateMode::SingleState, D3D12_RESOURCE_STATE_RAYTRACING_ACCELERATION_STRUCTURE, bHasInitialData,
-			GPUMask, ResourceAllocator, *DebugNameString, ERHIAccess::BVHWrite);
+			ResourceDesc,
+			D3D12_RAYTRACING_ACCELERATION_STRUCTURE_BYTE_ALIGNMENT,
+			CreateDesc,
+			ED3D12ResourceStateMode::SingleState,
+			D3D12_RESOURCE_STATE_RAYTRACING_ACCELERATION_STRUCTURE
+		);
 	}
 	else if (Type == ERayTracingBufferType::Scratch)
 	{
 		// Scratch doesn't need single state anymore because there are only a few scratch allocations left and allocating a 
 		// dedicated single state heap for it wastes memory - ideally all scratch allocations should be transient
-		D3D12_RESOURCE_DESC BufferDesc = CD3DX12_RESOURCE_DESC::Buffer(Size, D3D12_RESOURCE_FLAG_ALLOW_UNORDERED_ACCESS);
+		const D3D12_RESOURCE_DESC ResourceDesc = CD3DX12_RESOURCE_DESC::Buffer(Size, D3D12_RESOURCE_FLAG_ALLOW_UNORDERED_ACCESS);
+		const FRHIBufferCreateDesc CreateDesc =
+			FRHIBufferCreateDesc::Create(*DebugNameString, Size, 0, BUF_UnorderedAccess)
+			.SetInitialState(ERHIAccess::BVHWrite)
+			.SetGPUMask(GPUMask);
+
 		Result = Adapter->CreateRHIBuffer(
-			BufferDesc, D3D12_RAYTRACING_ACCELERATION_STRUCTURE_BYTE_ALIGNMENT,
-			FRHIBufferDesc(BufferDesc.Width, 0, BUF_UnorderedAccess),
-			ED3D12ResourceStateMode::Default, D3D12_RESOURCE_STATE_UNORDERED_ACCESS, bHasInitialData,
-			GPUMask, ResourceAllocator, *DebugNameString, ERHIAccess::BVHWrite);
+			ResourceDesc,
+			D3D12_RAYTRACING_ACCELERATION_STRUCTURE_BYTE_ALIGNMENT,
+			CreateDesc,
+			ED3D12ResourceStateMode::Default,
+			D3D12_RESOURCE_STATE_UNORDERED_ACCESS
+		);
 
 		// Elevates the scratch buffer heap priority, which may help performance / stability in low memory conditions 
 		// (Acceleration structure already boosted from allocation side)
@@ -3526,13 +3560,15 @@ FD3D12RayTracingGeometry::FD3D12RayTracingGeometry(FRHICommandListBase& RHICmdLi
 
 	if(!FD3D12RayTracingGeometry::NullTransformBuffer.IsValid())
 	{
-		TResourceArray<float> NullTransformData;
+		TArray<float> NullTransformData;
 		NullTransformData.SetNumZeroed(12);
 
-		FRHIResourceCreateInfo CreateInfo(TEXT("NullTransformBuffer"));
-		CreateInfo.ResourceArray = &NullTransformData;
-
-		FD3D12RayTracingGeometry::NullTransformBuffer = RHICmdList.CreateBuffer(NullTransformData.GetResourceDataSize(), BUF_VertexBuffer | BUF_Static, 0, ERHIAccess::VertexOrIndexBuffer, CreateInfo);
+		FD3D12RayTracingGeometry::NullTransformBuffer = UE::RHIResourceUtils::CreateVertexBufferFromArray(
+			RHICmdList,
+			TEXT("NullTransformBuffer"),
+			EBufferUsageFlags::Static,
+			MakeConstArrayView(NullTransformData)
+		);
 	}
 
 	RegisterD3D12RayTracingGeometry(this);
diff --git a/Engine/Source/Runtime/D3D12RHI/Private/D3D12Resources.cpp b/Engine/Source/Runtime/D3D12RHI/Private/D3D12Resources.cpp
--- a/Engine/Source/Runtime/D3D12RHI/Private/D3D12Resources.cpp
+++ b/Engine/Source/Runtime/D3D12RHI/Private/D3D12Resources.cpp
@@ -1737,11 +1737,6 @@ void FD3D12ResourceBarrierBatcher::FlushIntoCommandList(FD3D12CommandList& Comma
 	Barriers.Reset();
 }
 
-uint32 FD3D12Buffer::GetParentGPUIndex() const
-{
-	return Parent->GetGPUIndex();
-}
-
 void FD3D12DynamicRHI::RHIReplaceResources(FRHICommandListBase& RHICmdList, TArray<FRHIResourceReplaceInfo>&& ReplaceInfos)
 {
 	RHICmdList.EnqueueLambdaMultiPipe(GetEnabledRHIPipelines(), FRHICommandListBase::EThreadFence::Enabled, TEXT("FD3D12DynamicRHI::RHIReplaceResources"),
diff --git a/Engine/Source/Runtime/D3D12RHI/Private/D3D12Resources.h b/Engine/Source/Runtime/D3D12RHI/Private/D3D12Resources.h
--- a/Engine/Source/Runtime/D3D12RHI/Private/D3D12Resources.h
+++ b/Engine/Source/Runtime/D3D12RHI/Private/D3D12Resources.h
@@ -993,19 +993,16 @@ public:
 class FD3D12Buffer : public FRHIBuffer, public FD3D12BaseShaderResource, public FD3D12LinkedAdapterObject<FD3D12Buffer>
 {
 public:
-	FD3D12Buffer(FD3D12Device* InParent, FRHIBufferDesc const& Desc)
-		: FRHIBuffer(Desc)
+	FD3D12Buffer(FD3D12Device* InParent, const FRHIBufferCreateDesc& InCreateDesc)
+		: FRHIBuffer(InCreateDesc)
 		, FD3D12BaseShaderResource(InParent)
 		, LockedData(InParent)
-	{}
+	{
+	}
 
 	virtual ~FD3D12Buffer();
 
-	virtual uint32 GetParentGPUIndex() const override;
-
 	static void UploadResourceData(FD3D12CommandContext& CommandContext, D3D12_RESOURCE_STATES InDestinationState, FD3D12ResourceLocation& DestinationResourceLocation, const FD3D12ResourceLocation& SourceResourceLocation, uint32 Size);
-	void UploadResourceData(FRHICommandListBase& InRHICmdList, FRHIGPUMask GPUMask, D3D12_RESOURCE_STATES InDestinationState, const void* SourceData, int32 SourceDataSize);
-	void UploadResourceData(FRHICommandListBase& InRHICmdList, FResourceArrayUploadInterface* InResourceArray, D3D12_RESOURCE_STATES InDestinationState, const TCHAR* AssetName = nullptr, const FName& ClassName = NAME_None, const FName& PackageName = NAME_None);
 
 	FD3D12SyncPointRef UploadResourceDataViaCopyQueue(FD3D12CommandContext& OwningContext, FResourceArrayUploadInterface* InResourceArray);
 
@@ -1044,7 +1041,7 @@ public:
 		return FRHIResource::GetRefCount();
 	}
 
-	static void GetResourceDescAndAlignment(uint64 InSize, uint32 InStride, EBufferUsageFlags InUsage, D3D12_RESOURCE_DESC& ResourceDesc, uint32& Alignment);
+	static void GetResourceDescAndAlignment(const FRHIBufferCreateDesc& CreateDesc, D3D12_RESOURCE_DESC& ResourceDesc, uint32& Alignment);
 
 	FD3D12LockedResource LockedData;
 	uint32 BufferAlignment = 0;
diff --git a/Engine/Source/Runtime/D3D12RHI/Private/D3D12TransientResourceAllocator.cpp b/Engine/Source/Runtime/D3D12RHI/Private/D3D12TransientResourceAllocator.cpp
--- a/Engine/Source/Runtime/D3D12RHI/Private/D3D12TransientResourceAllocator.cpp
+++ b/Engine/Source/Runtime/D3D12RHI/Private/D3D12TransientResourceAllocator.cpp
@@ -218,29 +218,31 @@ void FD3D12TransientResourceHeapAllocator::FResourceAllocatorAdapter::AllocateRe
 #endif
 }
 
-FRHITransientBuffer* FD3D12TransientResourceHeapAllocator::CreateBuffer(const FRHIBufferCreateInfo& InCreateInfo, const TCHAR* InDebugName, const FRHITransientAllocationFences& Fences)
+FRHITransientBuffer* FD3D12TransientResourceHeapAllocator::CreateBuffer(const FRHIBufferDesc& InCreateInfo, const TCHAR* InDebugName, const FRHITransientAllocationFences& Fences)
 {
-	D3D12_RESOURCE_DESC Desc;
-	uint32 Alignment;
-	EBufferUsageFlags BufferUsage = InCreateInfo.Usage;
-	FD3D12Buffer::GetResourceDescAndAlignment(InCreateInfo.Size, InCreateInfo.Stride, BufferUsage, Desc, Alignment);
+	const FRHIBufferCreateDesc CreateDesc =
+		FRHIBufferCreateDesc::Create(InDebugName, InCreateInfo)
+		.SetInitialState(ERHIAccess::Discard);
+
+	D3D12_RESOURCE_DESC Desc{};
+	uint32 Alignment{};
+	FD3D12Buffer::GetResourceDescAndAlignment(CreateDesc, Desc, Alignment);
 
 	Alignment = D3D12_DEFAULT_RESOURCE_PLACEMENT_ALIGNMENT;
 	uint64 Size = Align(Desc.Width, Alignment);
 
-	return CreateBufferInternal(InCreateInfo, InDebugName, Fences, Size, Alignment,
+	return CreateBufferInternal(CreateDesc, InDebugName, Fences, Size, Alignment,
 		[&](const FRHITransientHeap::FResourceInitializer& Initializer)
 	{
 		FD3D12TransientHeap& Heap = static_cast<FD3D12TransientHeap&>(Initializer.Heap);
 
-		return CreateTransientResource<FRHITransientBuffer>([Adapter = GetParentAdapter(), &Heap, Allocation = Initializer.Allocation, Desc, InCreateInfo, InDebugName]
+		return CreateTransientResource<FRHITransientBuffer>([Adapter = GetParentAdapter(), &Heap, Allocation = Initializer.Allocation, Desc, CreateDesc, InDebugName]
 		{
 			TRACE_CPUPROFILER_EVENT_SCOPE(AllocatePlacedBuffer);
 			FResourceAllocatorAdapter ResourceAllocatorAdapter(Adapter, Heap, Allocation, Desc);
-			FRHIBuffer* Buffer = FD3D12DynamicRHI::GetD3DRHI()->CreateBuffer(InCreateInfo, InDebugName, ERHIAccess::Discard, &ResourceAllocatorAdapter);
-			Buffer->SetTrackedAccess_Unsafe(ERHIAccess::Discard);
+			FRHIBuffer* Buffer = FD3D12DynamicRHI::GetD3DRHI()->CreateD3D12Buffer(CreateDesc, &ResourceAllocatorAdapter);
 			return FRHITransientResource::FResourceTaskResult{ Buffer, ResourceAllocatorAdapter.GpuVirtualAddress };
 
-		}, Initializer.Hash, Size, InCreateInfo);
+		}, Initializer.Hash, Size, CreateDesc);
 	});
 }
\ No newline at end of file
diff --git a/Engine/Source/Runtime/D3D12RHI/Private/D3D12WorkGraph.cpp b/Engine/Source/Runtime/D3D12RHI/Private/D3D12WorkGraph.cpp
--- a/Engine/Source/Runtime/D3D12RHI/Private/D3D12WorkGraph.cpp
+++ b/Engine/Source/Runtime/D3D12RHI/Private/D3D12WorkGraph.cpp
@@ -951,12 +951,25 @@ void FD3D12CommandContext::DispatchWorkGraphShaderBundle(
 	D3D12_GPU_VIRTUAL_ADDRESS_RANGE_AND_STRIDE NodeLocalRootArgumentsTable{ 0, 0, 0 };
 	if (ValidRecords.Num() && LocalRootArgs.Num())
 	{
+		const uint32 DataSize = LocalRootArgs.GetResourceDataSize();
+
 		// todo: Check if copy queue is the optimal way to upload the root args.
 		// todo: Use a single buffer owned by the shader bundle RHI object (needs a copy operation that doesn't complain about multiple uploads).
-		D3D12_RESOURCE_DESC Desc = CD3DX12_RESOURCE_DESC::Buffer(LocalRootArgs.GetResourceDataSize(), D3D12_RESOURCE_FLAG_NONE);
+		const D3D12_RESOURCE_DESC ResourceDesc = CD3DX12_RESOURCE_DESC::Buffer(DataSize, D3D12_RESOURCE_FLAG_NONE);
+
+		const FRHIBufferCreateDesc CreateDesc =
+			FRHIBufferCreateDesc::Create(TEXT("BundleRecordBuffer"), DataSize, 0, EBufferUsageFlags::Static)
+			.SetInitialState(ERHIAccess::CopyDest)
+			.SetGPUMask(FRHIGPUMask::FromIndex(GetParentDevice()->GetGPUIndex()));
+
 		FD3D12Buffer* RootArgBuffer = GetParentDevice()->GetParentAdapter()->CreateRHIBuffer(
-			Desc, 16, FRHIBufferDesc(Desc.Width, 0, BUF_Static), ED3D12ResourceStateMode::MultiState,
-			D3D12_RESOURCE_STATE_COPY_DEST, true, FRHIGPUMask::FromIndex(GetParentDevice()->GetGPUIndex()), nullptr, TEXT("BundleRecordBuffer"), ERHIAccess::CopyDest);
+			ResourceDesc,
+			16,
+			CreateDesc,
+			ED3D12ResourceStateMode::MultiState,
+			D3D12_RESOURCE_STATE_COPY_DEST,
+			true
+		);
 
 		BatchedSyncPoints.ToWait.Emplace(RootArgBuffer->UploadResourceDataViaCopyQueue(*this, &LocalRootArgs));
 		TransitionResource(RootArgBuffer->GetResource(), D3D12_RESOURCE_STATE_COPY_DEST, D3D12_RESOURCE_STATE_COMMON, 0);
@@ -1296,12 +1309,25 @@ void FD3D12CommandContext::DispatchWorkGraphShaderBundle(
 	D3D12_GPU_VIRTUAL_ADDRESS_RANGE_AND_STRIDE NodeLocalRootArgumentsTable{ 0, 0, 0 };
 	if (LocalRootArgs.Num())
 	{
+		const uint32 DataSize = LocalRootArgs.GetResourceDataSize();
+
 		// todo: Check if copy queue is the optimal way to upload the root args.
 		// todo: Use a single buffer owned by the shader bundle RHI object (needs a copy operation that doesn't complain about multiple uploads).
-		D3D12_RESOURCE_DESC Desc = CD3DX12_RESOURCE_DESC::Buffer(LocalRootArgs.GetResourceDataSize(), D3D12_RESOURCE_FLAG_NONE);
+		const D3D12_RESOURCE_DESC ResourceDesc = CD3DX12_RESOURCE_DESC::Buffer(DataSize, D3D12_RESOURCE_FLAG_NONE);
+
+		const FRHIBufferCreateDesc CreateDesc =
+			FRHIBufferCreateDesc::Create(TEXT("BundleRecordBuffer"), DataSize, 0, EBufferUsageFlags::Static)
+			.SetInitialState(ERHIAccess::CopyDest)
+			.SetGPUMask(FRHIGPUMask::FromIndex(GetParentDevice()->GetGPUIndex()));
+
 		FD3D12Buffer* RootArgBuffer = GetParentDevice()->GetParentAdapter()->CreateRHIBuffer(
-			Desc, 16, FRHIBufferDesc(Desc.Width, 0, BUF_Static), ED3D12ResourceStateMode::MultiState,
-			D3D12_RESOURCE_STATE_COPY_DEST, true, FRHIGPUMask::FromIndex(GetParentDevice()->GetGPUIndex()), nullptr, TEXT("BundleRecordBuffer"), ERHIAccess::CopyDest);
+			ResourceDesc,
+			16,
+			CreateDesc,
+			ED3D12ResourceStateMode::MultiState,
+			D3D12_RESOURCE_STATE_COPY_DEST,
+			true
+		);
 
 		BatchedSyncPoints.ToWait.Emplace(RootArgBuffer->UploadResourceDataViaCopyQueue(*this, &LocalRootArgs));
 		TransitionResource(RootArgBuffer->GetResource(), D3D12_RESOURCE_STATE_COPY_DEST, D3D12_RESOURCE_STATE_COMMON, 0);
diff --git a/Engine/Source/Runtime/D3D12RHI/Public/D3D12RHI.h b/Engine/Source/Runtime/D3D12RHI/Public/D3D12RHI.h
--- a/Engine/Source/Runtime/D3D12RHI/Public/D3D12RHI.h
+++ b/Engine/Source/Runtime/D3D12RHI/Public/D3D12RHI.h
@@ -24,6 +24,7 @@
 #endif
 
 #if PLATFORM_WINDOWS
+	#define D3D12RHI_PLATFORM_HAS_UNIFIED_MEMORY	0
 	#define ENABLE_RESIDENCY_MANAGEMENT				1
 	#define PIPELINE_STATE_FILE_LOCATION			FPaths::ProjectSavedDir()
 	#define USE_PIX									WITH_PIX_EVENT_RUNTIME
diff --git a/Engine/Source/Runtime/Engine/Private/BatchedElements.cpp b/Engine/Source/Runtime/Engine/Private/BatchedElements.cpp
--- a/Engine/Source/Runtime/Engine/Private/BatchedElements.cpp
+++ b/Engine/Source/Runtime/Engine/Private/BatchedElements.cpp
@@ -13,6 +13,7 @@
 #include "ShaderParameterUtils.h"
 #include "GlobalRenderResources.h"
 #include "HDRHelper.h"
+#include "RHIResourceUtils.h"
 
 DEFINE_LOG_CATEGORY_STATIC(LogBatchedElements, Log, All);
 
@@ -839,11 +840,13 @@ void FBatchedElements::DrawPointElements(FRHICommandList& RHICmdList, const FMat
 			return;
 		}
 
-		FRHIResourceCreateInfo CreateInfo(TEXT("FBatchedElements_Points"));
-		FBufferRHIRef VertexBufferRHI = RHICmdList.CreateBuffer(sizeof(FSimpleElementVertex) * NumVertices, BUF_VertexBuffer | BUF_Volatile, 0, ERHIAccess::VertexOrIndexBuffer, CreateInfo);
-		void* VerticesPtr = RHICmdList.LockBuffer(VertexBufferRHI, 0, sizeof(FSimpleElementVertex) * NumVertices, RLM_WriteOnly);
+		const FRHIBufferCreateDesc Desc =
+			FRHIBufferCreateDesc::CreateVertex<FSimpleElementVertex>(TEXT("FBatchedElements_Points"), NumVertices)
+			.AddUsage(EBufferUsageFlags::Volatile)
+			.SetInitialState(ERHIAccess::VertexOrIndexBuffer)
+			.SetInitActionInitializer();
 
-		FSimpleElementVertex* PointVertices = (FSimpleElementVertex*)VerticesPtr;
+		TRHIBufferInitializer<FSimpleElementVertex> Initializer = RHICmdList.CreateBufferInitializer(Desc);
 
 		uint32 VertIdx = 0;
 		for(int32 PointIndex = 0; PointIndex < Points.Num(); PointIndex++)
@@ -856,18 +859,19 @@ void FBatchedElements::DrawPointElements(FRHICommandList& RHICmdList, const FMat
 			const uint32 ViewportMajorAxis = ViewportSizeX;//FMath::Max(ViewportSizeX, ViewportSizeY);
 			const FVector WorldPointX = CameraX * Point.Size / ViewportMajorAxis * TransformedPosition.W;
 			const FVector WorldPointY = CameraY * -Point.Size / ViewportMajorAxis * TransformedPosition.W;
-					
-			PointVertices[VertIdx + 0] = FSimpleElementVertex(Point.Position + WorldPointX - WorldPointY,FVector2D(1,0),Point.Color,Point.HitProxyColor);
-			PointVertices[VertIdx + 1] = FSimpleElementVertex(Point.Position + WorldPointX + WorldPointY,FVector2D(1,1),Point.Color,Point.HitProxyColor);
-			PointVertices[VertIdx + 2] = FSimpleElementVertex(Point.Position - WorldPointX - WorldPointY,FVector2D(0,0),Point.Color,Point.HitProxyColor);
-			PointVertices[VertIdx + 3] = FSimpleElementVertex(Point.Position + WorldPointX + WorldPointY,FVector2D(1,1),Point.Color,Point.HitProxyColor);
-			PointVertices[VertIdx + 4] = FSimpleElementVertex(Point.Position - WorldPointX - WorldPointY,FVector2D(0,0),Point.Color,Point.HitProxyColor);
-			PointVertices[VertIdx + 5] = FSimpleElementVertex(Point.Position - WorldPointX + WorldPointY,FVector2D(0,1),Point.Color,Point.HitProxyColor);
+
+			Initializer[VertIdx + 0] = FSimpleElementVertex(Point.Position + WorldPointX - WorldPointY,FVector2D(1,0),Point.Color,Point.HitProxyColor);
+			Initializer[VertIdx + 1] = FSimpleElementVertex(Point.Position + WorldPointX + WorldPointY,FVector2D(1,1),Point.Color,Point.HitProxyColor);
+			Initializer[VertIdx + 2] = FSimpleElementVertex(Point.Position - WorldPointX - WorldPointY,FVector2D(0,0),Point.Color,Point.HitProxyColor);
+			Initializer[VertIdx + 3] = FSimpleElementVertex(Point.Position + WorldPointX + WorldPointY,FVector2D(1,1),Point.Color,Point.HitProxyColor);
+			Initializer[VertIdx + 4] = FSimpleElementVertex(Point.Position - WorldPointX - WorldPointY,FVector2D(0,0),Point.Color,Point.HitProxyColor);
+			Initializer[VertIdx + 5] = FSimpleElementVertex(Point.Position - WorldPointX + WorldPointY,FVector2D(0,1),Point.Color,Point.HitProxyColor);
 
 			VertIdx += 6;
 		}
 
-		RHICmdList.UnlockBuffer(VertexBufferRHI);
+		FBufferRHIRef VertexBufferRHI = Initializer.Finalize();
+
 		RHICmdList.SetStreamSource(0, VertexBufferRHI, 0);
 		RHICmdList.DrawPrimitive(0, NumTris, 1);
 	}
@@ -934,12 +938,7 @@ bool FBatchedElements::Draw(FRHICommandList& RHICmdList, const FMeshPassProcesso
 				// Set the appropriate pixel shader parameters & shader state for the non-textured elements.
 				PrepareShaders(RHICmdList, GraphicsPSOInit, StencilRef, FeatureLevel, SE_BLEND_Opaque, RelativeMatrices, BatchedElementParameters, GWhiteTexture, bHitTesting, Gamma, NULL, &View);
 
-				FRHIResourceCreateInfo CreateInfo(TEXT("Lines"));
-				FBufferRHIRef VertexBufferRHI = RHICmdList.CreateBuffer(sizeof(FSimpleElementVertex) * LineVertices.Num(), BUF_VertexBuffer | BUF_Volatile, 0, ERHIAccess::VertexOrIndexBuffer, CreateInfo);
-				void* VoidPtr = RHICmdList.LockBuffer(VertexBufferRHI, 0, sizeof(FSimpleElementVertex) * LineVertices.Num(), RLM_WriteOnly);
-
-				FMemory::Memcpy(VoidPtr, LineVertices.GetData(), sizeof(FSimpleElementVertex) * LineVertices.Num());
-				RHICmdList.UnlockBuffer(VertexBufferRHI);
+				FBufferRHIRef VertexBufferRHI = UE::RHIResourceUtils::CreateVertexBufferFromArray(RHICmdList, TEXT("Lines"), BUF_Volatile, MakeConstArrayView(LineVertices));
 
 				RHICmdList.SetStreamSource(0, VertexBufferRHI, 0);
 
@@ -1013,11 +1012,14 @@ bool FBatchedElements::Draw(FRHICommandList& RHICmdList, const FMeshPassProcesso
 					GraphicsPSOInit.RasterizerState = RasterState.GetReference();
 					PrepareShaders(RHICmdList, GraphicsPSOInit, StencilRef, FeatureLevel, SE_BLEND_AlphaBlend, RelativeMatrices, BatchedElementParameters, GWhiteTexture, bHitTesting, Gamma, NULL, &View);
 
-					FRHIResourceCreateInfo CreateInfo(TEXT("ThickLines"));
-					FBufferRHIRef VertexBufferRHI = RHICmdList.CreateBuffer(sizeof(FSimpleElementVertex) * VertPerLine * NumLinesThisBatch, BUF_VertexBuffer | BUF_Volatile, 0, ERHIAccess::VertexOrIndexBuffer, CreateInfo);
-					void* ThickVertexData = RHICmdList.LockBuffer(VertexBufferRHI, 0, sizeof(FSimpleElementVertex) * VertPerLine * NumLinesThisBatch, RLM_WriteOnly);
-					FSimpleElementVertex* ThickVertices = (FSimpleElementVertex*)ThickVertexData;
-					check(ThickVertices);
+					const FRHIBufferCreateDesc Desc =
+						FRHIBufferCreateDesc::CreateVertex<FSimpleElementVertex>(TEXT("ThickLines"), VertPerLine * NumLinesThisBatch)
+						.AddUsage(EBufferUsageFlags::Volatile)
+						.SetInitialState(ERHIAccess::VertexOrIndexBuffer)
+						.SetInitActionInitializer();
+
+					TRHIBufferInitializer<FSimpleElementVertex> InitialData = RHICmdList.CreateBufferInitializer(Desc);
+					FSimpleElementVertex* ThickVertices = InitialData.GetWritableData();
 
 					int32 AddedTris = 0;
 					for (int i = 0; i < NumLinesThisBatch; ++i)
@@ -1227,7 +1229,8 @@ bool FBatchedElements::Draw(FRHICommandList& RHICmdList, const FMeshPassProcesso
 						AddedTris += 4;
 					}
 
-					RHICmdList.UnlockBuffer(VertexBufferRHI);
+					FBufferRHIRef VertexBufferRHI = InitialData.Finalize();
+
 					if (AddedTris > 0)
 					{
 						RHICmdList.SetStreamSource(0, VertexBufferRHI, 0);
@@ -1242,8 +1245,13 @@ bool FBatchedElements::Draw(FRHICommandList& RHICmdList, const FMeshPassProcesso
 			{
 				check(WireTriVerts.Num() == WireTris.Num() * 3);
 
-				FRHIResourceCreateInfo CreateInfo(TEXT("WireTris"), &WireTriVerts);
-				FBufferRHIRef VertexBufferRHI = RHICmdList.CreateBuffer(WireTriVerts.GetResourceDataSize(), BUF_VertexBuffer | BUF_Volatile, 0, ERHIAccess::VertexOrIndexBuffer, CreateInfo);
+				const FRHIBufferCreateDesc CreateDesc =
+					FRHIBufferCreateDesc::CreateVertex(TEXT("WireTris"), WireTriVerts.GetResourceDataSize())
+					.AddUsage(BUF_Volatile)
+					.SetInitActionResourceArray(&WireTriVerts)
+					.SetInitialState(ERHIAccess::VertexOrIndexBuffer);
+
+				FBufferRHIRef VertexBufferRHI = RHICmdList.CreateBuffer(CreateDesc);
 
 				RHICmdList.SetStreamSource(0, VertexBufferRHI, 0);
 
@@ -1335,10 +1343,14 @@ bool FBatchedElements::Draw(FRHICommandList& RHICmdList, const FMeshPassProcesso
 					return false;
 				}
 
-				FRHIResourceCreateInfo CreateInfo(TEXT("Sprites"));
-				FBufferRHIRef VertexBufferRHI = RHICmdList.CreateBuffer(sizeof(FSimpleElementVertex) * ValidSpriteCount * 6, BUF_VertexBuffer | BUF_Volatile, 0, ERHIAccess::VertexOrIndexBuffer, CreateInfo);
-				void* VoidPtr = RHICmdList.LockBuffer(VertexBufferRHI, 0, sizeof(FSimpleElementVertex) * ValidSpriteCount * 6, RLM_WriteOnly);
-				FSimpleElementVertex* SpriteList = reinterpret_cast<FSimpleElementVertex*>(VoidPtr);
+				const FRHIBufferCreateDesc Desc =
+					FRHIBufferCreateDesc::CreateVertex<FSimpleElementVertex>(TEXT("Sprites"), ValidSpriteCount * 6)
+					.AddUsage(EBufferUsageFlags::Volatile)
+					.SetInitialState(ERHIAccess::VertexOrIndexBuffer)
+					.SetInitActionInitializer();
+
+				TRHIBufferInitializer<FSimpleElementVertex> Initializer = RHICmdList.CreateBufferInitializer(Desc);
+				FSimpleElementVertex* SpriteList = Initializer.GetWritableData();
 
 				for (int32 SpriteIndex = 0; SpriteIndex < ValidSpriteCount; SpriteIndex++)
 				{
@@ -1362,7 +1374,8 @@ bool FBatchedElements::Draw(FRHICommandList& RHICmdList, const FMeshPassProcesso
 					Vertex[4] = FSimpleElementVertex(Sprite.Position - WorldSpriteX - WorldSpriteY, FVector2D(UStart, VStart), Sprite.Color, Sprite.HitProxyColor);
 					Vertex[5] = FSimpleElementVertex(Sprite.Position - WorldSpriteX + WorldSpriteY, FVector2D(UStart, VEnd), Sprite.Color, Sprite.HitProxyColor);
 				}
-				RHICmdList.UnlockBuffer(VertexBufferRHI);
+
+				FBufferRHIRef VertexBufferRHI = Initializer.Finalize();
 
 				//First time init
 				const FTexture* CurrentTexture = SortedSprites[0].Texture;
@@ -1409,11 +1422,7 @@ bool FBatchedElements::Draw(FRHICommandList& RHICmdList, const FMeshPassProcesso
 				return false;
 			}
 
-			FRHIResourceCreateInfo CreateInfo(TEXT("MeshElements"));
-			FBufferRHIRef VertexBufferRHI = RHICmdList.CreateBuffer(sizeof(FSimpleElementVertex) * MeshVertices.Num(), BUF_VertexBuffer | BUF_Volatile, 0, ERHIAccess::VertexOrIndexBuffer, CreateInfo);
-			void* VoidPtr = RHICmdList.LockBuffer(VertexBufferRHI, 0, sizeof(FSimpleElementVertex) * MeshVertices.Num(), RLM_WriteOnly);
-			FPlatformMemory::Memcpy(VoidPtr, MeshVertices.GetData(), sizeof(FSimpleElementVertex) * MeshVertices.Num());
-			RHICmdList.UnlockBuffer(VertexBufferRHI);
+			FBufferRHIRef VertexBufferRHI = UE::RHIResourceUtils::CreateVertexBufferFromArray(RHICmdList, TEXT("MeshElements"), BUF_Volatile, MakeConstArrayView(MeshVertices));
 
 			// Draw the mesh elements.
 			for(int32 MeshIndex = 0;MeshIndex < MeshElements.Num();MeshIndex++)
@@ -1435,10 +1444,7 @@ bool FBatchedElements::Draw(FRHICommandList& RHICmdList, const FMeshPassProcesso
 					// Set the appropriate pixel shader for the mesh.
 					PrepareShaders(RHICmdList, GraphicsPSOInit, StencilRef, FeatureLevel, MeshElement.BlendMode, RelativeMatrices, MeshElement.BatchedElementParameters, MeshElement.Texture, bHitTesting, Gamma, &MeshElement.GlowInfo, &View);
 
-					FBufferRHIRef IndexBufferRHI = RHICmdList.CreateBuffer(sizeof(uint16) * MeshElement.Indices.Num(), BUF_IndexBuffer | BUF_Volatile, sizeof(uint16), ERHIAccess::VertexOrIndexBuffer, CreateInfo);
-					void* VoidPtr2 = RHICmdList.LockBuffer(IndexBufferRHI, 0, sizeof(uint16) * MeshElement.Indices.Num(), RLM_WriteOnly);
-					FPlatformMemory::Memcpy(VoidPtr2, MeshElement.Indices.GetData(), sizeof(uint16) * MeshElement.Indices.Num());
-					RHICmdList.UnlockBuffer(IndexBufferRHI);
+					FBufferRHIRef IndexBufferRHI = UE::RHIResourceUtils::CreateIndexBufferFromArray(RHICmdList, TEXT("MeshElements"), EBufferUsageFlags::Volatile, MakeConstArrayView(MeshElement.Indices));
 
 					// Draw the mesh.
 					RHICmdList.SetStreamSource(0, VertexBufferRHI, MeshElement.MinVertex * sizeof(FSimpleElementVertex));
diff --git a/Engine/Source/Runtime/Engine/Private/Components/BrushComponent.cpp b/Engine/Source/Runtime/Engine/Private/Components/BrushComponent.cpp
--- a/Engine/Source/Runtime/Engine/Private/Components/BrushComponent.cpp
+++ b/Engine/Source/Runtime/Engine/Private/Components/BrushComponent.cpp
@@ -60,22 +60,29 @@ public:
 	{
 		if(NumEdges)
 		{
-			FRHIResourceCreateInfo CreateInfo(TEXT("FModelWireIndexBuffer"));
-			IndexBufferRHI = RHICmdList.CreateIndexBuffer(sizeof(uint16),NumEdges * 2 * sizeof(uint16),BUF_Static, CreateInfo);
+			const FRHIBufferCreateDesc Desc =
+				FRHIBufferCreateDesc::CreateIndex<uint16>(TEXT("FModelWireIndexBuffer"), NumEdges * 2)
+				.AddUsage(EBufferUsageFlags::Static)
+				.SetInitActionInitializer()
+				.DetermineInitialState();
 
-			uint16* DestIndex = (uint16*)RHICmdList.LockBuffer(IndexBufferRHI,0,NumEdges * 2 * sizeof(uint16),RLM_WriteOnly);
+			TRHIBufferInitializer<uint16> Initializer = RHICmdList.CreateBufferInitializer(Desc);
+
+			size_t WritingIndex = 0;
 			uint16 BaseIndex = 0;
 			for(int32 PolyIndex = 0;PolyIndex < Polys.Num();PolyIndex++)
 			{
 				FPoly&	Poly = Polys[PolyIndex];
-				for(int32 VertexIndex = 0;VertexIndex < Poly.Vertices.Num();VertexIndex++)
+				for(int32 VertexIndex = 0; VertexIndex < Poly.Vertices.Num(); VertexIndex++)
 				{
-					*DestIndex++ = BaseIndex + VertexIndex;
-					*DestIndex++ = BaseIndex + ((VertexIndex + 1) % Poly.Vertices.Num());
+					Initializer[WritingIndex + 0] = BaseIndex + VertexIndex;
+					Initializer[WritingIndex + 1] = BaseIndex + ((VertexIndex + 1) % Poly.Vertices.Num());
+
+					WritingIndex += 2;
 				}
 				BaseIndex += Poly.Vertices.Num();
 			}
-			RHICmdList.UnlockBuffer(IndexBufferRHI);
+			IndexBufferRHI = Initializer.Finalize();
 		}
 	}
 
diff --git a/Engine/Source/Runtime/Engine/Private/Components/MaterialBillboardComponent.cpp b/Engine/Source/Runtime/Engine/Private/Components/MaterialBillboardComponent.cpp
--- a/Engine/Source/Runtime/Engine/Private/Components/MaterialBillboardComponent.cpp
+++ b/Engine/Source/Runtime/Engine/Private/Components/MaterialBillboardComponent.cpp
@@ -34,8 +34,12 @@ public:
 
 	virtual void InitRHI(FRHICommandListBase& RHICmdList) override
 	{
-		FRHIResourceCreateInfo CreateInfo(TEXT("FMaterialSpriteVertexBuffer"));
-		VertexBufferRHI = RHICmdList.CreateVertexBuffer(sizeof(FMaterialSpriteVertex),BUF_Static,CreateInfo);
+		const FRHIBufferCreateDesc CreateDesc =
+			FRHIBufferCreateDesc::CreateVertex(TEXT("FMaterialSpriteVertexBuffer"), sizeof(FMaterialSpriteVertex))
+			.AddUsage(EBufferUsageFlags::Static)
+			.DetermineInitialState();
+
+		VertexBufferRHI = RHICmdList.CreateBuffer(CreateDesc);
 	}
 };
 static TGlobalResource<FMaterialSpriteVertexBuffer> GDummyMaterialSpriteVertexBuffer;
diff --git a/Engine/Source/Runtime/Engine/Private/DynamicMeshBuilder.cpp b/Engine/Source/Runtime/Engine/Private/DynamicMeshBuilder.cpp
--- a/Engine/Source/Runtime/Engine/Private/DynamicMeshBuilder.cpp
+++ b/Engine/Source/Runtime/Engine/Private/DynamicMeshBuilder.cpp
@@ -101,11 +101,14 @@ public:
 	FBufferRHIRef CreateResource(FRHICommandListBase& RHICmdList, FGlobalDynamicMeshPoolPolicy::CreationArguments Args)
 	{
 		FGlobalDynamicMeshPoolPolicy::CreationArguments BufferSize = GetPoolBucketSize(GetPoolBucketIndex(Args));
-		// The use of BUF_Static is deliberate - on OS X the buffer backing-store orphaning & reallocation will dominate execution time
-		// so to avoid this we don't reuse a buffer for several frames, thereby avoiding the pipeline stall and the reallocation cost.
-		FRHIResourceCreateInfo CreateInfo(TEXT("FGlobalDynamicMeshIndexPolicy"));
-		FBufferRHIRef VertexBuffer = RHICmdList.CreateIndexBuffer(sizeof(DynamicMeshIndexType), BufferSize, BUF_Static, CreateInfo);
-		return VertexBuffer;
+		const FRHIBufferCreateDesc CreateDesc =
+			FRHIBufferCreateDesc::CreateIndex(TEXT("FGlobalDynamicMeshIndexPolicy"), BufferSize, sizeof(DynamicMeshIndexType))
+			// The use of BUF_Static is deliberate - on OS X the buffer backing-store orphaning & reallocation will dominate execution time
+			// so to avoid this we don't reuse a buffer for several frames, thereby avoiding the pipeline stall and the reallocation cost.
+			.AddUsage(EBufferUsageFlags::Static)
+			.DetermineInitialState();
+
+		return RHICmdList.CreateBuffer(CreateDesc);
 	}
 	
 	/** Gets the arguments used to create resource
@@ -159,9 +162,13 @@ public:
 	FBufferRHIRef CreateResource(FRHICommandListBase& RHICmdList, FGlobalDynamicMeshPoolPolicy::CreationArguments Args)
 	{
 		FGlobalDynamicMeshPoolPolicy::CreationArguments BufferSize = GetPoolBucketSize(GetPoolBucketIndex(Args));
-		FRHIResourceCreateInfo CreateInfo(TEXT("FGlobalDynamicMeshVertexPolicy"));
-		FBufferRHIRef VertexBuffer = RHICmdList.CreateVertexBuffer(BufferSize, BUF_Volatile | BUF_ShaderResource, CreateInfo);
-		return VertexBuffer;
+
+		const FRHIBufferCreateDesc CreateDesc =
+			FRHIBufferCreateDesc::CreateVertex(TEXT("FGlobalDynamicMeshVertexPolicy"), BufferSize)
+			.AddUsage(EBufferUsageFlags::Volatile | EBufferUsageFlags::ShaderResource)
+			.DetermineInitialState();
+
+		return RHICmdList.CreateBuffer(CreateDesc);
 	}
 	
 	/** Gets the arguments used to create resource
@@ -231,8 +238,12 @@ FBufferRHIRef FDynamicMeshBufferAllocator::AllocIndexBuffer(FRHICommandListBase&
 {
 	uint32 SizeInBytes = GetIndexBufferSize(NumElements);
 
-	FRHIResourceCreateInfo CreateInfo(TEXT("FDynamicMeshBufferAllocator"));
-	return RHICmdList.CreateIndexBuffer(sizeof(DynamicMeshIndexType), SizeInBytes, BUF_Volatile, CreateInfo);
+	const FRHIBufferCreateDesc CreateDesc =
+		FRHIBufferCreateDesc::CreateIndex(TEXT("FDynamicMeshBufferAllocator"), SizeInBytes, sizeof(DynamicMeshIndexType))
+		.AddUsage(EBufferUsageFlags::Volatile)
+		.DetermineInitialState();
+
+	return RHICmdList.CreateBuffer(CreateDesc);
 }
 
 void FDynamicMeshBufferAllocator::ReleaseIndexBuffer(FBufferRHIRef& IndexBufferRHI)
@@ -249,8 +260,12 @@ FBufferRHIRef FDynamicMeshBufferAllocator::AllocVertexBuffer(FRHICommandListBase
 {
 	uint32 SizeInBytes = GetVertexBufferSize(Stride, NumElements);
 
-	FRHIResourceCreateInfo CreateInfo(TEXT("FDynamicMeshBufferAllocator"));
-	return RHICmdList.CreateVertexBuffer(SizeInBytes, BUF_Volatile | BUF_ShaderResource, CreateInfo);
+	const FRHIBufferCreateDesc CreateDesc =
+		FRHIBufferCreateDesc::CreateVertex(TEXT("FDynamicMeshBufferAllocator"), SizeInBytes)
+		.AddUsage(EBufferUsageFlags::Volatile | EBufferUsageFlags::ShaderResource)
+		.DetermineInitialState();
+
+	return RHICmdList.CreateBuffer(CreateDesc);
 }
 
 void FDynamicMeshBufferAllocator::ReleaseVertexBuffer(FBufferRHIRef& VertexBufferRHI)
diff --git a/Engine/Source/Runtime/Engine/Private/GPUSkinCache.cpp b/Engine/Source/Runtime/Engine/Private/GPUSkinCache.cpp
--- a/Engine/Source/Runtime/Engine/Private/GPUSkinCache.cpp
+++ b/Engine/Source/Runtime/Engine/Private/GPUSkinCache.cpp
@@ -22,11 +22,11 @@ GPUSkinCache.cpp: Performs skinning on a compute shader into a buffer to avoid v
 #include "RenderUtils.h"
 #include "RendererInterface.h"
 #include "RenderingThread.h"
+#include "RHIResourceUtils.h"
 #include "Stats/StatsTrace.h"
 #include "UObject/UObjectIterator.h"
 #include "Algo/Sort.h"
 #include "ComponentRecreateRenderStateContext.h"
-#include "Algo/Sort.h"
 
 DEFINE_STAT(STAT_GPUSkinCache_TotalNumChunks);
 DEFINE_STAT(STAT_GPUSkinCache_TotalNumVertices);
@@ -1825,8 +1825,13 @@ bool FGPUSkinCache::ProcessEntry(
 	        FResourceArrayInterface* ResourceArray = VertexAndNormalData.GetResourceArray();
 	        check(ResourceArray->GetResourceDataSize() > 0);
 
-	        FRHIResourceCreateInfo CreateInfo(TEXT("ClothPositionAndNormalsBuffer"), ResourceArray);
-	        ClothPositionAndNormalsBuffer.VertexBufferRHI = RHICmdList.CreateVertexBuffer( ResourceArray->GetResourceDataSize(), BUF_Static | BUF_ShaderResource, CreateInfo);
+			const FRHIBufferCreateDesc CreateDesc =
+				FRHIBufferCreateDesc::CreateVertex(TEXT("ClothPositionAndNormalsBuffer"), ResourceArray->GetResourceDataSize())
+				.AddUsage(BUF_Static | BUF_ShaderResource)
+				.SetInitActionResourceArray(ResourceArray)
+				.DetermineInitialState();
+
+	        ClothPositionAndNormalsBuffer.VertexBufferRHI = RHICmdList.CreateBuffer(CreateDesc);
 	        ClothPositionAndNormalsBuffer.VertexBufferSRV = RHICmdList.CreateShaderResourceView(
 				ClothPositionAndNormalsBuffer.VertexBufferRHI, 
 				FRHIViewDesc::CreateBufferSRV()
diff --git a/Engine/Source/Runtime/Engine/Private/GPUSkinVertexFactory.cpp b/Engine/Source/Runtime/Engine/Private/GPUSkinVertexFactory.cpp
--- a/Engine/Source/Runtime/Engine/Private/GPUSkinVertexFactory.cpp
+++ b/Engine/Source/Runtime/Engine/Private/GPUSkinVertexFactory.cpp
@@ -131,12 +131,13 @@ public:
 	virtual void InitRHI(FRHICommandListBase& RHICmdList) override
 	{
 		// Enough data for 64k vertices mesh
-		uint32 Size = sizeof(FMorphGPUSkinVertex) * 65535;
-		FRHIResourceCreateInfo CreateInfo(TEXT("FNullMorphVertexBuffer"));
-		VertexBufferRHI = RHICmdList.CreateBuffer(Size, BUF_Static | BUF_VertexBuffer | BUF_ShaderResource, 0, ERHIAccess::VertexOrIndexBuffer | ERHIAccess::SRVMask, CreateInfo);
-		void* LockedData = RHICmdList.LockBuffer(VertexBufferRHI, 0, Size, RLM_WriteOnly);
-		FMemory::Memzero(LockedData, Size);
-		RHICmdList.UnlockBuffer(VertexBufferRHI);
+		const FRHIBufferCreateDesc CreateDesc =
+			FRHIBufferCreateDesc::CreateVertex<FMorphGPUSkinVertex>(TEXT("FNullMorphVertexBuffer"), 65535)
+			.AddUsage(EBufferUsageFlags::Static | EBufferUsageFlags::ShaderResource)
+			.SetInitialState(ERHIAccess::VertexOrIndexBuffer | ERHIAccess::SRVMask)
+			.SetInitActionZeroData();
+
+		VertexBufferRHI = RHICmdList.CreateBuffer(CreateDesc);
 	}
 };
 
@@ -192,8 +193,13 @@ FVertexBufferAndSRV FBoneBufferPoolPolicy::CreateResource(FRHICommandListBase& R
 	// in VisualStudio the copy constructor call on the return argument can be optimized out
 	// see https://msdn.microsoft.com/en-us/library/ms364057.aspx#nrvo_cpp05_topic3
 	FVertexBufferAndSRV Buffer;
-	FRHIResourceCreateInfo CreateInfo(TEXT("FBoneBufferPoolPolicy"));
-	Buffer.VertexBufferRHI = RHICmdList.CreateVertexBuffer( BufferSize, (BUF_Dynamic | BUF_ShaderResource), CreateInfo );
+
+	const FRHIBufferCreateDesc CreateDesc =
+		FRHIBufferCreateDesc::CreateVertex(TEXT("FBoneBufferPoolPolicy"), BufferSize)
+		.AddUsage(EBufferUsageFlags::Dynamic | EBufferUsageFlags::ShaderResource)
+		.DetermineInitialState();
+
+	Buffer.VertexBufferRHI = RHICmdList.CreateBuffer(CreateDesc);
 	Buffer.VertexBufferSRV = RHICmdList.CreateShaderResourceView(
 		Buffer.VertexBufferRHI, 
 		FRHIViewDesc::CreateBufferSRV()
@@ -217,8 +223,13 @@ FVertexBufferAndSRV FClothBufferPoolPolicy::CreateResource(FRHICommandListBase&
 	// in VisualStudio the copy constructor call on the return argument can be optimized out
 	// see https://msdn.microsoft.com/en-us/library/ms364057.aspx#nrvo_cpp05_topic3
 	FVertexBufferAndSRV Buffer;
-	FRHIResourceCreateInfo CreateInfo(TEXT("FClothBufferPoolPolicy"));
-	Buffer.VertexBufferRHI = RHICmdList.CreateVertexBuffer( BufferSize, (BUF_Dynamic | BUF_ShaderResource), CreateInfo );
+
+	const FRHIBufferCreateDesc CreateDesc =
+		FRHIBufferCreateDesc::CreateVertex(TEXT("FClothBufferPoolPolicy"), BufferSize)
+		.AddUsage(EBufferUsageFlags::Dynamic | EBufferUsageFlags::ShaderResource)
+		.DetermineInitialState();
+
+	Buffer.VertexBufferRHI = RHICmdList.CreateBuffer(CreateDesc);
 	Buffer.VertexBufferSRV = RHICmdList.CreateShaderResourceView(
 		Buffer.VertexBufferRHI, 
 		FRHIViewDesc::CreateBufferSRV()
diff --git a/Engine/Source/Runtime/Engine/Private/GPUSort.cpp b/Engine/Source/Runtime/Engine/Private/GPUSort.cpp
--- a/Engine/Source/Runtime/Engine/Private/GPUSort.cpp
+++ b/Engine/Source/Runtime/Engine/Private/GPUSort.cpp
@@ -88,14 +88,15 @@ public:
 	{
 		const int32 OffsetsCount = DIGIT_COUNT * MAX_GROUP_COUNT;
 		const int32 OffsetsBufferSize = OffsetsCount * sizeof(uint32);
-		
+
+		const FRHIBufferCreateDesc CreateDesc =
+			FRHIBufferCreateDesc::CreateVertex(TEXT("SortOffset"), OffsetsBufferSize)
+			.AddUsage(EBufferUsageFlags::Static | EBufferUsageFlags::ShaderResource | EBufferUsageFlags::UnorderedAccess)
+			.DetermineInitialState();
+
 		for (int32 BufferIndex = 0; BufferIndex < 2; ++BufferIndex)
 		{
-			FRHIResourceCreateInfo CreateInfo(TEXT("SortOffset"));
-			Buffers[BufferIndex] = RHICmdList.CreateVertexBuffer(
-				OffsetsBufferSize,
-				BUF_Static | BUF_ShaderResource | BUF_UnorderedAccess,
-				CreateInfo);
+			Buffers[BufferIndex] = RHICmdList.CreateBuffer(CreateDesc);
 			BufferSRVs[BufferIndex] = RHICmdList.CreateShaderResourceView(
 				Buffers[BufferIndex],
 				FRHIViewDesc::CreateBufferSRV()
@@ -187,12 +188,13 @@ public:
 	 */
 	virtual void InitRHI(FRHICommandListBase& RHICmdList) override
 	{
-			FRHIResourceCreateInfo CreateInfo(TEXT("FRadixSortParametersBuffer"));
-			SortParametersBufferRHI = RHICmdList.CreateVertexBuffer(
-				/*Size=*/ sizeof(FRadixSortParameters),
-				/*Usage=*/ BUF_Volatile | BUF_ShaderResource,
-				CreateInfo);
-			SortParametersBufferSRV = RHICmdList.CreateShaderResourceView(
+		const FRHIBufferCreateDesc CreateDesc =
+			FRHIBufferCreateDesc::CreateVertex(TEXT("FRadixSortParametersBuffer"), sizeof(FRadixSortParameters))
+			.AddUsage(EBufferUsageFlags::Volatile | EBufferUsageFlags::ShaderResource)
+			.DetermineInitialState();
+
+		SortParametersBufferRHI = RHICmdList.CreateBuffer(CreateDesc);
+		SortParametersBufferSRV = RHICmdList.CreateShaderResourceView(
 				SortParametersBufferRHI,
 				FRHIViewDesc::CreateBufferSRV()
 					.SetType(FRHIViewDesc::EBufferType::Typed)
@@ -736,11 +738,36 @@ static bool RunGPUSortTest(FRHICommandListImmediate& RHICmdList, int32 TestSize,
 	RefSortedKeys = Keys;
 	RefSortedKeys.Sort();
 
+	// Upload initial keys and values to the GPU.
+	{
+		const FRHIBufferCreateDesc KeysCreateDesc =
+			FRHIBufferCreateDesc::CreateVertex(TEXT("KeysBuffer"), BufferSize)
+			.AddUsage(EBufferUsageFlags::Static | EBufferUsageFlags::ShaderResource | EBufferUsageFlags::UnorderedAccess)
+			.SetInitActionInitializer()
+			.DetermineInitialState();
+
+		FRHIBufferInitializer Initializer = RHICmdList.CreateBufferInitializer(KeysCreateDesc);
+		Initializer.WriteData(Keys.GetData(), BufferSize);
+		KeysBufferRHI[0] = Initializer.Finalize();
+		KeysBufferRHI[1] = RHICmdList.CreateBuffer(KeysCreateDesc);
+	}
+
+	{
+		const FRHIBufferCreateDesc ValuesCreateDesc =
+			FRHIBufferCreateDesc::CreateVertex(TEXT("ValuesBuffer"), BufferSize)
+			.AddUsage(EBufferUsageFlags::Static | EBufferUsageFlags::ShaderResource | EBufferUsageFlags::UnorderedAccess)
+			.SetInitActionInitializer()
+			.DetermineInitialState();
+
+		FRHIBufferInitializer Initializer = RHICmdList.CreateBufferInitializer(ValuesCreateDesc);
+		Initializer.WriteData(Keys.GetData(), BufferSize);
+		ValuesBufferRHI[0] = Initializer.Finalize();
+		ValuesBufferRHI[1] = RHICmdList.CreateBuffer(ValuesCreateDesc);
+	}
+
 	// Allocate GPU resources.
 	for (int32 BufferIndex = 0; BufferIndex < 2; ++BufferIndex)
 	{
-		FRHIResourceCreateInfo CreateInfo(TEXT("KeysBuffer"));
-		KeysBufferRHI[BufferIndex] = RHICmdList.CreateVertexBuffer(BufferSize, BUF_Static | BUF_ShaderResource | BUF_UnorderedAccess, CreateInfo);
 		KeysBufferSRV[BufferIndex] = RHICmdList.CreateShaderResourceView(
 			KeysBufferRHI[BufferIndex],
 			FRHIViewDesc::CreateBufferSRV()
@@ -751,8 +778,7 @@ static bool RunGPUSortTest(FRHICommandListImmediate& RHICmdList, int32 TestSize,
 			FRHIViewDesc::CreateBufferUAV()
 				.SetType(FRHIViewDesc::EBufferType::Typed)
 				.SetFormat(PF_R32_UINT));
-		CreateInfo.DebugName = TEXT("ValuesBuffer");
-		ValuesBufferRHI[BufferIndex] = RHICmdList.CreateVertexBuffer(BufferSize, BUF_Static | BUF_ShaderResource | BUF_UnorderedAccess, CreateInfo);
+
 		ValuesBufferSRV[BufferIndex] = RHICmdList.CreateShaderResourceView(
 			ValuesBufferRHI[BufferIndex],
 			FRHIViewDesc::CreateBufferSRV()
@@ -765,18 +791,6 @@ static bool RunGPUSortTest(FRHICommandListImmediate& RHICmdList, int32 TestSize,
 				.SetFormat(PF_R32_UINT));
 	}
 
-	// Upload initial keys and values to the GPU.
-	{
-		uint32* Buffer;
-
-		Buffer = (uint32*)RHICmdList.LockBuffer(KeysBufferRHI[0], /*Offset=*/ 0, BufferSize, RLM_WriteOnly);
-		FMemory::Memcpy(Buffer, Keys.GetData(), BufferSize);
-		RHICmdList.UnlockBuffer(KeysBufferRHI[0]);
-		Buffer = (uint32*)RHICmdList.LockBuffer(ValuesBufferRHI[0], /*Offset=*/ 0, BufferSize, RLM_WriteOnly);
-		FMemory::Memcpy(Buffer, Keys.GetData(), BufferSize);
-		RHICmdList.UnlockBuffer(ValuesBufferRHI[0]);
-	}
-
 	// Execute the GPU sort.
 	for (int32 BufferIndex = 0; BufferIndex < 2; ++BufferIndex)
 	{
diff --git a/Engine/Source/Runtime/Engine/Private/GPUSortManager.cpp b/Engine/Source/Runtime/Engine/Private/GPUSortManager.cpp
--- a/Engine/Source/Runtime/Engine/Private/GPUSortManager.cpp
+++ b/Engine/Source/Runtime/Engine/Private/GPUSortManager.cpp
@@ -268,8 +268,12 @@ FGPUSortManager::FValueBuffer::FValueBuffer(FRHICommandListBase& RHICmdList, int
 {
 	check(InUsedCount >= 0 && InUsedCount <= InAllocatedCount);
 
-	FRHIResourceCreateInfo CreateInfo(TEXT("ValueBuffer"));
-	VertexBufferRHI = RHICmdList.CreateVertexBuffer((uint32)InAllocatedCount * sizeof(uint32), BUF_Static | BUF_ShaderResource | BUF_UnorderedAccess, ERHIAccess::SRVGraphics, CreateInfo);
+	const FRHIBufferCreateDesc CreateDesc =
+		FRHIBufferCreateDesc::CreateVertex<uint32>(TEXT("ValueBuffer"), InAllocatedCount)
+		.AddUsage(EBufferUsageFlags::Static | EBufferUsageFlags::ShaderResource | EBufferUsageFlags::UnorderedAccess)
+		.SetInitialState(ERHIAccess::SRVGraphics);
+
+	VertexBufferRHI = RHICmdList.CreateBuffer(CreateDesc);
 
 	UInt32SRV = RHICmdList.CreateShaderResourceView(
 		VertexBufferRHI, 
diff --git a/Engine/Source/Runtime/Engine/Private/GeometryCacheVertexFactory.cpp b/Engine/Source/Runtime/Engine/Private/GeometryCacheVertexFactory.cpp
--- a/Engine/Source/Runtime/Engine/Private/GeometryCacheVertexFactory.cpp
+++ b/Engine/Source/Runtime/Engine/Private/GeometryCacheVertexFactory.cpp
@@ -120,13 +120,17 @@ public:
 
 	virtual void InitRHI(FRHICommandListBase& RHICmdList) override
 	{
-		FRHIResourceCreateInfo CreateInfo(TEXT("DefaultGeometryCacheVertexBuffer"));
-		VertexBufferRHI = RHICmdList.CreateBuffer(sizeof(FVector4f) * 2, BUF_Static | BUF_VertexBuffer | BUF_ShaderResource, 0, ERHIAccess::VertexOrIndexBuffer | ERHIAccess::SRVMask, CreateInfo);
-		FVector4f* DummyContents = (FVector4f*)RHICmdList.LockBuffer(VertexBufferRHI, 0, sizeof(FVector4f) * 2, RLM_WriteOnly);
-		DummyContents[0] = FVector4f(0.0f, 0.0f, 0.0f, 0.0f);
-		DummyContents[1] = FVector4f(1.0f, 1.0f, 1.0f, 1.0f);
-		RHICmdList.UnlockBuffer(VertexBufferRHI);
+		const FRHIBufferCreateDesc CreateDesc =
+			FRHIBufferCreateDesc::CreateVertex<FVector4f>(TEXT("DefaultGeometryCacheVertexBuffer"), 2)
+			.AddUsage(EBufferUsageFlags::Static | EBufferUsageFlags::ShaderResource)
+			.SetInitialState(ERHIAccess::VertexOrIndexBuffer | ERHIAccess::SRVMask)
+			.SetInitActionInitializer();
 
+		TRHIBufferInitializer<FVector4f> Initializer = RHICmdList.CreateBufferInitializer(CreateDesc);
+		Initializer[0] = FVector4f(0.0f, 0.0f, 0.0f, 0.0f);
+		Initializer[1] = FVector4f(1.0f, 1.0f, 1.0f, 1.0f);
+
+		VertexBufferRHI = Initializer.Finalize();
 		SRV = RHICmdList.CreateShaderResourceView(
 			VertexBufferRHI, 
 			FRHIViewDesc::CreateBufferSRV()
@@ -149,13 +153,17 @@ public:
 
 	virtual void InitRHI(FRHICommandListBase& RHICmdList) override
 	{
-		FRHIResourceCreateInfo CreateInfo(TEXT("DummyTangentBuffer"));
-		VertexBufferRHI = RHICmdList.CreateBuffer(sizeof(FVector4f) * 2, BUF_Static | BUF_VertexBuffer | BUF_ShaderResource, 0, ERHIAccess::VertexOrIndexBuffer | ERHIAccess::SRVMask, CreateInfo);
-		FVector4f* DummyContents = (FVector4f*)RHICmdList.LockBuffer(VertexBufferRHI, 0, sizeof(FVector4f) * 2, RLM_WriteOnly);
-		DummyContents[0] = FVector4f(0.0f, 0.0f, 0.0f, 0.0f);
-		DummyContents[1] = FVector4f(1.0f, 1.0f, 1.0f, 1.0f);
-		RHICmdList.UnlockBuffer(VertexBufferRHI);
+		const FRHIBufferCreateDesc CreateDesc =
+			FRHIBufferCreateDesc::CreateVertex<FVector4f>(TEXT("DummyTangentBuffer"), 2)
+			.AddUsage(EBufferUsageFlags::Static | EBufferUsageFlags::ShaderResource)
+			.SetInitialState(ERHIAccess::VertexOrIndexBuffer | ERHIAccess::SRVMask)
+			.SetInitActionInitializer();
+
+		TRHIBufferInitializer<FVector4f> Initializer = RHICmdList.CreateBufferInitializer(CreateDesc);
+		Initializer[0] = FVector4f(0.0f, 0.0f, 0.0f, 0.0f);
+		Initializer[1] = FVector4f(1.0f, 1.0f, 1.0f, 1.0f);
 
+		VertexBufferRHI = Initializer.Finalize();
 		SRV = RHICmdList.CreateShaderResourceView(
 			VertexBufferRHI, 
 			FRHIViewDesc::CreateBufferSRV()
diff --git a/Engine/Source/Runtime/Engine/Private/InstancedStaticMesh.cpp b/Engine/Source/Runtime/Engine/Private/InstancedStaticMesh.cpp
--- a/Engine/Source/Runtime/Engine/Private/InstancedStaticMesh.cpp
+++ b/Engine/Source/Runtime/Engine/Private/InstancedStaticMesh.cpp
@@ -302,15 +302,16 @@ class FDummyFloatBuffer : public FVertexBufferWithSRV
 public:
 	virtual void InitRHI(FRHICommandListBase& RHICmdList) override
 	{
-		// Create the texture RHI.  		
-		FRHIResourceCreateInfo CreateInfo(TEXT("DummyFloatBuffer"));
-
+		// Create the texture RHI.
 		const int32 NumFloats = 4;
-		VertexBufferRHI = RHICmdList.CreateVertexBuffer(sizeof(float)*NumFloats, BUF_Static | BUF_ShaderResource, CreateInfo);
 
-		float* BufferData = (float*)RHICmdList.LockBuffer(VertexBufferRHI, 0, sizeof(float)*NumFloats, RLM_WriteOnly);
-		FMemory::Memzero(BufferData, sizeof(float)*NumFloats);
-		RHICmdList.UnlockBuffer(VertexBufferRHI);
+		const FRHIBufferCreateDesc CreateDesc =
+			FRHIBufferCreateDesc::CreateVertex(TEXT("DummyFloatBuffer"), sizeof(float) * NumFloats)
+			.AddUsage(EBufferUsageFlags::Static | EBufferUsageFlags::ShaderResource)
+			.DetermineInitialState()
+			.SetInitActionZeroData();
+
+		VertexBufferRHI = RHICmdList.CreateBuffer(CreateDesc);
 
 		// Create a view of the buffer
 		ShaderResourceViewRHI = RHICmdList.CreateShaderResourceView(
@@ -486,10 +487,15 @@ void FStaticMeshInstanceBuffer::CreateVertexBuffer(FRHICommandListBase& RHICmdLi
 	check(InResourceArray);
 	check(InResourceArray->GetResourceDataSize() > 0);
 
+	const FRHIBufferCreateDesc CreateDesc =
+		FRHIBufferCreateDesc::CreateVertex(TEXT("FStaticMeshInstanceBuffer"), InResourceArray->GetResourceDataSize())
+		.AddUsage(InUsage)
+		.SetInitActionResourceArray(InResourceArray)
+		.DetermineInitialState();
+
 	// TODO: possibility over allocated the vertex buffer when we support partial update for when working in the editor
-	FRHIResourceCreateInfo CreateInfo(TEXT("FStaticMeshInstanceBuffer"), InResourceArray);
-	OutVertexBufferRHI = RHICmdList.CreateVertexBuffer(InResourceArray->GetResourceDataSize(), InUsage, CreateInfo);
-	
+	OutVertexBufferRHI = RHICmdList.CreateBuffer(CreateDesc);
+
 	if (RHISupportsManualVertexFetch(GMaxRHIShaderPlatform))
 	{
 		OutInstanceSRV = RHICmdList.CreateShaderResourceView(
diff --git a/Engine/Source/Runtime/Engine/Private/MeshParticleVertexFactory.cpp b/Engine/Source/Runtime/Engine/Private/MeshParticleVertexFactory.cpp
--- a/Engine/Source/Runtime/Engine/Private/MeshParticleVertexFactory.cpp
+++ b/Engine/Source/Runtime/Engine/Private/MeshParticleVertexFactory.cpp
@@ -56,8 +56,12 @@ public:
 
 	virtual void InitRHI(FRHICommandListBase& RHICmdList)
 	{
-		FRHIResourceCreateInfo CreateInfo(TEXT("FDummyPrevTransformBuffer"));
-		VB = RHICmdList.CreateVertexBuffer(sizeof(FVector4f) * 3, BUF_Static | BUF_ShaderResource, CreateInfo);
+		const FRHIBufferCreateDesc CreateDesc =
+			FRHIBufferCreateDesc::CreateVertex<FVector4f>(TEXT("FDummyPrevTransformBuffer"), 3)
+			.AddUsage(EBufferUsageFlags::Static | EBufferUsageFlags::ShaderResource)
+			.DetermineInitialState();
+
+		VB = RHICmdList.CreateBuffer(CreateDesc);
 		SRV = RHICmdList.CreateShaderResourceView(
 			VB,
 			FRHIViewDesc::CreateBufferSRV()
diff --git a/Engine/Source/Runtime/Engine/Private/Particles/ParticleCurveTexture.cpp b/Engine/Source/Runtime/Engine/Private/Particles/ParticleCurveTexture.cpp
--- a/Engine/Source/Runtime/Engine/Private/Particles/ParticleCurveTexture.cpp
+++ b/Engine/Source/Runtime/Engine/Private/Particles/ParticleCurveTexture.cpp
@@ -221,9 +221,13 @@ static void InjectCurves(
 
 		// get a buffer for all curve textures at once, and copy curve data over
 		//
-		FRHIResourceCreateInfo CreateInfo(TEXT("ScratchVertexBuffer"));
-		FBufferRHIRef ScratchVertexBufferRHI = RHICmdList.CreateBuffer(TotalSamples * sizeof(FColor), BUF_Volatile | BUF_VertexBuffer, 0, ERHIAccess::VertexOrIndexBuffer, CreateInfo);
-		FColor* RESTRICT DestSamples = (FColor*)RHICmdList.LockBuffer(ScratchVertexBufferRHI, 0, TotalSamples * sizeof(FColor), RLM_WriteOnly);
+		const FRHIBufferCreateDesc CreateDesc =
+			FRHIBufferCreateDesc::CreateVertex<FColor>(TEXT("ScratchVertexBuffer"), TotalSamples)
+			.AddUsage(EBufferUsageFlags::Volatile)
+			.SetInitialState(ERHIAccess::VertexOrIndexBuffer)
+			.SetInitActionInitializer();
+
+		TRHIBufferInitializer<FColor> Initializer = RHICmdList.CreateBufferInitializer(CreateDesc);
 
 		int32 CurrOffset = 0;
 
@@ -233,15 +237,14 @@ static void InjectCurves(
 
 			// Copy curve samples in to the scratch vertex buffer.
 			const int32 SampleCount = Curve.TexelAllocation.Size;
-			FMemory::Memcpy(DestSamples + CurrOffset, Curve.Samples, SampleCount * sizeof(FColor));
+			Initializer.WriteArray(CurrOffset, MakeConstArrayView(Curve.Samples, SampleCount));
 			FMemory::Free(Curve.Samples);
 			Curve.Samples = NULL;
 
 			CurrOffset += SampleCount;
 		}
 
-		RHICmdList.UnlockBuffer(ScratchVertexBufferRHI);
-
+		FBufferRHIRef ScratchVertexBufferRHI = Initializer.Finalize();
 
 		// now draw the curves into the curve texture target, reading from the single buffer we just filled
 		//
diff --git a/Engine/Source/Runtime/Engine/Private/Particles/ParticleGpuSimulation.cpp b/Engine/Source/Runtime/Engine/Private/Particles/ParticleGpuSimulation.cpp
--- a/Engine/Source/Runtime/Engine/Private/Particles/ParticleGpuSimulation.cpp
+++ b/Engine/Source/Runtime/Engine/Private/Particles/ParticleGpuSimulation.cpp
@@ -2289,12 +2289,14 @@ static FBox ComputeParticleBounds(
 		ParticleBoundsParameters.ParticleCount = ParticleCount;
 
 		// Create a buffer for storing bounds.
-		const int32 BufferSize = GroupCount * 2 * sizeof(FVector4f);
-		FRHIResourceCreateInfo CreateInfo(TEXT("BoundsVertexBuffer"));
-		FBufferRHIRef BoundsVertexBufferRHI = RHICmdList.CreateVertexBuffer(
-			BufferSize,
-			BUF_Static | BUF_UnorderedAccess | BUF_KeepCPUAccessible,
-			CreateInfo);
+		const FRHIBufferCreateDesc CreateDesc =
+			FRHIBufferCreateDesc::CreateVertex<FVector4f>(TEXT("BoundsVertexBuffer"), GroupCount * 2)
+			.AddUsage(EBufferUsageFlags::Static | EBufferUsageFlags::UnorderedAccess | EBufferUsageFlags::KeepCPUAccessible)
+			.DetermineInitialState();
+
+		const int32 BufferSize = CreateDesc.Size;
+
+		FBufferRHIRef BoundsVertexBufferRHI = RHICmdList.CreateBuffer(CreateDesc);
 		FUnorderedAccessViewRHIRef BoundsVertexBufferUAV = RHICmdList.CreateUnorderedAccessView(
 			BoundsVertexBufferRHI,
 			FRHIViewDesc::CreateBufferUAV()
@@ -2446,14 +2448,17 @@ public:
 	 */
 	virtual void InitRHI(FRHICommandListBase& RHICmdList) override
 	{
-		if ( AlignedTileCount > 0 )
+		if (AlignedTileCount > 0)
 		{
 			int32 BufferAlignedTileCount = (GMaxRHIFeatureLevel <= ERHIFeatureLevel::ES3_1 ? TileCount : AlignedTileCount);
 			const uint32 Stride = sizeof(FVector4f);
-			const int32 TileBufferSize = BufferAlignedTileCount * Stride;
-			check(TileBufferSize > 0);
-			FRHIResourceCreateInfo CreateInfo(TEXT("FParticleTileVertexBuffer"));
-			VertexBufferRHI = RHICmdList.CreateVertexBuffer( TileBufferSize, BUF_Static | BUF_KeepCPUAccessible | BUF_ShaderResource, CreateInfo );
+
+			const FRHIBufferCreateDesc CreateDesc =
+				FRHIBufferCreateDesc::CreateVertex<FVector4f>(TEXT("FParticleTileVertexBuffer"), BufferAlignedTileCount)
+				.AddUsage(EBufferUsageFlags::Static | EBufferUsageFlags::KeepCPUAccessible | EBufferUsageFlags::ShaderResource)
+				.DetermineInitialState();
+
+			VertexBufferRHI = RHICmdList.CreateBuffer(CreateDesc);
 			VertexBufferSRV = RHICmdList.CreateShaderResourceView(
 				VertexBufferRHI, 
 				FRHIViewDesc::CreateBufferSRV()
@@ -2514,10 +2519,13 @@ public:
 			// Metal *requires* that a buffer be bound - you cannot protect access with a branch in the shader.
 			int32 Count = FMath::Max(ParticleCount, 1);
 			const int32 BufferStride = sizeof(FParticleIndex);
-			const int32 BufferSize = Count * BufferStride;
-			const EBufferUsageFlags Flags = BUF_Static | /*BUF_KeepCPUAccessible | */BUF_ShaderResource;
-			FRHIResourceCreateInfo CreateInfo(TEXT("FGPUParticleVertexBuffer"));
-			VertexBufferRHI = RHICmdList.CreateVertexBuffer(BufferSize, Flags, CreateInfo);
+
+			const FRHIBufferCreateDesc CreateDesc =
+				FRHIBufferCreateDesc::CreateVertex<FParticleIndex>(TEXT("FGPUParticleVertexBuffer"), Count)
+				.AddUsage(EBufferUsageFlags::Static | /*EBufferUsageFlags::KeepCPUAccessible |*/ EBufferUsageFlags::ShaderResource)
+				.DetermineInitialState();
+
+			VertexBufferRHI = RHICmdList.CreateBuffer(CreateDesc);
 			VertexBufferSRV = RHICmdList.CreateShaderResourceView(
 				VertexBufferRHI, 
 				FRHIViewDesc::CreateBufferSRV()
diff --git a/Engine/Source/Runtime/Engine/Private/Particles/ParticleResources.cpp b/Engine/Source/Runtime/Engine/Private/Particles/ParticleResources.cpp
--- a/Engine/Source/Runtime/Engine/Private/Particles/ParticleResources.cpp
+++ b/Engine/Source/Runtime/Engine/Private/Particles/ParticleResources.cpp
@@ -18,10 +18,14 @@ const int32 GParticleScratchVertexBufferSize = 64 * (1 << 10); // 64KB
  */
 void FParticleTexCoordVertexBuffer::InitRHI(FRHICommandListBase& RHICmdList)
 {
-	const uint32 Size = sizeof(FVector2f) * 4 * MAX_PARTICLES_PER_INSTANCE;
-	FRHIResourceCreateInfo CreateInfo(TEXT("FParticleTexCoordVertexBuffer"));
-	VertexBufferRHI = RHICmdList.CreateBuffer(Size, BUF_Static | BUF_VertexBuffer, 0, ERHIAccess::VertexOrIndexBuffer, CreateInfo);
-	FVector2f* Vertices = (FVector2f*)RHICmdList.LockBuffer(VertexBufferRHI, 0, Size, RLM_WriteOnly);
+	const FRHIBufferCreateDesc CreateDesc =
+		FRHIBufferCreateDesc::CreateVertex<FVector2f>(TEXT("FParticleTexCoordVertexBuffer"), 4 * MAX_PARTICLES_PER_INSTANCE)
+		.AddUsage(EBufferUsageFlags::Static)
+		.SetInitialState(ERHIAccess::VertexOrIndexBuffer)
+		.SetInitActionInitializer();
+
+	TRHIBufferInitializer<FVector2f> Vertices = RHICmdList.CreateBufferInitializer(CreateDesc);
+
 	for (uint32 SpriteIndex = 0; SpriteIndex < MAX_PARTICLES_PER_INSTANCE; ++SpriteIndex)
 	{
 		Vertices[SpriteIndex*4 + 0] = FVector2f(0.0f, 0.0f);
@@ -29,7 +33,8 @@ void FParticleTexCoordVertexBuffer::InitRHI(FRHICommandListBase& RHICmdList)
 		Vertices[SpriteIndex*4 + 2] = FVector2f(1.0f, 1.0f);
 		Vertices[SpriteIndex*4 + 3] = FVector2f(1.0f, 0.0f);
 	}
-	RHICmdList.UnlockBuffer( VertexBufferRHI );
+
+	VertexBufferRHI = Vertices.Finalize();
 }
 
 /** Global particle texture coordinate vertex buffer. */
@@ -40,10 +45,14 @@ TGlobalResource<FParticleTexCoordVertexBuffer> GParticleTexCoordVertexBuffer;
  */
 void FParticleEightTexCoordVertexBuffer::InitRHI(FRHICommandListBase& RHICmdList)
 {
-	const uint32 Size = sizeof(FVector2f) * 8 * MAX_PARTICLES_PER_INSTANCE;
-	FRHIResourceCreateInfo CreateInfo(TEXT("FParticleEightTexCoordVertexBuffer"));
-	VertexBufferRHI = RHICmdList.CreateBuffer(Size, BUF_Static | BUF_VertexBuffer, 0, ERHIAccess::VertexOrIndexBuffer, CreateInfo);
-	FVector2f* Vertices = (FVector2f*)RHICmdList.LockBuffer(VertexBufferRHI, 0, Size, RLM_WriteOnly);
+	const FRHIBufferCreateDesc CreateDesc =
+		FRHIBufferCreateDesc::CreateVertex<FVector2f>(TEXT("FParticleEightTexCoordVertexBuffer"), 8 * MAX_PARTICLES_PER_INSTANCE)
+		.AddUsage(EBufferUsageFlags::Static)
+		.SetInitialState(ERHIAccess::VertexOrIndexBuffer)
+		.SetInitActionInitializer();
+
+	TRHIBufferInitializer<FVector2f> Vertices = RHICmdList.CreateBufferInitializer(CreateDesc);
+
 	for (uint32 SpriteIndex = 0; SpriteIndex < MAX_PARTICLES_PER_INSTANCE; ++SpriteIndex)
 	{
 		// The contents of this buffer does not matter, whenever it is used, cutout geometry will override
@@ -56,7 +65,8 @@ void FParticleEightTexCoordVertexBuffer::InitRHI(FRHICommandListBase& RHICmdList
 		Vertices[SpriteIndex*8 + 6] = FVector2f(1.0f, 0.0f);
 		Vertices[SpriteIndex*8 + 7] = FVector2f(1.0f, 0.0f);
 	}
-	RHICmdList.UnlockBuffer( VertexBufferRHI );
+
+	VertexBufferRHI = Vertices.Finalize();
 }
 
 /** Global particle texture coordinate vertex buffer. */
@@ -71,11 +81,15 @@ void FParticleIndexBuffer::InitRHI(FRHICommandListBase& RHICmdList)
 	// but using the maximum needed for the non-instanced path
 	// in prep for future flipping of both instanced and non-instanced at runtime.
 	const uint32 MaxParticles = 65536 / 4;
-	const uint32 Size = sizeof(uint16) * 6 * MaxParticles;
-	const uint32 Stride = sizeof(uint16);
-	FRHIResourceCreateInfo CreateInfo(TEXT("FParticleIndexBuffer"));
-	IndexBufferRHI = RHICmdList.CreateBuffer( Size, BUF_Static | BUF_IndexBuffer, Stride, ERHIAccess::VertexOrIndexBuffer, CreateInfo );
-	uint16* Indices = (uint16*)RHICmdList.LockBuffer( IndexBufferRHI, 0, Size, RLM_WriteOnly );
+
+	const FRHIBufferCreateDesc CreateDesc =
+		FRHIBufferCreateDesc::CreateIndex<uint16>(TEXT("FParticleIndexBuffer"), 6 * MaxParticles)
+		.AddUsage(EBufferUsageFlags::Static)
+		.SetInitialState(ERHIAccess::VertexOrIndexBuffer)
+		.SetInitActionInitializer();
+
+	TRHIBufferInitializer<uint16> Indices = RHICmdList.CreateBufferInitializer(CreateDesc);
+
 	for (uint32 SpriteIndex = 0; SpriteIndex < MaxParticles; ++SpriteIndex)
 	{
 		Indices[SpriteIndex*6 + 0] = SpriteIndex*4 + 0;
@@ -85,7 +99,8 @@ void FParticleIndexBuffer::InitRHI(FRHICommandListBase& RHICmdList)
 		Indices[SpriteIndex*6 + 4] = SpriteIndex*4 + 1;
 		Indices[SpriteIndex*6 + 5] = SpriteIndex*4 + 2;
 	}
-	RHICmdList.UnlockBuffer( IndexBufferRHI );
+
+	IndexBufferRHI = Indices.Finalize();
 }
 
 /** Global particle index buffer. */
@@ -100,11 +115,15 @@ void FSixTriangleParticleIndexBuffer::InitRHI(FRHICommandListBase& RHICmdList)
 	// but using the maximum needed for the non-instanced path
 	// in prep for future flipping of both instanced and non-instanced at runtime.
 	const uint32 MaxParticles = 65536 / 8;
-	const uint32 Size = sizeof(uint16) * 6 * 3 * MaxParticles;
-	const uint32 Stride = sizeof(uint16);
-	FRHIResourceCreateInfo CreateInfo(TEXT("FSixTriangleParticleIndexBuffer"));
-	IndexBufferRHI = RHICmdList.CreateBuffer( Size, BUF_Static | BUF_IndexBuffer, Stride, ERHIAccess::VertexOrIndexBuffer, CreateInfo );
-	uint16* Indices = (uint16*)RHICmdList.LockBuffer( IndexBufferRHI, 0, Size, RLM_WriteOnly );
+
+	const FRHIBufferCreateDesc CreateDesc =
+		FRHIBufferCreateDesc::CreateIndex<uint16>(TEXT("FSixTriangleParticleIndexBuffer"), 6 * 3 * MaxParticles)
+		.AddUsage(EBufferUsageFlags::Static)
+		.SetInitialState(ERHIAccess::VertexOrIndexBuffer)
+		.SetInitActionInitializer();
+
+	TRHIBufferInitializer<uint16> Indices = RHICmdList.CreateBufferInitializer(CreateDesc);
+
 	for (uint32 SpriteIndex = 0; SpriteIndex < MaxParticles; ++SpriteIndex)
 	{
 		Indices[SpriteIndex*18 + 0] = SpriteIndex*8 + 0;
@@ -128,7 +147,8 @@ void FSixTriangleParticleIndexBuffer::InitRHI(FRHICommandListBase& RHICmdList)
 		Indices[SpriteIndex*18 + 16] = SpriteIndex*8 + 6;
 		Indices[SpriteIndex*18 + 17] = SpriteIndex*8 + 7;
 	}
-	RHICmdList.UnlockBuffer( IndexBufferRHI );
+
+	IndexBufferRHI = Indices.Finalize();
 }
 
 /** Global particle index buffer. */
@@ -140,10 +160,12 @@ TGlobalResource<FSixTriangleParticleIndexBuffer> GSixTriangleParticleIndexBuffer
 void FParticleScratchVertexBuffer::InitRHI(FRHICommandListBase& RHICmdList)
 {
 	// Create a scratch vertex buffer for injecting particles and rendering tiles.
-	EBufferUsageFlags Flags = BUF_Volatile | BUF_ShaderResource;
+	const FRHIBufferCreateDesc CreateDesc =
+		FRHIBufferCreateDesc::CreateVertex(TEXT("FParticleScratchVertexBuffer"), GParticleScratchVertexBufferSize)
+		.AddUsage(EBufferUsageFlags::Volatile | EBufferUsageFlags::ShaderResource)
+		.DetermineInitialState();
 
-	FRHIResourceCreateInfo CreateInfo(TEXT("FParticleScratchVertexBuffer"));
-	VertexBufferRHI = RHICmdList.CreateVertexBuffer(GParticleScratchVertexBufferSize, Flags, CreateInfo);
+	VertexBufferRHI = RHICmdList.CreateBuffer(CreateDesc);
 	VertexBufferSRV_A32B32G32R32F = RHICmdList.CreateShaderResourceView(
 		VertexBufferRHI, 
 		FRHIViewDesc::CreateBufferSRV()
diff --git a/Engine/Source/Runtime/Engine/Private/Particles/ParticleSortingGPU.cpp b/Engine/Source/Runtime/Engine/Private/Particles/ParticleSortingGPU.cpp
--- a/Engine/Source/Runtime/Engine/Private/Particles/ParticleSortingGPU.cpp
+++ b/Engine/Source/Runtime/Engine/Private/Particles/ParticleSortingGPU.cpp
@@ -188,33 +188,43 @@ void FParticleSortBuffers::InitRHI(FRHICommandListBase& RHICmdList)
 {
 	for (int32 BufferIndex = 0; BufferIndex < 2; ++BufferIndex)
 	{
-		FRHIResourceCreateInfo CreateInfo(TEXT("PartialSortKeyBuffer"));
-
-		KeyBuffers[BufferIndex] = RHICmdList.CreateVertexBuffer( BufferSize * sizeof(uint32), BUF_Static | BUF_ShaderResource | BUF_UnorderedAccess, CreateInfo);
-		KeyBufferSRVs[BufferIndex] = RHICmdList.CreateShaderResourceView(
-			KeyBuffers[BufferIndex],
-			FRHIViewDesc::CreateBufferSRV()
-				.SetType(FRHIViewDesc::EBufferType::Typed)
-				.SetFormat(PF_R32_UINT));
-		KeyBufferUAVs[BufferIndex] = RHICmdList.CreateUnorderedAccessView(
-			KeyBuffers[BufferIndex],
-			FRHIViewDesc::CreateBufferUAV()
-				.SetType(FRHIViewDesc::EBufferType::Typed)
-				.SetFormat(PF_R32_UINT));
-
-		CreateInfo.DebugName = TEXT("PartialSortVertexBuffer");
-		VertexBuffers[BufferIndex] = RHICmdList.CreateVertexBuffer( BufferSize * sizeof(uint32), BUF_Static | BUF_ShaderResource | BUF_UnorderedAccess, CreateInfo);
-
-		VertexBufferSortSRVs[BufferIndex] = RHICmdList.CreateShaderResourceView(
-			VertexBuffers[BufferIndex],
-			FRHIViewDesc::CreateBufferSRV()
-				.SetType(FRHIViewDesc::EBufferType::Typed)
-				.SetFormat(PF_R32_UINT));
-		VertexBufferSortUAVs[BufferIndex] = RHICmdList.CreateUnorderedAccessView(
-			VertexBuffers[BufferIndex],
-			FRHIViewDesc::CreateBufferUAV()
-				.SetType(FRHIViewDesc::EBufferType::Typed)
-				.SetFormat(PF_R32_UINT));
+		{
+			const FRHIBufferCreateDesc CreateDesc =
+				FRHIBufferCreateDesc::CreateVertex<uint32>(TEXT("PartialSortKeyBuffer"), BufferSize)
+				.AddUsage(EBufferUsageFlags::Static | EBufferUsageFlags::ShaderResource | EBufferUsageFlags::UnorderedAccess)
+				.DetermineInitialState();
+
+			KeyBuffers[BufferIndex] = RHICmdList.CreateBuffer(CreateDesc);
+			KeyBufferSRVs[BufferIndex] = RHICmdList.CreateShaderResourceView(
+				KeyBuffers[BufferIndex],
+				FRHIViewDesc::CreateBufferSRV()
+					.SetType(FRHIViewDesc::EBufferType::Typed)
+					.SetFormat(PF_R32_UINT));
+			KeyBufferUAVs[BufferIndex] = RHICmdList.CreateUnorderedAccessView(
+				KeyBuffers[BufferIndex],
+				FRHIViewDesc::CreateBufferUAV()
+					.SetType(FRHIViewDesc::EBufferType::Typed)
+					.SetFormat(PF_R32_UINT));
+		}
+
+		{
+			const FRHIBufferCreateDesc CreateDesc =
+				FRHIBufferCreateDesc::CreateVertex<uint32>(TEXT("PartialSortVertexBuffer"), BufferSize)
+				.AddUsage(EBufferUsageFlags::Static | EBufferUsageFlags::ShaderResource | EBufferUsageFlags::UnorderedAccess)
+				.DetermineInitialState();
+
+			VertexBuffers[BufferIndex] = RHICmdList.CreateBuffer(CreateDesc);
+			VertexBufferSortSRVs[BufferIndex] = RHICmdList.CreateShaderResourceView(
+				VertexBuffers[BufferIndex],
+				FRHIViewDesc::CreateBufferSRV()
+					.SetType(FRHIViewDesc::EBufferType::Typed)
+					.SetFormat(PF_R32_UINT));
+			VertexBufferSortUAVs[BufferIndex] = RHICmdList.CreateUnorderedAccessView(
+				VertexBuffers[BufferIndex],
+				FRHIViewDesc::CreateBufferUAV()
+					.SetType(FRHIViewDesc::EBufferType::Typed)
+					.SetFormat(PF_R32_UINT));
+		}
 	}
 }
 
diff --git a/Engine/Source/Runtime/Engine/Private/Particles/ParticleVertexFactory.cpp b/Engine/Source/Runtime/Engine/Private/Particles/ParticleVertexFactory.cpp
--- a/Engine/Source/Runtime/Engine/Private/Particles/ParticleVertexFactory.cpp
+++ b/Engine/Source/Runtime/Engine/Private/Particles/ParticleVertexFactory.cpp
@@ -24,15 +24,14 @@ public:
 	 */
 	virtual void InitRHI(FRHICommandListBase& RHICmdList) override
 	{
-		const uint32 Size = sizeof(FVector2f) * 4;
-
 		// create a static vertex buffer
-		FRHIResourceCreateInfo CreateInfo(TEXT("FNullSubUVCutoutVertexBuffer"));
-		VertexBufferRHI = RHICmdList.CreateBuffer(Size, BUF_Static | BUF_VertexBuffer | BUF_ShaderResource, 0, ERHIAccess::VertexOrIndexBuffer | ERHIAccess::SRVMask, CreateInfo);
-		void* BufferData = RHICmdList.LockBuffer(VertexBufferRHI, 0, Size, RLM_WriteOnly);
-		FMemory::Memzero(BufferData, Size);
-		RHICmdList.UnlockBuffer(VertexBufferRHI);
+		const FRHIBufferCreateDesc CreateDesc =
+			FRHIBufferCreateDesc::CreateVertex<FVector2f>(TEXT("FNullSubUVCutoutVertexBuffer"), 4)
+			.AddUsage(EBufferUsageFlags::Static | EBufferUsageFlags::ShaderResource)
+			.SetInitialState(ERHIAccess::VertexOrIndexBuffer | ERHIAccess::SRVMask)
+			.SetInitActionZeroData();
 
+		VertexBufferRHI = RHICmdList.CreateBuffer(CreateDesc);
 		VertexBufferSRV = RHICmdList.CreateShaderResourceView(
 			VertexBufferRHI,
 			FRHIViewDesc::CreateBufferSRV()
diff --git a/Engine/Source/Runtime/Engine/Private/PrecomputedVolumetricLightmap.cpp b/Engine/Source/Runtime/Engine/Private/PrecomputedVolumetricLightmap.cpp
--- a/Engine/Source/Runtime/Engine/Private/PrecomputedVolumetricLightmap.cpp
+++ b/Engine/Source/Runtime/Engine/Private/PrecomputedVolumetricLightmap.cpp
@@ -270,8 +270,13 @@ ENGINE_API void FPrecomputedVolumetricLightmapData::InitRHIForSubLevelResources(
 		IndirectionTextureOriginalValues.SetAllowCPUAccess(true);
 
 		{
-			FRHIResourceCreateInfo CreateInfo(TEXT("SubLevelBrickPositionsBuffer"), &SubLevelBrickPositions);
-			SubLevelBrickPositionsBuffer = RHICmdList.CreateVertexBuffer(SubLevelBrickPositions.Num() * SubLevelBrickPositions.GetTypeSize(), BUF_Static | BUF_ShaderResource, CreateInfo);
+			const FRHIBufferCreateDesc CreateDesc =
+				FRHIBufferCreateDesc::CreateVertex(TEXT("SubLevelBrickPositionsBuffer"), SubLevelBrickPositions.Num() * SubLevelBrickPositions.GetTypeSize())
+				.AddUsage(BUF_Static | BUF_ShaderResource)
+				.SetInitActionResourceArray(&SubLevelBrickPositions)
+				.DetermineInitialState();
+
+			SubLevelBrickPositionsBuffer = RHICmdList.CreateBuffer(CreateDesc);
 			SubLevelBrickPositionsSRV = RHICmdList.CreateShaderResourceView(
 				SubLevelBrickPositionsBuffer, 
 				FRHIViewDesc::CreateBufferSRV()
@@ -280,8 +285,13 @@ ENGINE_API void FPrecomputedVolumetricLightmapData::InitRHIForSubLevelResources(
 		}
 
 		{
-			FRHIResourceCreateInfo CreateInfo(TEXT("IndirectionTextureOriginalValuesBuffer"), &IndirectionTextureOriginalValues);
-			IndirectionTextureOriginalValuesBuffer = RHICmdList.CreateVertexBuffer(IndirectionTextureOriginalValues.Num() * IndirectionTextureOriginalValues.GetTypeSize(), BUF_Static | BUF_ShaderResource, CreateInfo);
+			const FRHIBufferCreateDesc CreateDesc =
+				FRHIBufferCreateDesc::CreateVertex(TEXT("IndirectionTextureOriginalValuesBuffer"), IndirectionTextureOriginalValues.Num() * IndirectionTextureOriginalValues.GetTypeSize())
+				.AddUsage(BUF_Static | BUF_ShaderResource)
+				.SetInitActionResourceArray(&IndirectionTextureOriginalValues)
+				.DetermineInitialState();
+
+			IndirectionTextureOriginalValuesBuffer = RHICmdList.CreateBuffer(CreateDesc);
 			IndirectionTextureOriginalValuesSRV = RHICmdList.CreateShaderResourceView(
 				IndirectionTextureOriginalValuesBuffer, 
 				FRHIViewDesc::CreateBufferSRV()
diff --git a/Engine/Source/Runtime/Engine/Private/PrimitiveUniformShaderParameters.cpp b/Engine/Source/Runtime/Engine/Private/PrimitiveUniformShaderParameters.cpp
--- a/Engine/Source/Runtime/Engine/Private/PrimitiveUniformShaderParameters.cpp
+++ b/Engine/Source/Runtime/Engine/Private/PrimitiveUniformShaderParameters.cpp
@@ -50,9 +50,14 @@ void FSinglePrimitiveStructured::InitRHI(FRHICommandListBase& RHICmdList)
 		PrimitiveSceneDataTextureSRV = RHICmdList.CreateShaderResourceView(PrimitiveSceneDataTextureRHI, FRHIViewDesc::CreateTextureSRV().SetDimensionFromTexture(PrimitiveSceneDataTextureRHI));
 	}
 
-	FRHIResourceCreateInfo CreateInfo(TEXT("SkyIrradianceEnvironmentMap"));
-	SkyIrradianceEnvironmentMapRHI = RHICmdList.CreateStructuredBuffer(sizeof(FVector4f), sizeof(FVector4f) * 8, BUF_Static | BUF_ShaderResource, CreateInfo);
-	SkyIrradianceEnvironmentMapSRV = RHICmdList.CreateShaderResourceView(SkyIrradianceEnvironmentMapRHI, FRHIViewDesc::CreateBufferSRV().SetTypeFromBuffer(SkyIrradianceEnvironmentMapRHI));
+	{
+		const FRHIBufferCreateDesc Desc =
+			FRHIBufferCreateDesc::CreateStructured<FVector4f>(TEXT("SkyIrradianceEnvironmentMap"), 8)
+			.AddUsage(BUF_Static | BUF_ShaderResource)
+			.DetermineInitialState();
+		SkyIrradianceEnvironmentMapRHI = RHICmdList.CreateBuffer(Desc);
+		SkyIrradianceEnvironmentMapSRV = RHICmdList.CreateShaderResourceView(SkyIrradianceEnvironmentMapRHI, FRHIViewDesc::CreateBufferSRV().SetTypeFromBuffer(SkyIrradianceEnvironmentMapRHI));
+	}
 
 	UploadToGPU(RHICmdList);
 }
diff --git a/Engine/Source/Runtime/Engine/Private/RawIndexBuffer.cpp b/Engine/Source/Runtime/Engine/Private/RawIndexBuffer.cpp
--- a/Engine/Source/Runtime/Engine/Private/RawIndexBuffer.cpp
+++ b/Engine/Source/Runtime/Engine/Private/RawIndexBuffer.cpp
@@ -42,16 +42,19 @@ void FRawIndexBuffer::CacheOptimize()
 void FRawIndexBuffer::InitRHI(FRHICommandListBase& RHICmdList)
 {
 	uint32 Size = Indices.Num() * sizeof(uint16);
-	if( Size > 0 )
+	if (Indices.Num() > 0)
 	{
 		// Create the index buffer.
-		FRHIResourceCreateInfo CreateInfo(TEXT("FRawIndexBuffer"));
-		IndexBufferRHI = RHICmdList.CreateBuffer(Size, BUF_Static | BUF_IndexBuffer, sizeof(uint16), ERHIAccess::VertexOrIndexBuffer | ERHIAccess::SRVMask, CreateInfo);
+		const FRHIBufferCreateDesc CreateDesc =
+			FRHIBufferCreateDesc::CreateIndex<uint16>(TEXT("FRawIndexBuffer"), Indices.Num())
+			.AddUsage(EBufferUsageFlags::Static)
+			.SetInitialState(ERHIAccess::VertexOrIndexBuffer | ERHIAccess::SRVMask)
+			.SetInitActionInitializer();
 
-		// Initialize the buffer.
-		void* Buffer = RHICmdList.LockBuffer(IndexBufferRHI, 0, Size, RLM_WriteOnly);
-		FMemory::Memcpy(Buffer, Indices.GetData(), Size);
-		RHICmdList.UnlockBuffer(IndexBufferRHI);
+		TRHIBufferInitializer<uint16> Initializer = RHICmdList.CreateBufferInitializer(CreateDesc);
+		Initializer.WriteArray(MakeConstArrayView(Indices));
+
+		IndexBufferRHI = Initializer.Finalize();
 	}
 }
 
@@ -105,26 +108,30 @@ void FRawIndexBuffer16or32::InitRHI(FRHICommandListBase& RHICmdList)
 	if (Size > 0)
 	{
 		// Create the index buffer.
-		FRHIResourceCreateInfo CreateInfo(TEXT("FRawIndexBuffer"));
-		IndexBufferRHI = RHICmdList.CreateBuffer(Size, BUF_Static | BUF_IndexBuffer, IndexStride, ERHIAccess::VertexOrIndexBuffer | ERHIAccess::SRVMask, CreateInfo);
+		const FRHIBufferCreateDesc CreateDesc =
+			FRHIBufferCreateDesc::CreateIndex(TEXT("FRawIndexBuffer"), Size, IndexStride)
+			.AddUsage(EBufferUsageFlags::Static)
+			.SetInitialState(ERHIAccess::VertexOrIndexBuffer | ERHIAccess::SRVMask)
+			.SetInitActionInitializer();
 
 		// Initialize the buffer.		
-		void* Buffer = RHICmdList.LockBuffer(IndexBufferRHI, 0, Size, RLM_WriteOnly);
-
 		if (b32Bit)
 		{
-			FMemory::Memcpy(Buffer, Indices.GetData(), Size);
+			TRHIBufferInitializer<uint32> Initializer = RHICmdList.CreateBufferInitializer(CreateDesc);
+			Initializer.WriteArray(Indices);
+			IndexBufferRHI = Initializer.Finalize();
 		}
 		else
 		{
-			uint16* DestIndices16Bit = (uint16*)Buffer;
-			for (int32 i = 0; i < NumIndices; ++i)
+			TRHIBufferInitializer<uint16> Initializer = RHICmdList.CreateBufferInitializer(CreateDesc);
+
+			for (int32 Index = 0; Index < NumIndices; ++Index)
 			{
-				DestIndices16Bit[i] = Indices[i];
+				Initializer[Index] = Indices[Index];
 			}
-		}
 
-		RHICmdList.UnlockBuffer(IndexBufferRHI);
+			IndexBufferRHI = Initializer.Finalize();
+		}
 	}
 
 	// Undo/redo can destroy and recreate the render resources for UModels without rebuilding the
@@ -341,15 +348,20 @@ FBufferRHIRef FRawStaticIndexBuffer::CreateRHIBuffer(FRHICommandListBase& RHICmd
 		const static FLazyName ClassName32(TEXT("FRawStaticIndexBuffer32"));
 		const static FLazyName ClassName16(TEXT("FRawStaticIndexBuffer16"));
 
+		const TCHAR* BufferName = Is32Bit() ? TEXT("FRawStaticIndexBuffer32") : TEXT("FRawStaticIndexBuffer16");
+
 		// Create the index buffer.
-		FRHIResourceCreateInfo CreateInfo(Is32Bit() ? TEXT("FRawStaticIndexBuffer32") : TEXT("FRawStaticIndexBuffer16"), &IndexStorage);
-		CreateInfo.ClassName = Is32Bit() ? ClassName32 : ClassName16;
-		CreateInfo.OwnerName = GetOwnerName();
-		CreateInfo.bWithoutNativeResource = !SizeInBytes;
-
-		FBufferRHIRef Buffer = RHICmdList.CreateBuffer(SizeInBytes, BufferFlags | EBufferUsageFlags::IndexBuffer, IndexStride, ERHIAccess::VertexOrIndexBuffer | ERHIAccess::SRVMask, CreateInfo);
-		Buffer->SetOwnerName(GetOwnerName());
-		return Buffer;
+		FRHIBufferCreateDesc Desc =
+			SizeInBytes
+			? (FRHIBufferCreateDesc::CreateIndex(BufferName, SizeInBytes, IndexStride).AddUsage(BufferFlags))
+			: FRHIBufferCreateDesc::CreateNull(BufferName);
+
+		Desc.SetClassName(Is32Bit() ? ClassName32 : ClassName16);
+		Desc.SetOwnerName(GetOwnerName());
+		Desc.SetInitialState(ERHIAccess::VertexOrIndexBuffer | ERHIAccess::SRVMask);
+		Desc.SetInitActionResourceArray(&IndexStorage);
+
+		return RHICmdList.CreateBuffer(Desc);
 	}
 	return nullptr;
 }
@@ -469,24 +481,28 @@ FBufferRHIRef FRawStaticIndexBuffer16or32Interface::CreateRHIIndexBufferInternal
 	bool bNeedSRV
 )
 {
-	// Create the index buffer.
-	FRHIResourceCreateInfo CreateInfo(InDebugName, ResourceArray);
-	CreateInfo.ClassName = InDebugName;
-	CreateInfo.OwnerName = InOwnerName;
-	EBufferUsageFlags Flags = EBufferUsageFlags::Static;
+	const uint32 Size = IndexCount * IndexSize;
+
+	FRHIBufferCreateDesc CreateDesc = Size
+		? (FRHIBufferCreateDesc::CreateIndex(InDebugName, Size, IndexSize).AddUsage(EBufferUsageFlags::Static))
+		: FRHIBufferCreateDesc::CreateNull(InDebugName);
+
+	CreateDesc.SetClassName(InDebugName)
+		.SetOwnerName(InOwnerName)
+		.SetInitialState(ERHIAccess::VertexOrIndexBuffer | ERHIAccess::SRVMask);
 
-	if (bNeedSRV)
+	if (!CreateDesc.IsNull() && bNeedSRV)
 	{
-		// BUF_ShaderResource is needed for SkinCache RecomputeSkinTangents
-		Flags |= EBufferUsageFlags::ShaderResource;
+		// EBufferUsageFlags::ShaderResource is needed for SkinCache RecomputeSkinTangents
+		CreateDesc.AddUsage(EBufferUsageFlags::ShaderResource);
 	}
 
-	const uint32 Size = IndexCount * IndexSize;
-	CreateInfo.bWithoutNativeResource = !Size;
+	if (ResourceArray)
+	{
+		CreateDesc.SetInitActionResourceArray(ResourceArray);
+	}
 
-	FBufferRHIRef Buffer = RHICmdList.CreateBuffer(Size, Flags | EBufferUsageFlags::IndexBuffer, IndexSize, ERHIAccess::VertexOrIndexBuffer | ERHIAccess::SRVMask, CreateInfo);
-	Buffer->SetOwnerName(InOwnerName);
-	return Buffer;
+	return RHICmdList.CreateBuffer(CreateDesc);
 }
 
 /*-----------------------------------------------------------------------------
diff --git a/Engine/Source/Runtime/Engine/Private/Rendering/MorphTargetVertexInfoBuffers.cpp b/Engine/Source/Runtime/Engine/Private/Rendering/MorphTargetVertexInfoBuffers.cpp
--- a/Engine/Source/Runtime/Engine/Private/Rendering/MorphTargetVertexInfoBuffers.cpp
+++ b/Engine/Source/Runtime/Engine/Private/Rendering/MorphTargetVertexInfoBuffers.cpp
@@ -5,6 +5,7 @@
 #include "Rendering/SkeletalMeshLODRenderData.h"
 #include "Animation/MorphTarget.h"
 #include "DataDrivenShaderPlatformInfo.h"
+#include "RHIResourceUtils.h"
 
 extern int32 GSkinCacheRecomputeTangents;
 
@@ -20,16 +21,15 @@ void FMorphTargetVertexInfoBuffers::InitRHI(FRHICommandListBase& RHICmdList)
 
 	const static FLazyName ClassName(TEXT("FMorphTargetVertexInfoBuffers"));
 
-	const uint32 BufferSize = MorphData.Num() * sizeof(uint32);
-	FRHIResourceCreateInfo CreateInfo(TEXT("MorphData"));
-	CreateInfo.ClassName = ClassName;
-	CreateInfo.OwnerName = GetOwnerName();
-	MorphDataBuffer = RHICmdList.CreateStructuredBuffer(sizeof(uint32), BufferSize, BUF_Static | BUF_ByteAddressBuffer | BUF_ShaderResource, ERHIAccess::SRVMask, CreateInfo);
-	MorphDataBuffer->SetOwnerName(GetOwnerName());
-	
-	void* BufferPtr = RHICmdList.LockBuffer(MorphDataBuffer, 0, BufferSize, RLM_WriteOnly);
-	FMemory::ParallelMemcpy(BufferPtr, MorphData.GetData(), BufferSize, EMemcpyCachePolicy::StoreUncached);
-	RHICmdList.UnlockBuffer(MorphDataBuffer);
+	const FRHIBufferCreateDesc Desc =
+		FRHIBufferCreateDesc::CreateStructured<uint32>(TEXT("MorphData"), MorphData.Num())
+		.AddUsage(BUF_Static | BUF_ByteAddressBuffer | BUF_ShaderResource)
+		.SetInitialState(ERHIAccess::SRVMask)
+		.SetClassName(ClassName)
+		.SetOwnerName(GetOwnerName())
+		.SetInitActionResourceArray(&MorphData);
+
+	MorphDataBuffer = RHICmdList.CreateBuffer(Desc);
 	MorphDataSRV = RHICmdList.CreateShaderResourceView(MorphDataBuffer, FRHIViewDesc::CreateBufferSRV().SetTypeFromBuffer(MorphDataBuffer));
 
 	if (bEmptyMorphCPUDataOnInitRHI)
diff --git a/Engine/Source/Runtime/Engine/Private/Rendering/SkeletalMeshAttributeVertexBuffer.cpp b/Engine/Source/Runtime/Engine/Private/Rendering/SkeletalMeshAttributeVertexBuffer.cpp
--- a/Engine/Source/Runtime/Engine/Private/Rendering/SkeletalMeshAttributeVertexBuffer.cpp
+++ b/Engine/Source/Runtime/Engine/Private/Rendering/SkeletalMeshAttributeVertexBuffer.cpp
@@ -174,8 +174,14 @@ int32 FSkeletalMeshAttributeVertexBuffer::GetResourceSize() const
 void FSkeletalMeshAttributeVertexBuffer::InitRHI(FRHICommandListBase& RHICmdList)
 {
 	FResourceArrayInterface* ResourceArray = ValueData->GetResourceArray();
-	FRHIResourceCreateInfo CreateInfo(TEXT("SkeletalMeshAttributeVertexBuffer"), ResourceArray);
-	ValueBuffer.VertexBufferRHI = RHICmdList.CreateVertexBuffer(ResourceArray->GetResourceDataSize(), BUF_Static | BUF_ShaderResource, CreateInfo);
+
+	const FRHIBufferCreateDesc CreateDesc =
+		FRHIBufferCreateDesc::CreateVertex(TEXT("SkeletalMeshAttributeVertexBuffer"), ResourceArray->GetResourceDataSize())
+		.AddUsage(EBufferUsageFlags::Static | EBufferUsageFlags::ShaderResource)
+		.SetInitActionResourceArray(ResourceArray)
+		.DetermineInitialState();
+
+	ValueBuffer.VertexBufferRHI = RHICmdList.CreateBuffer(CreateDesc);
 	ValueBuffer.VertexBufferSRV = RHICmdList.CreateShaderResourceView(
 		ValueBuffer.VertexBufferRHI,
 		FRHIViewDesc::CreateBufferSRV()
diff --git a/Engine/Source/Runtime/Engine/Private/Rendering/SkeletalMeshHalfEdgeBuffer.cpp b/Engine/Source/Runtime/Engine/Private/Rendering/SkeletalMeshHalfEdgeBuffer.cpp
--- a/Engine/Source/Runtime/Engine/Private/Rendering/SkeletalMeshHalfEdgeBuffer.cpp
+++ b/Engine/Source/Runtime/Engine/Private/Rendering/SkeletalMeshHalfEdgeBuffer.cpp
@@ -233,30 +233,39 @@ FSkeletalMeshHalfEdgeBuffer::FRHIInfo FSkeletalMeshHalfEdgeBuffer::CreateRHIBuff
 	uint32 MinBufferSize = GetMinBufferSize();
 	{
 		const uint32 BufferSize = VertexToEdgeData.GetResourceDataSize();
-		FRHIResourceCreateInfo CreateInfo(TEXT("VertexToEdgeData"));
-		CreateInfo.ClassName = ClassName;
-		CreateInfo.OwnerName = GetOwnerName();
+
+		FRHIBufferCreateDesc CreateDesc =
+			FRHIBufferCreateDesc::CreateStructured(TEXT("VertexToEdgeData"), FMath::Max(BufferSize, MinBufferSize), MinBufferSize)
+			.AddUsage(EBufferUsageFlags::Static | EBufferUsageFlags::ShaderResource)
+			.SetInitialState(ERHIAccess::SRVMask)
+			.SetClassName(ClassName)
+			.SetOwnerName(GetOwnerName());
+
 		if (BufferSize > 0)
 		{
-			CreateInfo.ResourceArray = &VertexToEdgeData;
+			CreateDesc.SetInitActionResourceArray(&VertexToEdgeData);
 		}
-		Buffers.VertexToEdgeBufferRHI = RHICmdList.CreateStructuredBuffer(MinBufferSize, FMath::Max(BufferSize, MinBufferSize) , BUF_Static | BUF_ShaderResource, ERHIAccess::SRVMask, CreateInfo);
-		Buffers.VertexToEdgeBufferRHI->SetOwnerName(GetOwnerName());
+
+		Buffers.VertexToEdgeBufferRHI = RHICmdList.CreateBuffer(CreateDesc);
 	}
 
 	{
 		const uint32 BufferSize = EdgeToTwinEdgeData.GetResourceDataSize();
-		FRHIResourceCreateInfo CreateInfo(TEXT("EdgeToTwinEdgeData"));
-		CreateInfo.ClassName = ClassName;
-		CreateInfo.OwnerName = GetOwnerName();
+
+		FRHIBufferCreateDesc CreateDesc =
+			FRHIBufferCreateDesc::CreateStructured(TEXT("EdgeToTwinEdgeData"), FMath::Max(BufferSize, MinBufferSize), sizeof(int32))
+			.AddUsage(EBufferUsageFlags::Static | EBufferUsageFlags::ShaderResource)
+			.SetInitialState(ERHIAccess::SRVMask)
+			.SetClassName(ClassName)
+			.SetOwnerName(GetOwnerName());
+
 		if (BufferSize > 0)
 		{
-			CreateInfo.ResourceArray = &EdgeToTwinEdgeData;
+			CreateDesc.SetInitActionResourceArray(&EdgeToTwinEdgeData);
 		}
-		Buffers.EdgeToTwinEdgeBufferRHI = RHICmdList.CreateStructuredBuffer(sizeof(int32), FMath::Max(BufferSize, MinBufferSize), BUF_Static | BUF_ShaderResource, ERHIAccess::SRVMask, CreateInfo);
-		Buffers.EdgeToTwinEdgeBufferRHI->SetOwnerName(GetOwnerName());
-	}
 
+		Buffers.EdgeToTwinEdgeBufferRHI = RHICmdList.CreateBuffer(CreateDesc);
+	}
 
 	return Buffers;
 }
diff --git a/Engine/Source/Runtime/Engine/Private/SkeletalRenderGPUSkin.cpp b/Engine/Source/Runtime/Engine/Private/SkeletalRenderGPUSkin.cpp
--- a/Engine/Source/Runtime/Engine/Private/SkeletalRenderGPUSkin.cpp
+++ b/Engine/Source/Runtime/Engine/Private/SkeletalRenderGPUSkin.cpp
@@ -18,6 +18,7 @@
 #include "RenderGraphUtils.h"
 #include "RenderCore.h"
 #include "SkinnedMeshSceneProxyDesc.h"
+#include "RHIResourceUtils.h"
 
 DEFINE_LOG_CATEGORY_STATIC(LogSkeletalGPUSkinMesh, Warning, All);
 
@@ -147,39 +148,25 @@ void FMorphVertexBuffer::InitRHI(FRHICommandListBase& RHICmdList)
 	// LOD of the skel mesh is used to find number of vertices in buffer
 	FSkeletalMeshLODRenderData& LodData = SkelMeshRenderData->LODRenderData[LODIdx];
 
-	// Create the buffer rendering resource
-	uint32 Size = LodData.GetNumVertices() * sizeof(FMorphGPUSkinVertex);
-	FRHIResourceCreateInfo CreateInfo(TEXT("MorphVertexBuffer"));
-
-	const EShaderPlatform ShaderPlatform = GetFeatureLevelShaderPlatform(FeatureLevel);
 	const bool bUseGPUMorphTargets = UseGPUMorphTargets(FeatureLevel);
 	bUsesComputeShader = bUseGPUMorphTargets;
 
-	EBufferUsageFlags Flags = bUseGPUMorphTargets ? (EBufferUsageFlags)(BUF_Static | BUF_UnorderedAccess) : BUF_Dynamic;
-
-	// BUF_ShaderResource is needed for Morph support of the SkinCache
-	Flags = (EBufferUsageFlags)(Flags | BUF_ShaderResource);
-
-	VertexBufferRHI = RHICmdList.CreateVertexBuffer(Size, Flags, CreateInfo);
-	VertexBufferRHI->SetOwnerName(GetOwnerName());
-	SRVValue = RHICmdList.CreateShaderResourceView(
-		VertexBufferRHI, 
-		FRHIViewDesc::CreateBufferSRV()
-			.SetType(FRHIViewDesc::EBufferType::Typed)
-			.SetFormat(PF_R32_FLOAT));
+	const FRHIBufferCreateDesc CreateDesc =
+		FRHIBufferCreateDesc::CreateVertex<FMorphGPUSkinVertex>(TEXT("MorphVertexBuffer"), LodData.GetNumVertices())
+		// EBufferUsageFlags::ShaderResource is needed for Morph support of the SkinCache
+		.AddUsage(EBufferUsageFlags::ShaderResource)
+		.AddUsage(bUseGPUMorphTargets ? (EBufferUsageFlags)(EBufferUsageFlags::Static | EBufferUsageFlags::UnorderedAccess) : EBufferUsageFlags::Dynamic)
+		.SetOwnerName(GetOwnerName())
+		.DetermineInitialState();
 
 	if (!bUseGPUMorphTargets)
 	{
-		// Lock the buffer.
-		void* BufferData = RHICmdList.LockBuffer(VertexBufferRHI, 0, sizeof(FMorphGPUSkinVertex)*LodData.GetNumVertices(), RLM_WriteOnly);
-		FMorphGPUSkinVertex* Buffer = (FMorphGPUSkinVertex*)BufferData;
-		FMemory::Memzero(&Buffer[0], sizeof(FMorphGPUSkinVertex)*LodData.GetNumVertices());
-		// Unlock the buffer.
-		RHICmdList.UnlockBuffer(VertexBufferRHI);
+		VertexBufferRHI = UE::RHIResourceUtils::CreateBufferZeroed(RHICmdList, CreateDesc);
 		bNeedsInitialClear = false;
 	}
 	else
 	{
+		VertexBufferRHI = RHICmdList.CreateBuffer(CreateDesc);
 		UAVValue = RHICmdList.CreateUnorderedAccessView(
 			VertexBufferRHI, 
 			FRHIViewDesc::CreateBufferUAV()
@@ -188,6 +175,12 @@ void FMorphVertexBuffer::InitRHI(FRHICommandListBase& RHICmdList)
 		bNeedsInitialClear = true;
 	}
 
+	SRVValue = RHICmdList.CreateShaderResourceView(
+		VertexBufferRHI, 
+		FRHIViewDesc::CreateBufferSRV()
+			.SetType(FRHIViewDesc::EBufferType::Typed)
+			.SetFormat(PF_R32_FLOAT));
+
 	// hasn't been updated yet
 	bHasBeenUpdated = false;
 }
diff --git a/Engine/Source/Runtime/Engine/Private/StaticMesh.cpp b/Engine/Source/Runtime/Engine/Private/StaticMesh.cpp
--- a/Engine/Source/Runtime/Engine/Private/StaticMesh.cpp
+++ b/Engine/Source/Runtime/Engine/Private/StaticMesh.cpp
@@ -270,27 +270,29 @@ void FStaticMeshSectionAreaWeightedTriangleSamplerBuffer::InitRHI(FRHICommandLis
 		}
 		const uint32 SizeByte = TriangleCount * sizeof(SectionTriangleInfo);
 
-		FRHIResourceCreateInfo CreateInfo(TEXT("StaticMeshSectionAreaWeightedTriangleSamplerBuffer"));
-		BufferSectionTriangleRHI = RHICmdList.CreateBuffer(SizeByte, BUF_Static | BUF_VertexBuffer | BUF_ShaderResource, 0, ERHIAccess::VertexOrIndexBuffer | ERHIAccess::SRVMask, CreateInfo);
+		const FRHIBufferCreateDesc Desc =
+			FRHIBufferCreateDesc::CreateVertex<SectionTriangleInfo>(TEXT("StaticMeshSectionAreaWeightedTriangleSamplerBuffer"), TriangleCount)
+			.AddUsage(BUF_Static | BUF_ShaderResource)
+			.SetInitialState(ERHIAccess::VertexOrIndexBuffer | ERHIAccess::SRVMask)
+			.SetInitActionInitializer();
+
+		TRHIBufferInitializer<SectionTriangleInfo> Initializer = RHICmdList.CreateBufferInitializer(Desc);
 
 		// Now compute the alias look up table for unifor; distribution for all section and all triangles
-		SectionTriangleInfo* SectionTriangleInfoBuffer = (SectionTriangleInfo*)RHICmdList.LockBuffer(BufferSectionTriangleRHI, 0, SizeByte, RLM_WriteOnly);
-		for (uint32 i = 0; i < AllSectionCount; ++i)
+		for (uint32 SectionIndex = 0; SectionIndex < AllSectionCount; SectionIndex++)
 		{
-			FStaticMeshSectionAreaWeightedTriangleSampler& sampler = (*Samplers)[i];
-			TArrayView<const float> ProbTris = sampler.GetProb();
-			TArrayView<const int32> AliasTris = sampler.GetAlias();
-			const uint32 NumTriangle = sampler.GetNumEntries();
+			const FStaticMeshSectionAreaWeightedTriangleSampler& TriangleSampler = (*Samplers)[SectionIndex];
+			TConstArrayView<float> ProbTris = TriangleSampler.GetProb();
+			TConstArrayView<int32> AliasTris = TriangleSampler.GetAlias();
+			const uint32 NumTriangle = TriangleSampler.GetNumEntries();
 
-			for (uint32 t = 0; t < NumTriangle; ++t)
+			for (uint32 TriangleIndex = 0; TriangleIndex < NumTriangle; TriangleIndex++)
 			{
-				SectionTriangleInfo NewTriangleInfo = { ProbTris[t], (uint32)AliasTris[t] };
-				*SectionTriangleInfoBuffer = NewTriangleInfo;
-				SectionTriangleInfoBuffer++;
+				SectionTriangleInfo NewTriangleInfo = { ProbTris[TriangleIndex], (uint32)AliasTris[TriangleIndex] };
+				Initializer[TriangleIndex] = NewTriangleInfo;
 			}
 		}
-		RHICmdList.UnlockBuffer(BufferSectionTriangleRHI);
-
+		BufferSectionTriangleRHI = Initializer.Finalize();
 		BufferSectionTriangleSRV = RHICmdList.CreateShaderResourceView(
 			BufferSectionTriangleRHI, 
 			FRHIViewDesc::CreateBufferSRV()
diff --git a/Engine/Source/Runtime/Engine/Private/UnrealEngine.cpp b/Engine/Source/Runtime/Engine/Private/UnrealEngine.cpp
--- a/Engine/Source/Runtime/Engine/Private/UnrealEngine.cpp
+++ b/Engine/Source/Runtime/Engine/Private/UnrealEngine.cpp
@@ -11405,12 +11405,17 @@ bool UEngine::PerformError(const TCHAR* Cmd, FOutputDevice& Ar)
 				const size_t CurrentAllocBytes = CurrentAllocMB * 1024 * 1024;
 
 				// create and fill buffer
-				FRHIResourceCreateInfo Info(TEXT("Debug GPUAlloc"));
-				FBufferRHIRef Buf = RHICmdList.CreateBuffer(CurrentAllocBytes, EBufferUsageFlags::VertexBuffer, 32, ERHIAccess::Unknown, Info);
+				const FRHIBufferCreateDesc CreateDesc =
+					FRHIBufferCreateDesc::CreateVertex(TEXT("Debug GPUAlloc"), CurrentAllocBytes)
+					.SetInitActionInitializer()
+					.DetermineInitialState();
+
+				FRHIBufferInitializer Initializer = RHICmdList.CreateBufferInitializer(CreateDesc);
+				
 				// only need to touch some of the memory
-				void* Data = RHICmdList.LockBuffer(Buf, 0, CurrentAllocBytes, EResourceLockMode::RLM_WriteOnly);
-				FMemory::Memset(Data, 'd', CurrentAllocBytes);
-				RHICmdList.UnlockBuffer(Buf);
+				Initializer.FillWithValue('d');
+				
+				FBufferRHIRef Buf = Initializer.Finalize();
 
 				// store temporarily to avoid RAII
 				BufferRefs.Add(Buf);
diff --git a/Engine/Source/Runtime/Engine/Private/VT/VirtualTextureUploadCache.cpp b/Engine/Source/Runtime/Engine/Private/VT/VirtualTextureUploadCache.cpp
--- a/Engine/Source/Runtime/Engine/Private/VT/VirtualTextureUploadCache.cpp
+++ b/Engine/Source/Runtime/Engine/Private/VT/VirtualTextureUploadCache.cpp
@@ -185,8 +185,12 @@ void FVTUploadTileAllocator::FStagingBuffer::Init(FRHICommandList& RHICmdList, u
 	if (GRHISupportsDirectGPUMemoryLock && GRHISupportsUpdateFromBufferTexture)
 	{
 		// Allocate staging buffer directly in GPU memory.
-		FRHIResourceCreateInfo CreateInfo(TEXT("StagingBuffer"));
-		RHIBuffer = RHICmdList.CreateBuffer(BufferSize, BUF_ShaderResource | BUF_Static | BUF_KeepCPUAccessible | BUF_StructuredBuffer, InBufferStrideBytes, ERHIAccess::SRVMask, CreateInfo);
+		const FRHIBufferCreateDesc CreateDesc =
+			FRHIBufferCreateDesc::CreateStructured(TEXT("StagingBuffer"), BufferSize, InBufferStrideBytes)
+			.AddUsage(EBufferUsageFlags::ShaderResource | EBufferUsageFlags::Static | EBufferUsageFlags::KeepCPUAccessible)
+			.SetInitialState(ERHIAccess::SRVMask);
+
+		RHIBuffer = RHICmdList.CreateBuffer(CreateDesc);
 
 		// Here we bypass 'normal' RHI operations in order to get a persistent pointer to GPU memory, on supported platforms
 		// This should be encapsulated into a proper RHI method at some point
diff --git a/Engine/Source/Runtime/Engine/Private/VectorFieldVisualization.cpp b/Engine/Source/Runtime/Engine/Private/VectorFieldVisualization.cpp
--- a/Engine/Source/Runtime/Engine/Private/VectorFieldVisualization.cpp
+++ b/Engine/Source/Runtime/Engine/Private/VectorFieldVisualization.cpp
@@ -16,6 +16,7 @@
 #include "FXSystem.h"
 #include "MeshMaterialShader.h"
 #include "DataDrivenShaderPlatformInfo.h"
+#include "RHIResourceUtils.h"
 
 IMPLEMENT_GLOBAL_SHADER_PARAMETER_STRUCT(FVectorFieldVisualizationParameters,"VectorFieldVis");
 
@@ -79,6 +80,12 @@ public:
 };
 TGlobalResource<FVectorFieldVisualizationVertexDeclaration> GVectorFieldVisualizationVertexDeclaration;
 
+const FVector4f GDummyVertexBufferData[2] =
+{
+	FVector4f(0.0f, 0.0f, 0.0f, 0.0f),
+	FVector4f(1.0f, 1.0f, 1.0f, 1.0f),
+};
+
 /**
  * A dummy vertex buffer to bind when visualizing vector fields. This prevents
  * some D3D debug warnings about zero-element input layouts but is not strictly
@@ -87,15 +94,9 @@ TGlobalResource<FVectorFieldVisualizationVertexDeclaration> GVectorFieldVisualiz
 class FDummyVertexBuffer : public FVertexBuffer
 {
 public:
-
 	virtual void InitRHI(FRHICommandListBase& RHICmdList) override
 	{
-		FRHIResourceCreateInfo CreateInfo(TEXT("FDummyVertexBuffer"));
-		VertexBufferRHI = RHICmdList.CreateBuffer(sizeof(FVector4f) * 2, BUF_Static | BUF_VertexBuffer, 0, ERHIAccess::VertexOrIndexBuffer, CreateInfo);
-		FVector4f* DummyContents = (FVector4f*)RHICmdList.LockBuffer(VertexBufferRHI, 0, sizeof(FVector4f) * 2, RLM_WriteOnly);
-		DummyContents[0] = FVector4f(0.0f, 0.0f, 0.0f, 0.0f);
-		DummyContents[1] = FVector4f(1.0f, 1.0f, 1.0f, 1.0f);
-		RHICmdList.UnlockBuffer(VertexBufferRHI);
+		VertexBufferRHI = UE::RHIResourceUtils::CreateVertexBufferFromArray(RHICmdList, TEXT("FDummyVertexBuffer"), BUF_Static, MakeConstArrayView(GDummyVertexBufferData));
 	}
 };
 TGlobalResource<FDummyVertexBuffer> GDummyVertexBuffer;
