Initialize struct defaults

diff --git a/Engine/Source/Runtime/Engine/Public/StateStream/SkinnedMeshStateStream.h b/Engine/Source/Runtime/Engine/Public/StateStream/SkinnedMeshStateStream.h
--- a/Engine/Source/Runtime/Engine/Public/StateStream/SkinnedMeshStateStream.h
+++ b/Engine/Source/Runtime/Engine/Public/StateStream/SkinnedMeshStateStream.h
@@ -20,7 +20,7 @@ struct FSkinnedMeshStaticState
 	GENERATED_USTRUCT_BODY()
 
 	UPROPERTY()
-	uint64 MaterialRelevance;
+	uint64 MaterialRelevance = 0;
 };
 
 
diff --git a/Engine/Source/Runtime/Engine/Public/StateStream/StaticMeshStateStream.h b/Engine/Source/Runtime/Engine/Public/StateStream/StaticMeshStateStream.h
--- a/Engine/Source/Runtime/Engine/Public/StateStream/StaticMeshStateStream.h
+++ b/Engine/Source/Runtime/Engine/Public/StateStream/StaticMeshStateStream.h
@@ -43,10 +43,10 @@ private: // Use accessors (SetX/GetX) instead of properties directly
 	TArray<uint32> Owners;
 
 	UPROPERTY()
-	uint32 bOnlyOwnerSee : 1;
+	uint32 bOnlyOwnerSee : 1 = false;
 
 	UPROPERTY()
-	uint32 bOwnerNoSee : 1;
+	uint32 bOwnerNoSee : 1 = false;
 };
 
 
