[Engine]
* Added default initializers to FSkyAtmosphereDynamicState members

diff --git a/Engine/Source/Runtime/Engine/Public/StateStream/SkyAtmosphereStateStream.h b/Engine/Source/Runtime/Engine/Public/StateStream/SkyAtmosphereStateStream.h
--- a/Engine/Source/Runtime/Engine/Public/StateStream/SkyAtmosphereStateStream.h
+++ b/Engine/Source/Runtime/Engine/Public/StateStream/SkyAtmosphereStateStream.h
@@ -35,6 +35,7 @@ struct FOverrideAtmosphericLight
 	UPROPERTY()
 	FVector Direction[NUM_ATMOSPHERE_LIGHTS];
 
+	FOverrideAtmosphericLight() { memset(this, 0, sizeof(FOverrideAtmosphericLight)); }
 	bool operator==(const FOverrideAtmosphericLight& O) const { return memcmp(this, &O, sizeof(FOverrideAtmosphericLight)) == 0; }
 };
 
@@ -46,39 +47,39 @@ struct FSkyAtmosphereDynamicState
 // private: // Should be private but complicates some of the code that needs support both this path and UComponent path
 
 	UPROPERTY()
-	ESkyAtmosphereTransformMode TransformMode;
+	ESkyAtmosphereTransformMode TransformMode = ESkyAtmosphereTransformMode(0);
 	UPROPERTY()
-	float BottomRadius;
+	float BottomRadius = 0.0f;
 	UPROPERTY()
-	FColor GroundAlbedo;
+	FColor GroundAlbedo = FColor::White;
 	UPROPERTY()
-	float AtmosphereHeight;
+	float AtmosphereHeight = 0.0f;
 	UPROPERTY()
-	float MultiScatteringFactor;
+	float MultiScatteringFactor = 0.0f;
 	UPROPERTY()
-	float TraceSampleCountScale;
+	float TraceSampleCountScale = 0.0f;
 	UPROPERTY()
-	float RayleighScatteringScale;
+	float RayleighScatteringScale = 0.0f;
 	UPROPERTY()
-	FLinearColor RayleighScattering;
+	FLinearColor RayleighScattering = FLinearColor::White;
 	UPROPERTY()
-	float RayleighExponentialDistribution;
+	float RayleighExponentialDistribution = 0.0f;
 	UPROPERTY()
-	float MieScatteringScale;
+	float MieScatteringScale = 0.0f;
 	UPROPERTY()
-	FLinearColor MieScattering;
+	FLinearColor MieScattering = FLinearColor::White;
 	UPROPERTY()
-	float MieAbsorptionScale;
+	float MieAbsorptionScale = 0.0f;
 	UPROPERTY()
-	FLinearColor MieAbsorption;
+	FLinearColor MieAbsorption = FLinearColor::White;
 	UPROPERTY()
-	float MieAnisotropy;
+	float MieAnisotropy = 0.0f;
 	UPROPERTY()
-	float MieExponentialDistribution;
+	float MieExponentialDistribution = 0.0f;
 	UPROPERTY()
-	float OtherAbsorptionScale;
+	float OtherAbsorptionScale = 0.0f;
 	UPROPERTY()
-	FLinearColor OtherAbsorption;
+	FLinearColor OtherAbsorption = FLinearColor::White;
 	UPROPERTY()
 	float OtherTentDistributionTipAltitude = 0.0f;
 	UPROPERTY()
@@ -86,23 +87,23 @@ struct FSkyAtmosphereDynamicState
 	UPROPERTY()
 	float OtherTentDistributionWidth = 1.0f;
 	UPROPERTY()
-	FLinearColor SkyLuminanceFactor;
+	FLinearColor SkyLuminanceFactor = FLinearColor::White;
 	UPROPERTY()
-	FLinearColor SkyAndAerialPerspectiveLuminanceFactor;
+	FLinearColor SkyAndAerialPerspectiveLuminanceFactor = FLinearColor::White;
 	UPROPERTY()
-	float AerialPespectiveViewDistanceScale;
+	float AerialPespectiveViewDistanceScale = 0.0f;
 	UPROPERTY()
-	float HeightFogContribution;
+	float HeightFogContribution = 0.0f;
 	UPROPERTY()
-	float TransmittanceMinLightElevationAngle;
+	float TransmittanceMinLightElevationAngle = 0.0f;
 	UPROPERTY()
-	float AerialPerspectiveStartDepth;
+	float AerialPerspectiveStartDepth = 0.0f;
 	UPROPERTY()
-	uint8 bHoldout : 1;
+	uint8 bHoldout : 1 = false;
 	UPROPERTY()
-	uint8 bRenderInMainPass : 1;
+	uint8 bRenderInMainPass : 1 = false;
 	UPROPERTY()
-	uint8 bBuilt : 1;
+	uint8 bBuilt : 1 = false;
 
 	UPROPERTY()
 	FTransform ComponentTransform;
