Fixed crash on null skeletal mesh bones
- Vertex factory updates will invalidate and reinitialize the RHI state, which invalidates all pending bone data for enqueued skin cache updates.
- Skin cache entry transfers will wrongly skip dequeueing on the original owners release, which may lead to invalid bone data for enqueued skin cache updates.

diff --git a/Engine/Source/Runtime/Engine/Private/GPUSkinCache.cpp b/Engine/Source/Runtime/Engine/Private/GPUSkinCache.cpp
--- a/Engine/Source/Runtime/Engine/Private/GPUSkinCache.cpp
+++ b/Engine/Source/Runtime/Engine/Private/GPUSkinCache.cpp
@@ -1922,31 +1922,48 @@ void FGPUSkinCache::ProcessRayTracingGeometryToUpdate(FRHICommandList& RHICmdLis
 
 #endif
 
-void FGPUSkinCache::Release(FGPUSkinCacheEntry*& SkinCacheEntry)
+void FGPUSkinCache::Dequeue(FGPUSkinCacheEntry* SkinCacheEntry)
 {
-	if (SkinCacheEntry)
+	if (!SkinCacheEntry)
 	{
-		FGPUSkinCache* SkinCache = SkinCacheEntry->SkinCache;
-		check(SkinCache);
-		SkinCache->PendingProcessRTGeometryEntries.Remove(SkinCacheEntry);
+		return;
+	}
+	
+	FGPUSkinCache* SkinCache = SkinCacheEntry->SkinCache;
+	checkf(SkinCache, TEXT("Attempting to dequeue a skin cache entry not linked to a parent skin cache"));
+	
+	SkinCache->PendingProcessRTGeometryEntries.Remove(SkinCacheEntry);
 
-		if (SkinCacheEntry->bQueuedForDispatch)
+	if (SkinCacheEntry->bQueuedForDispatch)
+	{
+		for (int32 Index = 0; Index < SkinCache->BatchDispatches.Num(); )
 		{
-			for (int32 Index = 0; Index < SkinCache->BatchDispatches.Num(); )
+			FDispatchEntry& Dispatch = SkinCache->BatchDispatches[Index];
+			
+			if (Dispatch.SkinCacheEntry == SkinCacheEntry)
 			{
-				if (SkinCache->BatchDispatches[Index].SkinCacheEntry == SkinCacheEntry)
-				{
-					SkinCache->BatchDispatches.RemoveAtSwap(Index);
+				// Reset the revision, may not kick off the update otherwise
+				SkinCacheEntry->DispatchData[Dispatch.Section].RevisionNumber = 0;
+				
+				SkinCache->BatchDispatches.RemoveAtSwap(Index);
 
-					// Continue to search for other sections associated with this skin cache entry.
-				}
-				else
-				{
-					++Index;
-				}
+				// Continue to search for other sections associated with this skin cache entry.
+			}
+			else
+			{
+				++Index;
 			}
-			SkinCacheEntry->bQueuedForDispatch = false;
 		}
+		
+		SkinCacheEntry->bQueuedForDispatch = false;
+	}
+}
+
+void FGPUSkinCache::Release(FGPUSkinCacheEntry*& SkinCacheEntry)
+{
+	if (SkinCacheEntry)
+	{
+		Dequeue(SkinCacheEntry);
 
 		ReleaseSkinCacheEntry(SkinCacheEntry);
 		SkinCacheEntry = nullptr;
@@ -2217,6 +2234,11 @@ void FGPUSkinCache::UpdateSkinWeightBuffer(FGPUSkinCacheEntry* Entry)
 {
 	if (Entry)
 	{
+		// Dequeue any pending updates to the entries
+		//   Skin weight updates reinitialize the vertex factories RHI state, which will in turn invalidate the bone data
+		//   for any pending update in the dispatch list.
+		FGPUSkinCache::Dequeue(Entry);
+		
 		Entry->UpdateSkinWeightBuffer();
 	}
 }
@@ -2225,6 +2247,13 @@ void FGPUSkinCache::SetEntryGPUSkin(FGPUSkinCacheEntry* Entry, FSkeletalMeshObje
 {
 	if (Entry)
 	{
+		// Dequeue any pending updates to the entries
+		//   When transferring owner there is a small window in which we may still reference the original vertex factory
+		//   before the new owner has updated the entry. If the entry is pending an update in the dispatch list, we risk
+		//   accessing invalid bone data if the original owner is released. The original owner *does* dequeue on release,
+		//   however, the transfer nulls the old entry.
+		FGPUSkinCache::Dequeue(Entry);
+		
 		Entry->GPUSkin = Skin;
 	}
 }
diff --git a/Engine/Source/Runtime/Engine/Public/GPUSkinCache.h b/Engine/Source/Runtime/Engine/Public/GPUSkinCache.h
--- a/Engine/Source/Runtime/Engine/Public/GPUSkinCache.h
+++ b/Engine/Source/Runtime/Engine/Public/GPUSkinCache.h
@@ -154,6 +154,9 @@ public:
 		FGPUSkinCacheEntry*& InOutEntry
 		);
 
+	/** Dequeue a skin cache entry from the pending skin cache dispatches */
+	static void Dequeue(FGPUSkinCacheEntry* SkinCacheEntry);
+	
 	static void Release(FGPUSkinCacheEntry*& SkinCacheEntry);
 
 	static bool IsEntryValid(FGPUSkinCacheEntry* SkinCacheEntry, int32 Section);
