Follow up MH component changes for Post-process anim graphs in AnimNext

Deprecated post process lod threshold functions.

diff --git a/Engine/Source/Runtime/Engine/Classes/Engine/SkeletalMesh.h b/Engine/Source/Runtime/Engine/Classes/Engine/SkeletalMesh.h
--- a/Engine/Source/Runtime/Engine/Classes/Engine/SkeletalMesh.h
+++ b/Engine/Source/Runtime/Engine/Classes/Engine/SkeletalMesh.h
@@ -2133,8 +2133,11 @@ private:
 	int32 PostProcessAnimBPLODThreshold = INDEX_NONE;
 
 public:
+	UE_DEPRECATED(5.6, "Use GetPostProcessAnimGraphLODThreshold() instead.")
 	ENGINE_API int32 GetPostProcessAnimBPLODThreshold() const { return GetPostProcessAnimGraphLODThreshold(); }
+	UE_DEPRECATED(5.6, "Use SetPostProcessAnimGraphLODThreshold() instead.")
 	ENGINE_API void SetPostProcessAnimBPLODThreshold(int32 InPostProcessAnimBPLODThreshold) { SetPostProcessAnimGraphLODThreshold(InPostProcessAnimBPLODThreshold); }
+	UE_DEPRECATED(5.6, "Use ShouldEvaluatePostProcessAnimGraph() instead.")
 	ENGINE_API bool ShouldEvaluatePostProcessAnimBP(int32 LODLevel) const { return ShouldEvaluatePostProcessAnimGraph(LODLevel); }
 
 	ENGINE_API int32 GetPostProcessAnimGraphLODThreshold() const;
