Removed all uses of MenuCategories from all MaterialExpression*.h classes and replaced with definitions in ini files.

- Wrote code that looped over all UMaterialExpression classes and dumped out their MenuCategories as localized FText.  The ini file sections in BaseEditor.ini were auto-generated in this way.
- Deleted all uses of MenuCategories and uses of static FConstructorStatics ConstructorStatics; in MaterialExpressions.cpp
- If the material expression's constructor was then empty I deleted the constructor
- We still need the constructor, but now auto-gen it via `GENERATED_BODY()` instead of the legacy `GENERATED_UCLASS_BODY` (which didn't generate the constructor)
- The GENERATED_BODY had a `private:` in it making all the properties and member functions private on the UMaterialExpression.
- I added back in a `public:` to each class so they behave and compile like they used to.
- Fix the Sobol Material Expressions which didn't have their categories setup.

[RN] minor | Rendering - Materials and Shaders

diff --git a/Engine/Config/BaseEditor.ini b/Engine/Config/BaseEditor.ini
--- a/Engine/Config/BaseEditor.ini
+++ b/Engine/Config/BaseEditor.ini
@@ -709,6 +709,857 @@ bUseWidgetTemplateSelector=True
 +Filters=(FilterString="TexCoordRHIBuffer",DisplayName="Tex Coord Buffers")
 +Filters=(FilterString="FColorVertexBuffer",DisplayName="Color Vertex Buffers")
 
+[/Script/Engine.MaterialExpressionBentNormalCustomOutput]
++MenuCategories=NSLOCTEXT("MaterialExpression", "Utility", "Utility")
+
+[/Script/Engine.MaterialExpressionClearCoatNormalCustomOutput]
++MenuCategories=NSLOCTEXT("MaterialExpression", "Utility", "Utility")
+
+[/Script/Engine.MaterialExpressionTangentOutput]
++MenuCategories=NSLOCTEXT("MaterialExpression", "Custom", "Custom")
+
+[/Script/Engine.MaterialExpressionThinTranslucentMaterialOutput]
++MenuCategories=NSLOCTEXT("MaterialExpression", "ThinTranslucent", "ThinTranslucent")
+
+[/Script/Engine.MaterialExpressionFirstPersonOutput]
++MenuCategories=NSLOCTEXT("MaterialExpression", "FirstPerson", "First Person")
+
+[/Script/Landscape.MaterialExpressionLandscapeGrassOutput]
++MenuCategories=NSLOCTEXT("Landscape", "Landscape", "Landscape")
+
+[/Script/Landscape.MaterialExpressionLandscapeLayerBlend]
++MenuCategories=NSLOCTEXT("Landscape", "Landscape", "Landscape")
+
+[/Script/Landscape.MaterialExpressionLandscapeLayerCoords]
++MenuCategories=NSLOCTEXT("Landscape", "Landscape", "Landscape")
+
+[/Script/Landscape.MaterialExpressionLandscapeLayerSample]
++MenuCategories=NSLOCTEXT("Landscape", "Landscape", "Landscape")
+
+[/Script/Landscape.MaterialExpressionLandscapeLayerSwitch]
++MenuCategories=NSLOCTEXT("Landscape", "Landscape", "Landscape")
+
+[/Script/Landscape.MaterialExpressionLandscapeLayerWeight]
++MenuCategories=NSLOCTEXT("Landscape", "Landscape", "Landscape")
+
+[/Script/Landscape.MaterialExpressionLandscapePhysicalMaterialOutput]
++MenuCategories=NSLOCTEXT("Landscape", "Landscape", "Landscape")
+
+[/Script/Landscape.MaterialExpressionLandscapeVisibilityMask]
++MenuCategories=NSLOCTEXT("Landscape", "Landscape", "Landscape")
+
+[/Script/Engine.MaterialExpressionAbs]
++MenuCategories=NSLOCTEXT("MaterialExpression", "Math", "Math")
+
+[/Script/Engine.MaterialExpressionAbsorptionMediumMaterialOutput]
++MenuCategories=NSLOCTEXT("MaterialExpression", "Shading", "Shading")
+
+[/Script/Engine.MaterialExpressionActorPositionWS]
++MenuCategories=NSLOCTEXT("MaterialExpression", "Vectors", "Vectors")
++MenuCategories=NSLOCTEXT("MaterialExpression", "Coordinates", "Coordinates")
+
+[/Script/Engine.MaterialExpressionAdd]
++MenuCategories=NSLOCTEXT("MaterialExpression", "Math", "Math")
+
+[/Script/Engine.MaterialExpressionTextureSample]
++MenuCategories=NSLOCTEXT("MaterialExpression", "Texture", "Texture")
+
+[/Script/Engine.MaterialExpressionTextureSampleParameter2D]
++MenuCategories=NSLOCTEXT("MaterialExpression", "Texture", "Texture")
++MenuCategories=NSLOCTEXT("MaterialExpression", "Parameters", "Parameters")
+
+[/Script/Engine.MaterialExpressionAntialiasedTextureMask]
++MenuCategories=NSLOCTEXT("MaterialExpression", "Utility", "Utility")
+
+[/Script/Engine.MaterialExpressionAppendVector]
++MenuCategories=NSLOCTEXT("MaterialExpression", "Math", "Math")
++MenuCategories=NSLOCTEXT("MaterialExpression", "VectorOps", "VectorOps")
+
+[/Script/Engine.MaterialExpressionArccosine]
++MenuCategories=NSLOCTEXT("MaterialExpression", "Math", "Math")
+
+[/Script/Engine.MaterialExpressionArccosineFast]
++MenuCategories=NSLOCTEXT("MaterialExpression", "Math", "Math")
+
+[/Script/Engine.MaterialExpressionArcsine]
++MenuCategories=NSLOCTEXT("MaterialExpression", "Math", "Math")
+
+[/Script/Engine.MaterialExpressionArcsineFast]
++MenuCategories=NSLOCTEXT("MaterialExpression", "Math", "Math")
+
+[/Script/Engine.MaterialExpressionArctangent]
++MenuCategories=NSLOCTEXT("MaterialExpression", "Math", "Math")
+
+[/Script/Engine.MaterialExpressionArctangent2]
++MenuCategories=NSLOCTEXT("MaterialExpression", "Math", "Math")
+
+[/Script/Engine.MaterialExpressionArctangent2Fast]
++MenuCategories=NSLOCTEXT("MaterialExpression", "Math", "Math")
+
+[/Script/Engine.MaterialExpressionArctangentFast]
++MenuCategories=NSLOCTEXT("MaterialExpression", "Math", "Math")
+
+[/Script/Engine.MaterialExpressionAtmosphericFogColor]
++MenuCategories=NSLOCTEXT("MaterialExpression", "Atmosphere", "Atmosphere")
+
+[/Script/Engine.MaterialExpressionAtmosphericLightColor]
++MenuCategories=NSLOCTEXT("MaterialExpression", "Utility", "Utility")
+
+[/Script/Engine.MaterialExpressionAtmosphericLightVector]
++MenuCategories=NSLOCTEXT("MaterialExpression", "Utility", "Utility")
+
+[/Script/Engine.MaterialExpressionBlackBody]
++MenuCategories=NSLOCTEXT("MaterialExpression", "Utility", "Utility")
+
+[/Script/Engine.MaterialExpressionBlendMaterialAttributes]
++MenuCategories=NSLOCTEXT("MaterialExpression", "MaterialAttributes", "Material Attributes")
+
+[/Script/Engine.MaterialExpressionLegacyBlendMaterialAttributes]
++MenuCategories=NSLOCTEXT("MaterialExpression", "MaterialAttributes", "Material Attributes")
+
+[/Script/Engine.MaterialExpressionBounds]
++MenuCategories=NSLOCTEXT("MaterialExpression", "Vectors", "Vectors")
+
+[/Script/Engine.MaterialExpressionBreakMaterialAttributes]
++MenuCategories=NSLOCTEXT("MaterialExpression", "MaterialAttributes", "Material Attributes")
+
+[/Script/Engine.MaterialExpressionBumpOffset]
++MenuCategories=NSLOCTEXT("MaterialExpression", "Utility", "Utility")
+
+[/Script/Engine.MaterialExpressionCameraPositionWS]
++MenuCategories=NSLOCTEXT("MaterialExpression", "Vectors", "Vectors")
++MenuCategories=NSLOCTEXT("MaterialExpression", "Coordinates", "Coordinates")
+
+[/Script/Engine.MaterialExpressionCameraVectorWS]
++MenuCategories=NSLOCTEXT("MaterialExpression", "Vectors", "Vectors")
+
+[/Script/Engine.MaterialExpressionCeil]
++MenuCategories=NSLOCTEXT("MaterialExpression", "Math", "Math")
+
+[/Script/Engine.MaterialExpressionVectorParameter]
++MenuCategories=NSLOCTEXT("MaterialExpression", "Parameters", "Parameters")
+
+[/Script/Engine.MaterialExpressionChannelMaskParameter]
++MenuCategories=NSLOCTEXT("MaterialExpression", "Parameters", "Parameters")
++MenuCategories=NSLOCTEXT("MaterialExpression", "Utility", "Utility")
+
+[/Script/Engine.MaterialExpressionClamp]
++MenuCategories=NSLOCTEXT("MaterialExpression", "Math", "Math")
+
+[/Script/Engine.MaterialExpressionCloudSampleAttribute]
++MenuCategories=NSLOCTEXT("MaterialExpression", "Volume", "Volume")
+
+[/Script/Engine.MaterialExpressionCollectionParameter]
++MenuCategories=NSLOCTEXT("MaterialExpression", "Parameters", "Parameters")
+
+[/Script/Engine.MaterialExpressionCollectionTransform]
++MenuCategories=NSLOCTEXT("MaterialExpression", "Parameters", "Parameters")
+
+[/Script/Engine.MaterialExpressionComponentMask]
++MenuCategories=NSLOCTEXT("MaterialExpression", "Math", "Math")
++MenuCategories=NSLOCTEXT("MaterialExpression", "VectorOps", "VectorOps")
+
+[/Script/Engine.MaterialExpressionComposite]
++MenuCategories=NSLOCTEXT("MaterialExpression", "Utility", "Utility")
+
+[/Script/Engine.MaterialExpressionConstant]
++MenuCategories=NSLOCTEXT("MaterialExpression", "Constants", "Constants")
+
+[/Script/Engine.MaterialExpressionConstant2Vector]
++MenuCategories=NSLOCTEXT("MaterialExpression", "Constants", "Constants")
++MenuCategories=NSLOCTEXT("MaterialExpression", "Vectors", "Vectors")
+
+[/Script/Engine.MaterialExpressionConstant3Vector]
++MenuCategories=NSLOCTEXT("MaterialExpression", "Constants", "Constants")
++MenuCategories=NSLOCTEXT("MaterialExpression", "Vectors", "Vectors")
+
+[/Script/Engine.MaterialExpressionConstant4Vector]
++MenuCategories=NSLOCTEXT("MaterialExpression", "Constants", "Constants")
++MenuCategories=NSLOCTEXT("MaterialExpression", "Vectors", "Vectors")
+
+[/Script/Engine.MaterialExpressionConstantBiasScale]
++MenuCategories=NSLOCTEXT("MaterialExpression", "Utility", "Utility")
+
+[/Script/Engine.MaterialExpressionCosine]
++MenuCategories=NSLOCTEXT("MaterialExpression", "Math", "Math")
+
+[/Script/Engine.MaterialExpressionCrossProduct]
++MenuCategories=NSLOCTEXT("MaterialExpression", "Math", "Math")
++MenuCategories=NSLOCTEXT("MaterialExpression", "VectorOps", "VectorOps")
+
+[/Script/Engine.MaterialExpressionScalarParameter]
++MenuCategories=NSLOCTEXT("MaterialExpression", "Parameters", "Parameters")
+
+[/Script/Engine.MaterialExpressionCurveAtlasRowParameter]
++MenuCategories=NSLOCTEXT("MaterialExpression", "Parameters", "Parameters")
+
+[/Script/Engine.MaterialExpressionCustom]
++MenuCategories=NSLOCTEXT("MaterialExpression", "Custom", "Custom")
+
+[/Script/Engine.MaterialExpressionDataDrivenShaderPlatformInfoSwitch]
++MenuCategories=NSLOCTEXT("MaterialExpression", "Utility", "Utility")
+
+[/Script/Engine.MaterialExpressionDBufferTexture]
++MenuCategories=NSLOCTEXT("MaterialExpression", "Decals", "Decals")
+
+[/Script/Engine.MaterialExpressionDDX]
++MenuCategories=NSLOCTEXT("MaterialExpression", "Utility", "Utility")
+
+[/Script/Engine.MaterialExpressionDDY]
++MenuCategories=NSLOCTEXT("MaterialExpression", "Utility", "Utility")
+
+[/Script/Engine.MaterialExpressionDecalColor]
++MenuCategories=NSLOCTEXT("MaterialExpression", "Utils", "Utils")
+
+[/Script/Engine.MaterialExpressionDecalDerivative]
++MenuCategories=NSLOCTEXT("MaterialExpression", "Utils", "Utils")
+
+[/Script/Engine.MaterialExpressionDecalLifetimeOpacity]
++MenuCategories=NSLOCTEXT("MaterialExpression", "Utils", "Utils")
+
+[/Script/Engine.MaterialExpressionDecalMipmapLevel]
++MenuCategories=NSLOCTEXT("MaterialExpression", "Utils", "Utils")
+
+[/Script/Engine.MaterialExpressionDeltaTime]
++MenuCategories=NSLOCTEXT("MaterialExpression", "Constants", "Constants")
+
+[/Script/Engine.MaterialExpressionDepthFade]
++MenuCategories=NSLOCTEXT("MaterialExpression", "Depth", "Depth")
++MenuCategories=NSLOCTEXT("MaterialExpression", "Utility", "Utility")
+
+[/Script/Engine.MaterialExpressionDepthOfFieldFunction]
++MenuCategories=NSLOCTEXT("MaterialExpression", "Utility", "Utility")
+
+[/Script/Engine.MaterialExpressionDeriveNormalZ]
++MenuCategories=NSLOCTEXT("MaterialExpression", "VectorOps", "VectorOps")
+
 [/Script/Engine.MaterialExpressionDesaturation]
 +MenuCategories=NSLOCTEXT("MaterialExpression", "Color", "Color")
-+MenuCategories=NSLOCTEXT("MaterialExpression", "Utility", "Utility")
\ No newline at end of file
++MenuCategories=NSLOCTEXT("MaterialExpression", "Utility", "Utility")
+
+[/Script/Engine.MaterialExpressionDistance]
++MenuCategories=NSLOCTEXT("MaterialExpression", "Utility", "Utility")
+
+[/Script/Engine.MaterialExpressionDistanceCullFade]
++MenuCategories=NSLOCTEXT("MaterialExpression", "Constants", "Constants")
+
+[/Script/Engine.MaterialExpressionDistanceFieldApproxAO]
++MenuCategories=NSLOCTEXT("MaterialExpression", "Utility", "Utility")
+
+[/Script/Engine.MaterialExpressionDistanceFieldGradient]
++MenuCategories=NSLOCTEXT("MaterialExpression", "Utility", "Utility")
+
+[/Script/Engine.MaterialExpressionDistanceFieldsRenderingSwitch]
++MenuCategories=NSLOCTEXT("MaterialExpression", "DistanceFieldsRendering", "DistanceFieldsRendering")
+
+[/Script/Engine.MaterialExpressionDistanceToNearestSurface]
++MenuCategories=NSLOCTEXT("MaterialExpression", "Utility", "Utility")
+
+[/Script/Engine.MaterialExpressionDivide]
++MenuCategories=NSLOCTEXT("MaterialExpression", "Math", "Math")
+
+[/Script/Engine.MaterialExpressionDotProduct]
++MenuCategories=NSLOCTEXT("MaterialExpression", "Math", "Math")
++MenuCategories=NSLOCTEXT("MaterialExpression", "VectorOps", "VectorOps")
+
+[/Script/Engine.MaterialExpressionDoubleVectorParameter]
++MenuCategories=NSLOCTEXT("MaterialExpression", "Parameters", "Parameters")
+
+[/Script/Engine.MaterialExpressionDynamicParameter]
++MenuCategories=NSLOCTEXT("MaterialExpression", "Particles", "Particles")
++MenuCategories=NSLOCTEXT("MaterialExpression", "Parameters", "Parameters")
+
+[/Script/Engine.MaterialExpressionExponential]
++MenuCategories=NSLOCTEXT("MaterialExpression", "Math", "Math")
+
+[/Script/Engine.MaterialExpressionExponential2]
++MenuCategories=NSLOCTEXT("MaterialExpression", "Math", "Math")
+
+[/Script/Engine.MaterialExpressionEyeAdaptation]
++MenuCategories=NSLOCTEXT("MaterialExpression", "Utility", "Utility")
+
+[/Script/Engine.MaterialExpressionEyeAdaptationInverse]
++MenuCategories=NSLOCTEXT("MaterialExpression", "Utility", "Utility")
+
+[/Script/Engine.MaterialExpressionFeatureLevelSwitch]
++MenuCategories=NSLOCTEXT("MaterialExpression", "Utility", "Utility")
+
+[/Script/Engine.MaterialExpressionFloatToUInt]
++MenuCategories=NSLOCTEXT("MaterialExpression", "Math", "Math")
+
+[/Script/Engine.MaterialExpressionUIntToFloat]
++MenuCategories=NSLOCTEXT("MaterialExpression", "Math", "Math")
+
+[/Script/Engine.MaterialExpressionFloor]
++MenuCategories=NSLOCTEXT("MaterialExpression", "Math", "Math")
+
+[/Script/Engine.MaterialExpressionFmod]
++MenuCategories=NSLOCTEXT("MaterialExpression", "Math", "Math")
+
+[/Script/Engine.MaterialExpressionFontSample]
++MenuCategories=NSLOCTEXT("MaterialExpression", "Font", "Font")
++MenuCategories=NSLOCTEXT("MaterialExpression", "Texture", "Texture")
+
+[/Script/Engine.MaterialExpressionFontSampleParameter]
++MenuCategories=NSLOCTEXT("MaterialExpression", "Font", "Font")
++MenuCategories=NSLOCTEXT("MaterialExpression", "Texture", "Texture")
++MenuCategories=NSLOCTEXT("MaterialExpression", "Font", "Font")
++MenuCategories=NSLOCTEXT("MaterialExpression", "Parameters", "Parameters")
+
+[/Script/Engine.MaterialExpressionFontSignedDistance]
++MenuCategories=NSLOCTEXT("MaterialExpression", "Constants", "Constants")
+
+[/Script/Engine.MaterialExpressionFrac]
++MenuCategories=NSLOCTEXT("MaterialExpression", "Math", "Math")
+
+[/Script/Engine.MaterialExpressionFresnel]
++MenuCategories=NSLOCTEXT("MaterialExpression", "Utility", "Utility")
+
+[/Script/Engine.MaterialExpressionFunctionInput]
++MenuCategories=NSLOCTEXT("MaterialExpression", "Functions", "Functions")
+
+[/Script/Engine.MaterialExpressionFunctionOutput]
++MenuCategories=NSLOCTEXT("MaterialExpression", "Functions", "Functions")
+
+[/Script/Engine.MaterialExpressionConstantDouble]
++MenuCategories=NSLOCTEXT("MaterialExpression", "Constants", "Constants")
+
+[/Script/Engine.MaterialExpressionGetMaterialAttributes]
++MenuCategories=NSLOCTEXT("MaterialExpression", "MaterialAttributes", "Material Attributes")
+
+[/Script/Engine.MaterialExpressionGIReplace]
++MenuCategories=NSLOCTEXT("MaterialExpression", "Utility", "Utility")
+
+[/Script/Engine.MaterialExpressionHairAttributes]
++MenuCategories=NSLOCTEXT("MaterialExpression", "Hair Attributes", "Hair Attributes")
+
+[/Script/Engine.MaterialExpressionHairColor]
++MenuCategories=NSLOCTEXT("MaterialExpression", "Hair Color", "Hair Color")
+
+[/Script/Engine.MaterialExpressionHsvToRgb]
++MenuCategories=NSLOCTEXT("MaterialExpression", "Math", "Math")
+
+[/Script/Engine.MaterialExpressionIf]
++MenuCategories=NSLOCTEXT("MaterialExpression", "Math", "Math")
+
+[/Script/Engine.MaterialExpressionIfThenElse]
++MenuCategories=NSLOCTEXT("MaterialExpression", "Execution", "Execution")
+
+[/Script/Engine.MaterialExpressionInverseLinearInterpolate]
++MenuCategories=NSLOCTEXT("MaterialExpression", "Math", "Math")
++MenuCategories=NSLOCTEXT("MaterialExpression", "Utility", "Utility")
+
+[/Script/Engine.MaterialExpressionIsFirstPerson]
++MenuCategories=NSLOCTEXT("MaterialExpression", "Constants", "Constants")
++MenuCategories=NSLOCTEXT("MaterialExpression", "FirstPerson", "First Person")
+
+[/Script/Engine.MaterialExpressionIsOrthographic]
++MenuCategories=NSLOCTEXT("MaterialExpression", "Constants", "Constants")
+
+[/Script/Engine.MaterialExpressionLength]
++MenuCategories=NSLOCTEXT("MaterialExpression", "Utility", "Utility")
+
+[/Script/Engine.MaterialExpressionLightmapUVs]
++MenuCategories=NSLOCTEXT("MaterialExpression", "Coordinates", "Coordinates")
+
+[/Script/Engine.MaterialExpressionLightmassReplace]
++MenuCategories=NSLOCTEXT("MaterialExpression", "Utility", "Utility")
+
+[/Script/Engine.MaterialExpressionLightVector]
++MenuCategories=NSLOCTEXT("MaterialExpression", "Vectors", "Vectors")
+
+[/Script/Engine.MaterialExpressionLinearInterpolate]
++MenuCategories=NSLOCTEXT("MaterialExpression", "Math", "Math")
++MenuCategories=NSLOCTEXT("MaterialExpression", "Utility", "Utility")
+
+[/Script/Engine.MaterialExpressionLocalPosition]
++MenuCategories=NSLOCTEXT("MaterialExpression", "Coordinates", "Coordinates")
+
+[/Script/Engine.MaterialExpressionLogarithm]
++MenuCategories=NSLOCTEXT("MaterialExpression", "Math", "Math")
+
+[/Script/Engine.MaterialExpressionLogarithm10]
++MenuCategories=NSLOCTEXT("MaterialExpression", "Math", "Math")
+
+[/Script/Engine.MaterialExpressionLogarithm2]
++MenuCategories=NSLOCTEXT("MaterialExpression", "Math", "Math")
+
+[/Script/Engine.MaterialExpressionMakeMaterialAttributes]
++MenuCategories=NSLOCTEXT("MaterialExpression", "MaterialAttributes", "Material Attributes")
+
+[/Script/Engine.MaterialExpressionMapARPassthroughCameraUV]
++MenuCategories=NSLOCTEXT("MaterialExpression", "Coordinates", "Coordinates")
+
+[/Script/Engine.MaterialExpressionMaterialAttributeLayers]
++MenuCategories=NSLOCTEXT("MaterialExpression", "MaterialAttributes", "Material Attributes")
++MenuCategories=NSLOCTEXT("MaterialExpression", "Parameters", "Parameters")
+
+[/Script/Engine.MaterialExpressionMaterialFunctionCall]
++MenuCategories=NSLOCTEXT("MaterialExpression", "Functions", "Functions")
+
+[/Script/Engine.MaterialExpressionMaterialProxyReplace]
++MenuCategories=NSLOCTEXT("MaterialExpression", "Utility", "Utility")
+
+[/Script/Engine.MaterialExpressionMax]
++MenuCategories=NSLOCTEXT("MaterialExpression", "Math", "Math")
+
+[/Script/Engine.MaterialExpressionMeshPaintTextureCoordinateIndex]
++MenuCategories=NSLOCTEXT("MaterialExpression", "MeshPaintTexture", "MeshPaintTexture")
+
+[/Script/Engine.MaterialExpressionMeshPaintTextureObject]
++MenuCategories=NSLOCTEXT("MaterialExpression", "MeshPaintTexture", "MeshPaintTexture")
+
+[/Script/Engine.MaterialExpressionMeshPaintTextureReplace]
++MenuCategories=NSLOCTEXT("MaterialExpression", "MeshPaintTexture", "MeshPaintTexture")
+
+[/Script/Engine.MaterialExpressionMin]
++MenuCategories=NSLOCTEXT("MaterialExpression", "Math", "Math")
+
+[/Script/Engine.MaterialExpressionModulo]
++MenuCategories=NSLOCTEXT("MaterialExpression", "Math", "Math")
+
+[/Script/Engine.MaterialExpressionMultiply]
++MenuCategories=NSLOCTEXT("MaterialExpression", "Math", "Math")
+
+[/Script/Engine.MaterialExpressionNamedRerouteDeclaration]
++MenuCategories=NSLOCTEXT("MaterialExpression", "Utility", "Utility")
+
+[/Script/Engine.MaterialExpressionNaniteReplace]
++MenuCategories=NSLOCTEXT("MaterialExpression", "Utility", "Utility")
+
+[/Script/Engine.MaterialExpressionNeuralNetworkInput]
++MenuCategories=NSLOCTEXT("MaterialExpression", "NeuralNetwork", "NeuralNetwork")
+
+[/Script/Engine.MaterialExpressionNeuralNetworkOutput]
++MenuCategories=NSLOCTEXT("MaterialExpression", "NeuralNetwork", "NeuralNetwork")
+
+[/Script/Engine.MaterialExpressionNoise]
++MenuCategories=NSLOCTEXT("MaterialExpressionNoise", "Utility", "Utility")
+
+[/Script/Engine.MaterialExpressionScalarBlueNoise]
++MenuCategories=NSLOCTEXT("MaterialExpressionNoise", "Noise", "Noise")
+
+[/Script/Engine.MaterialExpressionNormalize]
++MenuCategories=NSLOCTEXT("MaterialExpression", "Math", "Math")
++MenuCategories=NSLOCTEXT("MaterialExpression", "VectorOps", "VectorOps")
+
+[/Script/Engine.MaterialExpressionObjectBounds]
++MenuCategories=NSLOCTEXT("MaterialExpression", "Vectors", "Vectors")
+
+[/Script/Engine.MaterialExpressionObjectLocalBounds]
++MenuCategories=NSLOCTEXT("MaterialExpression", "Vectors", "Vectors")
+
+[/Script/Engine.MaterialExpressionObjectOrientation]
++MenuCategories=NSLOCTEXT("MaterialExpression", "Vectors", "Vectors")
++MenuCategories=NSLOCTEXT("MaterialExpression", "Coordinates", "Coordinates")
+
+[/Script/Engine.MaterialExpressionObjectPositionWS]
++MenuCategories=NSLOCTEXT("MaterialExpression", "Vectors", "Vectors")
++MenuCategories=NSLOCTEXT("MaterialExpression", "Coordinates", "Coordinates")
+
+[/Script/Engine.MaterialExpressionObjectRadius]
++MenuCategories=NSLOCTEXT("MaterialExpression", "Coordinates", "Coordinates")
+
+[/Script/Engine.MaterialExpressionOneMinus]
++MenuCategories=NSLOCTEXT("MaterialExpression", "Math", "Math")
+
+[/Script/Engine.MaterialExpressionPanner]
++MenuCategories=NSLOCTEXT("MaterialExpression", "Coordinates", "Coordinates")
+
+[/Script/Engine.MaterialExpressionParticleColor]
++MenuCategories=NSLOCTEXT("MaterialExpression", "Particles", "Particles")
++MenuCategories=NSLOCTEXT("MaterialExpression", "Constants", "Constants")
+
+[/Script/Engine.MaterialExpressionParticleDirection]
++MenuCategories=NSLOCTEXT("MaterialExpression", "Particles", "Particles")
++MenuCategories=NSLOCTEXT("MaterialExpression", "Constants", "Constants")
+
+[/Script/Engine.MaterialExpressionParticleMacroUV]
++MenuCategories=NSLOCTEXT("MaterialExpression", "Particles", "Particles")
+
+[/Script/Engine.MaterialExpressionParticleMotionBlurFade]
++MenuCategories=NSLOCTEXT("MaterialExpression", "Particles", "Particles")
++MenuCategories=NSLOCTEXT("MaterialExpression", "Constants", "Constants")
+
+[/Script/Engine.MaterialExpressionParticlePositionWS]
++MenuCategories=NSLOCTEXT("MaterialExpression", "Particles", "Particles")
++MenuCategories=NSLOCTEXT("MaterialExpression", "Coordinates", "Coordinates")
++MenuCategories=NSLOCTEXT("MaterialExpression", "Vectors", "Vectors")
+
+[/Script/Engine.MaterialExpressionParticleRadius]
++MenuCategories=NSLOCTEXT("MaterialExpression", "Particles", "Particles")
++MenuCategories=NSLOCTEXT("MaterialExpression", "Constants", "Constants")
+
+[/Script/Engine.MaterialExpressionParticleRandom]
++MenuCategories=NSLOCTEXT("MaterialExpression", "Particles", "Particles")
++MenuCategories=NSLOCTEXT("MaterialExpression", "Constants", "Constants")
+
+[/Script/Engine.MaterialExpressionParticleRelativeTime]
++MenuCategories=NSLOCTEXT("MaterialExpression", "Particles", "Particles")
++MenuCategories=NSLOCTEXT("MaterialExpression", "Constants", "Constants")
+
+[/Script/Engine.MaterialExpressionParticleSize]
++MenuCategories=NSLOCTEXT("MaterialExpression", "Particles", "Particles")
++MenuCategories=NSLOCTEXT("MaterialExpression", "Constants", "Constants")
+
+[/Script/Engine.MaterialExpressionParticleSpeed]
++MenuCategories=NSLOCTEXT("MaterialExpression", "Particles", "Particles")
++MenuCategories=NSLOCTEXT("MaterialExpression", "Constants", "Constants")
+
+[/Script/Engine.MaterialExpressionParticleSpriteRotation]
++MenuCategories=NSLOCTEXT("MaterialExpression", "Particles", "Particles")
++MenuCategories=NSLOCTEXT("MaterialExpression", "Constants", "Constants")
+
+[/Script/Engine.MaterialExpressionParticleSubUV]
++MenuCategories=NSLOCTEXT("MaterialExpression", "Texture", "Texture")
++MenuCategories=NSLOCTEXT("MaterialExpression", "Particles", "Particles")
+
+[/Script/Engine.MaterialExpressionParticleSubUVProperties]
++MenuCategories=NSLOCTEXT("MaterialExpression", "Particles", "Particles")
++MenuCategories=NSLOCTEXT("MaterialExpression", "Coordinates", "Coordinates")
+
+[/Script/Engine.MaterialExpressionPathTracingBufferTexture]
++MenuCategories=NSLOCTEXT("MaterialExpression", "PathTracing", "PathTracing")
+
+[/Script/Engine.MaterialExpressionPathTracingQualitySwitch]
++MenuCategories=NSLOCTEXT("MaterialExpression", "Utility", "Utility")
+
+[/Script/Engine.MaterialExpressionPathTracingRayTypeSwitch]
++MenuCategories=NSLOCTEXT("MaterialExpression", "Utility", "Utility")
+
+[/Script/Engine.MaterialExpressionPerInstanceCustomData]
++MenuCategories=NSLOCTEXT("MaterialExpression", "Custom", "Custom")
+
+[/Script/Engine.MaterialExpressionPerInstanceCustomData3Vector]
++MenuCategories=NSLOCTEXT("MaterialExpression", "Custom", "Custom")
+
+[/Script/Engine.MaterialExpressionPerInstanceFadeAmount]
++MenuCategories=NSLOCTEXT("MaterialExpression", "Constants", "Constants")
+
+[/Script/Engine.MaterialExpressionPerInstanceRandom]
++MenuCategories=NSLOCTEXT("MaterialExpression", "Constants", "Constants")
+
+[/Script/Engine.MaterialExpressionPinBase]
++MenuCategories=NSLOCTEXT("MaterialExpression", "Utility", "Utility")
+
+[/Script/Engine.MaterialExpressionPixelDepth]
++MenuCategories=NSLOCTEXT("MaterialExpression", "Depth", "Depth")
+
+[/Script/Engine.MaterialExpressionPixelNormalWS]
++MenuCategories=NSLOCTEXT("MaterialExpression", "Vectors", "Vectors")
++MenuCategories=NSLOCTEXT("MaterialExpression", "Coordinates", "Coordinates")
+
+[/Script/Engine.MaterialExpressionPostVolumeUserFlagTest]
++MenuCategories=NSLOCTEXT("MaterialExpression", "Get Post Process Setting", "Get Post Process Setting")
+
+[/Script/Engine.MaterialExpressionPower]
++MenuCategories=NSLOCTEXT("MaterialExpression", "Math", "Math")
+
+[/Script/Engine.MaterialExpressionPrecomputedAOMask]
++MenuCategories=NSLOCTEXT("MaterialExpression", "Constants", "Constants")
+
+[/Script/Engine.MaterialExpressionPreSkinnedLocalBounds]
++MenuCategories=NSLOCTEXT("MaterialExpression", "Vectors", "Vectors")
+
+[/Script/Engine.MaterialExpressionPreSkinnedNormal]
++MenuCategories=NSLOCTEXT("MaterialExpression", "Vectors", "Vectors")
+
+[/Script/Engine.MaterialExpressionPreSkinnedPosition]
++MenuCategories=NSLOCTEXT("MaterialExpression", "Vectors", "Vectors")
+
+[/Script/Engine.MaterialExpressionPreviousFrameSwitch]
++MenuCategories=NSLOCTEXT("MaterialExpression", "Functions", "Functions")
+
+[/Script/Engine.MaterialExpressionQualitySwitch]
++MenuCategories=NSLOCTEXT("MaterialExpression", "Utility", "Utility")
+
+[/Script/Engine.MaterialExpressionRayTracingQualitySwitch]
++MenuCategories=NSLOCTEXT("MaterialExpression", "Utility", "Utility")
+
+[/Script/Engine.MaterialExpressionReflectionCapturePassSwitch]
++MenuCategories=NSLOCTEXT("MaterialExpression", "Utility", "Utility")
+
+[/Script/Engine.MaterialExpressionReflectionVectorWS]
++MenuCategories=NSLOCTEXT("MaterialExpression", "Vectors", "Vectors")
+
+[/Script/Engine.MaterialExpressionReroute]
++MenuCategories=NSLOCTEXT("MaterialExpression", "Utility", "Utility")
+
+[/Script/Engine.MaterialExpressionRgbToHsv]
++MenuCategories=NSLOCTEXT("MaterialExpression", "Math", "Math")
+
+[/Script/Engine.MaterialExpressionRotateAboutAxis]
++MenuCategories=NSLOCTEXT("MaterialExpression", "Utility", "Utility")
+
+[/Script/Engine.MaterialExpressionRotator]
++MenuCategories=NSLOCTEXT("MaterialExpression", "Coordinates", "Coordinates")
+
+[/Script/Engine.MaterialExpressionRound]
++MenuCategories=NSLOCTEXT("MaterialExpression", "Math", "Math")
+
+[/Script/Engine.MaterialExpressionRuntimeVirtualTextureCustomData]
++MenuCategories=NSLOCTEXT("MaterialExpression", "VirtualTexture", "VirtualTexture")
+
+[/Script/Engine.MaterialExpressionRuntimeVirtualTextureOutput]
++MenuCategories=NSLOCTEXT("MaterialExpression", "VirtualTexture", "VirtualTexture")
+
+[/Script/Engine.MaterialExpressionRuntimeVirtualTextureReplace]
++MenuCategories=NSLOCTEXT("MaterialExpression", "VirtualTexture", "VirtualTexture")
+
+[/Script/Engine.MaterialExpressionRuntimeVirtualTextureSample]
++MenuCategories=NSLOCTEXT("MaterialExpression", "VirtualTexture", "VirtualTexture")
+
+[/Script/Engine.MaterialExpressionRuntimeVirtualTextureSampleParameter]
++MenuCategories=NSLOCTEXT("MaterialExpression", "VirtualTexture", "VirtualTexture")
++MenuCategories=NSLOCTEXT("MaterialExpression", "Parameters", "Parameters")
+
+[/Script/Engine.MaterialExpressionSamplePhysicsVectorField]
++MenuCategories=NSLOCTEXT("MaterialExpression", "Utility", "Utility")
+
+[/Script/Engine.MaterialExpressionSamplePhysicsScalarField]
++MenuCategories=NSLOCTEXT("MaterialExpression", "Utility", "Utility")
+
+[/Script/Engine.MaterialExpressionSamplePhysicsIntegerField]
++MenuCategories=NSLOCTEXT("MaterialExpression", "Utility", "Utility")
+
+[/Script/Engine.MaterialExpressionSaturate]
++MenuCategories=NSLOCTEXT("MaterialExpression", "Math", "Math")
+
+[/Script/Engine.MaterialExpressionSceneColor]
++MenuCategories=NSLOCTEXT("MaterialExpression", "Texture", "Texture")
+
+[/Script/Engine.MaterialExpressionSceneDepth]
++MenuCategories=NSLOCTEXT("MaterialExpression", "Depth", "Depth")
+
+[/Script/Engine.MaterialExpressionSceneDepthWithoutWater]
++MenuCategories=NSLOCTEXT("MaterialExpression", "Water", "Water")
+
+[/Script/Engine.MaterialExpressionSceneTexelSize]
++MenuCategories=NSLOCTEXT("MaterialExpression", "Coordinates", "Coordinates")
+
+[/Script/Engine.MaterialExpressionSceneTexture]
++MenuCategories=NSLOCTEXT("MaterialExpression", "Texture", "Texture")
+
+[/Script/Engine.MaterialExpressionScreenPosition]
++MenuCategories=NSLOCTEXT("MaterialExpression", "Coordinates", "Coordinates")
+
+[/Script/Engine.MaterialExpressionSetMaterialAttributes]
++MenuCategories=NSLOCTEXT("MaterialExpression", "MaterialAttributes", "Material Attributes")
+
+[/Script/Engine.MaterialExpressionShaderStageSwitch]
++MenuCategories=NSLOCTEXT("MaterialExpression", "Utility", "Utility")
+
+[/Script/Engine.MaterialExpressionShadingModel]
++MenuCategories=NSLOCTEXT("MaterialExpression", "Shading Model", "Shading Model")
+
+[/Script/Engine.MaterialExpressionShadingPathSwitch]
++MenuCategories=NSLOCTEXT("MaterialExpression", "Utility", "Utility")
+
+[/Script/Engine.MaterialExpressionShadowReplace]
++MenuCategories=NSLOCTEXT("MaterialExpression", "Utility", "Utility")
+
+[/Script/Engine.MaterialExpressionSign]
++MenuCategories=NSLOCTEXT("MaterialExpression", "Math", "Math")
+
+[/Script/Engine.MaterialExpressionSine]
++MenuCategories=NSLOCTEXT("MaterialExpression", "Math", "Math")
+
+[/Script/Engine.MaterialExpressionSingleLayerWaterMaterialOutput]
++MenuCategories=NSLOCTEXT("MaterialExpression", "Water", "Water")
+
+[/Script/Engine.MaterialExpressionSkyAtmosphereLightDirection]
++MenuCategories=NSLOCTEXT("MaterialExpression", "Sky", "Sky")
+
+[/Script/Engine.MaterialExpressionSkyAtmosphereLightIlluminance]
++MenuCategories=NSLOCTEXT("MaterialExpression", "Sky", "Sky")
+
+[/Script/Engine.MaterialExpressionSkyAtmosphereLightIlluminanceOnGround]
++MenuCategories=NSLOCTEXT("MaterialExpression", "Sky", "Sky")
+
+[/Script/Engine.MaterialExpressionSkyAtmosphereLightDiskLuminance]
++MenuCategories=NSLOCTEXT("MaterialExpression", "Sky", "Sky")
+
+[/Script/Engine.MaterialExpressionSkyAtmosphereAerialPerspective]
++MenuCategories=NSLOCTEXT("MaterialExpression", "Sky", "Sky")
+
+[/Script/Engine.MaterialExpressionSkyAtmosphereDistantLightScatteredLuminance]
++MenuCategories=NSLOCTEXT("MaterialExpression", "Sky", "Sky")
+
+[/Script/Engine.MaterialExpressionSkyAtmosphereViewLuminance]
++MenuCategories=NSLOCTEXT("MaterialExpression", "Sky", "Sky")
+
+[/Script/Engine.MaterialExpressionSkyLightEnvMapSample]
++MenuCategories=NSLOCTEXT("MaterialExpression", "Sky", "Sky")
+
+[/Script/Engine.MaterialExpressionSobol]
++MenuCategories=NSLOCTEXT("MaterialExpression", "Utility", "Utility")
+
+[/Script/Engine.MaterialExpressionSmoothStep]
++MenuCategories=NSLOCTEXT("MaterialExpression", "Math", "Math")
++MenuCategories=NSLOCTEXT("MaterialExpression", "Utility", "Utility")
+
+[/Script/Engine.MaterialExpressionSparseVolumeTextureObject]
++MenuCategories=NSLOCTEXT("MaterialExpression", "Texture", "Texture")
++MenuCategories=NSLOCTEXT("MaterialExpression", "Functions", "Functions")
+
+[/Script/Engine.MaterialExpressionSparseVolumeTextureSample]
++MenuCategories=NSLOCTEXT("MaterialExpression", "Texture", "Texture")
+
+[/Script/Engine.MaterialExpressionSparseVolumeTextureSampleParameter]
++MenuCategories=NSLOCTEXT("MaterialExpression", "Texture", "Texture")
++MenuCategories=NSLOCTEXT("MaterialExpression", "Parameters", "Parameters")
+
+[/Script/Engine.MaterialExpressionSparseVolumeTextureObjectParameter]
++MenuCategories=NSLOCTEXT("MaterialExpression", "Texture", "Texture")
++MenuCategories=NSLOCTEXT("MaterialExpression", "Parameters", "Parameters")
+
+[/Script/Engine.MaterialExpressionSpeedTree]
++MenuCategories=NSLOCTEXT("MaterialExpression", "SpeedTree", "SpeedTree")
+
+[/Script/Engine.MaterialExpressionSphereMask]
++MenuCategories=NSLOCTEXT("MaterialExpression", "Utility", "Utility")
+
+[/Script/Engine.MaterialExpressionSphericalParticleOpacity]
++MenuCategories=NSLOCTEXT("MaterialExpression", "Particles", "Particles")
+
+[/Script/Engine.MaterialExpressionSquareRoot]
++MenuCategories=NSLOCTEXT("MaterialExpression", "Math", "Math")
+
+[/Script/Engine.MaterialExpressionSRGBColorToWorkingColorSpace]
++MenuCategories=NSLOCTEXT("MaterialExpression", "Color", "Color")
+
+[/Script/Engine.MaterialExpressionStaticBool]
++MenuCategories=NSLOCTEXT("MaterialExpression", "Functions", "Functions")
+
+[/Script/Engine.MaterialExpressionStaticBoolParameter]
++MenuCategories=NSLOCTEXT("MaterialExpression", "Parameters", "Parameters")
+
+[/Script/Engine.MaterialExpressionStaticComponentMaskParameter]
++MenuCategories=NSLOCTEXT("MaterialExpression", "Parameters", "Parameters")
++MenuCategories=NSLOCTEXT("MaterialExpression", "Parameters", "Parameters")
+
+[/Script/Engine.MaterialExpressionStaticSwitch]
++MenuCategories=NSLOCTEXT("MaterialExpression", "Functions", "Functions")
+
+[/Script/Engine.MaterialExpressionStaticSwitchParameter]
++MenuCategories=NSLOCTEXT("MaterialExpression", "Parameters", "Parameters")
+
+[/Script/Engine.MaterialExpressionStep]
++MenuCategories=NSLOCTEXT("MaterialExpression", "Math", "Math")
++MenuCategories=NSLOCTEXT("MaterialExpression", "Utility", "Utility")
+
+[/Script/Engine.MaterialExpressionSubsurfaceMediumMaterialOutput]
++MenuCategories=NSLOCTEXT("MaterialExpression", "Shading", "Shading")
+
+[/Script/Engine.MaterialExpressionSubtract]
++MenuCategories=NSLOCTEXT("MaterialExpression", "Math", "Math")
+
+[/Script/Engine.MaterialExpressionSwitch]
++MenuCategories=NSLOCTEXT("MaterialExpression", "Math", "Math")
+
+[/Script/Engine.MaterialExpressionTangent]
++MenuCategories=NSLOCTEXT("MaterialExpression", "Math", "Math")
+
+[/Script/Engine.MaterialExpressionTemporalSobol]
++MenuCategories=NSLOCTEXT("MaterialExpression", "Utility", "Utility")
+
+[/Script/Engine.MaterialExpressionTextureCoordinate]
++MenuCategories=NSLOCTEXT("MaterialExpression", "Coordinates", "Coordinates")
+
+[/Script/Engine.MaterialExpressionTextureObject]
++MenuCategories=NSLOCTEXT("MaterialExpression", "Texture", "Texture")
++MenuCategories=NSLOCTEXT("MaterialExpression", "Functions", "Functions")
+
+[/Script/Engine.MaterialExpressionTextureObjectParameter]
++MenuCategories=NSLOCTEXT("MaterialExpression", "Texture", "Texture")
++MenuCategories=NSLOCTEXT("MaterialExpression", "Parameters", "Parameters")
+
+[/Script/Engine.MaterialExpressionTextureProperty]
++MenuCategories=NSLOCTEXT("MaterialExpression", "Texture", "Texture")
+
+[/Script/Engine.MaterialExpressionTextureSampleParameter2DArray]
++MenuCategories=NSLOCTEXT("MaterialExpression", "Texture", "Texture")
++MenuCategories=NSLOCTEXT("MaterialExpression", "Parameters", "Parameters")
+
+[/Script/Engine.MaterialExpressionTextureSampleParameterCube]
++MenuCategories=NSLOCTEXT("MaterialExpression", "Texture", "Texture")
++MenuCategories=NSLOCTEXT("MaterialExpression", "Parameters", "Parameters")
+
+[/Script/Engine.MaterialExpressionTextureSampleParameterCubeArray]
++MenuCategories=NSLOCTEXT("MaterialExpression", "Texture", "Texture")
++MenuCategories=NSLOCTEXT("MaterialExpression", "Parameters", "Parameters")
+
+[/Script/Engine.MaterialExpressionTextureSampleParameterSubUV]
++MenuCategories=NSLOCTEXT("MaterialExpression", "Texture", "Texture")
++MenuCategories=NSLOCTEXT("MaterialExpression", "Parameters", "Parameters")
++MenuCategories=NSLOCTEXT("MaterialExpression", "Particles", "Particles")
+
+[/Script/Engine.MaterialExpressionTextureSampleParameterVolume]
++MenuCategories=NSLOCTEXT("MaterialExpression", "Texture", "Texture")
++MenuCategories=NSLOCTEXT("MaterialExpression", "Parameters", "Parameters")
+
+[/Script/Engine.MaterialExpressionTime]
++MenuCategories=NSLOCTEXT("MaterialExpression", "Constants", "Constants")
+
+[/Script/Engine.MaterialExpressionTransform]
++MenuCategories=NSLOCTEXT("MaterialExpression", "VectorOps", "VectorOps")
+
+[/Script/Engine.MaterialExpressionTransformPosition]
++MenuCategories=NSLOCTEXT("MaterialExpression", "VectorOps", "VectorOps")
+
+[/Script/Engine.MaterialExpressionTruncate]
++MenuCategories=NSLOCTEXT("MaterialExpression", "Math", "Math")
+
+[/Script/Engine.MaterialExpressionTruncateLWC]
++MenuCategories=NSLOCTEXT("MaterialExpression", "Coordinates", "Coordinates")
+
+[/Script/Engine.MaterialExpressionTwoSidedSign]
++MenuCategories=NSLOCTEXT("MaterialExpression", "Constants", "Constants")
+
+[/Script/Engine.MaterialExpressionUserSceneTexture]
++MenuCategories=NSLOCTEXT("MaterialExpression", "Texture", "Texture")
+
+[/Script/Engine.MaterialExpressionVectorNoise]
++MenuCategories=NSLOCTEXT("MaterialExpressionNoise", "Utility", "Utility")
+
+[/Script/Engine.MaterialExpressionVertexColor]
++MenuCategories=NSLOCTEXT("MaterialExpression", "Constants", "Constants")
+
+[/Script/Engine.MaterialExpressionVertexInterpolator]
++MenuCategories=NSLOCTEXT("MaterialExpression", "Utility", "Utility")
+
+[/Script/Engine.MaterialExpressionVertexNormalWS]
++MenuCategories=NSLOCTEXT("MaterialExpression", "Vectors", "Vectors")
++MenuCategories=NSLOCTEXT("MaterialExpression", "Coordinates", "Coordinates")
+
+[/Script/Engine.MaterialExpressionVertexTangentWS]
++MenuCategories=NSLOCTEXT("MaterialExpression", "Vectors", "Vectors")
++MenuCategories=NSLOCTEXT("MaterialExpression", "Coordinates", "Coordinates")
+
+[/Script/Engine.MaterialExpressionViewProperty]
++MenuCategories=NSLOCTEXT("MaterialExpression", "Constants", "Constants")
+
+[/Script/Engine.MaterialExpressionViewSize]
++MenuCategories=NSLOCTEXT("MaterialExpression", "Coordinates", "Coordinates")
+
+[/Script/Engine.MaterialExpressionVirtualTextureFeatureSwitch]
++MenuCategories=NSLOCTEXT("MaterialExpression", "VirtualTexture", "VirtualTexture")
+
+[/Script/Engine.MaterialExpressionVolumetricAdvancedMaterialInput]
++MenuCategories=NSLOCTEXT("MaterialExpression", "Volume", "Volume")
+
+[/Script/Engine.MaterialExpressionVolumetricCloudEmptySpaceSkippingInput]
++MenuCategories=NSLOCTEXT("MaterialExpression", "Volume", "Volume")
+
+[/Script/Engine.MaterialExpressionVolumetricAdvancedMaterialOutput]
++MenuCategories=NSLOCTEXT("MaterialExpression", "Volume", "Volume")
+
+[/Script/Engine.MaterialExpressionVolumetricCloudEmptySpaceSkippingOutput]
++MenuCategories=NSLOCTEXT("MaterialExpression", "Volume", "Volume")
+
+[/Script/Engine.MaterialExpressionWorldPosition]
++MenuCategories=NSLOCTEXT("MaterialExpression", "Coordinates", "Coordinates")
+
+[/Script/Paper2D.MaterialExpressionSpriteTextureSampler]
++MenuCategories=NSLOCTEXT("MaterialExpression", "Texture", "Texture")
++MenuCategories=NSLOCTEXT("MaterialExpression", "Parameters", "Parameters")
\ No newline at end of file
diff --git a/Engine/Source/Runtime/Engine/Private/Materials/MaterialExpressions.cpp b/Engine/Source/Runtime/Engine/Private/Materials/MaterialExpressions.cpp
--- a/Engine/Source/Runtime/Engine/Private/Materials/MaterialExpressions.cpp
+++ b/Engine/Source/Runtime/Engine/Private/Materials/MaterialExpressions.cpp
@@ -2472,19 +2472,6 @@ UMaterialExpressionTextureSample::UMaterialExpressionTextureSample(const FObject
 	: Super(ObjectInitializer)
 {
 #if WITH_EDITORONLY_DATA
-	// Structure to hold one-time initialization
-	struct FConstructorStatics
-	{
-		FText NAME_Texture;
-		FConstructorStatics()
-			: NAME_Texture(LOCTEXT( "Texture", "Texture" ))
-		{
-		}
-	};
-	static FConstructorStatics ConstructorStatics;
-
-	MenuCategories.Add(ConstructorStatics.NAME_Texture);
-
 	Outputs.Reset();
 	Outputs.Add(FExpressionOutput(TEXT("RGB"), 1, 1, 1, 1, 0));
 	Outputs.Add(FExpressionOutput(TEXT("R"), 1, 1, 0, 0, 0));
@@ -2981,21 +2968,6 @@ int32 UMaterialExpressionTextureSample::CompileMipValue1(class FMaterialCompiler
 UMaterialExpressionRuntimeVirtualTextureOutput::UMaterialExpressionRuntimeVirtualTextureOutput(const FObjectInitializer& ObjectInitializer)
 	: Super(ObjectInitializer)
 {
-	// Structure to hold one-time initialization
-	struct FConstructorStatics
-	{
-		FText NAME_VirtualTexture;
-		FConstructorStatics()
-			: NAME_VirtualTexture(LOCTEXT("VirtualTexture", "VirtualTexture"))
-		{
-		}
-	};
-	static FConstructorStatics ConstructorStatics;
-
-#if WITH_EDITORONLY_DATA
-	MenuCategories.Add(ConstructorStatics.NAME_VirtualTexture);
-#endif
-
 #if WITH_EDITOR
 	Outputs.Reset();
 #endif
@@ -3091,21 +3063,6 @@ FString UMaterialExpressionRuntimeVirtualTextureOutput::GetDisplayName() const
 UMaterialExpressionRuntimeVirtualTextureSample::UMaterialExpressionRuntimeVirtualTextureSample(const FObjectInitializer& ObjectInitializer)
 	: Super(ObjectInitializer)
 {
-	// Structure to hold one-time initialization
-	struct FConstructorStatics
-	{
-		FText NAME_VirtualTexture;
-		FConstructorStatics()
-			: NAME_VirtualTexture(LOCTEXT("VirtualTexture", "VirtualTexture"))
-		{
-		}
-	};
-	static FConstructorStatics ConstructorStatics;
-
-#if WITH_EDITORONLY_DATA
-	MenuCategories.Add(ConstructorStatics.NAME_VirtualTexture);
-#endif
-
 #if WITH_EDITOR
 	InitOutputs();
 	bShowOutputNameOnPin = true;
@@ -3649,22 +3606,7 @@ void UMaterialExpressionRuntimeVirtualTextureSample::GetCaption(TArray<FString>&
 UMaterialExpressionRuntimeVirtualTextureSampleParameter::UMaterialExpressionRuntimeVirtualTextureSampleParameter(const FObjectInitializer& ObjectInitializer)
 	: Super(ObjectInitializer)
 {
-	// Structure to hold one-time initialization
-	struct FConstructorStatics
-	{
-		FText NAME_Parameters;
-		FConstructorStatics()
-			: NAME_Parameters(LOCTEXT("Parameters", "Parameters"))
-		{
-		}
-	};
-	static FConstructorStatics ConstructorStatics;
-
 	bIsParameterExpression = true;
-
-#if WITH_EDITORONLY_DATA
-	MenuCategories.Add(ConstructorStatics.NAME_Parameters);
-#endif
 }
 
 #if WITH_EDITOR
@@ -3724,25 +3666,6 @@ bool UMaterialExpressionRuntimeVirtualTextureSampleParameter::MatchesSearchQuery
 }
 #endif
 
-UMaterialExpressionRuntimeVirtualTextureReplace::UMaterialExpressionRuntimeVirtualTextureReplace(const FObjectInitializer& ObjectInitializer)
-	: Super(ObjectInitializer)
-{
-	// Structure to hold one-time initialization
-	struct FConstructorStatics
-	{
-		FText NAME_VirtualTexture;
-		FConstructorStatics()
-			: NAME_VirtualTexture(LOCTEXT("VirtualTexture", "VirtualTexture"))
-		{
-		}
-	};
-	static FConstructorStatics ConstructorStatics;
-
-#if WITH_EDITORONLY_DATA
-	MenuCategories.Add(ConstructorStatics.NAME_VirtualTexture);
-#endif
-}
-
 #if WITH_EDITOR
 
 int32 UMaterialExpressionRuntimeVirtualTextureReplace::Compile(class FMaterialCompiler* Compiler, int32 OutputIndex)
@@ -3778,25 +3701,6 @@ void UMaterialExpressionRuntimeVirtualTextureReplace::GetCaption(TArray<FString>
 
 #endif // WITH_EDITOR
 
-UMaterialExpressionRuntimeVirtualTextureCustomData::UMaterialExpressionRuntimeVirtualTextureCustomData(const FObjectInitializer& ObjectInitializer)
-	: Super(ObjectInitializer)
-{
-	// Structure to hold one-time initialization
-	struct FConstructorStatics
-	{
-		FText NAME_VirtualTexture;
-		FConstructorStatics()
-			: NAME_VirtualTexture(LOCTEXT("VirtualTexture", "VirtualTexture"))
-		{
-		}
-	};
-	static FConstructorStatics ConstructorStatics;
-
-#if WITH_EDITORONLY_DATA
-	MenuCategories.Add(ConstructorStatics.NAME_VirtualTexture);
-#endif
-}
-
 #if WITH_EDITOR
 
 int32 UMaterialExpressionRuntimeVirtualTextureCustomData::Compile(class FMaterialCompiler* Compiler, int32 OutputIndex)
@@ -3811,25 +3715,6 @@ void UMaterialExpressionRuntimeVirtualTextureCustomData::GetCaption(TArray<FStri
 
 #endif // WITH_EDITOR
 
-UMaterialExpressionVirtualTextureFeatureSwitch::UMaterialExpressionVirtualTextureFeatureSwitch(const FObjectInitializer& ObjectInitializer)
-	: Super(ObjectInitializer)
-{
-	// Structure to hold one-time initialization
-	struct FConstructorStatics
-	{
-		FText NAME_VirtualTexture;
-		FConstructorStatics()
-			: NAME_VirtualTexture(LOCTEXT("VirtualTexture", "VirtualTexture"))
-		{
-		}
-	};
-	static FConstructorStatics ConstructorStatics;
-
-#if WITH_EDITORONLY_DATA
-	MenuCategories.Add(ConstructorStatics.NAME_VirtualTexture);
-#endif
-}
-
 #if WITH_EDITOR
 
 int32 UMaterialExpressionVirtualTextureFeatureSwitch::Compile(class FMaterialCompiler* Compiler, int32 OutputIndex)
@@ -3874,21 +3759,6 @@ void UMaterialExpressionVirtualTextureFeatureSwitch::GetCaption(TArray<FString>&
 UMaterialExpressionMeshPaintTextureCoordinateIndex::UMaterialExpressionMeshPaintTextureCoordinateIndex(const FObjectInitializer& ObjectInitializer)
 	: Super(ObjectInitializer)
 {
-	// Structure to hold one-time initialization
-	struct FConstructorStatics
-	{
-		FText NAME_MeshPaintTexture;
-		FConstructorStatics()
-			: NAME_MeshPaintTexture(LOCTEXT("MeshPaintTexture", "MeshPaintTexture"))
-		{
-		}
-	};
-	static FConstructorStatics ConstructorStatics;
-
-#if WITH_EDITORONLY_DATA
-	MenuCategories.Add(ConstructorStatics.NAME_MeshPaintTexture);
-#endif
-
 	Outputs.Reset();
 	Outputs.Add(FExpressionOutput(TEXT("")));
 }
@@ -3915,21 +3785,6 @@ int32 UMaterialExpressionMeshPaintTextureCoordinateIndex::Compile(class FMateria
 UMaterialExpressionMeshPaintTextureObject::UMaterialExpressionMeshPaintTextureObject(const FObjectInitializer& ObjectInitializer)
 	: Super(ObjectInitializer)
 {
-	// Structure to hold one-time initialization
-	struct FConstructorStatics
-	{
-		FText NAME_MeshPaintTexture;
-		FConstructorStatics()
-			: NAME_MeshPaintTexture(LOCTEXT("MeshPaintTexture", "MeshPaintTexture"))
-		{
-		}
-	};
-	static FConstructorStatics ConstructorStatics;
-
-#if WITH_EDITORONLY_DATA
-	MenuCategories.Add(ConstructorStatics.NAME_MeshPaintTexture);
-#endif
-
 	Outputs.Reset();
 	Outputs.Add(FExpressionOutput(TEXT("")));
 }
@@ -3961,21 +3816,6 @@ int32 UMaterialExpressionMeshPaintTextureObject::Compile(class FMaterialCompiler
 UMaterialExpressionMeshPaintTextureReplace::UMaterialExpressionMeshPaintTextureReplace(const FObjectInitializer& ObjectInitializer)
 	: Super(ObjectInitializer)
 {
-	// Structure to hold one-time initialization
-	struct FConstructorStatics
-	{
-		FText NAME_MeshPaintTexture;
-		FConstructorStatics()
-			: NAME_MeshPaintTexture(LOCTEXT("MeshPaintTexture", "MeshPaintTexture"))
-		{
-		}
-	};
-	static FConstructorStatics ConstructorStatics;
-
-#if WITH_EDITORONLY_DATA
-	MenuCategories.Add(ConstructorStatics.NAME_MeshPaintTexture);
-#endif
-
 	Outputs.Reset();
 	Outputs.Add(FExpressionOutput(TEXT("")));
 }
@@ -4015,23 +3855,8 @@ int32 UMaterialExpressionMeshPaintTextureReplace::Compile(class FMaterialCompile
 UMaterialExpressionAdd::UMaterialExpressionAdd(const FObjectInitializer& ObjectInitializer)
 	: Super(ObjectInitializer)
 {
-	// Structure to hold one-time initialization
-	struct FConstructorStatics
-	{
-		FText NAME_Math;
-		FConstructorStatics()
-			: NAME_Math(LOCTEXT( "Math", "Math" ))
-		{
-		}
-	};
-	static FConstructorStatics ConstructorStatics;
-
 	ConstA = 0.0f;
 	ConstB = 1.0f;
-
-#if WITH_EDITORONLY_DATA
-	MenuCategories.Add(ConstructorStatics.NAME_Math);
-#endif
 }
 
 //
@@ -4041,21 +3866,9 @@ UMaterialExpressionTextureSampleParameter::UMaterialExpressionTextureSampleParam
 	: Super(ObjectInitializer)
 {
 #if WITH_EDITORONLY_DATA
-	// Structure to hold one-time initialization
-	struct FConstructorStatics
-	{
-		FText NAME_Obsolete;
-		FConstructorStatics()
-			: NAME_Obsolete(LOCTEXT( "Obsolete", "Obsolete" ))
-		{
-		}
-	};
-	static FConstructorStatics ConstructorStatics;
 	bIsParameterExpression = true;
 	bShowTextureInputPin = false;
 
-	MenuCategories.Empty();
-	MenuCategories.Add( ConstructorStatics.NAME_Obsolete);
 	ApplyChannelNames();
 #endif
 }
@@ -4171,12 +3984,8 @@ UMaterialExpressionTextureObjectParameter::UMaterialExpressionTextureObjectParam
 	struct FConstructorStatics
 	{
 		ConstructorHelpers::FObjectFinder<UTexture2D> DefaultTexture2D;
-		FText NAME_Texture;
-		FText NAME_Parameters;
 		FConstructorStatics()
 			: DefaultTexture2D(TEXT("/Engine/EngineResources/DefaultTexture"))
-			, NAME_Texture(LOCTEXT( "Texture", "Texture" ))
-			, NAME_Parameters(LOCTEXT( "Parameters", "Parameters" ))
 		{
 		}
 	};
@@ -4184,10 +3993,6 @@ UMaterialExpressionTextureObjectParameter::UMaterialExpressionTextureObjectParam
 
 	Texture = ConstructorStatics.DefaultTexture2D.Object;
 
-	MenuCategories.Empty();
-	MenuCategories.Add(ConstructorStatics.NAME_Texture);
-	MenuCategories.Add(ConstructorStatics.NAME_Parameters);
-
 	Outputs.Empty();
 	Outputs.Add(FExpressionOutput(TEXT("")));
 
@@ -4258,12 +4063,8 @@ UMaterialExpressionTextureObject::UMaterialExpressionTextureObject(const FObject
 	struct FConstructorStatics
 	{
 		ConstructorHelpers::FObjectFinder<UTexture2D> Object0;
-		FText NAME_Texture;
-		FText NAME_Functions;
 		FConstructorStatics()
 			: Object0(TEXT("/Engine/EngineResources/DefaultTexture"))
-			, NAME_Texture(LOCTEXT( "Texture", "Texture" ))
-			, NAME_Functions(LOCTEXT( "Functions", "Functions" ))
 		{
 		}
 	};
@@ -4271,9 +4072,6 @@ UMaterialExpressionTextureObject::UMaterialExpressionTextureObject(const FObject
 
 	Texture = ConstructorStatics.Object0.Object;
 
-	MenuCategories.Add(ConstructorStatics.NAME_Texture);
-	MenuCategories.Add(ConstructorStatics.NAME_Functions);
-
 	Outputs.Empty();
 	Outputs.Add(FExpressionOutput(TEXT("")));
 
@@ -4358,22 +4156,9 @@ uint32 UMaterialExpressionTextureObject::GetOutputType(int32 OutputIndex)
 UMaterialExpressionTextureProperty::UMaterialExpressionTextureProperty(const FObjectInitializer& ObjectInitializer)
 	: Super(ObjectInitializer)
 {
-	// Structure to hold one-time initialization
-	struct FConstructorStatics
-	{
-		FText NAME_Constants;
-		FConstructorStatics()
-			: NAME_Constants(LOCTEXT( "Texture", "Texture" ))
-		{
-		}
-	};
-	static FConstructorStatics ConstructorStatics;
-
 	Property = TMTM_TextureSize;
 
 #if WITH_EDITORONLY_DATA
-	MenuCategories.Add(ConstructorStatics.NAME_Constants);
-
 	bShaderInputData = true;
 	bShowOutputNameOnPin = false;
 	
@@ -4466,22 +4251,14 @@ UMaterialExpressionTextureSampleParameter2D::UMaterialExpressionTextureSamplePar
 	struct FConstructorStatics
 	{
 		ConstructorHelpers::FObjectFinder<UTexture2D> DefaultTexture;
-		FText NAME_Texture;
-		FText NAME_Parameters;
 		FConstructorStatics()
 			: DefaultTexture(TEXT("/Engine/EngineResources/DefaultTexture"))
-			, NAME_Texture(LOCTEXT( "Texture", "Texture" ))
-			, NAME_Parameters(LOCTEXT( "Parameters", "Parameters" ))
 		{
 		}
 	};
 	static FConstructorStatics ConstructorStatics;
 
 	Texture = ConstructorStatics.DefaultTexture.Object;
-
-	MenuCategories.Empty();
-	MenuCategories.Add(ConstructorStatics.NAME_Texture);
-	MenuCategories.Add(ConstructorStatics.NAME_Parameters);
 #endif
 }
 
@@ -4557,22 +4334,14 @@ UMaterialExpressionTextureSampleParameterCube::UMaterialExpressionTextureSampleP
 	struct FConstructorStatics
 	{
 		ConstructorHelpers::FObjectFinder<UTextureCube> DefaultTextureCube;
-		FText NAME_Texture;
-		FText NAME_Parameters;
 		FConstructorStatics()
 			: DefaultTextureCube(TEXT("/Engine/EngineResources/DefaultTextureCube"))
-			, NAME_Texture(LOCTEXT( "Texture", "Texture" ))
-			, NAME_Parameters(LOCTEXT( "Parameters", "Parameters" ))
 		{
 		}
 	};
 	static FConstructorStatics ConstructorStatics;
 
 	Texture = ConstructorStatics.DefaultTextureCube.Object;
-
-	MenuCategories.Empty();
-	MenuCategories.Add(ConstructorStatics.NAME_Texture);
-	MenuCategories.Add(ConstructorStatics.NAME_Parameters);
 #endif
 }
 
@@ -4618,29 +4387,6 @@ void UMaterialExpressionTextureSampleParameterCube::SetDefaultTexture()
 //
 //  UMaterialExpressionTextureSampleParameter2DArray
 //
-UMaterialExpressionTextureSampleParameter2DArray::UMaterialExpressionTextureSampleParameter2DArray(const FObjectInitializer& ObjectInitializer)
-	: Super(ObjectInitializer)
-{
-	// Structure to hold one-time initialization. No default texture array added.
-	struct FConstructorStatics
-	{
-		FText NAME_Texture;
-		FText NAME_Parameters;
-		FConstructorStatics()
-			: NAME_Texture(LOCTEXT("Texture", "Texture"))
-			, NAME_Parameters(LOCTEXT("Parameters", "Parameters"))
-		{
-		}
-	};
-	static FConstructorStatics ConstructorStatics;
-
-#if WITH_EDITORONLY_DATA
-	MenuCategories.Empty();
-	MenuCategories.Add(ConstructorStatics.NAME_Texture);
-	MenuCategories.Add(ConstructorStatics.NAME_Parameters);
-#endif
-}
-
 #if WITH_EDITOR
 int32 UMaterialExpressionTextureSampleParameter2DArray::Compile(class FMaterialCompiler* Compiler, int32 OutputIndex)
 {
@@ -4687,29 +4433,6 @@ const TCHAR* UMaterialExpressionTextureSampleParameter2DArray::GetRequirements()
 //
 //  UMaterialExpressionTextureSampleParameterCubeArray
 //
-UMaterialExpressionTextureSampleParameterCubeArray::UMaterialExpressionTextureSampleParameterCubeArray(const FObjectInitializer& ObjectInitializer)
-	: Super(ObjectInitializer)
-{
-	// Structure to hold one-time initialization
-	struct FConstructorStatics
-	{
-		FText NAME_Texture;
-		FText NAME_Parameters;
-		FConstructorStatics()
-			: NAME_Texture(LOCTEXT("Texture", "Texture"))
-			, NAME_Parameters(LOCTEXT("Parameters", "Parameters"))
-		{
-		}
-	};
-	static FConstructorStatics ConstructorStatics;
-
-#if WITH_EDITORONLY_DATA
-	MenuCategories.Empty();
-	MenuCategories.Add(ConstructorStatics.NAME_Texture);
-	MenuCategories.Add(ConstructorStatics.NAME_Parameters);
-#endif
-}
-
 #if WITH_EDITOR
 int32 UMaterialExpressionTextureSampleParameterCubeArray::Compile(class FMaterialCompiler* Compiler, int32 OutputIndex)
 {
@@ -4760,22 +4483,14 @@ UMaterialExpressionTextureSampleParameterVolume::UMaterialExpressionTextureSampl
 	struct FConstructorStatics
 	{
 		ConstructorHelpers::FObjectFinder<UVolumeTexture> DefaultVolumeTexture;
-		FText NAME_Texture;
-		FText NAME_Parameters;
 		FConstructorStatics()
 			: DefaultVolumeTexture(TEXT("/Engine/EngineResources/DefaultVolumeTexture"))
-			, NAME_Texture(LOCTEXT( "Texture", "Texture" ))
-			, NAME_Parameters(LOCTEXT( "Parameters", "Parameters" ))
 		{
 		}
 	};
 	static FConstructorStatics ConstructorStatics;
 
 	Texture = ConstructorStatics.DefaultVolumeTexture.Object;
-
-	MenuCategories.Empty();
-	MenuCategories.Add(ConstructorStatics.NAME_Texture);
-	MenuCategories.Add(ConstructorStatics.NAME_Parameters);
 #endif
 }
 
@@ -4840,22 +4555,7 @@ static int32 ParticleSubUV(
 UMaterialExpressionTextureSampleParameterSubUV::UMaterialExpressionTextureSampleParameterSubUV(const FObjectInitializer& ObjectInitializer)
 	: Super(ObjectInitializer)
 {
-	// Structure to hold one-time initialization
-	struct FConstructorStatics
-	{
-		FText NAME_Particles;
-		FConstructorStatics()
-			: NAME_Particles(LOCTEXT( "Particles", "Particles" ))
-		{
-		}
-	};
-	static FConstructorStatics ConstructorStatics;
-
 	bBlend = true;
-
-#if WITH_EDITORONLY_DATA
-	MenuCategories.Add(ConstructorStatics.NAME_Particles);
-#endif
 }
 
 #if WITH_EDITOR
@@ -4925,23 +4625,8 @@ void UMaterialExpressionAdd::GetCaption(TArray<FString>& OutCaptions) const
 UMaterialExpressionMultiply::UMaterialExpressionMultiply(const FObjectInitializer& ObjectInitializer)
 	: Super(ObjectInitializer)
 {
-	// Structure to hold one-time initialization
-	struct FConstructorStatics
-	{
-		FText NAME_Math;
-		FConstructorStatics()
-			: NAME_Math(LOCTEXT( "Math", "Math" ))
-		{
-		}
-	};
-	static FConstructorStatics ConstructorStatics;
-
 	ConstA = 0.0f;
 	ConstB = 1.0f;
-
-#if WITH_EDITORONLY_DATA
-	MenuCategories.Add(ConstructorStatics.NAME_Math);
-#endif
 }
 
 #if WITH_EDITOR
@@ -4976,21 +4661,6 @@ void UMaterialExpressionMultiply::GetCaption(TArray<FString>& OutCaptions) const
 UMaterialExpressionDivide::UMaterialExpressionDivide(const FObjectInitializer& ObjectInitializer)
 	: Super(ObjectInitializer)
 {
-	// Structure to hold one-time initialization
-	struct FConstructorStatics
-	{
-		FText NAME_Math;
-		FConstructorStatics()
-			: NAME_Math(LOCTEXT( "Math", "Math" ))
-		{
-		}
-	};
-	static FConstructorStatics ConstructorStatics;
-
-#if WITH_EDITORONLY_DATA
-	MenuCategories.Add(ConstructorStatics.NAME_Math);
-#endif
-
 	ConstA = 1.0f;
 	ConstB = 2.0f;
 }
@@ -5030,23 +4700,8 @@ void UMaterialExpressionDivide::GetCaption(TArray<FString>& OutCaptions) const
 UMaterialExpressionSubtract::UMaterialExpressionSubtract(const FObjectInitializer& ObjectInitializer)
 	: Super(ObjectInitializer)
 {
-	// Structure to hold one-time initialization
-	struct FConstructorStatics
-	{
-		FText NAME_Math;
-		FConstructorStatics()
-			: NAME_Math(LOCTEXT( "Math", "Math" ))
-		{
-		}
-	};
-	static FConstructorStatics ConstructorStatics;
-
 	ConstA = 1.0f;
 	ConstB = 1.0f;
-
-#if WITH_EDITORONLY_DATA
-	MenuCategories.Add(ConstructorStatics.NAME_Math);
-#endif
 }
 
 #if WITH_EDITOR
@@ -5084,27 +4739,9 @@ void UMaterialExpressionSubtract::GetCaption(TArray<FString>& OutCaptions) const
 UMaterialExpressionSmoothStep::UMaterialExpressionSmoothStep(const FObjectInitializer& ObjectInitializer)
 	: Super(ObjectInitializer)
 {
-	// Structure to hold one-time initialization
-	struct FConstructorStatics
-	{
-		FText NAME_Math;
-		FText NAME_Utility;
-		FConstructorStatics()
-			: NAME_Math(LOCTEXT("Math", "Math"))
-			, NAME_Utility(LOCTEXT("Utility", "Utility"))
-		{
-		}
-	};
-	static FConstructorStatics ConstructorStatics;
-
 	ConstMin = 0.0f;
 	ConstMax = 1.0f;
 	ConstValue = 0.0f;
-
-#if WITH_EDITORONLY_DATA
-	MenuCategories.Add(ConstructorStatics.NAME_Math);
-	MenuCategories.Add(ConstructorStatics.NAME_Utility);
-#endif
 }
 
 #if WITH_EDITOR
@@ -5147,26 +4784,8 @@ void UMaterialExpressionSmoothStep::GetCaption(TArray<FString>& OutCaptions) con
 UMaterialExpressionStep::UMaterialExpressionStep(const FObjectInitializer& ObjectInitializer)
 	: Super(ObjectInitializer)
 {
-	// Structure to hold one-time initialization
-	struct FConstructorStatics
-	{
-		FText NAME_Math;
-		FText NAME_Utility;
-		FConstructorStatics()
-			: NAME_Math(LOCTEXT("Math", "Math"))
-			, NAME_Utility(LOCTEXT("Utility", "Utility"))
-		{
-		}
-	};
-	static FConstructorStatics ConstructorStatics;
-
 	ConstY = 0.0f;
 	ConstX = 1.0f;
-
-#if WITH_EDITORONLY_DATA
-	MenuCategories.Add(ConstructorStatics.NAME_Math);
-	MenuCategories.Add(ConstructorStatics.NAME_Utility);
-#endif
 }
 
 #if WITH_EDITOR
@@ -5205,28 +4824,10 @@ void UMaterialExpressionStep::GetCaption(TArray<FString>& OutCaptions) const
 UMaterialExpressionInverseLinearInterpolate::UMaterialExpressionInverseLinearInterpolate(const FObjectInitializer& ObjectInitializer)
 	: Super(ObjectInitializer)
 {
-	// Structure to hold one-time initialization
-	struct FConstructorStatics
-	{
-		FText NAME_Math;
-		FText NAME_Utility;
-		FConstructorStatics()
-			: NAME_Math(LOCTEXT("Math", "Math"))
-			, NAME_Utility(LOCTEXT("Utility", "Utility"))
-		{
-		}
-	};
-	static FConstructorStatics ConstructorStatics;
-
 	ConstA = 0.0f;
 	ConstB = 1.0f;
 	ConstValue = 0.0f;
 	bClampResult = false;
-
-#if WITH_EDITORONLY_DATA
-	MenuCategories.Add(ConstructorStatics.NAME_Math);
-	MenuCategories.Add(ConstructorStatics.NAME_Utility);
-#endif
 }
 
 #if WITH_EDITOR
@@ -5275,27 +4876,9 @@ void UMaterialExpressionInverseLinearInterpolate::GetCaption(TArray<FString>& Ou
 UMaterialExpressionLinearInterpolate::UMaterialExpressionLinearInterpolate(const FObjectInitializer& ObjectInitializer)
 	: Super(ObjectInitializer)
 {
-	// Structure to hold one-time initialization
-	struct FConstructorStatics
-	{
-		FText NAME_Math;
-		FText NAME_Utility;
-		FConstructorStatics()
-			: NAME_Math(LOCTEXT( "Math", "Math" ))
-			, NAME_Utility(LOCTEXT( "Utility", "Utility" ))
-		{
-		}
-	};
-	static FConstructorStatics ConstructorStatics;
-
 	ConstA = 0;
 	ConstB = 1;
 	ConstAlpha = 0.5f;
-
-#if WITH_EDITORONLY_DATA
-	MenuCategories.Add(ConstructorStatics.NAME_Math);
-	MenuCategories.Add(ConstructorStatics.NAME_Utility);
-#endif
 }
 
 #if WITH_EDITOR
@@ -5334,19 +4917,7 @@ void UMaterialExpressionLinearInterpolate::GetCaption(TArray<FString>& OutCaptio
 UMaterialExpressionGenericConstant::UMaterialExpressionGenericConstant(const FObjectInitializer& ObjectInitializer)
 	: Super(ObjectInitializer)
 {
-	// Structure to hold one-time initialization
-	struct FConstructorStatics
-	{
-		FText NAME_Constants;
-		FConstructorStatics()
-			: NAME_Constants(LOCTEXT("Constants", "Constants"))
-		{
-		}
-	};
-	static FConstructorStatics ConstructorStatics;
-
 #if WITH_EDITORONLY_DATA
-	MenuCategories.Add(ConstructorStatics.NAME_Constants);
 	bCollapsed = true;
 #endif
 }
@@ -5379,20 +4950,7 @@ FString UMaterialExpressionGenericConstant::GetDescription() const
 UMaterialExpressionConstant::UMaterialExpressionConstant(const FObjectInitializer& ObjectInitializer)
 	: Super(ObjectInitializer)
 {
-	// Structure to hold one-time initialization
-	struct FConstructorStatics
-	{
-		FText NAME_Constants;
-		FConstructorStatics()
-			: NAME_Constants(LOCTEXT( "Constants", "Constants" ))
-		{
-		}
-	};
-	static FConstructorStatics ConstructorStatics;
-
 #if WITH_EDITORONLY_DATA
-	MenuCategories.Add(ConstructorStatics.NAME_Constants);
-
 	bCollapsed = true;
 #endif
 }
@@ -5421,23 +4979,7 @@ FString UMaterialExpressionConstant::GetDescription() const
 UMaterialExpressionConstant2Vector::UMaterialExpressionConstant2Vector(const FObjectInitializer& ObjectInitializer)
 	: Super(ObjectInitializer)
 {
-	// Structure to hold one-time initialization
-	struct FConstructorStatics
-	{
-		FText NAME_Constants;
-		FText NAME_Vectors;
-		FConstructorStatics()
-			: NAME_Constants(LOCTEXT( "Constants", "Constants" ))
-			, NAME_Vectors(LOCTEXT( "Vectors", "Vectors" ))
-		{
-		}
-	};
-	static FConstructorStatics ConstructorStatics;
-
 #if WITH_EDITORONLY_DATA
-	MenuCategories.Add(ConstructorStatics.NAME_Constants);
-	MenuCategories.Add(ConstructorStatics.NAME_Vectors);
-
 	bCollapsed = true;
 
 	Outputs.Reset();
@@ -5471,23 +5013,7 @@ FString UMaterialExpressionConstant2Vector::GetDescription() const
 UMaterialExpressionConstant3Vector::UMaterialExpressionConstant3Vector(const FObjectInitializer& ObjectInitializer)
 	: Super(ObjectInitializer)
 {
-	// Structure to hold one-time initialization
-	struct FConstructorStatics
-	{
-		FText NAME_Constants;
-		FText NAME_Vectors;
-		FConstructorStatics()
-			: NAME_Constants(LOCTEXT( "Constants", "Constants" ))
-			, NAME_Vectors(LOCTEXT( "Vectors", "Vectors" ))
-		{
-		}
-	};
-	static FConstructorStatics ConstructorStatics;
-
 #if WITH_EDITORONLY_DATA
-	MenuCategories.Add(ConstructorStatics.NAME_Constants);
-	MenuCategories.Add(ConstructorStatics.NAME_Vectors);
-
 	bCollapsed = false;
 
 	Outputs.Reset();
@@ -5522,23 +5048,7 @@ FString UMaterialExpressionConstant3Vector::GetDescription() const
 UMaterialExpressionConstant4Vector::UMaterialExpressionConstant4Vector(const FObjectInitializer& ObjectInitializer)
 	: Super(ObjectInitializer)
 {
-	// Structure to hold one-time initialization
-	struct FConstructorStatics
-	{
-		FText NAME_Constants;
-		FText NAME_Vectors;
-		FConstructorStatics()
-			: NAME_Constants(LOCTEXT( "Constants", "Constants" ))
-			, NAME_Vectors(LOCTEXT( "Vectors", "Vectors" ))
-		{
-		}
-	};
-	static FConstructorStatics ConstructorStatics;
-
 #if WITH_EDITORONLY_DATA
-	MenuCategories.Add(ConstructorStatics.NAME_Constants);
-	MenuCategories.Add(ConstructorStatics.NAME_Vectors);
-
 	bCollapsed = false;
 
 	Outputs.Reset();
@@ -5575,24 +5085,9 @@ FString UMaterialExpressionConstant4Vector::GetDescription() const
 UMaterialExpressionClamp::UMaterialExpressionClamp(const FObjectInitializer& ObjectInitializer)
 	: Super(ObjectInitializer)
 {
-	// Structure to hold one-time initialization
-	struct FConstructorStatics
-	{
-		FText NAME_Math;
-		FConstructorStatics()
-			: NAME_Math(LOCTEXT( "Math", "Math" ))
-		{
-		}
-	};
-	static FConstructorStatics ConstructorStatics;
-
 	ClampMode = CMODE_Clamp;
 	MinDefault = 0.0f;
 	MaxDefault = 1.0f;
-
-#if WITH_EDITORONLY_DATA
-	MenuCategories.Add(ConstructorStatics.NAME_Math);
-#endif
 }
 
 void UMaterialExpressionClamp::Serialize(FStructuredArchive::FRecord Record)
@@ -5651,25 +5146,6 @@ void UMaterialExpressionClamp::GetCaption(TArray<FString>& OutCaptions) const
 //
 //	UMaterialExpressionSaturate
 //
-UMaterialExpressionSaturate::UMaterialExpressionSaturate(const FObjectInitializer& ObjectInitializer)
-	: Super(ObjectInitializer)
-{
-	// Structure to hold one-time initialization
-	struct FConstructorStatics
-	{
-		FText NAME_Math;
-		FConstructorStatics()
-			: NAME_Math(LOCTEXT( "Math", "Math" ))
-		{
-		}
-	};
-	static FConstructorStatics ConstructorStatics;
-
-#if WITH_EDITORONLY_DATA
-	MenuCategories.Add(ConstructorStatics.NAME_Math);
-#endif
-}
-
 #if WITH_EDITOR
 int32 UMaterialExpressionSaturate::Compile(class FMaterialCompiler* Compiler, int32 OutputIndex)
 {
@@ -5699,23 +5175,8 @@ void UMaterialExpressionSaturate::GetExpressionToolTip(TArray<FString>& OutToolT
 UMaterialExpressionMin::UMaterialExpressionMin(const FObjectInitializer& ObjectInitializer)
 	: Super(ObjectInitializer)
 {
-	// Structure to hold one-time initialization
-	struct FConstructorStatics
-	{
-		FText NAME_Math;
-		FConstructorStatics()
-			: NAME_Math(LOCTEXT( "Math", "Math" ))
-		{
-		}
-	};
-	static FConstructorStatics ConstructorStatics;
-
 	ConstA = 0.0f;
 	ConstB = 1.0f;
-
-#if WITH_EDITORONLY_DATA
-	MenuCategories.Add(ConstructorStatics.NAME_Math);
-#endif
 }
 
 #if WITH_EDITOR
@@ -5754,23 +5215,8 @@ void UMaterialExpressionMin::GetCaption(TArray<FString>& OutCaptions) const
 UMaterialExpressionMax::UMaterialExpressionMax(const FObjectInitializer& ObjectInitializer)
 	: Super(ObjectInitializer)
 {
-	// Structure to hold one-time initialization
-	struct FConstructorStatics
-	{
-		FText NAME_Math;
-		FConstructorStatics()
-			: NAME_Math(LOCTEXT( "Math", "Math" ))
-		{
-		}
-	};
-	static FConstructorStatics ConstructorStatics;
-
 	ConstA = 0.0f;
 	ConstB = 1.0f;
-
-#if WITH_EDITORONLY_DATA
-	MenuCategories.Add(ConstructorStatics.NAME_Math);
-#endif
 }
 
 #if WITH_EDITOR
@@ -5809,23 +5255,10 @@ void UMaterialExpressionMax::GetCaption(TArray<FString>& OutCaptions) const
 UMaterialExpressionTextureCoordinate::UMaterialExpressionTextureCoordinate(const FObjectInitializer& ObjectInitializer)
 	: Super(ObjectInitializer)
 {
-	// Structure to hold one-time initialization
-	struct FConstructorStatics
-	{
-		FText NAME_Coordinates;
-		FConstructorStatics()
-			: NAME_Coordinates(LOCTEXT( "Coordinates", "Coordinates" ))
-		{
-		}
-	};
-	static FConstructorStatics ConstructorStatics;
-
 	UTiling = 1.0f;
 	VTiling = 1.0f;
 
 #if WITH_EDITORONLY_DATA
-	MenuCategories.Add(ConstructorStatics.NAME_Coordinates);
-
 	bCollapsed = true;
 #endif
 }
@@ -5876,28 +5309,6 @@ bool UMaterialExpressionTextureCoordinate::MatchesSearchQuery(const TCHAR* Searc
 }
 #endif // WITH_EDITOR
 
-UMaterialExpressionDotProduct::UMaterialExpressionDotProduct(const FObjectInitializer& ObjectInitializer)
-	: Super(ObjectInitializer)
-{
-	// Structure to hold one-time initialization
-	struct FConstructorStatics
-	{
-		FText NAME_Math;
-		FText NAME_VectorOps;
-		FConstructorStatics()
-			: NAME_Math(LOCTEXT( "Math", "Math" ))
-			, NAME_VectorOps(LOCTEXT( "VectorOps", "VectorOps" ))
-		{
-		}
-	};
-	static FConstructorStatics ConstructorStatics;
-
-#if WITH_EDITORONLY_DATA
-	MenuCategories.Add(ConstructorStatics.NAME_Math);
-	MenuCategories.Add(ConstructorStatics.NAME_VectorOps);
-#endif
-}
-
 #if WITH_EDITOR
 int32 UMaterialExpressionDotProduct::Compile(class FMaterialCompiler* Compiler, int32 OutputIndex)
 {
@@ -5926,28 +5337,6 @@ void UMaterialExpressionDotProduct::GetCaption(TArray<FString>& OutCaptions) con
 }
 #endif // WITH_EDITOR
 
-UMaterialExpressionCrossProduct::UMaterialExpressionCrossProduct(const FObjectInitializer& ObjectInitializer)
-	: Super(ObjectInitializer)
-{
-	// Structure to hold one-time initialization
-	struct FConstructorStatics
-	{
-		FText NAME_Math;
-		FText NAME_VectorOps;
-		FConstructorStatics()
-			: NAME_Math(LOCTEXT( "Math", "Math" ))
-			, NAME_VectorOps(LOCTEXT( "VectorOps", "VectorOps" ))
-		{
-		}
-	};
-	static FConstructorStatics ConstructorStatics;
-
-#if WITH_EDITORONLY_DATA
-	MenuCategories.Add(ConstructorStatics.NAME_Math);
-	MenuCategories.Add(ConstructorStatics.NAME_VectorOps);
-#endif
-}
-
 #if WITH_EDITOR
 int32 UMaterialExpressionCrossProduct::Compile(class FMaterialCompiler* Compiler, int32 OutputIndex)
 {
@@ -5976,28 +5365,6 @@ void UMaterialExpressionCrossProduct::GetCaption(TArray<FString>& OutCaptions) c
 }
 #endif // WITH_EDITOR
 
-UMaterialExpressionComponentMask::UMaterialExpressionComponentMask(const FObjectInitializer& ObjectInitializer)
-	: Super(ObjectInitializer)
-{
-	// Structure to hold one-time initialization
-	struct FConstructorStatics
-	{
-		FText NAME_Math;
-		FText NAME_VectorOps;
-		FConstructorStatics()
-			: NAME_Math(LOCTEXT( "Math", "Math" ))
-			, NAME_VectorOps(LOCTEXT( "VectorOps", "VectorOps" ))
-		{
-		}
-	};
-	static FConstructorStatics ConstructorStatics;
-
-#if WITH_EDITORONLY_DATA
-	MenuCategories.Add(ConstructorStatics.NAME_Math);
-	MenuCategories.Add(ConstructorStatics.NAME_VectorOps);
-#endif
-}
-
 #if WITH_EDITOR
 int32 UMaterialExpressionComponentMask::Compile(class FMaterialCompiler* Compiler, int32 OutputIndex)
 {
@@ -6030,25 +5397,6 @@ void UMaterialExpressionComponentMask::GetCaption(TArray<FString>& OutCaptions)
 //
 //	UMaterialExpressionStaticComponentMaskParameter
 //
-UMaterialExpressionStaticComponentMaskParameter::UMaterialExpressionStaticComponentMaskParameter(const FObjectInitializer& ObjectInitializer)
-	: Super(ObjectInitializer)
-{
-	// Structure to hold one-time initialization
-	struct FConstructorStatics
-	{
-		FText NAME_Parameters;
-		FConstructorStatics()
-			: NAME_Parameters(LOCTEXT( "Parameters", "Parameters" ))
-		{
-		}
-	};
-	static FConstructorStatics ConstructorStatics;
-
-#if WITH_EDITORONLY_DATA
-	MenuCategories.Add(ConstructorStatics.NAME_Parameters);
-#endif
-}
-
 #if WITH_EDITOR
 int32 UMaterialExpressionStaticComponentMaskParameter::Compile(class FMaterialCompiler* Compiler, int32 OutputIndex)
 {
@@ -6108,20 +5456,7 @@ bool  UMaterialExpressionStaticComponentMaskParameter::SetParameterValue(FName I
 UMaterialExpressionTime::UMaterialExpressionTime(const FObjectInitializer& ObjectInitializer)
 	: Super(ObjectInitializer)
 {
-	// Structure to hold one-time initialization
-	struct FConstructorStatics
-	{
-		FText NAME_Constants;
-		FConstructorStatics()
-			: NAME_Constants(LOCTEXT( "Constants", "Constants" ))
-		{
-		}
-	};
-	static FConstructorStatics ConstructorStatics;
-
 #if WITH_EDITORONLY_DATA
-	MenuCategories.Add(ConstructorStatics.NAME_Constants);
-
 	bShaderInputData = true;
 #endif
 
@@ -6158,20 +5493,7 @@ void UMaterialExpressionTime::GetCaption(TArray<FString>& OutCaptions) const
 UMaterialExpressionCameraVectorWS::UMaterialExpressionCameraVectorWS(const FObjectInitializer& ObjectInitializer)
 	: Super(ObjectInitializer)
 {
-	// Structure to hold one-time initialization
-	struct FConstructorStatics
-	{
-		FText NAME_Vectors;
-		FConstructorStatics()
-			: NAME_Vectors(LOCTEXT( "Vectors", "Vectors" ))
-		{
-		}
-	};
-	static FConstructorStatics ConstructorStatics;
-
 #if WITH_EDITORONLY_DATA
-	MenuCategories.Add(ConstructorStatics.NAME_Vectors);
-
 	bShaderInputData = true;
 #endif
 }
@@ -6191,23 +5513,7 @@ void UMaterialExpressionCameraVectorWS::GetCaption(TArray<FString>& OutCaptions)
 UMaterialExpressionCameraPositionWS::UMaterialExpressionCameraPositionWS(const FObjectInitializer& ObjectInitializer)
 	: Super(ObjectInitializer)
 {
-	// Structure to hold one-time initialization
-	struct FConstructorStatics
-	{
-		FText NAME_Vectors;
-		FText NAME_Coordinates;
-		FConstructorStatics()
-			: NAME_Vectors(LOCTEXT( "Vectors", "Vectors" ))
-			, NAME_Coordinates(LOCTEXT( "Coordinates", "Coordinates" ))
-		{
-		}
-	};
-	static FConstructorStatics ConstructorStatics;
-
 #if WITH_EDITORONLY_DATA
-	MenuCategories.Add(ConstructorStatics.NAME_Vectors);
-	MenuCategories.Add(ConstructorStatics.NAME_Coordinates);
-
 	bShaderInputData = true;
 #endif
 }
@@ -6231,20 +5537,7 @@ void UMaterialExpressionCameraPositionWS::GetCaption(TArray<FString>& OutCaption
 UMaterialExpressionReflectionVectorWS::UMaterialExpressionReflectionVectorWS(const FObjectInitializer& ObjectInitializer)
 	: Super(ObjectInitializer)
 {
-	// Structure to hold one-time initialization
-	struct FConstructorStatics
-	{
-		FText NAME_Vectors;
-		FConstructorStatics()
-			: NAME_Vectors(LOCTEXT( "Vectors", "Vectors" ))
-		{
-		}
-	};
-	static FConstructorStatics ConstructorStatics;
-
 #if WITH_EDITORONLY_DATA
-	MenuCategories.Add(ConstructorStatics.NAME_Vectors);
-
 	bShaderInputData = true;
 #endif
 }
@@ -6276,20 +5569,7 @@ void UMaterialExpressionReflectionVectorWS::GetCaption(TArray<FString>& OutCapti
 UMaterialExpressionPanner::UMaterialExpressionPanner(const FObjectInitializer& ObjectInitializer)
 	: Super(ObjectInitializer)
 {
-	// Structure to hold one-time initialization
-	struct FConstructorStatics
-	{
-		FText NAME_Coordinates;
-		FConstructorStatics()
-			: NAME_Coordinates(LOCTEXT( "Coordinates", "Coordinates" ))
-		{
-		}
-	};
-	static FConstructorStatics ConstructorStatics;
-
 #if WITH_EDITORONLY_DATA
-	MenuCategories.Add(ConstructorStatics.NAME_Coordinates);
-
 	bCollapsed = true;
 #endif
 	ConstCoordinate = 0;
@@ -6340,25 +5620,12 @@ void UMaterialExpressionPanner::GetCaption(TArray<FString>& OutCaptions) const
 UMaterialExpressionRotator::UMaterialExpressionRotator(const FObjectInitializer& ObjectInitializer)
 	: Super(ObjectInitializer)
 {
-	// Structure to hold one-time initialization
-	struct FConstructorStatics
-	{
-		FText NAME_Coordinates;
-		FConstructorStatics()
-			: NAME_Coordinates(LOCTEXT( "Coordinates", "Coordinates" ))
-		{
-		}
-	};
-	static FConstructorStatics ConstructorStatics;
-
 	CenterX = 0.5f;
 	CenterY = 0.5f;
 	Speed = 0.25f;
 	ConstCoordinate = 0;
 
 #if WITH_EDITORONLY_DATA
-	MenuCategories.Add(ConstructorStatics.NAME_Coordinates);
-
 	bCollapsed = true;
 #endif
 }
@@ -6414,22 +5681,7 @@ void UMaterialExpressionRotator::GetCaption(TArray<FString>& OutCaptions) const
 UMaterialExpressionSine::UMaterialExpressionSine(const FObjectInitializer& ObjectInitializer)
 	: Super(ObjectInitializer)
 {
-	// Structure to hold one-time initialization
-	struct FConstructorStatics
-	{
-		FText NAME_Math;
-		FConstructorStatics()
-			: NAME_Math(LOCTEXT( "Math", "Math" ))
-		{
-		}
-	};
-	static FConstructorStatics ConstructorStatics;
-
 	Period = 1.0f;
-
-#if WITH_EDITORONLY_DATA
-	MenuCategories.Add(ConstructorStatics.NAME_Math);
-#endif
 }
 
 #if WITH_EDITOR
@@ -6455,22 +5707,7 @@ void UMaterialExpressionSine::GetCaption(TArray<FString>& OutCaptions) const
 UMaterialExpressionCosine::UMaterialExpressionCosine(const FObjectInitializer& ObjectInitializer)
 	: Super(ObjectInitializer)
 {
-	// Structure to hold one-time initialization
-	struct FConstructorStatics
-	{
-		FText NAME_Math;
-		FConstructorStatics()
-			: NAME_Math(LOCTEXT( "Math", "Math" ))
-		{
-		}
-	};
-	static FConstructorStatics ConstructorStatics;
-
 	Period = 1.0f;
-
-#if WITH_EDITORONLY_DATA
-	MenuCategories.Add(ConstructorStatics.NAME_Math);
-#endif
 }
 
 #if WITH_EDITOR
@@ -6496,22 +5733,7 @@ void UMaterialExpressionCosine::GetCaption(TArray<FString>& OutCaptions) const
 UMaterialExpressionTangent::UMaterialExpressionTangent(const FObjectInitializer& ObjectInitializer)
 	: Super(ObjectInitializer)
 {
-	// Structure to hold one-time initialization
-	struct FConstructorStatics
-	{
-		FText NAME_Math;
-		FConstructorStatics()
-			: NAME_Math(LOCTEXT( "Math", "Math" ))
-		{
-		}
-	};
-	static FConstructorStatics ConstructorStatics;
-
 	Period = 1.0f;
-
-#if WITH_EDITORONLY_DATA
-	MenuCategories.Add(ConstructorStatics.NAME_Math);
-#endif
 }
 
 #if WITH_EDITOR
@@ -6534,25 +5756,6 @@ void UMaterialExpressionTangent::GetCaption(TArray<FString>& OutCaptions) const
 //
 //	UMaterialExpressionArcsine
 //
-UMaterialExpressionArcsine::UMaterialExpressionArcsine(const FObjectInitializer& ObjectInitializer)
-	: Super(ObjectInitializer)
-{
-	// Structure to hold one-time initialization
-	struct FConstructorStatics
-	{
-		FText NAME_Math;
-		FConstructorStatics()
-			: NAME_Math(LOCTEXT( "Math", "Math" ))
-		{
-		}
-	};
-	static FConstructorStatics ConstructorStatics;
-
-#if WITH_EDITORONLY_DATA
-	MenuCategories.Add(ConstructorStatics.NAME_Math);
-#endif
-}
-
 #if WITH_EDITOR
 int32 UMaterialExpressionArcsine::Compile(class FMaterialCompiler* Compiler, int32 OutputIndex)
 {
@@ -6578,25 +5781,6 @@ void UMaterialExpressionArcsine::GetExpressionToolTip(TArray<FString>& OutToolTi
 //
 //	UMaterialExpressionArcsineFast
 //
-UMaterialExpressionArcsineFast::UMaterialExpressionArcsineFast(const FObjectInitializer& ObjectInitializer)
-	: Super(ObjectInitializer)
-{
-	// Structure to hold one-time initialization
-	struct FConstructorStatics
-	{
-		FText NAME_Math;
-		FConstructorStatics()
-			: NAME_Math(LOCTEXT( "Math", "Math" ))
-		{
-		}
-	};
-	static FConstructorStatics ConstructorStatics;
-
-#if WITH_EDITORONLY_DATA
-	MenuCategories.Add(ConstructorStatics.NAME_Math);
-#endif
-}
-
 #if WITH_EDITOR
 int32 UMaterialExpressionArcsineFast::Compile(class FMaterialCompiler* Compiler, int32 OutputIndex)
 {
@@ -6622,25 +5806,6 @@ void UMaterialExpressionArcsineFast::GetExpressionToolTip(TArray<FString>& OutTo
 //
 //	UMaterialExpressionArccosine
 //
-UMaterialExpressionArccosine::UMaterialExpressionArccosine(const FObjectInitializer& ObjectInitializer)
-	: Super(ObjectInitializer)
-{
-	// Structure to hold one-time initialization
-	struct FConstructorStatics
-	{
-		FText NAME_Math;
-		FConstructorStatics()
-			: NAME_Math(LOCTEXT( "Math", "Math" ))
-		{
-		}
-	};
-	static FConstructorStatics ConstructorStatics;
-
-#if WITH_EDITORONLY_DATA
-	MenuCategories.Add(ConstructorStatics.NAME_Math);
-#endif
-}
-
 #if WITH_EDITOR
 int32 UMaterialExpressionArccosine::Compile(class FMaterialCompiler* Compiler, int32 OutputIndex)
 {
@@ -6666,25 +5831,6 @@ void UMaterialExpressionArccosine::GetExpressionToolTip(TArray<FString>& OutTool
 //
 //	UMaterialExpressionArccosineFast
 //
-UMaterialExpressionArccosineFast::UMaterialExpressionArccosineFast(const FObjectInitializer& ObjectInitializer)
-	: Super(ObjectInitializer)
-{
-	// Structure to hold one-time initialization
-	struct FConstructorStatics
-	{
-		FText NAME_Math;
-		FConstructorStatics()
-			: NAME_Math(LOCTEXT( "Math", "Math" ))
-		{
-		}
-	};
-	static FConstructorStatics ConstructorStatics;
-
-#if WITH_EDITORONLY_DATA
-	MenuCategories.Add(ConstructorStatics.NAME_Math);
-#endif
-}
-
 #if WITH_EDITOR
 int32 UMaterialExpressionArccosineFast::Compile(class FMaterialCompiler* Compiler, int32 OutputIndex)
 {
@@ -6710,25 +5856,6 @@ void UMaterialExpressionArccosineFast::GetExpressionToolTip(TArray<FString>& Out
 //
 //	UMaterialExpressionArctangent
 //
-UMaterialExpressionArctangent::UMaterialExpressionArctangent(const FObjectInitializer& ObjectInitializer)
-	: Super(ObjectInitializer)
-{
-	// Structure to hold one-time initialization
-	struct FConstructorStatics
-	{
-		FText NAME_Math;
-		FConstructorStatics()
-			: NAME_Math(LOCTEXT( "Math", "Math" ))
-		{
-		}
-	};
-	static FConstructorStatics ConstructorStatics;
-
-#if WITH_EDITORONLY_DATA
-	MenuCategories.Add(ConstructorStatics.NAME_Math);
-#endif
-}
-
 #if WITH_EDITOR
 int32 UMaterialExpressionArctangent::Compile(class FMaterialCompiler* Compiler, int32 OutputIndex)
 {
@@ -6754,25 +5881,6 @@ void UMaterialExpressionArctangent::GetExpressionToolTip(TArray<FString>& OutToo
 //
 //	UMaterialExpressionArctangentFast
 //
-UMaterialExpressionArctangentFast::UMaterialExpressionArctangentFast(const FObjectInitializer& ObjectInitializer)
-	: Super(ObjectInitializer)
-{
-	// Structure to hold one-time initialization
-	struct FConstructorStatics
-	{
-		FText NAME_Math;
-		FConstructorStatics()
-			: NAME_Math(LOCTEXT( "Math", "Math" ))
-		{
-		}
-	};
-	static FConstructorStatics ConstructorStatics;
-
-#if WITH_EDITORONLY_DATA
-	MenuCategories.Add(ConstructorStatics.NAME_Math);
-#endif
-}
-
 #if WITH_EDITOR
 int32 UMaterialExpressionArctangentFast::Compile(class FMaterialCompiler* Compiler, int32 OutputIndex)
 {
@@ -6798,25 +5906,6 @@ void UMaterialExpressionArctangentFast::GetExpressionToolTip(TArray<FString>& Ou
 //
 //	UMaterialExpressionArctangent2
 //
-UMaterialExpressionArctangent2::UMaterialExpressionArctangent2(const FObjectInitializer& ObjectInitializer)
-	: Super(ObjectInitializer)
-{
-	// Structure to hold one-time initialization
-	struct FConstructorStatics
-	{
-		FText NAME_Math;
-		FConstructorStatics()
-			: NAME_Math(LOCTEXT( "Math", "Math" ))
-		{
-		}
-	};
-	static FConstructorStatics ConstructorStatics;
-
-#if WITH_EDITORONLY_DATA
-	MenuCategories.Add(ConstructorStatics.NAME_Math);
-#endif
-}
-
 #if WITH_EDITOR
 int32 UMaterialExpressionArctangent2::Compile(class FMaterialCompiler* Compiler, int32 OutputIndex)
 {
@@ -6844,25 +5933,6 @@ void UMaterialExpressionArctangent2::GetExpressionToolTip(TArray<FString>& OutTo
 //
 //	UMaterialExpressionArctangent2Fast
 //
-UMaterialExpressionArctangent2Fast::UMaterialExpressionArctangent2Fast(const FObjectInitializer& ObjectInitializer)
-	: Super(ObjectInitializer)
-{
-	// Structure to hold one-time initialization
-	struct FConstructorStatics
-	{
-		FText NAME_Math;
-		FConstructorStatics()
-			: NAME_Math(LOCTEXT( "Math", "Math" ))
-		{
-		}
-	};
-	static FConstructorStatics ConstructorStatics;
-
-#if WITH_EDITORONLY_DATA
-	MenuCategories.Add(ConstructorStatics.NAME_Math);
-#endif
-}
-
 #if WITH_EDITOR
 int32 UMaterialExpressionArctangent2Fast::Compile(class FMaterialCompiler* Compiler, int32 OutputIndex)
 {
@@ -6890,24 +5960,11 @@ void UMaterialExpressionArctangent2Fast::GetExpressionToolTip(TArray<FString>& O
 UMaterialExpressionBumpOffset::UMaterialExpressionBumpOffset(const FObjectInitializer& ObjectInitializer)
 	: Super(ObjectInitializer)
 {
-	// Structure to hold one-time initialization
-	struct FConstructorStatics
-	{
-		FText NAME_Utility;
-		FConstructorStatics()
-			: NAME_Utility(LOCTEXT( "Utility", "Utility" ))
-		{
-		}
-	};
-	static FConstructorStatics ConstructorStatics;
-
 	HeightRatio = 0.05f;
 	ReferencePlane = 0.5f;
 	ConstCoordinate = 0;
 #if WITH_EDITORONLY_DATA
 	bCollapsed = false;
-
-	MenuCategories.Add(ConstructorStatics.NAME_Utility);
 #endif
 }
 
@@ -6940,28 +5997,6 @@ void UMaterialExpressionBumpOffset::GetCaption(TArray<FString>& OutCaptions) con
 }
 #endif // WITH_EDITOR
 
-UMaterialExpressionAppendVector::UMaterialExpressionAppendVector(const FObjectInitializer& ObjectInitializer)
-	: Super(ObjectInitializer)
-{
-	// Structure to hold one-time initialization
-	struct FConstructorStatics
-	{
-		FText NAME_Math;
-		FText NAME_VectorOps;
-		FConstructorStatics()
-			: NAME_Math(LOCTEXT( "Math", "Math" ))
-			, NAME_VectorOps(LOCTEXT( "VectorOps", "VectorOps" ))
-		{
-		}
-	};
-	static FConstructorStatics ConstructorStatics;
-
-#if WITH_EDITORONLY_DATA
-	MenuCategories.Add(ConstructorStatics.NAME_Math);
-	MenuCategories.Add(ConstructorStatics.NAME_VectorOps);
-#endif
-}
-
 #if WITH_EDITOR
 int32 UMaterialExpressionAppendVector::Compile(class FMaterialCompiler* Compiler, int32 OutputIndex)
 {
@@ -6991,26 +6026,6 @@ void UMaterialExpressionAppendVector::GetCaption(TArray<FString>& OutCaptions) c
 #endif // WITH_EDITOR
 
 // -----
-
-UMaterialExpressionMakeMaterialAttributes::UMaterialExpressionMakeMaterialAttributes(const FObjectInitializer& ObjectInitializer)
-	: Super(ObjectInitializer)
-{
-	// Structure to hold one-time initialization
-	struct FConstructorStatics
-	{
-		FText NAME_MaterialAttributes;
-		FConstructorStatics()
-			: NAME_MaterialAttributes(LOCTEXT( "MaterialAttributes", "Material Attributes" ))
-		{
-		}
-	};
-	static FConstructorStatics ConstructorStatics;
-
-#if WITH_EDITORONLY_DATA
-	MenuCategories.Add(ConstructorStatics.NAME_MaterialAttributes);
-#endif
-}
-
 FExpressionInput* UMaterialExpressionMakeMaterialAttributes::GetExpressionInput(EMaterialProperty InProperty)
 {
 	switch (InProperty)
@@ -7165,23 +6180,9 @@ UMaterialExpressionBreakMaterialAttributes::UMaterialExpressionBreakMaterialAttr
 	: Super(ObjectInitializer)
 {
 #if WITH_EDITORONLY_DATA
-	// Structure to hold one-time initialization
-	struct FConstructorStatics
-	{
-		FText NAME_MaterialAttributes;
-		FConstructorStatics()
-			: NAME_MaterialAttributes(LOCTEXT( "MaterialAttributes", "Material Attributes" ))
-		{
-		}
-	};
-
-	static FConstructorStatics ConstructorStatics;
-
 	bShowOutputNameOnPin = true;
 	bShowMaskColorsOnPin = false;
 
-	MenuCategories.Add(ConstructorStatics.NAME_MaterialAttributes);
-	
  	static_assert(MP_MAX == 35, 
 		"New material properties should be added to the end of the outputs for this expression. \
 		The order of properties here should match the material results pins, the make material attributes node inputs and the mapping of IO indices to properties in GetMaterialPropertyFromInputOutputIndex().\
@@ -7356,19 +6357,7 @@ uint32 UMaterialExpressionBreakMaterialAttributes::GetOutputType(int32 OutputInd
 UMaterialExpressionGetMaterialAttributes::UMaterialExpressionGetMaterialAttributes(const FObjectInitializer& ObjectInitializer)
 	: Super(ObjectInitializer)
 {
-	// Structure to hold one-time initialization
-	struct FConstructorStatics
-	{
-		FText NAME_MaterialAttributes;
-		FConstructorStatics()
-			: NAME_MaterialAttributes(LOCTEXT( "MaterialAttributes", "Material Attributes" ))
-		{
-		}
-	};
-	static FConstructorStatics ConstructorStatics;
 #if WITH_EDITORONLY_DATA
-	MenuCategories.Add(ConstructorStatics.NAME_MaterialAttributes);
-
 	bShowOutputNameOnPin = true;
 #endif
 
@@ -7638,20 +6627,6 @@ void UMaterialExpressionGetMaterialAttributes::PostLoad()
 UMaterialExpressionSetMaterialAttributes::UMaterialExpressionSetMaterialAttributes(const FObjectInitializer& ObjectInitializer)
 	: Super(ObjectInitializer)
 {
-	// Structure to hold one-time initialization
-	struct FConstructorStatics
-	{
-		FText NAME_MaterialAttributes;
-		FConstructorStatics()
-			: NAME_MaterialAttributes(LOCTEXT( "MaterialAttributes", "Material Attributes" ))
-		{
-		}
-	};
-	static FConstructorStatics ConstructorStatics;
-#if WITH_EDITORONLY_DATA
-	MenuCategories.Add(ConstructorStatics.NAME_MaterialAttributes);
-#endif
-
 #if WITH_EDITOR
 	// Add default input pins
 	Inputs.Reset();
@@ -8015,19 +6990,7 @@ UMaterialExpressionBlendMaterialAttributes::UMaterialExpressionBlendMaterialAttr
 	, PixelAttributeBlendType(EMaterialAttributeBlend::Blend)
 	, VertexAttributeBlendType(EMaterialAttributeBlend::Blend)
 {
-	// Structure to hold one-time initialization
-	struct FConstructorStatics
-	{
-		FText NAME_MaterialAttributes;
-		FConstructorStatics()
-			: NAME_MaterialAttributes(LOCTEXT( "MaterialAttributes", "Material Attributes" ))
-		{
-		}
-	};
-	static FConstructorStatics ConstructorStatics;
 #if WITH_EDITORONLY_DATA
-	MenuCategories.Add(ConstructorStatics.NAME_MaterialAttributes);
-
 	Outputs.Reset();
 	Outputs.Add(FExpressionOutput(TEXT(""), 0, 0, 0, 0, 0));
 #endif
@@ -8355,29 +7318,12 @@ UMaterialExpressionMaterialAttributeLayers::UMaterialExpressionMaterialAttribute
 	, bIsLayerGraphBuilt(false)
 	, ParamLayers(nullptr)
 {
-	// Structure to hold one-time initialization
-	struct FConstructorStatics
-	{
-		FText NAME_MaterialAttributes;
-		FText NAME_Parameters;
-		FConstructorStatics()
-			: NAME_MaterialAttributes(LOCTEXT( "MaterialAttributes", "Material Attributes" ))
-			, NAME_Parameters(LOCTEXT("Parameters", "Parameters"))
-		{
-		}
-	};
-	static FConstructorStatics ConstructorStatics;
-
 #if WITH_EDITOR
 	DefaultLayers.AddDefaultBackgroundLayer();
 	
 	CachedInputs.Empty();
 	CachedInputs.Push(&Input);
 #endif
-#if WITH_EDITORONLY_DATA
-	MenuCategories.Add(ConstructorStatics.NAME_MaterialAttributes);
-	MenuCategories.Add(ConstructorStatics.NAME_Parameters);
-#endif
 }
 
 void UMaterialExpressionMaterialAttributeLayers::PostLoad()
@@ -9052,29 +7998,9 @@ uint32 UMaterialExpressionMaterialAttributeLayers::GetInputType(int32 InputIndex
 
 
 // -----
+#if WITH_EDITOR
 
-UMaterialExpressionFloatToUInt::UMaterialExpressionFloatToUInt(const FObjectInitializer& ObjectInitializer)
-: Super(ObjectInitializer)
-{
-	// Structure to hold one-time initialization
-	struct FConstructorStatics
-	{
-		FText NAME_Math;
-		FConstructorStatics()
-		: NAME_Math(LOCTEXT( "Math", "Math" ))
-		{
-		}
-	};
-	static FConstructorStatics ConstructorStatics;
-
-	#if WITH_EDITORONLY_DATA
-	MenuCategories.Add(ConstructorStatics.NAME_Math);
-	#endif
-}
-
-#if WITH_EDITOR
-
-int32 UMaterialExpressionFloatToUInt::Compile(class FMaterialCompiler* Compiler, int32 OutputIndex)
+int32 UMaterialExpressionFloatToUInt::Compile(class FMaterialCompiler* Compiler, int32 OutputIndex)
 {
 	if (!Input.GetTracedInput().Expression)
 	{
@@ -9122,26 +8048,6 @@ uint32 UMaterialExpressionFloatToUInt::GetOutputType(int32 OutputIndex)
 
 
 // -----
-
-UMaterialExpressionUIntToFloat::UMaterialExpressionUIntToFloat(const FObjectInitializer& ObjectInitializer)
-: Super(ObjectInitializer)
-{
-	// Structure to hold one-time initialization
-	struct FConstructorStatics
-	{
-		FText NAME_Math;
-		FConstructorStatics()
-		: NAME_Math(LOCTEXT( "Math", "Math" ))
-		{
-		}
-	};
-	static FConstructorStatics ConstructorStatics;
-
-	#if WITH_EDITORONLY_DATA
-	MenuCategories.Add(ConstructorStatics.NAME_Math);
-	#endif
-}
-
 #if WITH_EDITOR
 int32 UMaterialExpressionUIntToFloat::Compile(class FMaterialCompiler* Compiler, int32 OutputIndex)
 {
@@ -9181,26 +8087,6 @@ uint32 UMaterialExpressionUIntToFloat::GetOutputType(int32 OutputIndex)
 
 
 // -----
-
-UMaterialExpressionFloor::UMaterialExpressionFloor(const FObjectInitializer& ObjectInitializer)
-	: Super(ObjectInitializer)
-{
-	// Structure to hold one-time initialization
-	struct FConstructorStatics
-	{
-		FText NAME_Math;
-		FConstructorStatics()
-			: NAME_Math(LOCTEXT( "Math", "Math" ))
-		{
-		}
-	};
-	static FConstructorStatics ConstructorStatics;
-
-#if WITH_EDITORONLY_DATA
-	MenuCategories.Add(ConstructorStatics.NAME_Math);
-#endif
-}
-
 #if WITH_EDITOR
 int32 UMaterialExpressionFloor::Compile(class FMaterialCompiler* Compiler, int32 OutputIndex)
 {
@@ -9223,25 +8109,6 @@ void UMaterialExpressionFloor::GetCaption(TArray<FString>& OutCaptions) const
 }
 #endif // WITH_EDITOR
 
-UMaterialExpressionCeil::UMaterialExpressionCeil(const FObjectInitializer& ObjectInitializer)
-	: Super(ObjectInitializer)
-{
-	// Structure to hold one-time initialization
-	struct FConstructorStatics
-	{
-		FText NAME_Math;
-		FConstructorStatics()
-			: NAME_Math(LOCTEXT( "Math", "Math" ))
-		{
-		}
-	};
-	static FConstructorStatics ConstructorStatics;
-
-#if WITH_EDITORONLY_DATA
-	MenuCategories.Add(ConstructorStatics.NAME_Math);
-#endif
-}
-
 #if WITH_EDITOR
 int32 UMaterialExpressionCeil::Compile(class FMaterialCompiler* Compiler, int32 OutputIndex)
 {
@@ -9266,25 +8133,6 @@ void UMaterialExpressionCeil::GetCaption(TArray<FString>& OutCaptions) const
 //
 //	UMaterialExpressionRound
 //
-UMaterialExpressionRound::UMaterialExpressionRound(const FObjectInitializer& ObjectInitializer)
-	: Super(ObjectInitializer)
-{
-	// Structure to hold one-time initialization
-	struct FConstructorStatics
-	{
-		FText NAME_Math;
-		FConstructorStatics()
-			: NAME_Math(LOCTEXT( "Math", "Math" ))
-		{
-		}
-	};
-	static FConstructorStatics ConstructorStatics;
-
-#if WITH_EDITORONLY_DATA
-	MenuCategories.Add(ConstructorStatics.NAME_Math);
-#endif
-}
-
 #if WITH_EDITOR
 int32 UMaterialExpressionRound::Compile(class FMaterialCompiler* Compiler, int32 OutputIndex)
 {
@@ -9314,25 +8162,6 @@ void UMaterialExpressionRound::GetExpressionToolTip(TArray<FString>& OutToolTip)
 //
 //	UMaterialExpressionTruncate
 //
-UMaterialExpressionTruncate::UMaterialExpressionTruncate(const FObjectInitializer& ObjectInitializer)
-	: Super(ObjectInitializer)
-{
-	// Structure to hold one-time initialization
-	struct FConstructorStatics
-	{
-		FText NAME_Math;
-		FConstructorStatics()
-			: NAME_Math(LOCTEXT( "Math", "Math" ))
-		{
-		}
-	};
-	static FConstructorStatics ConstructorStatics;
-
-#if WITH_EDITORONLY_DATA
-	MenuCategories.Add(ConstructorStatics.NAME_Math);
-#endif
-}
-
 #if WITH_EDITOR
 int32 UMaterialExpressionTruncate::Compile(class FMaterialCompiler* Compiler, int32 OutputIndex)
 {
@@ -9362,25 +8191,6 @@ void UMaterialExpressionTruncate::GetExpressionToolTip(TArray<FString>& OutToolT
 //
 //	UMaterialExpressionSign
 //
-UMaterialExpressionSign::UMaterialExpressionSign(const FObjectInitializer& ObjectInitializer)
-	: Super(ObjectInitializer)
-{
-	// Structure to hold one-time initialization
-	struct FConstructorStatics
-	{
-		FText NAME_Math;
-		FConstructorStatics()
-			: NAME_Math(LOCTEXT( "Math", "Math" ))
-		{
-		}
-	};
-	static FConstructorStatics ConstructorStatics;
-
-#if WITH_EDITORONLY_DATA
-	MenuCategories.Add(ConstructorStatics.NAME_Math);
-#endif
-}
-
 #if WITH_EDITOR
 int32 UMaterialExpressionSign::Compile(class FMaterialCompiler* Compiler, int32 OutputIndex)
 {
@@ -9405,26 +8215,6 @@ void UMaterialExpressionSign::GetExpressionToolTip(TArray<FString>& OutToolTip)
 //
 //	UMaterialExpressionFmod
 //
-
-UMaterialExpressionFmod::UMaterialExpressionFmod(const FObjectInitializer& ObjectInitializer)
-	: Super(ObjectInitializer)
-{
-	// Structure to hold one-time initialization
-	struct FConstructorStatics
-	{
-		FText NAME_Math;
-		FConstructorStatics()
-			: NAME_Math(LOCTEXT( "Math", "Math" ))
-		{
-		}
-	};
-	static FConstructorStatics ConstructorStatics;
-
-#if WITH_EDITORONLY_DATA
-	MenuCategories.Add(ConstructorStatics.NAME_Math);
-#endif
-}
-
 #if WITH_EDITOR
 int32 UMaterialExpressionFmod::Compile(class FMaterialCompiler* Compiler, int32 OutputIndex)
 {
@@ -9448,26 +8238,6 @@ void UMaterialExpressionFmod::GetCaption(TArray<FString>& OutCaptions) const
 //
 //	UMaterialExpressionModulo
 //
-
-UMaterialExpressionModulo::UMaterialExpressionModulo(const FObjectInitializer& ObjectInitializer)
-	: Super(ObjectInitializer)
-{
-	// Structure to hold one-time initialization
-	struct FConstructorStatics
-	{
-		FText NAME_Math;
-		FConstructorStatics()
-			: NAME_Math(LOCTEXT( "Math", "Math" ))
-		{
-		}
-	};
-	static FConstructorStatics ConstructorStatics;
-
-#if WITH_EDITORONLY_DATA
-	MenuCategories.Add(ConstructorStatics.NAME_Math);
-#endif
-}
-
 #if WITH_EDITOR
 int32 UMaterialExpressionModulo::Compile(class FMaterialCompiler* Compiler, int32 OutputIndex)
 {
@@ -9503,25 +8273,6 @@ FText UMaterialExpressionModulo::GetKeywords() const
 }
 #endif // WITH_EDITOR
 
-UMaterialExpressionFrac::UMaterialExpressionFrac(const FObjectInitializer& ObjectInitializer)
-	: Super(ObjectInitializer)
-{
-	// Structure to hold one-time initialization
-	struct FConstructorStatics
-	{
-		FText NAME_Math;
-		FConstructorStatics()
-			: NAME_Math(LOCTEXT( "Math", "Math" ))
-		{
-		}
-	};
-	static FConstructorStatics ConstructorStatics;
-
-#if WITH_EDITORONLY_DATA
-	MenuCategories.Add(ConstructorStatics.NAME_Math);
-#endif
-}
-
 #if WITH_EDITOR
 int32 UMaterialExpressionFrac::Compile(class FMaterialCompiler* Compiler, int32 OutputIndex)
 {
@@ -9580,11 +8331,9 @@ UMaterialExpressionParameter::UMaterialExpressionParameter(const FObjectInitiali
 	// Structure to hold one-time initialization
 	struct FConstructorStatics
 	{
-		FText NAME_Parameters;
 		FName ParameterName;
 		FConstructorStatics()
-			: NAME_Parameters(LOCTEXT( "Parameters", "Parameters" ))
-			, ParameterName(UMaterialExpressionParameter::ParameterDefaultName)
+			: ParameterName(UMaterialExpressionParameter::ParameterDefaultName)
 		{
 		}
 	};
@@ -9594,8 +8343,6 @@ UMaterialExpressionParameter::UMaterialExpressionParameter(const FObjectInitiali
 	ParameterName = ConstructorStatics.ParameterName;
 
 #if WITH_EDITORONLY_DATA
-	MenuCategories.Add(ConstructorStatics.NAME_Parameters);
-
 	bCollapsed = false;
 #endif
 }
@@ -9863,18 +8610,7 @@ void UMaterialExpressionDoubleVectorParameter::PostEditChangeProperty(FPropertyC
 UMaterialExpressionChannelMaskParameter::UMaterialExpressionChannelMaskParameter(const FObjectInitializer& ObjectInitializer)
 	: Super(ObjectInitializer)
 {
-	struct FConstructorStatics
-	{
-		FText NAME_Utility;
-		FConstructorStatics()
-			: NAME_Utility(LOCTEXT( "Utility", "Utility" ))
-		{
-		}
-	};
-	static FConstructorStatics ConstructorStatics;
 #if WITH_EDITORONLY_DATA
-	MenuCategories.Add(ConstructorStatics.NAME_Utility);
-
 	Outputs.Reset();
 	Outputs.Add(FExpressionOutput(TEXT(""), 1, 1, 0, 0, 0));
 	bShowMaskColorsOnPin = false;
@@ -10327,21 +9063,7 @@ UMaterialExpressionStaticBool::UMaterialExpressionStaticBool(const FObjectInitia
 	: Super(ObjectInitializer)
 {
 #if WITH_EDITORONLY_DATA
-	// Structure to hold one-time initialization
-	struct FConstructorStatics
-	{
-		FText NAME_Functions;
-		FConstructorStatics()
-			: NAME_Functions(LOCTEXT( "Functions", "Functions" ))
-		{
-		}
-	};
-	static FConstructorStatics ConstructorStatics;
-
 	bHidePreviewWindow = true;
-
-	MenuCategories.Add(ConstructorStatics.NAME_Functions);
-
 	bShaderInputData = true;
 #endif
 }
@@ -10366,25 +9088,6 @@ void UMaterialExpressionStaticBool::GetCaption(TArray<FString>& OutCaptions) con
 //
 //	UMaterialExpressionStaticSwitch
 //
-UMaterialExpressionStaticSwitch::UMaterialExpressionStaticSwitch(const FObjectInitializer& ObjectInitializer)
-	: Super(ObjectInitializer)
-{
-	// Structure to hold one-time initialization
-	struct FConstructorStatics
-	{
-		FText NAME_Functions;
-		FConstructorStatics()
-			: NAME_Functions(LOCTEXT( "Functions", "Functions" ))
-		{
-		}
-	};
-	static FConstructorStatics ConstructorStatics;
-
-#if WITH_EDITORONLY_DATA
-	MenuCategories.Add(ConstructorStatics.NAME_Functions);
-#endif
-}
-
 #if WITH_EDITOR
 bool UMaterialExpressionStaticSwitch::IsResultMaterialAttributes(int32 OutputIndex)
 {
@@ -10515,25 +9218,6 @@ FSubstrateOperator* UMaterialExpressionStaticSwitch::SubstrateGenerateMaterialTo
 //
 //	UMaterialExpressionPreviousFrameSwitch
 //
-UMaterialExpressionPreviousFrameSwitch::UMaterialExpressionPreviousFrameSwitch(const FObjectInitializer& ObjectInitializer)
-	: Super(ObjectInitializer)
-{
-	// Structure to hold one-time initialization
-	struct FConstructorStatics
-	{
-		FText NAME_Functions;
-		FConstructorStatics()
-			: NAME_Functions(LOCTEXT("Functions", "Functions"))
-		{
-		}
-	};
-	static FConstructorStatics ConstructorStatics;
-
-#if WITH_EDITORONLY_DATA
-	MenuCategories.Add(ConstructorStatics.NAME_Functions);
-#endif
-}
-
 #if WITH_EDITOR
 bool UMaterialExpressionPreviousFrameSwitch::IsResultMaterialAttributes(int32 OutputIndex)
 {
@@ -10590,26 +9274,6 @@ uint32 UMaterialExpressionPreviousFrameSwitch::GetInputType(int32 InputIndex)
 //
 //	UMaterialExpressionQualitySwitch
 //
-
-UMaterialExpressionQualitySwitch::UMaterialExpressionQualitySwitch(const FObjectInitializer& ObjectInitializer)
-: Super(ObjectInitializer)
-{
-	// Structure to hold one-time initialization
-	struct FConstructorStatics
-	{
-		FText NAME_Utility;
-		FConstructorStatics()
-			: NAME_Utility(LOCTEXT( "Utility", "Utility" ))
-		{
-		}
-	};
-	static FConstructorStatics ConstructorStatics;
-
-#if WITH_EDITORONLY_DATA
-	MenuCategories.Add(ConstructorStatics.NAME_Utility);
-#endif
-}
-
 #if WITH_EDITOR
 FExpressionInput* UMaterialExpressionQualitySwitch::GetEffectiveInput(class FMaterialCompiler* Compiler)
 {
@@ -10719,26 +9383,6 @@ FSubstrateOperator* UMaterialExpressionQualitySwitch::SubstrateGenerateMaterialT
 //
 //	UMaterialExpressionFeatureLevelSwitch
 //
-
-UMaterialExpressionFeatureLevelSwitch::UMaterialExpressionFeatureLevelSwitch(const FObjectInitializer& ObjectInitializer)
-: Super(ObjectInitializer)
-{
-	// Structure to hold one-time initialization
-	struct FConstructorStatics
-	{
-		FText NAME_Utility;
-		FConstructorStatics()
-			: NAME_Utility(LOCTEXT( "Utility", "Utility" ))
-		{
-		}
-	};
-	static FConstructorStatics ConstructorStatics;
-
-#if WITH_EDITORONLY_DATA
-	MenuCategories.Add(ConstructorStatics.NAME_Utility);
-#endif
-}
-
 #if WITH_EDITOR
 int32 UMaterialExpressionFeatureLevelSwitch::Compile(class FMaterialCompiler* Compiler, int32 OutputIndex)
 {
@@ -11368,26 +10012,6 @@ bool UMaterialExpressionRequiredSamplersSwitch::IsResultMaterialAttributes(int32
 //
 //	UMaterialExpressionShadingPathSwitch
 //
-
-UMaterialExpressionShadingPathSwitch::UMaterialExpressionShadingPathSwitch(const FObjectInitializer& ObjectInitializer)
-	: Super(ObjectInitializer)
-{
-	// Structure to hold one-time initialization
-	struct FConstructorStatics
-	{
-		FText NAME_Utility;
-		FConstructorStatics()
-			: NAME_Utility(LOCTEXT("Utility", "Utility"))
-		{
-		}
-	};
-	static FConstructorStatics ConstructorStatics;
-
-#if WITH_EDITORONLY_DATA
-	MenuCategories.Add(ConstructorStatics.NAME_Utility);
-#endif
-}
-
 #if WITH_EDITOR
 FExpressionInput* UMaterialExpressionShadingPathSwitch::GetEffectiveInput(class FMaterialCompiler* Compiler)
 {
@@ -11507,28 +10131,6 @@ FSubstrateOperator* UMaterialExpressionShadingPathSwitch::SubstrateGenerateMater
 //
 //	UMaterialExpressionNormalize
 //
-UMaterialExpressionNormalize::UMaterialExpressionNormalize(const FObjectInitializer& ObjectInitializer)
-	: Super(ObjectInitializer)
-{
-	// Structure to hold one-time initialization
-	struct FConstructorStatics
-	{
-		FText NAME_Math;
-		FText NAME_VectorOps;
-		FConstructorStatics()
-			: NAME_Math(LOCTEXT( "Math", "Math" ))
-			, NAME_VectorOps(LOCTEXT( "VectorOps", "VectorOps" ))
-		{
-		}
-	};
-	static FConstructorStatics ConstructorStatics;
-
-#if WITH_EDITORONLY_DATA
-	MenuCategories.Add(ConstructorStatics.NAME_Math);
-	MenuCategories.Add(ConstructorStatics.NAME_VectorOps);
-#endif
-}
-
 #if WITH_EDITOR
 int32 UMaterialExpressionNormalize::Compile(class FMaterialCompiler* Compiler, int32 OutputIndex)
 {
@@ -11543,25 +10145,6 @@ int32 UMaterialExpressionNormalize::Compile(class FMaterialCompiler* Compiler, i
 //
 //	UMaterialExpressionTruncateLWC
 //
-UMaterialExpressionTruncateLWC::UMaterialExpressionTruncateLWC(const FObjectInitializer& ObjectInitializer)
-	: Super(ObjectInitializer)
-{
-#if WITH_EDITORONLY_DATA
-	// Structure to hold one-time initialization
-	struct FConstructorStatics
-	{
-		FText NAME_Coordinates;
-		FConstructorStatics()
-			: NAME_Coordinates(LOCTEXT("Coordinates", "Coordinates"))
-		{
-		}
-	};
-	static FConstructorStatics ConstructorStatics;
-
-	MenuCategories.Add(ConstructorStatics.NAME_Coordinates);
-#endif
-}
-
 #if WITH_EDITOR
 int32 UMaterialExpressionTruncateLWC::Compile(class FMaterialCompiler* Compiler, int32 OutputIndex)
 {
@@ -11579,19 +10162,6 @@ UMaterialExpressionVertexColor::UMaterialExpressionVertexColor(const FObjectInit
 	: Super(ObjectInitializer)
 {
 #if WITH_EDITORONLY_DATA
-	// Structure to hold one-time initialization
-	struct FConstructorStatics
-	{
-		FText NAME_Constants;
-		FConstructorStatics()
-			: NAME_Constants(LOCTEXT( "Constants", "Constants" ))
-		{
-		}
-	};
-	static FConstructorStatics ConstructorStatics;
-
-	MenuCategories.Add(ConstructorStatics.NAME_Constants);
-
 	Outputs.Reset();
 	Outputs.Add(FExpressionOutput(TEXT(""), 1, 1, 1, 1, 0));
 	Outputs.Add(FExpressionOutput(TEXT(""), 1, 1, 0, 0, 0));
@@ -11618,19 +10188,6 @@ UMaterialExpressionFontSignedDistance::UMaterialExpressionFontSignedDistance(con
 	: Super(ObjectInitializer)
 {
 #if WITH_EDITORONLY_DATA
-	// Structure to hold one-time initialization
-	struct FConstructorStatics
-	{
-		FText NAME_Constants;
-		FConstructorStatics()
-			: NAME_Constants(LOCTEXT("Constants", "Constants"))
-		{
-		}
-	};
-	static FConstructorStatics ConstructorStatics;
-
-	MenuCategories.Add(ConstructorStatics.NAME_Constants);
-
 	Outputs.Reset();
 	Outputs.Add(FExpressionOutput(TEXT("Signed Distance"), 1, 1, 0, 0, 0));
 	Outputs.Add(FExpressionOutput(TEXT("Smooth Signed Distance"), 1, 0, 1, 0, 0));
@@ -11663,22 +10220,6 @@ UMaterialExpressionParticleColor::UMaterialExpressionParticleColor(const FObject
 	: Super(ObjectInitializer)
 {
 #if WITH_EDITORONLY_DATA
-	// Structure to hold one-time initialization
-	struct FConstructorStatics
-	{
-		FText NAME_Particles;
-		FText NAME_Constants;
-		FConstructorStatics()
-			: NAME_Particles(LOCTEXT( "Particles", "Particles" ))
-			, NAME_Constants(LOCTEXT( "Constants", "Constants" ))
-		{
-		}
-	};
-	static FConstructorStatics ConstructorStatics;
-
-	MenuCategories.Add(ConstructorStatics.NAME_Particles);
-	MenuCategories.Add(ConstructorStatics.NAME_Constants);
-
 	Outputs.Reset();
 	Outputs.Add(FExpressionOutput(TEXT("RGB"), 1, 1, 1, 1, 0));
 	Outputs.Add(FExpressionOutput(TEXT("R"), 1, 1, 0, 0, 0));
@@ -11708,25 +10249,6 @@ UMaterialExpressionParticlePositionWS::UMaterialExpressionParticlePositionWS(con
 	: Super(ObjectInitializer)
 {
 #if WITH_EDITORONLY_DATA
-	// Structure to hold one-time initialization
-	struct FConstructorStatics
-	{
-		FText NAME_Particles;
-		FText NAME_Coordinates;
-		FText NAME_Vectors;
-		FConstructorStatics()
-			: NAME_Particles(LOCTEXT( "Particles", "Particles" ))
-			, NAME_Coordinates(LOCTEXT( "Coordinates", "Coordinates" ))
-			, NAME_Vectors(LOCTEXT( "Vectors", "Vectors" ))
-		{
-		}
-	};
-	static FConstructorStatics ConstructorStatics;
-
-	MenuCategories.Add(ConstructorStatics.NAME_Particles);
-	MenuCategories.Add(ConstructorStatics.NAME_Coordinates);
-	MenuCategories.Add(ConstructorStatics.NAME_Vectors);
-
 	bShaderInputData = true;
 #endif
 }
@@ -11766,22 +10288,6 @@ UMaterialExpressionParticleRadius::UMaterialExpressionParticleRadius(const FObje
 	: Super(ObjectInitializer)
 {
 #if WITH_EDITORONLY_DATA
-	// Structure to hold one-time initialization
-	struct FConstructorStatics
-	{
-		FText NAME_Particles;
-		FText NAME_Constants;
-		FConstructorStatics()
-			: NAME_Particles(LOCTEXT( "Particles", "Particles" ))
-			, NAME_Constants(LOCTEXT( "Constants", "Constants" ))
-		{
-		}
-	};
-	static FConstructorStatics ConstructorStatics;
-
-	MenuCategories.Add(ConstructorStatics.NAME_Particles);
-	MenuCategories.Add(ConstructorStatics.NAME_Constants);
-
 	bShaderInputData = true;
 #endif
 }
@@ -11802,19 +10308,6 @@ UMaterialExpressionDynamicParameter::UMaterialExpressionDynamicParameter(const F
 	: Super(ObjectInitializer)
 {
 #if WITH_EDITORONLY_DATA
-	// Structure to hold one-time initialization
-	struct FConstructorStatics
-	{
-		FText NAME_Particles;
-		FText NAME_Parameters;
-		FConstructorStatics()
-			: NAME_Particles(LOCTEXT( "Particles", "Particles" ))
-			, NAME_Parameters(LOCTEXT( "Parameters", "Parameters" ))
-		{
-		}
-	};
-	static FConstructorStatics ConstructorStatics;
-
 	bShowOutputNameOnPin = true;
 	bHidePreviewWindow = true;
 #endif // WITH_EDITORONLY_DATA
@@ -11825,9 +10318,6 @@ UMaterialExpressionDynamicParameter::UMaterialExpressionDynamicParameter(const F
 	ParamNames.Add(TEXT("Param4"));
 
 #if WITH_EDITORONLY_DATA
-	MenuCategories.Add(ConstructorStatics.NAME_Particles);
-	MenuCategories.Add(ConstructorStatics.NAME_Parameters);
-	
 	Outputs.Reset();
 	Outputs.Add(FExpressionOutput(TEXT(""), 1, 1, 0, 0, 0));
 	Outputs.Add(FExpressionOutput(TEXT(""), 1, 0, 1, 0, 0));
@@ -11944,22 +10434,7 @@ bool UMaterialExpressionDynamicParameter::CopyDynamicParameterProperties(const U
 UMaterialExpressionParticleSubUV::UMaterialExpressionParticleSubUV(const FObjectInitializer& ObjectInitializer)
 	: Super(ObjectInitializer)
 {
-	// Structure to hold one-time initialization
-	struct FConstructorStatics
-	{
-		FText NAME_Particles;
-		FConstructorStatics()
-			: NAME_Particles(LOCTEXT( "Particles", "Particles" ))
-		{
-		}
-	};
-	static FConstructorStatics ConstructorStatics;
-
 	bBlend = true;
-
-#if WITH_EDITORONLY_DATA
-	MenuCategories.Add(ConstructorStatics.NAME_Particles);
-#endif
 }
 
 #if WITH_EDITOR
@@ -12011,23 +10486,7 @@ void UMaterialExpressionParticleSubUV::GetCaption(TArray<FString>& OutCaptions)
 UMaterialExpressionParticleSubUVProperties::UMaterialExpressionParticleSubUVProperties(const FObjectInitializer& ObjectInitializer)
 	: Super(ObjectInitializer)
 {
-	// Structure to hold one-time initialization
-	struct FConstructorStatics
-	{
-		FText NAME_Particles;
-		FText NAME_Coordinates;
-		FConstructorStatics()
-			: NAME_Particles(LOCTEXT("Particles", "Particles"))
-			, NAME_Coordinates(LOCTEXT("Coordinates", "Coordinates"))
-		{
-		}
-	};
-	static FConstructorStatics ConstructorStatics;
-
 #if WITH_EDITORONLY_DATA
-	MenuCategories.Add(ConstructorStatics.NAME_Particles);
-	MenuCategories.Add(ConstructorStatics.NAME_Coordinates);
-
 	bShaderInputData = true;
 	bShowOutputNameOnPin = true;
 
@@ -12062,19 +10521,6 @@ UMaterialExpressionParticleMacroUV::UMaterialExpressionParticleMacroUV(const FOb
 	: Super(ObjectInitializer)
 {
 #if WITH_EDITORONLY_DATA
-	// Structure to hold one-time initialization
-	struct FConstructorStatics
-	{
-		FText NAME_Particles;
-		FConstructorStatics()
-			: NAME_Particles(LOCTEXT( "Particles", "Particles" ))
-		{
-		}
-	};
-	static FConstructorStatics ConstructorStatics;
-
-	MenuCategories.Add(ConstructorStatics.NAME_Particles);
-
 	bShaderInputData = true;
 #endif
 }
@@ -12095,19 +10541,6 @@ UMaterialExpressionLightVector::UMaterialExpressionLightVector(const FObjectInit
 	: Super(ObjectInitializer)
 {
 #if WITH_EDITORONLY_DATA
-	// Structure to hold one-time initialization
-	struct FConstructorStatics
-	{
-		FText NAME_Vectors;
-		FConstructorStatics()
-			: NAME_Vectors(LOCTEXT( "Vectors", "Vectors" ))
-		{
-		}
-	};
-	static FConstructorStatics ConstructorStatics;
-
-	MenuCategories.Add(ConstructorStatics.NAME_Vectors);
-
 	bShaderInputData = true;
 #endif
 }
@@ -12128,19 +10561,6 @@ UMaterialExpressionScreenPosition::UMaterialExpressionScreenPosition(const FObje
 	: Super(ObjectInitializer)
 {
 #if WITH_EDITORONLY_DATA
-	// Structure to hold one-time initialization
-	struct FConstructorStatics
-	{
-		FText NAME_Coordinates;
-		FConstructorStatics()
-			: NAME_Coordinates(LOCTEXT( "Coordinates", "Coordinates" ))
-		{
-		}
-	};
-	static FConstructorStatics ConstructorStatics;
-
-	MenuCategories.Add(ConstructorStatics.NAME_Coordinates);
-
 	bShaderInputData = true;
 	bShowOutputNameOnPin = true;
 
@@ -12170,19 +10590,6 @@ UMaterialExpressionViewProperty::UMaterialExpressionViewProperty(const FObjectIn
 	: Super(ObjectInitializer)
 {
 #if WITH_EDITORONLY_DATA
-	// Structure to hold one-time initialization
-	struct FConstructorStatics
-	{
-		FText NAME_Constants;
-		FConstructorStatics()
-			: NAME_Constants(LOCTEXT( "Constants", "Constants" ))
-		{
-		}
-	};
-	static FConstructorStatics ConstructorStatics;
-
-	MenuCategories.Add(ConstructorStatics.NAME_Constants);
-
 	bShaderInputData = true;
 	bShowOutputNameOnPin = true;
 	
@@ -12227,19 +10634,6 @@ UMaterialExpressionViewSize::UMaterialExpressionViewSize(const FObjectInitialize
 	: Super(ObjectInitializer)
 {
 #if WITH_EDITORONLY_DATA
-	// Structure to hold one-time initialization
-	struct FConstructorStatics
-	{
-		FText NAME_Coordinates;
-		FConstructorStatics()
-			: NAME_Coordinates(LOCTEXT( "Coordinates", "Coordinates" ))
-		{
-		}
-	};
-	static FConstructorStatics ConstructorStatics;
-
-	MenuCategories.Add(ConstructorStatics.NAME_Coordinates);
-
 	bShaderInputData = true;
 #endif
 }
@@ -12260,19 +10654,6 @@ UMaterialExpressionIsOrthographic::UMaterialExpressionIsOrthographic(const FObje
 	: Super(ObjectInitializer)
 {
 #if WITH_EDITORONLY_DATA
-	// Structure to hold one-time initialization
-	struct FConstructorStatics
-	{
-		FText NAME_Constants;
-		FConstructorStatics()
-			: NAME_Constants(LOCTEXT("Constants", "Constants"))
-		{
-		}
-	};
-	static FConstructorStatics ConstructorStatics;
-
-	MenuCategories.Add(ConstructorStatics.NAME_Constants);
-
 	bShaderInputData = true;
 #endif
 }
@@ -12293,19 +10674,6 @@ UMaterialExpressionDeltaTime::UMaterialExpressionDeltaTime(const FObjectInitiali
 	: Super(ObjectInitializer)
 {
 #if WITH_EDITORONLY_DATA
-	// Structure to hold one-time initialization
-	struct FConstructorStatics
-	{
-		FText NAME_Constants;
-		FConstructorStatics()
-			: NAME_Constants(LOCTEXT("Constants", "Constants"))
-		{
-		}
-	};
-	static FConstructorStatics ConstructorStatics;
-
-	MenuCategories.Add(ConstructorStatics.NAME_Constants);
-
 	bShaderInputData = true;
 #endif
 }
@@ -12326,19 +10694,6 @@ UMaterialExpressionSceneTexelSize::UMaterialExpressionSceneTexelSize(const FObje
 	: Super(ObjectInitializer)
 {
 #if WITH_EDITORONLY_DATA
-	// Structure to hold one-time initialization
-	struct FConstructorStatics
-	{
-		FText NAME_Coordinates;
-		FConstructorStatics()
-			: NAME_Coordinates(LOCTEXT( "Coordinates", "Coordinates" ))
-		{
-		}
-	};
-	static FConstructorStatics ConstructorStatics;
-
-	MenuCategories.Add(ConstructorStatics.NAME_Coordinates);
-
 	bShaderInputData = true;
 #endif
 }
@@ -12357,25 +10712,6 @@ void UMaterialExpressionSceneTexelSize::GetCaption(TArray<FString>& OutCaptions)
 }
 #endif // WITH_EDITOR
 
-UMaterialExpressionSquareRoot::UMaterialExpressionSquareRoot(const FObjectInitializer& ObjectInitializer)
-	: Super(ObjectInitializer)
-{
-	// Structure to hold one-time initialization
-	struct FConstructorStatics
-	{
-		FText NAME_Math;
-		FConstructorStatics()
-			: NAME_Math(LOCTEXT( "Math", "Math" ))
-		{
-		}
-	};
-	static FConstructorStatics ConstructorStatics;
-
-#if WITH_EDITORONLY_DATA
-	MenuCategories.Add(ConstructorStatics.NAME_Math);
-#endif
-}
-
 #if WITH_EDITOR
 int32 UMaterialExpressionSquareRoot::Compile(class FMaterialCompiler* Compiler, int32 OutputIndex)
 {
@@ -12392,25 +10728,6 @@ void UMaterialExpressionSquareRoot::GetCaption(TArray<FString>& OutCaptions) con
 }
 #endif // WITH_EDITOR
 
-UMaterialExpressionSRGBColorToWorkingColorSpace::UMaterialExpressionSRGBColorToWorkingColorSpace(const FObjectInitializer& ObjectInitializer)
-	: Super(ObjectInitializer)
-{
-#if WITH_EDITORONLY_DATA
-	// Structure to hold one-time initialization
-	struct FConstructorStatics
-	{
-		FText NAME_Color;
-		FConstructorStatics()
-			: NAME_Color(LOCTEXT("Color", "Color"))
-		{
-		}
-	};
-	static FConstructorStatics ConstructorStatics;
-
-	MenuCategories.Add(ConstructorStatics.NAME_Color);
-#endif
-}
-
 #if WITH_EDITOR
 int32 UMaterialExpressionSRGBColorToWorkingColorSpace::Compile(class FMaterialCompiler* Compiler, int32 OutputIndex)
 {
@@ -12455,19 +10772,6 @@ UMaterialExpressionPixelDepth::UMaterialExpressionPixelDepth(const FObjectInitia
 	: Super(ObjectInitializer)
 {
 #if WITH_EDITORONLY_DATA
-	// Structure to hold one-time initialization
-	struct FConstructorStatics
-	{
-		FText NAME_Depth;
-		FConstructorStatics()
-			: NAME_Depth(LOCTEXT( "Depth", "Depth" ))
-		{
-		}
-	};
-	static FConstructorStatics ConstructorStatics;
-
-	MenuCategories.Add(ConstructorStatics.NAME_Depth);
-
 	Outputs.Reset();
 	Outputs.Add(FExpressionOutput(TEXT(""), 1, 1, 0, 0, 0));
 	bShaderInputData = true;
@@ -12496,19 +10800,6 @@ UMaterialExpressionSceneDepth::UMaterialExpressionSceneDepth(const FObjectInitia
 	: Super(ObjectInitializer)
 {
 #if WITH_EDITORONLY_DATA
-	// Structure to hold one-time initialization
-	struct FConstructorStatics
-	{
-		FText NAME_Depth;
-		FConstructorStatics()
-			: NAME_Depth(LOCTEXT( "Depth", "Depth" ))
-		{
-		}
-	};
-	static FConstructorStatics ConstructorStatics;
-
-	MenuCategories.Add(ConstructorStatics.NAME_Depth);
-
 	Outputs.Reset();
 	Outputs.Add(FExpressionOutput(TEXT(""), 1, 1, 0, 0, 0));
 	bShaderInputData = true;
@@ -12587,19 +10878,6 @@ UMaterialExpressionSceneTexture::UMaterialExpressionSceneTexture(const FObjectIn
 	: Super(ObjectInitializer)
 {
 #if WITH_EDITORONLY_DATA
-	// Structure to hold one-time initialization
-	struct FConstructorStatics
-	{
-		FText NAME_Texture;
-		FConstructorStatics()
-			: NAME_Texture(LOCTEXT( "Texture", "Texture" ))
-		{
-		}
-	};
-	static FConstructorStatics ConstructorStatics;
-
-	MenuCategories.Add(ConstructorStatics.NAME_Texture);
-
 	bShaderInputData = true;
 	bShowOutputNameOnPin = true;
 #endif
@@ -12662,19 +10940,6 @@ UMaterialExpressionUserSceneTexture::UMaterialExpressionUserSceneTexture(const F
 	: Super(ObjectInitializer)
 {
 #if WITH_EDITORONLY_DATA
-	// Structure to hold one-time initialization
-	struct FConstructorStatics
-	{
-		FText NAME_Texture;
-		FConstructorStatics()
-			: NAME_Texture(LOCTEXT("Texture", "Texture"))
-		{
-		}
-	};
-	static FConstructorStatics ConstructorStatics;
-
-	MenuCategories.Add(ConstructorStatics.NAME_Texture);
-
 	bShaderInputData = true;
 	bShowOutputNameOnPin = true;
 #endif
@@ -12738,19 +11003,6 @@ UMaterialExpressionSceneColor::UMaterialExpressionSceneColor(const FObjectInitia
 	: Super(ObjectInitializer)
 {
 #if WITH_EDITORONLY_DATA
-	// Structure to hold one-time initialization
-	struct FConstructorStatics
-	{
-		FText NAME_Texture;
-		FConstructorStatics()
-			: NAME_Texture(LOCTEXT( "Texture", "Texture" ))
-		{
-		}
-	};
-	static FConstructorStatics ConstructorStatics;
-
-	MenuCategories.Add(ConstructorStatics.NAME_Texture);
-
 	bShaderInputData = true;
 	bShowOutputNameOnPin = true;
 #endif
@@ -12819,19 +11071,6 @@ UMaterialExpressionDBufferTexture::UMaterialExpressionDBufferTexture(const FObje
 	: Super(ObjectInitializer)
 {
 #if WITH_EDITORONLY_DATA
-	// Structure to hold one-time initialization
-	struct FConstructorStatics
-	{
-		FText NAME_Decals;
-		FConstructorStatics()
-			: NAME_Decals(LOCTEXT("Decals", "Decals"))
-		{
-		}
-	};
-	static FConstructorStatics ConstructorStatics;
-
-	MenuCategories.Add(ConstructorStatics.NAME_Decals);
-
 	bShaderInputData = true;
 	bShowOutputNameOnPin = true;
 #endif
@@ -12870,21 +11109,6 @@ void UMaterialExpressionDBufferTexture::GetCaption(TArray<FString>& OutCaptions)
 UMaterialExpressionPower::UMaterialExpressionPower(const FObjectInitializer& ObjectInitializer)
 	: Super(ObjectInitializer)
 {
-	// Structure to hold one-time initialization
-	struct FConstructorStatics
-	{
-		FText NAME_Math;
-		FConstructorStatics()
-			: NAME_Math(LOCTEXT( "Math", "Math" ))
-		{
-		}
-	};
-	static FConstructorStatics ConstructorStatics;
-
-#if WITH_EDITORONLY_DATA
-	MenuCategories.Add(ConstructorStatics.NAME_Math);
-#endif
-
 	ConstExponent = 2;
 }
 
@@ -12922,24 +11146,6 @@ void UMaterialExpressionPower::GetExpressionToolTip(TArray<FString>& OutToolTip)
 }
 #endif // WITH_EDITOR
 
-UMaterialExpressionLength::UMaterialExpressionLength(const FObjectInitializer& ObjectInitializer)
-	: Super(ObjectInitializer)
-{
-	// Structure to hold one-time initialization
-	struct FConstructorStatics
-	{
-		FText NAME_Utility;
-		FConstructorStatics()
-			: NAME_Utility(LOCTEXT("Utility", "Utility"))
-		{}
-	};
-	static FConstructorStatics ConstructorStatics;
-
-#if WITH_EDITORONLY_DATA
-	MenuCategories.Add(ConstructorStatics.NAME_Utility);
-#endif
-}
-
 #if WITH_EDITOR
 int32 UMaterialExpressionLength::Compile(class FMaterialCompiler* Compiler, int32 OutputIndex)
 {
@@ -12969,24 +11175,6 @@ void UMaterialExpressionLength::GetExpressionToolTip(TArray<FString>& OutToolTip
 }
 #endif // WITH_EDITOR
 
-UMaterialExpressionHsvToRgb::UMaterialExpressionHsvToRgb(const FObjectInitializer& ObjectInitializer)
-	: Super(ObjectInitializer)
-{
-	// Structure to hold one-time initialization
-	struct FConstructorStatics
-	{
-		FText NAME_Math;
-		FConstructorStatics()
-			: NAME_Math(LOCTEXT("Math", "Math"))
-		{}
-	};
-	static FConstructorStatics ConstructorStatics;
-
-#if WITH_EDITORONLY_DATA
-	MenuCategories.Add(ConstructorStatics.NAME_Math);
-#endif
-}
-
 #if WITH_EDITOR
 int32 UMaterialExpressionHsvToRgb::Compile(class FMaterialCompiler* Compiler, int32 OutputIndex)
 {
@@ -13018,24 +11206,6 @@ void UMaterialExpressionHsvToRgb::GetExpressionToolTip(TArray<FString>& OutToolT
 }
 #endif // WITH_EDITOR
 
-UMaterialExpressionRgbToHsv::UMaterialExpressionRgbToHsv(const FObjectInitializer& ObjectInitializer)
-	: Super(ObjectInitializer)
-{
-	// Structure to hold one-time initialization
-	struct FConstructorStatics
-	{
-		FText NAME_Math;
-		FConstructorStatics()
-			: NAME_Math(LOCTEXT("Math", "Math"))
-		{}
-	};
-	static FConstructorStatics ConstructorStatics;
-
-#if WITH_EDITORONLY_DATA
-	MenuCategories.Add(ConstructorStatics.NAME_Math);
-#endif
-}
-
 #if WITH_EDITOR
 int32 UMaterialExpressionRgbToHsv::Compile(class FMaterialCompiler* Compiler, int32 OutputIndex)
 {
@@ -13067,24 +11237,6 @@ void UMaterialExpressionRgbToHsv::GetExpressionToolTip(TArray<FString>& OutToolT
 }
 #endif // WITH_EDITOR
 
-UMaterialExpressionExponential::UMaterialExpressionExponential(const FObjectInitializer& ObjectInitializer)
-	: Super(ObjectInitializer)
-{
-	// Structure to hold one-time initialization
-	struct FConstructorStatics
-	{
-		FText NAME_Math;
-		FConstructorStatics()
-			: NAME_Math(LOCTEXT("Math", "Math"))
-		{}
-	};
-	static FConstructorStatics ConstructorStatics;
-
-#if WITH_EDITORONLY_DATA
-	MenuCategories.Add(ConstructorStatics.NAME_Math);
-#endif
-}
-
 #if WITH_EDITOR
 int32 UMaterialExpressionExponential::Compile(class FMaterialCompiler* Compiler, int32 OutputIndex)
 {
@@ -13107,24 +11259,6 @@ void UMaterialExpressionExponential::GetExpressionToolTip(TArray<FString>& OutTo
 }
 #endif // WITH_EDITOR
 
-UMaterialExpressionExponential2::UMaterialExpressionExponential2(const FObjectInitializer& ObjectInitializer)
-	: Super(ObjectInitializer)
-{
-	// Structure to hold one-time initialization
-	struct FConstructorStatics
-	{
-		FText NAME_Math;
-		FConstructorStatics()
-			: NAME_Math(LOCTEXT("Math", "Math"))
-		{}
-	};
-	static FConstructorStatics ConstructorStatics;
-
-#if WITH_EDITORONLY_DATA
-	MenuCategories.Add(ConstructorStatics.NAME_Math);
-#endif
-}
-
 #if WITH_EDITOR
 int32 UMaterialExpressionExponential2::Compile(class FMaterialCompiler* Compiler, int32 OutputIndex)
 {
@@ -13147,24 +11281,6 @@ void UMaterialExpressionExponential2::GetExpressionToolTip(TArray<FString>& OutT
 }
 #endif // WITH_EDITOR
 
-UMaterialExpressionLogarithm::UMaterialExpressionLogarithm(const FObjectInitializer& ObjectInitializer)
-	: Super(ObjectInitializer)
-{
-	// Structure to hold one-time initialization
-	struct FConstructorStatics
-	{
-		FText NAME_Math;
-		FConstructorStatics()
-			: NAME_Math(LOCTEXT("Math", "Math"))
-		{}
-	};
-	static FConstructorStatics ConstructorStatics;
-
-#if WITH_EDITORONLY_DATA
-	MenuCategories.Add(ConstructorStatics.NAME_Math);
-#endif
-}
-
 #if WITH_EDITOR
 int32 UMaterialExpressionLogarithm::Compile(class FMaterialCompiler* Compiler, int32 OutputIndex)
 {
@@ -13187,25 +11303,6 @@ void UMaterialExpressionLogarithm::GetExpressionToolTip(TArray<FString>& OutTool
 }
 #endif // WITH_EDITOR
 
-UMaterialExpressionLogarithm2::UMaterialExpressionLogarithm2(const FObjectInitializer& ObjectInitializer)
-	: Super(ObjectInitializer)
-{
-	// Structure to hold one-time initialization
-	struct FConstructorStatics
-	{
-		FText NAME_Math;
-		FConstructorStatics()
-			: NAME_Math(LOCTEXT( "Math", "Math" ))
-		{
-		}
-	};
-	static FConstructorStatics ConstructorStatics;
-
-#if WITH_EDITORONLY_DATA
-	MenuCategories.Add(ConstructorStatics.NAME_Math);
-#endif
-}
-
 #if WITH_EDITOR
 int32 UMaterialExpressionLogarithm2::Compile(class FMaterialCompiler* Compiler, int32 OutputIndex)
 {
@@ -13228,25 +11325,6 @@ void UMaterialExpressionLogarithm2::GetExpressionToolTip(TArray<FString>& OutToo
 }
 #endif // WITH_EDITOR
 
-UMaterialExpressionLogarithm10::UMaterialExpressionLogarithm10(const FObjectInitializer& ObjectInitializer)
-	: Super(ObjectInitializer)
-{
-	// Structure to hold one-time initialization
-	struct FConstructorStatics
-	{
-		FText NAME_Math;
-		FConstructorStatics()
-			: NAME_Math(LOCTEXT( "Math", "Math" ))
-		{
-		}
-	};
-	static FConstructorStatics ConstructorStatics;
-
-#if WITH_EDITORONLY_DATA
-	MenuCategories.Add(ConstructorStatics.NAME_Math);
-#endif
-}
-
 #if WITH_EDITOR
 int32 UMaterialExpressionLogarithm10::Compile(class FMaterialCompiler* Compiler, int32 OutputIndex)
 {
@@ -13284,21 +11362,6 @@ FName UMaterialExpressionSceneColor::GetInputName(int32 InputIndex) const
 UMaterialExpressionIf::UMaterialExpressionIf(const FObjectInitializer& ObjectInitializer)
 	: Super(ObjectInitializer)
 {
-	// Structure to hold one-time initialization
-	struct FConstructorStatics
-	{
-		FText NAME_Math;
-		FConstructorStatics()
-			: NAME_Math(LOCTEXT( "Math", "Math" ))
-		{
-		}
-	};
-	static FConstructorStatics ConstructorStatics;
-
-#if WITH_EDITORONLY_DATA
-	MenuCategories.Add(ConstructorStatics.NAME_Math);
-#endif
-
 	EqualsThreshold = 0.00001f;
 	ConstB = 0.0f;
 }
@@ -13382,25 +11445,6 @@ bool UMaterialExpressionIf::IsResultMaterialAttributes(int32 OutputIndex)
 }
 #endif // WITH_EDITOR
 
-UMaterialExpressionOneMinus::UMaterialExpressionOneMinus(const FObjectInitializer& ObjectInitializer)
-	: Super(ObjectInitializer)
-{
-	// Structure to hold one-time initialization
-	struct FConstructorStatics
-	{
-		FText NAME_Math;
-		FConstructorStatics()
-			: NAME_Math(LOCTEXT( "Math", "Math" ))
-		{
-		}
-	};
-	static FConstructorStatics ConstructorStatics;
-
-#if WITH_EDITORONLY_DATA
-	MenuCategories.Add(ConstructorStatics.NAME_Math);
-#endif
-}
-
 #if WITH_EDITOR
 int32 UMaterialExpressionOneMinus::Compile(class FMaterialCompiler* Compiler, int32 OutputIndex)
 {
@@ -13417,25 +11461,6 @@ void UMaterialExpressionOneMinus::GetCaption(TArray<FString>& OutCaptions) const
 }
 #endif // WITH_EDITOR
 
-UMaterialExpressionAbs::UMaterialExpressionAbs(const FObjectInitializer& ObjectInitializer)
-	: Super(ObjectInitializer)
-{
-	// Structure to hold one-time initialization
-	struct FConstructorStatics
-	{
-		FText NAME_Math;
-		FConstructorStatics()
-			: NAME_Math(LOCTEXT( "Math", "Math" ))
-		{
-		}
-	};
-	static FConstructorStatics ConstructorStatics;
-
-#if WITH_EDITORONLY_DATA
-	MenuCategories.Add(ConstructorStatics.NAME_Math);
-#endif
-}
-
 #if WITH_EDITOR
 int32 UMaterialExpressionAbs::Compile( FMaterialCompiler* Compiler, int32 OutputIndex)
 {
@@ -13538,21 +11563,6 @@ void UMaterialExpressionTransform::GetCaption(TArray<FString>& OutCaptions) cons
 UMaterialExpressionTransform::UMaterialExpressionTransform(const FObjectInitializer& ObjectInitializer)
 	: Super(ObjectInitializer)
 {
-	// Structure to hold one-time initialization
-	struct FConstructorStatics
-	{
-		FText NAME_VectorOps;
-		FConstructorStatics()
-			: NAME_VectorOps(LOCTEXT( "VectorOps", "VectorOps" ))
-		{
-		}
-	};
-	static FConstructorStatics ConstructorStatics;
-
-#if WITH_EDITORONLY_DATA
-	MenuCategories.Add(ConstructorStatics.NAME_VectorOps);
-#endif
-
 	TransformSourceType = TRANSFORMSOURCE_Tangent;
 	TransformType = TRANSFORM_World;
 }
@@ -13564,21 +11574,6 @@ UMaterialExpressionTransform::UMaterialExpressionTransform(const FObjectInitiali
 UMaterialExpressionTransformPosition::UMaterialExpressionTransformPosition(const FObjectInitializer& ObjectInitializer)
 	: Super(ObjectInitializer)
 {
-	// Structure to hold one-time initialization
-	struct FConstructorStatics
-	{
-		FText NAME_VectorOps;
-		FConstructorStatics()
-			: NAME_VectorOps(LOCTEXT( "VectorOps", "VectorOps" ))
-		{
-		}
-	};
-	static FConstructorStatics ConstructorStatics;
-
-#if WITH_EDITORONLY_DATA
-	MenuCategories.Add(ConstructorStatics.NAME_VectorOps);
-#endif
-
 	TransformSourceType = TRANSFORMPOSSOURCE_Local;
 	TransformType = TRANSFORMPOSSOURCE_Local;
 }
@@ -13743,20 +11738,6 @@ UMaterialExpressionComment::UMaterialExpressionComment(const FObjectInitializer&
 	, bGroupMode(true)
 
 {
-	// Structure to hold one-time initialization
-	struct FConstructorStatics
-	{
-		FText NAME_Utility;
-		FConstructorStatics()
-			: NAME_Utility(LOCTEXT( "Utility", "Utility" ))
-		{
-		}
-	};
-	static FConstructorStatics ConstructorStatics;
-
-#if WITH_EDITORONLY_DATA
-	MenuCategories.Add(ConstructorStatics.NAME_Utility);
-#endif
 }
 
 #if WITH_EDITOR
@@ -13855,23 +11836,11 @@ bool UMaterialExpressionComment::MatchesSearchQuery( const TCHAR* SearchQuery )
 UMaterialExpressionComposite::UMaterialExpressionComposite(const FObjectInitializer& ObjectInitializer)
 	: Super(ObjectInitializer)
 {
-	// Structure to hold one-time initialization
-	struct FConstructorStatics
-	{
-		FText NAME_Utility;
-		FConstructorStatics()
-			: NAME_Utility(LOCTEXT( "Utility", "Utility" ))
-		{
-		}
-	};
-	static FConstructorStatics ConstructorStatics;
-
 	SubgraphName = "Collapsed Nodes";
 
 #if WITH_EDITORONLY_DATA
 	bShowOutputNameOnPin = true;
 
-	MenuCategories.Add(ConstructorStatics.NAME_Utility);
 	Outputs.Reset();
 #endif
 }
@@ -14054,21 +12023,9 @@ bool UMaterialExpressionComposite::Modify(bool bAlwaysMarkDirty)
 UMaterialExpressionPinBase::UMaterialExpressionPinBase(const FObjectInitializer& ObjectInitializer)
 	: Super(ObjectInitializer)
 {
-	// Structure to hold one-time initialization
-	struct FConstructorStatics
-	{
-		FText NAME_Utility;
-		FConstructorStatics()
-			: NAME_Utility(LOCTEXT("Utility", "Utility"))
-		{
-		}
-	};
-	static FConstructorStatics ConstructorStatics;
-
 #if WITH_EDITORONLY_DATA
 	bShowOutputNameOnPin = true;
 
-	MenuCategories.Add(ConstructorStatics.NAME_Utility);
 	Outputs.Reset();
 #endif
 }
@@ -14308,23 +12265,8 @@ bool UMaterialExpressionPinBase::Modify(bool bAlwaysMarkDirty)
 UMaterialExpressionFresnel::UMaterialExpressionFresnel(const FObjectInitializer& ObjectInitializer)
 	: Super(ObjectInitializer)
 {
-	// Structure to hold one-time initialization
-	struct FConstructorStatics
-	{
-		FText NAME_Utility;
-		FConstructorStatics()
-			: NAME_Utility(LOCTEXT( "Utility", "Utility" ))
-		{
-		}
-	};
-	static FConstructorStatics ConstructorStatics;
-
 	Exponent = 5.0f;
 	BaseReflectFraction = 0.04f;
-
-#if WITH_EDITORONLY_DATA
-	MenuCategories.Add(ConstructorStatics.NAME_Utility);
-#endif
 }
 
 #if WITH_EDITOR
@@ -14355,22 +12297,6 @@ UMaterialExpressionFontSample::UMaterialExpressionFontSample(const FObjectInitia
 	: Super(ObjectInitializer)
 {
 #if WITH_EDITORONLY_DATA
-	// Structure to hold one-time initialization
-	struct FConstructorStatics
-	{
-		FText NAME_Font;
-		FText NAME_Texture;
-		FConstructorStatics()
-			: NAME_Font(LOCTEXT( "Font", "Font" ))
-			, NAME_Texture(LOCTEXT( "Texture", "Texture" ))
-		{
-		}
-	};
-	static FConstructorStatics ConstructorStatics;
-
-	MenuCategories.Add(ConstructorStatics.NAME_Font);
-	MenuCategories.Add(ConstructorStatics.NAME_Texture);
-
 	Outputs.Reset();
 	Outputs.Add(FExpressionOutput(TEXT(""), 1, 1, 1, 1, 0));
 	Outputs.Add(FExpressionOutput(TEXT(""), 1, 1, 0, 0, 0));
@@ -14486,25 +12412,7 @@ UMaterialExpressionFontSampleParameter
 UMaterialExpressionFontSampleParameter::UMaterialExpressionFontSampleParameter(const FObjectInitializer& ObjectInitializer)
 	: Super(ObjectInitializer)
 {
-	// Structure to hold one-time initialization
-	struct FConstructorStatics
-	{
-		FText NAME_Font;
-		FText NAME_Parameters;
-		FConstructorStatics()
-			: NAME_Font(LOCTEXT( "Font", "Font" ))
-			, NAME_Parameters(LOCTEXT( "Parameters", "Parameters" ))
-		{
-		}
-	};
-	static FConstructorStatics ConstructorStatics;
-
 	bIsParameterExpression = true;
-
-#if WITH_EDITORONLY_DATA
-	MenuCategories.Add(ConstructorStatics.NAME_Font);
-	MenuCategories.Add(ConstructorStatics.NAME_Parameters);
-#endif
 }
 
 #if WITH_EDITOR
@@ -14619,19 +12527,6 @@ UMaterialExpressionLocalPosition::UMaterialExpressionLocalPosition(const FObject
 	: Super(ObjectInitializer)
 {
 #if WITH_EDITORONLY_DATA
-	// Structure to hold one-time initialization
-	struct FConstructorStatics
-	{
-		FText NAME_Coordinates;
-		FConstructorStatics()
-		: NAME_Coordinates(LOCTEXT( "Coordinates", "Coordinates" ))
-		{
-		}
-	};
-	static FConstructorStatics ConstructorStatics;
-
-	MenuCategories.Add(ConstructorStatics.NAME_Coordinates);
-
 	Outputs.Reset();
 	Outputs.Add(FExpressionOutput(TEXT("XYZ"), 1, 1, 1, 1, 0));
 	Outputs.Add(FExpressionOutput(TEXT("XY"), 1, 1, 1, 0, 0));
@@ -14692,19 +12587,6 @@ UMaterialExpressionWorldPosition::UMaterialExpressionWorldPosition(const FObject
 	: Super(ObjectInitializer)
 {
 #if WITH_EDITORONLY_DATA
-	// Structure to hold one-time initialization
-	struct FConstructorStatics
-	{
-		FText NAME_Coordinates;
-		FConstructorStatics()
-			: NAME_Coordinates(LOCTEXT( "Coordinates", "Coordinates" ))
-		{
-		}
-	};
-	static FConstructorStatics ConstructorStatics;
-
-	MenuCategories.Add(ConstructorStatics.NAME_Coordinates);
-
 	Outputs.Reset();
 	Outputs.Add(FExpressionOutput(TEXT("XYZ"), 1, 1, 1, 1, 0));
 	Outputs.Add(FExpressionOutput(TEXT("XY"), 1, 1, 1, 0, 0));
@@ -14767,22 +12649,6 @@ UMaterialExpressionObjectPositionWS::UMaterialExpressionObjectPositionWS(const F
 	: Super(ObjectInitializer)
 {
 #if WITH_EDITORONLY_DATA
-	// Structure to hold one-time initialization
-	struct FConstructorStatics
-	{
-		FText NAME_Vectors;
-		FText NAME_Coordinates;
-		FConstructorStatics()
-			: NAME_Vectors(LOCTEXT( "Vectors", "Vectors" ))
-			, NAME_Coordinates(LOCTEXT( "Coordinates", "Coordinates" ))
-		{
-		}
-	};
-	static FConstructorStatics ConstructorStatics;
-
-	MenuCategories.Add(ConstructorStatics.NAME_Vectors);
-	MenuCategories.Add(ConstructorStatics.NAME_Coordinates);
-
 	bShaderInputData = true;
 #endif
 }
@@ -14852,19 +12718,6 @@ UMaterialExpressionObjectRadius::UMaterialExpressionObjectRadius(const FObjectIn
 	: Super(ObjectInitializer)
 {
 #if WITH_EDITORONLY_DATA
-	// Structure to hold one-time initialization
-	struct FConstructorStatics
-	{
-		FText NAME_Coordinates;
-		FConstructorStatics()
-			: NAME_Coordinates(LOCTEXT( "Coordinates", "Coordinates" ))
-		{
-		}
-	};
-	static FConstructorStatics ConstructorStatics;
-
-	MenuCategories.Add(ConstructorStatics.NAME_Coordinates);
-
 	bShaderInputData = true;
 #endif
 }
@@ -14893,19 +12746,6 @@ UMaterialExpressionObjectBounds::UMaterialExpressionObjectBounds(const FObjectIn
 	: Super(ObjectInitializer)
 {
 #if WITH_EDITORONLY_DATA
-	// Structure to hold one-time initialization
-	struct FConstructorStatics
-	{
-		FText NAME_Vectors;
-		FConstructorStatics()
-			: NAME_Vectors(LOCTEXT( "Vectors", "Vectors" ))
-		{
-		}
-	};
-	static FConstructorStatics ConstructorStatics;
-
-	MenuCategories.Add(ConstructorStatics.NAME_Vectors);
-
 	bShaderInputData = true;
 #endif
 }
@@ -14934,19 +12774,6 @@ UMaterialExpressionObjectLocalBounds::UMaterialExpressionObjectLocalBounds(const
 	: Super(ObjectInitializer)
 {
 #if WITH_EDITORONLY_DATA
-	// Structure to hold one-time initialization
-	struct FConstructorStatics
-	{
-		FText NAME_Vectors;
-		FConstructorStatics()
-			: NAME_Vectors(LOCTEXT("Vectors", "Vectors"))
-		{
-		}
-	};
-	static FConstructorStatics ConstructorStatics;
-
-	MenuCategories.Add(ConstructorStatics.NAME_Vectors);
-
 	bShaderInputData = true;
 	bShowOutputNameOnPin = true;
 
@@ -15004,19 +12831,6 @@ UMaterialExpressionBounds::UMaterialExpressionBounds(const FObjectInitializer& O
 	: Super(ObjectInitializer)
 {
 #if WITH_EDITORONLY_DATA
-	// Structure to hold one-time initialization
-	struct FConstructorStatics
-	{
-		FText NAME_Vectors;
-		FConstructorStatics()
-			: NAME_Vectors(LOCTEXT("Vectors", "Vectors"))
-		{
-		}
-	};
-	static FConstructorStatics ConstructorStatics;
-
-	MenuCategories.Add(ConstructorStatics.NAME_Vectors);
-
 	bShaderInputData = true;
 	bShowOutputNameOnPin = true;
 
@@ -15097,19 +12911,6 @@ UMaterialExpressionPreSkinnedLocalBounds::UMaterialExpressionPreSkinnedLocalBoun
 	: Super(ObjectInitializer)
 {
 #if WITH_EDITORONLY_DATA
-	// Structure to hold one-time initialization
-	struct FConstructorStatics
-	{
-		FText NAME_Vectors;
-		FConstructorStatics()
-			: NAME_Vectors(LOCTEXT("Vectors", "Vectors"))
-		{
-		}
-	};
-	static FConstructorStatics ConstructorStatics;
-
-	MenuCategories.Add(ConstructorStatics.NAME_Vectors);
-
 	bShaderInputData = true;
 	bShowOutputNameOnPin = true;
 
@@ -15167,19 +12968,6 @@ UMaterialExpressionDistanceCullFade::UMaterialExpressionDistanceCullFade(const F
 	: Super(ObjectInitializer)
 {
 #if WITH_EDITORONLY_DATA
-	// Structure to hold one-time initialization
-	struct FConstructorStatics
-	{
-		FText NAME_Constants;
-		FConstructorStatics()
-			: NAME_Constants(LOCTEXT( "Constants", "Constants" ))
-		{
-		}
-	};
-	static FConstructorStatics ConstructorStatics;
-
-	MenuCategories.Add(ConstructorStatics.NAME_Constants);
-
 	bShaderInputData = true;
 #endif
 }
@@ -15200,25 +12988,6 @@ void UMaterialExpressionDistanceCullFade::GetCaption(TArray<FString>& OutCaption
 // UMaterialExpressionDistanceFieldsRenderingSwitch
 ///////////////////////////////////////////////////////////////////////////////
 
-UMaterialExpressionDistanceFieldsRenderingSwitch::UMaterialExpressionDistanceFieldsRenderingSwitch(const FObjectInitializer& ObjectInitializer)
-	: Super(ObjectInitializer)
-{
-	// Structure to hold one-time initialization
-	struct FConstructorStatics
-	{
-		FText NAME_DistanceFieldsRendering;
-		FConstructorStatics()
-			: NAME_DistanceFieldsRendering(LOCTEXT("DistanceFieldsRendering", "DistanceFieldsRendering"))
-		{
-		}
-	};
-	static FConstructorStatics ConstructorStatics;
-
-#if WITH_EDITORONLY_DATA
-	MenuCategories.Add(ConstructorStatics.NAME_DistanceFieldsRendering);
-#endif
-}
-
 #if WITH_EDITOR
 
 int32 UMaterialExpressionDistanceFieldsRenderingSwitch::Compile(class FMaterialCompiler* Compiler, int32 OutputIndex)
@@ -15272,22 +13041,6 @@ UMaterialExpressionActorPositionWS::UMaterialExpressionActorPositionWS(const FOb
 	: Super(ObjectInitializer)
 {
 #if WITH_EDITORONLY_DATA
-	// Structure to hold one-time initialization
-	struct FConstructorStatics
-	{
-		FText NAME_Vectors;
-		FText NAME_Coordinates;
-		FConstructorStatics()
-			: NAME_Vectors(LOCTEXT( "Vectors", "Vectors" ))
-			, NAME_Coordinates(LOCTEXT( "Coordinates", "Coordinates" ))
-		{
-		}
-	};
-	static FConstructorStatics ConstructorStatics;
-
-	MenuCategories.Add(ConstructorStatics.NAME_Vectors);
-	MenuCategories.Add(ConstructorStatics.NAME_Coordinates);
-
 	bShaderInputData = true;
 #endif
 }
@@ -15330,26 +13083,7 @@ void UMaterialExpressionActorPositionWS::GetCaption(TArray<FString>& OutCaptions
 
 ///////////////////////////////////////////////////////////////////////////////
 // UMaterialExpressionDeriveNormalZ
-///////////////////////////////////////////////////////////////////////////////
-UMaterialExpressionDeriveNormalZ::UMaterialExpressionDeriveNormalZ(const FObjectInitializer& ObjectInitializer)
-	: Super(ObjectInitializer)
-{
-	// Structure to hold one-time initialization
-	struct FConstructorStatics
-	{
-		FText NAME_VectorOps;
-		FConstructorStatics()
-			: NAME_VectorOps(LOCTEXT( "VectorOps", "VectorOps" ))
-		{
-		}
-	};
-	static FConstructorStatics ConstructorStatics;
-
-#if WITH_EDITORONLY_DATA
-	MenuCategories.Add(ConstructorStatics.NAME_VectorOps);
-#endif
-}
-	
+///////////////////////////////////////////////////////////////////////////////	
 #if WITH_EDITOR
 int32 UMaterialExpressionDeriveNormalZ::Compile(class FMaterialCompiler* Compiler, int32 OutputIndex)
 {
@@ -15381,23 +13115,8 @@ void UMaterialExpressionDeriveNormalZ::GetCaption(TArray<FString>& OutCaptions)
 UMaterialExpressionConstantBiasScale::UMaterialExpressionConstantBiasScale(const FObjectInitializer& ObjectInitializer)
 	: Super(ObjectInitializer)
 {
-	// Structure to hold one-time initialization
-	struct FConstructorStatics
-	{
-		FText NAME_Utility;
-		FConstructorStatics()
-			: NAME_Utility(LOCTEXT( "Utility", "Utility" ))
-		{
-		}
-	};
-	static FConstructorStatics ConstructorStatics;
-
 	Bias = 1.0f;
 	Scale = 0.5f;
-
-#if WITH_EDITORONLY_DATA
-	MenuCategories.Add(ConstructorStatics.NAME_Utility);
-#endif
 }
 
 #if WITH_EDITOR
@@ -15424,28 +13143,11 @@ void UMaterialExpressionConstantBiasScale::GetCaption(TArray<FString>& OutCaptio
 UMaterialExpressionCustom::UMaterialExpressionCustom(const FObjectInitializer& ObjectInitializer)
 	: Super(ObjectInitializer)
 {
-#if WITH_EDITORONLY_DATA
-	// Structure to hold one-time initialization
-	struct FConstructorStatics
-	{
-		FText NAME_Custom;
-		FConstructorStatics()
-			: NAME_Custom(LOCTEXT( "Custom", "Custom" ))
-		{
-		}
-	};
-	static FConstructorStatics ConstructorStatics;
-#endif // WITH_EDITORONLY_DATA
-
 	Description = TEXT("Custom");
 	Code = TEXT("// The below expression will get compiled\n// into the output of this node\nfloat3(1, 1, 1)");
 
 	ShowCode = false;
 
-#if WITH_EDITORONLY_DATA
-	MenuCategories.Add(ConstructorStatics.NAME_Custom);
-#endif
-
 	OutputType = CMOT_Float3;
 
 	Inputs.Add(FCustomInput());
@@ -15878,25 +13580,8 @@ void UMaterialFunctionInterfaceEditorOnlyData::Serialize(FArchive& Ar)
 UMaterialExpressionSwitch::UMaterialExpressionSwitch(const FObjectInitializer& ObjectInitializer)
 	: Super(ObjectInitializer)
 {
-#if WITH_EDITORONLY_DATA
-	// Structure to hold one-time initialization
-	struct FConstructorStatics
-	{
-		FText NAME_Math;
-		FConstructorStatics()
-			: NAME_Math(LOCTEXT("Math", "Math"))
-		{
-		}
-	};
-	static FConstructorStatics ConstructorStatics;
-#endif // WITH_EDITORONLY_DATA
-
 	Description = TEXT("Switch");
 
-#if WITH_EDITORONLY_DATA
-	MenuCategories.Add(ConstructorStatics.NAME_Math);
-#endif
-
 	Inputs.Add(FSwitchCustomInput());
 	Inputs[0].InputName = TEXT("");
 
@@ -19886,22 +17571,9 @@ UMaterialExpressionMaterialFunctionCall::UMaterialExpressionMaterialFunctionCall
 	: Super(ObjectInitializer)
 {
 #if WITH_EDITORONLY_DATA
-	// Structure to hold one-time initialization
-	struct FConstructorStatics
-	{
-		FText NAME_Functions;
-		FConstructorStatics()
-			: NAME_Functions(LOCTEXT( "Functions", "Functions" ))
-		{
-		}
-	};
-	static FConstructorStatics ConstructorStatics;
-
 	bShowOutputNameOnPin = true;
 	bHidePreviewWindow = true;
 
-	MenuCategories.Add(ConstructorStatics.NAME_Functions);
-
 	// Function calls created without a function should be pinless by default
 	FunctionInputs.Empty();
 	FunctionOutputs.Empty();
@@ -20563,27 +18235,12 @@ uint32 UMaterialExpressionMaterialFunctionCall::GetInputType(int32 InputIndex)
 UMaterialExpressionFunctionInput::UMaterialExpressionFunctionInput(const FObjectInitializer& ObjectInitializer)
 	: Super(ObjectInitializer)
 {
-#if WITH_EDITORONLY_DATA
-	// Structure to hold one-time initialization
-	struct FConstructorStatics
-	{
-		FText NAME_Functions;
-		FConstructorStatics()
-			: NAME_Functions(LOCTEXT( "Functions", "Functions" ))
-		{
-		}
-	};
-	static FConstructorStatics ConstructorStatics;
-#endif
-
 	bCompilingFunctionPreview = true;
 	InputType = FunctionInput_Vector3;
 	InputName = TEXT("In");
 
 #if WITH_EDITORONLY_DATA
 	bCollapsed = false;
-
-	MenuCategories.Add(ConstructorStatics.NAME_Functions);
 #endif
 }
 
@@ -20996,24 +18653,12 @@ UMaterialExpressionFunctionOutput::UMaterialExpressionFunctionOutput(const FObje
 	: Super(ObjectInitializer)
 {
 #if WITH_EDITORONLY_DATA
-	// Structure to hold one-time initialization
-	struct FConstructorStatics
-	{
-		FText NAME_Functions;
-		FConstructorStatics()
-			: NAME_Functions(LOCTEXT( "Functions", "Functions" ))
-		{
-		}
-	};
-	static FConstructorStatics ConstructorStatics;
-
 	bShowOutputs = false;
 #endif
 
 	OutputName = TEXT("Result");
 
 #if WITH_EDITORONLY_DATA
-	MenuCategories.Add(ConstructorStatics.NAME_Functions);
 	bCollapsed = false;
 #endif
 }
@@ -21223,19 +18868,6 @@ UMaterialExpressionCollectionParameter::UMaterialExpressionCollectionParameter(c
 	: Super(ObjectInitializer)
 {
 #if WITH_EDITORONLY_DATA
-	// Structure to hold one-time initialization
-	struct FConstructorStatics
-	{
-		FText NAME_Parameters;
-		FConstructorStatics()
-			: NAME_Parameters(LOCTEXT( "Parameters", "Parameters" ))
-		{
-		}
-	};
-	static FConstructorStatics ConstructorStatics;
-
-	MenuCategories.Add(ConstructorStatics.NAME_Parameters);
-
 	bCollapsed = false;
 #endif
 }
@@ -21354,19 +18986,6 @@ UMaterialExpressionCollectionTransform::UMaterialExpressionCollectionTransform(c
 	: Super(ObjectInitializer)
 {
 #if WITH_EDITORONLY_DATA
-	// Structure to hold one-time initialization
-	struct FConstructorStatics
-	{
-		FText NAME_Parameters;
-		FConstructorStatics()
-			: NAME_Parameters(LOCTEXT("Parameters", "Parameters"))
-		{
-		}
-	};
-	static FConstructorStatics ConstructorStatics;
-
-	MenuCategories.Add(ConstructorStatics.NAME_Parameters);
-
 	bCollapsed = false;
 	bHidePreviewWindow = true;
 #endif
@@ -21589,22 +19208,9 @@ UMaterialExpressionLightmapUVs::UMaterialExpressionLightmapUVs(const FObjectInit
 	: Super(ObjectInitializer)
 {
 #if WITH_EDITORONLY_DATA
-	// Structure to hold one-time initialization
-	struct FConstructorStatics
-	{
-		FText NAME_Coordinates;
-		FConstructorStatics()
-			: NAME_Coordinates(LOCTEXT( "Coordinates", "Coordinates" ))
-		{
-		}
-	};
-	static FConstructorStatics ConstructorStatics;
-
 	bShowOutputNameOnPin = true;
 	bHidePreviewWindow = true;
 
-	MenuCategories.Add(ConstructorStatics.NAME_Coordinates);
-
 	Outputs.Reset();
 	Outputs.Add(FExpressionOutput(TEXT(""), 1, 1, 1, 0, 0));
 
@@ -21632,22 +19238,9 @@ UMaterialExpressionPrecomputedAOMask::UMaterialExpressionPrecomputedAOMask(const
 	: Super(ObjectInitializer)
 {
 #if WITH_EDITORONLY_DATA
-	// Structure to hold one-time initialization
-	struct FConstructorStatics
-	{
-		FText NAME_Constants;
-		FConstructorStatics()
-			: NAME_Constants(LOCTEXT( "Constants", "Constants" ))
-		{
-		}
-	};
-	static FConstructorStatics ConstructorStatics;
-
 	bShowOutputNameOnPin = true;
 	bHidePreviewWindow = true;
 
-	MenuCategories.Add(ConstructorStatics.NAME_Constants);
-
 	Outputs.Reset();
 	Outputs.Add(FExpressionOutput(TEXT("")));
 
@@ -21671,25 +19264,6 @@ void UMaterialExpressionPrecomputedAOMask::GetCaption(TArray<FString>& OutCaptio
 //
 //	UMaterialExpressionLightmassReplace
 //
-UMaterialExpressionLightmassReplace::UMaterialExpressionLightmassReplace(const FObjectInitializer& ObjectInitializer)
-	: Super(ObjectInitializer)
-{
-	// Structure to hold one-time initialization
-	struct FConstructorStatics
-	{
-		FText NAME_Utility;
-		FConstructorStatics()
-			: NAME_Utility(LOCTEXT( "Utility", "Utility" ))
-		{
-		}
-	};
-	static FConstructorStatics ConstructorStatics;
-
-#if WITH_EDITORONLY_DATA
-	MenuCategories.Add(ConstructorStatics.NAME_Utility);
-#endif
-}
-
 #if WITH_EDITOR
 int32 UMaterialExpressionLightmassReplace::Compile(class FMaterialCompiler* Compiler, int32 OutputIndex)
 {
@@ -21736,25 +19310,6 @@ void UMaterialExpressionLightmassReplace::GetCaption(TArray<FString>& OutCaption
 //
 //	UMaterialExpressionShadowReplace
 //
-UMaterialExpressionShadowReplace::UMaterialExpressionShadowReplace(const FObjectInitializer& ObjectInitializer)
-	: Super(ObjectInitializer)
-{
-	// Structure to hold one-time initialization
-	struct FConstructorStatics
-	{
-		FText NAME_Utility;
-		FConstructorStatics()
-			: NAME_Utility(LOCTEXT("Utility", "Utility"))
-		{
-		}
-	};
-	static FConstructorStatics ConstructorStatics;
-
-#if WITH_EDITORONLY_DATA
-	MenuCategories.Add(ConstructorStatics.NAME_Utility);
-#endif
-}
-
 #if WITH_EDITOR
 int32 UMaterialExpressionShadowReplace::Compile(class FMaterialCompiler* Compiler, int32 OutputIndex)
 {
@@ -21788,25 +19343,6 @@ void UMaterialExpressionShadowReplace::GetExpressionToolTip(TArray<FString>& Out
 //
 //	UMaterialExpressionShaderStageSwitch
 //
-UMaterialExpressionShaderStageSwitch::UMaterialExpressionShaderStageSwitch(const FObjectInitializer& ObjectInitializer)
-	: Super(ObjectInitializer)
-{
-	// Structure to hold one-time initialization
-	struct FConstructorStatics
-	{
-		FText NAME_Utility;
-		FConstructorStatics()
-			: NAME_Utility(LOCTEXT("Utility", "Utility"))
-		{
-		}
-	};
-	static FConstructorStatics ConstructorStatics;
-
-#if WITH_EDITORONLY_DATA
-	MenuCategories.Add(ConstructorStatics.NAME_Utility);
-#endif
-}
-
 #if WITH_EDITOR
 int32 UMaterialExpressionShaderStageSwitch::Compile(class FMaterialCompiler* Compiler, int32 OutputIndex)
 {
@@ -21840,25 +19376,6 @@ void UMaterialExpressionShaderStageSwitch::GetExpressionToolTip(TArray<FString>&
 //
 //	UMaterialExpressionMaterialProxy
 //
-UMaterialExpressionMaterialProxyReplace::UMaterialExpressionMaterialProxyReplace(const FObjectInitializer& ObjectInitializer)
-	: Super(ObjectInitializer)
-{
-	// Structure to hold one-time initialization
-	struct FConstructorStatics
-	{
-		FText NAME_Utility;
-		FConstructorStatics()
-			: NAME_Utility(LOCTEXT("Utility", "Utility"))
-		{
-		}
-	};
-	static FConstructorStatics ConstructorStatics;
-
-#if WITH_EDITORONLY_DATA
-	MenuCategories.Add(ConstructorStatics.NAME_Utility);
-#endif
-}
-
 #if WITH_EDITOR
 int32 UMaterialExpressionMaterialProxyReplace::Compile(class FMaterialCompiler* Compiler, int32 OutputIndex)
 {
@@ -21897,25 +19414,6 @@ void UMaterialExpressionMaterialProxyReplace::GetCaption(TArray<FString>& OutCap
 //
 //	UMaterialExpressionGIReplace
 //
-UMaterialExpressionGIReplace::UMaterialExpressionGIReplace(const FObjectInitializer& ObjectInitializer)
-	: Super(ObjectInitializer)
-{
-	// Structure to hold one-time initialization
-	struct FConstructorStatics
-	{
-		FText NAME_Utility;
-		FConstructorStatics()
-			: NAME_Utility(LOCTEXT( "Utility", "Utility" ))
-		{
-		}
-	};
-	static FConstructorStatics ConstructorStatics;
-
-#if WITH_EDITORONLY_DATA
-	MenuCategories.Add(ConstructorStatics.NAME_Utility);
-#endif
-}
-
 #if WITH_EDITOR
 int32 UMaterialExpressionGIReplace::Compile(class FMaterialCompiler* Compiler, int32 OutputIndex)
 {
@@ -21946,20 +19444,6 @@ void UMaterialExpressionGIReplace::GetCaption(TArray<FString>& OutCaptions) cons
 UMaterialExpressionRayTracingQualitySwitch::UMaterialExpressionRayTracingQualitySwitch(const FObjectInitializer& ObjectInitializer)
 	: Super(ObjectInitializer)
 {
-	// Structure to hold one-time initialization
-	struct FConstructorStatics
-	{
-		FText NAME_Utility;
-		FConstructorStatics()
-			: NAME_Utility(LOCTEXT("Utility", "Utility"))
-		{
-		}
-	};
-	static FConstructorStatics ConstructorStatics;
-
-#if WITH_EDITORONLY_DATA
-	MenuCategories.Add(ConstructorStatics.NAME_Utility);
-#endif
 }
 
 #if WITH_EDITOR
@@ -22022,20 +19506,6 @@ uint32 UMaterialExpressionRayTracingQualitySwitch::GetInputType(int32 InputIndex
 UMaterialExpressionPathTracingQualitySwitch::UMaterialExpressionPathTracingQualitySwitch(const FObjectInitializer& ObjectInitializer)
 	: Super(ObjectInitializer)
 {
-	// Structure to hold one-time initialization
-	struct FConstructorStatics
-	{
-		FText NAME_Utility;
-		FConstructorStatics()
-			: NAME_Utility(LOCTEXT("Utility", "Utility"))
-		{
-		}
-	};
-	static FConstructorStatics ConstructorStatics;
-
-#if WITH_EDITORONLY_DATA
-	MenuCategories.Add(ConstructorStatics.NAME_Utility);
-#endif
 }
 
 #if WITH_EDITOR
@@ -22099,20 +19569,6 @@ uint32 UMaterialExpressionPathTracingQualitySwitch::GetInputType(int32 InputInde
 UMaterialExpressionPathTracingRayTypeSwitch::UMaterialExpressionPathTracingRayTypeSwitch(const FObjectInitializer& ObjectInitializer)
 	: Super(ObjectInitializer)
 {
-	// Structure to hold one-time initialization
-	struct FConstructorStatics
-	{
-		FText NAME_Utility;
-		FConstructorStatics()
-			: NAME_Utility(LOCTEXT("Utility", "Utility"))
-		{
-		}
-	};
-	static FConstructorStatics ConstructorStatics;
-
-#if WITH_EDITORONLY_DATA
-	MenuCategories.Add(ConstructorStatics.NAME_Utility);
-#endif
 }
 
 #if WITH_EDITOR
@@ -22160,19 +19616,6 @@ UMaterialExpressionPathTracingBufferTexture::UMaterialExpressionPathTracingBuffe
 	: Super(ObjectInitializer)
 {
 #if WITH_EDITORONLY_DATA
-	// Structure to hold one-time initialization
-	struct FConstructorStatics
-	{
-		FText NAME_Decals;
-		FConstructorStatics()
-			: NAME_Decals(LOCTEXT("PathTracing", "PathTracing"))
-		{
-		}
-	};
-	static FConstructorStatics ConstructorStatics;
-
-	MenuCategories.Add(ConstructorStatics.NAME_Decals);
-
 	bShaderInputData = true;
 	bShowOutputNameOnPin = true;
 #endif
@@ -22212,22 +19655,6 @@ UMaterialExpressionObjectOrientation::UMaterialExpressionObjectOrientation(const
 	: Super(ObjectInitializer)
 {
 #if WITH_EDITORONLY_DATA
-	// Structure to hold one-time initialization
-	struct FConstructorStatics
-	{
-		FText NAME_Vectors;
-		FText NAME_Coordinates;
-		FConstructorStatics()
-			: NAME_Vectors(LOCTEXT( "Vectors", "Vectors" ))
-			, NAME_Coordinates(LOCTEXT( "Coordinates", "Coordinates" ))
-		{
-		}
-	};
-	static FConstructorStatics ConstructorStatics;
-
-	MenuCategories.Add(ConstructorStatics.NAME_Vectors);
-	MenuCategories.Add(ConstructorStatics.NAME_Coordinates);
-
 	bShaderInputData = true;
 #endif
 }
@@ -22423,20 +19850,6 @@ FSubstrateOperator* UMaterialExpressionRerouteBase::SubstrateGenerateMaterialTop
 UMaterialExpressionReroute::UMaterialExpressionReroute(const FObjectInitializer& ObjectInitializer)
 	: Super(ObjectInitializer)
 {
-	// Structure to hold one-time initialization
-	struct FConstructorStatics
-	{
-		FText NAME_Utility;
-		FConstructorStatics()
-			: NAME_Utility(LOCTEXT("Utility", "Utility"))
-		{
-		}
-	};
-	static FConstructorStatics ConstructorStatics;
-
-#if WITH_EDITORONLY_DATA
-	MenuCategories.Add(ConstructorStatics.NAME_Utility);
-#endif
 }
 
 #if WITH_EDITOR
@@ -22507,18 +19920,7 @@ UMaterialExpressionNamedRerouteDeclaration* UMaterialExpressionNamedRerouteBase:
 UMaterialExpressionNamedRerouteDeclaration::UMaterialExpressionNamedRerouteDeclaration(const FObjectInitializer& ObjectInitializer)
 	: Super(ObjectInitializer)
 {
-	// Structure to hold one-time initialization
-	struct FConstructorStatics
-	{
-		FText NAME_Utility;
-		FConstructorStatics()
-			: NAME_Utility(LOCTEXT("Utility", "Utility"))
-		{
-		}
-	};
-	static FConstructorStatics ConstructorStatics;
 #if WITH_EDITORONLY_DATA
-	MenuCategories.Add(ConstructorStatics.NAME_Utility);
 	NodeColor = FLinearColor::MakeRandomColor();
 #endif
 	Name = TEXT("Name");
@@ -22827,22 +20229,7 @@ bool UMaterialExpressionNamedRerouteUsage::GetRerouteInput(FExpressionInput& Out
 UMaterialExpressionRotateAboutAxis::UMaterialExpressionRotateAboutAxis(const FObjectInitializer& ObjectInitializer)
 	: Super(ObjectInitializer)
 {
-	// Structure to hold one-time initialization
-	struct FConstructorStatics
-	{
-		FText NAME_Utility;
-		FConstructorStatics()
-			: NAME_Utility(LOCTEXT( "Utility", "Utility" ))
-		{
-		}
-	};
-	static FConstructorStatics ConstructorStatics;
-
 	Period = 1.0f;
-
-#if WITH_EDITORONLY_DATA
-	MenuCategories.Add(ConstructorStatics.NAME_Utility);
-#endif
 }
 
 #if WITH_EDITOR
@@ -22906,23 +20293,8 @@ static int32 CompileHelperLength( FMaterialCompiler* Compiler, int32 A, int32 B
 UMaterialExpressionSphereMask::UMaterialExpressionSphereMask(const FObjectInitializer& ObjectInitializer)
 	: Super(ObjectInitializer)
 {
-	// Structure to hold one-time initialization
-	struct FConstructorStatics
-	{
-		FText NAME_Utility;
-		FConstructorStatics()
-			: NAME_Utility(LOCTEXT( "Utility", "Utility" ))
-		{
-		}
-	};
-	static FConstructorStatics ConstructorStatics;
-
 	AttenuationRadius = 256.0f;
 	HardnessPercent = 100.0f;
-
-#if WITH_EDITORONLY_DATA
-	MenuCategories.Add(ConstructorStatics.NAME_Utility);
-#endif
 }
 
 #if WITH_EDITOR
@@ -22992,17 +20364,6 @@ void UMaterialExpressionSphereMask::GetCaption(TArray<FString>& OutCaptions) con
 UMaterialExpressionSobol::UMaterialExpressionSobol(const FObjectInitializer& ObjectInitializer)
 	: Super(ObjectInitializer)
 {
-	// Structure to hold one-time initialization
-	struct FConstructorStatics
-	{
-		FText NAME_Utility;
-		FConstructorStatics()
-			: NAME_Utility(LOCTEXT("Utility", "Utility"))
-		{
-		}
-	};
-	static FConstructorStatics ConstructorStatics;
-
 	ConstIndex = 0;
 	ConstSeed = FVector2D(0.f, 0.f);
 }
@@ -23035,17 +20396,6 @@ void UMaterialExpressionSobol::GetCaption(TArray<FString>& OutCaptions) const
 UMaterialExpressionTemporalSobol::UMaterialExpressionTemporalSobol(const FObjectInitializer& ObjectInitializer)
 	: Super(ObjectInitializer)
 {
-	// Structure to hold one-time initialization
-	struct FConstructorStatics
-	{
-		FText NAME_Utility;
-		FConstructorStatics()
-			: NAME_Utility(LOCTEXT("Utility", "Utility"))
-		{
-		}
-	};
-	static FConstructorStatics ConstructorStatics;
-
 	ConstIndex = 0;
 	ConstSeed = FVector2D(0.f, 0.f);
 }
@@ -23077,20 +20427,6 @@ void UMaterialExpressionTemporalSobol::GetCaption(TArray<FString>& OutCaptions)
 UMaterialExpressionNaniteReplace::UMaterialExpressionNaniteReplace(const FObjectInitializer& ObjectInitializer)
 : Super(ObjectInitializer)
 {
-	// Structure to hold one-time initialization
-	struct FConstructorStatics
-	{
-		FText NAME_Utility;
-		FConstructorStatics()
-			: NAME_Utility(LOCTEXT("Utility", "Utility"))
-		{
-		}
-	};
-	static FConstructorStatics ConstructorStatics;
-
-#if WITH_EDITORONLY_DATA
-	MenuCategories.Add(ConstructorStatics.NAME_Utility);
-#endif
 }
 
 #if WITH_EDITOR
@@ -23129,20 +20465,6 @@ void UMaterialExpressionNaniteReplace::GetExpressionToolTip(TArray<FString>& Out
 UMaterialExpressionBlackBody::UMaterialExpressionBlackBody(const FObjectInitializer& ObjectInitializer)
 	: Super(ObjectInitializer)
 {
-	// Structure to hold one-time initialization
-	struct FConstructorStatics
-	{
-		FText NAME_Utility;
-		FConstructorStatics()
-			: NAME_Utility(LOCTEXT( "Utility", "Utility" ))
-		{
-		}
-	};
-	static FConstructorStatics ConstructorStatics;
-
-#if WITH_EDITORONLY_DATA
-	MenuCategories.Add(ConstructorStatics.NAME_Utility);
-#endif
 }
 
 #if WITH_EDITOR
@@ -23175,20 +20497,6 @@ void UMaterialExpressionBlackBody::GetCaption(TArray<FString>& OutCaptions) cons
 UMaterialExpressionDistanceToNearestSurface::UMaterialExpressionDistanceToNearestSurface(const FObjectInitializer& ObjectInitializer)
 	: Super(ObjectInitializer)
 {
-	// Structure to hold one-time initialization
-	struct FConstructorStatics
-	{
-		FText NAME_Utility;
-		FConstructorStatics()
-			: NAME_Utility(LOCTEXT( "Utility", "Utility" ))
-		{
-		}
-	};
-	static FConstructorStatics ConstructorStatics;
-
-#if WITH_EDITORONLY_DATA
-	MenuCategories.Add(ConstructorStatics.NAME_Utility);
-#endif
 }
 
 #if WITH_EDITOR
@@ -23244,20 +20552,6 @@ void UMaterialExpressionDistanceToNearestSurface::GetCaption(TArray<FString>& Ou
 UMaterialExpressionDistanceFieldGradient::UMaterialExpressionDistanceFieldGradient(const FObjectInitializer& ObjectInitializer)
 	: Super(ObjectInitializer)
 {
-	// Structure to hold one-time initialization
-	struct FConstructorStatics
-	{
-		FText NAME_Utility;
-		FConstructorStatics()
-			: NAME_Utility(LOCTEXT( "Utility", "Utility" ))
-		{
-		}
-	};
-	static FConstructorStatics ConstructorStatics;
-
-#if WITH_EDITORONLY_DATA
-	MenuCategories.Add(ConstructorStatics.NAME_Utility);
-#endif
 }
 
 #if WITH_EDITOR
@@ -23313,21 +20607,6 @@ void UMaterialExpressionDistanceFieldGradient::GetCaption(TArray<FString>& OutCa
 UMaterialExpressionDistanceFieldApproxAO::UMaterialExpressionDistanceFieldApproxAO(const FObjectInitializer& ObjectInitializer)
 	: Super(ObjectInitializer)
 {
-	// Structure to hold one-time initialization
-	struct FConstructorStatics
-	{
-		FText NAME_Utility;
-		FConstructorStatics()
-			: NAME_Utility(LOCTEXT("Utility", "Utility"))
-		{
-		}
-	};
-	static FConstructorStatics ConstructorStatics;
-
-#if WITH_EDITORONLY_DATA
-	MenuCategories.Add(ConstructorStatics.NAME_Utility);
-#endif
-
 	NumSteps = 1;
 	BaseDistanceDefault = 15;
 	RadiusDefault = 150;
@@ -23402,28 +20681,12 @@ void UMaterialExpressionDistanceFieldApproxAO::GetCaption(TArray<FString>& OutCa
 UMaterialExpressionSamplePhysicsVectorField::UMaterialExpressionSamplePhysicsVectorField(const FObjectInitializer& ObjectInitializer)
 	: Super(ObjectInitializer)
 {
-	// Structure to hold one-time initialization
-	struct FConstructorStatics
-	{
-		FText NAME_Utility;
-		FConstructorStatics()
-			: NAME_Utility(LOCTEXT("Utility", "Utility"))
-		{
-		}
-	};
-	static FConstructorStatics ConstructorStatics;
-
-	FieldTarget = EFieldVectorType::Vector_LinearForce;
-
-#if WITH_EDITORONLY_DATA
-	MenuCategories.Add(ConstructorStatics.NAME_Utility);
-#endif
-}
-
-#if WITH_EDITOR
-FName UMaterialExpressionSamplePhysicsVectorField::GetInputName(int32 InputIndex) const
-{
-	if (GetInput(InputIndex) == &WorldPosition)
+}
+
+#if WITH_EDITOR
+FName UMaterialExpressionSamplePhysicsVectorField::GetInputName(int32 InputIndex) const
+{
+	if (GetInput(InputIndex) == &WorldPosition)
 	{
 		return GetWorldPositionInputName(WorldPositionOriginType);
 	}
@@ -23474,22 +20737,7 @@ void UMaterialExpressionSamplePhysicsVectorField::GetCaption(TArray<FString>& Ou
 UMaterialExpressionSamplePhysicsScalarField::UMaterialExpressionSamplePhysicsScalarField(const FObjectInitializer& ObjectInitializer)
 	: Super(ObjectInitializer)
 {
-	// Structure to hold one-time initialization
-	struct FConstructorStatics
-	{
-		FText NAME_Utility;
-		FConstructorStatics()
-			: NAME_Utility(LOCTEXT("Utility", "Utility"))
-		{
-		}
-	};
-	static FConstructorStatics ConstructorStatics;
-
 	FieldTarget = EFieldScalarType::Scalar_DynamicConstraint;
-
-#if WITH_EDITORONLY_DATA
-	MenuCategories.Add(ConstructorStatics.NAME_Utility);
-#endif
 }
 
 #if WITH_EDITOR
@@ -23547,26 +20795,9 @@ void UMaterialExpressionSamplePhysicsScalarField::GetCaption(TArray<FString>& Ou
 UMaterialExpressionSamplePhysicsIntegerField::UMaterialExpressionSamplePhysicsIntegerField(const FObjectInitializer& ObjectInitializer)
 	: Super(ObjectInitializer)
 {
-	// Structure to hold one-time initialization
-	struct FConstructorStatics
-	{
-		FText NAME_Utility;
-		FConstructorStatics()
-			: NAME_Utility(LOCTEXT("Utility", "Utility"))
-		{
-		}
-	};
-	static FConstructorStatics ConstructorStatics;
-
 	FieldTarget = EFieldIntegerType::Integer_DynamicState;
-
-#if WITH_EDITORONLY_DATA
-	MenuCategories.Add(ConstructorStatics.NAME_Utility);
-#endif
 }
 
-
-
 #if WITH_EDITOR
 FName UMaterialExpressionSamplePhysicsIntegerField::GetInputName(int32 InputIndex) const
 {
@@ -23619,25 +20850,6 @@ void UMaterialExpressionSamplePhysicsIntegerField::GetCaption(TArray<FString>& O
 ///////////////////////////////////////////////////////////////////////////////
 // UMaterialExpressionDistance
 ///////////////////////////////////////////////////////////////////////////////
-UMaterialExpressionDistance::UMaterialExpressionDistance(const FObjectInitializer& ObjectInitializer)
-	: Super(ObjectInitializer)
-{
-	// Structure to hold one-time initialization
-	struct FConstructorStatics
-	{
-		FText NAME_Utility;
-		FConstructorStatics()
-			: NAME_Utility(LOCTEXT( "Utility", "Utility" ))
-		{
-		}
-	};
-	static FConstructorStatics ConstructorStatics;
-
-#if WITH_EDITORONLY_DATA
-	MenuCategories.Add(ConstructorStatics.NAME_Utility);
-#endif
-}
-
 #if WITH_EDITOR
 int32 UMaterialExpressionDistance::Compile(class FMaterialCompiler* Compiler, int32 OutputIndex)
 {
@@ -23667,19 +20879,6 @@ UMaterialExpressionTwoSidedSign::UMaterialExpressionTwoSidedSign(const FObjectIn
 	: Super(ObjectInitializer)
 {
 #if WITH_EDITORONLY_DATA
-	// Structure to hold one-time initialization
-	struct FConstructorStatics
-	{
-		FText NAME_Constants;
-		FConstructorStatics()
-			: NAME_Constants(LOCTEXT( "Constants", "Constants" ))
-		{
-		}
-	};
-	static FConstructorStatics ConstructorStatics;
-
-	MenuCategories.Add(ConstructorStatics.NAME_Constants);
-
 	bShaderInputData = true;
 #endif
 }
@@ -23700,22 +20899,6 @@ UMaterialExpressionVertexNormalWS::UMaterialExpressionVertexNormalWS(const FObje
 	: Super(ObjectInitializer)
 {
 #if WITH_EDITORONLY_DATA
-	// Structure to hold one-time initialization
-	struct FConstructorStatics
-	{
-		FText NAME_Vectors;
-		FText NAME_Coordinates;
-		FConstructorStatics()
-			: NAME_Vectors(LOCTEXT( "Vectors", "Vectors" ))
-			, NAME_Coordinates(LOCTEXT( "Coordinates", "Coordinates" ))
-		{
-		}
-	};
-	static FConstructorStatics ConstructorStatics;
-
-	MenuCategories.Add(ConstructorStatics.NAME_Vectors);
-	MenuCategories.Add(ConstructorStatics.NAME_Coordinates);
-
 	bShaderInputData = true;
 #endif
 }
@@ -23736,22 +20919,6 @@ UMaterialExpressionVertexTangentWS::UMaterialExpressionVertexTangentWS(const FOb
 	: Super(ObjectInitializer)
 {
 #if WITH_EDITORONLY_DATA
-	// Structure to hold one-time initialization
-	struct FConstructorStatics
-	{
-		FText NAME_Vectors;
-		FText NAME_Coordinates;
-		FConstructorStatics()
-			: NAME_Vectors(LOCTEXT("Vectors", "Vectors"))
-			, NAME_Coordinates(LOCTEXT("Coordinates", "Coordinates"))
-		{
-		}
-	};
-	static FConstructorStatics ConstructorStatics;
-
-	MenuCategories.Add(ConstructorStatics.NAME_Vectors);
-	MenuCategories.Add(ConstructorStatics.NAME_Coordinates);
-
 	bShaderInputData = true;
 #endif
 }
@@ -23772,22 +20939,6 @@ UMaterialExpressionPixelNormalWS::UMaterialExpressionPixelNormalWS(const FObject
 	: Super(ObjectInitializer)
 {
 #if WITH_EDITORONLY_DATA
-	// Structure to hold one-time initialization
-	struct FConstructorStatics
-	{
-		FText NAME_Vectors;
-		FText NAME_Coordinates;
-		FConstructorStatics()
-			: NAME_Vectors(LOCTEXT( "Vectors", "Vectors" ))
-			, NAME_Coordinates(LOCTEXT( "Coordinates", "Coordinates" ))
-		{
-		}
-	};
-	static FConstructorStatics ConstructorStatics;
-
-	MenuCategories.Add(ConstructorStatics.NAME_Vectors);
-	MenuCategories.Add(ConstructorStatics.NAME_Coordinates);
-
 	bShaderInputData = true;
 #endif
 }
@@ -23811,19 +20962,6 @@ UMaterialExpressionPerInstanceRandom::UMaterialExpressionPerInstanceRandom(const
 	: Super(ObjectInitializer)
 {
 #if WITH_EDITORONLY_DATA
-	// Structure to hold one-time initialization
-	struct FConstructorStatics
-	{
-		FText NAME_Constants;
-		FConstructorStatics()
-			: NAME_Constants(LOCTEXT( "Constants", "Constants" ))
-		{
-		}
-	};
-	static FConstructorStatics ConstructorStatics;
-
-	MenuCategories.Add(ConstructorStatics.NAME_Constants);
-
 	bShaderInputData = true;
 #endif
 }
@@ -23847,19 +20985,6 @@ UMaterialExpressionPerInstanceFadeAmount::UMaterialExpressionPerInstanceFadeAmou
 	: Super(ObjectInitializer)
 {
 #if WITH_EDITORONLY_DATA
-	// Structure to hold one-time initialization
-	struct FConstructorStatics
-	{
-		FText NAME_Constants;
-		FConstructorStatics()
-			: NAME_Constants(LOCTEXT( "Constants", "Constants" ))
-		{
-		}
-	};
-	static FConstructorStatics ConstructorStatics;
-
-	MenuCategories.Add(ConstructorStatics.NAME_Constants);
-
 	bShaderInputData = true;
 #endif
 }
@@ -23883,19 +21008,6 @@ UMaterialExpressionPerInstanceCustomData::UMaterialExpressionPerInstanceCustomDa
 	: Super(ObjectInitializer)
 {
 #if WITH_EDITORONLY_DATA
-	// Structure to hold one-time initialization
-	struct FConstructorStatics
-	{
-		FText NAME_Custom;
-		FConstructorStatics()
-			: NAME_Custom(LOCTEXT("Custom", "Custom"))
-		{
-		}
-	};
-	static FConstructorStatics ConstructorStatics;
-
-	MenuCategories.Add(ConstructorStatics.NAME_Custom);
-
 	bShaderInputData = true;
 #endif
 }
@@ -23920,19 +21032,6 @@ UMaterialExpressionPerInstanceCustomData3Vector::UMaterialExpressionPerInstanceC
 	: Super(ObjectInitializer)
 {
 #if WITH_EDITORONLY_DATA
-	// Structure to hold one-time initialization
-	struct FConstructorStatics
-	{
-		FText NAME_Custom;
-		FConstructorStatics()
-			: NAME_Custom(LOCTEXT("Custom", "Custom"))
-		{
-		}
-	};
-	static FConstructorStatics ConstructorStatics;
-
-	MenuCategories.Add(ConstructorStatics.NAME_Custom);
-
 	bShaderInputData = true;
 #endif
 }
@@ -23964,21 +21063,14 @@ UMaterialExpressionAntialiasedTextureMask::UMaterialExpressionAntialiasedTexture
 	struct FConstructorStatics
 	{
 		ConstructorHelpers::FObjectFinder<UTexture2D> DefaultTexture;
-		FText NAME_Utility;
-		FName NAME_None;
 		FConstructorStatics()
 			: DefaultTexture(TEXT("/Engine/EngineResources/DefaultTexture"))
-			, NAME_Utility(LOCTEXT( "Utility", "Utility" ))
-			, NAME_None(TEXT("None"))
 		{
 		}
 	};
 	static FConstructorStatics ConstructorStatics;
 
 	Texture = ConstructorStatics.DefaultTexture.Object;
-
-	MenuCategories.Empty();
-	MenuCategories.Add(ConstructorStatics.NAME_Utility);
 #endif
 
 	Threshold = 0.5f;
@@ -24068,19 +21160,6 @@ UMaterialExpressionDecalDerivative::UMaterialExpressionDecalDerivative(const FOb
 	: Super(ObjectInitializer)
 {
 #if WITH_EDITORONLY_DATA
-	// Structure to hold one-time initialization
-	struct FConstructorStatics
-	{
-		FText NAME_Vectors;
-		FConstructorStatics()
-			: NAME_Vectors(LOCTEXT("Utils", "Utils"))
-		{
-		}
-	};
-	static FConstructorStatics ConstructorStatics;
-
-	MenuCategories.Add(ConstructorStatics.NAME_Vectors);
-
 	//bCollapsed = true;
 	bShaderInputData = true;
 	bShowOutputNameOnPin = true;
@@ -24110,19 +21189,6 @@ UMaterialExpressionDecalColor::UMaterialExpressionDecalColor(const FObjectInitia
 	: Super(ObjectInitializer)
 {
 #if WITH_EDITORONLY_DATA
-	// Structure to hold one-time initialization
-	struct FConstructorStatics
-	{
-		FText NAME_Vectors;
-		FConstructorStatics()
-			: NAME_Vectors(LOCTEXT("Utils", "Utils"))
-		{
-		}
-	};
-	static FConstructorStatics ConstructorStatics;
-
-	MenuCategories.Add(ConstructorStatics.NAME_Vectors);
-
 	bShowOutputNameOnPin = true;
 	bShaderInputData = true;
 
@@ -24160,19 +21226,6 @@ UMaterialExpressionDecalLifetimeOpacity::UMaterialExpressionDecalLifetimeOpacity
 	: Super(ObjectInitializer)
 {
 #if WITH_EDITORONLY_DATA
-	// Structure to hold one-time initialization
-	struct FConstructorStatics
-	{
-		FText NAME_Vectors;
-		FConstructorStatics()
-			: NAME_Vectors(LOCTEXT("Utils", "Utils"))
-		{
-		}
-	};
-	static FConstructorStatics ConstructorStatics;
-
-	MenuCategories.Add(ConstructorStatics.NAME_Vectors);
-
 	bShaderInputData = true;
 
 	Outputs.Reset();
@@ -24206,19 +21259,6 @@ UMaterialExpressionDecalMipmapLevel::UMaterialExpressionDecalMipmapLevel(const F
 	, ConstHeight(ConstWidth)
 {
 #if WITH_EDITORONLY_DATA
-	// Structure to hold one-time initialization
-	struct FConstructorStatics
-	{
-		FText NAME_Vectors;
-		FConstructorStatics()
-			: NAME_Vectors(LOCTEXT("Utils", "Utils"))
-		{
-		}
-	};
-	static FConstructorStatics ConstructorStatics;
-
-	MenuCategories.Add(ConstructorStatics.NAME_Vectors);
-
 	bCollapsed = true;
 #endif
 }
@@ -24259,28 +21299,10 @@ void UMaterialExpressionDecalMipmapLevel::GetCaption(TArray<FString>& OutCaption
 UMaterialExpressionDepthFade::UMaterialExpressionDepthFade(const FObjectInitializer& ObjectInitializer)
 	: Super(ObjectInitializer)
 {
-#if WITH_EDITORONLY_DATA
-	// Structure to hold one-time initialization
-	struct FConstructorStatics
-	{
-		FText NAME_Depth;
-		FText NAME_Utility;
-		FConstructorStatics()
-			: NAME_Depth(LOCTEXT( "Depth", "Depth" ))
-			, NAME_Utility(LOCTEXT( "Utility", "Utility" ))
-		{
-		}
-	};
-	static FConstructorStatics ConstructorStatics;
-#endif
-
 	FadeDistanceDefault = 100.0f;
 	OpacityDefault = 1.0f;
 
 #if WITH_EDITORONLY_DATA
-	MenuCategories.Add(ConstructorStatics.NAME_Depth);
-	MenuCategories.Add(ConstructorStatics.NAME_Utility);
-
 	bCollapsed = false;
 #endif
 }
@@ -24317,24 +21339,9 @@ int32 UMaterialExpressionDepthFade::Compile(class FMaterialCompiler* Compiler, i
 UMaterialExpressionSphericalParticleOpacity::UMaterialExpressionSphericalParticleOpacity(const FObjectInitializer& ObjectInitializer)
 	: Super(ObjectInitializer)
 {
-#if WITH_EDITORONLY_DATA
-	// Structure to hold one-time initialization
-	struct FConstructorStatics
-	{
-		FText NAME_Particles;
-		FConstructorStatics()
-			: NAME_Particles(LOCTEXT( "Particles", "Particles" ))
-		{
-		}
-	};
-	static FConstructorStatics ConstructorStatics;
-#endif
-
 	ConstantDensity = 1;
 
 #if WITH_EDITORONLY_DATA
-	MenuCategories.Add(ConstructorStatics.NAME_Particles);
-
 	bCollapsed = false;
 #endif
 }
@@ -24354,19 +21361,6 @@ UMaterialExpressionDepthOfFieldFunction::UMaterialExpressionDepthOfFieldFunction
 	: Super(ObjectInitializer)
 {
 #if WITH_EDITORONLY_DATA
-	// Structure to hold one-time initialization
-	struct FConstructorStatics
-	{
-		FText NAME_Utility;
-		FConstructorStatics()
-			: NAME_Utility(LOCTEXT( "Utility", "Utility" ))
-		{
-		}
-	};
-	static FConstructorStatics ConstructorStatics;
-
-	MenuCategories.Add(ConstructorStatics.NAME_Utility);
-
 	bCollapsed = false;
 #endif
 }
@@ -24409,19 +21403,6 @@ UMaterialExpressionPostVolumeUserFlagTest::UMaterialExpressionPostVolumeUserFlag
 	: Super(ObjectInitializer)
 {
 #if WITH_EDITORONLY_DATA
-	// Structure to hold one-time initialization
-	struct FConstructorStatics
-	{
-		FText NAME_GetPostProcessSetting;
-		FConstructorStatics()
-			: NAME_GetPostProcessSetting(LOCTEXT("Get Post Process Setting", "Get Post Process Setting"))
-		{
-		}
-	};
-	static FConstructorStatics ConstructorStatics;
-
-	MenuCategories.Add(ConstructorStatics.NAME_GetPostProcessSetting);
-
 	bCollapsed = true;
 #endif
 }
@@ -24456,19 +21437,6 @@ UMaterialExpressionDDX::UMaterialExpressionDDX(const FObjectInitializer& ObjectI
 	: Super(ObjectInitializer)
 {
 #if WITH_EDITORONLY_DATA
-	// Structure to hold one-time initialization
-	struct FConstructorStatics
-	{
-		FText NAME_Utility;
-		FConstructorStatics()
-			: NAME_Utility(LOCTEXT( "Utility", "Utility" ))
-		{
-		}
-	};
-	static FConstructorStatics ConstructorStatics;
-
-	MenuCategories.Add(ConstructorStatics.NAME_Utility);
-
 	bCollapsed = true;
 #endif
 }
@@ -24505,19 +21473,6 @@ UMaterialExpressionDDY::UMaterialExpressionDDY(const FObjectInitializer& ObjectI
 	: Super(ObjectInitializer)
 {
 #if WITH_EDITORONLY_DATA
-	// Structure to hold one-time initialization
-	struct FConstructorStatics
-	{
-		FText NAME_Utility;
-		FConstructorStatics()
-			: NAME_Utility(LOCTEXT( "Utility", "Utility" ))
-		{
-		}
-	};
-	static FConstructorStatics ConstructorStatics;
-
-	MenuCategories.Add(ConstructorStatics.NAME_Utility);
-
 	bCollapsed = true;
 #endif
 }
@@ -24554,22 +21509,6 @@ UMaterialExpressionParticleRelativeTime::UMaterialExpressionParticleRelativeTime
 	: Super(ObjectInitializer)
 {
 #if WITH_EDITORONLY_DATA
-	// Structure to hold one-time initialization
-	struct FConstructorStatics
-	{
-		FText NAME_Particles;
-		FText NAME_Constants;
-		FConstructorStatics()
-			: NAME_Particles(LOCTEXT( "Particles", "Particles" ))
-			, NAME_Constants(LOCTEXT( "Constants", "Constants" ))
-		{
-		}
-	};
-	static FConstructorStatics ConstructorStatics;
-
-	MenuCategories.Add(ConstructorStatics.NAME_Particles);
-	MenuCategories.Add(ConstructorStatics.NAME_Constants);
-
 	bShaderInputData = true;
 #endif
 }
@@ -24593,22 +21532,6 @@ UMaterialExpressionParticleMotionBlurFade::UMaterialExpressionParticleMotionBlur
 	: Super(ObjectInitializer)
 {
 #if WITH_EDITORONLY_DATA
-	// Structure to hold one-time initialization
-	struct FConstructorStatics
-	{
-		FText NAME_Particles;
-		FText NAME_Constants;
-		FConstructorStatics()
-			: NAME_Particles(LOCTEXT( "Particles", "Particles" ))
-			, NAME_Constants(LOCTEXT( "Constants", "Constants" ))
-		{
-		}
-	};
-	static FConstructorStatics ConstructorStatics;
-
-	MenuCategories.Add(ConstructorStatics.NAME_Particles);
-	MenuCategories.Add(ConstructorStatics.NAME_Constants);
-
 	bShaderInputData = true;
 #endif
 }
@@ -24632,22 +21555,6 @@ UMaterialExpressionParticleRandom::UMaterialExpressionParticleRandom(const FObje
 	: Super(ObjectInitializer)
 {
 #if WITH_EDITORONLY_DATA
-	// Structure to hold one-time initialization
-	struct FConstructorStatics
-	{
-		FText NAME_Particles;
-		FText NAME_Constants;
-		FConstructorStatics()
-			: NAME_Particles(LOCTEXT( "Particles", "Particles" ))
-			, NAME_Constants(LOCTEXT( "Constants", "Constants" ))
-		{
-		}
-	};
-	static FConstructorStatics ConstructorStatics;
-
-	MenuCategories.Add(ConstructorStatics.NAME_Particles);
-	MenuCategories.Add(ConstructorStatics.NAME_Constants);
-
 	bShaderInputData = true;
 #endif
 }
@@ -24671,22 +21578,6 @@ UMaterialExpressionParticleDirection::UMaterialExpressionParticleDirection(const
 	: Super(ObjectInitializer)
 {
 #if WITH_EDITORONLY_DATA
-	// Structure to hold one-time initialization
-	struct FConstructorStatics
-	{
-		FText NAME_Particles;
-		FText NAME_Constants;
-		FConstructorStatics()
-			: NAME_Particles(LOCTEXT( "Particles", "Particles" ))
-			, NAME_Constants(LOCTEXT( "Constants", "Constants" ))
-		{
-		}
-	};
-	static FConstructorStatics ConstructorStatics;
-
-	MenuCategories.Add(ConstructorStatics.NAME_Particles);
-	MenuCategories.Add(ConstructorStatics.NAME_Constants);
-
 	bShaderInputData = true;
 #endif
 }
@@ -24710,22 +21601,6 @@ UMaterialExpressionParticleSpeed::UMaterialExpressionParticleSpeed(const FObject
 	: Super(ObjectInitializer)
 {
 #if WITH_EDITORONLY_DATA
-	// Structure to hold one-time initialization
-	struct FConstructorStatics
-	{
-		FText NAME_Particles;
-		FText NAME_Constants;
-		FConstructorStatics()
-			: NAME_Particles(LOCTEXT( "Particles", "Particles" ))
-			, NAME_Constants(LOCTEXT( "Constants", "Constants" ))
-		{
-		}
-	};
-	static FConstructorStatics ConstructorStatics;
-
-	MenuCategories.Add(ConstructorStatics.NAME_Particles);
-	MenuCategories.Add(ConstructorStatics.NAME_Constants);
-
 	bShaderInputData = true;
 #endif
 }
@@ -24749,22 +21624,6 @@ UMaterialExpressionParticleSize::UMaterialExpressionParticleSize(const FObjectIn
 	: Super(ObjectInitializer)
 {
 #if WITH_EDITORONLY_DATA
-	// Structure to hold one-time initialization
-	struct FConstructorStatics
-	{
-		FText NAME_Particles;
-		FText NAME_Constants;
-		FConstructorStatics()
-			: NAME_Particles(LOCTEXT( "Particles", "Particles" ))
-			, NAME_Constants(LOCTEXT( "Constants", "Constants" ))
-		{
-		}
-	};
-	static FConstructorStatics ConstructorStatics;
-
-	MenuCategories.Add(ConstructorStatics.NAME_Particles);
-	MenuCategories.Add(ConstructorStatics.NAME_Constants);
-
 	bShaderInputData = true;
 #endif
 }
@@ -24788,22 +21647,6 @@ UMaterialExpressionParticleSpriteRotation::UMaterialExpressionParticleSpriteRota
 	: Super(ObjectInitializer)
 {
 #if WITH_EDITORONLY_DATA
-	// Structure to hold one-time initialization
-	struct FConstructorStatics
-	{
-		FText NAME_Particles;
-		FText NAME_Constants;
-		FConstructorStatics()
-			: NAME_Particles(LOCTEXT("Particles", "Particles"))
-			, NAME_Constants(LOCTEXT("Constants", "Constants"))
-		{
-		}
-	};
-	static FConstructorStatics ConstructorStatics;
-
-	MenuCategories.Add(ConstructorStatics.NAME_Particles);
-	MenuCategories.Add(ConstructorStatics.NAME_Constants);
-
 	Outputs.Reset();
 	Outputs.Emplace(TEXT("Rad"), 1, 1, 0, 0, 0);
 	Outputs.Emplace(TEXT("Deg"), 1, 0, 1, 0, 0);
@@ -24828,23 +21671,10 @@ void UMaterialExpressionParticleSpriteRotation::GetCaption(TArray<FString>& OutC
 /*------------------------------------------------------------------------------
 	Atmospheric fog material expression.
 ------------------------------------------------------------------------------*/
-UMaterialExpressionAtmosphericFogColor::UMaterialExpressionAtmosphericFogColor(const FObjectInitializer& ObjectInitializer)
-	: Super(ObjectInitializer)
-{
-#if WITH_EDITORONLY_DATA
-	// Structure to hold one-time initialization
-	struct FConstructorStatics
-	{
-		FText NAME_Atmosphere;
-		FConstructorStatics()
-			: NAME_Atmosphere(LOCTEXT( "Atmosphere", "Atmosphere" ))
-		{
-		}
-	};
-	static FConstructorStatics ConstructorStatics;
-
-	MenuCategories.Add(ConstructorStatics.NAME_Atmosphere);
-
+UMaterialExpressionAtmosphericFogColor::UMaterialExpressionAtmosphericFogColor(const FObjectInitializer& ObjectInitializer)
+	: Super(ObjectInitializer)
+{
+#if WITH_EDITORONLY_DATA
 	bCollapsed = false;
 #endif
 }
@@ -24898,26 +21728,11 @@ void UMaterialExpressionAtmosphericFogColor::GetCaption(TArray<FString>& OutCapt
 UMaterialExpressionSpeedTree::UMaterialExpressionSpeedTree(const FObjectInitializer& ObjectInitializer)
 	: Super(ObjectInitializer)
 {
-	// Structure to hold one-time initialization
-	struct FConstructorStatics
-	{
-		FText NAME_SpeedTree;
-		FConstructorStatics()
-			: NAME_SpeedTree(LOCTEXT( "SpeedTree", "SpeedTree" ))
-		{
-		}
-	};
-	static FConstructorStatics ConstructorStatics;
-
 	GeometryType = STG_Branch;
 	WindType = STW_None;
 	LODType = STLOD_Pop;
 	BillboardThreshold = 0.9f;
 	bAccurateWindVelocities = false;
-
-#if WITH_EDITORONLY_DATA
-	MenuCategories.Add(ConstructorStatics.NAME_SpeedTree);
-#endif
 }
 
 #if WITH_EDITOR
@@ -25004,19 +21819,6 @@ UMaterialExpressionEyeAdaptation::UMaterialExpressionEyeAdaptation(const FObject
 	: Super(ObjectInitializer)
 {
 #if WITH_EDITORONLY_DATA
-	// Structure to hold one-time initialization
-	struct FConstructorStatics
-	{
-		FText NAME_Utility;
-		FConstructorStatics()
-			: NAME_Utility(LOCTEXT( "Utility", "Utility" ))
-		{
-		}
-	};
-	static FConstructorStatics ConstructorStatics;
-
-	MenuCategories.Add(ConstructorStatics.NAME_Utility);
-
 	Outputs.Reset();
 	Outputs.Add(FExpressionOutput(TEXT("EyeAdaptation")));
 	bShaderInputData = true;
@@ -25042,19 +21844,6 @@ UMaterialExpressionEyeAdaptationInverse::UMaterialExpressionEyeAdaptationInverse
 	: Super(ObjectInitializer)
 {
 #if WITH_EDITORONLY_DATA
-	// Structure to hold one-time initialization
-	struct FConstructorStatics
-	{
-		FText NAME_Utility;
-		FConstructorStatics()
-			: NAME_Utility(LOCTEXT( "Utility", "Utility" ))
-		{
-		}
-	};
-	static FConstructorStatics ConstructorStatics;
-
-	MenuCategories.Add(ConstructorStatics.NAME_Utility);
-
 	Outputs.Reset();
 	Outputs.Add(FExpressionOutput(TEXT("EyeAdaptationInverse")));
 	bShaderInputData = true;
@@ -25083,19 +21872,6 @@ UMaterialExpressionTangentOutput::UMaterialExpressionTangentOutput(const FObject
 : Super(ObjectInitializer)
 {
 #if WITH_EDITORONLY_DATA
-	// Structure to hold one-time initialization
-	struct FConstructorStatics
-	{
-		FText NAME_Custom;
-		FConstructorStatics(const FString& DisplayName, const FString& FunctionName)
-			: NAME_Custom(LOCTEXT( "Custom", "Custom" ))
-		{
-		}
-	};
-	static FConstructorStatics ConstructorStatics(GetDisplayName(), GetFunctionName());
-
-	MenuCategories.Add(ConstructorStatics.NAME_Custom);
-
 	// No outputs
 	Outputs.Reset();
 #endif
@@ -25128,19 +21904,6 @@ UMaterialExpressionClearCoatNormalCustomOutput::UMaterialExpressionClearCoatNorm
 : Super(ObjectInitializer)
 {
 #if WITH_EDITORONLY_DATA
-	// Structure to hold one-time initialization
-	struct FConstructorStatics
-	{
-		FText NAME_Utility;
-		FConstructorStatics(const FString& DisplayName, const FString& FunctionName)
-			: NAME_Utility(LOCTEXT("Utility", "Utility"))
-		{
-		}
-	};
-	static FConstructorStatics ConstructorStatics(GetDisplayName(), GetFunctionName());
-
-	MenuCategories.Add(ConstructorStatics.NAME_Utility);
-
 	bCollapsed = true;
 
 	// No outputs
@@ -25186,19 +21949,6 @@ UMaterialExpressionBentNormalCustomOutput::UMaterialExpressionBentNormalCustomOu
 : Super(ObjectInitializer)
 {
 #if WITH_EDITORONLY_DATA
-	// Structure to hold one-time initialization
-	struct FConstructorStatics
-	{
-		FText NAME_Utility;
-		FConstructorStatics(const FString& DisplayName, const FString& FunctionName)
-			: NAME_Utility(LOCTEXT("Utility", "Utility"))
-		{
-		}
-	};
-	static FConstructorStatics ConstructorStatics(GetDisplayName(), GetFunctionName());
-
-	MenuCategories.Add(ConstructorStatics.NAME_Utility);
-
 	bCollapsed = true;
 
 	// No outputs
@@ -25239,19 +21989,6 @@ UMaterialExpressionVertexInterpolator::UMaterialExpressionVertexInterpolator(con
 : Super(ObjectInitializer)
 {
 #if WITH_EDITORONLY_DATA
-	// Structure to hold one-time initialization
-	struct FConstructorStatics
-	{
-		FText NAME_Utility;
-		FConstructorStatics()
-			: NAME_Utility(LOCTEXT("Utility", "Utility"))
-		{
-		}
-	};
-	static FConstructorStatics ConstructorStatics;
-
-	MenuCategories.Add(ConstructorStatics.NAME_Utility);
-
 	Outputs.Reset();
 	Outputs.Add(FExpressionOutput(TEXT("PS"), 0, 0, 0, 0, 0));
 	bShowOutputNameOnPin = true;
@@ -25360,25 +22097,6 @@ uint32 UMaterialExpressionVertexInterpolator::GetOutputType(int32 OutputIndex)
 ///////////////////////////////////////////////////////////////////////////////
 // UMaterialExpressionrAtmosphericLightVector
 ///////////////////////////////////////////////////////////////////////////////
-UMaterialExpressionAtmosphericLightVector::UMaterialExpressionAtmosphericLightVector(const FObjectInitializer& ObjectInitializer)
-	: Super(ObjectInitializer)
-{
-	// Structure to hold one-time initialization
-	struct FConstructorStatics
-	{
-		FText NAME_Utility;
-		FConstructorStatics()
-			: NAME_Utility(LOCTEXT("Utility", "Utility"))
-		{
-		}
-	};
-	static FConstructorStatics ConstructorStatics;
-
-#if WITH_EDITORONLY_DATA
-	MenuCategories.Add(ConstructorStatics.NAME_Utility);
-#endif
-}
-
 #if WITH_EDITOR
 int32 UMaterialExpressionAtmosphericLightVector::Compile(class FMaterialCompiler* Compiler, int32 OutputIndex)
 {
@@ -25395,25 +22113,6 @@ void UMaterialExpressionAtmosphericLightVector::GetCaption(TArray<FString>& OutC
 ///////////////////////////////////////////////////////////////////////////////
 // UMaterialExpressionrAtmosphericLightColor
 ///////////////////////////////////////////////////////////////////////////////
-UMaterialExpressionAtmosphericLightColor ::UMaterialExpressionAtmosphericLightColor(const FObjectInitializer& ObjectInitializer)
-	: Super(ObjectInitializer)
-{
-	// Structure to hold one-time initialization
-	struct FConstructorStatics
-	{
-		FText NAME_Utility;
-		FConstructorStatics()
-			: NAME_Utility(LOCTEXT("Utility", "Utility"))
-		{
-		}
-	};
-	static FConstructorStatics ConstructorStatics;
-
-#if WITH_EDITORONLY_DATA
-	MenuCategories.Add(ConstructorStatics.NAME_Utility);
-#endif
-}
-
 #if WITH_EDITOR
 int32 UMaterialExpressionAtmosphericLightColor::Compile(class FMaterialCompiler* Compiler, int32 OutputIndex)
 {
@@ -25430,25 +22129,6 @@ void UMaterialExpressionAtmosphericLightColor::GetCaption(TArray<FString>& OutCa
 ///////////////////////////////////////////////////////////////////////////////
 // UMaterialExpressionSkyAtmosphereLightIlluminance
 ///////////////////////////////////////////////////////////////////////////////
-UMaterialExpressionSkyAtmosphereLightIlluminance::UMaterialExpressionSkyAtmosphereLightIlluminance(const FObjectInitializer& ObjectInitializer)
-	: Super(ObjectInitializer)
-{
-	// Structure to hold one-time initialization
-	struct FConstructorStatics
-	{
-		FText NAME_Sky;
-		FConstructorStatics()
-			: NAME_Sky(LOCTEXT("Sky", "Sky"))
-		{
-		}
-	};
-	static FConstructorStatics ConstructorStatics;
-
-#if WITH_EDITORONLY_DATA
-	MenuCategories.Add(ConstructorStatics.NAME_Sky);
-#endif
-}
-
 #if WITH_EDITOR
 FName UMaterialExpressionSkyAtmosphereLightIlluminance::GetInputName(int32 InputIndex) const
 {
@@ -25503,25 +22183,6 @@ void UMaterialExpressionSkyAtmosphereLightIlluminance::GetExpressionToolTip(TArr
 ///////////////////////////////////////////////////////////////////////////////
 // UMaterialExpressionSkyAtmosphereLightIlluminanceOnGround
 ///////////////////////////////////////////////////////////////////////////////
-UMaterialExpressionSkyAtmosphereLightIlluminanceOnGround::UMaterialExpressionSkyAtmosphereLightIlluminanceOnGround(const FObjectInitializer& ObjectInitializer)
-	: Super(ObjectInitializer)
-{
-	// Structure to hold one-time initialization
-	struct FConstructorStatics
-	{
-		FText NAME_Sky;
-		FConstructorStatics()
-			: NAME_Sky(LOCTEXT("Sky", "Sky"))
-		{
-		}
-	};
-	static FConstructorStatics ConstructorStatics;
-
-#if WITH_EDITORONLY_DATA
-	MenuCategories.Add(ConstructorStatics.NAME_Sky);
-#endif
-}
-
 #if WITH_EDITOR
 int32 UMaterialExpressionSkyAtmosphereLightIlluminanceOnGround::Compile(class FMaterialCompiler* Compiler, int32 OutputIndex)
 {
@@ -25543,25 +22204,6 @@ void UMaterialExpressionSkyAtmosphereLightIlluminanceOnGround::GetExpressionTool
 ///////////////////////////////////////////////////////////////////////////////
 // UMaterialExpressionSkyAtmosphereLightDirection
 ///////////////////////////////////////////////////////////////////////////////
-UMaterialExpressionSkyAtmosphereLightDirection::UMaterialExpressionSkyAtmosphereLightDirection(const FObjectInitializer& ObjectInitializer)
-	: Super(ObjectInitializer)
-{
-	// Structure to hold one-time initialization
-	struct FConstructorStatics
-	{
-		FText NAME_Sky;
-		FConstructorStatics()
-			: NAME_Sky(LOCTEXT("Sky", "Sky"))
-		{
-		}
-	};
-	static FConstructorStatics ConstructorStatics;
-
-#if WITH_EDITORONLY_DATA
-	MenuCategories.Add(ConstructorStatics.NAME_Sky);
-#endif
-}
-
 #if WITH_EDITOR
 int32 UMaterialExpressionSkyAtmosphereLightDirection::Compile(class FMaterialCompiler* Compiler, int32 OutputIndex)
 {
@@ -25577,25 +22219,6 @@ void UMaterialExpressionSkyAtmosphereLightDirection::GetCaption(TArray<FString>&
 ///////////////////////////////////////////////////////////////////////////////
 // UMaterialExpressionSkyAtmosphereLightDiskLuminance
 ///////////////////////////////////////////////////////////////////////////////
-UMaterialExpressionSkyAtmosphereLightDiskLuminance::UMaterialExpressionSkyAtmosphereLightDiskLuminance(const FObjectInitializer& ObjectInitializer)
-	: Super(ObjectInitializer)
-{
-	// Structure to hold one-time initialization
-	struct FConstructorStatics
-	{
-		FText NAME_Sky;
-		FConstructorStatics()
-			: NAME_Sky(LOCTEXT("Sky", "Sky"))
-		{
-		}
-	};
-	static FConstructorStatics ConstructorStatics;
-
-#if WITH_EDITORONLY_DATA
-	MenuCategories.Add(ConstructorStatics.NAME_Sky);
-#endif
-}
-
 #if WITH_EDITOR
 int32 UMaterialExpressionSkyAtmosphereLightDiskLuminance::Compile(class FMaterialCompiler* Compiler, int32 OutputIndex)
 {
@@ -25617,25 +22240,6 @@ void UMaterialExpressionSkyAtmosphereLightDiskLuminance::GetCaption(TArray<FStri
 ///////////////////////////////////////////////////////////////////////////////
 // UMaterialExpressionSkyAtmosphereViewLuminance
 ///////////////////////////////////////////////////////////////////////////////
-UMaterialExpressionSkyAtmosphereViewLuminance::UMaterialExpressionSkyAtmosphereViewLuminance(const FObjectInitializer& ObjectInitializer)
-	: Super(ObjectInitializer)
-{
-	// Structure to hold one-time initialization
-	struct FConstructorStatics
-	{
-		FText NAME_Sky;
-		FConstructorStatics()
-			: NAME_Sky(LOCTEXT("Sky", "Sky"))
-		{
-		}
-	};
-	static FConstructorStatics ConstructorStatics;
-
-#if WITH_EDITORONLY_DATA
-	MenuCategories.Add(ConstructorStatics.NAME_Sky);
-#endif
-}
-
 #if WITH_EDITOR
 int32 UMaterialExpressionSkyAtmosphereViewLuminance::Compile(class FMaterialCompiler* Compiler, int32 OutputIndex)
 {
@@ -25651,25 +22255,6 @@ void UMaterialExpressionSkyAtmosphereViewLuminance::GetCaption(TArray<FString>&
 ///////////////////////////////////////////////////////////////////////////////
 // UMaterialExpressionSkyAtmosphereAerialPerpsective
 ///////////////////////////////////////////////////////////////////////////////
-UMaterialExpressionSkyAtmosphereAerialPerspective::UMaterialExpressionSkyAtmosphereAerialPerspective(const FObjectInitializer& ObjectInitializer)
-	: Super(ObjectInitializer)
-{
-	// Structure to hold one-time initialization
-	struct FConstructorStatics
-	{
-		FText NAME_Sky;
-		FConstructorStatics()
-			: NAME_Sky(LOCTEXT("Sky", "Sky"))
-		{
-		}
-	};
-	static FConstructorStatics ConstructorStatics;
-
-#if WITH_EDITORONLY_DATA
-	MenuCategories.Add(ConstructorStatics.NAME_Sky);
-#endif
-}
-
 #if WITH_EDITOR
 FName UMaterialExpressionSkyAtmosphereAerialPerspective::GetInputName(int32 InputIndex) const
 {
@@ -25718,25 +22303,6 @@ void UMaterialExpressionSkyAtmosphereAerialPerspective::GetCaption(TArray<FStrin
 ///////////////////////////////////////////////////////////////////////////////
 // UMaterialExpressionSkyAtmosphereAerialPerpsective
 ///////////////////////////////////////////////////////////////////////////////
-UMaterialExpressionSkyAtmosphereDistantLightScatteredLuminance::UMaterialExpressionSkyAtmosphereDistantLightScatteredLuminance(const FObjectInitializer& ObjectInitializer)
-	: Super(ObjectInitializer)
-{
-	// Structure to hold one-time initialization
-	struct FConstructorStatics
-	{
-		FText NAME_Sky;
-		FConstructorStatics()
-			: NAME_Sky(LOCTEXT("Sky", "Sky"))
-		{
-		}
-	};
-	static FConstructorStatics ConstructorStatics;
-
-#if WITH_EDITORONLY_DATA
-	MenuCategories.Add(ConstructorStatics.NAME_Sky);
-#endif
-}
-
 #if WITH_EDITOR
 int32 UMaterialExpressionSkyAtmosphereDistantLightScatteredLuminance::Compile(class FMaterialCompiler* Compiler, int32 OutputIndex)
 {
@@ -25756,19 +22322,6 @@ UMaterialExpressionPreSkinnedPosition::UMaterialExpressionPreSkinnedPosition(con
 	: Super(ObjectInitializer)
 {
 #if WITH_EDITORONLY_DATA
-	// Structure to hold one-time initialization
-	struct FConstructorStatics
-	{
-		FText NAME_Constants;
-		FConstructorStatics()
-			: NAME_Constants(LOCTEXT( "Vectors", "Vectors" ))
-		{
-		}
-	};
-	static FConstructorStatics ConstructorStatics;
-
-	MenuCategories.Add(ConstructorStatics.NAME_Constants);
-
 	Outputs.Reset();
 	Outputs.Add(FExpressionOutput(TEXT(""), 1, 1, 1, 1, 0));
 	bShaderInputData = true;
@@ -25817,19 +22370,6 @@ UMaterialExpressionPreSkinnedNormal::UMaterialExpressionPreSkinnedNormal(const F
 	: Super(ObjectInitializer)
 {
 #if WITH_EDITORONLY_DATA
-	// Structure to hold one-time initialization
-	struct FConstructorStatics
-	{
-		FText NAME_Constants;
-		FConstructorStatics()
-			: NAME_Constants(LOCTEXT( "Vectors", "Vectors" ))
-		{
-		}
-	};
-	static FConstructorStatics ConstructorStatics;
-
-	MenuCategories.Add(ConstructorStatics.NAME_Constants);
-
 	Outputs.Reset();
 	Outputs.Add(FExpressionOutput(TEXT(""), 1, 1, 1, 1, 0));
 	bShaderInputData = true;
@@ -26044,21 +22584,6 @@ void UMaterialExpressionCurveAtlasRowParameter::PostEditChangeProperty(FProperty
 UMaterialExpressionHairAttributes::UMaterialExpressionHairAttributes(const FObjectInitializer& ObjectInitializer)
 	: Super(ObjectInitializer)
 {
-#if WITH_EDITORONLY_DATA
-	// Structure to hold one-time initialization
-	struct FConstructorStatics
-	{
-		FText NAME_Utility;
-		FConstructorStatics()
-			: NAME_Utility(LOCTEXT("Hair Attributes", "Hair Attributes"))
-		{
-		}
-	};
-	static FConstructorStatics ConstructorStatics;
-	MenuCategories.Add(ConstructorStatics.NAME_Utility);
-
-#endif
-
 	bUseTangentSpace = true;
 
 #if WITH_EDITORONLY_DATA
@@ -26150,31 +22675,16 @@ int32 UMaterialExpressionHairAttributes::Compile(class FMaterialCompiler* Compil
 void UMaterialExpressionHairAttributes::GetCaption(TArray<FString>& OutCaptions) const
 {
 	OutCaptions.Add(TEXT("Hair Attributes"));
-}
-#endif // WITH_EDITOR
-
-//
-// Hair Color
-//
-
-UMaterialExpressionHairColor::UMaterialExpressionHairColor(const FObjectInitializer& ObjectInitializer)
-	: Super(ObjectInitializer)
-{
-#if WITH_EDITORONLY_DATA
-	// Structure to hold one-time initialization
-	struct FConstructorStatics
-	{
-		FText NAME_Utility;
-		FConstructorStatics()
-			: NAME_Utility(LOCTEXT("Hair Color", "Hair Color"))
-		{
-		}
-	};
-	static FConstructorStatics ConstructorStatics;
-	MenuCategories.Add(ConstructorStatics.NAME_Utility);
+}
+#endif // WITH_EDITOR
 
-#endif
+//
+// Hair Color
+//
 
+UMaterialExpressionHairColor::UMaterialExpressionHairColor(const FObjectInitializer& ObjectInitializer)
+	: Super(ObjectInitializer)
+{
 #if WITH_EDITORONLY_DATA
 	bShowOutputNameOnPin = true;
 	Outputs.Reset();
@@ -26204,26 +22714,6 @@ void UMaterialExpressionHairColor::GetCaption(TArray<FString>& OutCaptions) cons
 //
 //  UMaterialExpressionARPassthroughCameraUVs
 //
-
-UMaterialExpressionMapARPassthroughCameraUV::UMaterialExpressionMapARPassthroughCameraUV(const FObjectInitializer& ObjectInitializer)
-	: Super(ObjectInitializer)
-{
-#if WITH_EDITORONLY_DATA
-	// Structure to hold one-time initialization
-	struct FConstructorStatics
-	{
-		FText NAME_Coordinates;
-		FConstructorStatics()
-			: NAME_Coordinates(LOCTEXT("Coordinates", "Coordinates"))
-		{
-		}
-	};
-	static FConstructorStatics ConstructorStatics;
-
-	MenuCategories.Add(ConstructorStatics.NAME_Coordinates);
-#endif // WITH_EDITORONLY_DATA
-}
-
 #if WITH_EDITOR
 int32 UMaterialExpressionMapARPassthroughCameraUV::Compile(class FMaterialCompiler* Compiler, int32 OutputIndex)
 {
@@ -26294,21 +22784,6 @@ uint32 UMaterialExpressionShadingModel::GetOutputType(int32 OutputIndex)
 UMaterialExpressionSingleLayerWaterMaterialOutput::UMaterialExpressionSingleLayerWaterMaterialOutput(const FObjectInitializer& ObjectInitializer)
 	: Super(ObjectInitializer)
 {
-	// Structure to hold one-time initialization
-	struct FConstructorStatics
-	{
-		FText NAME_Water;
-		FConstructorStatics()
-			: NAME_Water(LOCTEXT("Water", "Water"))
-		{
-		}
-	};
-	static FConstructorStatics ConstructorStatics;
-
-#if WITH_EDITORONLY_DATA
-	MenuCategories.Add(ConstructorStatics.NAME_Water);
-#endif
-
 #if WITH_EDITOR
 	Outputs.Reset();
 #endif
@@ -26379,17 +22854,6 @@ FString UMaterialExpressionSingleLayerWaterMaterialOutput::GetDisplayName() cons
 UMaterialExpressionVolumetricAdvancedMaterialOutput::UMaterialExpressionVolumetricAdvancedMaterialOutput(const FObjectInitializer& ObjectInitializer)
 	: Super(ObjectInitializer)
 {
-	// Structure to hold one-time initialization
-	struct FConstructorStatics
-	{
-		FText NAME_VolumetricAdvancedOutput;
-		FConstructorStatics()
-			: NAME_VolumetricAdvancedOutput(LOCTEXT("Volume", "Volume"))
-		{
-		}
-	};
-	static FConstructorStatics ConstructorStatics;
-
 	ConstPhaseG = 0.0f;
 	ConstPhaseG2 = 0.0f;
 	ConstPhaseBlend = 0.0f;
@@ -26405,10 +22869,6 @@ UMaterialExpressionVolumetricAdvancedMaterialOutput::UMaterialExpressionVolumetr
 	bRayMarchVolumeShadow = true;
 	bClampMultiScatteringContribution = true;
 
-#if WITH_EDITORONLY_DATA
-	MenuCategories.Add(ConstructorStatics.NAME_VolumetricAdvancedOutput);
-#endif
-
 #if WITH_EDITOR
 	Outputs.Reset();
 #endif
@@ -26497,18 +22957,6 @@ UMaterialExpressionVolumetricAdvancedMaterialInput::UMaterialExpressionVolumetri
 	: Super(ObjectInitializer)
 {
 #if WITH_EDITORONLY_DATA
-	// Structure to hold one-time initialization
-	struct FConstructorStatics
-	{
-		FText NAME_VolumetricAdvancedInput;
-		FConstructorStatics()
-			: NAME_VolumetricAdvancedInput(LOCTEXT("Volume", "Volume"))
-		{
-		}
-	};
-	static FConstructorStatics ConstructorStatics;
-	MenuCategories.Add(ConstructorStatics.NAME_VolumetricAdvancedInput);
-
 	bShowOutputNameOnPin = true;
 
 	Outputs.Reset();
@@ -26546,21 +22994,6 @@ void UMaterialExpressionVolumetricAdvancedMaterialInput::GetCaption(TArray<FStri
 UMaterialExpressionVolumetricCloudEmptySpaceSkippingOutput::UMaterialExpressionVolumetricCloudEmptySpaceSkippingOutput(const FObjectInitializer& ObjectInitializer)
 	: Super(ObjectInitializer)
 {
-	// Structure to hold one-time initialization
-	struct FConstructorStatics
-	{
-		FText NAME_VolumetricCloudEmptySpaceSkippingOutput;
-		FConstructorStatics()
-			: NAME_VolumetricCloudEmptySpaceSkippingOutput(LOCTEXT("Volume", "Volume"))
-		{
-		}
-	};
-	static FConstructorStatics ConstructorStatics;
-
-#if WITH_EDITORONLY_DATA
-	MenuCategories.Add(ConstructorStatics.NAME_VolumetricCloudEmptySpaceSkippingOutput);
-#endif
-
 #if WITH_EDITOR
 	Outputs.Reset();
 #endif
@@ -26616,18 +23049,6 @@ UMaterialExpressionVolumetricCloudEmptySpaceSkippingInput::UMaterialExpressionVo
 	: Super(ObjectInitializer)
 {
 #if WITH_EDITORONLY_DATA
-	// Structure to hold one-time initialization
-	struct FConstructorStatics
-	{
-		FText NAME_UMaterialExpressionVolumetricCloudEmptySpaceSkippingInputInput;
-		FConstructorStatics()
-			: NAME_UMaterialExpressionVolumetricCloudEmptySpaceSkippingInputInput(LOCTEXT("Volume", "Volume"))
-		{
-		}
-	};
-	static FConstructorStatics ConstructorStatics;
-	MenuCategories.Add(ConstructorStatics.NAME_UMaterialExpressionVolumetricCloudEmptySpaceSkippingInputInput);
-
 	bShowOutputNameOnPin = true;
 
 	Outputs.Reset();
@@ -26689,26 +23110,6 @@ uint32 UMaterialExpressionVolumetricCloudEmptySpaceSkippingInput::GetInputType(i
 ///////////////////////////////////////////////////////////////////////////////
 // UMaterialExpressionReflectionCapturePassSwitch
 ///////////////////////////////////////////////////////////////////////////////
-
-UMaterialExpressionReflectionCapturePassSwitch::UMaterialExpressionReflectionCapturePassSwitch(const FObjectInitializer& ObjectInitializer)
-	: Super(ObjectInitializer)
-{
-	// Structure to hold one-time initialization
-	struct FConstructorStatics
-	{
-		FText NAME_Utility;
-		FConstructorStatics()
-			: NAME_Utility(LOCTEXT("Utility", "Utility"))
-		{
-		}
-	};
-	static FConstructorStatics ConstructorStatics;
-
-#if WITH_EDITORONLY_DATA
-	MenuCategories.Add(ConstructorStatics.NAME_Utility);
-#endif
-}
-
 #if WITH_EDITOR
 int32 UMaterialExpressionReflectionCapturePassSwitch::Compile(class FMaterialCompiler* Compiler, int32 OutputIndex)
 {
@@ -26750,18 +23151,6 @@ UMaterialExpressionCloudSampleAttribute::UMaterialExpressionCloudSampleAttribute
 	: Super(ObjectInitializer)
 {
 #if WITH_EDITORONLY_DATA
-	// Structure to hold one-time initialization
-	struct FConstructorStatics
-	{
-		FText NAME_CloudSampleAttributes;
-		FConstructorStatics()
-			: NAME_CloudSampleAttributes(LOCTEXT("Volume", "Volume"))
-		{
-		}
-	};
-	static FConstructorStatics ConstructorStatics;
-	MenuCategories.Add(ConstructorStatics.NAME_CloudSampleAttributes);
-
 	bShowOutputNameOnPin = true;
 
 	Outputs.Reset();
@@ -26814,21 +23203,6 @@ void UMaterialExpressionCloudSampleAttribute::GetExpressionToolTip(TArray<FStrin
 UMaterialExpressionThinTranslucentMaterialOutput::UMaterialExpressionThinTranslucentMaterialOutput(const FObjectInitializer& ObjectInitializer)
 	: Super(ObjectInitializer)
 {
-	// Structure to hold one-time initialization
-	struct FConstructorStatics
-	{
-		FText NAME_ThinTranslucent;
-		FConstructorStatics()
-			: NAME_ThinTranslucent(LOCTEXT("ThinTranslucent", "ThinTranslucent"))
-		{
-		}
-	};
-	static FConstructorStatics ConstructorStatics;
-
-#if WITH_EDITORONLY_DATA
-	MenuCategories.Add(ConstructorStatics.NAME_ThinTranslucent);
-#endif
-
 #if WITH_EDITOR
 	Outputs.Reset();
 #endif
@@ -26897,21 +23271,6 @@ FString UMaterialExpressionThinTranslucentMaterialOutput::GetDisplayName() const
 UMaterialExpressionAbsorptionMediumMaterialOutput::UMaterialExpressionAbsorptionMediumMaterialOutput(const FObjectInitializer& ObjectInitializer)
 	: Super(ObjectInitializer)
 {
-	// Structure to hold one-time initialization
-	struct FConstructorStatics
-	{
-		FText NAME_MediumAbsorption;
-		FConstructorStatics()
-			: NAME_MediumAbsorption(LOCTEXT("Shading", "Shading"))
-		{
-		}
-	};
-	static FConstructorStatics ConstructorStatics;
-
-#if WITH_EDITORONLY_DATA
-	MenuCategories.Add(ConstructorStatics.NAME_MediumAbsorption);
-#endif
-
 #if WITH_EDITOR
 	Outputs.Reset();
 #endif
@@ -27019,21 +23378,6 @@ FName UMaterialExpressionSceneDepthWithoutWater::GetInputName(int32 InputIndex)
 
 #endif // WITH_EDITOR
 
-
-UMaterialExpressionIfThenElse::UMaterialExpressionIfThenElse(const FObjectInitializer& ObjectInitializer)
-: Super(ObjectInitializer)
-{
-	struct FConstructorStatics
-	{
-		FText NAME_Execution;
-		FConstructorStatics() : NAME_Execution(LOCTEXT("Execution", "Execution")) { }
-	};
-	static FConstructorStatics ConstructorStatics;
-#if WITH_EDITORONLY_DATA
-	MenuCategories.Add(ConstructorStatics.NAME_Execution);
-#endif
-}
-
 #if WITH_EDITOR
 void UMaterialExpressionIfThenElse::GetCaption(TArray<FString>& OutCaptions) const
 {
@@ -27055,25 +23399,6 @@ uint32 UMaterialExpressionIfThenElse::GetInputType(int32 InputIndex)
 ///////////////////////////////////////////////////////////////////////////////
 // UMaterialExpressionSkyLightEnvMapSample
 ///////////////////////////////////////////////////////////////////////////////
-UMaterialExpressionSkyLightEnvMapSample::UMaterialExpressionSkyLightEnvMapSample(const FObjectInitializer& ObjectInitializer)
-	: Super(ObjectInitializer)
-{
-	// Structure to hold one-time initialization
-	struct FConstructorStatics
-	{
-		FText NAME_Sky;
-		FConstructorStatics()
-			: NAME_Sky(LOCTEXT("Sky", "Sky"))
-		{
-		}
-	};
-	static FConstructorStatics ConstructorStatics;
-
-#if WITH_EDITORONLY_DATA
-	MenuCategories.Add(ConstructorStatics.NAME_Sky);
-#endif
-}
-
 #if WITH_EDITOR
 int32 UMaterialExpressionSkyLightEnvMapSample::Compile(class FMaterialCompiler* Compiler, int32 OutputIndex)
 {
@@ -27120,23 +23445,7 @@ UObject* UMaterialExpressionSparseVolumeTextureBase::GetReferencedTexture() cons
 UMaterialExpressionSparseVolumeTextureObject::UMaterialExpressionSparseVolumeTextureObject(const FObjectInitializer& ObjectInitializer)
 	: Super(ObjectInitializer)
 {
-	// Structure to hold one-time initialization
-	struct FConstructorStatics
-	{
-		FText NAME_Texture;
-		FText NAME_Functions;
-		FConstructorStatics()
-			: NAME_Texture(LOCTEXT("Texture", "Texture"))
-			, NAME_Functions(LOCTEXT("Functions", "Functions"))
-		{
-		}
-	};
-	static FConstructorStatics ConstructorStatics;
-
 #if WITH_EDITORONLY_DATA
-	MenuCategories.Add(ConstructorStatics.NAME_Texture);
-	MenuCategories.Add(ConstructorStatics.NAME_Functions);
-
 	Outputs.Empty();
 	Outputs.Add(FExpressionOutput(TEXT("")));
 
@@ -27191,24 +23500,7 @@ uint32 UMaterialExpressionSparseVolumeTextureObject::GetOutputType(int32 OutputI
 UMaterialExpressionSparseVolumeTextureObjectParameter::UMaterialExpressionSparseVolumeTextureObjectParameter(const FObjectInitializer& ObjectInitializer)
 	: Super(ObjectInitializer)
 {
-	// Structure to hold one-time initialization
-	struct FConstructorStatics
-	{
-		FText NAME_Texture;
-		FText NAME_Parameters;
-		FConstructorStatics()
-			: NAME_Texture(LOCTEXT("Texture", "Texture"))
-			, NAME_Parameters(LOCTEXT("Parameters", "Parameters"))
-		{
-		}
-	};
-	static FConstructorStatics ConstructorStatics;
-
 #if WITH_EDITORONLY_DATA
-	MenuCategories.Empty();
-	MenuCategories.Add(ConstructorStatics.NAME_Texture);
-	MenuCategories.Add(ConstructorStatics.NAME_Parameters);
-
 	Outputs.Empty();
 	Outputs.Add(FExpressionOutput(TEXT("")));
 #endif
@@ -27253,21 +23545,6 @@ uint32 UMaterialExpressionSparseVolumeTextureObjectParameter::GetOutputType(int3
 UMaterialExpressionSparseVolumeTextureSample::UMaterialExpressionSparseVolumeTextureSample(const FObjectInitializer& ObjectInitializer)
 	: Super(ObjectInitializer)
 {
-	// Structure to hold one-time initialization
-	struct FConstructorStatics
-	{
-		FText NAME_Texture;
-		FConstructorStatics()
-			: NAME_Texture(LOCTEXT("Texture", "Texture"))
-		{
-		}
-	};
-	static FConstructorStatics ConstructorStatics;
-
-#if WITH_EDITORONLY_DATA
-	MenuCategories.Add(ConstructorStatics.NAME_Texture);
-#endif
-
 #if WITH_EDITOR
 	Outputs.Reset();
 	Outputs.Add(FExpressionOutput(TEXT("Attributes A"), 1, 1, 1, 1, 1));
@@ -27533,22 +23810,7 @@ void UMaterialExpressionSparseVolumeTextureSample::GetCaption(TArray<FString>& O
 UMaterialExpressionSparseVolumeTextureSampleParameter::UMaterialExpressionSparseVolumeTextureSampleParameter(const FObjectInitializer& ObjectInitializer)
 	: Super(ObjectInitializer)
 {
-	// Structure to hold one-time initialization
-	struct FConstructorStatics
-	{
-		FText NAME_Parameters;
-		FConstructorStatics()
-			: NAME_Parameters(LOCTEXT("Parameters", "Parameters"))
-		{
-		}
-	};
-	static FConstructorStatics ConstructorStatics;
-
 	bIsParameterExpression = true;
-
-#if WITH_EDITORONLY_DATA
-	MenuCategories.Add(ConstructorStatics.NAME_Parameters);
-#endif
 }
 
 #if WITH_EDITOR
@@ -27625,21 +23887,6 @@ bool UMaterialExpressionSparseVolumeTextureSampleParameter::SetParameterValue(co
 UMaterialExpressionSubsurfaceMediumMaterialOutput::UMaterialExpressionSubsurfaceMediumMaterialOutput(const FObjectInitializer& ObjectInitializer)
 	: Super(ObjectInitializer)
 {
-	// Structure to hold one-time initialization
-	struct FConstructorStatics
-	{
-		FText NAME_SubsurfaceMedium;
-		FConstructorStatics()
-			: NAME_SubsurfaceMedium(LOCTEXT("Shading", "Shading"))
-		{
-		}
-	};
-	static FConstructorStatics ConstructorStatics;
-
-#if WITH_EDITORONLY_DATA
-	MenuCategories.Add(ConstructorStatics.NAME_SubsurfaceMedium);
-#endif
-
 #if WITH_EDITOR
 	Outputs.Reset();
 #endif
@@ -27710,21 +23957,6 @@ FString UMaterialExpressionSubsurfaceMediumMaterialOutput::GetDisplayName() cons
 UMaterialExpressionNeuralNetworkInput::UMaterialExpressionNeuralNetworkInput(const FObjectInitializer& ObjectInitializer)
 	: Super(ObjectInitializer)
 {
-	// Structure to hold one-time initialization
-	struct FConstructorStatics
-	{
-		FText NAME_NeuralNetwork;
-		FConstructorStatics()
-			: NAME_NeuralNetwork(LOCTEXT("NeuralNetwork", "NeuralNetwork"))
-		{
-		}
-	};
-	static FConstructorStatics ConstructorStatics;
-
-#if WITH_EDITORONLY_DATA
-	MenuCategories.Add(ConstructorStatics.NAME_NeuralNetwork);
-#endif
-
 #if WITH_EDITOR
 	Outputs.Reset();
 #endif
@@ -27812,21 +24044,7 @@ FString UMaterialExpressionNeuralNetworkInput::GetFunctionName() const
 UMaterialExpressionNeuralNetworkOutput::UMaterialExpressionNeuralNetworkOutput(const FObjectInitializer& ObjectInitializer)
 	: Super(ObjectInitializer)
 {
-
-	// Structure to hold one-time initialization
-	struct FConstructorStatics
-	{
-		FText NAME_NeuralNetwork;
-		FConstructorStatics()
-			: NAME_NeuralNetwork(LOCTEXT("NeuralNetwork", "NeuralNetwork"))
-		{
-		}
-	};
-	static FConstructorStatics ConstructorStatics;
-
 #if WITH_EDITORONLY_DATA
-	MenuCategories.Add(ConstructorStatics.NAME_NeuralNetwork);
-
 	bShaderInputData = true;
 	bShowOutputNameOnPin = true;
 #endif
@@ -27884,23 +24102,7 @@ uint32 UMaterialExpressionNeuralNetworkOutput::GetInputType(int32 InputIndex)
 UMaterialExpressionFirstPersonOutput::UMaterialExpressionFirstPersonOutput(const FObjectInitializer& ObjectInitializer)
 	: Super(ObjectInitializer)
 {
-	// Structure to hold one-time initialization
-	struct FConstructorStatics
-	{
-		FText NAME_FirstPerson;
-		FConstructorStatics()
-			: NAME_FirstPerson(LOCTEXT("FirstPerson", "First Person"))
-		{
-		}
-	};
-	static FConstructorStatics ConstructorStatics;
-
 	ConstFirstPersonInterpolationAlpha = 1.0f;
-
-#if WITH_EDITORONLY_DATA
-	MenuCategories.Add(ConstructorStatics.NAME_FirstPerson);
-#endif
-
 #if WITH_EDITOR
 	Outputs.Reset();
 #endif
@@ -27946,22 +24148,6 @@ UMaterialExpressionIsFirstPerson::UMaterialExpressionIsFirstPerson(const FObject
 	: Super(ObjectInitializer)
 {
 #if WITH_EDITORONLY_DATA
-	// Structure to hold one-time initialization
-	struct FConstructorStatics
-	{
-		FText NAME_Constants;
-		FText NAME_FirstPerson;
-		FConstructorStatics()
-			: NAME_Constants(LOCTEXT("Constants", "Constants")),
-			NAME_FirstPerson(LOCTEXT("FirstPerson", "First Person"))
-		{
-		}
-	};
-	static FConstructorStatics ConstructorStatics;
-
-	MenuCategories.Add(ConstructorStatics.NAME_Constants);
-	MenuCategories.Add(ConstructorStatics.NAME_FirstPerson);
-
 	bShaderInputData = true;
 #endif
 }
diff --git a/Engine/Source/Runtime/Engine/Public/Materials/MaterialExpressionAbs.h b/Engine/Source/Runtime/Engine/Public/Materials/MaterialExpressionAbs.h
--- a/Engine/Source/Runtime/Engine/Public/Materials/MaterialExpressionAbs.h
+++ b/Engine/Source/Runtime/Engine/Public/Materials/MaterialExpressionAbs.h
@@ -16,7 +16,9 @@
 UCLASS(MinimalAPI, collapsecategories, hidecategories=Object)
 class UMaterialExpressionAbs : public UMaterialExpression
 {
-	GENERATED_UCLASS_BODY()
+	GENERATED_BODY()
+
+public:
 
 	/** Link to the input expression to be evaluated */
 	UPROPERTY()
diff --git a/Engine/Source/Runtime/Engine/Public/Materials/MaterialExpressionAppendVector.h b/Engine/Source/Runtime/Engine/Public/Materials/MaterialExpressionAppendVector.h
--- a/Engine/Source/Runtime/Engine/Public/Materials/MaterialExpressionAppendVector.h
+++ b/Engine/Source/Runtime/Engine/Public/Materials/MaterialExpressionAppendVector.h
@@ -12,7 +12,9 @@
 UCLASS(MinimalAPI, collapsecategories, hidecategories=Object)
 class UMaterialExpressionAppendVector : public UMaterialExpression
 {
-	GENERATED_UCLASS_BODY()
+	GENERATED_BODY()
+
+public:
 
 	UPROPERTY()
 	FExpressionInput A;
diff --git a/Engine/Source/Runtime/Engine/Public/Materials/MaterialExpressionArccosine.h b/Engine/Source/Runtime/Engine/Public/Materials/MaterialExpressionArccosine.h
--- a/Engine/Source/Runtime/Engine/Public/Materials/MaterialExpressionArccosine.h
+++ b/Engine/Source/Runtime/Engine/Public/Materials/MaterialExpressionArccosine.h
@@ -11,7 +11,9 @@
 UCLASS(MinimalAPI, collapsecategories, hidecategories=Object)
 class UMaterialExpressionArccosine : public UMaterialExpression
 {
-	GENERATED_UCLASS_BODY()
+	GENERATED_BODY()
+
+public:
 
 	UPROPERTY()
 	FExpressionInput Input;
diff --git a/Engine/Source/Runtime/Engine/Public/Materials/MaterialExpressionArccosineFast.h b/Engine/Source/Runtime/Engine/Public/Materials/MaterialExpressionArccosineFast.h
--- a/Engine/Source/Runtime/Engine/Public/Materials/MaterialExpressionArccosineFast.h
+++ b/Engine/Source/Runtime/Engine/Public/Materials/MaterialExpressionArccosineFast.h
@@ -11,7 +11,9 @@
 UCLASS(collapsecategories, hidecategories=Object)
 class UMaterialExpressionArccosineFast : public UMaterialExpression
 {
-	GENERATED_UCLASS_BODY()
+	GENERATED_BODY()
+
+public:
 
 	UPROPERTY()
 	FExpressionInput Input;
diff --git a/Engine/Source/Runtime/Engine/Public/Materials/MaterialExpressionArcsine.h b/Engine/Source/Runtime/Engine/Public/Materials/MaterialExpressionArcsine.h
--- a/Engine/Source/Runtime/Engine/Public/Materials/MaterialExpressionArcsine.h
+++ b/Engine/Source/Runtime/Engine/Public/Materials/MaterialExpressionArcsine.h
@@ -11,7 +11,9 @@
 UCLASS(MinimalAPI, collapsecategories, hidecategories=Object)
 class UMaterialExpressionArcsine : public UMaterialExpression
 {
-	GENERATED_UCLASS_BODY()
+	GENERATED_BODY()
+
+public:
 
 	UPROPERTY()
 	FExpressionInput Input;
diff --git a/Engine/Source/Runtime/Engine/Public/Materials/MaterialExpressionArcsineFast.h b/Engine/Source/Runtime/Engine/Public/Materials/MaterialExpressionArcsineFast.h
--- a/Engine/Source/Runtime/Engine/Public/Materials/MaterialExpressionArcsineFast.h
+++ b/Engine/Source/Runtime/Engine/Public/Materials/MaterialExpressionArcsineFast.h
@@ -11,7 +11,9 @@
 UCLASS(collapsecategories, hidecategories=Object)
 class UMaterialExpressionArcsineFast : public UMaterialExpression
 {
-	GENERATED_UCLASS_BODY()
+	GENERATED_BODY()
+
+public:
 
 	UPROPERTY()
 	FExpressionInput Input;
diff --git a/Engine/Source/Runtime/Engine/Public/Materials/MaterialExpressionArctangent.h b/Engine/Source/Runtime/Engine/Public/Materials/MaterialExpressionArctangent.h
--- a/Engine/Source/Runtime/Engine/Public/Materials/MaterialExpressionArctangent.h
+++ b/Engine/Source/Runtime/Engine/Public/Materials/MaterialExpressionArctangent.h
@@ -11,7 +11,9 @@
 UCLASS(MinimalAPI, collapsecategories, hidecategories=Object)
 class UMaterialExpressionArctangent : public UMaterialExpression
 {
-	GENERATED_UCLASS_BODY()
+	GENERATED_BODY()
+
+public:
 
 	UPROPERTY()
 	FExpressionInput Input;
diff --git a/Engine/Source/Runtime/Engine/Public/Materials/MaterialExpressionArctangent2.h b/Engine/Source/Runtime/Engine/Public/Materials/MaterialExpressionArctangent2.h
--- a/Engine/Source/Runtime/Engine/Public/Materials/MaterialExpressionArctangent2.h
+++ b/Engine/Source/Runtime/Engine/Public/Materials/MaterialExpressionArctangent2.h
@@ -11,7 +11,9 @@
 UCLASS(MinimalAPI, collapsecategories, hidecategories=Object)
 class UMaterialExpressionArctangent2 : public UMaterialExpression
 {
-	GENERATED_UCLASS_BODY()
+	GENERATED_BODY()
+
+public:
 
 	UPROPERTY()
 	FExpressionInput Y;
diff --git a/Engine/Source/Runtime/Engine/Public/Materials/MaterialExpressionArctangent2Fast.h b/Engine/Source/Runtime/Engine/Public/Materials/MaterialExpressionArctangent2Fast.h
--- a/Engine/Source/Runtime/Engine/Public/Materials/MaterialExpressionArctangent2Fast.h
+++ b/Engine/Source/Runtime/Engine/Public/Materials/MaterialExpressionArctangent2Fast.h
@@ -11,7 +11,9 @@
 UCLASS(collapsecategories, hidecategories=Object)
 class UMaterialExpressionArctangent2Fast : public UMaterialExpression
 {
-	GENERATED_UCLASS_BODY()
+	GENERATED_BODY()
+
+public:
 
 	UPROPERTY()
 	FExpressionInput Y;
diff --git a/Engine/Source/Runtime/Engine/Public/Materials/MaterialExpressionArctangentFast.h b/Engine/Source/Runtime/Engine/Public/Materials/MaterialExpressionArctangentFast.h
--- a/Engine/Source/Runtime/Engine/Public/Materials/MaterialExpressionArctangentFast.h
+++ b/Engine/Source/Runtime/Engine/Public/Materials/MaterialExpressionArctangentFast.h
@@ -11,7 +11,9 @@
 UCLASS(collapsecategories, hidecategories=Object)
 class UMaterialExpressionArctangentFast : public UMaterialExpression
 {
-	GENERATED_UCLASS_BODY()
+	GENERATED_BODY()
+
+public:
 
 	UPROPERTY()
 	FExpressionInput Input;
diff --git a/Engine/Source/Runtime/Engine/Public/Materials/MaterialExpressionAtmosphericLightColor.h b/Engine/Source/Runtime/Engine/Public/Materials/MaterialExpressionAtmosphericLightColor.h
--- a/Engine/Source/Runtime/Engine/Public/Materials/MaterialExpressionAtmosphericLightColor.h
+++ b/Engine/Source/Runtime/Engine/Public/Materials/MaterialExpressionAtmosphericLightColor.h
@@ -11,7 +11,9 @@
 UCLASS(collapsecategories, hidecategories = Object, DisplayName = "Atmosphere Sun Light Illuminance On Ground")
 class UMaterialExpressionAtmosphericLightColor : public UMaterialExpression
 {
-	GENERATED_UCLASS_BODY()
+	GENERATED_BODY()
+
+public:
 
 	//~ Begin UMaterialExpression Interface
 #if WITH_EDITOR
diff --git a/Engine/Source/Runtime/Engine/Public/Materials/MaterialExpressionAtmosphericLightVector.h b/Engine/Source/Runtime/Engine/Public/Materials/MaterialExpressionAtmosphericLightVector.h
--- a/Engine/Source/Runtime/Engine/Public/Materials/MaterialExpressionAtmosphericLightVector.h
+++ b/Engine/Source/Runtime/Engine/Public/Materials/MaterialExpressionAtmosphericLightVector.h
@@ -11,7 +11,9 @@
 UCLASS(collapsecategories, hidecategories = Object, DisplayName = "Atmosphere Sun Light Vector")
 class UMaterialExpressionAtmosphericLightVector : public UMaterialExpression
 {
-	GENERATED_UCLASS_BODY()
+	GENERATED_BODY()
+
+public:
 
 	//~ Begin UMaterialExpression Interface
 #if WITH_EDITOR
diff --git a/Engine/Source/Runtime/Engine/Public/Materials/MaterialExpressionCeil.h b/Engine/Source/Runtime/Engine/Public/Materials/MaterialExpressionCeil.h
--- a/Engine/Source/Runtime/Engine/Public/Materials/MaterialExpressionCeil.h
+++ b/Engine/Source/Runtime/Engine/Public/Materials/MaterialExpressionCeil.h
@@ -12,7 +12,9 @@
 UCLASS(MinimalAPI, collapsecategories, hidecategories=Object)
 class UMaterialExpressionCeil : public UMaterialExpression
 {
-	GENERATED_UCLASS_BODY()
+	GENERATED_BODY()
+
+public:
 
 	UPROPERTY()
 	FExpressionInput Input;
diff --git a/Engine/Source/Runtime/Engine/Public/Materials/MaterialExpressionComponentMask.h b/Engine/Source/Runtime/Engine/Public/Materials/MaterialExpressionComponentMask.h
--- a/Engine/Source/Runtime/Engine/Public/Materials/MaterialExpressionComponentMask.h
+++ b/Engine/Source/Runtime/Engine/Public/Materials/MaterialExpressionComponentMask.h
@@ -12,7 +12,9 @@
 UCLASS(collapsecategories, hidecategories=Object, MinimalAPI)
 class UMaterialExpressionComponentMask : public UMaterialExpression
 {
-	GENERATED_UCLASS_BODY()
+	GENERATED_BODY()
+
+public:
 
 	UPROPERTY()
 	FExpressionInput Input;
diff --git a/Engine/Source/Runtime/Engine/Public/Materials/MaterialExpressionCrossProduct.h b/Engine/Source/Runtime/Engine/Public/Materials/MaterialExpressionCrossProduct.h
--- a/Engine/Source/Runtime/Engine/Public/Materials/MaterialExpressionCrossProduct.h
+++ b/Engine/Source/Runtime/Engine/Public/Materials/MaterialExpressionCrossProduct.h
@@ -12,7 +12,9 @@
 UCLASS(MinimalAPI, collapsecategories, hidecategories=Object)
 class UMaterialExpressionCrossProduct : public UMaterialExpression
 {
-	GENERATED_UCLASS_BODY()
+	GENERATED_BODY()
+
+public:
 
 	UPROPERTY()
 	FExpressionInput A;
diff --git a/Engine/Source/Runtime/Engine/Public/Materials/MaterialExpressionDeriveNormalZ.h b/Engine/Source/Runtime/Engine/Public/Materials/MaterialExpressionDeriveNormalZ.h
--- a/Engine/Source/Runtime/Engine/Public/Materials/MaterialExpressionDeriveNormalZ.h
+++ b/Engine/Source/Runtime/Engine/Public/Materials/MaterialExpressionDeriveNormalZ.h
@@ -12,7 +12,9 @@
 UCLASS(collapsecategories, hidecategories=Object, MinimalAPI)
 class UMaterialExpressionDeriveNormalZ : public UMaterialExpression
 {
-	GENERATED_UCLASS_BODY()
+	GENERATED_BODY()
+
+public:
 
 	UPROPERTY()
 	FExpressionInput InXY;
diff --git a/Engine/Source/Runtime/Engine/Public/Materials/MaterialExpressionDistance.h b/Engine/Source/Runtime/Engine/Public/Materials/MaterialExpressionDistance.h
--- a/Engine/Source/Runtime/Engine/Public/Materials/MaterialExpressionDistance.h
+++ b/Engine/Source/Runtime/Engine/Public/Materials/MaterialExpressionDistance.h
@@ -12,7 +12,9 @@
 UCLASS(MinimalAPI)
 class UMaterialExpressionDistance : public UMaterialExpression
 {
-	GENERATED_UCLASS_BODY()
+	GENERATED_BODY()
+
+public:
 
 	UPROPERTY()
 	FExpressionInput A;
diff --git a/Engine/Source/Runtime/Engine/Public/Materials/MaterialExpressionDistanceFieldsRenderingSwitch.h b/Engine/Source/Runtime/Engine/Public/Materials/MaterialExpressionDistanceFieldsRenderingSwitch.h
--- a/Engine/Source/Runtime/Engine/Public/Materials/MaterialExpressionDistanceFieldsRenderingSwitch.h
+++ b/Engine/Source/Runtime/Engine/Public/Materials/MaterialExpressionDistanceFieldsRenderingSwitch.h
@@ -12,7 +12,9 @@
 UCLASS(collapsecategories, hidecategories = Object)
 class UMaterialExpressionDistanceFieldsRenderingSwitch : public UMaterialExpression
 {
-	GENERATED_UCLASS_BODY()
+	GENERATED_BODY()
+
+public:
 
 	/** Used if distance fields renderering is not supported. */
 	UPROPERTY()
diff --git a/Engine/Source/Runtime/Engine/Public/Materials/MaterialExpressionDotProduct.h b/Engine/Source/Runtime/Engine/Public/Materials/MaterialExpressionDotProduct.h
--- a/Engine/Source/Runtime/Engine/Public/Materials/MaterialExpressionDotProduct.h
+++ b/Engine/Source/Runtime/Engine/Public/Materials/MaterialExpressionDotProduct.h
@@ -12,7 +12,9 @@
 UCLASS(MinimalAPI, collapsecategories, hidecategories=Object)
 class UMaterialExpressionDotProduct : public UMaterialExpression
 {
-	GENERATED_UCLASS_BODY()
+	GENERATED_BODY()
+
+public:
 
 	UPROPERTY()
 	FExpressionInput A;
diff --git a/Engine/Source/Runtime/Engine/Public/Materials/MaterialExpressionExponential.h b/Engine/Source/Runtime/Engine/Public/Materials/MaterialExpressionExponential.h
--- a/Engine/Source/Runtime/Engine/Public/Materials/MaterialExpressionExponential.h
+++ b/Engine/Source/Runtime/Engine/Public/Materials/MaterialExpressionExponential.h
@@ -9,7 +9,9 @@
 UCLASS(collapsecategories, hidecategories = Object, MinimalAPI)
 class UMaterialExpressionExponential : public UMaterialExpression
 {
-	GENERATED_UCLASS_BODY()
+	GENERATED_BODY()
+
+public:
 
 	UPROPERTY()
 	FExpressionInput Input;
diff --git a/Engine/Source/Runtime/Engine/Public/Materials/MaterialExpressionExponential2.h b/Engine/Source/Runtime/Engine/Public/Materials/MaterialExpressionExponential2.h
--- a/Engine/Source/Runtime/Engine/Public/Materials/MaterialExpressionExponential2.h
+++ b/Engine/Source/Runtime/Engine/Public/Materials/MaterialExpressionExponential2.h
@@ -9,7 +9,9 @@
 UCLASS(collapsecategories, hidecategories = Object, MinimalAPI)
 class UMaterialExpressionExponential2 : public UMaterialExpression
 {
-	GENERATED_UCLASS_BODY()
+	GENERATED_BODY()
+
+public:
 
 	UPROPERTY()
 	FExpressionInput Input;
diff --git a/Engine/Source/Runtime/Engine/Public/Materials/MaterialExpressionFeatureLevelSwitch.h b/Engine/Source/Runtime/Engine/Public/Materials/MaterialExpressionFeatureLevelSwitch.h
--- a/Engine/Source/Runtime/Engine/Public/Materials/MaterialExpressionFeatureLevelSwitch.h
+++ b/Engine/Source/Runtime/Engine/Public/Materials/MaterialExpressionFeatureLevelSwitch.h
@@ -14,7 +14,9 @@
 UCLASS(collapsecategories, hidecategories=Object, MinimalAPI)
 class UMaterialExpressionFeatureLevelSwitch : public UMaterialExpression
 {
-	GENERATED_UCLASS_BODY()
+	GENERATED_BODY()
+
+public:
 
 	/** Default connection, used when a certain feature level doesn't have an override. */
 	UPROPERTY()
diff --git a/Engine/Source/Runtime/Engine/Public/Materials/MaterialExpressionFloatToUInt.h b/Engine/Source/Runtime/Engine/Public/Materials/MaterialExpressionFloatToUInt.h
--- a/Engine/Source/Runtime/Engine/Public/Materials/MaterialExpressionFloatToUInt.h
+++ b/Engine/Source/Runtime/Engine/Public/Materials/MaterialExpressionFloatToUInt.h
@@ -20,7 +20,9 @@ enum class EFloatToIntMode : uint8
 UCLASS(MinimalAPI, collapsecategories, hidecategories=Object)
 class UMaterialExpressionFloatToUInt : public UMaterialExpression
 {
-	GENERATED_UCLASS_BODY()
+	GENERATED_BODY()
+
+public:
 
 	UPROPERTY()
 	FExpressionInput Input;
@@ -41,7 +43,9 @@ class UMaterialExpressionFloatToUInt : public UMaterialExpression
 UCLASS(MinimalAPI, collapsecategories, hidecategories=Object)
 class UMaterialExpressionUIntToFloat : public UMaterialExpression
 {
-	GENERATED_UCLASS_BODY()
+	GENERATED_BODY()
+
+public:
 
 	UPROPERTY()
 	FExpressionInput Input;
diff --git a/Engine/Source/Runtime/Engine/Public/Materials/MaterialExpressionFloor.h b/Engine/Source/Runtime/Engine/Public/Materials/MaterialExpressionFloor.h
--- a/Engine/Source/Runtime/Engine/Public/Materials/MaterialExpressionFloor.h
+++ b/Engine/Source/Runtime/Engine/Public/Materials/MaterialExpressionFloor.h
@@ -12,7 +12,9 @@
 UCLASS(MinimalAPI, collapsecategories, hidecategories=Object)
 class UMaterialExpressionFloor : public UMaterialExpression
 {
-	GENERATED_UCLASS_BODY()
+	GENERATED_BODY()
+
+public:
 
 	UPROPERTY()
 	FExpressionInput Input;
diff --git a/Engine/Source/Runtime/Engine/Public/Materials/MaterialExpressionFmod.h b/Engine/Source/Runtime/Engine/Public/Materials/MaterialExpressionFmod.h
--- a/Engine/Source/Runtime/Engine/Public/Materials/MaterialExpressionFmod.h
+++ b/Engine/Source/Runtime/Engine/Public/Materials/MaterialExpressionFmod.h
@@ -12,7 +12,9 @@
 UCLASS(MinimalAPI, collapsecategories, hidecategories=Object)
 class UMaterialExpressionFmod : public UMaterialExpression
 {
-	GENERATED_UCLASS_BODY()
+	GENERATED_BODY()
+
+public:
 
 	UPROPERTY()
 	FExpressionInput A;
diff --git a/Engine/Source/Runtime/Engine/Public/Materials/MaterialExpressionFrac.h b/Engine/Source/Runtime/Engine/Public/Materials/MaterialExpressionFrac.h
--- a/Engine/Source/Runtime/Engine/Public/Materials/MaterialExpressionFrac.h
+++ b/Engine/Source/Runtime/Engine/Public/Materials/MaterialExpressionFrac.h
@@ -12,7 +12,9 @@
 UCLASS(MinimalAPI, collapsecategories, hidecategories=Object)
 class UMaterialExpressionFrac : public UMaterialExpression
 {
-	GENERATED_UCLASS_BODY()
+	GENERATED_BODY()
+
+public:
 
 	UPROPERTY()
 	FExpressionInput Input;
diff --git a/Engine/Source/Runtime/Engine/Public/Materials/MaterialExpressionGIReplace.h b/Engine/Source/Runtime/Engine/Public/Materials/MaterialExpressionGIReplace.h
--- a/Engine/Source/Runtime/Engine/Public/Materials/MaterialExpressionGIReplace.h
+++ b/Engine/Source/Runtime/Engine/Public/Materials/MaterialExpressionGIReplace.h
@@ -12,7 +12,9 @@
 UCLASS()
 class UMaterialExpressionGIReplace : public UMaterialExpression
 {
-	GENERATED_UCLASS_BODY()
+	GENERATED_BODY()
+
+public:
 
 	/** Used for direct lighting computations e.g. real-time shaders */
 	UPROPERTY()
diff --git a/Engine/Source/Runtime/Engine/Public/Materials/MaterialExpressionHsvToRgb.h b/Engine/Source/Runtime/Engine/Public/Materials/MaterialExpressionHsvToRgb.h
--- a/Engine/Source/Runtime/Engine/Public/Materials/MaterialExpressionHsvToRgb.h
+++ b/Engine/Source/Runtime/Engine/Public/Materials/MaterialExpressionHsvToRgb.h
@@ -7,7 +7,9 @@
 UCLASS(collapsecategories, hidecategories = Object, MinimalAPI)
 class UMaterialExpressionHsvToRgb : public UMaterialExpression
 {
-	GENERATED_UCLASS_BODY()
+	GENERATED_BODY()
+
+public:
 
 	UPROPERTY()
 	FExpressionInput Input;
diff --git a/Engine/Source/Runtime/Engine/Public/Materials/MaterialExpressionIfThenElse.h b/Engine/Source/Runtime/Engine/Public/Materials/MaterialExpressionIfThenElse.h
--- a/Engine/Source/Runtime/Engine/Public/Materials/MaterialExpressionIfThenElse.h
+++ b/Engine/Source/Runtime/Engine/Public/Materials/MaterialExpressionIfThenElse.h
@@ -13,7 +13,9 @@
 UCLASS(collapsecategories, hidecategories = Object, MinimalAPI)
 class UMaterialExpressionIfThenElse : public UMaterialExpression
 {
-	GENERATED_UCLASS_BODY()
+	GENERATED_BODY()
+
+public:
 	
 	UPROPERTY()
 	FExpressionInput Condition;
diff --git a/Engine/Source/Runtime/Engine/Public/Materials/MaterialExpressionLength.h b/Engine/Source/Runtime/Engine/Public/Materials/MaterialExpressionLength.h
--- a/Engine/Source/Runtime/Engine/Public/Materials/MaterialExpressionLength.h
+++ b/Engine/Source/Runtime/Engine/Public/Materials/MaterialExpressionLength.h
@@ -7,7 +7,9 @@
 UCLASS(collapsecategories, hidecategories = Object, MinimalAPI)
 class UMaterialExpressionLength : public UMaterialExpression
 {
-	GENERATED_UCLASS_BODY()
+	GENERATED_BODY()
+
+public:
 
 	UPROPERTY()
 	FExpressionInput Input;
diff --git a/Engine/Source/Runtime/Engine/Public/Materials/MaterialExpressionLightmassReplace.h b/Engine/Source/Runtime/Engine/Public/Materials/MaterialExpressionLightmassReplace.h
--- a/Engine/Source/Runtime/Engine/Public/Materials/MaterialExpressionLightmassReplace.h
+++ b/Engine/Source/Runtime/Engine/Public/Materials/MaterialExpressionLightmassReplace.h
@@ -12,7 +12,9 @@
 UCLASS()
 class UMaterialExpressionLightmassReplace : public UMaterialExpression
 {
-	GENERATED_UCLASS_BODY()
+	GENERATED_BODY()
+
+public:
 
 	UPROPERTY()
 	FExpressionInput Realtime;
diff --git a/Engine/Source/Runtime/Engine/Public/Materials/MaterialExpressionLogarithm.h b/Engine/Source/Runtime/Engine/Public/Materials/MaterialExpressionLogarithm.h
--- a/Engine/Source/Runtime/Engine/Public/Materials/MaterialExpressionLogarithm.h
+++ b/Engine/Source/Runtime/Engine/Public/Materials/MaterialExpressionLogarithm.h
@@ -9,7 +9,9 @@
 UCLASS(collapsecategories, hidecategories = Object, MinimalAPI)
 class UMaterialExpressionLogarithm : public UMaterialExpression
 {
-	GENERATED_UCLASS_BODY()
+	GENERATED_BODY()
+
+public:
 
 	UPROPERTY()
 	FExpressionInput Input;
diff --git a/Engine/Source/Runtime/Engine/Public/Materials/MaterialExpressionLogarithm10.h b/Engine/Source/Runtime/Engine/Public/Materials/MaterialExpressionLogarithm10.h
--- a/Engine/Source/Runtime/Engine/Public/Materials/MaterialExpressionLogarithm10.h
+++ b/Engine/Source/Runtime/Engine/Public/Materials/MaterialExpressionLogarithm10.h
@@ -12,7 +12,9 @@
 UCLASS(collapsecategories, hidecategories=Object, MinimalAPI)
 class UMaterialExpressionLogarithm10 : public UMaterialExpression
 {
-	GENERATED_UCLASS_BODY()
+	GENERATED_BODY()
+
+public:
 	
 	UPROPERTY()
 	FExpressionInput X;
diff --git a/Engine/Source/Runtime/Engine/Public/Materials/MaterialExpressionLogarithm2.h b/Engine/Source/Runtime/Engine/Public/Materials/MaterialExpressionLogarithm2.h
--- a/Engine/Source/Runtime/Engine/Public/Materials/MaterialExpressionLogarithm2.h
+++ b/Engine/Source/Runtime/Engine/Public/Materials/MaterialExpressionLogarithm2.h
@@ -12,7 +12,9 @@
 UCLASS(collapsecategories, hidecategories=Object, MinimalAPI)
 class UMaterialExpressionLogarithm2 : public UMaterialExpression
 {
-	GENERATED_UCLASS_BODY()
+	GENERATED_BODY()
+
+public:
 	
 	UPROPERTY()
 	FExpressionInput X;
diff --git a/Engine/Source/Runtime/Engine/Public/Materials/MaterialExpressionMakeMaterialAttributes.h b/Engine/Source/Runtime/Engine/Public/Materials/MaterialExpressionMakeMaterialAttributes.h
--- a/Engine/Source/Runtime/Engine/Public/Materials/MaterialExpressionMakeMaterialAttributes.h
+++ b/Engine/Source/Runtime/Engine/Public/Materials/MaterialExpressionMakeMaterialAttributes.h
@@ -12,7 +12,9 @@
 UCLASS(collapsecategories, hidecategories = Object, MinimalAPI)
 class UMaterialExpressionMakeMaterialAttributes : public UMaterialExpression
 {
-	GENERATED_UCLASS_BODY()
+	GENERATED_BODY()
+
+public:
 
 	UPROPERTY()
 	FExpressionInput BaseColor;
diff --git a/Engine/Source/Runtime/Engine/Public/Materials/MaterialExpressionMapARPassthroughCameraUV.h b/Engine/Source/Runtime/Engine/Public/Materials/MaterialExpressionMapARPassthroughCameraUV.h
--- a/Engine/Source/Runtime/Engine/Public/Materials/MaterialExpressionMapARPassthroughCameraUV.h
+++ b/Engine/Source/Runtime/Engine/Public/Materials/MaterialExpressionMapARPassthroughCameraUV.h
@@ -15,7 +15,9 @@
 UCLASS(collapsecategories, hidecategories = Object)
 class UMaterialExpressionMapARPassthroughCameraUV : public UMaterialExpression
 {
-	GENERATED_UCLASS_BODY()
+	GENERATED_BODY()
+
+public:
 
 	UPROPERTY(meta = (RequiredInput = "true", ToolTip = "Input UV to map to AR Camera UV"))
 	FExpressionInput Coordinates;
diff --git a/Engine/Source/Runtime/Engine/Public/Materials/MaterialExpressionMaterialProxyReplace.h b/Engine/Source/Runtime/Engine/Public/Materials/MaterialExpressionMaterialProxyReplace.h
--- a/Engine/Source/Runtime/Engine/Public/Materials/MaterialExpressionMaterialProxyReplace.h
+++ b/Engine/Source/Runtime/Engine/Public/Materials/MaterialExpressionMaterialProxyReplace.h
@@ -13,7 +13,9 @@
 UCLASS()
 class UMaterialExpressionMaterialProxyReplace : public UMaterialExpression
 {
-	GENERATED_UCLASS_BODY()
+	GENERATED_BODY()
+
+public:
 
 	UPROPERTY()
 	FExpressionInput Realtime;
diff --git a/Engine/Source/Runtime/Engine/Public/Materials/MaterialExpressionModulo.h b/Engine/Source/Runtime/Engine/Public/Materials/MaterialExpressionModulo.h
--- a/Engine/Source/Runtime/Engine/Public/Materials/MaterialExpressionModulo.h
+++ b/Engine/Source/Runtime/Engine/Public/Materials/MaterialExpressionModulo.h
@@ -8,7 +8,9 @@
 UCLASS(MinimalAPI, collapsecategories, hidecategories=Object)
 class UMaterialExpressionModulo : public UMaterialExpression
 {
-	GENERATED_UCLASS_BODY()
+	GENERATED_BODY()
+
+public:
 
 	UPROPERTY()
 	FExpressionInput A;
diff --git a/Engine/Source/Runtime/Engine/Public/Materials/MaterialExpressionNormalize.h b/Engine/Source/Runtime/Engine/Public/Materials/MaterialExpressionNormalize.h
--- a/Engine/Source/Runtime/Engine/Public/Materials/MaterialExpressionNormalize.h
+++ b/Engine/Source/Runtime/Engine/Public/Materials/MaterialExpressionNormalize.h
@@ -12,7 +12,9 @@
 UCLASS(MinimalAPI)
 class UMaterialExpressionNormalize : public UMaterialExpression
 {
-	GENERATED_UCLASS_BODY()
+	GENERATED_BODY()
+
+public:
 
 	UPROPERTY()
 	FExpressionInput VectorInput;
diff --git a/Engine/Source/Runtime/Engine/Public/Materials/MaterialExpressionOneMinus.h b/Engine/Source/Runtime/Engine/Public/Materials/MaterialExpressionOneMinus.h
--- a/Engine/Source/Runtime/Engine/Public/Materials/MaterialExpressionOneMinus.h
+++ b/Engine/Source/Runtime/Engine/Public/Materials/MaterialExpressionOneMinus.h
@@ -12,7 +12,9 @@
 UCLASS(collapsecategories, hidecategories=Object, MinimalAPI)
 class UMaterialExpressionOneMinus : public UMaterialExpression
 {
-	GENERATED_UCLASS_BODY()
+	GENERATED_BODY()
+
+public:
 
 	UPROPERTY()
 	FExpressionInput Input;
diff --git a/Engine/Source/Runtime/Engine/Public/Materials/MaterialExpressionPreviousFrameSwitch.h b/Engine/Source/Runtime/Engine/Public/Materials/MaterialExpressionPreviousFrameSwitch.h
--- a/Engine/Source/Runtime/Engine/Public/Materials/MaterialExpressionPreviousFrameSwitch.h
+++ b/Engine/Source/Runtime/Engine/Public/Materials/MaterialExpressionPreviousFrameSwitch.h
@@ -9,7 +9,9 @@
 UCLASS(collapsecategories, hidecategories=Object, MinimalAPI)
 class UMaterialExpressionPreviousFrameSwitch : public UMaterialExpression
 {
-	GENERATED_UCLASS_BODY()
+	GENERATED_BODY()
+
+public:
 
 	UPROPERTY()
 	FExpressionInput CurrentFrame;
diff --git a/Engine/Source/Runtime/Engine/Public/Materials/MaterialExpressionQualitySwitch.h b/Engine/Source/Runtime/Engine/Public/Materials/MaterialExpressionQualitySwitch.h
--- a/Engine/Source/Runtime/Engine/Public/Materials/MaterialExpressionQualitySwitch.h
+++ b/Engine/Source/Runtime/Engine/Public/Materials/MaterialExpressionQualitySwitch.h
@@ -13,7 +13,9 @@
 UCLASS(collapsecategories, hidecategories=Object, MinimalAPI)
 class UMaterialExpressionQualitySwitch : public UMaterialExpression
 {
-	GENERATED_UCLASS_BODY()
+	GENERATED_BODY()
+
+public:
 
 	/** Default connection, used when a specific quality level input is missing. */
 	UPROPERTY()
diff --git a/Engine/Source/Runtime/Engine/Public/Materials/MaterialExpressionReflectionCapturePassSwitch.h b/Engine/Source/Runtime/Engine/Public/Materials/MaterialExpressionReflectionCapturePassSwitch.h
--- a/Engine/Source/Runtime/Engine/Public/Materials/MaterialExpressionReflectionCapturePassSwitch.h
+++ b/Engine/Source/Runtime/Engine/Public/Materials/MaterialExpressionReflectionCapturePassSwitch.h
@@ -12,7 +12,9 @@
 UCLASS(meta = (DisplayName = "ReflectionCapturePassSwitch"))
 class UMaterialExpressionReflectionCapturePassSwitch : public UMaterialExpression
 {
-	GENERATED_UCLASS_BODY()
+	GENERATED_BODY()
+
+public:
 
 	UPROPERTY(meta = (RequiredInput = "true", ToolTip = "Input will be used when rendering into non-reflection passes."))
 	FExpressionInput Default;
diff --git a/Engine/Source/Runtime/Engine/Public/Materials/MaterialExpressionRgbToHsv.h b/Engine/Source/Runtime/Engine/Public/Materials/MaterialExpressionRgbToHsv.h
--- a/Engine/Source/Runtime/Engine/Public/Materials/MaterialExpressionRgbToHsv.h
+++ b/Engine/Source/Runtime/Engine/Public/Materials/MaterialExpressionRgbToHsv.h
@@ -7,7 +7,9 @@
 UCLASS(collapsecategories, hidecategories = Object, MinimalAPI)
 class UMaterialExpressionRgbToHsv : public UMaterialExpression
 {
-	GENERATED_UCLASS_BODY()
+	GENERATED_BODY()
+
+public:
 
 	UPROPERTY()
 	FExpressionInput Input;
diff --git a/Engine/Source/Runtime/Engine/Public/Materials/MaterialExpressionRound.h b/Engine/Source/Runtime/Engine/Public/Materials/MaterialExpressionRound.h
--- a/Engine/Source/Runtime/Engine/Public/Materials/MaterialExpressionRound.h
+++ b/Engine/Source/Runtime/Engine/Public/Materials/MaterialExpressionRound.h
@@ -11,7 +11,9 @@
 UCLASS(collapsecategories, hidecategories=Object)
 class UMaterialExpressionRound : public UMaterialExpression
 {
-	GENERATED_UCLASS_BODY()
+	GENERATED_BODY()
+
+public:
 
 	UPROPERTY()
 	FExpressionInput Input;
diff --git a/Engine/Source/Runtime/Engine/Public/Materials/MaterialExpressionRuntimeVirtualTextureCustomData.h b/Engine/Source/Runtime/Engine/Public/Materials/MaterialExpressionRuntimeVirtualTextureCustomData.h
--- a/Engine/Source/Runtime/Engine/Public/Materials/MaterialExpressionRuntimeVirtualTextureCustomData.h
+++ b/Engine/Source/Runtime/Engine/Public/Materials/MaterialExpressionRuntimeVirtualTextureCustomData.h
@@ -9,7 +9,9 @@
 UCLASS(collapsecategories, hidecategories=Object)
 class UMaterialExpressionRuntimeVirtualTextureCustomData : public UMaterialExpression
 {
-	GENERATED_UCLASS_BODY()
+	GENERATED_BODY()
+
+public:
 
 	//~ Begin UMaterialExpression Interface
 #if WITH_EDITOR
diff --git a/Engine/Source/Runtime/Engine/Public/Materials/MaterialExpressionRuntimeVirtualTextureReplace.h b/Engine/Source/Runtime/Engine/Public/Materials/MaterialExpressionRuntimeVirtualTextureReplace.h
--- a/Engine/Source/Runtime/Engine/Public/Materials/MaterialExpressionRuntimeVirtualTextureReplace.h
+++ b/Engine/Source/Runtime/Engine/Public/Materials/MaterialExpressionRuntimeVirtualTextureReplace.h
@@ -12,7 +12,9 @@
 UCLASS(collapsecategories, hidecategories = Object)
 class UMaterialExpressionRuntimeVirtualTextureReplace : public UMaterialExpression
 {
-	GENERATED_UCLASS_BODY()
+	GENERATED_BODY()
+
+public:
 
 	/** Used by the default rendering passes. */
 	UPROPERTY()
diff --git a/Engine/Source/Runtime/Engine/Public/Materials/MaterialExpressionSRGBColorToWorkingColorSpace.h b/Engine/Source/Runtime/Engine/Public/Materials/MaterialExpressionSRGBColorToWorkingColorSpace.h
--- a/Engine/Source/Runtime/Engine/Public/Materials/MaterialExpressionSRGBColorToWorkingColorSpace.h
+++ b/Engine/Source/Runtime/Engine/Public/Materials/MaterialExpressionSRGBColorToWorkingColorSpace.h
@@ -12,7 +12,9 @@
 UCLASS(MinimalAPI, collapsecategories, hidecategories=Object)
 class UMaterialExpressionSRGBColorToWorkingColorSpace : public UMaterialExpression
 {
-	GENERATED_UCLASS_BODY()
+	GENERATED_BODY()
+
+public:
 
 	UPROPERTY()
 	FExpressionInput Input;
diff --git a/Engine/Source/Runtime/Engine/Public/Materials/MaterialExpressionSaturate.h b/Engine/Source/Runtime/Engine/Public/Materials/MaterialExpressionSaturate.h
--- a/Engine/Source/Runtime/Engine/Public/Materials/MaterialExpressionSaturate.h
+++ b/Engine/Source/Runtime/Engine/Public/Materials/MaterialExpressionSaturate.h
@@ -11,7 +11,9 @@
 UCLASS(MinimalAPI, collapsecategories, hidecategories=Object)
 class UMaterialExpressionSaturate : public UMaterialExpression
 {
-	GENERATED_UCLASS_BODY()
+	GENERATED_BODY()
+
+public:
 
 	UPROPERTY()
 	FExpressionInput Input;
diff --git a/Engine/Source/Runtime/Engine/Public/Materials/MaterialExpressionShaderStageSwitch.h b/Engine/Source/Runtime/Engine/Public/Materials/MaterialExpressionShaderStageSwitch.h
--- a/Engine/Source/Runtime/Engine/Public/Materials/MaterialExpressionShaderStageSwitch.h
+++ b/Engine/Source/Runtime/Engine/Public/Materials/MaterialExpressionShaderStageSwitch.h
@@ -13,7 +13,9 @@
 UCLASS(collapsecategories, hidecategories = Object, MinimalAPI)
 class UMaterialExpressionShaderStageSwitch : public UMaterialExpression
 {
-	GENERATED_UCLASS_BODY()
+	GENERATED_BODY()
+
+public:
 
 	static inline bool ShouldUsePixelShaderInput(EShaderFrequency InShaderFrequency)
 	{
diff --git a/Engine/Source/Runtime/Engine/Public/Materials/MaterialExpressionShadingPathSwitch.h b/Engine/Source/Runtime/Engine/Public/Materials/MaterialExpressionShadingPathSwitch.h
--- a/Engine/Source/Runtime/Engine/Public/Materials/MaterialExpressionShadingPathSwitch.h
+++ b/Engine/Source/Runtime/Engine/Public/Materials/MaterialExpressionShadingPathSwitch.h
@@ -14,7 +14,9 @@
 UCLASS(collapsecategories, hidecategories=Object, MinimalAPI)
 class UMaterialExpressionShadingPathSwitch : public UMaterialExpression
 {
-	GENERATED_UCLASS_BODY()
+	GENERATED_BODY()
+
+public:
 
 	/** Default connection, used when a specific shading type input is missing. */
 	UPROPERTY()
diff --git a/Engine/Source/Runtime/Engine/Public/Materials/MaterialExpressionShadowReplace.h b/Engine/Source/Runtime/Engine/Public/Materials/MaterialExpressionShadowReplace.h
--- a/Engine/Source/Runtime/Engine/Public/Materials/MaterialExpressionShadowReplace.h
+++ b/Engine/Source/Runtime/Engine/Public/Materials/MaterialExpressionShadowReplace.h
@@ -12,7 +12,9 @@
 UCLASS(meta = (DisplayName = "ShadowPassSwitch"))
 class UMaterialExpressionShadowReplace : public UMaterialExpression
 {
-	GENERATED_UCLASS_BODY()
+	GENERATED_BODY()
+
+public:
 
 	UPROPERTY(meta = (RequiredInput = "true", ToolTip = "Input will be used when rendering into non-ShadowMap passes"))
 	FExpressionInput Default;
diff --git a/Engine/Source/Runtime/Engine/Public/Materials/MaterialExpressionSign.h b/Engine/Source/Runtime/Engine/Public/Materials/MaterialExpressionSign.h
--- a/Engine/Source/Runtime/Engine/Public/Materials/MaterialExpressionSign.h
+++ b/Engine/Source/Runtime/Engine/Public/Materials/MaterialExpressionSign.h
@@ -11,7 +11,9 @@
 UCLASS(MinimalAPI, collapsecategories, hidecategories=Object)
 class UMaterialExpressionSign : public UMaterialExpression
 {
-	GENERATED_UCLASS_BODY()
+	GENERATED_BODY()
+
+public:
 
 	UPROPERTY()
 	FExpressionInput Input;
diff --git a/Engine/Source/Runtime/Engine/Public/Materials/MaterialExpressionSkyAtmosphereLightDirection.h b/Engine/Source/Runtime/Engine/Public/Materials/MaterialExpressionSkyAtmosphereLightDirection.h
--- a/Engine/Source/Runtime/Engine/Public/Materials/MaterialExpressionSkyAtmosphereLightDirection.h
+++ b/Engine/Source/Runtime/Engine/Public/Materials/MaterialExpressionSkyAtmosphereLightDirection.h
@@ -11,7 +11,9 @@
 UCLASS()
 class UMaterialExpressionSkyAtmosphereLightDirection : public UMaterialExpression
 {
-	GENERATED_UCLASS_BODY()
+	GENERATED_BODY()
+
+public:
 
 	/** Index of the atmosphere light to sample. */
 	UPROPERTY(EditAnywhere, BlueprintReadWrite, Category = MaterialExpressionTextureCoordinate, meta = (UIMin = "0", UIMax = "1", ClampMin = "0", ClampMax = "1", ShowAsInputPin = "Primary"))
diff --git a/Engine/Source/Runtime/Engine/Public/Materials/MaterialExpressionSkyAtmosphereLightIlluminance.h b/Engine/Source/Runtime/Engine/Public/Materials/MaterialExpressionSkyAtmosphereLightIlluminance.h
--- a/Engine/Source/Runtime/Engine/Public/Materials/MaterialExpressionSkyAtmosphereLightIlluminance.h
+++ b/Engine/Source/Runtime/Engine/Public/Materials/MaterialExpressionSkyAtmosphereLightIlluminance.h
@@ -11,7 +11,9 @@
 UCLASS()
 class UMaterialExpressionSkyAtmosphereLightIlluminance : public UMaterialExpression
 {
-	GENERATED_UCLASS_BODY()
+	GENERATED_BODY()
+
+public:
 
 	/** Index of the atmosphere light to sample. */
 	UPROPERTY(EditAnywhere, BlueprintReadWrite, Category = MaterialExpressionTextureCoordinate, meta = (UIMin = "0", UIMax = "1", ClampMin = "0", ClampMax = "1", ShowAsInputPin = "Primary"))
@@ -41,7 +43,9 @@ class UMaterialExpressionSkyAtmosphereLightIlluminance : public UMaterialExpress
 UCLASS()
 class UMaterialExpressionSkyAtmosphereLightIlluminanceOnGround : public UMaterialExpression
 {
-	GENERATED_UCLASS_BODY()
+	GENERATED_BODY()
+
+public:
 
 	/** Index of the atmosphere light to sample. */
 	UPROPERTY(EditAnywhere, BlueprintReadWrite, Category = MaterialExpressionTextureCoordinate, meta = (UIMin = "0", UIMax = "1", ClampMin = "0", ClampMax = "1", ShowAsInputPin = "Primary"))
@@ -62,7 +66,9 @@ class UMaterialExpressionSkyAtmosphereLightIlluminanceOnGround : public UMateria
 UCLASS()
 class UMaterialExpressionSkyAtmosphereLightDiskLuminance : public UMaterialExpression
 {
-	GENERATED_UCLASS_BODY()
+	GENERATED_BODY()
+
+public:
 
 	/** Index of the atmosphere light to sample. */
 	UPROPERTY(EditAnywhere, BlueprintReadWrite, Category = MaterialExpressionTextureCoordinate, meta = (UIMin = "0", UIMax = "1", ClampMin = "0", ClampMax = "1", ShowAsInputPin = "Primary"))
@@ -85,7 +91,9 @@ class UMaterialExpressionSkyAtmosphereLightDiskLuminance : public UMaterialExpre
 UCLASS()
 class UMaterialExpressionSkyAtmosphereAerialPerspective : public UMaterialExpression
 {
-	GENERATED_UCLASS_BODY()
+	GENERATED_BODY()
+
+public:
 
 	/** World position of the sample. If not specified, the pixel world position will be used. Larger distance will result in more fog. Please make sure .SkyAtmosphere.AerialPerspectiveLUT.Depth is set far enough to have fog data.
 		If you are scaling the sky dome pixel world position, make sure it is centered around the origin.*/
@@ -112,7 +120,9 @@ class UMaterialExpressionSkyAtmosphereAerialPerspective : public UMaterialExpres
 UCLASS()
 class UMaterialExpressionSkyAtmosphereDistantLightScatteredLuminance : public UMaterialExpression
 {
-	GENERATED_UCLASS_BODY()
+	GENERATED_BODY()
+
+public:
 
 	//~ Begin UMaterialExpression Interface
 #if WITH_EDITOR
diff --git a/Engine/Source/Runtime/Engine/Public/Materials/MaterialExpressionSkyAtmosphereViewLuminance.h b/Engine/Source/Runtime/Engine/Public/Materials/MaterialExpressionSkyAtmosphereViewLuminance.h
--- a/Engine/Source/Runtime/Engine/Public/Materials/MaterialExpressionSkyAtmosphereViewLuminance.h
+++ b/Engine/Source/Runtime/Engine/Public/Materials/MaterialExpressionSkyAtmosphereViewLuminance.h
@@ -11,7 +11,9 @@
 UCLASS()
 class UMaterialExpressionSkyAtmosphereViewLuminance : public UMaterialExpression
 {
-	GENERATED_UCLASS_BODY()
+	GENERATED_BODY()
+
+public:
 
 	//~ Begin UMaterialExpression Interface
 #if WITH_EDITOR
diff --git a/Engine/Source/Runtime/Engine/Public/Materials/MaterialExpressionSkyLightEnvMapSample.h b/Engine/Source/Runtime/Engine/Public/Materials/MaterialExpressionSkyLightEnvMapSample.h
--- a/Engine/Source/Runtime/Engine/Public/Materials/MaterialExpressionSkyLightEnvMapSample.h
+++ b/Engine/Source/Runtime/Engine/Public/Materials/MaterialExpressionSkyLightEnvMapSample.h
@@ -11,7 +11,9 @@
 UCLASS()
 class UMaterialExpressionSkyLightEnvMapSample : public UMaterialExpression
 {
-	GENERATED_UCLASS_BODY()
+	GENERATED_BODY()
+
+public:
 
 	UPROPERTY(meta = (RequiredInput = "true", ToolTip = "The direction to sample the cubemap from"))
 	FExpressionInput Direction;
diff --git a/Engine/Source/Runtime/Engine/Public/Materials/MaterialExpressionSquareRoot.h b/Engine/Source/Runtime/Engine/Public/Materials/MaterialExpressionSquareRoot.h
--- a/Engine/Source/Runtime/Engine/Public/Materials/MaterialExpressionSquareRoot.h
+++ b/Engine/Source/Runtime/Engine/Public/Materials/MaterialExpressionSquareRoot.h
@@ -12,7 +12,9 @@
 UCLASS(MinimalAPI, collapsecategories, hidecategories=Object)
 class UMaterialExpressionSquareRoot : public UMaterialExpression
 {
-	GENERATED_UCLASS_BODY()
+	GENERATED_BODY()
+
+public:
 
 	UPROPERTY()
 	FExpressionInput Input;
diff --git a/Engine/Source/Runtime/Engine/Public/Materials/MaterialExpressionStaticComponentMaskParameter.h b/Engine/Source/Runtime/Engine/Public/Materials/MaterialExpressionStaticComponentMaskParameter.h
--- a/Engine/Source/Runtime/Engine/Public/Materials/MaterialExpressionStaticComponentMaskParameter.h
+++ b/Engine/Source/Runtime/Engine/Public/Materials/MaterialExpressionStaticComponentMaskParameter.h
@@ -12,7 +12,9 @@
 UCLASS(collapsecategories, hidecategories=Object, MinimalAPI)
 class UMaterialExpressionStaticComponentMaskParameter : public UMaterialExpressionParameter
 {
-	GENERATED_UCLASS_BODY()
+	GENERATED_BODY()
+
+public:
 
 	UPROPERTY()
 	FExpressionInput Input;
diff --git a/Engine/Source/Runtime/Engine/Public/Materials/MaterialExpressionStaticSwitch.h b/Engine/Source/Runtime/Engine/Public/Materials/MaterialExpressionStaticSwitch.h
--- a/Engine/Source/Runtime/Engine/Public/Materials/MaterialExpressionStaticSwitch.h
+++ b/Engine/Source/Runtime/Engine/Public/Materials/MaterialExpressionStaticSwitch.h
@@ -13,7 +13,9 @@
 UCLASS(collapsecategories, hidecategories=Object, MinimalAPI)
 class UMaterialExpressionStaticSwitch : public UMaterialExpression
 {
-	GENERATED_UCLASS_BODY()
+	GENERATED_BODY()
+
+public:
 
 	UPROPERTY(EditAnywhere, Category=MaterialExpressionStaticSwitch, meta=(OverridingInputProperty = "Value"))
 	uint32 DefaultValue:1;
diff --git a/Engine/Source/Runtime/Engine/Public/Materials/MaterialExpressionTextureSampleParameter2DArray.h b/Engine/Source/Runtime/Engine/Public/Materials/MaterialExpressionTextureSampleParameter2DArray.h
--- a/Engine/Source/Runtime/Engine/Public/Materials/MaterialExpressionTextureSampleParameter2DArray.h
+++ b/Engine/Source/Runtime/Engine/Public/Materials/MaterialExpressionTextureSampleParameter2DArray.h
@@ -9,7 +9,9 @@
 UCLASS(collapsecategories, hidecategories = Object, MinimalAPI)
 class UMaterialExpressionTextureSampleParameter2DArray : public UMaterialExpressionTextureSampleParameter
 {
-	GENERATED_UCLASS_BODY()
+	GENERATED_BODY()
+
+public:
 
 #if WITH_EDITOR
 	// Begin UMaterialExpression Interface
diff --git a/Engine/Source/Runtime/Engine/Public/Materials/MaterialExpressionTextureSampleParameterCubeArray.h b/Engine/Source/Runtime/Engine/Public/Materials/MaterialExpressionTextureSampleParameterCubeArray.h
--- a/Engine/Source/Runtime/Engine/Public/Materials/MaterialExpressionTextureSampleParameterCubeArray.h
+++ b/Engine/Source/Runtime/Engine/Public/Materials/MaterialExpressionTextureSampleParameterCubeArray.h
@@ -13,7 +13,9 @@ class UTexture;
 UCLASS(collapsecategories, hidecategories = Object, MinimalAPI)
 class UMaterialExpressionTextureSampleParameterCubeArray : public UMaterialExpressionTextureSampleParameter
 {
-	GENERATED_UCLASS_BODY()
+	GENERATED_BODY()
+
+public:
 
 #if WITH_EDITOR
 	//~ Begin UMaterialExpression Interface
diff --git a/Engine/Source/Runtime/Engine/Public/Materials/MaterialExpressionTruncate.h b/Engine/Source/Runtime/Engine/Public/Materials/MaterialExpressionTruncate.h
--- a/Engine/Source/Runtime/Engine/Public/Materials/MaterialExpressionTruncate.h
+++ b/Engine/Source/Runtime/Engine/Public/Materials/MaterialExpressionTruncate.h
@@ -11,7 +11,9 @@
 UCLASS(collapsecategories, hidecategories=Object)
 class UMaterialExpressionTruncate : public UMaterialExpression
 {
-	GENERATED_UCLASS_BODY()
+	GENERATED_BODY()
+
+public:
 
 	UPROPERTY()
 	FExpressionInput Input;
diff --git a/Engine/Source/Runtime/Engine/Public/Materials/MaterialExpressionTruncateLWC.h b/Engine/Source/Runtime/Engine/Public/Materials/MaterialExpressionTruncateLWC.h
--- a/Engine/Source/Runtime/Engine/Public/Materials/MaterialExpressionTruncateLWC.h
+++ b/Engine/Source/Runtime/Engine/Public/Materials/MaterialExpressionTruncateLWC.h
@@ -12,7 +12,9 @@
 UCLASS(MinimalAPI)
 class UMaterialExpressionTruncateLWC : public UMaterialExpression
 {
-	GENERATED_UCLASS_BODY()
+	GENERATED_BODY()
+
+public:
 
 	UPROPERTY()
 	FExpressionInput Input;
diff --git a/Engine/Source/Runtime/Engine/Public/Materials/MaterialExpressionVirtualTextureFeatureSwitch.h b/Engine/Source/Runtime/Engine/Public/Materials/MaterialExpressionVirtualTextureFeatureSwitch.h
--- a/Engine/Source/Runtime/Engine/Public/Materials/MaterialExpressionVirtualTextureFeatureSwitch.h
+++ b/Engine/Source/Runtime/Engine/Public/Materials/MaterialExpressionVirtualTextureFeatureSwitch.h
@@ -13,7 +13,9 @@
 UCLASS(collapsecategories, hidecategories = Object)
 class UMaterialExpressionVirtualTextureFeatureSwitch : public UMaterialExpression
 {
-	GENERATED_UCLASS_BODY()
+	GENERATED_BODY()
+
+public:
 
 	/** Used if virtual texture is not supported. */
 	UPROPERTY()
