Fixed data race in FParallelAnimationEvaluationStats calculations.

Total, Min, and Max times could be incorrect if two threads call AddTiming simultaneously.
Logic still assumes that AddTiming won't be called during OnEndFrame.

#tests win64 test client, multi platform replayruns

diff --git a/Engine/Source/Runtime/Engine/Private/Components/SkeletalMeshComponent.cpp b/Engine/Source/Runtime/Engine/Private/Components/SkeletalMeshComponent.cpp
--- a/Engine/Source/Runtime/Engine/Private/Components/SkeletalMeshComponent.cpp
+++ b/Engine/Source/Runtime/Engine/Private/Components/SkeletalMeshComponent.cpp
@@ -131,20 +131,20 @@ struct FParallelAnimationEvaluationStats
 
 	void AddTiming(float Value)
 	{
-		TotalTime.Store(TotalTime.Load() + Value);
-		MinTime.Store(FMath::Min(MinTime.Load(), Value));
-		MaxTime.Store(FMath::Max(MaxTime.Load(), Value));
-		NumberOfTasks.IncrementExchange();
+		AtomicFloatAddRelaxed(TotalTime, Value);
+		AtomicFloatMinRelaxed(MinTime, Value);
+		AtomicFloatMaxRelaxed(MaxTime, Value);
+		NumberOfTasks.fetch_add(1, std::memory_order_release);
 	}
 
 private:
 	FParallelAnimationEvaluationStats()
 	{
 		EndOfFrameHandle = FCoreDelegates::OnEndFrame.AddRaw(this, &FParallelAnimationEvaluationStats::OnEndFrame);
-		MinTime = TNumericLimits<float>::Max();
-		MaxTime = 0.f;
-		TotalTime = 0.f;
-		NumberOfTasks = 0;
+		MinTime.store(TNumericLimits<float>::Max(), std::memory_order_relaxed);
+		MaxTime.store(0.f, std::memory_order_relaxed);
+		TotalTime.store(0.f, std::memory_order_relaxed);
+		NumberOfTasks.store(0, std::memory_order_release);
 	}
 	
 	~FParallelAnimationEvaluationStats()
@@ -155,26 +155,62 @@ private:
 	void OnEndFrame()
 	{
 		// Only set stats when a task has actually run
-		if (NumberOfTasks > 0)
+		const int32 NumTasks = NumberOfTasks.load(std::memory_order_acquire);
+		if (NumTasks > 0)
 		{	
-		    CSV_CUSTOM_STAT(AnimationParallelEvaluation, TotalTaskTime, TotalTime, ECsvCustomStatOp::Set);	
-		    CSV_CUSTOM_STAT(AnimationParallelEvaluation, AverageTaskTime, NumberOfTasks == 0 || FMath::IsNearlyZero(TotalTime) ? 0.f : TotalTime / static_cast<float>(NumberOfTasks), ECsvCustomStatOp::Set);				
-		    CSV_CUSTOM_STAT(AnimationParallelEvaluation, NumberOfTasks, NumberOfTasks, ECsvCustomStatOp::Set);
+			const float Min = MinTime.load(std::memory_order_relaxed);
+			const float Max = MaxTime.load(std::memory_order_relaxed);
+			const float Total = TotalTime.load(std::memory_order_relaxed);
+			const float AvgTime = FMath::IsNearlyZero(Total) ? 0.f : Total / NumTasks;
+
+		    CSV_CUSTOM_STAT(AnimationParallelEvaluation, TotalTaskTime, Total, ECsvCustomStatOp::Set);
+		    CSV_CUSTOM_STAT(AnimationParallelEvaluation, AverageTaskTime, AvgTime, ECsvCustomStatOp::Set);
+		    CSV_CUSTOM_STAT(AnimationParallelEvaluation, NumberOfTasks, NumTasks, ECsvCustomStatOp::Set);
 				    
-		    CSV_CUSTOM_STAT(AnimationParallelEvaluation, MinTaskTime, MinTime, ECsvCustomStatOp::Set);
-		    CSV_CUSTOM_STAT(AnimationParallelEvaluation, MaxTaskTime, MaxTime, ECsvCustomStatOp::Set);
+		    CSV_CUSTOM_STAT(AnimationParallelEvaluation, MinTaskTime, Min, ECsvCustomStatOp::Set);
+		    CSV_CUSTOM_STAT(AnimationParallelEvaluation, MaxTaskTime, Max, ECsvCustomStatOp::Set);
 		}
 
-		MinTime = TNumericLimits<float>::Max();
-		MaxTime = 0.f;
-		TotalTime = 0.f;
-		NumberOfTasks = 0;
+		MinTime.store(TNumericLimits<float>::Max(), std::memory_order_relaxed);
+		MaxTime.store(0.f, std::memory_order_relaxed);
+		TotalTime.store(0.f, std::memory_order_relaxed);
+		NumberOfTasks.store(0, std::memory_order_release);
 	}
 
-	TAtomic<float> MinTime;
-	TAtomic<float> MaxTime;
-	TAtomic<float> TotalTime;
-	TAtomic<int32> NumberOfTasks;
+	static void AtomicFloatAddRelaxed(std::atomic<float>& Atomic, float Value)
+	{
+		float OldAtomic = Atomic.load(std::memory_order_relaxed);
+		float NewAtomic;
+		do
+		{
+			NewAtomic = OldAtomic + Value;
+		} while (!Atomic.compare_exchange_weak(OldAtomic, NewAtomic, std::memory_order_relaxed));
+	}
+
+	static void AtomicFloatMinRelaxed(std::atomic<float>& Atomic, float Value)
+	{
+		float OldAtomic = Atomic.load(std::memory_order_relaxed);
+		float NewAtomic;
+		do
+		{
+			NewAtomic = (OldAtomic < Value) ? OldAtomic : Value;
+		} while (!Atomic.compare_exchange_weak(OldAtomic, NewAtomic, std::memory_order_relaxed));
+	}
+
+	static void AtomicFloatMaxRelaxed(std::atomic<float>& Atomic, float Value)
+	{
+		float OldAtomic = Atomic.load(std::memory_order_relaxed);
+		float NewAtomic;
+		do
+		{
+			NewAtomic = (OldAtomic < Value) ? Value : OldAtomic;
+		} while (!Atomic.compare_exchange_weak(OldAtomic, NewAtomic, std::memory_order_relaxed));
+	}
+
+	std::atomic<float> MinTime;
+	std::atomic<float> MaxTime;
+	std::atomic<float> TotalTime;
+	std::atomic<int32> NumberOfTasks;
 
 	FDelegateHandle EndOfFrameHandle;
 };
