UMaterialExpression::bShaderInputData is now driven via ini file config

- Convert all "input" expressions to be driven via ini file
- Continue the deletion and conversion for empty constructors to be UBT generated defaults

[RN] minor | Rendering - Materials and Shaders

diff --git a/Engine/Config/BaseEditor.ini b/Engine/Config/BaseEditor.ini
--- a/Engine/Config/BaseEditor.ini
+++ b/Engine/Config/BaseEditor.ini
@@ -757,6 +757,7 @@ bUseWidgetTemplateSelector=True
 [/Script/Engine.MaterialExpressionActorPositionWS]
 +MenuCategories=NSLOCTEXT("MaterialExpression", "Vectors", "Vectors")
 +MenuCategories=NSLOCTEXT("MaterialExpression", "Coordinates", "Coordinates")
+bShaderInputData=true
 
 [/Script/Engine.MaterialExpressionAdd]
 +MenuCategories=NSLOCTEXT("MaterialExpression", "Math", "Math")
@@ -807,6 +808,7 @@ bUseWidgetTemplateSelector=True
 
 [/Script/Engine.MaterialExpressionAtmosphericLightVector]
 +MenuCategories=NSLOCTEXT("MaterialExpression", "Utility", "Utility")
+bShaderInputData=true
 
 [/Script/Engine.MaterialExpressionBindlessSwitch]
 +MenuCategories=NSLOCTEXT("MaterialExpression", "Utility", "Utility")
@@ -822,6 +824,7 @@ bUseWidgetTemplateSelector=True
 
 [/Script/Engine.MaterialExpressionBounds]
 +MenuCategories=NSLOCTEXT("MaterialExpression", "Vectors", "Vectors")
+bShaderInputData=true
 
 [/Script/Engine.MaterialExpressionBreakMaterialAttributes]
 +MenuCategories=NSLOCTEXT("MaterialExpression", "MaterialAttributes", "Material Attributes")
@@ -832,9 +835,11 @@ bUseWidgetTemplateSelector=True
 [/Script/Engine.MaterialExpressionCameraPositionWS]
 +MenuCategories=NSLOCTEXT("MaterialExpression", "Vectors", "Vectors")
 +MenuCategories=NSLOCTEXT("MaterialExpression", "Coordinates", "Coordinates")
+bShaderInputData=true
 
 [/Script/Engine.MaterialExpressionCameraVectorWS]
 +MenuCategories=NSLOCTEXT("MaterialExpression", "Vectors", "Vectors")
+bShaderInputData=true
 
 [/Script/Engine.MaterialExpressionCeil]
 +MenuCategories=NSLOCTEXT("MaterialExpression", "Math", "Math")
@@ -907,6 +912,7 @@ bUseWidgetTemplateSelector=True
 
 [/Script/Engine.MaterialExpressionDBufferTexture]
 +MenuCategories=NSLOCTEXT("MaterialExpression", "Decals", "Decals")
+bShaderInputData=true
 
 [/Script/Engine.MaterialExpressionDDX]
 +MenuCategories=NSLOCTEXT("MaterialExpression", "Utility", "Utility")
@@ -916,18 +922,22 @@ bUseWidgetTemplateSelector=True
 
 [/Script/Engine.MaterialExpressionDecalColor]
 +MenuCategories=NSLOCTEXT("MaterialExpression", "Utils", "Utils")
+bShaderInputData=true
 
 [/Script/Engine.MaterialExpressionDecalDerivative]
 +MenuCategories=NSLOCTEXT("MaterialExpression", "Utils", "Utils")
+bShaderInputData=true
 
 [/Script/Engine.MaterialExpressionDecalLifetimeOpacity]
 +MenuCategories=NSLOCTEXT("MaterialExpression", "Utils", "Utils")
+bShaderInputData=true
 
 [/Script/Engine.MaterialExpressionDecalMipmapLevel]
 +MenuCategories=NSLOCTEXT("MaterialExpression", "Utils", "Utils")
 
 [/Script/Engine.MaterialExpressionDeltaTime]
 +MenuCategories=NSLOCTEXT("MaterialExpression", "Constants", "Constants")
+bShaderInputData=true
 
 [/Script/Engine.MaterialExpressionDepthFade]
 +MenuCategories=NSLOCTEXT("MaterialExpression", "Depth", "Depth")
@@ -948,6 +958,7 @@ bUseWidgetTemplateSelector=True
 
 [/Script/Engine.MaterialExpressionDistanceCullFade]
 +MenuCategories=NSLOCTEXT("MaterialExpression", "Constants", "Constants")
+bShaderInputData=true
 
 [/Script/Engine.MaterialExpressionDistanceFieldApproxAO]
 +MenuCategories=NSLOCTEXT("MaterialExpression", "Utility", "Utility")
@@ -974,6 +985,7 @@ bUseWidgetTemplateSelector=True
 [/Script/Engine.MaterialExpressionDynamicParameter]
 +MenuCategories=NSLOCTEXT("MaterialExpression", "Particles", "Particles")
 +MenuCategories=NSLOCTEXT("MaterialExpression", "Parameters", "Parameters")
+bShaderInputData=true
 
 [/Script/Engine.MaterialExpressionExponential]
 +MenuCategories=NSLOCTEXT("MaterialExpression", "Math", "Math")
@@ -983,9 +995,11 @@ bUseWidgetTemplateSelector=True
 
 [/Script/Engine.MaterialExpressionEyeAdaptation]
 +MenuCategories=NSLOCTEXT("MaterialExpression", "Utility", "Utility")
+bShaderInputData=true
 
 [/Script/Engine.MaterialExpressionEyeAdaptationInverse]
 +MenuCategories=NSLOCTEXT("MaterialExpression", "Utility", "Utility")
+bShaderInputData=true
 
 [/Script/Engine.MaterialExpressionFeatureLevelSwitch]
 +MenuCategories=NSLOCTEXT("MaterialExpression", "Utility", "Utility")
@@ -1058,15 +1072,18 @@ bUseWidgetTemplateSelector=True
 [/Script/Engine.MaterialExpressionIsFirstPerson]
 +MenuCategories=NSLOCTEXT("MaterialExpression", "Constants", "Constants")
 +MenuCategories=NSLOCTEXT("MaterialExpression", "FirstPerson", "First Person")
+bShaderInputData=true
 
 [/Script/Engine.MaterialExpressionIsOrthographic]
 +MenuCategories=NSLOCTEXT("MaterialExpression", "Constants", "Constants")
+bShaderInputData=true
 
 [/Script/Engine.MaterialExpressionLength]
 +MenuCategories=NSLOCTEXT("MaterialExpression", "Utility", "Utility")
 
 [/Script/Engine.MaterialExpressionLightmapUVs]
 +MenuCategories=NSLOCTEXT("MaterialExpression", "Coordinates", "Coordinates")
+bShaderInputData=true
 
 [/Script/Engine.MaterialExpressionLightmassReplace]
 +MenuCategories=NSLOCTEXT("MaterialExpression", "Utility", "Utility")
@@ -1080,6 +1097,7 @@ bUseWidgetTemplateSelector=True
 
 [/Script/Engine.MaterialExpressionLocalPosition]
 +MenuCategories=NSLOCTEXT("MaterialExpression", "Coordinates", "Coordinates")
+bShaderInputData=true
 
 [/Script/Engine.MaterialExpressionLogarithm]
 +MenuCategories=NSLOCTEXT("MaterialExpression", "Math", "Math")
@@ -1138,6 +1156,7 @@ bUseWidgetTemplateSelector=True
 
 [/Script/Engine.MaterialExpressionNeuralNetworkOutput]
 +MenuCategories=NSLOCTEXT("MaterialExpression", "NeuralNetwork", "NeuralNetwork")
+bShaderInputData=true
 
 [/Script/Engine.MaterialExpressionNoise]
 +MenuCategories=NSLOCTEXT("MaterialExpressionNoise", "Utility", "Utility")
@@ -1151,20 +1170,25 @@ bUseWidgetTemplateSelector=True
 
 [/Script/Engine.MaterialExpressionObjectBounds]
 +MenuCategories=NSLOCTEXT("MaterialExpression", "Vectors", "Vectors")
+bShaderInputData=true
 
 [/Script/Engine.MaterialExpressionObjectLocalBounds]
 +MenuCategories=NSLOCTEXT("MaterialExpression", "Vectors", "Vectors")
+bShaderInputData=true
 
 [/Script/Engine.MaterialExpressionObjectOrientation]
 +MenuCategories=NSLOCTEXT("MaterialExpression", "Vectors", "Vectors")
 +MenuCategories=NSLOCTEXT("MaterialExpression", "Coordinates", "Coordinates")
+bShaderInputData=true
 
 [/Script/Engine.MaterialExpressionObjectPositionWS]
 +MenuCategories=NSLOCTEXT("MaterialExpression", "Vectors", "Vectors")
 +MenuCategories=NSLOCTEXT("MaterialExpression", "Coordinates", "Coordinates")
+bShaderInputData=true
 
 [/Script/Engine.MaterialExpressionObjectRadius]
 +MenuCategories=NSLOCTEXT("MaterialExpression", "Coordinates", "Coordinates")
+bShaderInputData=true
 
 [/Script/Engine.MaterialExpressionOneMinus]
 +MenuCategories=NSLOCTEXT("MaterialExpression", "Math", "Math")
@@ -1175,46 +1199,57 @@ bUseWidgetTemplateSelector=True
 [/Script/Engine.MaterialExpressionParticleColor]
 +MenuCategories=NSLOCTEXT("MaterialExpression", "Particles", "Particles")
 +MenuCategories=NSLOCTEXT("MaterialExpression", "Constants", "Constants")
+bShaderInputData=true
 
 [/Script/Engine.MaterialExpressionParticleDirection]
 +MenuCategories=NSLOCTEXT("MaterialExpression", "Particles", "Particles")
 +MenuCategories=NSLOCTEXT("MaterialExpression", "Constants", "Constants")
+bShaderInputData=true
 
 [/Script/Engine.MaterialExpressionParticleMacroUV]
 +MenuCategories=NSLOCTEXT("MaterialExpression", "Particles", "Particles")
+bShaderInputData=true
 
 [/Script/Engine.MaterialExpressionParticleMotionBlurFade]
 +MenuCategories=NSLOCTEXT("MaterialExpression", "Particles", "Particles")
 +MenuCategories=NSLOCTEXT("MaterialExpression", "Constants", "Constants")
+bShaderInputData=true
 
 [/Script/Engine.MaterialExpressionParticlePositionWS]
 +MenuCategories=NSLOCTEXT("MaterialExpression", "Particles", "Particles")
 +MenuCategories=NSLOCTEXT("MaterialExpression", "Coordinates", "Coordinates")
 +MenuCategories=NSLOCTEXT("MaterialExpression", "Vectors", "Vectors")
+bShaderInputData=true
 
 [/Script/Engine.MaterialExpressionParticleRadius]
 +MenuCategories=NSLOCTEXT("MaterialExpression", "Particles", "Particles")
 +MenuCategories=NSLOCTEXT("MaterialExpression", "Constants", "Constants")
+bShaderInputData=true
 
 [/Script/Engine.MaterialExpressionParticleRandom]
 +MenuCategories=NSLOCTEXT("MaterialExpression", "Particles", "Particles")
 +MenuCategories=NSLOCTEXT("MaterialExpression", "Constants", "Constants")
+bShaderInputData=true
 
 [/Script/Engine.MaterialExpressionParticleRelativeTime]
 +MenuCategories=NSLOCTEXT("MaterialExpression", "Particles", "Particles")
 +MenuCategories=NSLOCTEXT("MaterialExpression", "Constants", "Constants")
+bShaderInputData=true
 
 [/Script/Engine.MaterialExpressionParticleSize]
 +MenuCategories=NSLOCTEXT("MaterialExpression", "Particles", "Particles")
 +MenuCategories=NSLOCTEXT("MaterialExpression", "Constants", "Constants")
+bShaderInputData=true
 
 [/Script/Engine.MaterialExpressionParticleSpeed]
 +MenuCategories=NSLOCTEXT("MaterialExpression", "Particles", "Particles")
 +MenuCategories=NSLOCTEXT("MaterialExpression", "Constants", "Constants")
+bShaderInputData=true
 
 [/Script/Engine.MaterialExpressionParticleSpriteRotation]
 +MenuCategories=NSLOCTEXT("MaterialExpression", "Particles", "Particles")
 +MenuCategories=NSLOCTEXT("MaterialExpression", "Constants", "Constants")
+bShaderInputData=true
 
 [/Script/Engine.MaterialExpressionParticleSubUV]
 +MenuCategories=NSLOCTEXT("MaterialExpression", "Texture", "Texture")
@@ -1223,9 +1258,11 @@ bUseWidgetTemplateSelector=True
 [/Script/Engine.MaterialExpressionParticleSubUVProperties]
 +MenuCategories=NSLOCTEXT("MaterialExpression", "Particles", "Particles")
 +MenuCategories=NSLOCTEXT("MaterialExpression", "Coordinates", "Coordinates")
+bShaderInputData=true
 
 [/Script/Engine.MaterialExpressionPathTracingBufferTexture]
 +MenuCategories=NSLOCTEXT("MaterialExpression", "PathTracing", "PathTracing")
+bShaderInputData=true
 
 [/Script/Engine.MaterialExpressionPathTracingQualitySwitch]
 +MenuCategories=NSLOCTEXT("MaterialExpression", "Utility", "Utility")
@@ -1235,25 +1272,31 @@ bUseWidgetTemplateSelector=True
 
 [/Script/Engine.MaterialExpressionPerInstanceCustomData]
 +MenuCategories=NSLOCTEXT("MaterialExpression", "Custom", "Custom")
+bShaderInputData=true
 
 [/Script/Engine.MaterialExpressionPerInstanceCustomData3Vector]
 +MenuCategories=NSLOCTEXT("MaterialExpression", "Custom", "Custom")
+bShaderInputData=true
 
 [/Script/Engine.MaterialExpressionPerInstanceFadeAmount]
 +MenuCategories=NSLOCTEXT("MaterialExpression", "Constants", "Constants")
+bShaderInputData=true
 
 [/Script/Engine.MaterialExpressionPerInstanceRandom]
 +MenuCategories=NSLOCTEXT("MaterialExpression", "Constants", "Constants")
+bShaderInputData=true
 
 [/Script/Engine.MaterialExpressionPinBase]
 +MenuCategories=NSLOCTEXT("MaterialExpression", "Utility", "Utility")
 
 [/Script/Engine.MaterialExpressionPixelDepth]
 +MenuCategories=NSLOCTEXT("MaterialExpression", "Depth", "Depth")
+bShaderInputData=true
 
 [/Script/Engine.MaterialExpressionPixelNormalWS]
 +MenuCategories=NSLOCTEXT("MaterialExpression", "Vectors", "Vectors")
 +MenuCategories=NSLOCTEXT("MaterialExpression", "Coordinates", "Coordinates")
+bShaderInputData=true
 
 [/Script/Engine.MaterialExpressionPostVolumeUserFlagTest]
 +MenuCategories=NSLOCTEXT("MaterialExpression", "Get Post Process Setting", "Get Post Process Setting")
@@ -1263,15 +1306,19 @@ bUseWidgetTemplateSelector=True
 
 [/Script/Engine.MaterialExpressionPrecomputedAOMask]
 +MenuCategories=NSLOCTEXT("MaterialExpression", "Constants", "Constants")
+bShaderInputData=true
 
 [/Script/Engine.MaterialExpressionPreSkinnedLocalBounds]
 +MenuCategories=NSLOCTEXT("MaterialExpression", "Vectors", "Vectors")
+bShaderInputData=true
 
 [/Script/Engine.MaterialExpressionPreSkinnedNormal]
 +MenuCategories=NSLOCTEXT("MaterialExpression", "Vectors", "Vectors")
+bShaderInputData=true
 
 [/Script/Engine.MaterialExpressionPreSkinnedPosition]
 +MenuCategories=NSLOCTEXT("MaterialExpression", "Vectors", "Vectors")
+bShaderInputData=true
 
 [/Script/Engine.MaterialExpressionPreviousFrameSwitch]
 +MenuCategories=NSLOCTEXT("MaterialExpression", "Functions", "Functions")
@@ -1287,6 +1334,7 @@ bUseWidgetTemplateSelector=True
 
 [/Script/Engine.MaterialExpressionReflectionVectorWS]
 +MenuCategories=NSLOCTEXT("MaterialExpression", "Vectors", "Vectors")
+bShaderInputData=true
 
 [/Script/Engine.MaterialExpressionReroute]
 +MenuCategories=NSLOCTEXT("MaterialExpression", "Utility", "Utility")
@@ -1336,21 +1384,27 @@ bUseWidgetTemplateSelector=True
 
 [/Script/Engine.MaterialExpressionSceneColor]
 +MenuCategories=NSLOCTEXT("MaterialExpression", "Texture", "Texture")
+bShaderInputData=true
 
 [/Script/Engine.MaterialExpressionSceneDepth]
 +MenuCategories=NSLOCTEXT("MaterialExpression", "Depth", "Depth")
+bShaderInputData=true
 
 [/Script/Engine.MaterialExpressionSceneDepthWithoutWater]
 +MenuCategories=NSLOCTEXT("MaterialExpression", "Water", "Water")
+bShaderInputData=true
 
 [/Script/Engine.MaterialExpressionSceneTexelSize]
 +MenuCategories=NSLOCTEXT("MaterialExpression", "Coordinates", "Coordinates")
+bShaderInputData=true
 
 [/Script/Engine.MaterialExpressionSceneTexture]
 +MenuCategories=NSLOCTEXT("MaterialExpression", "Texture", "Texture")
+bShaderInputData=true
 
 [/Script/Engine.MaterialExpressionScreenPosition]
 +MenuCategories=NSLOCTEXT("MaterialExpression", "Coordinates", "Coordinates")
+bShaderInputData=true
 
 [/Script/Engine.MaterialExpressionSetMaterialAttributes]
 +MenuCategories=NSLOCTEXT("MaterialExpression", "MaterialAttributes", "Material Attributes")
@@ -1439,6 +1493,7 @@ bUseWidgetTemplateSelector=True
 
 [/Script/Engine.MaterialExpressionStaticBool]
 +MenuCategories=NSLOCTEXT("MaterialExpression", "Functions", "Functions")
+bShaderInputData=true
 
 [/Script/Engine.MaterialExpressionStaticBoolParameter]
 +MenuCategories=NSLOCTEXT("MaterialExpression", "Parameters", "Parameters")
@@ -1495,6 +1550,7 @@ bUseWidgetTemplateSelector=True
 
 [/Script/Engine.MaterialExpressionTextureProperty]
 +MenuCategories=NSLOCTEXT("MaterialExpression", "Texture", "Texture")
+bShaderInputData=true
 
 [/Script/Engine.MaterialExpressionTextureSampleParameter2DArray]
 +MenuCategories=NSLOCTEXT("MaterialExpression", "Texture", "Texture")
@@ -1519,6 +1575,7 @@ bUseWidgetTemplateSelector=True
 
 [/Script/Engine.MaterialExpressionTime]
 +MenuCategories=NSLOCTEXT("MaterialExpression", "Constants", "Constants")
+bShaderInputData=true
 
 [/Script/Engine.MaterialExpressionTransform]
 +MenuCategories=NSLOCTEXT("MaterialExpression", "VectorOps", "VectorOps")
@@ -1534,15 +1591,18 @@ bUseWidgetTemplateSelector=True
 
 [/Script/Engine.MaterialExpressionTwoSidedSign]
 +MenuCategories=NSLOCTEXT("MaterialExpression", "Constants", "Constants")
+bShaderInputData=true
 
 [/Script/Engine.MaterialExpressionUserSceneTexture]
 +MenuCategories=NSLOCTEXT("MaterialExpression", "Texture", "Texture")
+bShaderInputData=true
 
 [/Script/Engine.MaterialExpressionVectorNoise]
 +MenuCategories=NSLOCTEXT("MaterialExpressionNoise", "Utility", "Utility")
 
 [/Script/Engine.MaterialExpressionVertexColor]
 +MenuCategories=NSLOCTEXT("MaterialExpression", "Constants", "Constants")
+bShaderInputData=true
 
 [/Script/Engine.MaterialExpressionVertexInterpolator]
 +MenuCategories=NSLOCTEXT("MaterialExpression", "Utility", "Utility")
@@ -1550,16 +1610,20 @@ bUseWidgetTemplateSelector=True
 [/Script/Engine.MaterialExpressionVertexNormalWS]
 +MenuCategories=NSLOCTEXT("MaterialExpression", "Vectors", "Vectors")
 +MenuCategories=NSLOCTEXT("MaterialExpression", "Coordinates", "Coordinates")
+bShaderInputData=true
 
 [/Script/Engine.MaterialExpressionVertexTangentWS]
 +MenuCategories=NSLOCTEXT("MaterialExpression", "Vectors", "Vectors")
 +MenuCategories=NSLOCTEXT("MaterialExpression", "Coordinates", "Coordinates")
+bShaderInputData=true
 
 [/Script/Engine.MaterialExpressionViewProperty]
 +MenuCategories=NSLOCTEXT("MaterialExpression", "Constants", "Constants")
+bShaderInputData=true
 
 [/Script/Engine.MaterialExpressionViewSize]
 +MenuCategories=NSLOCTEXT("MaterialExpression", "Coordinates", "Coordinates")
+bShaderInputData=true
 
 [/Script/Engine.MaterialExpressionVirtualTextureFeatureSwitch]
 +MenuCategories=NSLOCTEXT("MaterialExpression", "VirtualTexture", "VirtualTexture")
@@ -1578,6 +1642,7 @@ bUseWidgetTemplateSelector=True
 
 [/Script/Engine.MaterialExpressionWorldPosition]
 +MenuCategories=NSLOCTEXT("MaterialExpression", "Coordinates", "Coordinates")
+bShaderInputData=true
 
 [/Script/Paper2D.MaterialExpressionSpriteTextureSampler]
 +MenuCategories=NSLOCTEXT("MaterialExpression", "Texture", "Texture")
diff --git a/Engine/Source/Runtime/Engine/Private/Materials/MaterialExpressions.cpp b/Engine/Source/Runtime/Engine/Private/Materials/MaterialExpressions.cpp
--- a/Engine/Source/Runtime/Engine/Private/Materials/MaterialExpressions.cpp
+++ b/Engine/Source/Runtime/Engine/Private/Materials/MaterialExpressions.cpp
@@ -4159,7 +4159,6 @@ UMaterialExpressionTextureProperty::UMaterialExpressionTextureProperty(const FOb
 	Property = TMTM_TextureSize;
 
 #if WITH_EDITORONLY_DATA
-	bShaderInputData = true;
 	bShowOutputNameOnPin = false;
 	
 	Outputs.Reset();
@@ -5456,10 +5455,6 @@ bool  UMaterialExpressionStaticComponentMaskParameter::SetParameterValue(FName I
 UMaterialExpressionTime::UMaterialExpressionTime(const FObjectInitializer& ObjectInitializer)
 	: Super(ObjectInitializer)
 {
-#if WITH_EDITORONLY_DATA
-	bShaderInputData = true;
-#endif
-
 	Period = 0.0f;
 	bOverride_Period = false;
 }
@@ -5490,14 +5485,6 @@ void UMaterialExpressionTime::GetCaption(TArray<FString>& OutCaptions) const
 }
 #endif // WITH_EDITOR
 
-UMaterialExpressionCameraVectorWS::UMaterialExpressionCameraVectorWS(const FObjectInitializer& ObjectInitializer)
-	: Super(ObjectInitializer)
-{
-#if WITH_EDITORONLY_DATA
-	bShaderInputData = true;
-#endif
-}
-
 #if WITH_EDITOR
 int32 UMaterialExpressionCameraVectorWS::Compile(class FMaterialCompiler* Compiler, int32 OutputIndex)
 {
@@ -5510,14 +5497,6 @@ void UMaterialExpressionCameraVectorWS::GetCaption(TArray<FString>& OutCaptions)
 }
 #endif // WITH_EDITOR
 
-UMaterialExpressionCameraPositionWS::UMaterialExpressionCameraPositionWS(const FObjectInitializer& ObjectInitializer)
-	: Super(ObjectInitializer)
-{
-#if WITH_EDITORONLY_DATA
-	bShaderInputData = true;
-#endif
-}
-
 #if WITH_EDITOR
 int32 UMaterialExpressionCameraPositionWS::Compile(class FMaterialCompiler* Compiler, int32 OutputIndex)
 {
@@ -5533,15 +5512,6 @@ void UMaterialExpressionCameraPositionWS::GetCaption(TArray<FString>& OutCaption
 //
 //	UMaterialExpressionReflectionVectorWS
 //
-
-UMaterialExpressionReflectionVectorWS::UMaterialExpressionReflectionVectorWS(const FObjectInitializer& ObjectInitializer)
-	: Super(ObjectInitializer)
-{
-#if WITH_EDITORONLY_DATA
-	bShaderInputData = true;
-#endif
-}
-
 #if WITH_EDITOR
 int32 UMaterialExpressionReflectionVectorWS::Compile(class FMaterialCompiler* Compiler, int32 OutputIndex)
 {
@@ -9059,7 +9029,6 @@ UMaterialExpressionStaticBool::UMaterialExpressionStaticBool(const FObjectInitia
 {
 #if WITH_EDITORONLY_DATA
 	bHidePreviewWindow = true;
-	bShaderInputData = true;
 #endif
 }
 
@@ -10152,7 +10121,6 @@ UMaterialExpressionVertexColor::UMaterialExpressionVertexColor(const FObjectInit
 	Outputs.Add(FExpressionOutput(TEXT(""), 1, 0, 1, 0, 0));
 	Outputs.Add(FExpressionOutput(TEXT(""), 1, 0, 0, 1, 0));
 	Outputs.Add(FExpressionOutput(TEXT(""), 1, 0, 0, 0, 1));
-	bShaderInputData = true;
 #endif
 }
 
@@ -10213,7 +10181,6 @@ UMaterialExpressionParticleColor::UMaterialExpressionParticleColor(const FObject
 	Outputs.Add(FExpressionOutput(TEXT("RGBA"), 1, 1, 1, 1, 1));
 
 	bShowOutputNameOnPin = true;
-	bShaderInputData = true;
 #endif
 }
 
@@ -10229,14 +10196,6 @@ void UMaterialExpressionParticleColor::GetCaption(TArray<FString>& OutCaptions)
 }
 #endif // WITH_EDITOR
 
-UMaterialExpressionParticlePositionWS::UMaterialExpressionParticlePositionWS(const FObjectInitializer& ObjectInitializer)
-	: Super(ObjectInitializer)
-{
-#if WITH_EDITORONLY_DATA
-	bShaderInputData = true;
-#endif
-}
-
 #if WITH_EDITOR
 int32 UMaterialExpressionParticlePositionWS::Compile(class FMaterialCompiler* Compiler, int32 OutputIndex)
 {
@@ -10268,14 +10227,6 @@ void UMaterialExpressionParticlePositionWS::GetCaption(TArray<FString>& OutCapti
 }
 #endif // WITH_EDITOR
 
-UMaterialExpressionParticleRadius::UMaterialExpressionParticleRadius(const FObjectInitializer& ObjectInitializer)
-	: Super(ObjectInitializer)
-{
-#if WITH_EDITORONLY_DATA
-	bShaderInputData = true;
-#endif
-}
-
 #if WITH_EDITOR
 int32 UMaterialExpressionParticleRadius::Compile(class FMaterialCompiler* Compiler, int32 OutputIndex)
 {
@@ -10309,8 +10260,6 @@ UMaterialExpressionDynamicParameter::UMaterialExpressionDynamicParameter(const F
 	Outputs.Add(FExpressionOutput(TEXT(""), 1, 0, 0, 0, 1));
 	Outputs.Add(FExpressionOutput(TEXT("RGB"), 1, 1, 1, 1, 0));
 	Outputs.Add(FExpressionOutput(TEXT("RGBA"), 1, 1, 1, 1, 1));
-
-	bShaderInputData = true;
 #endif // WITH_EDITORONLY_DATA
 
 	DefaultValue = FLinearColor::White;
@@ -10471,7 +10420,6 @@ UMaterialExpressionParticleSubUVProperties::UMaterialExpressionParticleSubUVProp
 	: Super(ObjectInitializer)
 {
 #if WITH_EDITORONLY_DATA
-	bShaderInputData = true;
 	bShowOutputNameOnPin = true;
 
 	Outputs.Reset();
@@ -10501,14 +10449,6 @@ void UMaterialExpressionParticleSubUVProperties::GetCaption(TArray<FString>& Out
 //
 //	MaterialExpressionParticleMacroUV
 //
-UMaterialExpressionParticleMacroUV::UMaterialExpressionParticleMacroUV(const FObjectInitializer& ObjectInitializer)
-	: Super(ObjectInitializer)
-{
-#if WITH_EDITORONLY_DATA
-	bShaderInputData = true;
-#endif
-}
-
 #if WITH_EDITOR
 int32 UMaterialExpressionParticleMacroUV::Compile(class FMaterialCompiler* Compiler, int32 OutputIndex)
 {
@@ -10521,14 +10461,6 @@ void UMaterialExpressionParticleMacroUV::GetCaption(TArray<FString>& OutCaptions
 }
 #endif // WITH_EDITOR
 
-UMaterialExpressionLightVector::UMaterialExpressionLightVector(const FObjectInitializer& ObjectInitializer)
-	: Super(ObjectInitializer)
-{
-#if WITH_EDITORONLY_DATA
-	bShaderInputData = true;
-#endif
-}
-
 #if WITH_EDITOR
 int32 UMaterialExpressionLightVector::Compile(class FMaterialCompiler* Compiler, int32 OutputIndex)
 {
@@ -10545,7 +10477,6 @@ UMaterialExpressionScreenPosition::UMaterialExpressionScreenPosition(const FObje
 	: Super(ObjectInitializer)
 {
 #if WITH_EDITORONLY_DATA
-	bShaderInputData = true;
 	bShowOutputNameOnPin = true;
 
 	Outputs.Reset();
@@ -10574,7 +10505,6 @@ UMaterialExpressionViewProperty::UMaterialExpressionViewProperty(const FObjectIn
 	: Super(ObjectInitializer)
 {
 #if WITH_EDITORONLY_DATA
-	bShaderInputData = true;
 	bShowOutputNameOnPin = true;
 	
 	Outputs.Reset();
@@ -10614,14 +10544,6 @@ void UMaterialExpressionViewProperty::GetCaption(TArray<FString>& OutCaptions) c
 }
 #endif // WITH_EDITOR
 
-UMaterialExpressionViewSize::UMaterialExpressionViewSize(const FObjectInitializer& ObjectInitializer)
-	: Super(ObjectInitializer)
-{
-#if WITH_EDITORONLY_DATA
-	bShaderInputData = true;
-#endif
-}
-
 #if WITH_EDITOR
 int32 UMaterialExpressionViewSize::Compile(class FMaterialCompiler* Compiler, int32 OutputIndex)
 {
@@ -10634,14 +10556,6 @@ void UMaterialExpressionViewSize::GetCaption(TArray<FString>& OutCaptions) const
 }
 #endif // WITH_EDITOR
 
-UMaterialExpressionIsOrthographic::UMaterialExpressionIsOrthographic(const FObjectInitializer& ObjectInitializer)
-	: Super(ObjectInitializer)
-{
-#if WITH_EDITORONLY_DATA
-	bShaderInputData = true;
-#endif
-}
-
 #if WITH_EDITOR
 int32 UMaterialExpressionIsOrthographic::Compile(class FMaterialCompiler* Compiler, int32 OutputIndex)
 {
@@ -10654,14 +10568,6 @@ void UMaterialExpressionIsOrthographic::GetCaption(TArray<FString>& OutCaptions)
 }
 #endif // WITH_EDITOR
 
-UMaterialExpressionDeltaTime::UMaterialExpressionDeltaTime(const FObjectInitializer& ObjectInitializer)
-	: Super(ObjectInitializer)
-{
-#if WITH_EDITORONLY_DATA
-	bShaderInputData = true;
-#endif
-}
-
 #if WITH_EDITOR
 int32 UMaterialExpressionDeltaTime::Compile(class FMaterialCompiler* Compiler, int32 OutputIndex)
 {
@@ -10674,14 +10580,6 @@ void UMaterialExpressionDeltaTime::GetCaption(TArray<FString>& OutCaptions) cons
 }
 #endif // WITH_EDITOR
 
-UMaterialExpressionSceneTexelSize::UMaterialExpressionSceneTexelSize(const FObjectInitializer& ObjectInitializer)
-	: Super(ObjectInitializer)
-{
-#if WITH_EDITORONLY_DATA
-	bShaderInputData = true;
-#endif
-}
-
 #if WITH_EDITOR
 int32 UMaterialExpressionSceneTexelSize::Compile(class FMaterialCompiler* Compiler, int32 OutputIndex)
 {
@@ -10758,7 +10656,6 @@ UMaterialExpressionPixelDepth::UMaterialExpressionPixelDepth(const FObjectInitia
 #if WITH_EDITORONLY_DATA
 	Outputs.Reset();
 	Outputs.Add(FExpressionOutput(TEXT(""), 1, 1, 0, 0, 0));
-	bShaderInputData = true;
 #endif
 }
 
@@ -10786,7 +10683,6 @@ UMaterialExpressionSceneDepth::UMaterialExpressionSceneDepth(const FObjectInitia
 #if WITH_EDITORONLY_DATA
 	Outputs.Reset();
 	Outputs.Add(FExpressionOutput(TEXT(""), 1, 1, 0, 0, 0));
-	bShaderInputData = true;
 #endif
 
 	ConstInput = FVector2D(0.f, 0.f);
@@ -10862,7 +10758,6 @@ UMaterialExpressionSceneTexture::UMaterialExpressionSceneTexture(const FObjectIn
 	: Super(ObjectInitializer)
 {
 #if WITH_EDITORONLY_DATA
-	bShaderInputData = true;
 	bShowOutputNameOnPin = true;
 #endif
 
@@ -10924,7 +10819,6 @@ UMaterialExpressionUserSceneTexture::UMaterialExpressionUserSceneTexture(const F
 	: Super(ObjectInitializer)
 {
 #if WITH_EDITORONLY_DATA
-	bShaderInputData = true;
 	bShowOutputNameOnPin = true;
 #endif
 
@@ -10987,7 +10881,6 @@ UMaterialExpressionSceneColor::UMaterialExpressionSceneColor(const FObjectInitia
 	: Super(ObjectInitializer)
 {
 #if WITH_EDITORONLY_DATA
-	bShaderInputData = true;
 	bShowOutputNameOnPin = true;
 #endif
 	ConstInput = FVector2D(0.f, 0.f);
@@ -11055,7 +10948,6 @@ UMaterialExpressionDBufferTexture::UMaterialExpressionDBufferTexture(const FObje
 	: Super(ObjectInitializer)
 {
 #if WITH_EDITORONLY_DATA
-	bShaderInputData = true;
 	bShowOutputNameOnPin = true;
 #endif
 
@@ -12517,7 +12409,6 @@ UMaterialExpressionLocalPosition::UMaterialExpressionLocalPosition(const FObject
 	Outputs.Add(FExpressionOutput(TEXT("Z"), 1, 0, 0, 1, 0));
 
 	bShowOutputNameOnPin = true;
-	bShaderInputData = true;
 #endif
 
 	IncludedOffsets = EPositionIncludedOffsets::IncludeOffsets;
@@ -12577,7 +12468,6 @@ UMaterialExpressionWorldPosition::UMaterialExpressionWorldPosition(const FObject
 	Outputs.Add(FExpressionOutput(TEXT("Z"), 1, 0, 0, 1, 0));
 
 	bShowOutputNameOnPin = true;
-	bShaderInputData = true;
 #endif
 	WorldPositionShaderOffset = WPT_Default;
 }
@@ -12629,14 +12519,6 @@ void UMaterialExpressionWorldPosition::GetCaption(TArray<FString>& OutCaptions)
 ///////////////////////////////////////////////////////////////////////////////
 // UMaterialExpressionObjectPositionWS
 ///////////////////////////////////////////////////////////////////////////////
-UMaterialExpressionObjectPositionWS::UMaterialExpressionObjectPositionWS(const FObjectInitializer& ObjectInitializer)
-	: Super(ObjectInitializer)
-{
-#if WITH_EDITORONLY_DATA
-	bShaderInputData = true;
-#endif
-}
-
 #if WITH_EDITOR
 int32 UMaterialExpressionObjectPositionWS::Compile(class FMaterialCompiler* Compiler, int32 OutputIndex)
 {
@@ -12698,14 +12580,6 @@ void UMaterialExpressionObjectPositionWS::GetExpressionToolTip(TArray<FString>&
 ///////////////////////////////////////////////////////////////////////////////
 // UMaterialExpressionObjectRadius
 ///////////////////////////////////////////////////////////////////////////////
-UMaterialExpressionObjectRadius::UMaterialExpressionObjectRadius(const FObjectInitializer& ObjectInitializer)
-	: Super(ObjectInitializer)
-{
-#if WITH_EDITORONLY_DATA
-	bShaderInputData = true;
-#endif
-}
-
 #if WITH_EDITOR
 int32 UMaterialExpressionObjectRadius::Compile(class FMaterialCompiler* Compiler, int32 OutputIndex)
 {
@@ -12726,14 +12600,6 @@ void UMaterialExpressionObjectRadius::GetCaption(TArray<FString>& OutCaptions) c
 ///////////////////////////////////////////////////////////////////////////////
 // UMaterialExpressionObjectBoundingBox
 ///////////////////////////////////////////////////////////////////////////////
-UMaterialExpressionObjectBounds::UMaterialExpressionObjectBounds(const FObjectInitializer& ObjectInitializer)
-	: Super(ObjectInitializer)
-{
-#if WITH_EDITORONLY_DATA
-	bShaderInputData = true;
-#endif
-}
-
 #if WITH_EDITOR
 int32 UMaterialExpressionObjectBounds::Compile(class FMaterialCompiler* Compiler, int32 OutputIndex)
 {
@@ -12758,7 +12624,6 @@ UMaterialExpressionObjectLocalBounds::UMaterialExpressionObjectLocalBounds(const
 	: Super(ObjectInitializer)
 {
 #if WITH_EDITORONLY_DATA
-	bShaderInputData = true;
 	bShowOutputNameOnPin = true;
 
 	Outputs.Reset();
@@ -12815,7 +12680,6 @@ UMaterialExpressionBounds::UMaterialExpressionBounds(const FObjectInitializer& O
 	: Super(ObjectInitializer)
 {
 #if WITH_EDITORONLY_DATA
-	bShaderInputData = true;
 	bShowOutputNameOnPin = true;
 
 	Outputs.Reset();
@@ -12895,7 +12759,6 @@ UMaterialExpressionPreSkinnedLocalBounds::UMaterialExpressionPreSkinnedLocalBoun
 	: Super(ObjectInitializer)
 {
 #if WITH_EDITORONLY_DATA
-	bShaderInputData = true;
 	bShowOutputNameOnPin = true;
 
 	Outputs.Reset();
@@ -12948,14 +12811,6 @@ void UMaterialExpressionPreSkinnedLocalBounds::GetExpressionToolTip(TArray<FStri
 ///////////////////////////////////////////////////////////////////////////////
 // UMaterialExpressionDistanceCullFade
 ///////////////////////////////////////////////////////////////////////////////
-UMaterialExpressionDistanceCullFade::UMaterialExpressionDistanceCullFade(const FObjectInitializer& ObjectInitializer)
-	: Super(ObjectInitializer)
-{
-#if WITH_EDITORONLY_DATA
-	bShaderInputData = true;
-#endif
-}
-
 #if WITH_EDITOR
 int32 UMaterialExpressionDistanceCullFade::Compile(class FMaterialCompiler* Compiler, int32 OutputIndex)
 {
@@ -13021,14 +12876,6 @@ void UMaterialExpressionDistanceFieldsRenderingSwitch::GetCaption(TArray<FString
 ///////////////////////////////////////////////////////////////////////////////
 // UMaterialExpressionActorPositionWS
 ///////////////////////////////////////////////////////////////////////////////
-UMaterialExpressionActorPositionWS::UMaterialExpressionActorPositionWS(const FObjectInitializer& ObjectInitializer)
-	: Super(ObjectInitializer)
-{
-#if WITH_EDITORONLY_DATA
-	bShaderInputData = true;
-#endif
-}
-
 #if WITH_EDITOR
 int32 UMaterialExpressionActorPositionWS::Compile(class FMaterialCompiler* Compiler, int32 OutputIndex)
 {
@@ -19164,8 +19011,6 @@ UMaterialExpressionLightmapUVs::UMaterialExpressionLightmapUVs(const FObjectInit
 
 	Outputs.Reset();
 	Outputs.Add(FExpressionOutput(TEXT(""), 1, 1, 1, 0, 0));
-
-	bShaderInputData = true;
 #endif
 }
 
@@ -19194,8 +19039,6 @@ UMaterialExpressionPrecomputedAOMask::UMaterialExpressionPrecomputedAOMask(const
 
 	Outputs.Reset();
 	Outputs.Add(FExpressionOutput(TEXT("")));
-
-	bShaderInputData = true;
 #endif
 }
 
@@ -19552,7 +19395,6 @@ UMaterialExpressionPathTracingBufferTexture::UMaterialExpressionPathTracingBuffe
 	: Super(ObjectInitializer)
 {
 #if WITH_EDITORONLY_DATA
-	bShaderInputData = true;
 	bShowOutputNameOnPin = true;
 #endif
 
@@ -19587,14 +19429,6 @@ void UMaterialExpressionPathTracingBufferTexture::GetCaption(TArray<FString>& Ou
 }
 #endif // WITH_EDITOR
 
-UMaterialExpressionObjectOrientation::UMaterialExpressionObjectOrientation(const FObjectInitializer& ObjectInitializer)
-	: Super(ObjectInitializer)
-{
-#if WITH_EDITORONLY_DATA
-	bShaderInputData = true;
-#endif
-}
-
 #if WITH_EDITOR
 int32 UMaterialExpressionObjectOrientation::Compile(class FMaterialCompiler* Compiler, int32 OutputIndex)
 {
@@ -20762,14 +20596,6 @@ void UMaterialExpressionDistance::GetCaption(TArray<FString>& OutCaptions) const
 }
 #endif // WITH_EDITOR
 
-UMaterialExpressionTwoSidedSign::UMaterialExpressionTwoSidedSign(const FObjectInitializer& ObjectInitializer)
-	: Super(ObjectInitializer)
-{
-#if WITH_EDITORONLY_DATA
-	bShaderInputData = true;
-#endif
-}
-
 #if WITH_EDITOR
 int32 UMaterialExpressionTwoSidedSign::Compile(class FMaterialCompiler* Compiler, int32 OutputIndex)
 {
@@ -20782,14 +20608,6 @@ void UMaterialExpressionTwoSidedSign::GetCaption(TArray<FString>& OutCaptions) c
 }
 #endif // WITH_EDITOR
 
-UMaterialExpressionVertexNormalWS::UMaterialExpressionVertexNormalWS(const FObjectInitializer& ObjectInitializer)
-	: Super(ObjectInitializer)
-{
-#if WITH_EDITORONLY_DATA
-	bShaderInputData = true;
-#endif
-}
-
 #if WITH_EDITOR
 int32 UMaterialExpressionVertexNormalWS::Compile(class FMaterialCompiler* Compiler, int32 OutputIndex)
 {
@@ -20802,14 +20620,6 @@ void UMaterialExpressionVertexNormalWS::GetCaption(TArray<FString>& OutCaptions)
 }
 #endif // WITH_EDITOR
 
-UMaterialExpressionVertexTangentWS::UMaterialExpressionVertexTangentWS(const FObjectInitializer& ObjectInitializer)
-	: Super(ObjectInitializer)
-{
-#if WITH_EDITORONLY_DATA
-	bShaderInputData = true;
-#endif
-}
-
 #if WITH_EDITOR
 int32 UMaterialExpressionVertexTangentWS::Compile(class FMaterialCompiler* Compiler, int32 OutputIndex)
 {
@@ -20822,14 +20632,6 @@ void UMaterialExpressionVertexTangentWS::GetCaption(TArray<FString>& OutCaptions
 }
 #endif // WITH_EDITOR
 
-UMaterialExpressionPixelNormalWS::UMaterialExpressionPixelNormalWS(const FObjectInitializer& ObjectInitializer)
-	: Super(ObjectInitializer)
-{
-#if WITH_EDITORONLY_DATA
-	bShaderInputData = true;
-#endif
-}
-
 #if WITH_EDITOR
 int32 UMaterialExpressionPixelNormalWS::Compile(class FMaterialCompiler* Compiler, int32 OutputIndex)
 {
@@ -20845,14 +20647,6 @@ void UMaterialExpressionPixelNormalWS::GetCaption(TArray<FString>& OutCaptions)
 ///////////////////////////////////////////////////////////////////////////////
 // UMaterialExpressionPerInstanceRandom
 ///////////////////////////////////////////////////////////////////////////////
-UMaterialExpressionPerInstanceRandom::UMaterialExpressionPerInstanceRandom(const FObjectInitializer& ObjectInitializer)
-	: Super(ObjectInitializer)
-{
-#if WITH_EDITORONLY_DATA
-	bShaderInputData = true;
-#endif
-}
-
 #if WITH_EDITOR
 int32 UMaterialExpressionPerInstanceRandom::Compile(class FMaterialCompiler* Compiler, int32 OutputIndex)
 {
@@ -20868,14 +20662,6 @@ void UMaterialExpressionPerInstanceRandom::GetCaption(TArray<FString>& OutCaptio
 ///////////////////////////////////////////////////////////////////////////////
 // UMaterialExpressionPerInstanceFadeAmount
 ///////////////////////////////////////////////////////////////////////////////
-UMaterialExpressionPerInstanceFadeAmount::UMaterialExpressionPerInstanceFadeAmount(const FObjectInitializer& ObjectInitializer)
-	: Super(ObjectInitializer)
-{
-#if WITH_EDITORONLY_DATA
-	bShaderInputData = true;
-#endif
-}
-
 #if WITH_EDITOR
 int32 UMaterialExpressionPerInstanceFadeAmount::Compile(class FMaterialCompiler* Compiler, int32 OutputIndex)
 {
@@ -20891,14 +20677,6 @@ void UMaterialExpressionPerInstanceFadeAmount::GetCaption(TArray<FString>& OutCa
 ///////////////////////////////////////////////////////////////////////////////
 // UMaterialExpressionPerInstanceCustomData
 ///////////////////////////////////////////////////////////////////////////////
-UMaterialExpressionPerInstanceCustomData::UMaterialExpressionPerInstanceCustomData(const FObjectInitializer& ObjectInitializer)
-	: Super(ObjectInitializer)
-{
-#if WITH_EDITORONLY_DATA
-	bShaderInputData = true;
-#endif
-}
-
 #if WITH_EDITOR
 int32 UMaterialExpressionPerInstanceCustomData::Compile(class FMaterialCompiler* Compiler, int32 OutputIndex)
 {
@@ -20915,14 +20693,6 @@ void UMaterialExpressionPerInstanceCustomData::GetCaption(TArray<FString>& OutCa
 ///////////////////////////////////////////////////////////////////////////////
 // UMaterialExpressionPerInstanceCustomData3Vector
 ///////////////////////////////////////////////////////////////////////////////
-UMaterialExpressionPerInstanceCustomData3Vector::UMaterialExpressionPerInstanceCustomData3Vector(const FObjectInitializer& ObjectInitializer)
-	: Super(ObjectInitializer)
-{
-#if WITH_EDITORONLY_DATA
-	bShaderInputData = true;
-#endif
-}
-
 #if WITH_EDITOR
 
 int32 UMaterialExpressionPerInstanceCustomData3Vector::Compile(class FMaterialCompiler* Compiler, int32 OutputIndex)
@@ -21047,8 +20817,6 @@ UMaterialExpressionDecalDerivative::UMaterialExpressionDecalDerivative(const FOb
 	: Super(ObjectInitializer)
 {
 #if WITH_EDITORONLY_DATA
-	//bCollapsed = true;
-	bShaderInputData = true;
 	bShowOutputNameOnPin = true;
 	
 	Outputs.Reset();
@@ -21077,7 +20845,6 @@ UMaterialExpressionDecalColor::UMaterialExpressionDecalColor(const FObjectInitia
 {
 #if WITH_EDITORONLY_DATA
 	bShowOutputNameOnPin = true;
-	bShaderInputData = true;
 
 	Outputs.Reset();
 	Outputs.Add(FExpressionOutput(TEXT("RGB"), 1, 1, 1, 1, 0));
@@ -21113,8 +20880,6 @@ UMaterialExpressionDecalLifetimeOpacity::UMaterialExpressionDecalLifetimeOpacity
 	: Super(ObjectInitializer)
 {
 #if WITH_EDITORONLY_DATA
-	bShaderInputData = true;
-
 	Outputs.Reset();
 	Outputs.Add(FExpressionOutput(TEXT("Opacity")));
 #endif
@@ -21392,14 +21157,6 @@ void UMaterialExpressionDDY::GetCaption(TArray<FString>& OutCaptions) const
 /*------------------------------------------------------------------------------
 	Particle relative time material expression.
 ------------------------------------------------------------------------------*/
-UMaterialExpressionParticleRelativeTime::UMaterialExpressionParticleRelativeTime(const FObjectInitializer& ObjectInitializer)
-	: Super(ObjectInitializer)
-{
-#if WITH_EDITORONLY_DATA
-	bShaderInputData = true;
-#endif
-}
-
 #if WITH_EDITOR
 int32 UMaterialExpressionParticleRelativeTime::Compile(class FMaterialCompiler* Compiler, int32 OutputIndex)
 {
@@ -21415,14 +21172,6 @@ void UMaterialExpressionParticleRelativeTime::GetCaption(TArray<FString>& OutCap
 /*------------------------------------------------------------------------------
 	Particle motion blur fade material expression.
 ------------------------------------------------------------------------------*/
-UMaterialExpressionParticleMotionBlurFade::UMaterialExpressionParticleMotionBlurFade(const FObjectInitializer& ObjectInitializer)
-	: Super(ObjectInitializer)
-{
-#if WITH_EDITORONLY_DATA
-	bShaderInputData = true;
-#endif
-}
-
 #if WITH_EDITOR
 int32 UMaterialExpressionParticleMotionBlurFade::Compile(class FMaterialCompiler* Compiler, int32 OutputIndex)
 {
@@ -21438,14 +21187,6 @@ void UMaterialExpressionParticleMotionBlurFade::GetCaption(TArray<FString>& OutC
 /*------------------------------------------------------------------------------
 	Particle motion blur fade material expression.
 ------------------------------------------------------------------------------*/
-UMaterialExpressionParticleRandom::UMaterialExpressionParticleRandom(const FObjectInitializer& ObjectInitializer)
-	: Super(ObjectInitializer)
-{
-#if WITH_EDITORONLY_DATA
-	bShaderInputData = true;
-#endif
-}
-
 #if WITH_EDITOR
 int32 UMaterialExpressionParticleRandom::Compile(class FMaterialCompiler* Compiler, int32 OutputIndex)
 {
@@ -21461,14 +21202,6 @@ void UMaterialExpressionParticleRandom::GetCaption(TArray<FString>& OutCaptions)
 /*------------------------------------------------------------------------------
 	Particle direction material expression.
 ------------------------------------------------------------------------------*/
-UMaterialExpressionParticleDirection::UMaterialExpressionParticleDirection(const FObjectInitializer& ObjectInitializer)
-	: Super(ObjectInitializer)
-{
-#if WITH_EDITORONLY_DATA
-	bShaderInputData = true;
-#endif
-}
-
 #if WITH_EDITOR
 int32 UMaterialExpressionParticleDirection::Compile(class FMaterialCompiler* Compiler, int32 OutputIndex)
 {
@@ -21484,14 +21217,6 @@ void UMaterialExpressionParticleDirection::GetCaption(TArray<FString>& OutCaptio
 /*------------------------------------------------------------------------------
 	Particle speed material expression.
 ------------------------------------------------------------------------------*/
-UMaterialExpressionParticleSpeed::UMaterialExpressionParticleSpeed(const FObjectInitializer& ObjectInitializer)
-	: Super(ObjectInitializer)
-{
-#if WITH_EDITORONLY_DATA
-	bShaderInputData = true;
-#endif
-}
-
 #if WITH_EDITOR
 int32 UMaterialExpressionParticleSpeed::Compile(class FMaterialCompiler* Compiler, int32 OutputIndex)
 {
@@ -21507,14 +21232,6 @@ void UMaterialExpressionParticleSpeed::GetCaption(TArray<FString>& OutCaptions)
 /*------------------------------------------------------------------------------
 	Particle size material expression.
 ------------------------------------------------------------------------------*/
-UMaterialExpressionParticleSize::UMaterialExpressionParticleSize(const FObjectInitializer& ObjectInitializer)
-	: Super(ObjectInitializer)
-{
-#if WITH_EDITORONLY_DATA
-	bShaderInputData = true;
-#endif
-}
-
 #if WITH_EDITOR
 int32 UMaterialExpressionParticleSize::Compile(class FMaterialCompiler* Compiler, int32 OutputIndex)
 {
@@ -21538,7 +21255,6 @@ UMaterialExpressionParticleSpriteRotation::UMaterialExpressionParticleSpriteRota
 	Outputs.Emplace(TEXT("Rad"), 1, 1, 0, 0, 0);
 	Outputs.Emplace(TEXT("Deg"), 1, 0, 1, 0, 0);
 
-	bShaderInputData = true;
 	bShowOutputNameOnPin = true;
 #endif
 }
@@ -21700,7 +21416,6 @@ UMaterialExpressionEyeAdaptation::UMaterialExpressionEyeAdaptation(const FObject
 #if WITH_EDITORONLY_DATA
 	Outputs.Reset();
 	Outputs.Add(FExpressionOutput(TEXT("EyeAdaptation")));
-	bShaderInputData = true;
 #endif
 }
 
@@ -21725,7 +21440,6 @@ UMaterialExpressionEyeAdaptationInverse::UMaterialExpressionEyeAdaptationInverse
 #if WITH_EDITORONLY_DATA
 	Outputs.Reset();
 	Outputs.Add(FExpressionOutput(TEXT("EyeAdaptationInverse")));
-	bShaderInputData = true;
 #endif
 }
 
@@ -22203,7 +21917,6 @@ UMaterialExpressionPreSkinnedPosition::UMaterialExpressionPreSkinnedPosition(con
 #if WITH_EDITORONLY_DATA
 	Outputs.Reset();
 	Outputs.Add(FExpressionOutput(TEXT(""), 1, 1, 1, 1, 0));
-	bShaderInputData = true;
 #endif
 }
 
@@ -22251,7 +21964,6 @@ UMaterialExpressionPreSkinnedNormal::UMaterialExpressionPreSkinnedNormal(const F
 #if WITH_EDITORONLY_DATA
 	Outputs.Reset();
 	Outputs.Add(FExpressionOutput(TEXT(""), 1, 1, 1, 1, 0));
-	bShaderInputData = true;
 #endif
 }
 
@@ -22622,7 +22334,6 @@ UMaterialExpressionShadingModel::UMaterialExpressionShadingModel(const FObjectIn
 
 	MenuCategories.Add(LOCTEXT("Shading Model", "Shading Model"));
 
-	// bShaderInputData = true;
 	bShowOutputNameOnPin = true;
 	
 	Outputs.Reset();
@@ -23201,7 +22912,6 @@ UMaterialExpressionSceneDepthWithoutWater::UMaterialExpressionSceneDepthWithoutW
 
 	Outputs.Reset();
 	Outputs.Add(FExpressionOutput(TEXT(""), 1, 1, 0, 0, 0));
-	bShaderInputData = true;
 #endif // WITH_EDITORONLY_DATA
 
 	ConstInput = FVector2D(0.f, 0.f);
@@ -23917,7 +23627,6 @@ UMaterialExpressionNeuralNetworkOutput::UMaterialExpressionNeuralNetworkOutput(c
 	: Super(ObjectInitializer)
 {
 #if WITH_EDITORONLY_DATA
-	bShaderInputData = true;
 	bShowOutputNameOnPin = true;
 #endif
 
@@ -24016,14 +23725,6 @@ FString UMaterialExpressionFirstPersonOutput::GetDisplayName() const
 	return TEXT("First Person Output");
 }
 
-UMaterialExpressionIsFirstPerson::UMaterialExpressionIsFirstPerson(const FObjectInitializer& ObjectInitializer)
-	: Super(ObjectInitializer)
-{
-#if WITH_EDITORONLY_DATA
-	bShaderInputData = true;
-#endif
-}
-
 #if WITH_EDITOR
 int32 UMaterialExpressionIsFirstPerson::Compile(class FMaterialCompiler* Compiler, int32 OutputIndex)
 {
diff --git a/Engine/Source/Runtime/Engine/Public/Materials/MaterialExpression.h b/Engine/Source/Runtime/Engine/Public/Materials/MaterialExpression.h
--- a/Engine/Source/Runtime/Engine/Public/Materials/MaterialExpression.h
+++ b/Engine/Source/Runtime/Engine/Public/Materials/MaterialExpression.h
@@ -221,7 +221,7 @@ class UMaterialExpression : public UObject
 	uint32 bCollapsed:1;
 
 	/** Whether the node represents an input to the shader or not.  Used to color the node's background. */
-	UPROPERTY()
+	UPROPERTY(config)
 	uint32 bShaderInputData:1;
 
 	/** Whether to draw the expression's inputs. */
diff --git a/Engine/Source/Runtime/Engine/Public/Materials/MaterialExpressionActorPositionWS.h b/Engine/Source/Runtime/Engine/Public/Materials/MaterialExpressionActorPositionWS.h
--- a/Engine/Source/Runtime/Engine/Public/Materials/MaterialExpressionActorPositionWS.h
+++ b/Engine/Source/Runtime/Engine/Public/Materials/MaterialExpressionActorPositionWS.h
@@ -11,7 +11,9 @@
 UCLASS(collapsecategories, hidecategories=Object)
 class UMaterialExpressionActorPositionWS : public UMaterialExpression
 {
-	GENERATED_UCLASS_BODY()
+	GENERATED_BODY()
+
+public:
 
 	UPROPERTY(EditAnywhere, Category=UMaterialExpressionActorPositionWS, meta=(DisplayName = "Origin", ShowAsInputPin = "Advanced"))
 	EPositionOrigin OriginType;
diff --git a/Engine/Source/Runtime/Engine/Public/Materials/MaterialExpressionCameraPositionWS.h b/Engine/Source/Runtime/Engine/Public/Materials/MaterialExpressionCameraPositionWS.h
--- a/Engine/Source/Runtime/Engine/Public/Materials/MaterialExpressionCameraPositionWS.h
+++ b/Engine/Source/Runtime/Engine/Public/Materials/MaterialExpressionCameraPositionWS.h
@@ -11,8 +11,9 @@
 UCLASS(MinimalAPI, collapsecategories, hidecategories=Object)
 class UMaterialExpressionCameraPositionWS : public UMaterialExpression
 {
-	GENERATED_UCLASS_BODY()
+	GENERATED_BODY()
 
+public:
 
 	//~ Begin UMaterialExpression Interface
 #if WITH_EDITOR
diff --git a/Engine/Source/Runtime/Engine/Public/Materials/MaterialExpressionCameraVectorWS.h b/Engine/Source/Runtime/Engine/Public/Materials/MaterialExpressionCameraVectorWS.h
--- a/Engine/Source/Runtime/Engine/Public/Materials/MaterialExpressionCameraVectorWS.h
+++ b/Engine/Source/Runtime/Engine/Public/Materials/MaterialExpressionCameraVectorWS.h
@@ -11,8 +11,9 @@
 UCLASS(MinimalAPI, collapsecategories, hidecategories=Object)
 class UMaterialExpressionCameraVectorWS : public UMaterialExpression
 {
-	GENERATED_UCLASS_BODY()
+	GENERATED_BODY()
 
+public:
 
 	//~ Begin UMaterialExpression Interface
 #if WITH_EDITOR
diff --git a/Engine/Source/Runtime/Engine/Public/Materials/MaterialExpressionDeltaTime.h b/Engine/Source/Runtime/Engine/Public/Materials/MaterialExpressionDeltaTime.h
--- a/Engine/Source/Runtime/Engine/Public/Materials/MaterialExpressionDeltaTime.h
+++ b/Engine/Source/Runtime/Engine/Public/Materials/MaterialExpressionDeltaTime.h
@@ -11,7 +11,9 @@
 UCLASS(collapsecategories, hidecategories=Object)
 class UMaterialExpressionDeltaTime : public UMaterialExpression
 {
-	GENERATED_UCLASS_BODY()
+	GENERATED_BODY()
+
+public:
 
 	//~ Begin UMaterialExpression Interface
 #if WITH_EDITOR
diff --git a/Engine/Source/Runtime/Engine/Public/Materials/MaterialExpressionDistanceCullFade.h b/Engine/Source/Runtime/Engine/Public/Materials/MaterialExpressionDistanceCullFade.h
--- a/Engine/Source/Runtime/Engine/Public/Materials/MaterialExpressionDistanceCullFade.h
+++ b/Engine/Source/Runtime/Engine/Public/Materials/MaterialExpressionDistanceCullFade.h
@@ -11,8 +11,9 @@
 UCLASS(collapsecategories, hidecategories=Object)
 class UMaterialExpressionDistanceCullFade : public UMaterialExpression
 {
-	GENERATED_UCLASS_BODY()
+	GENERATED_BODY()
 
+public:
 
 	//~ Begin UMaterialExpression Interface
 #if WITH_EDITOR
diff --git a/Engine/Source/Runtime/Engine/Public/Materials/MaterialExpressionIsFirstPerson.h b/Engine/Source/Runtime/Engine/Public/Materials/MaterialExpressionIsFirstPerson.h
--- a/Engine/Source/Runtime/Engine/Public/Materials/MaterialExpressionIsFirstPerson.h
+++ b/Engine/Source/Runtime/Engine/Public/Materials/MaterialExpressionIsFirstPerson.h
@@ -11,7 +11,9 @@
 UCLASS(collapsecategories, hidecategories = Object, MinimalAPI)
 class UMaterialExpressionIsFirstPerson : public UMaterialExpression
 {
-	GENERATED_UCLASS_BODY()
+	GENERATED_BODY()
+
+public:
 
 	//~ Begin UMaterialExpression Interface
 #if WITH_EDITOR
diff --git a/Engine/Source/Runtime/Engine/Public/Materials/MaterialExpressionIsOrthographic.h b/Engine/Source/Runtime/Engine/Public/Materials/MaterialExpressionIsOrthographic.h
--- a/Engine/Source/Runtime/Engine/Public/Materials/MaterialExpressionIsOrthographic.h
+++ b/Engine/Source/Runtime/Engine/Public/Materials/MaterialExpressionIsOrthographic.h
@@ -11,7 +11,9 @@
 UCLASS(collapsecategories, hidecategories=Object)
 class UMaterialExpressionIsOrthographic : public UMaterialExpression
 {
-	GENERATED_UCLASS_BODY()
+	GENERATED_BODY()
+
+public:
 
 	//~ Begin UMaterialExpression Interface
 #if WITH_EDITOR
diff --git a/Engine/Source/Runtime/Engine/Public/Materials/MaterialExpressionLightVector.h b/Engine/Source/Runtime/Engine/Public/Materials/MaterialExpressionLightVector.h
--- a/Engine/Source/Runtime/Engine/Public/Materials/MaterialExpressionLightVector.h
+++ b/Engine/Source/Runtime/Engine/Public/Materials/MaterialExpressionLightVector.h
@@ -11,8 +11,9 @@
 UCLASS(collapsecategories, hidecategories=Object)
 class UMaterialExpressionLightVector : public UMaterialExpression
 {
-	GENERATED_UCLASS_BODY()
+	GENERATED_BODY()
 
+public:
 
 	//~ Begin UMaterialExpression Interface
 #if WITH_EDITOR
diff --git a/Engine/Source/Runtime/Engine/Public/Materials/MaterialExpressionObjectBounds.h b/Engine/Source/Runtime/Engine/Public/Materials/MaterialExpressionObjectBounds.h
--- a/Engine/Source/Runtime/Engine/Public/Materials/MaterialExpressionObjectBounds.h
+++ b/Engine/Source/Runtime/Engine/Public/Materials/MaterialExpressionObjectBounds.h
@@ -11,8 +11,9 @@
 UCLASS(collapsecategories, hidecategories=Object)
 class UMaterialExpressionObjectBounds : public UMaterialExpression
 {
-	GENERATED_UCLASS_BODY()
+	GENERATED_BODY()
 
+public:
 
 	//~ Begin UMaterialExpression Interface
 #if WITH_EDITOR
diff --git a/Engine/Source/Runtime/Engine/Public/Materials/MaterialExpressionObjectOrientation.h b/Engine/Source/Runtime/Engine/Public/Materials/MaterialExpressionObjectOrientation.h
--- a/Engine/Source/Runtime/Engine/Public/Materials/MaterialExpressionObjectOrientation.h
+++ b/Engine/Source/Runtime/Engine/Public/Materials/MaterialExpressionObjectOrientation.h
@@ -11,8 +11,9 @@
 UCLASS(collapsecategories, hidecategories=Object, MinimalAPI)
 class UMaterialExpressionObjectOrientation : public UMaterialExpression
 {
-	GENERATED_UCLASS_BODY()
+	GENERATED_BODY()
 
+public:
 
 	//~ Begin UMaterialExpression Interface
 #if WITH_EDITOR
diff --git a/Engine/Source/Runtime/Engine/Public/Materials/MaterialExpressionObjectPositionWS.h b/Engine/Source/Runtime/Engine/Public/Materials/MaterialExpressionObjectPositionWS.h
--- a/Engine/Source/Runtime/Engine/Public/Materials/MaterialExpressionObjectPositionWS.h
+++ b/Engine/Source/Runtime/Engine/Public/Materials/MaterialExpressionObjectPositionWS.h
@@ -11,7 +11,9 @@
 UCLASS(collapsecategories, hidecategories=Object, MinimalAPI)
 class UMaterialExpressionObjectPositionWS : public UMaterialExpression
 {
-	GENERATED_UCLASS_BODY()
+	GENERATED_BODY()
+
+public:
 
 	UPROPERTY(EditAnywhere, Category=UMaterialExpressionObjectPositionWS, meta=(DisplayName = "Origin", ShowAsInputPin = "Advanced"))
 	EPositionOrigin OriginType;
diff --git a/Engine/Source/Runtime/Engine/Public/Materials/MaterialExpressionObjectRadius.h b/Engine/Source/Runtime/Engine/Public/Materials/MaterialExpressionObjectRadius.h
--- a/Engine/Source/Runtime/Engine/Public/Materials/MaterialExpressionObjectRadius.h
+++ b/Engine/Source/Runtime/Engine/Public/Materials/MaterialExpressionObjectRadius.h
@@ -11,8 +11,9 @@
 UCLASS(collapsecategories, hidecategories=Object)
 class UMaterialExpressionObjectRadius : public UMaterialExpression
 {
-	GENERATED_UCLASS_BODY()
+	GENERATED_BODY()
 
+public:
 
 	//~ Begin UMaterialExpression Interface
 #if WITH_EDITOR
diff --git a/Engine/Source/Runtime/Engine/Public/Materials/MaterialExpressionParticleDirection.h b/Engine/Source/Runtime/Engine/Public/Materials/MaterialExpressionParticleDirection.h
--- a/Engine/Source/Runtime/Engine/Public/Materials/MaterialExpressionParticleDirection.h
+++ b/Engine/Source/Runtime/Engine/Public/Materials/MaterialExpressionParticleDirection.h
@@ -15,8 +15,9 @@
 UCLASS(collapsecategories, hidecategories=Object)
 class UMaterialExpressionParticleDirection : public UMaterialExpression
 {
-	GENERATED_UCLASS_BODY()
+	GENERATED_BODY()
 
+public:
 
 	//~ Begin UMaterialExpression Interface
 #if WITH_EDITOR
diff --git a/Engine/Source/Runtime/Engine/Public/Materials/MaterialExpressionParticleMacroUV.h b/Engine/Source/Runtime/Engine/Public/Materials/MaterialExpressionParticleMacroUV.h
--- a/Engine/Source/Runtime/Engine/Public/Materials/MaterialExpressionParticleMacroUV.h
+++ b/Engine/Source/Runtime/Engine/Public/Materials/MaterialExpressionParticleMacroUV.h
@@ -15,8 +15,9 @@
 UCLASS()
 class UMaterialExpressionParticleMacroUV : public UMaterialExpression
 {
-	GENERATED_UCLASS_BODY()
+	GENERATED_BODY()
 
+public:
 
 	//~ Begin UMaterialExpression Interface
 #if WITH_EDITOR
diff --git a/Engine/Source/Runtime/Engine/Public/Materials/MaterialExpressionParticleMotionBlurFade.h b/Engine/Source/Runtime/Engine/Public/Materials/MaterialExpressionParticleMotionBlurFade.h
--- a/Engine/Source/Runtime/Engine/Public/Materials/MaterialExpressionParticleMotionBlurFade.h
+++ b/Engine/Source/Runtime/Engine/Public/Materials/MaterialExpressionParticleMotionBlurFade.h
@@ -15,8 +15,9 @@
 UCLASS(collapsecategories, hidecategories=Object)
 class UMaterialExpressionParticleMotionBlurFade : public UMaterialExpression
 {
-	GENERATED_UCLASS_BODY()
+	GENERATED_BODY()
 
+public:
 
 	//~ Begin UMaterialExpression Interface
 #if WITH_EDITOR
diff --git a/Engine/Source/Runtime/Engine/Public/Materials/MaterialExpressionParticlePositionWS.h b/Engine/Source/Runtime/Engine/Public/Materials/MaterialExpressionParticlePositionWS.h
--- a/Engine/Source/Runtime/Engine/Public/Materials/MaterialExpressionParticlePositionWS.h
+++ b/Engine/Source/Runtime/Engine/Public/Materials/MaterialExpressionParticlePositionWS.h
@@ -11,7 +11,9 @@
 UCLASS(collapsecategories, hidecategories=Object)
 class UMaterialExpressionParticlePositionWS : public UMaterialExpression
 {
-	GENERATED_UCLASS_BODY()
+	GENERATED_BODY()
+
+public:
 
 	UPROPERTY(EditAnywhere, Category=UMaterialExpressionParticlePositionWS, meta=(DisplayName = "Origin", ShowAsInputPin = "Advanced"))
 	EPositionOrigin OriginType;
diff --git a/Engine/Source/Runtime/Engine/Public/Materials/MaterialExpressionParticleRadius.h b/Engine/Source/Runtime/Engine/Public/Materials/MaterialExpressionParticleRadius.h
--- a/Engine/Source/Runtime/Engine/Public/Materials/MaterialExpressionParticleRadius.h
+++ b/Engine/Source/Runtime/Engine/Public/Materials/MaterialExpressionParticleRadius.h
@@ -11,8 +11,9 @@
 UCLASS(collapsecategories, hidecategories=Object)
 class UMaterialExpressionParticleRadius : public UMaterialExpression
 {
-	GENERATED_UCLASS_BODY()
+	GENERATED_BODY()
 
+public:
 
 	//~ Begin UMaterialExpression Interface
 #if WITH_EDITOR
diff --git a/Engine/Source/Runtime/Engine/Public/Materials/MaterialExpressionParticleRandom.h b/Engine/Source/Runtime/Engine/Public/Materials/MaterialExpressionParticleRandom.h
--- a/Engine/Source/Runtime/Engine/Public/Materials/MaterialExpressionParticleRandom.h
+++ b/Engine/Source/Runtime/Engine/Public/Materials/MaterialExpressionParticleRandom.h
@@ -15,8 +15,9 @@
 UCLASS(collapsecategories, hidecategories=Object)
 class UMaterialExpressionParticleRandom: public UMaterialExpression
 {
-	GENERATED_UCLASS_BODY()
+	GENERATED_BODY()
 
+public:
 
 	//~ Begin UMaterialExpression Interface
 #if WITH_EDITOR
diff --git a/Engine/Source/Runtime/Engine/Public/Materials/MaterialExpressionParticleRelativeTime.h b/Engine/Source/Runtime/Engine/Public/Materials/MaterialExpressionParticleRelativeTime.h
--- a/Engine/Source/Runtime/Engine/Public/Materials/MaterialExpressionParticleRelativeTime.h
+++ b/Engine/Source/Runtime/Engine/Public/Materials/MaterialExpressionParticleRelativeTime.h
@@ -15,8 +15,9 @@
 UCLASS(collapsecategories, hidecategories=Object)
 class UMaterialExpressionParticleRelativeTime : public UMaterialExpression
 {
-	GENERATED_UCLASS_BODY()
+	GENERATED_BODY()
 
+public:
 
 	//~ Begin UMaterialExpression Interface
 #if WITH_EDITOR
diff --git a/Engine/Source/Runtime/Engine/Public/Materials/MaterialExpressionParticleSize.h b/Engine/Source/Runtime/Engine/Public/Materials/MaterialExpressionParticleSize.h
--- a/Engine/Source/Runtime/Engine/Public/Materials/MaterialExpressionParticleSize.h
+++ b/Engine/Source/Runtime/Engine/Public/Materials/MaterialExpressionParticleSize.h
@@ -15,8 +15,9 @@
 UCLASS(collapsecategories, hidecategories=Object)
 class UMaterialExpressionParticleSize : public UMaterialExpression
 {
-	GENERATED_UCLASS_BODY()
+	GENERATED_BODY()
 
+public:
 
 	//~ Begin UMaterialExpression Interface
 #if WITH_EDITOR
diff --git a/Engine/Source/Runtime/Engine/Public/Materials/MaterialExpressionParticleSpeed.h b/Engine/Source/Runtime/Engine/Public/Materials/MaterialExpressionParticleSpeed.h
--- a/Engine/Source/Runtime/Engine/Public/Materials/MaterialExpressionParticleSpeed.h
+++ b/Engine/Source/Runtime/Engine/Public/Materials/MaterialExpressionParticleSpeed.h
@@ -15,8 +15,9 @@
 UCLASS(collapsecategories, hidecategories=Object)
 class UMaterialExpressionParticleSpeed : public UMaterialExpression
 {
-	GENERATED_UCLASS_BODY()
+	GENERATED_BODY()
 
+public:
 
 	//~ Begin UMaterialExpression Interface
 #if WITH_EDITOR
diff --git a/Engine/Source/Runtime/Engine/Public/Materials/MaterialExpressionPerInstanceCustomData.h b/Engine/Source/Runtime/Engine/Public/Materials/MaterialExpressionPerInstanceCustomData.h
--- a/Engine/Source/Runtime/Engine/Public/Materials/MaterialExpressionPerInstanceCustomData.h
+++ b/Engine/Source/Runtime/Engine/Public/Materials/MaterialExpressionPerInstanceCustomData.h
@@ -10,7 +10,9 @@
 UCLASS(collapsecategories, hidecategories=Object)
 class UMaterialExpressionPerInstanceCustomData : public UMaterialExpression
 {
-	GENERATED_UCLASS_BODY()
+	GENERATED_BODY()
+
+public:
 
 	UPROPERTY(meta = (RequiredInput = "false", ToolTip = "Defaults to 'ConstDefaultValue' if not specified; Default value is used when no instances provided."))
 	FExpressionInput DefaultValue;
@@ -33,7 +35,9 @@ class UMaterialExpressionPerInstanceCustomData : public UMaterialExpression
 UCLASS(collapsecategories, hidecategories=Object)
 class UMaterialExpressionPerInstanceCustomData3Vector : public UMaterialExpression
 {
-	GENERATED_UCLASS_BODY()
+	GENERATED_BODY()
+
+public:
 
 	UPROPERTY(meta = (RequiredInput = "false", ToolTip = "Defaults to 'ConstDefaultValue' if not specified; Default value is used when no instances provided."))
 	FExpressionInput DefaultValue;
diff --git a/Engine/Source/Runtime/Engine/Public/Materials/MaterialExpressionPerInstanceFadeAmount.h b/Engine/Source/Runtime/Engine/Public/Materials/MaterialExpressionPerInstanceFadeAmount.h
--- a/Engine/Source/Runtime/Engine/Public/Materials/MaterialExpressionPerInstanceFadeAmount.h
+++ b/Engine/Source/Runtime/Engine/Public/Materials/MaterialExpressionPerInstanceFadeAmount.h
@@ -11,8 +11,9 @@
 UCLASS(collapsecategories, hidecategories=Object)
 class UMaterialExpressionPerInstanceFadeAmount : public UMaterialExpression
 {
-	GENERATED_UCLASS_BODY()
+	GENERATED_BODY()
 
+public:
 
 	//~ Begin UMaterialExpression Interface
 #if WITH_EDITOR
diff --git a/Engine/Source/Runtime/Engine/Public/Materials/MaterialExpressionPerInstanceRandom.h b/Engine/Source/Runtime/Engine/Public/Materials/MaterialExpressionPerInstanceRandom.h
--- a/Engine/Source/Runtime/Engine/Public/Materials/MaterialExpressionPerInstanceRandom.h
+++ b/Engine/Source/Runtime/Engine/Public/Materials/MaterialExpressionPerInstanceRandom.h
@@ -11,8 +11,9 @@
 UCLASS(collapsecategories, hidecategories=Object)
 class UMaterialExpressionPerInstanceRandom : public UMaterialExpression
 {
-	GENERATED_UCLASS_BODY()
+	GENERATED_BODY()
 
+public:
 
 	//~ Begin UMaterialExpression Interface
 #if WITH_EDITOR
diff --git a/Engine/Source/Runtime/Engine/Public/Materials/MaterialExpressionPixelNormalWS.h b/Engine/Source/Runtime/Engine/Public/Materials/MaterialExpressionPixelNormalWS.h
--- a/Engine/Source/Runtime/Engine/Public/Materials/MaterialExpressionPixelNormalWS.h
+++ b/Engine/Source/Runtime/Engine/Public/Materials/MaterialExpressionPixelNormalWS.h
@@ -11,8 +11,9 @@
 UCLASS(MinimalAPI, collapsecategories, hidecategories=Object)
 class UMaterialExpressionPixelNormalWS : public UMaterialExpression
 {
-	GENERATED_UCLASS_BODY()
+	GENERATED_BODY()
 
+public:
 
 	//~ Begin UMaterialExpression Interface
 #if WITH_EDITOR
diff --git a/Engine/Source/Runtime/Engine/Public/Materials/MaterialExpressionReflectionVectorWS.h b/Engine/Source/Runtime/Engine/Public/Materials/MaterialExpressionReflectionVectorWS.h
--- a/Engine/Source/Runtime/Engine/Public/Materials/MaterialExpressionReflectionVectorWS.h
+++ b/Engine/Source/Runtime/Engine/Public/Materials/MaterialExpressionReflectionVectorWS.h
@@ -12,7 +12,9 @@
 UCLASS(collapsecategories, hidecategories=Object, MinimalAPI)
 class UMaterialExpressionReflectionVectorWS : public UMaterialExpression
 {
-	GENERATED_UCLASS_BODY()
+	GENERATED_BODY()
+
+public:
 
 	/** Optional world normal to reflect the camera view vector about. If unconnected, pixel normal is used */
 	UPROPERTY(meta = (RequiredInput = "false", ToolTip = "Defaults to Reflection Vector if not specified"))
diff --git a/Engine/Source/Runtime/Engine/Public/Materials/MaterialExpressionSceneTexelSize.h b/Engine/Source/Runtime/Engine/Public/Materials/MaterialExpressionSceneTexelSize.h
--- a/Engine/Source/Runtime/Engine/Public/Materials/MaterialExpressionSceneTexelSize.h
+++ b/Engine/Source/Runtime/Engine/Public/Materials/MaterialExpressionSceneTexelSize.h
@@ -11,7 +11,9 @@
 UCLASS(collapsecategories, hidecategories=Object)
 class UMaterialExpressionSceneTexelSize : public UMaterialExpression
 {
-	GENERATED_UCLASS_BODY()
+	GENERATED_BODY()
+
+public:
 
 	//~ Begin UMaterialExpression Interface
 #if WITH_EDITOR
diff --git a/Engine/Source/Runtime/Engine/Public/Materials/MaterialExpressionTwoSidedSign.h b/Engine/Source/Runtime/Engine/Public/Materials/MaterialExpressionTwoSidedSign.h
--- a/Engine/Source/Runtime/Engine/Public/Materials/MaterialExpressionTwoSidedSign.h
+++ b/Engine/Source/Runtime/Engine/Public/Materials/MaterialExpressionTwoSidedSign.h
@@ -11,8 +11,9 @@
 UCLASS(collapsecategories, hidecategories=Object, MinimalAPI)
 class UMaterialExpressionTwoSidedSign : public UMaterialExpression
 {
-	GENERATED_UCLASS_BODY()
+	GENERATED_BODY()
 
+public:
 
 	//~ Begin UMaterialExpression Interface
 #if WITH_EDITOR
diff --git a/Engine/Source/Runtime/Engine/Public/Materials/MaterialExpressionVertexNormalWS.h b/Engine/Source/Runtime/Engine/Public/Materials/MaterialExpressionVertexNormalWS.h
--- a/Engine/Source/Runtime/Engine/Public/Materials/MaterialExpressionVertexNormalWS.h
+++ b/Engine/Source/Runtime/Engine/Public/Materials/MaterialExpressionVertexNormalWS.h
@@ -11,8 +11,9 @@
 UCLASS(MinimalAPI, collapsecategories, hidecategories=Object)
 class UMaterialExpressionVertexNormalWS : public UMaterialExpression
 {
-	GENERATED_UCLASS_BODY()
+	GENERATED_BODY()
 
+public:
 
 	//~ Begin UMaterialExpression Interface
 #if WITH_EDITOR
diff --git a/Engine/Source/Runtime/Engine/Public/Materials/MaterialExpressionVertexTangentWS.h b/Engine/Source/Runtime/Engine/Public/Materials/MaterialExpressionVertexTangentWS.h
--- a/Engine/Source/Runtime/Engine/Public/Materials/MaterialExpressionVertexTangentWS.h
+++ b/Engine/Source/Runtime/Engine/Public/Materials/MaterialExpressionVertexTangentWS.h
@@ -11,8 +11,9 @@
 UCLASS(MinimalAPI, collapsecategories, hidecategories=Object)
 class UMaterialExpressionVertexTangentWS : public UMaterialExpression
 {
-	GENERATED_UCLASS_BODY()
+	GENERATED_BODY()
 
+public:
 
 	//~ Begin UMaterialExpression Interface
 #if WITH_EDITOR
diff --git a/Engine/Source/Runtime/Engine/Public/Materials/MaterialExpressionViewSize.h b/Engine/Source/Runtime/Engine/Public/Materials/MaterialExpressionViewSize.h
--- a/Engine/Source/Runtime/Engine/Public/Materials/MaterialExpressionViewSize.h
+++ b/Engine/Source/Runtime/Engine/Public/Materials/MaterialExpressionViewSize.h
@@ -11,7 +11,9 @@
 UCLASS(collapsecategories, hidecategories=Object)
 class UMaterialExpressionViewSize : public UMaterialExpression
 {
-	GENERATED_UCLASS_BODY()
+	GENERATED_BODY()
+
+public:
 
 	//~ Begin UMaterialExpression Interface
 #if WITH_EDITOR
