Add new CVar r.ForceStartInVR which overrides the project setting (or lack of -vr commandline argument) and forces the application to attempt to start in VR. Add r.ForceStartInVR=1 to the Android_XR device profile.

diff --git a/Engine/Config/BaseDeviceProfiles.ini b/Engine/Config/BaseDeviceProfiles.ini
--- a/Engine/Config/BaseDeviceProfiles.ini
+++ b/Engine/Config/BaseDeviceProfiles.ini
@@ -1208,6 +1208,7 @@ BaseProfileName=Android_Mid
 +CVars=vr.MobileMultiView=1
 +CVars=r.MobileHDR=0
 +CVars=r.Mobile.AntiAliasing=3
++CVars=r.ForceStartInVR=1
 
 [Meta_Quest_3 DeviceProfile]
 DeviceType=Android
diff --git a/Engine/Source/Runtime/Engine/Private/StereoRendering.cpp b/Engine/Source/Runtime/Engine/Private/StereoRendering.cpp
--- a/Engine/Source/Runtime/Engine/Private/StereoRendering.cpp
+++ b/Engine/Source/Runtime/Engine/Private/StereoRendering.cpp
@@ -6,6 +6,11 @@
 #include "Misc/CommandLine.h"
 #include "Misc/ConfigCacheIni.h"
 
+static TAutoConsoleVariable<int32> CVarForceStartInVR(
+	TEXT("r.ForceStartInVR"),
+	0,
+	TEXT("If true, the game will attempt to start in VR, regardless of whether \"Start in VR\" is true in the project settings or -vr was set on the commandline."),
+	ECVF_ReadOnly);
 
 bool IStereoRendering::IsStereoEyeView(const FSceneView& View)
 {
@@ -26,7 +31,11 @@ bool IStereoRendering::IsStartInVR()
 {
 	bool bStartInVR = false;
 
-	if (IsClassLoaded<UGeneralProjectSettings>())
+	if (CVarForceStartInVR.GetValueOnAnyThread() > 0)
+	{
+		bStartInVR = true;
+	}
+	else if (IsClassLoaded<UGeneralProjectSettings>())
 	{
 		bStartInVR = GetDefault<UGeneralProjectSettings>()->bStartInVR;
 	}
