- Incorporate cooked asset check into Is(Bone/Curve)DataValid checks
- Incorporate cooked asset check into platform specific compressed animation data fetching

diff --git a/Engine/Source/Runtime/Engine/Private/Animation/AnimCompressionTypes.cpp b/Engine/Source/Runtime/Engine/Private/Animation/AnimCompressionTypes.cpp
--- a/Engine/Source/Runtime/Engine/Private/Animation/AnimCompressionTypes.cpp
+++ b/Engine/Source/Runtime/Engine/Private/Animation/AnimCompressionTypes.cpp
@@ -25,6 +25,7 @@
 #include "Animation/SkeletonRemapping.h"
 #include "Animation/SkeletonRemappingRegistry.h"
 #include "Animation/AttributesRuntime.h"
+#include "UObject/Package.h"
 
 static FString GCompressionJsonOutput;
 static FAutoConsoleVariableRef CVarCompressionJsonOutput(
@@ -1614,7 +1615,7 @@ bool FCompressedAnimSequence::IsValid(const UAnimSequence* AnimSequence, bool bL
 bool FCompressedAnimSequence::IsBoneDataValid(const UAnimSequence* AnimSequence, bool bLogInformation) const
 {
 #if WITH_EDITOR
-	const bool bHasBoneTracks = AnimSequence->GetDataModelInterface()->GetNumBoneTracks() != 0;
+	const bool bHasBoneTracks = !AnimSequence->GetOutermost()->HasAnyPackageFlags(PKG_Cooked) && AnimSequence->GetDataModelInterface()->GetNumBoneTracks() != 0;
 	const bool bHasValidCompressedBoneData = CompressedDataStructure != nullptr || CompressedTrackToSkeletonMapTable.Num() == 0;
 	const bool bValidCompressedData = !bHasBoneTracks || bHasValidCompressedBoneData;
 	if(bLogInformation && !bValidCompressedData)
@@ -1630,7 +1631,7 @@ bool FCompressedAnimSequence::IsBoneDataValid(const UAnimSequence* AnimSequence,
 bool FCompressedAnimSequence::IsCurveDataValid(const UAnimSequence* AnimSequence, bool bLogInformation) const
 {
 #if WITH_EDITOR
-	const bool bHasCurveData = AnimSequence->GetDataModelInterface()->GetNumberOfFloatCurves() != 0;
+	const bool bHasCurveData = !AnimSequence->GetOutermost()->HasAnyPackageFlags(PKG_Cooked) && AnimSequence->GetDataModelInterface()->GetNumberOfFloatCurves() != 0;
 	const bool bHasValidCompressedCurveData = CompressedCurveByteStream.Num() != 0;
 	const bool bValidCompressedData = (!bHasCurveData || bHasValidCompressedCurveData || AnimSequence->IsValidAdditive());
 
diff --git a/Engine/Source/Runtime/Engine/Private/Animation/AnimSequence.cpp b/Engine/Source/Runtime/Engine/Private/Animation/AnimSequence.cpp
--- a/Engine/Source/Runtime/Engine/Private/Animation/AnimSequence.cpp
+++ b/Engine/Source/Runtime/Engine/Private/Animation/AnimSequence.cpp
@@ -763,71 +763,72 @@ void UAnimSequence::Serialize(FArchive& Ar)
 			else
 			{
 #if WITH_EDITOR
-				const bool bIsLoading = Ar.IsLoading();
-				const bool bIsSaving = Ar.IsSaving();
-
-				if (bIsLoading)
-				{
-					CompressedDataLock.WriteLock();
-				}
-				else
+				if (!GetOutermost()->HasAnyPackageFlags(PKG_Cooked))
 				{
-					CompressedDataLock.ReadLock();
-				}
-				
-				if (bCookingTargetNeedsCompressedData)
-				{
-					const ITargetPlatform* Platform = Ar.CookingTarget();
-					checkf(Platform, TEXT("Invalid cooking target platform"));
-					// During cook, serialize out singular (target platform) compressed data 
-					FCompressedAnimSequence& PlatformCompressedData = GetPlatformCompressedData(Platform);
-					check(PlatformCompressedData.IsValid(this, true));
-					SerializeCompressedData(Ar, PlatformCompressedData);
-				}
-				else if(bIsDuplicating || bIsTransacting || bIsCountingMemory)
-				{
-					// During duplication/transacting etc. serialize all contained compressed data
-					if (bIsSaving)
+					const bool bIsLoading = Ar.IsLoading();
+					const bool bIsSaving = Ar.IsSaving();
+
+					if (bIsLoading)
 					{
-						uint32 NumPlatforms = DataByPlatformKeyHash.Num();
-						Ar << NumPlatforms;
-						for (TPair<FIoHash, TUniquePtr<FCompressedAnimSequence>>& CompressedDataPair : DataByPlatformKeyHash)
-						{
-							Ar << CompressedDataPair.Key;
-							SerializeCompressedData(Ar, *CompressedDataPair.Value.Get());
-						}
+						CompressedDataLock.WriteLock();
 					}
 					else
 					{
-						uint32 NumPlatforms = 0;
-						Ar << NumPlatforms;
-						for (uint32 Index = 0; Index < NumPlatforms; ++Index)
+						CompressedDataLock.ReadLock();
+					}
+
+					if (bCookingTargetNeedsCompressedData)
+					{
+						const ITargetPlatform* Platform = Ar.CookingTarget();
+						checkf(Platform, TEXT("Invalid cooking target platform"));
+						// During cook, serialize out singular (target platform) compressed data 
+						FCompressedAnimSequence& PlatformCompressedData = GetPlatformCompressedData(Platform);
+						check(PlatformCompressedData.IsValid(this, true));
+						SerializeCompressedData(Ar, PlatformCompressedData);
+					}
+					else if (bIsDuplicating || bIsTransacting || bIsCountingMemory)
+					{
+						// During duplication/transacting etc. serialize all contained compressed data
+						if (bIsSaving)
 						{
-							FIoHash PlatformHash;
-							Ar << PlatformHash;
-							FCompressedAnimSequence& TargetData = *DataByPlatformKeyHash.Emplace(PlatformHash, MakeUnique<FCompressedAnimSequence>()).Get();
-							SerializeCompressedData(Ar, TargetData);
+							uint32 NumPlatforms = DataByPlatformKeyHash.Num();
+							Ar << NumPlatforms;
+							for (TPair<FIoHash, TUniquePtr<FCompressedAnimSequence>>& CompressedDataPair : DataByPlatformKeyHash)
+							{
+								Ar << CompressedDataPair.Key;
+								SerializeCompressedData(Ar, *CompressedDataPair.Value.Get());
+							}
 						}
-					}					
-				}
-				
-				if (bIsLoading)
-				{
-					CompressedDataLock.WriteUnlock();
+						else
+						{
+							uint32 NumPlatforms = 0;
+							Ar << NumPlatforms;
+							for (uint32 Index = 0; Index < NumPlatforms; ++Index)
+							{
+								FIoHash PlatformHash;
+								Ar << PlatformHash;
+								FCompressedAnimSequence& TargetData = *DataByPlatformKeyHash.Emplace(PlatformHash, MakeUnique<FCompressedAnimSequence>()).Get();
+								SerializeCompressedData(Ar, TargetData);
+							}
+						}
+					}
+
+					if (bIsLoading)
+					{
+						CompressedDataLock.WriteUnlock();
+					}
+					else
+					{
+						CompressedDataLock.ReadUnlock();
+					}
 				}
 				else
-				{
-					CompressedDataLock.ReadUnlock();
-				}
-#else
+#endif
 				{
 					// During cooked runtime serialize in singular compressed data
 					SerializeCompressedData(Ar, GetCompressedData_Internal());
 				}
-#endif
-			}
-			
-			
+			}			
 		}
 	}
 
@@ -1940,154 +1941,170 @@ bool UAnimSequence::CanEvaluateRawAnimationData() const
 const FCompressedAnimSequence& UAnimSequence::GetCompressedData_Internal() const
 {
 #if WITH_EDITOR
-	static thread_local FCompressedAnimSequence DefaultCompressedAnimSequence;
-	const ITargetPlatform* RunningPlatform = GetTargetPlatformManagerRef().GetRunningTargetPlatform();
-
-	if(DataByPlatformKeyHash.Num())
+	if(!GetOutermost()->HasAnyPackageFlags(PKG_Cooked))
 	{
-		const FIoHash CurrentPlatformHash = GetDerivedDataKeyHash(RunningPlatform);
-		if (const TUniquePtr<FCompressedAnimSequence>* PlatformCompressedData = DataByPlatformKeyHash.Find(CurrentPlatformHash))
+		static thread_local FCompressedAnimSequence DefaultCompressedAnimSequence;
+		const ITargetPlatform* RunningPlatform = GetTargetPlatformManagerRef().GetRunningTargetPlatform();
+
+		if(DataByPlatformKeyHash.Num())
 		{
-			return *PlatformCompressedData->Get();
+			const FIoHash CurrentPlatformHash = GetDerivedDataKeyHash(RunningPlatform);
+			if (const TUniquePtr<FCompressedAnimSequence>* PlatformCompressedData = DataByPlatformKeyHash.Find(CurrentPlatformHash))
+			{
+				return *PlatformCompressedData->Get();
+			}
 		}
+
+		return DefaultCompressedAnimSequence;
 	}
+#endif // WITH_EDITOR
 
-	return DefaultCompressedAnimSequence;
-#else	
 	PRAGMA_DISABLE_DEPRECATION_WARNINGS
 	return CompressedData;
 	PRAGMA_ENABLE_DEPRECATION_WARNINGS
-#endif // WITH_EDITOR
 }
 
 
 const FCompressedAnimSequence& UAnimSequence::GetPlatformCompressedData(const FAnimExtractContext& AnimExtractContext) const
 {
 #if WITH_EDITOR
-	FIoHash Hash = FIoHash::Zero;	
-	static thread_local FCompressedAnimSequence DefaultCompressedAnimSequence;	
-	if (AnimExtractContext.bEnforceCompressedDataSampling)
+	if(!GetOutermost()->HasAnyPackageFlags(PKG_Cooked))
 	{
-		if(AnimExtractContext.TargetPlatformHash.IsZero())
+		FIoHash Hash = FIoHash::Zero;	
+		static thread_local FCompressedAnimSequence DefaultCompressedAnimSequence;	
+		if (AnimExtractContext.bEnforceCompressedDataSampling)
 		{
-			// Not set, check platform instead
-			if (AnimExtractContext.TargetPlatform)
+			if(AnimExtractContext.TargetPlatformHash.IsZero())
 			{
-				Hash = GetDerivedDataKeyHash(AnimExtractContext.TargetPlatform);
+				// Not set, check platform instead
+				if (AnimExtractContext.TargetPlatform)
+				{
+					Hash = GetDerivedDataKeyHash(AnimExtractContext.TargetPlatform);
+				}
+				else
+				{
+					const ITargetPlatform* RunningPlatform = GetTargetPlatformManagerRef().GetRunningTargetPlatform();
+					Hash = GetDerivedDataKeyHash(RunningPlatform);
+				}
 			}
 			else
 			{
-				const ITargetPlatform* RunningPlatform = GetTargetPlatformManagerRef().GetRunningTargetPlatform();
-				Hash = GetDerivedDataKeyHash(RunningPlatform);
+				Hash = AnimExtractContext.TargetPlatformHash;
 			}
 		}
 		else
 		{
-			Hash = AnimExtractContext.TargetPlatformHash;
+			const ITargetPlatform* RunningPlatform = GetTargetPlatformManagerRef().GetRunningTargetPlatform();
+			Hash = GetDerivedDataKeyHash(RunningPlatform);
 		}
-	}
-	else
-	{
-		const ITargetPlatform* RunningPlatform = GetTargetPlatformManagerRef().GetRunningTargetPlatform();
-		Hash = GetDerivedDataKeyHash(RunningPlatform);
-	}
 	
-	if (const TUniquePtr<FCompressedAnimSequence>* PlatformData = DataByPlatformKeyHash.Find(Hash))
-	{
-		return *PlatformData->Get();
+		if (const TUniquePtr<FCompressedAnimSequence>* PlatformData = DataByPlatformKeyHash.Find(Hash))
+		{
+			return *PlatformData->Get();
+		}
+
+		return DefaultCompressedAnimSequence;
 	}
+#endif // WITH_EDITOR
 
-	return DefaultCompressedAnimSequence;
-#else
 	PRAGMA_DISABLE_DEPRECATION_WARNINGS
 	return CompressedData;
 	PRAGMA_ENABLE_DEPRECATION_WARNINGS
-#endif // WITH_EDITOR
 }
 
 FCompressedAnimSequence& UAnimSequence::GetPlatformCompressedData(const FAnimExtractContext& AnimExtractContext)
 {	
 #if WITH_EDITOR
-	FIoHash Hash = FIoHash::Zero;	
-	static thread_local FCompressedAnimSequence DefaultCompressedAnimSequence;	
-	if (AnimExtractContext.bEnforceCompressedDataSampling)
-	{		
-		if(AnimExtractContext.TargetPlatformHash.IsZero())
-		{
-			// Not set, check platform instead
-			if (AnimExtractContext.TargetPlatform)
+	if(!GetOutermost()->HasAnyPackageFlags(PKG_Cooked))
+	{
+		FIoHash Hash = FIoHash::Zero;	
+		static thread_local FCompressedAnimSequence DefaultCompressedAnimSequence;	
+		if (AnimExtractContext.bEnforceCompressedDataSampling)
+		{		
+			if(AnimExtractContext.TargetPlatformHash.IsZero())
 			{
-				Hash = GetDerivedDataKeyHash(AnimExtractContext.TargetPlatform);
+				// Not set, check platform instead
+				if (AnimExtractContext.TargetPlatform)
+				{
+					Hash = GetDerivedDataKeyHash(AnimExtractContext.TargetPlatform);
+				}
+				else
+				{
+					const ITargetPlatform* RunningPlatform = GetTargetPlatformManagerRef().GetRunningTargetPlatform();
+					Hash = GetDerivedDataKeyHash(RunningPlatform);
+				}
 			}
 			else
 			{
-				const ITargetPlatform* RunningPlatform = GetTargetPlatformManagerRef().GetRunningTargetPlatform();
-				Hash = GetDerivedDataKeyHash(RunningPlatform);
+				Hash = AnimExtractContext.TargetPlatformHash;
 			}
 		}
 		else
 		{
-			Hash = AnimExtractContext.TargetPlatformHash;
+			const ITargetPlatform* RunningPlatform = GetTargetPlatformManagerRef().GetRunningTargetPlatform();
+			Hash = GetDerivedDataKeyHash(RunningPlatform);
 		}
+		
+		if (TUniquePtr<FCompressedAnimSequence>* PlatformData = DataByPlatformKeyHash.Find(Hash))
+		{
+			return *PlatformData->Get();
+		}
+
+		return DefaultCompressedAnimSequence;
 	}
-	else
-	{
-		const ITargetPlatform* RunningPlatform = GetTargetPlatformManagerRef().GetRunningTargetPlatform();
-		Hash = GetDerivedDataKeyHash(RunningPlatform);
-	}
-	
-	if (TUniquePtr<FCompressedAnimSequence>* PlatformData = DataByPlatformKeyHash.Find(Hash))
-	{
-		return *PlatformData->Get();
-	}
+#endif // WITH_EDITOR
 
-	return DefaultCompressedAnimSequence;
-#else
 	PRAGMA_DISABLE_DEPRECATION_WARNINGS
 	return CompressedData;
 	PRAGMA_ENABLE_DEPRECATION_WARNINGS
-#endif // WITH_EDITOR
 }
 
 
 FCompressedAnimSequence& UAnimSequence::GetPlatformCompressedData(const ITargetPlatform* InTargetPlatform)
 {	
 #if WITH_EDITOR
-	thread_local FCompressedAnimSequence DefaultCompressedAnimSequence;
-	
-	check(InTargetPlatform);
-	
-	const FIoHash Hash = GetDerivedDataKeyHash(InTargetPlatform);
-	check(!Hash.IsZero());
-
-	if (const TUniquePtr<FCompressedAnimSequence>* PlatformData = DataByPlatformKeyHash.Find(Hash))
+	if(!GetOutermost()->HasAnyPackageFlags(PKG_Cooked))
 	{
-		return *PlatformData->Get();
-	}
+		thread_local FCompressedAnimSequence DefaultCompressedAnimSequence;
+		
+		check(InTargetPlatform);
+		
+		const FIoHash Hash = GetDerivedDataKeyHash(InTargetPlatform);
+		check(!Hash.IsZero());
 
-	return DefaultCompressedAnimSequence;
+		if (const TUniquePtr<FCompressedAnimSequence>* PlatformData = DataByPlatformKeyHash.Find(Hash))
+		{
+			return *PlatformData->Get();
+		}
 
-#else
-	return GetCompressedData_Internal();
+		return DefaultCompressedAnimSequence;
+	}
 #endif // WITH_EDITOR
+	
+	return GetCompressedData_Internal();
 }
 
 #if WITH_EDITOR
 const FCompressedAnimSequence& UAnimSequence::GetPlatformCompressedData(const ITargetPlatform* InTargetPlatform) const
 {
-	thread_local FCompressedAnimSequence DefaultCompressedAnimSequence;
-	
-	check(InTargetPlatform);
-	
-	const FIoHash Hash = GetDerivedDataKeyHash(InTargetPlatform);
-	check(!Hash.IsZero());
-
-	if (const TUniquePtr<FCompressedAnimSequence>* PlatformData = DataByPlatformKeyHash.Find(Hash))
+	if(!GetOutermost()->HasAnyPackageFlags(PKG_Cooked))
 	{
-		return *PlatformData->Get();
+		thread_local FCompressedAnimSequence DefaultCompressedAnimSequence;
+		
+		check(InTargetPlatform);
+		
+		const FIoHash Hash = GetDerivedDataKeyHash(InTargetPlatform);
+		check(!Hash.IsZero());
+
+		if (const TUniquePtr<FCompressedAnimSequence>* PlatformData = DataByPlatformKeyHash.Find(Hash))
+		{
+			return *PlatformData->Get();
+		}
+
+		return DefaultCompressedAnimSequence;
 	}
 
-	return DefaultCompressedAnimSequence;
+	return GetCompressedData_Internal();
 }
 #endif // WITH_EDITOR
 
