[WorldPartitionRuntimeHashSet] Fix source queries not respecting bSpatialQuery parameter. In some cases Location and Radius were still used to find cells, despite bSpatialQuery set to False.

diff --git a/Engine/Source/Runtime/Engine/Public/WorldPartition/WorldPartitionStreamingSource.h b/Engine/Source/Runtime/Engine/Public/WorldPartition/WorldPartitionStreamingSource.h
--- a/Engine/Source/Runtime/Engine/Public/WorldPartition/WorldPartitionStreamingSource.h
+++ b/Engine/Source/Runtime/Engine/Public/WorldPartition/WorldPartitionStreamingSource.h
@@ -307,13 +307,16 @@ struct FWorldPartitionStreamingQuerySource
 	/** Helper method that iterates over all shapes. If none is provided, it will still pass a sphere shape using Radius or grid's loading range (see bUseGridLoadingRange). */
 	FORCEINLINE void ForEachShape(float InGridLoadingRange, bool bInProjectIn2D, TFunctionRef<void(const FSphericalSector&)> InOperation) const
 	{
-		FStreamingSourceShapeHelper::ForEachShape(InGridLoadingRange, bUseGridLoadingRange ? InGridLoadingRange : Radius, bInProjectIn2D, Location, Rotation, Shapes, InOperation);
+		if (bSpatialQuery)
+		{
+			FStreamingSourceShapeHelper::ForEachShape(InGridLoadingRange, bUseGridLoadingRange ? InGridLoadingRange : Radius, bInProjectIn2D, Location, Rotation, Shapes, InOperation);
+		}
 	}
 
 	/** Helper method that iterates over all shapes. If none is provided, it will still pass a sphere shape using Radius or grid's loading range (see bUseGridLoadingRange). */
 	FORCEINLINE void ForEachShape(float InGridLoadingRange, FName InGridName, bool bInProjectIn2D, TFunctionRef<void(const FSphericalSector&)> InOperation) const
 	{
-		if (bSpatialQuery && FStreamingSourceShapeHelper::IsSourceAffectingGrid(TargetGrids, TargetBehavior, InGridName))
+		if (FStreamingSourceShapeHelper::IsSourceAffectingGrid(TargetGrids, TargetBehavior, InGridName))
 		{
 			ForEachShape(InGridLoadingRange, bInProjectIn2D, InOperation);
 		}
