Refactor operator validation and folding in the new material translator

This change reworks the way operators are validated and folded in the new material translator. It now uses an internal declarative system to define the prototypes of all the operators supported.

- Introduced two new enums, internal to the emitter:
- A parameter type filter (bitflag) specifying requirements for each operator's parameter.
- An abstract enum specifying the result type.
- This approach covers all current operators and adds support for `dot` and `cross`.
- Operator validation is now exhaustive and more robust.
- Additionally, the operator folding logic has been streamlined. The new `TryFoldOperator` entry point handles folding as follows:
- Ad hoc folding for complex operators like `dot` and `cross`.
- Component-wise folding for other operators that support it.

#rnx

diff --git a/Engine/Source/Runtime/Engine/Private/Materials/MaterialExpressionsIR.cpp b/Engine/Source/Runtime/Engine/Private/Materials/MaterialExpressionsIR.cpp
--- a/Engine/Source/Runtime/Engine/Private/Materials/MaterialExpressionsIR.cpp
+++ b/Engine/Source/Runtime/Engine/Private/Materials/MaterialExpressionsIR.cpp
@@ -633,6 +633,16 @@ void UMaterialExpressionMin::Build(MIR::FEmitter& Emitter)
 	BuildBinaryOperatorFromInputWithDefault(Emitter, MIR::BO_Min, &A, ConstA, &B, ConstB);
 }
 
+void UMaterialExpressionDotProduct::Build(UE::MIR::FEmitter& Emitter)
+{
+	BuildBinaryOperatorFromInput(Emitter, MIR::BO_Dot, &A, &B);
+}
+
+void UMaterialExpressionCrossProduct::Build(UE::MIR::FEmitter& Emitter)
+{
+	BuildBinaryOperatorFromInput(Emitter, MIR::BO_Cross, &A, &B);
+}
+
 void UMaterialExpressionOneMinus::Build(MIR::FEmitter& Emitter)
 {
 	// Default input to zero if not connected, then get it as a primitive.
@@ -730,7 +740,7 @@ void UMaterialExpressionIf::Build(MIR::FEmitter& Emitter)
 	UE_MIR_CHECKPOINT();
 
 	// Emit the comparison expressions.
-	MIR::FValue* ALessThanBConditionValue = Emitter.EmitOperator(MIR::BO_LowerThan, AValue, BValue);
+	MIR::FValue* ALessThanBConditionValue = Emitter.EmitOperator(MIR::BO_LessThan, AValue, BValue);
 	MIR::FValue* AEqualsBConditionValue = Emitter.EmitOperator(MIR::BO_Equals, AValue, BValue);
 
 	// And finally emit the full conditional expression.
diff --git a/Engine/Source/Runtime/Engine/Private/Materials/MaterialIR.cpp b/Engine/Source/Runtime/Engine/Private/Materials/MaterialIR.cpp
--- a/Engine/Source/Runtime/Engine/Private/Materials/MaterialIR.cpp
+++ b/Engine/Source/Runtime/Engine/Private/Materials/MaterialIR.cpp
@@ -401,11 +401,6 @@ bool IsTernaryOperator(EOperator Op)
 	return Op >= TO_FirstTernaryOperator;
 }
 
-bool IsBinaryInfixOperator(EOperator Op)
-{
-	return Op >= BO_Add && Op <= BO_Xor;
-}
-
 const TCHAR* OperatorToString(EOperator Op)
 {
 	switch (Op)
@@ -433,8 +428,8 @@ const TCHAR* OperatorToString(EOperator Op)
 		/* Binary Comparisons */
 		case BO_GreaterThan: return TEXT("GreaterThan");
 		case BO_GreaterThanOrEquals: return TEXT("GreaterThanOrEquals");
-		case BO_LowerThan: return TEXT("LowerThan");
-		case BO_LowerThanOrEquals: return TEXT("LowerThanOrEquals");
+		case BO_LessThan: return TEXT("LowerThan");
+		case BO_LessThanOrEquals: return TEXT("LowerThanOrEquals");
 		case BO_Equals: return TEXT("Equals");
 		case BO_NotEquals: return TEXT("NotEquals");
 	
diff --git a/Engine/Source/Runtime/Engine/Private/Materials/MaterialIREmitter.cpp b/Engine/Source/Runtime/Engine/Private/Materials/MaterialIREmitter.cpp
--- a/Engine/Source/Runtime/Engine/Private/Materials/MaterialIREmitter.cpp
+++ b/Engine/Source/Runtime/Engine/Private/Materials/MaterialIREmitter.cpp
@@ -18,7 +18,7 @@ namespace UE::MIR
 const TCHAR* VectorComponentToString(EVectorComponent Component)
 {
 	static const TCHAR* Strings[] = { TEXT("x"), TEXT("y"), TEXT("z"), TEXT("w") };
-	return Strings[(int)Component];
+	return Strings[(int32)Component];
 }
 
 FSwizzleMask::FSwizzleMask(EVectorComponent X)
@@ -73,7 +73,7 @@ struct FEmitter::FPrivate
 	}
 
 	// Allocates a block of memory using the emitter's allocator.
-	static uint8* Allocate(FEmitter* Emitter, int Size, int Alignment)
+	static uint8* Allocate(FEmitter* Emitter, int32 Size, int32 Alignment)
 	{
 		uint8* Bytes = (uint8*) FMemory::Malloc(Size, Alignment); // Emitter->Module->Allocator.PushBytes(Size, Alignment);
 		FMemory::Memzero(Bytes, Size);
@@ -93,8 +93,8 @@ static FDimensional* NewDimensionalValue(FEmitter* Emitter, FPrimitiveTypePtr Ty
 {
 	check(!Type->IsScalar());
 
-	int Dimensions = Type->NumRows * Type->NumColumns;
-	int SizeInBytes = sizeof(FDimensional) + sizeof(FValue*) * Dimensions;
+	int32 Dimensions = Type->NumRows * Type->NumColumns;
+	int32 SizeInBytes = sizeof(FDimensional) + sizeof(FValue*) * Dimensions;
 
 	uint8* Bytes = FEmitter::FPrivate::Allocate(Emitter, SizeInBytes, alignof(FDimensional));
 
@@ -149,7 +149,7 @@ static FValue* EmitPrototype(FEmitter* Emitter, const TValueType& Prototype)
 }
 
 // Finds the expression input index. Although the implementation has O(n) complexity, it is only used for error reporting.
-static int SlowFindExpressionInputIndex(UMaterialExpression* Expression, const FExpressionInput* InInput)
+static int32 SlowFindExpressionInputIndex(UMaterialExpression* Expression, const FExpressionInput* InInput)
 {
 	for (FExpressionInputIterator It{ Expression }; It; ++It)
 	{
@@ -165,7 +165,7 @@ static int SlowFindExpressionInputIndex(UMaterialExpression* Expression, const F
 // Finds the expression input name. Although the implementation has O(n) complexity, it is only used for error reporting.
 static FName SlowFindInputName(UMaterialExpression* Expression, const FExpressionInput* InInput)
 {
-	int InputIndex = SlowFindExpressionInputIndex(Expression, InInput);
+	int32 InputIndex = SlowFindExpressionInputIndex(Expression, InInput);
 	return Expression->GetInputName(InputIndex);
 }
 
@@ -290,7 +290,7 @@ FValue* FEmitter::InputTexture(const FExpressionInput* InInput)
 	return Value;
 }
 
-FEmitter& FEmitter::Output(int OutputIndex, FValue* Value)
+FEmitter& FEmitter::Output(int32 OutputIndex, FValue* Value)
 {
 	Output(Expression->GetOutput(OutputIndex), Value);
 	return *this;
@@ -535,7 +535,7 @@ FValue* FEmitter::EmitExternalInput(EExternalInput Id)
 	return EmitPrototype(this, Prototype);
 }
 
-FValue* FEmitter::EmitSubscript(FValue* Value, int Index)
+FValue* FEmitter::EmitSubscript(FValue* Value, int32 Index)
 {
 	FPrimitiveTypePtr PrimitiveType = Value->Type->AsPrimitive();
 	if (!PrimitiveType)
@@ -587,7 +587,7 @@ FValue* FEmitter::EmitSwizzle(FValue* Value, FSwizzleMask Mask)
 	// Make sure each component in the mask fits the number of components in Value.
 	for (EVectorComponent Component : Mask)
 	{
-		if ((int)Component >= PrimitiveType->NumRows)
+		if ((int32)Component >= PrimitiveType->NumRows)
 		{
 			Errorf(TEXT("Value of type `%s` has no component `%s`."), PrimitiveType->Spelling.GetData(), VectorComponentToString(Component));
 			return MIR::FPoison::Get();
@@ -599,7 +599,7 @@ FValue* FEmitter::EmitSwizzle(FValue* Value, FSwizzleMask Mask)
 	if (Mask.NumComponents == PrimitiveType->GetNumComponents())
 	{
 		bool InOrder = true;
-		for (int i = 0; i < Mask.NumComponents; ++i)
+		for (int32 i = 0; i < Mask.NumComponents; ++i)
 		{
 			if (Mask.Components[i] != (EVectorComponent)i)
 			{
@@ -617,16 +617,16 @@ FValue* FEmitter::EmitSwizzle(FValue* Value, FSwizzleMask Mask)
 	// If only one component is requested, we can use EmitSubscript() to return the single component.
 	if (Mask.NumComponents == 1)
 	{
-		return EmitSubscript(Value, (int)Mask.Components[0]);
+		return EmitSubscript(Value, (int32)Mask.Components[0]);
 	}
 
 	// Make the result vector type.
 	FPrimitiveTypePtr ResultType = FPrimitiveType::GetVector(PrimitiveType->ScalarKind, Mask.NumComponents);
 	FDimensional* Result = NewDimensionalValue(this, ResultType);
 
-	for (int i = 0; i < Mask.NumComponents; ++i)
+	for (int32 i = 0; i < Mask.NumComponents; ++i)
 	{
-		Result->GetComponents()[i] = EmitSubscript(Value, (int)Mask.Components[i]);
+		Result->GetComponents()[i] = EmitSubscript(Value, (int32)Mask.Components[i]);
 	}
 
 	return Result;
@@ -743,8 +743,8 @@ static bool FoldComparisonOperatorScalar(EOperator Operator, T A, T B)
 		case UO_Not: return !A;
 		case BO_GreaterThan: return A > B;
 		case BO_GreaterThanOrEquals: return A >= B;
-		case BO_LowerThan: return A < B;
-		case BO_LowerThanOrEquals: return A <= B;
+		case BO_LessThan: return A < B;
+		case BO_LessThanOrEquals: return A <= B;
 		case BO_Equals: return A == B;
 		case BO_NotEquals: return A != B;
 		default: UE_MIR_UNREACHABLE();
@@ -858,7 +858,7 @@ static FValue* TryFoldOperatorScalar(FEmitter* Emitter, EOperator Op, FValue* A,
 	{
 		/* Binary Comparisons */
 		case BO_GreaterThan:
-		case BO_LowerThan:
+		case BO_LessThan:
 		case BO_NotEquals:
 			if (A->Equals(B))
 			{
@@ -867,7 +867,7 @@ static FValue* TryFoldOperatorScalar(FEmitter* Emitter, EOperator Op, FValue* A,
 			break;
 
 		case BO_GreaterThanOrEquals:
-		case BO_LowerThanOrEquals:
+		case BO_LessThanOrEquals:
 		case BO_Equals:
 			if (A->Equals(B))
 			{
@@ -1021,9 +1021,103 @@ static FValue* TryFoldOperatorScalar(FEmitter* Emitter, EOperator Op, FValue* A,
 	}
 }
 
-FValue* FEmitter::EmitOperator(EOperator Op, FValue* A, FValue* B, FValue* C)
+/// Used to filter what parameter *primitive* types operators can take.
+enum EOperatorParameterTypeFilter
+{
+	OT_Unknown = 0xff,					///< Unspecified
+	OT_Any = 0,							///< Any primitive type (no filters)
+	OT_Boolean = 1 << 0,				///< Boolean primitive type of any dimension
+	OT_Arithmetic = 1 << 1,				///< Arithmetic primitive type of any dimension (i.e. that supports arithmetic operations)
+	OT_AnyFloat = 1 << 2,				///< Floating point primitive type of any dimension
+	OT_SameAsFirstArgument = 1 << 3,	///< Must be the same type as the first argument
+	OT_NotMatrix = 1 << 4,				///< Any primitive type except matrices
+	OT_Vector3 = 1 << 5,				///< 3D vector of any scalar type
+};
+
+inline EOperatorParameterTypeFilter operator|(EOperatorParameterTypeFilter A, EOperatorParameterTypeFilter B)
 {
-	// Assert that if C is specified, B must too
+	return EOperatorParameterTypeFilter(uint32(A) | uint32(B));
+}
+
+/// Used to determine the operator result type based on argument types
+enum EOperatorReturnType
+{
+	ORT_Unknown,							///< Unspecified
+	ORT_FirstArgumentType,					///< The same type as the first argument
+	ORT_BooleanWithFirstArgumentDimensions,	///< A boolean primitive type with the same dimensions (rows and columns) as the first argument type
+	ORT_FirstArgumentTypeToScalar			///< A scalar primitive type with the same kind as the scalar type of the first argument
+};
+
+/// The signature of an operator consisting of its parameter and return type information.
+struct FOperatorSignature
+{
+	EOperatorParameterTypeFilter ParameterTypes[3] = { OT_Unknown, OT_Unknown, OT_Unknown };
+	EOperatorReturnType ReturnType = ORT_Unknown;
+};
+
+/// Returns the signature of an operator.
+static const FOperatorSignature* GetOperatorSignature(EOperator Op)
+{
+	static const FOperatorSignature* Signatures = [] ()
+	{
+		const FOperatorSignature UnaryArithmetic = { { OT_Arithmetic }, ORT_FirstArgumentType };
+		const FOperatorSignature BinaryArithmetic = { { OT_Arithmetic, OT_SameAsFirstArgument }, ORT_FirstArgumentType };
+		const FOperatorSignature TernaryArithmetic = { { OT_Arithmetic, OT_SameAsFirstArgument, OT_SameAsFirstArgument }, ORT_FirstArgumentType };
+		const FOperatorSignature BinaryComparison = { { OT_Any, OT_SameAsFirstArgument }, ORT_BooleanWithFirstArgumentDimensions };
+		const FOperatorSignature BinaryLogical = { { OT_Boolean, OT_SameAsFirstArgument }, ORT_FirstArgumentType };
+
+		static FOperatorSignature S[MIR::OperatorCount];
+		S[UO_Negate] 				= UnaryArithmetic;
+		S[UO_Not] 					= { { OT_Boolean }, ORT_FirstArgumentType };
+		S[UO_Abs] 					= UnaryArithmetic;
+		S[UO_ACos] 					= UnaryArithmetic;
+		S[UO_ASin] 					= UnaryArithmetic;
+		S[UO_ATan] 					= UnaryArithmetic;
+		S[UO_Ceil] 					= UnaryArithmetic;
+		S[UO_Cos] 					= UnaryArithmetic;
+		S[UO_Exponential] 			= UnaryArithmetic;
+		S[UO_Exponential2] 			= UnaryArithmetic;
+		S[UO_Floor] 				= UnaryArithmetic;
+		S[UO_Logarithm] 			= UnaryArithmetic;
+		S[UO_Logarithm2] 			= UnaryArithmetic;
+		S[UO_Logarithm10] 			= UnaryArithmetic;
+		S[UO_Round] 				= UnaryArithmetic;
+		S[UO_Sin] 					= UnaryArithmetic;
+		S[UO_Tan] 					= UnaryArithmetic;
+		S[UO_Truncate] 				= UnaryArithmetic;
+		S[BO_GreaterThan] 			= BinaryComparison;
+		S[BO_GreaterThanOrEquals]	= BinaryComparison;
+		S[BO_LessThan]				= BinaryComparison;
+		S[BO_LessThanOrEquals]		= BinaryComparison;
+		S[BO_Equals]				= BinaryComparison;
+		S[BO_NotEquals]				= BinaryComparison;
+		S[BO_Add] 					= BinaryArithmetic;
+		S[BO_Subtract] 				= BinaryArithmetic;
+		S[BO_Multiply] 				= BinaryArithmetic;
+		S[BO_Divide] 				= BinaryArithmetic;
+		S[BO_Fmod] 					= BinaryArithmetic;
+		S[BO_Max] 					= BinaryArithmetic;
+		S[BO_Min] 					= BinaryArithmetic;
+		S[BO_Or] 					= BinaryLogical;
+		S[BO_And] 					= BinaryLogical;
+		S[BO_Xor] 					= BinaryLogical;
+		S[BO_Dot] 					= { { OT_Arithmetic | OT_NotMatrix, OT_SameAsFirstArgument }, ORT_FirstArgumentTypeToScalar };
+		S[BO_Cross] 				= { { OT_Arithmetic | OT_Vector3, OT_SameAsFirstArgument }, ORT_FirstArgumentType };
+		S[TO_Lerp] 					= TernaryArithmetic;
+		S[TO_Smoothstep]			= TernaryArithmetic;
+		return S;
+	} ();
+	return &Signatures[Op];
+}
+
+/// Validates that the types of the arguments are valid for specified operator.
+/// If valid, it returns the type of the result. Otherwise if it is not valid, it returns nullptr.
+static FPrimitiveTypePtr ValidateOperatorAndGetResultType(FEmitter* Emitter, EOperator Op, FValue* A, FValue* B, FValue* C)
+{
+	// Argument A must have always been provided.
+	check(A);
+
+	// Assert that if C is specified, B must too.
 	check(!C || B);
 
 	// Verify that B argument has been provided if operator is binary.
@@ -1032,35 +1126,120 @@ FValue* FEmitter::EmitOperator(EOperator Op, FValue* A, FValue* B, FValue* C)
 	// Verify that C argument has been provided if operator is ternary.
 	check(!MIR::IsTernaryOperator(Op) || C);
 
-	// Argument value types must always match.
-	if ((B && A->Type != B->Type) || 
-		(C && (A->Type != C->Type || B->Type != C->Type)))
+	// Verify that the first argument type is primitive.
+	FValue* Arguments[] = { A, B, C, nullptr };
+	FPrimitiveTypePtr FirstArgumentPrimitiveType = A->Type->AsPrimitive();
+	static const TCHAR* ArgumentsStr[] = { TEXT("first"), TEXT("second"), TEXT("third") };
+	const FOperatorSignature* Signature = GetOperatorSignature(Op);
+
+	bool bValid = true;
+	for (int32 i = 0; Arguments[i]; ++i)
 	{
-		Errorf(TEXT("Operands of operator `%s` do not have the same type."), OperatorToString(Op));
-		return MIR::FPoison::Get();
+		// Cast this argument type to primitive.
+		FPrimitiveTypePtr ArgumentPrimitiveType = Arguments[i]->Type->AsPrimitive();
+
+		// All operators have primitive parameter types.
+		if (!ArgumentPrimitiveType)
+		{
+			Emitter->Errorf(TEXT("Operator %s argument does not have primitive type."), ArgumentsStr[i]);
+			return nullptr;
+		}
+
+		EOperatorParameterTypeFilter Filter = Signature->ParameterTypes[i];
+		check(Filter != OT_Unknown); // No signature specified for this operator.
+
+		if (Filter & OT_Boolean)
+		{
+			if (ArgumentPrimitiveType->ScalarKind != ScalarKind_Bool)
+			{
+				Emitter->Errorf(TEXT("Operator %s argument component types are not boolean."), ArgumentsStr[i]);
+				bValid = false;
+			}
+		}
+		
+		if (Filter & OT_Arithmetic)
+		{
+			if (!ScalarKindIsArithmetic(ArgumentPrimitiveType->ScalarKind))
+			{
+				Emitter->Errorf(TEXT("Operator %s argument component types expected to be arithmetic (int, float, double, etc)."), ArgumentsStr[i]);
+				bValid = false;
+			}
+		}
+		
+		if (Filter & OT_AnyFloat)
+		{
+			if (!ScalarKindIsAnyFloat(ArgumentPrimitiveType->ScalarKind))
+			{
+				Emitter->Errorf(TEXT("Operator %s argument component types expected to be floating point (float, double, etc)."), ArgumentsStr[i]);
+				bValid = false;
+			}
+		}
+
+		if (Filter & OT_SameAsFirstArgument)
+		{
+			if (ArgumentPrimitiveType != FirstArgumentPrimitiveType)
+			{
+				Emitter->Errorf(TEXT("Operator %s argument type expected to be the same as first argument."), ArgumentsStr[i]);
+				bValid = false;
+			}
+		}
+		
+		if (Filter & OT_NotMatrix)
+		{
+			if (ArgumentPrimitiveType->IsMatrix())
+			{
+				Emitter->Errorf(TEXT("Operator %s argument type expected to be a scalar or vector."), ArgumentsStr[i]);
+				bValid = false;
+			}
+		}
+
+		if (Filter & OT_Vector3)
+		{
+			if (!ArgumentPrimitiveType->IsVector() || ArgumentPrimitiveType->GetNumComponents() != 3)
+			{
+				Emitter->Errorf(TEXT("Operator %s argument type expected to be a 3D vector."), ArgumentsStr[i]);
+				bValid = false;
+			}
+		}
 	}
 
-	// Get the arguments primitive type.
-	FPrimitiveTypePtr PrimitiveType = A->Type->AsPrimitive();
-	if (!PrimitiveType)
+	if (!bValid)
 	{
-		Errorf(TEXT("Operator on arguments of non-primitive type."));
-		return MIR::FPoison::Get();
+		return nullptr;
 	}
 	
-	// If operator is logical, arguments must be booleans.
-	if (IsBitOperator(Op) && !A->Type->IsBoolean() && !A->Type->IsInteger())
+	// Finally, determine operator result type.
+	switch (Signature->ReturnType)
 	{
-		Errorf(TEXT("Bit operator on non-boolean, non-integer arguments."));
-		return MIR::FPoison::Get();
+		case ORT_Unknown:
+			UE_MIR_UNREACHABLE(); // missing operator signature declaration
+
+		case ORT_FirstArgumentType:
+			return FirstArgumentPrimitiveType;
+
+		case ORT_BooleanWithFirstArgumentDimensions:
+			return FPrimitiveType::Get(ScalarKind_Bool, FirstArgumentPrimitiveType->NumRows, FirstArgumentPrimitiveType->NumColumns);
+
+		case ORT_FirstArgumentTypeToScalar:
+			return FPrimitiveType::GetScalar(FirstArgumentPrimitiveType->ScalarKind);
 	}
 
-	// Determine the result type. If the operator is arithmetic, the result type will be the same
-	// as the arguments type. Otherwise it will have the same number of components but boolean.
-	FPrimitiveTypePtr ResultType = IsComparisonOperator(Op) || IsBitOperator(Op)
-		? FPrimitiveType::Get(ScalarKind_Bool, PrimitiveType->NumRows, PrimitiveType->NumColumns)
-		: PrimitiveType;
+	UE_MIR_UNREACHABLE();
+}
+
+/// Returns whether the operator supports componentwise application. In other words, if the following is true:
+/// 	op(v, w) == [op(v_x, w_x), ..., op(v_n, w_n)]
+static bool IsComponentwiseOperator(EOperator Op)
+{
+	return Op != BO_Dot && Op != BO_Cross;
+}
 
+/// Tries to fold the operator by applying the operator componentwise on arguments components.
+/// If a value is returned, it will be a dimensional with some component folded to a constant. If some argument
+/// isn't a dimensional, or all arguments components are non-constant, the folding will not be carried out.
+/// If no folding is carried out, this function simply returns nullptr.
+static FValue* TryFoldComponentwiseOperator(FEmitter* Emitter, EOperator Op, FValue* A, FValue* B, FValue* C, FPrimitiveTypePtr ResultType)
+{
 	// Check that at least one component of the resulting dimensional value would folded.
 	// If all components of resulting dimensional value are not folded, then instead of emitting
 	// an individual operator instruction for each component, simply emit a single binary operator
@@ -1074,15 +1253,15 @@ FValue* FEmitter::EmitOperator(EOperator Op, FValue* A, FValue* B, FValue* C)
 	FMemMark Mark(FMemStack::Get());
 	TArrayView<FValue*> TempResultComponents = MakeTemporaryArray<FValue*>(Mark, ResultType->GetNumComponents());
 	
-	for (int i = 0; i < ResultType->GetNumComponents(); ++i)
+	for (int32 i = 0; i < ResultType->GetNumComponents(); ++i)
 	{
 		// Extract the arguments individual components
-		FValue* AComponent = EmitSubscript(A, i);
-		FValue* BComponent = B ? EmitSubscript(B, i) : nullptr;
-		FValue* CComponent = C ? EmitSubscript(C, i) : nullptr;
+		FValue* AComponent = Emitter->EmitSubscript(A, i);
+		FValue* BComponent = B ? Emitter->EmitSubscript(B, i) : nullptr;
+		FValue* CComponent = C ? Emitter->EmitSubscript(C, i) : nullptr;
 
 		// Try folding the operation, it may return null
-		FValue* ResultComponent = TryFoldOperatorScalar(this, Op, AComponent, BComponent, CComponent);
+		FValue* ResultComponent = TryFoldOperatorScalar(Emitter, Op, AComponent, BComponent, CComponent);
 
 		// Update the flags
 		bSomeResultComponentWasFolded |= (bool)ResultComponent;
@@ -1124,7 +1303,7 @@ FValue* FEmitter::EmitOperator(EOperator Op, FValue* A, FValue* B, FValue* C)
 		}
 
 		// Make the new dimensional value
-		FDimensional* Result = NewDimensionalValue(this, ResultType);
+		FDimensional* Result = NewDimensionalValue(Emitter, ResultType);
 
 		// Fetch the components array from the result dimensional
 		TArrayView<FValue*> ResultComponents = Result->GetComponents();
@@ -1133,7 +1312,7 @@ FValue* FEmitter::EmitOperator(EOperator Op, FValue* A, FValue* B, FValue* C)
 		FPrimitiveTypePtr ComponentType = ResultType->ToScalar();
 
 		// Create the operator instruction for each component pair
-		for (int i = 0; i < ResultType->GetNumComponents(); ++i)
+		for (int32 i = 0; i < ResultType->GetNumComponents(); ++i)
 		{
 			// Reuse cached result if possible
 			ResultComponents[i] = TempResultComponents[i];
@@ -1143,17 +1322,120 @@ FValue* FEmitter::EmitOperator(EOperator Op, FValue* A, FValue* B, FValue* C)
 			{
 				FOperator Proto = MakePrototype<FOperator>(ComponentType);
 				Proto.Op = Op;
-				Proto.AArg = EmitSubscript(A, i);
-				Proto.BArg = B ? EmitSubscript(B, i) : nullptr;
-				Proto.CArg = C ? EmitSubscript(C, i) : nullptr;
-				ResultComponents[i] = EmitPrototype(this, Proto);
+				Proto.AArg = Emitter->EmitSubscript(A, i);
+				Proto.BArg = B ? Emitter->EmitSubscript(B, i) : nullptr;
+				Proto.CArg = C ? Emitter->EmitSubscript(C, i) : nullptr;
+				ResultComponents[i] = EmitPrototype(Emitter, Proto);
 			}
 		}
 
 		return Result;
 	}
 
-	// Finally, simply emit the binary operator instruction between lhs and rhs (i.e. "lhs + rhs")
+	return nullptr;
+}
+
+/// If V is a dimensional and all its components are constants, it unpacks the components into OutComponents and returns true.
+/// If this is not possible for any reason, it returns false.
+static bool TryUnpackConstantVector3(FValue* V, TArrayView<FConstant*> OutComponents)
+{
+	check(V->Type->AsPrimitive()->IsVector() && V->Type->AsPrimitive()->GetNumComponents() == 3);
+	FDimensional* Dimensional = V->As<FDimensional>();
+	if (!V)
+	{
+		return false;
+	}
+
+	TArrayView<FValue*> Components = Dimensional->GetComponents();
+	for (int32 i = 0; i < 3; ++i)
+	{
+		OutComponents[i] = Components[i]->As<FConstant>();
+		if (!OutComponents[i])
+		{
+			return false;
+		}
+	}
+
+	return true;
+}
+
+template <typename T>
+static Math::TVector<T> Cross(Math::TVector<T> A, Math::TVector<T> B)
+{
+	return Math::TVector<T>(
+		A.Y * B.Z - A.Z * B.Y,
+		A.Z * B.X - A.X * B.Z,
+		A.X * B.Y - A.Y * B.X);
+}
+
+/// Tries to fold the operator, that is to evaluate its result now at translation time if its arguments are constant.
+/// If the operator could not be folded in any way, it returns nullptr.
+static FValue* TryFoldOperator(FEmitter* Emitter, EOperator Op, FValue* A, FValue* B, FValue* C, FPrimitiveTypePtr ResultType)
+{
+	if (Op == BO_Dot)
+	{
+		FPrimitiveTypePtr ArgumentsType = A->Type->AsPrimitive();
+
+		// Componentwise multiply the two vectors.
+		FValue* MultipliedComponents = TryFoldComponentwiseOperator(Emitter, BO_Multiply, A, B, nullptr, ArgumentsType);
+
+		// If the operation was folded (i.e. some componentwise multiplication was folded to a constant),
+		// simply add the components together to compute the final dot result.
+		if (MultipliedComponents)
+		{
+			FValue* Accumulator = Emitter->EmitSubscript(MultipliedComponents, 0);
+			for (int32 i = 1; i < ArgumentsType->NumRows; ++i)
+			{
+				Accumulator = Emitter->EmitOperator(BO_Add, Accumulator, Emitter->EmitSubscript(MultipliedComponents, i));
+			}
+			return Accumulator;
+		}
+	}
+	else if (Op == BO_Cross)
+	{
+		// Fold the cross operator only if all arguments components are constant.
+		FConstant* AVector[3];
+		FConstant* BVector[3];
+		if (TryUnpackConstantVector3(A, AVector) && TryUnpackConstantVector3(B, BVector))
+		{
+			if (ResultType->ScalarKind == ScalarKind_Float)
+			{
+				FVector3f Result = Cross<float>({ AVector[0]->Float, AVector[1]->Float, AVector[2]->Float }, { BVector[0]->Float, BVector[1]->Float, BVector[2]->Float });
+				return Emitter->EmitConstantFloat3(Result);
+			}
+			else
+			{
+				UE_MIR_UNREACHABLE();
+			}
+		}
+	}
+
+	// If the operation supports componentwise application, try folding the operator componentwise.
+	if (IsComponentwiseOperator(Op))
+	{
+		return TryFoldComponentwiseOperator(Emitter, Op, A, B, C, ResultType);
+	}
+
+	// No folding was possible, simply return null to indicate this.
+	return nullptr;
+}
+
+FValue* FEmitter::EmitOperator(EOperator Op, FValue* A, FValue* B, FValue* C)
+{
+	// Validate the operation and retrieve the result type.
+	FPrimitiveTypePtr ResultType = ValidateOperatorAndGetResultType(this, Op, A, B, C);
+	if (!ResultType)
+	{
+		return FPoison::Get();
+	}
+
+	// Try folding the operator folding first.
+	if (FValue* FoldedValue = TryFoldOperator(this, Op, A, B, C, ResultType))
+	{
+		return FoldedValue;
+	}
+
+	// Otherwise, we must emit a new instruction that executes the operator.
 	FOperator Proto = MakePrototype<FOperator>(ResultType);
 	Proto.Op = Op;
 	Proto.AArg = A;
@@ -1216,7 +1498,7 @@ static FValue* CastConstant(FEmitter* Emitter, FConstant* Constant, EScalarKind
 			switch (TargetKind)
 			{
 				case ScalarKind_Bool: UE_MIR_TODO();
-				case ScalarKind_Int: return Emitter->EmitConstantIntScalar((int)Constant->Float);
+				case ScalarKind_Int: return Emitter->EmitConstantIntScalar((int32)Constant->Float);
 				default: UE_MIR_UNREACHABLE();
 			}
 		}
@@ -1272,7 +1554,7 @@ static FValue* ConstructPrimitiveValue(FEmitter* Emitter, FPrimitiveTypePtr Targ
 		FValue* Component = Emitter->EmitConstruct(TargetPrimitiveType->ToScalar(), Initializer);
 
 		// Initialize all result components to the same scalar.
-		for (int i = 0; i < TargetPrimitiveType->GetNumComponents(); ++i)
+		for (int32 i = 0; i < TargetPrimitiveType->GetNumComponents(); ++i)
 		{
 			Result->GetComponents()[i] = Component;
 		}
@@ -1285,8 +1567,8 @@ static FValue* ConstructPrimitiveValue(FEmitter* Emitter, FPrimitiveTypePtr Targ
 	// the necessary components.
 	if (TargetPrimitiveType->IsVector() && InitializerPrimitiveType->IsVector())
 	{
-		int TargetNumComponents = TargetPrimitiveType->GetNumComponents();
-		int InitializerNumComponents = InitializerPrimitiveType->GetNumComponents();
+		int32 TargetNumComponents = TargetPrimitiveType->GetNumComponents();
+		int32 InitializerNumComponents = InitializerPrimitiveType->GetNumComponents();
 
 		// Create the result dimensional value.
 		FDimensional* Result = NewDimensionalValue(Emitter, TargetPrimitiveType);
@@ -1295,10 +1577,10 @@ static FValue* ConstructPrimitiveValue(FEmitter* Emitter, FPrimitiveTypePtr Targ
 		FPrimitiveTypePtr ResultComponentType = TargetPrimitiveType->ToScalar();
 
 		// For iterating over the components of the result dimensional value.
-		int Index = 0;
+		int32 Index = 0;
 
 		// Convert components from the initializer vector.
-		int MinNumComponents = FMath::Min(TargetNumComponents, InitializerNumComponents);
+		int32 MinNumComponents = FMath::Min(TargetNumComponents, InitializerNumComponents);
 		for (; Index < MinNumComponents; ++Index)
 		{
 			Result->GetComponents()[Index] = Emitter->EmitConstruct(ResultComponentType, Emitter->EmitSubscript(Initializer, Index));
@@ -1329,7 +1611,7 @@ static FValue* ConstructPrimitiveValue(FEmitter* Emitter, FPrimitiveTypePtr Targ
 			FPrimitiveTypePtr ResultComponentType = TargetPrimitiveType->ToScalar();
 
 			// Convert components from the initializer vector.
-			for (int Index = 0, Num = TargetPrimitiveType->GetNumComponents(); Index < Num; ++Index)
+			for (int32 Index = 0, Num = TargetPrimitiveType->GetNumComponents(); Index < Num; ++Index)
 			{
 				Result->GetComponents()[Index] = Emitter->EmitConstruct(ResultComponentType, DimensionalInitializer->GetComponents()[Index]);
 			}
@@ -1405,7 +1687,7 @@ FPrimitiveTypePtr FEmitter::GetCommonPrimitiveType(FPrimitiveTypePtr A, FPrimiti
 	// Neither A nor B are matrices, but single scalar or vector. Return the largest.
 	check(A->NumColumns == 1 && B->NumColumns == 1);
 	EScalarKind ScalarKind = FMath::Max(A->ScalarKind, B->ScalarKind);
-	int NumRows = FMath::Max(A->NumRows, B->NumRows);
+	int32 NumRows = FMath::Max(A->NumRows, B->NumRows);
 	return FPrimitiveType::Get(ScalarKind, NumRows, 1);
 }
 
diff --git a/Engine/Source/Runtime/Engine/Private/Materials/MaterialIRToHLSLTranslator.cpp b/Engine/Source/Runtime/Engine/Private/Materials/MaterialIRToHLSLTranslator.cpp
--- a/Engine/Source/Runtime/Engine/Private/Materials/MaterialIRToHLSLTranslator.cpp
+++ b/Engine/Source/Runtime/Engine/Private/Materials/MaterialIRToHLSLTranslator.cpp
@@ -197,6 +197,32 @@ static bool IsFoldable(const MIR::FInstruction* Instr)
 	return true;
 }
 
+/// Returns whether the operator in HLSL is infix between its arguments, e.g. "4 + 4"
+static bool IsOperatorInfix(MIR::EOperator Op)
+{
+	switch (Op)
+	{
+		case MIR::BO_GreaterThan:
+		case MIR::BO_GreaterThanOrEquals:
+		case MIR::BO_LessThan:
+		case MIR::BO_LessThanOrEquals:
+		case MIR::BO_Equals:
+		case MIR::BO_NotEquals:
+		case MIR::BO_Add:
+		case MIR::BO_Subtract:
+		case MIR::BO_Multiply:
+		case MIR::BO_Divide:
+		case MIR::BO_Fmod:
+		case MIR::BO_Max:
+		case MIR::BO_Min:
+		case MIR::BO_Or:
+		case MIR::BO_And:
+		case MIR::BO_Xor:
+			return true;
+		default: return false;
+	}
+}
+
 struct FTranslator : FMaterialIRToHLSLTranslation
 {
 	int32 NumLocals{};
@@ -429,7 +455,7 @@ struct FTranslator : FMaterialIRToHLSLTranslation
 
 	void LowerOperator(const MIR::FOperator* Operator)
 	{
-		if (MIR::IsComparisonOperator(Operator->Op) || MIR::IsBinaryInfixOperator(Operator->Op))
+		if (IsOperatorInfix(Operator->Op))
 		{
 			LowerValue(Operator->AArg);
 
@@ -438,8 +464,8 @@ struct FTranslator : FMaterialIRToHLSLTranslation
 			{
 				case MIR::BO_GreaterThan: OpString = TEXT(">"); break;
 				case MIR::BO_GreaterThanOrEquals: OpString = TEXT(">="); break;
-				case MIR::BO_LowerThan: OpString = TEXT("<"); break;
-				case MIR::BO_LowerThanOrEquals: OpString = TEXT("<="); break;
+				case MIR::BO_LessThan: OpString = TEXT("<"); break;
+				case MIR::BO_LessThanOrEquals: OpString = TEXT("<="); break;
 				case MIR::BO_Equals: OpString = TEXT("=="); break;
 				case MIR::BO_NotEquals: OpString = TEXT("!="); break;
 				
@@ -478,6 +504,8 @@ struct FTranslator : FMaterialIRToHLSLTranslation
 				case MIR::BO_Fmod: OpString = TEXT("fmod"); break;
 				case MIR::BO_Min: OpString = TEXT("min"); break;
 				case MIR::BO_Max: OpString = TEXT("max"); break;
+				case MIR::BO_Dot: OpString = TEXT("dot"); break;
+				case MIR::BO_Cross: OpString = TEXT("cross"); break;
 				case MIR::TO_Lerp: OpString = TEXT("lerp"); break;
 				case MIR::TO_Smoothstep: OpString = TEXT("smoothstep"); break;
 				default: UE_MIR_UNREACHABLE();
diff --git a/Engine/Source/Runtime/Engine/Private/Materials/MaterialIRTypes.cpp b/Engine/Source/Runtime/Engine/Private/Materials/MaterialIRTypes.cpp
--- a/Engine/Source/Runtime/Engine/Private/Materials/MaterialIRTypes.cpp
+++ b/Engine/Source/Runtime/Engine/Private/Materials/MaterialIRTypes.cpp
@@ -238,6 +238,16 @@ FPrimitiveTypePtr FType::AsMatrix() const
 	return Type->IsMatrix() ? Type : nullptr;
 }
 
+bool ScalarKindIsArithmetic(EScalarKind Kind)
+{
+	return Kind != ScalarKind_Bool;
+}
+
+bool ScalarKindIsAnyFloat(EScalarKind Kind)
+{
+	return Kind == ScalarKind_Float;
+}
+
 const TCHAR* ScalarKindToString(EScalarKind Kind)
 {
 	switch (Kind)
@@ -370,7 +380,6 @@ FTypePtr FObjectType::GetTexture2D()
 	return &Instance;
 }
 
-
 } // namespace UE::MIR
 
 #endif // #if WITH_EDITOR
diff --git a/Engine/Source/Runtime/Engine/Public/Materials/MaterialExpressionCrossProduct.h b/Engine/Source/Runtime/Engine/Public/Materials/MaterialExpressionCrossProduct.h
--- a/Engine/Source/Runtime/Engine/Public/Materials/MaterialExpressionCrossProduct.h
+++ b/Engine/Source/Runtime/Engine/Public/Materials/MaterialExpressionCrossProduct.h
@@ -25,6 +25,7 @@ public:
 
 	//~ Begin UMaterialExpression Interface
 #if WITH_EDITOR
+	virtual void Build(UE::MIR::FEmitter& Emitter) override;
 	virtual int32 Compile(class FMaterialCompiler* Compiler, int32 OutputIndex) override;
 	virtual void GetCaption(TArray<FString>& OutCaptions) const override;
 #endif
diff --git a/Engine/Source/Runtime/Engine/Public/Materials/MaterialExpressionDotProduct.h b/Engine/Source/Runtime/Engine/Public/Materials/MaterialExpressionDotProduct.h
--- a/Engine/Source/Runtime/Engine/Public/Materials/MaterialExpressionDotProduct.h
+++ b/Engine/Source/Runtime/Engine/Public/Materials/MaterialExpressionDotProduct.h
@@ -25,6 +25,7 @@ public:
 
 	//~ Begin UMaterialExpression Interface
 #if WITH_EDITOR
+	virtual void Build(UE::MIR::FEmitter& Emitter) override;
 	virtual int32 Compile(class FMaterialCompiler* Compiler, int32 OutputIndex) override;
 	virtual void GetCaption(TArray<FString>& OutCaptions) const override;
 #endif
diff --git a/Engine/Source/Runtime/Engine/Public/Materials/MaterialIR.h b/Engine/Source/Runtime/Engine/Public/Materials/MaterialIR.h
--- a/Engine/Source/Runtime/Engine/Public/Materials/MaterialIR.h
+++ b/Engine/Source/Runtime/Engine/Public/Materials/MaterialIR.h
@@ -125,7 +125,7 @@ struct FPoison : TValue<VK_Poison>
 };
 
 using TInteger = int64_t;
-using TFloat = double;
+using TFloat = float;
 
 struct FConstant : TValue<VK_Constant>
 {
@@ -310,7 +310,7 @@ enum EOperator
 	UO_FirstUnaryOperator,
 
 	UO_Negate = UO_FirstUnaryOperator, 	// Arithmetic negation: negate(5) -> -5
-	UO_Not,							// Logical negation: not(true) -> false
+	UO_Not,								// Logical negation: not(true) -> false
 	
 	/* Unary Intrinsic Operators */
 	UO_Abs,
@@ -335,8 +335,8 @@ enum EOperator
 
 	BO_GreaterThan = BO_FirstBinaryOperator,
 	BO_GreaterThanOrEquals,
-	BO_LowerThan,
-	BO_LowerThanOrEquals,
+	BO_LessThan,
+	BO_LessThanOrEquals,
 	BO_Equals,
 	BO_NotEquals,
 
@@ -345,20 +345,26 @@ enum EOperator
 	BO_Subtract,
 	BO_Multiply,
 	BO_Divide,
+	BO_Fmod,
+	BO_Max,
+	BO_Min,
+
+	/* Binary Logical Operators */
 	BO_Or,
 	BO_And,
 	BO_Xor,
 
 	/* Binary Intrinsic Operators */
-	BO_Fmod,
-	BO_Max,
-	BO_Min,
+	BO_Dot,
+	BO_Cross,
 
 	/* Ternary Arithmetic Operators */
 	TO_FirstTernaryOperator,
 
 	TO_Lerp = TO_FirstTernaryOperator,
 	TO_Smoothstep,
+
+	OperatorCount,
 };
 
 //
@@ -371,10 +377,10 @@ bool IsBitOperator(EOperator Op);
 bool IsBinaryOperator(EOperator Op);
 
 //
-bool IsTernaryOperator(EOperator Op);
+bool IsBinaryOperator(EOperator Op);
 
 //
-bool IsBinaryInfixOperator(EOperator Op);
+bool IsTernaryOperator(EOperator Op);
 
 //
 const TCHAR* OperatorToString(EOperator Op);
diff --git a/Engine/Source/Runtime/Engine/Public/Materials/MaterialIRTypes.h b/Engine/Source/Runtime/Engine/Public/Materials/MaterialIRTypes.h
--- a/Engine/Source/Runtime/Engine/Public/Materials/MaterialIRTypes.h
+++ b/Engine/Source/Runtime/Engine/Public/Materials/MaterialIRTypes.h
@@ -1,4 +1,4 @@
-// Copyright Epic Games, Inc. All Rights Reserved.
+/// Copyright Epic Games, Inc. All Rights Reserved.
 #pragma once
 
 #include "Materials/MaterialIRCommon.h"
@@ -28,66 +28,66 @@ const TCHAR* TypeKindToString(ETypeKind Kind);
 
 struct FType
 {
-	// Identifies what derived type this is.
+	/// Identifies what derived type this is.
 	ETypeKind Kind;
 
-	// Returns the type matching specified UE::Shader::FType.
+	/// Returns the type matching specified UE::Shader::FType.
 	static FTypePtr FromShaderType(const UE::Shader::FType& InShaderType);
 	
-	// Returns the type matching specified EMaterialValueType.
+	/// Returns the type matching specified EMaterialValueType.
 	static FTypePtr FromMaterialValueType(EMaterialValueType Type);
 
-	// Returns the type matching specified EMaterialParameterType.
+	/// Returns the type matching specified EMaterialParameterType.
 	static FTypePtr FromMaterialParameterType(EMaterialParameterType Type);
 
-	// Returns the `void` type.
+	/// Returns the `void` type.
 	static FTypePtr GetVoid();
 
-	// Returns the `void` type.
+	/// Returns the `void` type.
 	static FTypePtr GetPoison();
 
-	// Returns whether this is the poison type.
+	/// Returns whether this is the poison type.
 	bool IsPoison() const;
 
-	// Returns whether this is a boolean primitive type (of any dimension).
+	/// Returns whether this is a boolean primitive type (of any dimension).
 	bool IsBoolean() const;
 	
-	// Returns whether this is an integer primitive type (of any dimension).
+	/// Returns whether this is an integer primitive type (of any dimension).
 	bool IsInteger() const;
 
-	// Returns whether this type is a `bool` scalar.
+	/// Returns whether this type is a `bool` scalar.
 	bool IsBoolScalar() const;
 
-	// Returns whether this type is of a specific object kind.
+	/// Returns whether this type is of a specific object kind.
 	bool IsObjectOfKind(EObjectKind ObjectKind) const;
 
-	// Returns whether this type is a texture.
+	/// Returns whether this type is a texture.
 	bool IsTexture() const;
 
-	// Returns this type upcast this type to PrimitiveType if it is one, otherwise nullptr.
+	/// Returns this type upcast this type to PrimitiveType if it is one, otherwise nullptr.
 	FPrimitiveTypePtr AsPrimitive() const;
 
-	// Returns this type upcast this type to PrimitiveType if it's a scalar, otherwise nullptr.
+	/// Returns this type upcast this type to PrimitiveType if it's a scalar, otherwise nullptr.
 	FPrimitiveTypePtr AsScalar() const;
 
-	// Returns this type upcast this type to PrimitiveType if it's a vector, otherwise nullptr.
+	/// Returns this type upcast this type to PrimitiveType if it's a vector, otherwise nullptr.
 	FPrimitiveTypePtr AsVector() const;
 
-	// Returns this type upcast this type to PrimitiveType if it's a matrix, otherwise nullptr.
+	/// Returns this type upcast this type to PrimitiveType if it's a matrix, otherwise nullptr.
 	FPrimitiveTypePtr AsMatrix() const;
 
-	// Returns the this type name spelling (e.g. float4x4).
+	/// Returns the this type name spelling (e.g. float4x4).
 	FStringView GetSpelling() const;
 	
-	// Returns this type cast as a FObjectPtr if it is one, otherwise nullptr.
+	/// Returns this type cast as a FObjectPtr if it is one, otherwise nullptr.
 	FObjectTypePtr AsObject() const;
 
-	// Converts this type to a UE::Shader::EValueType.
+	/// Converts this type to a UE::Shader::EValueType.
 	UE::Shader::EValueType ToValueType() const;
 };
 
-// Primitive types of a single scalar.
-// Note: These are listed in precision order. Converting one to the other is then simply performed taking the max EScalarKind.
+/// Primitive types of a single scalar.
+/// Note: These are listed in precision order. Converting one to the other is then simply performed taking the max EScalarKind.
 enum EScalarKind
 {
 	ScalarKind_Bool,
@@ -95,13 +95,29 @@ enum EScalarKind
 	ScalarKind_Float,
 };
 
+/// Returns whether the specified scalar kind supports arithmetic operators (plus, minus, etc).
+bool ScalarKindIsArithmetic(EScalarKind Kind);
+
+/// Returns whether the specified scalar kind is a floating point type (float, double, etc).
+bool ScalarKindIsAnyFloat(EScalarKind Kind);
+
+/// Returns the string representation of specified scalar kind.
 const TCHAR* ScalarKindToString(EScalarKind Kind);
 
+/// Represents the type of scalars, vectors and matrices. It indicates what kind of scalar
+/// data type it has and type dimensions (rows and columns).
 struct FPrimitiveType : FType
 {
+	/// String representation of this type (e.g. 'float3', "bool4x2")
 	FStringView Spelling;
+
+	/// Scalar data type kind.
 	EScalarKind ScalarKind;
+
+	/// Number of rows.
 	int NumRows;
+	
+	/// Number of columns (rows > 1 and columns == 1, this is a vector).
 	int NumColumns;
 
 	static FPrimitiveTypePtr GetBool1();
