Fixed static code analysis warning in MaterialIREmitter.cpp

diff --git a/Engine/Source/Runtime/Engine/Private/Materials/MaterialIREmitter.cpp b/Engine/Source/Runtime/Engine/Private/Materials/MaterialIREmitter.cpp
--- a/Engine/Source/Runtime/Engine/Private/Materials/MaterialIREmitter.cpp
+++ b/Engine/Source/Runtime/Engine/Private/Materials/MaterialIREmitter.cpp
@@ -1339,7 +1339,7 @@ static FValue* TryFoldComponentwiseOperator(FEmitter* Emitter, EOperator Op, FVa
 /// If this is not possible for any reason, it returns false.
 static bool TryUnpackConstantVector3(FValue* V, TArrayView<FConstant*> OutComponents)
 {
-	check(V->Type->AsPrimitive()->IsVector() && V->Type->AsPrimitive()->GetNumComponents() == 3);
+	check(V && V->Type->AsPrimitive()->IsVector() && V->Type->AsPrimitive()->GetNumComponents() == 3);
 	FDimensional* Dimensional = V->As<FDimensional>();
 	if (!V)
 	{
