Convert all Material Expressions to use default member initializers instead of constructors

- C.48: Prefer default member initializers to member initializers in constructors for constant initializers
- If the material expression's constructor was then empty I deleted the constructor
- We still need the constructor, but now auto-gen it via `GENERATED_BODY()` instead of the legacy `GENERATED_UCLASS_BODY` (which didn't generate the constructor)
- The GENERATED_BODY had a `private:` in it making all the properties and member functions private on the UMaterialExpression.
- I added back in a `public:` to each class so they behave and compile like they used to.
- Where possible I combined `WITH_EDITOR` blocks.

[RN] minor | Rendering - Materials and Shaders

diff --git a/Engine/Source/Runtime/Engine/Private/Materials/MaterialExpressions.cpp b/Engine/Source/Runtime/Engine/Private/Materials/MaterialExpressions.cpp
--- a/Engine/Source/Runtime/Engine/Private/Materials/MaterialExpressions.cpp
+++ b/Engine/Source/Runtime/Engine/Private/Materials/MaterialExpressions.cpp
@@ -2341,13 +2341,6 @@ FSubstrateOperator* UMaterialExpression::SubstrateGenerateMaterialTopologyTree(c
 
 #endif // WITH_EDITOR
 
-UMaterialExpressionTextureBase::UMaterialExpressionTextureBase(const FObjectInitializer& ObjectInitializer)
-	: Super(ObjectInitializer)
-#if WITH_EDITORONLY_DATA
-	, IsDefaultMeshpaintTexture(false)
-#endif
-{}
-
 #if WITH_EDITOR
 void UMaterialExpressionTextureBase::PostEditChangeProperty(FPropertyChangedEvent& PropertyChangedEvent)
 {
@@ -3852,13 +3845,6 @@ int32 UMaterialExpressionMeshPaintTextureReplace::Compile(class FMaterialCompile
 
 #endif // WITH_EDITOR
 
-UMaterialExpressionAdd::UMaterialExpressionAdd(const FObjectInitializer& ObjectInitializer)
-	: Super(ObjectInitializer)
-{
-	ConstA = 0.0f;
-	ConstB = 1.0f;
-}
-
 //
 //  UMaterialExpressionTextureSampleParameter
 //
@@ -4156,8 +4142,6 @@ uint32 UMaterialExpressionTextureObject::GetOutputType(int32 OutputIndex)
 UMaterialExpressionTextureProperty::UMaterialExpressionTextureProperty(const FObjectInitializer& ObjectInitializer)
 	: Super(ObjectInitializer)
 {
-	Property = TMTM_TextureSize;
-
 #if WITH_EDITORONLY_DATA
 	bShowOutputNameOnPin = false;
 	
@@ -4551,12 +4535,6 @@ static int32 ParticleSubUV(
 /** 
  *	UMaterialExpressionTextureSampleParameterSubUV
  */
-UMaterialExpressionTextureSampleParameterSubUV::UMaterialExpressionTextureSampleParameterSubUV(const FObjectInitializer& ObjectInitializer)
-	: Super(ObjectInitializer)
-{
-	bBlend = true;
-}
-
 #if WITH_EDITOR
 int32 UMaterialExpressionTextureSampleParameterSubUV::Compile(class FMaterialCompiler* Compiler, int32 OutputIndex)
 {
@@ -4621,13 +4599,6 @@ void UMaterialExpressionAdd::GetCaption(TArray<FString>& OutCaptions) const
 //
 //	UMaterialExpressionMultiply
 //
-UMaterialExpressionMultiply::UMaterialExpressionMultiply(const FObjectInitializer& ObjectInitializer)
-	: Super(ObjectInitializer)
-{
-	ConstA = 0.0f;
-	ConstB = 1.0f;
-}
-
 #if WITH_EDITOR
 int32 UMaterialExpressionMultiply::Compile(class FMaterialCompiler* Compiler, int32 OutputIndex)
 {
@@ -4657,13 +4628,6 @@ void UMaterialExpressionMultiply::GetCaption(TArray<FString>& OutCaptions) const
 }
 #endif // WITH_EDITOR
 
-UMaterialExpressionDivide::UMaterialExpressionDivide(const FObjectInitializer& ObjectInitializer)
-	: Super(ObjectInitializer)
-{
-	ConstA = 1.0f;
-	ConstB = 2.0f;
-}
-
 #if WITH_EDITOR
 int32 UMaterialExpressionDivide::Compile(class FMaterialCompiler* Compiler, int32 OutputIndex)
 {
@@ -4691,19 +4655,7 @@ void UMaterialExpressionDivide::GetCaption(TArray<FString>& OutCaptions) const
 
 	OutCaptions.Add(ret);
 }
-#endif // WITH_EDITOR
-
-//
-//	UMaterialExpressionSubtract
-//
-UMaterialExpressionSubtract::UMaterialExpressionSubtract(const FObjectInitializer& ObjectInitializer)
-	: Super(ObjectInitializer)
-{
-	ConstA = 1.0f;
-	ConstB = 1.0f;
-}
 
-#if WITH_EDITOR
 int32 UMaterialExpressionSubtract::Compile(class FMaterialCompiler* Compiler, int32 OutputIndex)
 {
 	// if the input is hooked up, use it, otherwise use the internal constant
@@ -4729,21 +4681,7 @@ void UMaterialExpressionSubtract::GetCaption(TArray<FString>& OutCaptions) const
 
 	OutCaptions.Add(ret);
 }
-#endif // WITH_EDITOR
-
-//
-//	UMaterialExpressionSmoothStep
-//
-
-UMaterialExpressionSmoothStep::UMaterialExpressionSmoothStep(const FObjectInitializer& ObjectInitializer)
-	: Super(ObjectInitializer)
-{
-	ConstMin = 0.0f;
-	ConstMax = 1.0f;
-	ConstValue = 0.0f;
-}
 
-#if WITH_EDITOR
 int32 UMaterialExpressionSmoothStep::Compile(class FMaterialCompiler* Compiler, int32 OutputIndex)
 {
 	// if the input is hooked up, use it, otherwise use the internal constant
@@ -4779,14 +4717,6 @@ void UMaterialExpressionSmoothStep::GetCaption(TArray<FString>& OutCaptions) con
 //
 //	UMaterialExpressionStep
 //
-
-UMaterialExpressionStep::UMaterialExpressionStep(const FObjectInitializer& ObjectInitializer)
-	: Super(ObjectInitializer)
-{
-	ConstY = 0.0f;
-	ConstX = 1.0f;
-}
-
 #if WITH_EDITOR
 int32 UMaterialExpressionStep::Compile(class FMaterialCompiler* Compiler, int32 OutputIndex)
 {
@@ -4819,16 +4749,6 @@ void UMaterialExpressionStep::GetCaption(TArray<FString>& OutCaptions) const
 //
 //	UMaterialExpressionInverseLerp
 //
-
-UMaterialExpressionInverseLinearInterpolate::UMaterialExpressionInverseLinearInterpolate(const FObjectInitializer& ObjectInitializer)
-	: Super(ObjectInitializer)
-{
-	ConstA = 0.0f;
-	ConstB = 1.0f;
-	ConstValue = 0.0f;
-	bClampResult = false;
-}
-
 #if WITH_EDITOR
 int32 UMaterialExpressionInverseLinearInterpolate::Compile(class FMaterialCompiler* Compiler, int32 OutputIndex)
 {
@@ -4871,15 +4791,6 @@ void UMaterialExpressionInverseLinearInterpolate::GetCaption(TArray<FString>& Ou
 //
 //	UMaterialExpressionLinearInterpolate
 //
-
-UMaterialExpressionLinearInterpolate::UMaterialExpressionLinearInterpolate(const FObjectInitializer& ObjectInitializer)
-	: Super(ObjectInitializer)
-{
-	ConstA = 0;
-	ConstB = 1;
-	ConstAlpha = 0.5f;
-}
-
 #if WITH_EDITOR
 int32 UMaterialExpressionLinearInterpolate::Compile(class FMaterialCompiler* Compiler, int32 OutputIndex)
 {
@@ -5081,14 +4992,6 @@ FString UMaterialExpressionConstant4Vector::GetDescription() const
 }
 #endif // WITH_EDITOR
 
-UMaterialExpressionClamp::UMaterialExpressionClamp(const FObjectInitializer& ObjectInitializer)
-	: Super(ObjectInitializer)
-{
-	ClampMode = CMODE_Clamp;
-	MinDefault = 0.0f;
-	MaxDefault = 1.0f;
-}
-
 void UMaterialExpressionClamp::Serialize(FStructuredArchive::FRecord Record)
 {
 	Super::Serialize(Record);
@@ -5170,14 +5073,6 @@ void UMaterialExpressionSaturate::GetExpressionToolTip(TArray<FString>& OutToolT
 //
 //	UMaterialExpressionMin
 //
-
-UMaterialExpressionMin::UMaterialExpressionMin(const FObjectInitializer& ObjectInitializer)
-	: Super(ObjectInitializer)
-{
-	ConstA = 0.0f;
-	ConstB = 1.0f;
-}
-
 #if WITH_EDITOR
 int32 UMaterialExpressionMin::Compile(class FMaterialCompiler* Compiler, int32 OutputIndex)
 {
@@ -5210,14 +5105,6 @@ void UMaterialExpressionMin::GetCaption(TArray<FString>& OutCaptions) const
 //
 //	UMaterialExpressionMax
 //
-
-UMaterialExpressionMax::UMaterialExpressionMax(const FObjectInitializer& ObjectInitializer)
-	: Super(ObjectInitializer)
-{
-	ConstA = 0.0f;
-	ConstB = 1.0f;
-}
-
 #if WITH_EDITOR
 int32 UMaterialExpressionMax::Compile(class FMaterialCompiler* Compiler, int32 OutputIndex)
 {
@@ -5254,9 +5141,6 @@ void UMaterialExpressionMax::GetCaption(TArray<FString>& OutCaptions) const
 UMaterialExpressionTextureCoordinate::UMaterialExpressionTextureCoordinate(const FObjectInitializer& ObjectInitializer)
 	: Super(ObjectInitializer)
 {
-	UTiling = 1.0f;
-	VTiling = 1.0f;
-
 #if WITH_EDITORONLY_DATA
 	bCollapsed = true;
 #endif
@@ -5446,20 +5330,10 @@ bool  UMaterialExpressionStaticComponentMaskParameter::SetParameterValue(FName I
 
 	return false;
 }
-#endif
 
 //
 //	UMaterialExpressionTime
 //
-
-UMaterialExpressionTime::UMaterialExpressionTime(const FObjectInitializer& ObjectInitializer)
-	: Super(ObjectInitializer)
-{
-	Period = 0.0f;
-	bOverride_Period = false;
-}
-
-#if WITH_EDITOR
 int32 UMaterialExpressionTime::Compile(class FMaterialCompiler* Compiler, int32 OutputIndex)
 {
 	return bIgnorePause ? Compiler->RealTime(bOverride_Period, Period) : Compiler->GameTime(bOverride_Period, Period);
@@ -5483,9 +5357,7 @@ void UMaterialExpressionTime::GetCaption(TArray<FString>& OutCaptions) const
 		OutCaptions.Add(TEXT("Time"));
 	}
 }
-#endif // WITH_EDITOR
 
-#if WITH_EDITOR
 int32 UMaterialExpressionCameraVectorWS::Compile(class FMaterialCompiler* Compiler, int32 OutputIndex)
 {
 	return Compiler->CameraVector();
@@ -5495,9 +5367,7 @@ void UMaterialExpressionCameraVectorWS::GetCaption(TArray<FString>& OutCaptions)
 {
 	OutCaptions.Add(TEXT("Camera Vector"));
 }
-#endif // WITH_EDITOR
 
-#if WITH_EDITOR
 int32 UMaterialExpressionCameraPositionWS::Compile(class FMaterialCompiler* Compiler, int32 OutputIndex)
 {
 	return Compiler->ViewProperty(MEVP_WorldSpaceCameraPosition);
@@ -5507,12 +5377,10 @@ void UMaterialExpressionCameraPositionWS::GetCaption(TArray<FString>& OutCaption
 {
 	OutCaptions.Add(TEXT("Camera Position"));
 }
-#endif // WITH_EDITOR
 
 //
 //	UMaterialExpressionReflectionVectorWS
 //
-#if WITH_EDITOR
 int32 UMaterialExpressionReflectionVectorWS::Compile(class FMaterialCompiler* Compiler, int32 OutputIndex)
 {
 	int32 Result = CustomWorldNormal.Compile(Compiler);
@@ -5542,7 +5410,6 @@ UMaterialExpressionPanner::UMaterialExpressionPanner(const FObjectInitializer& O
 #if WITH_EDITORONLY_DATA
 	bCollapsed = true;
 #endif
-	ConstCoordinate = 0;
 }
 
 #if WITH_EDITOR
@@ -5590,11 +5457,6 @@ void UMaterialExpressionPanner::GetCaption(TArray<FString>& OutCaptions) const
 UMaterialExpressionRotator::UMaterialExpressionRotator(const FObjectInitializer& ObjectInitializer)
 	: Super(ObjectInitializer)
 {
-	CenterX = 0.5f;
-	CenterY = 0.5f;
-	Speed = 0.25f;
-	ConstCoordinate = 0;
-
 #if WITH_EDITORONLY_DATA
 	bCollapsed = true;
 #endif
@@ -5643,18 +5505,10 @@ void UMaterialExpressionRotator::GetCaption(TArray<FString>& OutCaptions) const
 {
 	OutCaptions.Add(TEXT("Rotator"));
 }
-#endif // WITH_EDITOR
 
 //
 //	UMaterialExpressionSine
 //
-UMaterialExpressionSine::UMaterialExpressionSine(const FObjectInitializer& ObjectInitializer)
-	: Super(ObjectInitializer)
-{
-	Period = 1.0f;
-}
-
-#if WITH_EDITOR
 int32 UMaterialExpressionSine::Compile(class FMaterialCompiler* Compiler, int32 OutputIndex)
 {
 	if(!Input.GetTracedInput().Expression)
@@ -5669,18 +5523,10 @@ void UMaterialExpressionSine::GetCaption(TArray<FString>& OutCaptions) const
 {
 	OutCaptions.Add(TEXT("Sine"));
 }
-#endif // WITH_EDITOR
 
 //
 //	UMaterialExpressionCosine
 //
-UMaterialExpressionCosine::UMaterialExpressionCosine(const FObjectInitializer& ObjectInitializer)
-	: Super(ObjectInitializer)
-{
-	Period = 1.0f;
-}
-
-#if WITH_EDITOR
 int32 UMaterialExpressionCosine::Compile(class FMaterialCompiler* Compiler, int32 OutputIndex)
 {
 	if(!Input.GetTracedInput().Expression)
@@ -5695,18 +5541,10 @@ void UMaterialExpressionCosine::GetCaption(TArray<FString>& OutCaptions) const
 {
 	OutCaptions.Add(TEXT("Cosine"));
 }
-#endif // WITH_EDITOR
 
 //
 //	UMaterialExpressionTangent
 //
-UMaterialExpressionTangent::UMaterialExpressionTangent(const FObjectInitializer& ObjectInitializer)
-	: Super(ObjectInitializer)
-{
-	Period = 1.0f;
-}
-
-#if WITH_EDITOR
 int32 UMaterialExpressionTangent::Compile(class FMaterialCompiler* Compiler, int32 OutputIndex)
 {
 	if(!Input.GetTracedInput().Expression)
@@ -5930,9 +5768,6 @@ void UMaterialExpressionArctangent2Fast::GetExpressionToolTip(TArray<FString>& O
 UMaterialExpressionBumpOffset::UMaterialExpressionBumpOffset(const FObjectInitializer& ObjectInitializer)
 	: Super(ObjectInitializer)
 {
-	HeightRatio = 0.05f;
-	ReferencePlane = 0.5f;
-	ConstCoordinate = 0;
 #if WITH_EDITORONLY_DATA
 	bCollapsed = false;
 #endif
@@ -6957,8 +6792,6 @@ void UMaterialExpressionSetMaterialAttributes::PostEditChangeProperty(FPropertyC
 
 UMaterialExpressionBlendMaterialAttributes::UMaterialExpressionBlendMaterialAttributes(const FObjectInitializer& ObjectInitializer)
 	: Super(ObjectInitializer)
-	, PixelAttributeBlendType(EMaterialAttributeBlend::Blend)
-	, VertexAttributeBlendType(EMaterialAttributeBlend::Blend)
 {
 #if WITH_EDITORONLY_DATA
 	Outputs.Reset();
@@ -7108,9 +6941,6 @@ FSubstrateOperator* UMaterialExpressionBlendMaterialAttributes::SubstrateGenerat
 }
 #endif // WITH_EDITOR
 
-UMaterialExpressionLegacyBlendMaterialAttributes::UMaterialExpressionLegacyBlendMaterialAttributes(const FObjectInitializer& ObjectInitializer) 
-	: Super(ObjectInitializer){}
-
 #if ENABLE_MATERIAL_LAYER_PROTOTYPE
 #if WITH_EDITOR
 void UMaterialExpressionLegacyBlendMaterialAttributes::GetCaption(TArray<FString>& OutCaptions) const
@@ -7283,10 +7113,6 @@ int32 UMaterialExpressionLegacyBlendMaterialAttributes::Compile(FMaterialCompile
 //
 UMaterialExpressionMaterialAttributeLayers::UMaterialExpressionMaterialAttributeLayers(const FObjectInitializer& ObjectInitializer)
 	: Super(ObjectInitializer)	
-	, NumActiveLayerCallers(0)
-	, NumActiveBlendCallers(0)
-	, bIsLayerGraphBuilt(false)
-	, ParamLayers(nullptr)
 {
 #if WITH_EDITOR
 	DefaultLayers.AddDefaultBackgroundLayer();
@@ -8585,10 +8411,6 @@ UMaterialExpressionChannelMaskParameter::UMaterialExpressionChannelMaskParameter
 	Outputs.Add(FExpressionOutput(TEXT(""), 1, 1, 0, 0, 0));
 	bShowMaskColorsOnPin = false;
 #endif
-
-	// Default mask to red channel
-	DefaultValue = FLinearColor(1.0f, 0.0f, 0.0f, 0.0f);
-	MaskChannel = EChannelMaskParameterColor::Red;
 }
 
 #if WITH_EDITOR
@@ -9893,20 +9715,10 @@ void UMaterialExpressionTextureObjectFromCollection::GetCaption(TArray<FString>&
 {
 	OutCaptions.Emplace(TEXT("Texture Object From Collection"));
 }
-#endif
 
 //
 //	UMaterialExpressionRequiredSamplersSwitch
 //
-
-UMaterialExpressionRequiredSamplersSwitch::UMaterialExpressionRequiredSamplersSwitch(const FObjectInitializer& ObjectInitializer)
-	: Super(ObjectInitializer),
-	RequiredSamplers(16)
-{
-}
-
-#if WITH_EDITOR
-
 int32 UMaterialExpressionRequiredSamplersSwitch::Compile(class FMaterialCompiler* Compiler, int32 OutputIndex)
 {
 	const EShaderPlatform ShaderPlatform = Compiler->GetShaderPlatform();
@@ -10261,11 +10073,6 @@ UMaterialExpressionDynamicParameter::UMaterialExpressionDynamicParameter(const F
 	Outputs.Add(FExpressionOutput(TEXT("RGB"), 1, 1, 1, 1, 0));
 	Outputs.Add(FExpressionOutput(TEXT("RGBA"), 1, 1, 1, 1, 1));
 #endif // WITH_EDITORONLY_DATA
-
-	DefaultValue = FLinearColor::White;
-
-	
-	ParameterIndex = 0;
 }
 
 #if WITH_EDITOR
@@ -10364,12 +10171,6 @@ bool UMaterialExpressionDynamicParameter::CopyDynamicParameterProperties(const U
 //
 //	MaterialExpressionParticleSubUV
 //
-UMaterialExpressionParticleSubUV::UMaterialExpressionParticleSubUV(const FObjectInitializer& ObjectInitializer)
-	: Super(ObjectInitializer)
-{
-	bBlend = true;
-}
-
 #if WITH_EDITOR
 int32 UMaterialExpressionParticleSubUV::Compile(class FMaterialCompiler* Compiler, int32 OutputIndex)
 {
@@ -10511,8 +10312,6 @@ UMaterialExpressionViewProperty::UMaterialExpressionViewProperty(const FObjectIn
 	Outputs.Add(FExpressionOutput(TEXT("Property")));
 	Outputs.Add(FExpressionOutput(TEXT("InvProperty")));
 #endif
-
-	Property = MEVP_FieldOfView;
 }
 
 #if WITH_EDITOR
@@ -10684,8 +10483,6 @@ UMaterialExpressionSceneDepth::UMaterialExpressionSceneDepth(const FObjectInitia
 	Outputs.Reset();
 	Outputs.Add(FExpressionOutput(TEXT(""), 1, 1, 0, 0, 0));
 #endif
-
-	ConstInput = FVector2D(0.f, 0.f);
 }
 
 void UMaterialExpressionSceneDepth::PostLoad()
@@ -10761,9 +10558,6 @@ UMaterialExpressionSceneTexture::UMaterialExpressionSceneTexture(const FObjectIn
 	bShowOutputNameOnPin = true;
 #endif
 
-	// by default faster, most lookup are read/write the same pixel so this is ralrely needed
-	bFiltered = false;
-
 #if WITH_EDITORONLY_DATA
 	Outputs.Reset();
 	Outputs.Add(FExpressionOutput(TEXT("Color"), 1, 1, 1, 1, 1));
@@ -10822,9 +10616,6 @@ UMaterialExpressionUserSceneTexture::UMaterialExpressionUserSceneTexture(const F
 	bShowOutputNameOnPin = true;
 #endif
 
-	// by default faster, most lookup are read/write the same pixel so this is rarely needed
-	bFiltered = false;
-
 #if WITH_EDITORONLY_DATA
 	Outputs.Reset();
 	Outputs.Add(FExpressionOutput(TEXT("Color"), 1, 1, 1, 1, 1));
@@ -10883,7 +10674,6 @@ UMaterialExpressionSceneColor::UMaterialExpressionSceneColor(const FObjectInitia
 #if WITH_EDITORONLY_DATA
 	bShowOutputNameOnPin = true;
 #endif
-	ConstInput = FVector2D(0.f, 0.f);
 
 	Outputs.Reset();
 	Outputs.Add(FExpressionOutput(TEXT("RGB"), 1, 1, 1, 1, 0));
@@ -10980,15 +10770,7 @@ void UMaterialExpressionDBufferTexture::GetCaption(TArray<FString>& OutCaptions)
 	FString Name = Enum->GetDisplayNameTextByValue(DBufferTextureId).ToString();
 	OutCaptions.Add(Name);
 }
-#endif // WITH_EDITOR
-
-UMaterialExpressionPower::UMaterialExpressionPower(const FObjectInitializer& ObjectInitializer)
-	: Super(ObjectInitializer)
-{
-	ConstExponent = 2;
-}
 
-#if WITH_EDITOR
 int32 UMaterialExpressionPower::Compile(class FMaterialCompiler* Compiler, int32 OutputIndex)
 {
 	if(!Base.GetTracedInput().Expression)
@@ -11233,16 +11015,7 @@ FName UMaterialExpressionSceneColor::GetInputName(int32 InputIndex) const
 	}
 	return NAME_None;
 }
-#endif // WITH_EDITOR
-
-UMaterialExpressionIf::UMaterialExpressionIf(const FObjectInitializer& ObjectInitializer)
-	: Super(ObjectInitializer)
-{
-	EqualsThreshold = 0.00001f;
-	ConstB = 0.0f;
-}
 
-#if WITH_EDITOR
 int32 UMaterialExpressionIf::Compile(class FMaterialCompiler* Compiler, int32 OutputIndex)
 {
 	if(!A.GetTracedInput().Expression)
@@ -11436,24 +11209,9 @@ void UMaterialExpressionTransform::GetCaption(TArray<FString>& OutCaptions) cons
 }
 #endif // WITH_EDITOR
 
-UMaterialExpressionTransform::UMaterialExpressionTransform(const FObjectInitializer& ObjectInitializer)
-	: Super(ObjectInitializer)
-{
-	TransformSourceType = TRANSFORMSOURCE_Tangent;
-	TransformType = TRANSFORM_World;
-}
-
-
 ///////////////////////////////////////////////////////////////////////////////
 // UMaterialExpressionTransformPosition
 ///////////////////////////////////////////////////////////////////////////////
-UMaterialExpressionTransformPosition::UMaterialExpressionTransformPosition(const FObjectInitializer& ObjectInitializer)
-	: Super(ObjectInitializer)
-{
-	TransformSourceType = TRANSFORMPOSSOURCE_Local;
-	TransformType = TRANSFORMPOSSOURCE_Local;
-}
-
 static EMaterialCommonBasis GetMaterialCommonBasis(EMaterialPositionTransformSource X)
 {
 	static const EMaterialCommonBasis ConversionTable[TRANSFORMPOSSOURCE_MAX] = {
@@ -11600,23 +11358,10 @@ int32 UMaterialExpressionTransformPosition::GetAbsoluteInputIndex(int32 InputInd
 	return INDEX_NONE;
 }
 #undef IF_INPUT_RETURN
-#endif // WITH_EDITOR
 
 ///////////////////////////////////////////////////////////////////////////////
 // UMaterialExpressionComment
 ///////////////////////////////////////////////////////////////////////////////
-UMaterialExpressionComment::UMaterialExpressionComment(const FObjectInitializer& ObjectInitializer)
-	: Super(ObjectInitializer)
-	, CommentColor(FLinearColor::White)
-	, FontSize(18)
-	, bCommentBubbleVisible_InDetailsPanel(false)
-	, bColorCommentBubble(false)
-	, bGroupMode(true)
-
-{
-}
-
-#if WITH_EDITOR
 
 void UMaterialExpressionComment::PostEditChangeProperty(FPropertyChangedEvent& PropertyChangedEvent)
 {
@@ -12136,16 +11881,6 @@ bool UMaterialExpressionPinBase::Modify(bool bAlwaysMarkDirty)
 	return Super::Modify(bAlwaysMarkDirty);
 }
 
-#endif // WITH_EDITOR
-
-UMaterialExpressionFresnel::UMaterialExpressionFresnel(const FObjectInitializer& ObjectInitializer)
-	: Super(ObjectInitializer)
-{
-	Exponent = 5.0f;
-	BaseReflectFraction = 0.04f;
-}
-
-#if WITH_EDITOR
 int32 UMaterialExpressionFresnel::Compile(class FMaterialCompiler* Compiler, int32 OutputIndex)
 {
 	// pow(1 - max(0,Normal dot Camera),Exponent) * (1 - BaseReflectFraction) + BaseReflectFraction
@@ -12410,9 +12145,6 @@ UMaterialExpressionLocalPosition::UMaterialExpressionLocalPosition(const FObject
 
 	bShowOutputNameOnPin = true;
 #endif
-
-	IncludedOffsets = EPositionIncludedOffsets::IncludeOffsets;
-	LocalOrigin = ELocalPositionOrigin::Instance;
 }
 
 #if WITH_EDITOR
@@ -12469,7 +12201,6 @@ UMaterialExpressionWorldPosition::UMaterialExpressionWorldPosition(const FObject
 
 	bShowOutputNameOnPin = true;
 #endif
-	WorldPositionShaderOffset = WPT_Default;
 }
 
 #if WITH_EDITOR
@@ -12938,19 +12669,11 @@ void UMaterialExpressionDeriveNormalZ::GetCaption(TArray<FString>& OutCaptions)
 {
 	OutCaptions.Add(TEXT("DeriveNormalZ"));
 }
-#endif // WITH_EDITOR
 
 ///////////////////////////////////////////////////////////////////////////////
 // UMaterialExpressionConstantBiasScale
 ///////////////////////////////////////////////////////////////////////////////
-UMaterialExpressionConstantBiasScale::UMaterialExpressionConstantBiasScale(const FObjectInitializer& ObjectInitializer)
-	: Super(ObjectInitializer)
-{
-	Bias = 1.0f;
-	Scale = 0.5f;
-}
 
-#if WITH_EDITOR
 int32 UMaterialExpressionConstantBiasScale::Compile(class FMaterialCompiler* Compiler, int32 OutputIndex)
 {
 	if (!Input.GetTracedInput().Expression)
@@ -13411,8 +13134,6 @@ void UMaterialFunctionInterfaceEditorOnlyData::Serialize(FArchive& Ar)
 UMaterialExpressionSwitch::UMaterialExpressionSwitch(const FObjectInitializer& ObjectInitializer)
 	: Super(ObjectInitializer)
 {
-	Description = TEXT("Switch");
-
 	Inputs.Add(FSwitchCustomInput());
 	Inputs[0].InputName = TEXT("");
 
@@ -13560,12 +13281,6 @@ namespace MaterialFunctionInterface
 	}
 }
 
-UMaterialFunctionInterface::UMaterialFunctionInterface(const FObjectInitializer& ObjectInitializer)
-	: Super(ObjectInitializer)
-	, MaterialFunctionUsage(EMaterialFunctionUsage::Default)
-{
-}
-
 void UMaterialFunctionInterface::PostInitProperties()
 {
 #if WITH_EDITORONLY_DATA
@@ -19973,12 +19688,6 @@ bool UMaterialExpressionNamedRerouteUsage::GetRerouteInput(FExpressionInput& Out
 	}
 }
 
-UMaterialExpressionRotateAboutAxis::UMaterialExpressionRotateAboutAxis(const FObjectInitializer& ObjectInitializer)
-	: Super(ObjectInitializer)
-{
-	Period = 1.0f;
-}
-
 #if WITH_EDITOR
 int32 UMaterialExpressionRotateAboutAxis::Compile(class FMaterialCompiler* Compiler, int32 OutputIndex)
 {
@@ -20037,12 +19746,6 @@ static int32 CompileHelperLength( FMaterialCompiler* Compiler, int32 A, int32 B
 ///////////////////////////////////////////////////////////////////////////////
 // UMaterialExpressionSphereMask
 ///////////////////////////////////////////////////////////////////////////////
-UMaterialExpressionSphereMask::UMaterialExpressionSphereMask(const FObjectInitializer& ObjectInitializer)
-	: Super(ObjectInitializer)
-{
-	AttenuationRadius = 256.0f;
-	HardnessPercent = 100.0f;
-}
 
 #if WITH_EDITOR
 int32 UMaterialExpressionSphereMask::Compile(class FMaterialCompiler* Compiler, int32 OutputIndex)
@@ -20103,19 +19806,11 @@ void UMaterialExpressionSphereMask::GetCaption(TArray<FString>& OutCaptions) con
 {
 	OutCaptions.Add(TEXT("SphereMask"));
 }
-#endif // WITH_EDITOR
 
 ///////////////////////////////////////////////////////////////////////////////
 // UMaterialExpressionSobol
 ///////////////////////////////////////////////////////////////////////////////
-UMaterialExpressionSobol::UMaterialExpressionSobol(const FObjectInitializer& ObjectInitializer)
-	: Super(ObjectInitializer)
-{
-	ConstIndex = 0;
-	ConstSeed = FVector2D(0.f, 0.f);
-}
 
-#if WITH_EDITOR
 int32 UMaterialExpressionSobol::Compile(class FMaterialCompiler* Compiler, int32 OutputIndex)
 {
 	int32 CellInput = Cell.GetTracedInput().Expression ? Cell.Compile(Compiler) : Compiler->Constant2(0.f, 0.f);
@@ -20135,19 +19830,11 @@ void UMaterialExpressionSobol::GetCaption(TArray<FString>& OutCaptions) const
 
 	OutCaptions.Add(Caption);
 }
-#endif // WITH_EDITOR
 
 ///////////////////////////////////////////////////////////////////////////////
 // UMaterialExpressionTemporalSobol
 ///////////////////////////////////////////////////////////////////////////////
-UMaterialExpressionTemporalSobol::UMaterialExpressionTemporalSobol(const FObjectInitializer& ObjectInitializer)
-	: Super(ObjectInitializer)
-{
-	ConstIndex = 0;
-	ConstSeed = FVector2D(0.f, 0.f);
-}
 
-#if WITH_EDITOR
 int32 UMaterialExpressionTemporalSobol::Compile(class FMaterialCompiler* Compiler, int32 OutputIndex)
 {
 	int32 IndexInput = Index.GetTracedInput().Expression ? Index.Compile(Compiler) : Compiler->Constant(ConstIndex);
@@ -20326,21 +20013,11 @@ void UMaterialExpressionDistanceFieldGradient::GetCaption(TArray<FString>& OutCa
 {
 	OutCaptions.Add(TEXT("DistanceFieldGradient"));
 }
-#endif // WITH_EDITOR
 
 ///////////////////////////////////////////////////////////////////////////////
 // UMaterialExpressionDistanceFieldApproxAO
 ///////////////////////////////////////////////////////////////////////////////
-UMaterialExpressionDistanceFieldApproxAO::UMaterialExpressionDistanceFieldApproxAO(const FObjectInitializer& ObjectInitializer)
-	: Super(ObjectInitializer)
-{
-	NumSteps = 1;
-	BaseDistanceDefault = 15;
-	RadiusDefault = 150;
-	StepScaleDefault = 3.0f;
-}
 
-#if WITH_EDITOR
 FName UMaterialExpressionDistanceFieldApproxAO::GetInputName(int32 InputIndex) const
 {
 	if (GetInput(InputIndex) == &Position)
@@ -20449,19 +20126,11 @@ void UMaterialExpressionSamplePhysicsVectorField::GetCaption(TArray<FString>& Ou
 {
 	OutCaptions.Add(TEXT("SamplePhysicsVectorField"));
 }
-#endif // WITH_EDITOR
 
 ///////////////////////////////////////////////////////////////////////////////
 // UMaterialExpressionSamplePhysicsScalarField
 ///////////////////////////////////////////////////////////////////////////////
 
-UMaterialExpressionSamplePhysicsScalarField::UMaterialExpressionSamplePhysicsScalarField(const FObjectInitializer& ObjectInitializer)
-	: Super(ObjectInitializer)
-{
-	FieldTarget = EFieldScalarType::Scalar_DynamicConstraint;
-}
-
-#if WITH_EDITOR
 FName UMaterialExpressionSamplePhysicsScalarField::GetInputName(int32 InputIndex) const
 {
 	if (GetInput(InputIndex) == &WorldPosition)
@@ -20506,20 +20175,11 @@ void UMaterialExpressionSamplePhysicsScalarField::GetCaption(TArray<FString>& Ou
 {
 	OutCaptions.Add(TEXT("SamplePhysicsScalarField"));
 }
-#endif // WITH_EDITOR
-
 
 ///////////////////////////////////////////////////////////////////////////////
 // UMaterialExpressionSamplePhysicsIntegerField
 ///////////////////////////////////////////////////////////////////////////////
 
-UMaterialExpressionSamplePhysicsIntegerField::UMaterialExpressionSamplePhysicsIntegerField(const FObjectInitializer& ObjectInitializer)
-	: Super(ObjectInitializer)
-{
-	FieldTarget = EFieldIntegerType::Integer_DynamicState;
-}
-
-#if WITH_EDITOR
 FName UMaterialExpressionSamplePhysicsIntegerField::GetInputName(int32 InputIndex) const
 {
 	if (GetInput(InputIndex) == &WorldPosition)
@@ -20730,9 +20390,7 @@ UMaterialExpressionAntialiasedTextureMask::UMaterialExpressionAntialiasedTexture
 	Texture = ConstructorStatics.DefaultTexture.Object;
 #endif
 
-	Threshold = 0.5f;
 	ParameterName = NAME_None;
-	Channel = TCC_Alpha;
 
 #if WITH_EDITORONLY_DATA
 	Outputs.Reset();
@@ -20951,9 +20609,6 @@ void UMaterialExpressionDecalMipmapLevel::GetCaption(TArray<FString>& OutCaption
 UMaterialExpressionDepthFade::UMaterialExpressionDepthFade(const FObjectInitializer& ObjectInitializer)
 	: Super(ObjectInitializer)
 {
-	FadeDistanceDefault = 100.0f;
-	OpacityDefault = 1.0f;
-
 #if WITH_EDITORONLY_DATA
 	bCollapsed = false;
 #endif
@@ -20991,8 +20646,6 @@ int32 UMaterialExpressionDepthFade::Compile(class FMaterialCompiler* Compiler, i
 UMaterialExpressionSphericalParticleOpacity::UMaterialExpressionSphericalParticleOpacity(const FObjectInitializer& ObjectInitializer)
 	: Super(ObjectInitializer)
 {
-	ConstantDensity = 1;
-
 #if WITH_EDITORONLY_DATA
 	bCollapsed = false;
 #endif
@@ -21323,22 +20976,11 @@ void UMaterialExpressionAtmosphericFogColor::GetCaption(TArray<FString>& OutCapt
 {
 	OutCaptions.Add(TEXT("Atmospheric Fog Color (deprecated)"));
 }
-#endif // WITH_EDITOR
 
 /*------------------------------------------------------------------------------
 	SpeedTree material expression.
 ------------------------------------------------------------------------------*/
-UMaterialExpressionSpeedTree::UMaterialExpressionSpeedTree(const FObjectInitializer& ObjectInitializer)
-	: Super(ObjectInitializer)
-{
-	GeometryType = STG_Branch;
-	WindType = STW_None;
-	LODType = STLOD_Pop;
-	BillboardThreshold = 0.9f;
-	bAccurateWindVelocities = false;
-}
 
-#if WITH_EDITOR
 int32 UMaterialExpressionSpeedTree::Compile(class FMaterialCompiler* Compiler, int32 OutputIndex)
 {
 	int32 GeometryArg = (GeometryInput.GetTracedInput().Expression ? GeometryInput.Compile(Compiler) : Compiler->Constant(GeometryType));
@@ -22175,8 +21817,6 @@ void UMaterialExpressionCurveAtlasRowParameter::PostEditChangeProperty(FProperty
 UMaterialExpressionHairAttributes::UMaterialExpressionHairAttributes(const FObjectInitializer& ObjectInitializer)
 	: Super(ObjectInitializer)
 {
-	bUseTangentSpace = true;
-
 #if WITH_EDITORONLY_DATA
 	bShowOutputNameOnPin = true;
 
@@ -22444,21 +22084,6 @@ FString UMaterialExpressionSingleLayerWaterMaterialOutput::GetDisplayName() cons
 UMaterialExpressionVolumetricAdvancedMaterialOutput::UMaterialExpressionVolumetricAdvancedMaterialOutput(const FObjectInitializer& ObjectInitializer)
 	: Super(ObjectInitializer)
 {
-	ConstPhaseG = 0.0f;
-	ConstPhaseG2 = 0.0f;
-	ConstPhaseBlend = 0.0f;
-	PerSamplePhaseEvaluation = false;
-
-	MultiScatteringApproximationOctaveCount = 0;
-	ConstMultiScatteringContribution = 0.5f;
-	ConstMultiScatteringOcclusion = 0.5f;
-	ConstMultiScatteringEccentricity = 0.5f;
-
-	bGroundContribution = false;
-	bGrayScaleMaterial = false;
-	bRayMarchVolumeShadow = true;
-	bClampMultiScatteringContribution = true;
-
 #if WITH_EDITOR
 	Outputs.Reset();
 #endif
@@ -23683,7 +23308,6 @@ uint32 UMaterialExpressionNeuralNetworkOutput::GetInputType(int32 InputIndex)
 UMaterialExpressionFirstPersonOutput::UMaterialExpressionFirstPersonOutput(const FObjectInitializer& ObjectInitializer)
 	: Super(ObjectInitializer)
 {
-	ConstFirstPersonInterpolationAlpha = 1.0f;
 #if WITH_EDITOR
 	Outputs.Reset();
 #endif
diff --git a/Engine/Source/Runtime/Engine/Public/Materials/MaterialExpressionAdd.h b/Engine/Source/Runtime/Engine/Public/Materials/MaterialExpressionAdd.h
--- a/Engine/Source/Runtime/Engine/Public/Materials/MaterialExpressionAdd.h
+++ b/Engine/Source/Runtime/Engine/Public/Materials/MaterialExpressionAdd.h
@@ -12,7 +12,9 @@
 UCLASS(MinimalAPI)
 class UMaterialExpressionAdd : public UMaterialExpression
 {
-	GENERATED_UCLASS_BODY()
+	GENERATED_BODY()
+
+public:
 
 	UPROPERTY(meta = (RequiredInput = "false", ToolTip = "Defaults to 'ConstA' if not specified"))
 	FExpressionInput A;
@@ -22,12 +24,11 @@ class UMaterialExpressionAdd : public UMaterialExpression
 
 	/** only used if A is not hooked up */
 	UPROPERTY(EditAnywhere, Category=MaterialExpressionAdd, meta=(OverridingInputProperty = "A"))
-	float ConstA;
+	float ConstA = 0.0f;
 
 	/** only used if B is not hooked up */
 	UPROPERTY(EditAnywhere, Category=MaterialExpressionAdd, meta=(OverridingInputProperty = "B"))
-	float ConstB;
-
+	float ConstB = 1.0f;
 
 	//~ Begin UMaterialExpression Interface
 #if WITH_EDITOR
diff --git a/Engine/Source/Runtime/Engine/Public/Materials/MaterialExpressionAntialiasedTextureMask.h b/Engine/Source/Runtime/Engine/Public/Materials/MaterialExpressionAntialiasedTextureMask.h
--- a/Engine/Source/Runtime/Engine/Public/Materials/MaterialExpressionAntialiasedTextureMask.h
+++ b/Engine/Source/Runtime/Engine/Public/Materials/MaterialExpressionAntialiasedTextureMask.h
@@ -26,10 +26,10 @@ class UMaterialExpressionAntialiasedTextureMask : public UMaterialExpressionText
 	GENERATED_UCLASS_BODY()
 
 	UPROPERTY(EditAnywhere, Category=MaterialExpressionAntialiasedTextureMask, meta=(UIMin = "0.0", UIMax = "1.0", ClampMin = "0.0", ClampMax = "1.0"))
-	float Threshold;
+	float Threshold = 0.5f;
 
 	UPROPERTY(EditAnywhere, Category=MaterialExpressionAntialiasedTextureMask)
-	TEnumAsByte<enum ETextureColorChannel> Channel;
+	TEnumAsByte<enum ETextureColorChannel> Channel = TCC_Alpha;
 
 #if WITH_EDITOR
 	//~ Begin UMaterialExpressionTextureSampleParameter Interface
diff --git a/Engine/Source/Runtime/Engine/Public/Materials/MaterialExpressionBlendMaterialAttributes.h b/Engine/Source/Runtime/Engine/Public/Materials/MaterialExpressionBlendMaterialAttributes.h
--- a/Engine/Source/Runtime/Engine/Public/Materials/MaterialExpressionBlendMaterialAttributes.h
+++ b/Engine/Source/Runtime/Engine/Public/Materials/MaterialExpressionBlendMaterialAttributes.h
@@ -46,11 +46,11 @@ class UMaterialExpressionBlendMaterialAttributes : public UMaterialExpression
 
 	// Optionally skip blending attributes of this type.
 	UPROPERTY(EditAnywhere, Category=MaterialAttributes)
-	TEnumAsByte<EMaterialAttributeBlend::Type> PixelAttributeBlendType;
+	TEnumAsByte<EMaterialAttributeBlend::Type> PixelAttributeBlendType = EMaterialAttributeBlend::Blend;
 
 	// Optionally skip blending attributes of this type.
 	UPROPERTY(EditAnywhere, Category=MaterialAttributes)
-	TEnumAsByte<EMaterialAttributeBlend::Type> VertexAttributeBlendType;
+	TEnumAsByte<EMaterialAttributeBlend::Type> VertexAttributeBlendType = EMaterialAttributeBlend::Blend;
 
 	//~ Begin UMaterialExpression Interface
 #if WITH_EDITOR
@@ -72,7 +72,9 @@ class UMaterialExpressionBlendMaterialAttributes : public UMaterialExpression
 UCLASS(collapsecategories, hidecategories = Object, MinimalAPI)
 class UMaterialExpressionLegacyBlendMaterialAttributes : public UMaterialExpressionBlendMaterialAttributes
 {
-	GENERATED_UCLASS_BODY()
+	GENERATED_BODY()
+
+public:
 
 	UPROPERTY()
 	FExpressionInput VertexAttribute_UseA;
diff --git a/Engine/Source/Runtime/Engine/Public/Materials/MaterialExpressionBumpOffset.h b/Engine/Source/Runtime/Engine/Public/Materials/MaterialExpressionBumpOffset.h
--- a/Engine/Source/Runtime/Engine/Public/Materials/MaterialExpressionBumpOffset.h
+++ b/Engine/Source/Runtime/Engine/Public/Materials/MaterialExpressionBumpOffset.h
@@ -25,14 +25,14 @@ class UMaterialExpressionBumpOffset : public UMaterialExpression
 	FExpressionInput HeightRatioInput;
 
 	UPROPERTY(EditAnywhere, Category=MaterialExpressionBumpOffset, meta=(OverridingInputProperty = "HeightRatioInput"))
-	float HeightRatio;    // Perceived height as a fraction of width.
+	float HeightRatio = 0.05f; // Perceived height as a fraction of width.
 
 	UPROPERTY(EditAnywhere, Category=MaterialExpressionBumpOffset)
-	float ReferencePlane;    // Height at which no offset is applied.
+	float ReferencePlane = 0.5f;    // Height at which no offset is applied.
 
 	/** only used if Coordinate is not hooked up */
 	UPROPERTY(EditAnywhere, Category = MaterialExpressionBumpOffset, meta = (OverridingInputProperty = "Coordinate"))
-	uint32 ConstCoordinate;
+	uint32 ConstCoordinate = 0;
 
 	//~ Begin UMaterialExpression Interface
 #if WITH_EDITOR
diff --git a/Engine/Source/Runtime/Engine/Public/Materials/MaterialExpressionChannelMaskParameter.h b/Engine/Source/Runtime/Engine/Public/Materials/MaterialExpressionChannelMaskParameter.h
--- a/Engine/Source/Runtime/Engine/Public/Materials/MaterialExpressionChannelMaskParameter.h
+++ b/Engine/Source/Runtime/Engine/Public/Materials/MaterialExpressionChannelMaskParameter.h
@@ -15,7 +15,7 @@ class UMaterialExpressionChannelMaskParameter : public UMaterialExpressionVector
 	GENERATED_UCLASS_BODY()
 
 	UPROPERTY(EditAnywhere, Category=MaterialExpressionChannelMaskParameter)
-	TEnumAsByte<EChannelMaskParameterColor::Type> MaskChannel;
+	TEnumAsByte<EChannelMaskParameterColor::Type> MaskChannel = EChannelMaskParameterColor::Red;
 
 	UPROPERTY()
 	FExpressionInput Input;
diff --git a/Engine/Source/Runtime/Engine/Public/Materials/MaterialExpressionClamp.h b/Engine/Source/Runtime/Engine/Public/Materials/MaterialExpressionClamp.h
--- a/Engine/Source/Runtime/Engine/Public/Materials/MaterialExpressionClamp.h
+++ b/Engine/Source/Runtime/Engine/Public/Materials/MaterialExpressionClamp.h
@@ -20,7 +20,9 @@ enum EClampMode : int
 UCLASS(MinimalAPI)
 class UMaterialExpressionClamp : public UMaterialExpression
 {
-	GENERATED_UCLASS_BODY()
+	GENERATED_BODY()
+
+public:
 
 	UPROPERTY()
 	FExpressionInput Input;
@@ -32,13 +34,13 @@ class UMaterialExpressionClamp : public UMaterialExpression
 	FExpressionInput Max;
 
 	UPROPERTY(EditAnywhere, Category=MaterialExpressionClamp, Meta = (ShowAsInputPin = "Advanced"))
-	TEnumAsByte<enum EClampMode> ClampMode;
+	TEnumAsByte<enum EClampMode> ClampMode = CMODE_Clamp;
 
 	UPROPERTY(EditAnywhere, Category=MaterialExpressionClamp, meta=(OverridingInputProperty = "Min"))
-	float MinDefault;
+	float MinDefault = 0.0f;
 
 	UPROPERTY(EditAnywhere, Category=MaterialExpressionClamp, meta=(OverridingInputProperty = "Max"))
-	float MaxDefault;
+	float MaxDefault = 1.0f;
 
 
 	//~ Begin UObject Interface
diff --git a/Engine/Source/Runtime/Engine/Public/Materials/MaterialExpressionComment.h b/Engine/Source/Runtime/Engine/Public/Materials/MaterialExpressionComment.h
--- a/Engine/Source/Runtime/Engine/Public/Materials/MaterialExpressionComment.h
+++ b/Engine/Source/Runtime/Engine/Public/Materials/MaterialExpressionComment.h
@@ -13,7 +13,9 @@ struct FPropertyChangedEvent;
 UCLASS(MinimalAPI)
 class UMaterialExpressionComment : public UMaterialExpression
 {
-	GENERATED_UCLASS_BODY()
+	GENERATED_BODY()
+
+public:
 
 	UPROPERTY()
 	int32 SizeX;
@@ -26,23 +28,23 @@ class UMaterialExpressionComment : public UMaterialExpression
 
 	/** Color to style comment with */
 	UPROPERTY(EditAnywhere, Category=MaterialExpressionComment)
-	FLinearColor CommentColor;
+	FLinearColor CommentColor = FLinearColor::White;
 
 	/** Size of the text in the comment box */
 	UPROPERTY(EditAnywhere, Category=MaterialExpressionComment, meta=(ClampMin=1, ClampMax=1000))
-	int32 FontSize;
+	int32 FontSize = 18;
 
 	/** Whether to show a zoom-invariant comment bubble when zoomed out (making the comment readable at any distance). */
 	UPROPERTY(EditAnywhere, Category=MaterialExpressionComment, meta=(DisplayName="Show Bubble When Zoomed"))
-	uint32 bCommentBubbleVisible_InDetailsPanel:1;
+	uint32 bCommentBubbleVisible_InDetailsPanel:1 = false;
 
 	/** Whether to use Comment Color to color the background of the comment bubble shown when zoomed out. */
 	UPROPERTY(EditAnywhere, Category=MaterialExpressionComment, meta=(DisplayName="Color Bubble", EditCondition=bCommentBubbleVisible_InDetailsPanel))
-	uint32 bColorCommentBubble:1;
+	uint32 bColorCommentBubble:1 = false;
 
 	/** Whether the comment should move any fully enclosed nodes around when it is moved */
 	UPROPERTY(EditAnywhere, Category=MaterialExpressionComment)
-	bool bGroupMode;
+	bool bGroupMode = true;
 	
 	//~ Begin UObject Interface
 #if WITH_EDITOR
diff --git a/Engine/Source/Runtime/Engine/Public/Materials/MaterialExpressionConstantBiasScale.h b/Engine/Source/Runtime/Engine/Public/Materials/MaterialExpressionConstantBiasScale.h
--- a/Engine/Source/Runtime/Engine/Public/Materials/MaterialExpressionConstantBiasScale.h
+++ b/Engine/Source/Runtime/Engine/Public/Materials/MaterialExpressionConstantBiasScale.h
@@ -12,17 +12,18 @@
 UCLASS(MinimalAPI)
 class UMaterialExpressionConstantBiasScale : public UMaterialExpression
 {
-	GENERATED_UCLASS_BODY()
+	GENERATED_BODY()
+
+public:
 
 	UPROPERTY()
 	FExpressionInput Input;
 
 	UPROPERTY(EditAnywhere, Category=MaterialExpressionConstantBiasScale, Meta = (ShowAsInputPin = "Advanced"))
-	float Bias;
+	float Bias = 1.0f;
 
 	UPROPERTY(EditAnywhere, Category=MaterialExpressionConstantBiasScale, Meta = (ShowAsInputPin = "Advanced"))
-	float Scale;
-
+	float Scale = 0.5f;
 
 	//~ Begin UMaterialExpression Interface
 #if WITH_EDITOR
diff --git a/Engine/Source/Runtime/Engine/Public/Materials/MaterialExpressionCosine.h b/Engine/Source/Runtime/Engine/Public/Materials/MaterialExpressionCosine.h
--- a/Engine/Source/Runtime/Engine/Public/Materials/MaterialExpressionCosine.h
+++ b/Engine/Source/Runtime/Engine/Public/Materials/MaterialExpressionCosine.h
@@ -12,14 +12,15 @@
 UCLASS(MinimalAPI, collapsecategories, hidecategories=Object)
 class UMaterialExpressionCosine : public UMaterialExpression
 {
-	GENERATED_UCLASS_BODY()
+	GENERATED_BODY()
+
+public:
 
 	UPROPERTY()
 	FExpressionInput Input;
 
 	UPROPERTY(EditAnywhere, Category=MaterialExpressionCosine, Meta = (ShowAsInputPin = "Advanced"))
-	float Period;
-
+	float Period = 1.0f;
 
 	//~ Begin UMaterialExpression Interface
 #if WITH_EDITOR
diff --git a/Engine/Source/Runtime/Engine/Public/Materials/MaterialExpressionDepthFade.h b/Engine/Source/Runtime/Engine/Public/Materials/MaterialExpressionDepthFade.h
--- a/Engine/Source/Runtime/Engine/Public/Materials/MaterialExpressionDepthFade.h
+++ b/Engine/Source/Runtime/Engine/Public/Materials/MaterialExpressionDepthFade.h
@@ -27,11 +27,11 @@ class UMaterialExpressionDepthFade : public UMaterialExpression
 
 	/** Opacity which will be scaled by the result of the fade.  This is used when InOpacity is unconnected. */
 	UPROPERTY(EditAnywhere, Category=MaterialExpressionDepthFade, meta=(OverridingInputProperty = "InOpacity"))
-	float OpacityDefault;
+	float OpacityDefault = 1.0f;
 
 	/** World space distance over which the fade should take place.  This is used when FadeDistance is unconnected. */
 	UPROPERTY(EditAnywhere, Category=MaterialExpressionDepthFade, meta=(OverridingInputProperty = "FadeDistance"))
-	float FadeDistanceDefault;
+	float FadeDistanceDefault = 100.0f;
 
 #if WITH_EDITOR
 
diff --git a/Engine/Source/Runtime/Engine/Public/Materials/MaterialExpressionDistanceFieldApproxAO.h b/Engine/Source/Runtime/Engine/Public/Materials/MaterialExpressionDistanceFieldApproxAO.h
--- a/Engine/Source/Runtime/Engine/Public/Materials/MaterialExpressionDistanceFieldApproxAO.h
+++ b/Engine/Source/Runtime/Engine/Public/Materials/MaterialExpressionDistanceFieldApproxAO.h
@@ -12,7 +12,9 @@
 UCLASS()
 class UMaterialExpressionDistanceFieldApproxAO : public UMaterialExpression
 {
-	GENERATED_UCLASS_BODY()
+	GENERATED_BODY()
+
+public:
 
 	UPROPERTY(meta = (RequiredInput = "false", ToolTip = "Defaults to current world position if not specified"))
 	FExpressionInput Position;
@@ -29,20 +31,20 @@ class UMaterialExpressionDistanceFieldApproxAO : public UMaterialExpression
 
 	/** only used if BaseDistance is not hooked up */
 	UPROPERTY(EditAnywhere, Category = MaterialExpressionDistanceFieldApproxAO, meta=(OverridingInputProperty = "BaseDistance"))
-	float BaseDistanceDefault;
+	float BaseDistanceDefault = 15;
 
 	UPROPERTY(meta = (RequiredInput = "false"))
 	FExpressionInput Radius;
 
 	/** only used if Radius is not hooked up */
 	UPROPERTY(EditAnywhere, Category = MaterialExpressionDistanceFieldApproxAO, meta=(OverridingInputProperty = "Radius"))
-	float RadiusDefault;
+	float RadiusDefault = 150;
 
 	UPROPERTY(EditAnywhere, Category = MaterialExpressionDistanceFieldApproxAO, meta = (UIMin = "1", UIMax = "4", ShowAsInputPin = "Primary", ToolTip = "Number of samples used to calculate occlusion"))
-	uint32 NumSteps;
+	uint32 NumSteps = 1;
 
 	UPROPERTY(EditAnywhere, Category = MaterialExpressionDistanceFieldApproxAO, DisplayName = "Step Scale", meta = (UIMin = "1.0", UIMax = "8.0", ShowAsInputPin = "Advanced", ToolTip = "Used to control step distance distribution"))
-	float StepScaleDefault;
+	float StepScaleDefault = 3.0f;
 
 	//~ Begin UMaterialExpression Interface
 #if WITH_EDITOR
diff --git a/Engine/Source/Runtime/Engine/Public/Materials/MaterialExpressionDivide.h b/Engine/Source/Runtime/Engine/Public/Materials/MaterialExpressionDivide.h
--- a/Engine/Source/Runtime/Engine/Public/Materials/MaterialExpressionDivide.h
+++ b/Engine/Source/Runtime/Engine/Public/Materials/MaterialExpressionDivide.h
@@ -12,7 +12,9 @@
 UCLASS(MinimalAPI)
 class UMaterialExpressionDivide : public UMaterialExpression
 {
-	GENERATED_UCLASS_BODY()
+	GENERATED_BODY()
+
+public:
 
 	UPROPERTY(meta = (RequiredInput = "false", ToolTip = "Defaults to 'ConstA' if not specified"))
 	FExpressionInput A;
@@ -22,11 +24,11 @@ class UMaterialExpressionDivide : public UMaterialExpression
 
 	/** only used if A is not hooked up */
 	UPROPERTY(EditAnywhere, Category=MaterialExpressionDivide, meta=(OverridingInputProperty = "A"))
-	float ConstA;
+	float ConstA = 1.0f;
 
 	/** only used if B is not hooked up */
 	UPROPERTY(EditAnywhere, Category=MaterialExpressionDivide, meta=(OverridingInputProperty = "B"))
-	float ConstB;
+	float ConstB = 2.0f;
 
 
 	//~ Begin UMaterialExpression Interface
diff --git a/Engine/Source/Runtime/Engine/Public/Materials/MaterialExpressionDynamicParameter.h b/Engine/Source/Runtime/Engine/Public/Materials/MaterialExpressionDynamicParameter.h
--- a/Engine/Source/Runtime/Engine/Public/Materials/MaterialExpressionDynamicParameter.h
+++ b/Engine/Source/Runtime/Engine/Public/Materials/MaterialExpressionDynamicParameter.h
@@ -30,11 +30,11 @@ class UMaterialExpressionDynamicParameter : public UMaterialExpression
 	TArray<FString> ParamNames;
 
 	UPROPERTY(EditAnywhere, Category = MaterialExpressionDynamicParameter, Meta = (ShowAsInputPin = "Primary"))
-	FLinearColor DefaultValue;
+	FLinearColor DefaultValue = FLinearColor::White;
 
 	/** The index of the dynamic parameter for use in tools that allow > 1 */
 	UPROPERTY(EditAnywhere, Category = MaterialExpressionDynamicParameter, meta = (UIMin = 0, ClampMin = 0, UIMax = 3, ClampMax = 3, ShowAsInputPin = "Primary"))
-	uint32 ParameterIndex;
+	uint32 ParameterIndex = 0;
 
 	//~ Begin UObject Interface
 #if WITH_EDITOR
diff --git a/Engine/Source/Runtime/Engine/Public/Materials/MaterialExpressionFirstPersonOutput.h b/Engine/Source/Runtime/Engine/Public/Materials/MaterialExpressionFirstPersonOutput.h
--- a/Engine/Source/Runtime/Engine/Public/Materials/MaterialExpressionFirstPersonOutput.h
+++ b/Engine/Source/Runtime/Engine/Public/Materials/MaterialExpressionFirstPersonOutput.h
@@ -19,7 +19,7 @@ class UMaterialExpressionFirstPersonOutput : public UMaterialExpressionCustomOut
 
 	/** Only used if FirstPersonInterpolationAlpha is not hooked up. Interpolates between world space and first person space. Valid range is [0, 1], from world space to first person space. */
 	UPROPERTY(EditAnywhere, Category = "First Person", meta = (OverridingInputProperty = "FirstPersonInterpolationAlpha", UIMin = 0.0f, UIMax = 1.0f, ClampMin = 0.0f, ClampMax = 1.0f))
-	float ConstFirstPersonInterpolationAlpha;
+	float ConstFirstPersonInterpolationAlpha = 1.0f;
 
 public:
 #if WITH_EDITOR
diff --git a/Engine/Source/Runtime/Engine/Public/Materials/MaterialExpressionFresnel.h b/Engine/Source/Runtime/Engine/Public/Materials/MaterialExpressionFresnel.h
--- a/Engine/Source/Runtime/Engine/Public/Materials/MaterialExpressionFresnel.h
+++ b/Engine/Source/Runtime/Engine/Public/Materials/MaterialExpressionFresnel.h
@@ -18,14 +18,16 @@
 UCLASS(MinimalAPI, collapsecategories, hidecategories=Object)
 class UMaterialExpressionFresnel : public UMaterialExpression
 {
-	GENERATED_UCLASS_BODY()
+	GENERATED_BODY()
+
+public:
 
 	UPROPERTY(meta = (RequiredInput = "false", ToolTip = "Defaults to 'Exponent' if not specified"))
 	FExpressionInput ExponentIn;
 
 	/** The exponent to pass into the pow() function */
 	UPROPERTY(EditAnywhere, Category=MaterialExpressionFresnel, meta=(OverridingInputProperty = "ExponentIn"))
-	float Exponent;
+	float Exponent = 5.0f;
 
 	UPROPERTY(meta = (RequiredInput = "false", ToolTip = "Defaults to 'BaseReflectFraction' if not specified"))
 	FExpressionInput BaseReflectFractionIn;
@@ -35,7 +37,7 @@ class UMaterialExpressionFresnel : public UMaterialExpression
 	 * A value of 1 effectively disables Fresnel.
 	 */
 	UPROPERTY(EditAnywhere, Category=MaterialExpressionFresnel, meta=(OverridingInputProperty = "BaseReflectFractionIn"))
-	float BaseReflectFraction;
+	float BaseReflectFraction = 0.04f;
 
 	/** The normal to dot with the camera FVector */
 	UPROPERTY(meta = (RequiredInput = "false", ToolTip = "Defaults to Pixel World Normal if not specified"))
diff --git a/Engine/Source/Runtime/Engine/Public/Materials/MaterialExpressionHairAttributes.h b/Engine/Source/Runtime/Engine/Public/Materials/MaterialExpressionHairAttributes.h
--- a/Engine/Source/Runtime/Engine/Public/Materials/MaterialExpressionHairAttributes.h
+++ b/Engine/Source/Runtime/Engine/Public/Materials/MaterialExpressionHairAttributes.h
@@ -16,7 +16,7 @@ class UMaterialExpressionHairAttributes : public UMaterialExpression
 
 	/** If enabled, this nodes outputs a tangent space tangent, otherwise it outputs a world space tangent. */
 	UPROPERTY(EditAnywhere, Category = MaterialExpressionHairAttributes)
-	uint32 bUseTangentSpace : 1;
+	uint32 bUseTangentSpace : 1 = true;
 
 	//~ Begin UMaterialExpression Interface
 #if WITH_EDITOR
diff --git a/Engine/Source/Runtime/Engine/Public/Materials/MaterialExpressionIf.h b/Engine/Source/Runtime/Engine/Public/Materials/MaterialExpressionIf.h
--- a/Engine/Source/Runtime/Engine/Public/Materials/MaterialExpressionIf.h
+++ b/Engine/Source/Runtime/Engine/Public/Materials/MaterialExpressionIf.h
@@ -13,7 +13,9 @@
 UCLASS(collapsecategories, hidecategories=Object, MinimalAPI)
 class UMaterialExpressionIf : public UMaterialExpression
 {
-	GENERATED_UCLASS_BODY()
+	GENERATED_BODY()
+
+public:
 
 	UPROPERTY()
 	FExpressionInput A;
@@ -31,11 +33,11 @@ class UMaterialExpressionIf : public UMaterialExpression
 	FExpressionInput ALessThanB;
 
 	UPROPERTY(EditAnywhere, Category=MaterialExpressionIf, meta = (ShowAsInputPin = "Advanced"))
-	float EqualsThreshold;
+	float EqualsThreshold = 0.00001f;
 
 	/** only used if B is not hooked up */
 	UPROPERTY(EditAnywhere, Category = MaterialExpressionIf, meta = (OverridingInputProperty = "B"))
-	float ConstB;
+	float ConstB = 0.0f;
 
 	UPROPERTY()
 	float ConstAEqualsB_DEPRECATED;
diff --git a/Engine/Source/Runtime/Engine/Public/Materials/MaterialExpressionInverseLinearInterpolate.h b/Engine/Source/Runtime/Engine/Public/Materials/MaterialExpressionInverseLinearInterpolate.h
--- a/Engine/Source/Runtime/Engine/Public/Materials/MaterialExpressionInverseLinearInterpolate.h
+++ b/Engine/Source/Runtime/Engine/Public/Materials/MaterialExpressionInverseLinearInterpolate.h
@@ -12,7 +12,9 @@
 UCLASS(MinimalAPI)
 class UMaterialExpressionInverseLinearInterpolate : public UMaterialExpression
 {
-	GENERATED_UCLASS_BODY()
+	GENERATED_BODY()
+
+public:
 
 	UPROPERTY(meta = (RequiredInput = "false", ToolTip = "Defaults to 'ConstA' if not specified"))
 	FExpressionInput A;
@@ -25,19 +27,19 @@ class UMaterialExpressionInverseLinearInterpolate : public UMaterialExpression
 
 	/** only used if A is not hooked up */
 	UPROPERTY(EditAnywhere, Category = MaterialExpressionInverseLinearInterpolate, meta = (OverridingInputProperty = "A"))
-	float ConstA;
+	float ConstA = 0.0f;
 
 	/** only used if B is not hooked up */
 	UPROPERTY(EditAnywhere, Category = MaterialExpressionInverseLinearInterpolate, meta = (OverridingInputProperty = "B"))
-	float ConstB;
+	float ConstB = 1.0f;
 
 	/** only used if Value is not hooked up */
 	UPROPERTY(EditAnywhere, Category = MaterialExpressionInverseLinearInterpolate, meta = (OverridingInputProperty = "Value"))
-	float ConstValue;
+	float ConstValue = 0.0f;
 
 	/** Clamp the result to 0 to 1 */
 	UPROPERTY(EditAnywhere, Category = MaterialExpressionInverseLinearInterpolate)
-	bool bClampResult;
+	bool bClampResult = false;
 
 	//~ Begin UMaterialExpression Interface
 #if WITH_EDITOR
diff --git a/Engine/Source/Runtime/Engine/Public/Materials/MaterialExpressionLinearInterpolate.h b/Engine/Source/Runtime/Engine/Public/Materials/MaterialExpressionLinearInterpolate.h
--- a/Engine/Source/Runtime/Engine/Public/Materials/MaterialExpressionLinearInterpolate.h
+++ b/Engine/Source/Runtime/Engine/Public/Materials/MaterialExpressionLinearInterpolate.h
@@ -12,8 +12,9 @@
 UCLASS(MinimalAPI)
 class UMaterialExpressionLinearInterpolate : public UMaterialExpression
 {
-	GENERATED_UCLASS_BODY()
+	GENERATED_BODY()
 
+public:
 	UPROPERTY(meta = (RequiredInput = "false", ToolTip = "Defaults to 'ConstA' if not specified"))
 	FExpressionInput A;
 
@@ -25,15 +26,15 @@ class UMaterialExpressionLinearInterpolate : public UMaterialExpression
 
 	/** only used if A is not hooked up */
 	UPROPERTY(EditAnywhere, Category=MaterialExpressionLinearInterpolate, meta=(OverridingInputProperty = "A"))
-	float ConstA;
+	float ConstA = 0.0f;
 
 	/** only used if B is not hooked up */
 	UPROPERTY(EditAnywhere, Category=MaterialExpressionLinearInterpolate, meta=(OverridingInputProperty = "B"))
-	float ConstB;
+	float ConstB = 1.0f;
 
 	/** only used if Alpha is not hooked up */
 	UPROPERTY(EditAnywhere, Category=MaterialExpressionLinearInterpolate, meta=(OverridingInputProperty = "Alpha"))
-	float ConstAlpha;
+	float ConstAlpha = 0.5f;
 
 
 	//~ Begin UMaterialExpression Interface
diff --git a/Engine/Source/Runtime/Engine/Public/Materials/MaterialExpressionLocalPosition.h b/Engine/Source/Runtime/Engine/Public/Materials/MaterialExpressionLocalPosition.h
--- a/Engine/Source/Runtime/Engine/Public/Materials/MaterialExpressionLocalPosition.h
+++ b/Engine/Source/Runtime/Engine/Public/Materials/MaterialExpressionLocalPosition.h
@@ -39,10 +39,10 @@ class UMaterialExpressionLocalPosition : public UMaterialExpression
 	GENERATED_UCLASS_BODY()
 
 	UPROPERTY(EditAnywhere, Category=MaterialExpressionLocalPosition, meta=(DisplayName = "Shader Offsets", ShowAsInputPin = "Advanced"))
-	EPositionIncludedOffsets IncludedOffsets;
+	EPositionIncludedOffsets IncludedOffsets = EPositionIncludedOffsets::IncludeOffsets;
 
 	UPROPERTY(EditAnywhere, Category=MaterialExpressionLocalPosition, meta=(DisplayName = "Local Origin", ShowAsInputPin = "Advanced"))
-	ELocalPositionOrigin LocalOrigin;
+	ELocalPositionOrigin LocalOrigin = ELocalPositionOrigin::Instance;
 	
 	//~ Begin UMaterialExpression Interface
 #if WITH_EDITOR
diff --git a/Engine/Source/Runtime/Engine/Public/Materials/MaterialExpressionMaterialAttributeLayers.h b/Engine/Source/Runtime/Engine/Public/Materials/MaterialExpressionMaterialAttributeLayers.h
--- a/Engine/Source/Runtime/Engine/Public/Materials/MaterialExpressionMaterialAttributeLayers.h
+++ b/Engine/Source/Runtime/Engine/Public/Materials/MaterialExpressionMaterialAttributeLayers.h
@@ -66,19 +66,19 @@ class UMaterialExpressionMaterialAttributeLayers : public UMaterialExpression
 	TArray<TObjectPtr<UMaterialExpressionMaterialFunctionCall>> LayerCallers;
 
 	UPROPERTY(Transient)
-	int32 NumActiveLayerCallers;
+	int32 NumActiveLayerCallers = 0;
 
 	UPROPERTY(Transient)
 	TArray<TObjectPtr<UMaterialExpressionMaterialFunctionCall>> BlendCallers;
 
 	UPROPERTY(Transient)
-	int32 NumActiveBlendCallers;
+	int32 NumActiveBlendCallers = 0;
 
 	UPROPERTY(Transient)
 	TObjectPtr<UMaterialExpressionMaterialFunctionCall> OutputCaller;
 
 	UPROPERTY(Transient)
-	bool bIsLayerGraphBuilt;
+	bool bIsLayerGraphBuilt = false;
 
 	//~ Begin UObject Interface
 	virtual void PostLoad() override;
@@ -121,5 +121,5 @@ class UMaterialExpressionMaterialAttributeLayers : public UMaterialExpression
 
 private:
 	/** Internal pointer to parameter-driven layer graph */
-	const FMaterialLayersFunctions* ParamLayers;
+	const FMaterialLayersFunctions* ParamLayers = nullptr;
 };
diff --git a/Engine/Source/Runtime/Engine/Public/Materials/MaterialExpressionMax.h b/Engine/Source/Runtime/Engine/Public/Materials/MaterialExpressionMax.h
--- a/Engine/Source/Runtime/Engine/Public/Materials/MaterialExpressionMax.h
+++ b/Engine/Source/Runtime/Engine/Public/Materials/MaterialExpressionMax.h
@@ -12,7 +12,9 @@
 UCLASS(MinimalAPI)
 class UMaterialExpressionMax : public UMaterialExpression
 {
-	GENERATED_UCLASS_BODY()
+	GENERATED_BODY()
+
+public:
 
 	UPROPERTY(meta = (RequiredInput = "false", ToolTip = "Defaults to 'ConstA' if not specified"))
 	FExpressionInput A;
@@ -22,11 +24,11 @@ class UMaterialExpressionMax : public UMaterialExpression
 
 	/** only used if A is not hooked up */
 	UPROPERTY(EditAnywhere, Category=MaterialExpressionMax, meta=(OverridingInputProperty = "A"))
-	float ConstA;
+	float ConstA = 0.0f;
 
 	/** only used if B is not hooked up */
 	UPROPERTY(EditAnywhere, Category=MaterialExpressionMax, meta=(OverridingInputProperty = "B"))
-	float ConstB;
+	float ConstB = 1.0f;
 
 	//~ Begin UMaterialExpression Interface
 #if WITH_EDITOR
diff --git a/Engine/Source/Runtime/Engine/Public/Materials/MaterialExpressionMin.h b/Engine/Source/Runtime/Engine/Public/Materials/MaterialExpressionMin.h
--- a/Engine/Source/Runtime/Engine/Public/Materials/MaterialExpressionMin.h
+++ b/Engine/Source/Runtime/Engine/Public/Materials/MaterialExpressionMin.h
@@ -12,7 +12,9 @@
 UCLASS(MinimalAPI)
 class UMaterialExpressionMin : public UMaterialExpression
 {
-	GENERATED_UCLASS_BODY()
+	GENERATED_BODY()
+
+public:
 
 	UPROPERTY(meta = (RequiredInput = "false", ToolTip = "Defaults to 'ConstA' if not specified"))
 	FExpressionInput A;
@@ -22,11 +24,11 @@ class UMaterialExpressionMin : public UMaterialExpression
 
 	/** only used if A is not hooked up */
 	UPROPERTY(EditAnywhere, Category=MaterialExpressionMin, meta=(OverridingInputProperty = "A"))
-	float ConstA;
+	float ConstA = 0.0f;
 
 	/** only used if B is not hooked up */
 	UPROPERTY(EditAnywhere, Category=MaterialExpressionMin, meta=(OverridingInputProperty = "B"))
-	float ConstB;
+	float ConstB = 1.0f;
 
 	//~ Begin UMaterialExpression Interface
 #if WITH_EDITOR
diff --git a/Engine/Source/Runtime/Engine/Public/Materials/MaterialExpressionMultiply.h b/Engine/Source/Runtime/Engine/Public/Materials/MaterialExpressionMultiply.h
--- a/Engine/Source/Runtime/Engine/Public/Materials/MaterialExpressionMultiply.h
+++ b/Engine/Source/Runtime/Engine/Public/Materials/MaterialExpressionMultiply.h
@@ -12,7 +12,9 @@
 UCLASS(MinimalAPI)
 class UMaterialExpressionMultiply : public UMaterialExpression
 {
-	GENERATED_UCLASS_BODY()
+	GENERATED_BODY()
+
+public:
 
 	UPROPERTY(meta = (RequiredInput = "false", ToolTip = "Defaults to 'ConstA' if not specified"))
 	FExpressionInput A;
@@ -22,11 +24,11 @@ class UMaterialExpressionMultiply : public UMaterialExpression
 
 	/** only used if A is not hooked up */
 	UPROPERTY(EditAnywhere, Category=MaterialExpressionMultiply, meta=(OverridingInputProperty = "A"))
-	float ConstA;
+	float ConstA = 0.0f;
 
 	/** only used if B is not hooked up */
 	UPROPERTY(EditAnywhere, Category=MaterialExpressionMultiply, meta=(OverridingInputProperty = "B"))
-	float ConstB;
+	float ConstB = 1.0f;
 
 
 	//~ Begin UMaterialExpression Interface
diff --git a/Engine/Source/Runtime/Engine/Public/Materials/MaterialExpressionPanner.h b/Engine/Source/Runtime/Engine/Public/Materials/MaterialExpressionPanner.h
--- a/Engine/Source/Runtime/Engine/Public/Materials/MaterialExpressionPanner.h
+++ b/Engine/Source/Runtime/Engine/Public/Materials/MaterialExpressionPanner.h
@@ -31,7 +31,7 @@ class UMaterialExpressionPanner : public UMaterialExpression
 
 	/** only used if Coordinate is not hooked up */
 	UPROPERTY(EditAnywhere, Category = MaterialExpressionPanner, meta = (OverridingInputProperty = "Coordinate"))
-	uint32 ConstCoordinate;
+	uint32 ConstCoordinate = 0;
 
 	// Output only the fractional part of the pan calculation for greater precision.
 	// Output is greater than or equal to 0 and less than 1.
diff --git a/Engine/Source/Runtime/Engine/Public/Materials/MaterialExpressionParticleSubUV.h b/Engine/Source/Runtime/Engine/Public/Materials/MaterialExpressionParticleSubUV.h
--- a/Engine/Source/Runtime/Engine/Public/Materials/MaterialExpressionParticleSubUV.h
+++ b/Engine/Source/Runtime/Engine/Public/Materials/MaterialExpressionParticleSubUV.h
@@ -11,11 +11,12 @@
 UCLASS(MinimalAPI)
 class UMaterialExpressionParticleSubUV : public UMaterialExpressionTextureSample
 {
-	GENERATED_UCLASS_BODY()
+	GENERATED_BODY()
 
-	UPROPERTY(EditAnywhere, Category=MaterialExpressionParticleSubUV)
-	uint32 bBlend:1;
+public:
 
+	UPROPERTY(EditAnywhere, Category=MaterialExpressionParticleSubUV)
+	uint32 bBlend:1 = true;
 
 	//~ Begin UMaterialExpression Interface
 #if WITH_EDITOR
diff --git a/Engine/Source/Runtime/Engine/Public/Materials/MaterialExpressionPower.h b/Engine/Source/Runtime/Engine/Public/Materials/MaterialExpressionPower.h
--- a/Engine/Source/Runtime/Engine/Public/Materials/MaterialExpressionPower.h
+++ b/Engine/Source/Runtime/Engine/Public/Materials/MaterialExpressionPower.h
@@ -12,7 +12,9 @@
 UCLASS(collapsecategories, hidecategories=Object, MinimalAPI)
 class UMaterialExpressionPower : public UMaterialExpression
 {
-	GENERATED_UCLASS_BODY()
+	GENERATED_BODY()
+
+public:
 
 	UPROPERTY()
 	FExpressionInput Base;
@@ -22,7 +24,7 @@ class UMaterialExpressionPower : public UMaterialExpression
 
 	/** only used if Exponent is not hooked up */
 	UPROPERTY(EditAnywhere, Category=MaterialExpressionPower, meta=(OverridingInputProperty = "Exponent"))
-	float ConstExponent;
+	float ConstExponent = 2;
 
 	//~ Begin UMaterialExpression Interface
 #if WITH_EDITOR
diff --git a/Engine/Source/Runtime/Engine/Public/Materials/MaterialExpressionRequiredSamplersSwitch.h b/Engine/Source/Runtime/Engine/Public/Materials/MaterialExpressionRequiredSamplersSwitch.h
--- a/Engine/Source/Runtime/Engine/Public/Materials/MaterialExpressionRequiredSamplersSwitch.h
+++ b/Engine/Source/Runtime/Engine/Public/Materials/MaterialExpressionRequiredSamplersSwitch.h
@@ -14,15 +14,18 @@
 UCLASS(MinimalAPI)
 class UMaterialExpressionRequiredSamplersSwitch : public UMaterialExpression
 {
-	GENERATED_UCLASS_BODY()
-		UPROPERTY()
-		FExpressionInput InputTrue;
+	GENERATED_BODY()
+
+public:
+
+	UPROPERTY()
+	FExpressionInput InputTrue;
 		
-		UPROPERTY()
-		FExpressionInput InputFalse;
+	UPROPERTY()
+	FExpressionInput InputFalse;
 
-		UPROPERTY(EditAnywhere, Category=MaterialExpressionRequiredSamplersSwitch, meta=(UIMin = "16", UIMax = "255"))
-		uint32 RequiredSamplers;
+	UPROPERTY(EditAnywhere, Category=MaterialExpressionRequiredSamplersSwitch, meta=(UIMin = "16", UIMax = "255"))
+	uint32 RequiredSamplers = 16;
 
 	//~ Begin UMaterialExpression Interface
 #if WITH_EDITOR
diff --git a/Engine/Source/Runtime/Engine/Public/Materials/MaterialExpressionRotateAboutAxis.h b/Engine/Source/Runtime/Engine/Public/Materials/MaterialExpressionRotateAboutAxis.h
--- a/Engine/Source/Runtime/Engine/Public/Materials/MaterialExpressionRotateAboutAxis.h
+++ b/Engine/Source/Runtime/Engine/Public/Materials/MaterialExpressionRotateAboutAxis.h
@@ -12,7 +12,9 @@
 UCLASS(MinimalAPI)
 class UMaterialExpressionRotateAboutAxis : public UMaterialExpression
 {
-	GENERATED_UCLASS_BODY()
+	GENERATED_BODY()
+
+public:
 
 	UPROPERTY()
 	FExpressionInput NormalizedRotationAxis;
@@ -27,7 +29,7 @@ class UMaterialExpressionRotateAboutAxis : public UMaterialExpression
 	FExpressionInput Position;
 
 	UPROPERTY(EditAnywhere, Category=MaterialExpressionRotateAboutAxis, meta = (ShowAsInputPin = "Advanced"))
-	float Period;
+	float Period = 1.0f;
 
 
 	//~ Begin UMaterialExpression Interface
diff --git a/Engine/Source/Runtime/Engine/Public/Materials/MaterialExpressionRotator.h b/Engine/Source/Runtime/Engine/Public/Materials/MaterialExpressionRotator.h
--- a/Engine/Source/Runtime/Engine/Public/Materials/MaterialExpressionRotator.h
+++ b/Engine/Source/Runtime/Engine/Public/Materials/MaterialExpressionRotator.h
@@ -21,17 +21,17 @@ class UMaterialExpressionRotator : public UMaterialExpression
 	FExpressionInput Time;
 
 	UPROPERTY(EditAnywhere, Category=MaterialExpressionRotator, meta = (ShowAsInputPin = "Advanced"))
-	float CenterX;
+	float CenterX = 0.5f;
 
 	UPROPERTY(EditAnywhere, Category=MaterialExpressionRotator, meta = (ShowAsInputPin = "Advanced"))
-	float CenterY;
+	float CenterY = 0.5f;
 
 	UPROPERTY(EditAnywhere, Category=MaterialExpressionRotator, meta = (ShowAsInputPin = "Advanced"))
-	float Speed;
+	float Speed = 0.25f;
 
 	/** only used if Coordinate is not hooked up */
 	UPROPERTY(EditAnywhere, Category = MaterialExpressionRotator, meta = (OverridingInputProperty = "Coordinate"))
-	uint32 ConstCoordinate;
+	uint32 ConstCoordinate = 0;
 
 	//~ Begin UMaterialExpression Interface
 #if WITH_EDITOR
diff --git a/Engine/Source/Runtime/Engine/Public/Materials/MaterialExpressionSamplePhysicsField.h b/Engine/Source/Runtime/Engine/Public/Materials/MaterialExpressionSamplePhysicsField.h
--- a/Engine/Source/Runtime/Engine/Public/Materials/MaterialExpressionSamplePhysicsField.h
+++ b/Engine/Source/Runtime/Engine/Public/Materials/MaterialExpressionSamplePhysicsField.h
@@ -45,7 +45,9 @@ public:
 UCLASS()
 class UMaterialExpressionSamplePhysicsScalarField : public UMaterialExpression
 {
-	GENERATED_UCLASS_BODY()
+	GENERATED_BODY()
+
+public:
 
 	UPROPERTY(meta = (RequiredInput = "false", ToolTip = "Defaults to current world position if not specified"))
 	FExpressionInput WorldPosition;
@@ -56,7 +58,7 @@ class UMaterialExpressionSamplePhysicsScalarField : public UMaterialExpression
 
 	/** Target Type to be accessed */
 	UPROPERTY(EditAnywhere, Category = "Field Setup", meta = (DisplayName = "Target Type", ShowAsInputPin = "Advanced"))
-	TEnumAsByte<EFieldScalarType> FieldTarget;
+	TEnumAsByte<EFieldScalarType> FieldTarget = EFieldScalarType::Scalar_DynamicConstraint;
 
 	//~ Begin UMaterialExpression Interface
 #if WITH_EDITOR
@@ -72,7 +74,9 @@ class UMaterialExpressionSamplePhysicsScalarField : public UMaterialExpression
 UCLASS()
 class UMaterialExpressionSamplePhysicsIntegerField : public UMaterialExpression
 {
-	GENERATED_UCLASS_BODY()
+	GENERATED_BODY()
+
+public:
 
 	UPROPERTY(meta = (RequiredInput = "false", ToolTip = "Defaults to current world position if not specified"))
 	FExpressionInput WorldPosition;
@@ -83,7 +87,7 @@ class UMaterialExpressionSamplePhysicsIntegerField : public UMaterialExpression
 
 	/** Target Type to be accessed */
 	UPROPERTY(EditAnywhere, Category = "Field Setup", meta = (DisplayName = "Target Type", ShowAsInputPin = "Advanced"))
-	TEnumAsByte<EFieldIntegerType> FieldTarget;
+	TEnumAsByte<EFieldIntegerType> FieldTarget = EFieldIntegerType::Integer_DynamicState;
 
 	//~ Begin UMaterialExpression Interface
 #if WITH_EDITOR
diff --git a/Engine/Source/Runtime/Engine/Public/Materials/MaterialExpressionSceneColor.h b/Engine/Source/Runtime/Engine/Public/Materials/MaterialExpressionSceneColor.h
--- a/Engine/Source/Runtime/Engine/Public/Materials/MaterialExpressionSceneColor.h
+++ b/Engine/Source/Runtime/Engine/Public/Materials/MaterialExpressionSceneColor.h
@@ -44,7 +44,7 @@ class UMaterialExpressionSceneColor : public UMaterialExpression
 
 	/** only used if Input is not hooked up */
 	UPROPERTY(EditAnywhere, Category = MaterialExpressionSceneColor, meta = (OverridingInputProperty = "Input"))
-	FVector2D ConstInput;
+	FVector2D ConstInput = FVector2D(0.f, 0.f);
 
 	//~ Begin UObject Interface.
 	virtual void PostLoad() override;
diff --git a/Engine/Source/Runtime/Engine/Public/Materials/MaterialExpressionSceneDepth.h b/Engine/Source/Runtime/Engine/Public/Materials/MaterialExpressionSceneDepth.h
--- a/Engine/Source/Runtime/Engine/Public/Materials/MaterialExpressionSceneDepth.h
+++ b/Engine/Source/Runtime/Engine/Public/Materials/MaterialExpressionSceneDepth.h
@@ -35,7 +35,7 @@ class UMaterialExpressionSceneDepth : public UMaterialExpression
 
 	/** only used if Input is not hooked up */
 	UPROPERTY(EditAnywhere, Category = MaterialExpressionSceneDepth, meta = (OverridingInputProperty = "Input"))
-	FVector2D ConstInput;
+	FVector2D ConstInput = FVector2D(0.f, 0.f);
 
 	//~ Begin UObject Interface.
 	virtual void PostLoad() override;
diff --git a/Engine/Source/Runtime/Engine/Public/Materials/MaterialExpressionSceneTexture.h b/Engine/Source/Runtime/Engine/Public/Materials/MaterialExpressionSceneTexture.h
--- a/Engine/Source/Runtime/Engine/Public/Materials/MaterialExpressionSceneTexture.h
+++ b/Engine/Source/Runtime/Engine/Public/Materials/MaterialExpressionSceneTexture.h
@@ -25,7 +25,7 @@ class UMaterialExpressionSceneTexture : public UMaterialExpression
 	
 	/** Whether to use point sampled texture lookup (default) or using [bi-linear] filtered (can be slower, avoid faceted lock with distortions), some SceneTextures cannot be filtered */
 	UPROPERTY(EditAnywhere, Category=UMaterialExpressionSceneTexture, meta=(DisplayName = "Filtered", ShowAsInputPin = "Advanced"))
-	bool bFiltered;
+	bool bFiltered = false;
 
 	//~ Begin UMaterialExpression Interface
 #if WITH_EDITOR
diff --git a/Engine/Source/Runtime/Engine/Public/Materials/MaterialExpressionSine.h b/Engine/Source/Runtime/Engine/Public/Materials/MaterialExpressionSine.h
--- a/Engine/Source/Runtime/Engine/Public/Materials/MaterialExpressionSine.h
+++ b/Engine/Source/Runtime/Engine/Public/Materials/MaterialExpressionSine.h
@@ -12,14 +12,15 @@
 UCLASS(MinimalAPI, collapsecategories, hidecategories=Object)
 class UMaterialExpressionSine : public UMaterialExpression
 {
-	GENERATED_UCLASS_BODY()
+	GENERATED_BODY()
+
+public:
 
 	UPROPERTY()
 	FExpressionInput Input;
 
 	UPROPERTY(EditAnywhere, Category=MaterialExpressionSine, Meta = (ShowAsInputPin = "Advanced"))
-	float Period;
-
+	float Period = 1.0f;
 
 	//~ Begin UMaterialExpression Interface
 #if WITH_EDITOR
diff --git a/Engine/Source/Runtime/Engine/Public/Materials/MaterialExpressionSmoothStep.h b/Engine/Source/Runtime/Engine/Public/Materials/MaterialExpressionSmoothStep.h
--- a/Engine/Source/Runtime/Engine/Public/Materials/MaterialExpressionSmoothStep.h
+++ b/Engine/Source/Runtime/Engine/Public/Materials/MaterialExpressionSmoothStep.h
@@ -12,7 +12,9 @@
 UCLASS(MinimalAPI)
 class UMaterialExpressionSmoothStep : public UMaterialExpression
 {
-	GENERATED_UCLASS_BODY()
+	GENERATED_BODY()
+
+public:
 
 	UPROPERTY(meta = (RequiredInput = "false", ToolTip = "Defaults to 'ConstMin' if not specified"))
 	FExpressionInput Min;
@@ -25,15 +27,15 @@ class UMaterialExpressionSmoothStep : public UMaterialExpression
 
 	/** only used if Min is not hooked up */
 	UPROPERTY(EditAnywhere, Category = MaterialExpressionSmoothStep, meta = (OverridingInputProperty = "Min"))
-	float ConstMin;
+	float ConstMin = 0.0f;
 
 	/** only used if Max is not hooked up */
 	UPROPERTY(EditAnywhere, Category = MaterialExpressionSmoothStep, meta = (OverridingInputProperty = "Max"))
-	float ConstMax;
+	float ConstMax = 1.0f;
 
 	/** only used if Value is not hooked up */
 	UPROPERTY(EditAnywhere, Category = MaterialExpressionSmoothStep, meta = (OverridingInputProperty = "Value"))
-	float ConstValue;
+	float ConstValue = 0.0f;
 
 
 	//~ Begin UMaterialExpression Interface
diff --git a/Engine/Source/Runtime/Engine/Public/Materials/MaterialExpressionSobol.h b/Engine/Source/Runtime/Engine/Public/Materials/MaterialExpressionSobol.h
--- a/Engine/Source/Runtime/Engine/Public/Materials/MaterialExpressionSobol.h
+++ b/Engine/Source/Runtime/Engine/Public/Materials/MaterialExpressionSobol.h
@@ -10,7 +10,9 @@
 UCLASS(MinimalAPI)
 class UMaterialExpressionSobol : public UMaterialExpression
 {
-	GENERATED_UCLASS_BODY()
+	GENERATED_BODY()
+
+public:
 
 	/** 2D integer cell in 256x256 grid.
 	  * Uses cell (0,0) if not connected */
@@ -30,12 +32,12 @@ class UMaterialExpressionSobol : public UMaterialExpression
 	UPROPERTY(EditAnywhere, Category = MaterialExpressionSobol, meta = (
 		OverridingInputProperty = "Index", 
 		Tooltip="Sobol point number. Only used if Index is not connected."))
-	uint32 ConstIndex;
+	uint32 ConstIndex = 0;
 	
 	UPROPERTY(EditAnywhere, Category = MaterialExpressionSobol, meta = (
 		OverridingInputProperty = "Seed",
 		Tooltip="2D Seed for sequence randomization. Only used if Seed is not connected."))
-	FVector2D ConstSeed;
+	FVector2D ConstSeed = FVector2D(0.f, 0.f);
 
 	//~ Begin UMaterialExpression Interface
 #if WITH_EDITOR
diff --git a/Engine/Source/Runtime/Engine/Public/Materials/MaterialExpressionSpeedTree.h b/Engine/Source/Runtime/Engine/Public/Materials/MaterialExpressionSpeedTree.h
--- a/Engine/Source/Runtime/Engine/Public/Materials/MaterialExpressionSpeedTree.h
+++ b/Engine/Source/Runtime/Engine/Public/Materials/MaterialExpressionSpeedTree.h
@@ -40,7 +40,9 @@ enum ESpeedTreeLODType : int
 UCLASS(collapsecategories, hidecategories=Object, MinimalAPI)
 class UMaterialExpressionSpeedTree : public UMaterialExpression
 {
-	GENERATED_UCLASS_BODY()	
+	GENERATED_BODY()	
+
+public:
 
 	UPROPERTY(meta = (RequiredInput = "false", ToolTip = "Defaults to 'WindType' if not specified"))
 	FExpressionInput GeometryInput;
@@ -55,19 +57,19 @@ class UMaterialExpressionSpeedTree : public UMaterialExpression
 	FExpressionInput ExtraBendWS;
 
 	UPROPERTY(EditAnywhere, Category=MaterialExpressionSpeedTree, meta=(DisplayName = "Geometry Type", ToolTip="The type of SpeedTree geometry on which this material will be used"))
-	TEnumAsByte<enum ESpeedTreeGeometryType> GeometryType;
+	TEnumAsByte<enum ESpeedTreeGeometryType> GeometryType = STG_Branch;
 
 	UPROPERTY(EditAnywhere, Category=MaterialExpressionSpeedTree, meta=(DisplayName = "Wind Type", ToolTip="The type of wind effect used on this tree. This can only go as high as it was in the SpeedTree Modeler, but you can set it to a lower option for lower quality wind and faster rendering."))
-	TEnumAsByte<enum ESpeedTreeWindType> WindType;
+	TEnumAsByte<enum ESpeedTreeWindType> WindType = STW_None;
 
 	UPROPERTY(EditAnywhere, Category=MaterialExpressionSpeedTree, meta=(DisplayName = "LOD Type", ToolTip="The type of LOD to use"))
-	TEnumAsByte<enum ESpeedTreeLODType> LODType;
+	TEnumAsByte<enum ESpeedTreeLODType> LODType = STLOD_Pop;
 
 	UPROPERTY(EditAnywhere, Category=MaterialExpressionSpeedTree, meta=(DisplayName = "Billboard Threshold", ToolTip="The threshold for triangles to be removed from the bilboard mesh when not facing the camera (0 = none pass, 1 = all pass).", ClampMin = "0.0", ClampMax = "1.0", UIMin = "0.0", UIMax = "1.0"))
-	float BillboardThreshold;
+	float BillboardThreshold = 0.9f;
 
 	UPROPERTY(EditAnywhere, Category=MaterialExpressionSpeedTree, AdvancedDisplay, meta=(DisplayName = "Accurate Wind Velocities", ToolTip="Support accurate velocities from wind. This will incur extra cost per vertex."))
-	bool bAccurateWindVelocities;
+	bool bAccurateWindVelocities = false;
 
 	//~ Begin UObject Interface
 	virtual void Serialize(FStructuredArchive::FRecord Record) override;
diff --git a/Engine/Source/Runtime/Engine/Public/Materials/MaterialExpressionSphereMask.h b/Engine/Source/Runtime/Engine/Public/Materials/MaterialExpressionSphereMask.h
--- a/Engine/Source/Runtime/Engine/Public/Materials/MaterialExpressionSphereMask.h
+++ b/Engine/Source/Runtime/Engine/Public/Materials/MaterialExpressionSphereMask.h
@@ -12,7 +12,9 @@
 UCLASS(MinimalAPI)
 class UMaterialExpressionSphereMask : public UMaterialExpression
 {
-	GENERATED_UCLASS_BODY()
+	GENERATED_BODY()
+
+public:
 
 	/** 1 to 4 dimensional vector, should be the same type as B */
 	UPROPERTY()
@@ -32,11 +34,11 @@ class UMaterialExpressionSphereMask : public UMaterialExpression
 
 	/** in the unit that A and B are measured */
 	UPROPERTY(EditAnywhere, Category=MaterialExpressionSphereMask, meta=(OverridingInputProperty = "Radius"), DisplayName = "Radius")
-	float AttenuationRadius;
+	float AttenuationRadius = 256.0f;
 
 	/** in percent 0%=soft .. 100%=hard */
 	UPROPERTY(EditAnywhere, Category=MaterialExpressionSphereMask, meta=(UIMin = "0.0", UIMax = "100.0", ClampMin = "0.0", ClampMax = "100.0", OverridingInputProperty = "Hardness", DisplayName = "Hardness"))
-	float HardnessPercent;
+	float HardnessPercent = 100.0f;
 
 
 	//~ Begin UMaterialExpression Interface
diff --git a/Engine/Source/Runtime/Engine/Public/Materials/MaterialExpressionSphericalParticleOpacity.h b/Engine/Source/Runtime/Engine/Public/Materials/MaterialExpressionSphericalParticleOpacity.h
--- a/Engine/Source/Runtime/Engine/Public/Materials/MaterialExpressionSphericalParticleOpacity.h
+++ b/Engine/Source/Runtime/Engine/Public/Materials/MaterialExpressionSphericalParticleOpacity.h
@@ -20,7 +20,7 @@ class UMaterialExpressionSphericalParticleOpacity : public UMaterialExpression
 
 	/** Constant density of the particle sphere.  Will be overridden if Density is connected. */
 	UPROPERTY(EditAnywhere, Category=MaterialExpressionSphericalParticleOpacity, meta=(OverridingInputProperty = "Density"))
-	float ConstantDensity;
+	float ConstantDensity = 1;
 
 
 	//~ Begin UMaterialExpression Interface
diff --git a/Engine/Source/Runtime/Engine/Public/Materials/MaterialExpressionStep.h b/Engine/Source/Runtime/Engine/Public/Materials/MaterialExpressionStep.h
--- a/Engine/Source/Runtime/Engine/Public/Materials/MaterialExpressionStep.h
+++ b/Engine/Source/Runtime/Engine/Public/Materials/MaterialExpressionStep.h
@@ -13,7 +13,9 @@
 UCLASS(MinimalAPI)
 class UMaterialExpressionStep : public UMaterialExpression
 {
-	GENERATED_UCLASS_BODY()
+	GENERATED_BODY()
+
+public:
 
 	UPROPERTY(meta = (RequiredInput = "false", ToolTip = "Defaults to 'ConstY' if not specified"))
 	FExpressionInput Y;
@@ -23,11 +25,11 @@ class UMaterialExpressionStep : public UMaterialExpression
 
 	/** only used if Y is not hooked up */
 	UPROPERTY(EditAnywhere, Category = MaterialExpressionStep, meta = (OverridingInputProperty = "Y"))
-	float ConstY;
+	float ConstY = 0.0f;
 
 	/** only used if X is not hooked up */
 	UPROPERTY(EditAnywhere, Category = MaterialExpressionStep, meta = (OverridingInputProperty = "X"))
-	float ConstX;
+	float ConstX = 1.0f;
 
 
 	//~ Begin UMaterialExpression Interface
diff --git a/Engine/Source/Runtime/Engine/Public/Materials/MaterialExpressionSubtract.h b/Engine/Source/Runtime/Engine/Public/Materials/MaterialExpressionSubtract.h
--- a/Engine/Source/Runtime/Engine/Public/Materials/MaterialExpressionSubtract.h
+++ b/Engine/Source/Runtime/Engine/Public/Materials/MaterialExpressionSubtract.h
@@ -12,7 +12,9 @@
 UCLASS(MinimalAPI)
 class UMaterialExpressionSubtract : public UMaterialExpression
 {
-	GENERATED_UCLASS_BODY()
+	GENERATED_BODY()
+
+public:
 
 	UPROPERTY(meta = (RequiredInput = "false", ToolTip = "Defaults to 'ConstA' if not specified"))
 	FExpressionInput A;
@@ -22,11 +24,11 @@ class UMaterialExpressionSubtract : public UMaterialExpression
 
 	/** only used if A is not hooked up */
 	UPROPERTY(EditAnywhere, Category=MaterialExpressionSubtract, meta=(OverridingInputProperty = "A"))
-	float ConstA;
+	float ConstA = 1.0f;
 
 	/** only used if B is not hooked up */
 	UPROPERTY(EditAnywhere, Category=MaterialExpressionSubtract, meta=(OverridingInputProperty = "B"))
-	float ConstB;
+	float ConstB = 1.0f;
 
 
 	//~ Begin UMaterialExpression Interface
diff --git a/Engine/Source/Runtime/Engine/Public/Materials/MaterialExpressionSwitch.h b/Engine/Source/Runtime/Engine/Public/Materials/MaterialExpressionSwitch.h
--- a/Engine/Source/Runtime/Engine/Public/Materials/MaterialExpressionSwitch.h
+++ b/Engine/Source/Runtime/Engine/Public/Materials/MaterialExpressionSwitch.h
@@ -40,7 +40,7 @@ class UMaterialExpressionSwitch : public UMaterialExpression
 	GENERATED_UCLASS_BODY()
 
 	UPROPERTY(EditAnywhere, Category= UMaterialExpressionSwitch)
-	FString Description;
+	FString Description = TEXT("Switch");
 
 	UPROPERTY(meta = (RequiredInput = "false", ToolTip = "Defaults to 'ConstSwitchValue' if not specified"))
 	FExpressionInput SwitchValue;
diff --git a/Engine/Source/Runtime/Engine/Public/Materials/MaterialExpressionTangent.h b/Engine/Source/Runtime/Engine/Public/Materials/MaterialExpressionTangent.h
--- a/Engine/Source/Runtime/Engine/Public/Materials/MaterialExpressionTangent.h
+++ b/Engine/Source/Runtime/Engine/Public/Materials/MaterialExpressionTangent.h
@@ -11,13 +11,15 @@
 UCLASS(MinimalAPI, collapsecategories, hidecategories=Object)
 class UMaterialExpressionTangent : public UMaterialExpression
 {
-	GENERATED_UCLASS_BODY()
+	GENERATED_BODY()
+
+public:
 
 	UPROPERTY()
 	FExpressionInput Input;
 
 	UPROPERTY(EditAnywhere, Category=MaterialExpressionTangent, Meta = (ShowAsInputPin = "Advanced"))
-	float Period;
+	float Period = 1.0f;
 
 	//~ Begin UMaterialExpression Interface
 #if WITH_EDITOR
diff --git a/Engine/Source/Runtime/Engine/Public/Materials/MaterialExpressionTemporalSobol.h b/Engine/Source/Runtime/Engine/Public/Materials/MaterialExpressionTemporalSobol.h
--- a/Engine/Source/Runtime/Engine/Public/Materials/MaterialExpressionTemporalSobol.h
+++ b/Engine/Source/Runtime/Engine/Public/Materials/MaterialExpressionTemporalSobol.h
@@ -10,7 +10,9 @@
 UCLASS(MinimalAPI)
 class UMaterialExpressionTemporalSobol : public UMaterialExpression
 {
-	GENERATED_UCLASS_BODY()
+	GENERATED_BODY()
+
+public:
 
 	UPROPERTY(meta = (
 		RequiredInput = "false", 
@@ -25,12 +27,12 @@ class UMaterialExpressionTemporalSobol : public UMaterialExpression
 	UPROPERTY(EditAnywhere, Category = MaterialExpressionSobol, meta = (
 		OverridingInputProperty = "Index", 
 		Tooltip="Sobol point number. Only used if Index is not connected."))
-	uint32 ConstIndex;
+	uint32 ConstIndex = 0;
 	
 	UPROPERTY(EditAnywhere, Category = MaterialExpressionSobol, meta = (
 		OverridingInputProperty = "Seed",
 		Tooltip="2D Seed for sequence randomization. Only used if Seed is not connected."))
-	FVector2D ConstSeed;
+	FVector2D ConstSeed = FVector2D(0.f, 0.f);
 
 	//~ Begin UMaterialExpression Interface
 #if WITH_EDITOR
diff --git a/Engine/Source/Runtime/Engine/Public/Materials/MaterialExpressionTextureBase.h b/Engine/Source/Runtime/Engine/Public/Materials/MaterialExpressionTextureBase.h
--- a/Engine/Source/Runtime/Engine/Public/Materials/MaterialExpressionTextureBase.h
+++ b/Engine/Source/Runtime/Engine/Public/Materials/MaterialExpressionTextureBase.h
@@ -20,7 +20,9 @@ enum EShaderPlatform : uint16;
 UCLASS(abstract, hidecategories=Object, MinimalAPI)
 class UMaterialExpressionTextureBase : public UMaterialExpression 
 {
-	GENERATED_UCLASS_BODY()
+	GENERATED_BODY()
+
+public:
 
 	UPROPERTY(EditAnywhere, BlueprintReadWrite, Category=MaterialExpressionTextureBase)
 	TObjectPtr<UTexture> Texture;
@@ -30,7 +32,7 @@ class UMaterialExpressionTextureBase : public UMaterialExpression
 	
 	/** Is default selected texture when using mesh paint mode texture painting */
 	UPROPERTY(EditAnywhere, Category=MaterialExpressionTextureBase)
-	uint8 IsDefaultMeshpaintTexture:1;
+	uint8 IsDefaultMeshpaintTexture:1 = false;
 	
 	//~ Begin UObject Interface
 #if WITH_EDITOR
diff --git a/Engine/Source/Runtime/Engine/Public/Materials/MaterialExpressionTextureCoordinate.h b/Engine/Source/Runtime/Engine/Public/Materials/MaterialExpressionTextureCoordinate.h
--- a/Engine/Source/Runtime/Engine/Public/Materials/MaterialExpressionTextureCoordinate.h
+++ b/Engine/Source/Runtime/Engine/Public/Materials/MaterialExpressionTextureCoordinate.h
@@ -19,11 +19,11 @@ class UMaterialExpressionTextureCoordinate : public UMaterialExpression
 
 	/** Controls how much the texture tiles horizontally, by scaling the U component of the vertex UVs by the specified amount. */
 	UPROPERTY(EditAnywhere, BlueprintReadWrite, Category=MaterialExpressionTextureCoordinate, meta = (ShowAsInputPin = "Advanced"))
-	float UTiling;
+	float UTiling = 1.0f;
 
 	/** Controls how much the texture tiles vertically, by scaling the V component of the vertex UVs by the specified amount. */
 	UPROPERTY(EditAnywhere, BlueprintReadWrite, Category=MaterialExpressionTextureCoordinate, meta = (ShowAsInputPin = "Advanced"))
-	float VTiling;
+	float VTiling = 1.0f;
 
 	/** Would like to unmirror U or V 
 	 *  - if the texture is mirrored and if you would like to undo mirroring for this texture sample, use this to unmirror */
diff --git a/Engine/Source/Runtime/Engine/Public/Materials/MaterialExpressionTextureProperty.h b/Engine/Source/Runtime/Engine/Public/Materials/MaterialExpressionTextureProperty.h
--- a/Engine/Source/Runtime/Engine/Public/Materials/MaterialExpressionTextureProperty.h
+++ b/Engine/Source/Runtime/Engine/Public/Materials/MaterialExpressionTextureProperty.h
@@ -34,7 +34,7 @@ class UMaterialExpressionTextureProperty : public UMaterialExpression
 	
 	/** Texture property to be accessed */
 	UPROPERTY(EditAnywhere, Category=UMaterialExpressionTextureProperty, meta=(DisplayName = "Texture Property", ShowAsInputPin = "Advanced"))
-	TEnumAsByte<EMaterialExposedTextureProperty> Property;
+	TEnumAsByte<EMaterialExposedTextureProperty> Property = TMTM_TextureSize;
 	
 	//~ Begin UMaterialExpression Interface
 #if WITH_EDITOR
diff --git a/Engine/Source/Runtime/Engine/Public/Materials/MaterialExpressionTextureSampleParameterSubUV.h b/Engine/Source/Runtime/Engine/Public/Materials/MaterialExpressionTextureSampleParameterSubUV.h
--- a/Engine/Source/Runtime/Engine/Public/Materials/MaterialExpressionTextureSampleParameterSubUV.h
+++ b/Engine/Source/Runtime/Engine/Public/Materials/MaterialExpressionTextureSampleParameterSubUV.h
@@ -13,10 +13,12 @@ class UTexture;
 UCLASS(collapsecategories, hidecategories=Object, MinimalAPI)
 class UMaterialExpressionTextureSampleParameterSubUV : public UMaterialExpressionTextureSampleParameter2D
 {
-	GENERATED_UCLASS_BODY()
+	GENERATED_BODY()
+
+public:
 
 	UPROPERTY(EditAnywhere, Category=MaterialExpressionTextureSampleParameterSubUV)
-	uint32 bBlend:1;
+	uint32 bBlend:1 = true;
 
 #if WITH_EDITOR
 	//~ Begin UMaterialExpressionTextureSampleParameter Interface
diff --git a/Engine/Source/Runtime/Engine/Public/Materials/MaterialExpressionTime.h b/Engine/Source/Runtime/Engine/Public/Materials/MaterialExpressionTime.h
--- a/Engine/Source/Runtime/Engine/Public/Materials/MaterialExpressionTime.h
+++ b/Engine/Source/Runtime/Engine/Public/Materials/MaterialExpressionTime.h
@@ -11,7 +11,9 @@
 UCLASS(MinimalAPI, collapsecategories, hidecategories=Object)
 class UMaterialExpressionTime : public UMaterialExpression
 {
-	GENERATED_UCLASS_BODY()
+	GENERATED_BODY()
+
+public:
 
 	/** This time continues advancing regardless of whether the game is paused. */
 	UPROPERTY(EditAnywhere, Category=MaterialExpressionTime)
@@ -19,11 +21,11 @@ class UMaterialExpressionTime : public UMaterialExpression
 
 	/** Enables or disables the Period value. */
 	UPROPERTY(EditAnywhere, Category = MaterialExpressionTime, meta=(InlineEditConditionToggle))
-	uint32 bOverride_Period:1;
+	uint32 bOverride_Period:1 = false;
 
 	/** Time will loop around once it gets to Period. */
 	UPROPERTY(EditAnywhere, Category=MaterialExpressionTime, meta = (editcondition = "bOverride_Period", ClampMin = "0.0", ToolTip = "Period at which to wrap around time"))
-	float Period;
+	float Period = 0.0f;
 
 	//~ Begin UMaterialExpression Interface
 #if WITH_EDITOR
diff --git a/Engine/Source/Runtime/Engine/Public/Materials/MaterialExpressionTransform.h b/Engine/Source/Runtime/Engine/Public/Materials/MaterialExpressionTransform.h
--- a/Engine/Source/Runtime/Engine/Public/Materials/MaterialExpressionTransform.h
+++ b/Engine/Source/Runtime/Engine/Public/Materials/MaterialExpressionTransform.h
@@ -66,7 +66,9 @@ enum EMaterialVectorCoordTransform : int
 UCLASS(MinimalAPI, collapsecategories, hidecategories=Object)
 class UMaterialExpressionTransform : public UMaterialExpression
 {
-	GENERATED_UCLASS_BODY()
+	GENERATED_BODY()
+
+public:
 
 	/** input expression for this transform */
 	UPROPERTY()
@@ -74,11 +76,11 @@ class UMaterialExpressionTransform : public UMaterialExpression
 
 	/** Source coordinate space of the FVector */
 	UPROPERTY(EditAnywhere, Category=MaterialExpressionTransform, meta=(DisplayName = "Source", ShowAsInputPin = "Advanced"))
-	TEnumAsByte<enum EMaterialVectorCoordTransformSource> TransformSourceType;
+	TEnumAsByte<enum EMaterialVectorCoordTransformSource> TransformSourceType = TRANSFORMSOURCE_Tangent;
 
 	/** Destination coordinate space of the FVector */
 	UPROPERTY(EditAnywhere, Category=MaterialExpressionTransform, meta=(DisplayName = "Destination", ShowAsInputPin = "Advanced"))
-	TEnumAsByte<enum EMaterialVectorCoordTransform> TransformType;
+	TEnumAsByte<enum EMaterialVectorCoordTransform> TransformType = TRANSFORM_World;
 
 	//~ Begin UMaterialExpression Interface
 #if WITH_EDITOR
diff --git a/Engine/Source/Runtime/Engine/Public/Materials/MaterialExpressionTransformPosition.h b/Engine/Source/Runtime/Engine/Public/Materials/MaterialExpressionTransformPosition.h
--- a/Engine/Source/Runtime/Engine/Public/Materials/MaterialExpressionTransformPosition.h
+++ b/Engine/Source/Runtime/Engine/Public/Materials/MaterialExpressionTransformPosition.h
@@ -51,7 +51,9 @@ enum EMaterialPositionTransformSource : int
 UCLASS(collapsecategories, hidecategories=Object, MinimalAPI)
 class UMaterialExpressionTransformPosition : public UMaterialExpression
 {
-	GENERATED_UCLASS_BODY()
+	GENERATED_BODY()
+
+public:
 
 	/** input expression for this transform */
 	UPROPERTY()
@@ -59,11 +61,11 @@ class UMaterialExpressionTransformPosition : public UMaterialExpression
 
 	/** source format of the position that will be transformed */
 	UPROPERTY(EditAnywhere, Category=MaterialExpressionTransformPosition, meta=(DisplayName = "Source", ShowAsInputPin = "Advanced", InvalidEnumValues="TRANSFORMPOSSOURCE_FirstPersonTranslatedWorld"))
-	TEnumAsByte<enum EMaterialPositionTransformSource> TransformSourceType;
+	TEnumAsByte<enum EMaterialPositionTransformSource> TransformSourceType = TRANSFORMPOSSOURCE_Local;
 
 	/** type of transform to apply to the input expression */
 	UPROPERTY(EditAnywhere, Category=MaterialExpressionTransformPosition, meta=(DisplayName = "Destination", ShowAsInputPin = "Advanced"))
-	TEnumAsByte<enum EMaterialPositionTransformSource> TransformType;
+	TEnumAsByte<enum EMaterialPositionTransformSource> TransformType = TRANSFORMPOSSOURCE_Local;
 
 	/** scale of the tiles used in Periodic World Space */
 	UPROPERTY(meta=(DisplayName = "Periodic World Tile Size", ToolTip = "Distance the camera can move before the world origin is moved", RequiredInput = "false"))
diff --git a/Engine/Source/Runtime/Engine/Public/Materials/MaterialExpressionUserSceneTexture.h b/Engine/Source/Runtime/Engine/Public/Materials/MaterialExpressionUserSceneTexture.h
--- a/Engine/Source/Runtime/Engine/Public/Materials/MaterialExpressionUserSceneTexture.h
+++ b/Engine/Source/Runtime/Engine/Public/Materials/MaterialExpressionUserSceneTexture.h
@@ -24,7 +24,7 @@ class UMaterialExpressionUserSceneTexture : public UMaterialExpression
 	
 	/** Whether to use point sampled texture lookup (default) or using [bi-linear] filtered (can be slower, avoid faceted lock with distortions) */
 	UPROPERTY(EditAnywhere, Category=UMaterialExpressionUserSceneTexture, meta=(DisplayName = "Filtered", ShowAsInputPin = "Advanced"))
-	bool bFiltered;
+	bool bFiltered = false;
 
 	/** Whether to clamp the texture lookup.  Necessary when sampling a UserSceneTexture at reduced resolution with filtering, to avoid blending out of bounds pixels. */
 	UPROPERTY(EditAnywhere, Category=UMaterialExpressionUserSceneTexture, meta=(DisplayName = "Clamped", ShowAsInputPin = "Advanced"))
diff --git a/Engine/Source/Runtime/Engine/Public/Materials/MaterialExpressionVectorParameter.h b/Engine/Source/Runtime/Engine/Public/Materials/MaterialExpressionVectorParameter.h
--- a/Engine/Source/Runtime/Engine/Public/Materials/MaterialExpressionVectorParameter.h
+++ b/Engine/Source/Runtime/Engine/Public/Materials/MaterialExpressionVectorParameter.h
@@ -15,7 +15,7 @@ class UMaterialExpressionVectorParameter : public UMaterialExpressionParameter
 	GENERATED_UCLASS_BODY()
 
 	UPROPERTY(EditAnywhere, Category=MaterialExpressionVectorParameter, meta = (ShowAsInputPin = "Primary"))
-	FLinearColor DefaultValue;
+	FLinearColor DefaultValue = FLinearColor(1.0f, 0.0f, 0.0f, 0.0f);
 
 	UPROPERTY(EditAnywhere, Category=CustomPrimitiveData)
 	bool bUseCustomPrimitiveData = false;
diff --git a/Engine/Source/Runtime/Engine/Public/Materials/MaterialExpressionViewProperty.h b/Engine/Source/Runtime/Engine/Public/Materials/MaterialExpressionViewProperty.h
--- a/Engine/Source/Runtime/Engine/Public/Materials/MaterialExpressionViewProperty.h
+++ b/Engine/Source/Runtime/Engine/Public/Materials/MaterialExpressionViewProperty.h
@@ -62,7 +62,7 @@ class UMaterialExpressionViewProperty : public UMaterialExpression
 	
 	/** View input property to be accessed */
 	UPROPERTY(EditAnywhere, Category=UMaterialExpressionViewProperty, meta=(DisplayName = "View Property", ShowAsInputPin = "Advanced"))
-	TEnumAsByte<EMaterialExposedViewProperty> Property;
+	TEnumAsByte<EMaterialExposedViewProperty> Property = MEVP_FieldOfView;
 	
 	//~ Begin UMaterialExpression Interface
 #if WITH_EDITOR
diff --git a/Engine/Source/Runtime/Engine/Public/Materials/MaterialExpressionVolumetricAdvancedMaterialOutput.h b/Engine/Source/Runtime/Engine/Public/Materials/MaterialExpressionVolumetricAdvancedMaterialOutput.h
--- a/Engine/Source/Runtime/Engine/Public/Materials/MaterialExpressionVolumetricAdvancedMaterialOutput.h
+++ b/Engine/Source/Runtime/Engine/Public/Materials/MaterialExpressionVolumetricAdvancedMaterialOutput.h
@@ -45,51 +45,51 @@ class UMaterialExpressionVolumetricAdvancedMaterialOutput : public UMaterialExpr
 
 	/** Only used if PhaseG is not hooked up. Parameter 'g' input to the phase function  describing how much forward(g<0) or backward (g>0) light scatter around. */
 	UPROPERTY(EditAnywhere, Category = "Phase", meta = (OverridingInputProperty = "PhaseG", UIMin = -0.999f, UIMax = 0.999f, ClampMin = -0.999f, ClampMax = 0.999f))
-	float ConstPhaseG;
+	float ConstPhaseG = 0.0f;
 
 	/** Only used if PhaseG2 is not hooked up. Parameter 'g' input to the second phase function  describing how much forward(g<0) or backward (g>0) light scatter around. Valid range is [-1,1]. */
 	UPROPERTY(EditAnywhere, Category = "Phase", meta = (OverridingInputProperty = "PhaseG2", UIMin = -0.999f, UIMax = 0.999f, ClampMin = -0.999f, ClampMax = 0.999f))
-	float ConstPhaseG2;
+	float ConstPhaseG2 = 0.0f;
 
 	/** Only used if PhaseBlend is not hooked up. Lerp factor when blending the two phase functions parameterized by G and G2. Valid range is [0,1]. */
 	UPROPERTY(EditAnywhere, Category = "Phase", meta = (OverridingInputProperty = "PhaseBlend", UIMin = 0.0f, UIMax = 1.0f, ClampMin = -0.999f, ClampMax = 0.999f))
-	float ConstPhaseBlend;
+	float ConstPhaseBlend = 0.0f;
 
 	/** Set this to true to force the phase function to be evaluated per sample, instead once per pixel (globally). Per sample evaluation is slower. */
 	UPROPERTY(EditAnywhere, Category = "Phase")
-	bool PerSamplePhaseEvaluation;
+	bool PerSamplePhaseEvaluation = false;
 
 	/** How many octave to use for the multiple-scattering approximation. This makes the shader more expensive so you should only use 0 or 1 for better performance, and tweak multiple scattering parameters accordingly. 0 means single scattering only. The maximum value is 2 (expenssive). */
 	UPROPERTY(EditAnywhere, Category = "Multi-Scattering", meta = (UIMin = 0, UIMax = 1, ClampMin = 0, ClampMax = 2))
-	uint32 MultiScatteringApproximationOctaveCount;
+	uint32 MultiScatteringApproximationOctaveCount = 0;
 
 	/** Only used if MultiScatteringContribution is not hooked up. Multi-scattering approximation: represents how much contribution each successive octave will add. Valid range is [0,1], from low to high contribution */
 	UPROPERTY(EditAnywhere, Category = "Multi-Scattering", meta = (OverridingInputProperty = "MultiScatteringContribution", UIMin = 0.0f, UIMax = 1.0f, ClampMin = 0.0f, ClampMax = 1.0f))
-	float ConstMultiScatteringContribution;
+	float ConstMultiScatteringContribution = 0.5f;
 
 	/** Only used if MultiScatteringOcclusion is not hooked up. Multi-scattering approximation: represents how much occlusion will be reduced for each successive octave. Valid range is [0,1], from low to high occlusion. */
 	UPROPERTY(EditAnywhere, Category = "Multi-Scattering", meta = (OverridingInputProperty = "MultiScatteringOcclusion", UIMin = 0.0f, UIMax = 1.0f, ClampMin = 0.0f, ClampMax = 1.0f))
-	float ConstMultiScatteringOcclusion;
+	float ConstMultiScatteringOcclusion = 0.5f;
 
 	/** Only used if MultiScatteringEccentricity is not hooked up. Multi-scattering approximation: represents how much the phase will become isotropic for each successive octave. Valid range is [0,1], from anisotropic to isotropic phase. */
 	UPROPERTY(EditAnywhere, Category = "Multi-Scattering", meta = (OverridingInputProperty = "MultiScatteringEccentricity", UIMin = 0.0f, UIMax = 1.0f, ClampMin = 0.0f, ClampMax = 1.0f))
-	float ConstMultiScatteringEccentricity;
+	float ConstMultiScatteringEccentricity = 0.5f;
 
 	/** Sample the shadowed lighting contribution from the ground onto the medium (single scattering). This adds some costs to the tracing when enabled.*/
 	UPROPERTY(EditAnywhere, Category = "Options")
-	bool bGroundContribution;
+	bool bGroundContribution = false;
 
 	/** Set this for the material to only be considered grey scale, only using the R chanel of the input parameters internally. The lighting will still be colored. This is an optimisation.*/
 	UPROPERTY(EditAnywhere, Category = "Options")
-	bool bGrayScaleMaterial;
+	bool bGrayScaleMaterial = false;
 
 	/** Disable this to use the cloud shadow map instead of secondary raymarching. This is usually enough for clouds viewed from the ground and it result in a performance boost. Shadow now have infinite length but also becomes less accurate and gray scale.*/
 	UPROPERTY(EditAnywhere, Category = "Options")
-	bool bRayMarchVolumeShadow;
+	bool bRayMarchVolumeShadow = true;
 
 	/** Set whether multiple scattering contribution entry is clamped in [0,1] or not. When disabled, the artist is in charge for ensuring the visual remain in a reasonable brighness range.*/
 	UPROPERTY(EditAnywhere, Category = "Options")
-	bool bClampMultiScatteringContribution;
+	bool bClampMultiScatteringContribution = true;
 
 public:
 #if WITH_EDITOR
diff --git a/Engine/Source/Runtime/Engine/Public/Materials/MaterialExpressionWorldPosition.h b/Engine/Source/Runtime/Engine/Public/Materials/MaterialExpressionWorldPosition.h
--- a/Engine/Source/Runtime/Engine/Public/Materials/MaterialExpressionWorldPosition.h
+++ b/Engine/Source/Runtime/Engine/Public/Materials/MaterialExpressionWorldPosition.h
@@ -33,7 +33,7 @@ class UMaterialExpressionWorldPosition : public UMaterialExpression
 	GENERATED_UCLASS_BODY()
 	
 	UPROPERTY(EditAnywhere, Category=UMaterialExpressionWorldPosition, meta=(DisplayName = "Shader Offsets", ShowAsInputPin = "Advanced"))
-	TEnumAsByte<EWorldPositionIncludedOffsets> WorldPositionShaderOffset;
+	TEnumAsByte<EWorldPositionIncludedOffsets> WorldPositionShaderOffset = WPT_Default;
 	
 	//~ Begin UMaterialExpression Interface
 #if WITH_EDITOR
diff --git a/Engine/Source/Runtime/Engine/Public/Materials/MaterialFunctionInterface.h b/Engine/Source/Runtime/Engine/Public/Materials/MaterialFunctionInterface.h
--- a/Engine/Source/Runtime/Engine/Public/Materials/MaterialFunctionInterface.h
+++ b/Engine/Source/Runtime/Engine/Public/Materials/MaterialFunctionInterface.h
@@ -57,7 +57,7 @@ public:
 UCLASS(abstract, hidecategories=object, MinimalAPI)
 class UMaterialFunctionInterface : public UObject
 {
-	GENERATED_UCLASS_BODY()
+	GENERATED_BODY()
 
 #if WITH_EDITORONLY_DATA
 protected:
@@ -74,6 +74,8 @@ public:
 	const UMaterialFunctionInterfaceEditorOnlyData* GetEditorOnlyData() const { return EditorOnlyData; }
 #endif // WITH_EDITORONLY_DATA
 
+public:
+
 	//~ Begin UObject Interface.
 	virtual void PostInitProperties() override;
 	virtual void PostDuplicate(bool bDuplicateForPIE) override;
@@ -98,7 +100,7 @@ public:
 protected:
 	/** The intended usage of this function, required for material layers. */
 	UPROPERTY(AssetRegistrySearchable)
-	EMaterialFunctionUsage MaterialFunctionUsage;
+	EMaterialFunctionUsage MaterialFunctionUsage = EMaterialFunctionUsage::Default;
 
 public:
 	virtual EMaterialFunctionUsage GetMaterialFunctionUsage()
