Fix a couple issues with -game or -server mode and world partitioning. In UWorldPartition::PrepareEditorGameWorld(), make sure we are running as the editor to check if we are in PIE mode, and place OFPAs in a sublevel instead of the Memory package to mimic those actors being added to the level when the level is cooked.

diff --git a/Engine/Source/Runtime/Engine/Private/WorldPartition/WorldPartition.cpp b/Engine/Source/Runtime/Engine/Private/WorldPartition/WorldPartition.cpp
--- a/Engine/Source/Runtime/Engine/Private/WorldPartition/WorldPartition.cpp
+++ b/Engine/Source/Runtime/Engine/Private/WorldPartition/WorldPartition.cpp
@@ -468,7 +468,7 @@ void UWorldPartition::OnCancelPIE()
 void UWorldPartition::PrepareEditorGameWorld()
 {
 	check(!bIsPIE);
-	bIsPIE = !IsRunningGame();
+	bIsPIE = GIsEditor && !IsRunningGame();
 
 	// In PIE, we always want to populate the map check dialog
 	FStreamingGenerationMapCheckErrorHandler MapCheckErrorHandler;
diff --git a/Engine/Source/Runtime/Engine/Private/WorldPartition/WorldPartitionLevelStreamingPolicy.cpp b/Engine/Source/Runtime/Engine/Private/WorldPartition/WorldPartitionLevelStreamingPolicy.cpp
--- a/Engine/Source/Runtime/Engine/Private/WorldPartition/WorldPartitionLevelStreamingPolicy.cpp
+++ b/Engine/Source/Runtime/Engine/Private/WorldPartition/WorldPartitionLevelStreamingPolicy.cpp
@@ -65,8 +65,15 @@ FString UWorldPartitionLevelStreamingPolicy::GetCellPackagePath(const FName& InC
 {
 	if (InWorld->IsGameWorld())
 	{
-		// Set as memory package to avoid wasting time in FPackageName::DoesPackageExist
-		return FString::Printf(TEXT("/Memory/%s_%s"), *FPackageName::GetShortName(InWorld->GetPackage()), *InCellName.ToString());
+		if (InWorld->WorldType == EWorldType::PIE)
+		{
+			// Set as memory package to avoid wasting time in FPackageName::DoesPackageExist
+			return FString::Printf(TEXT("/Memory/%s_%s"), *FPackageName::GetShortName(InWorld->GetPackage()), *InCellName.ToString());
+		}
+		
+		// This is for non-PIE mode in editor, such as with -game and -server. It kind of mimics OFPAs being
+		// flattened into the level list upon cooking.
+		return FString::Printf(TEXT("%s/%s"), *InWorld->GetPackage()->GetName(), *InCellName.ToString());
 	}
 
 	return FString::Printf(TEXT("/%s"), *InCellName.ToString());
