Scene Capture:  Remove complexity added to interleave Custom Render Passes and Scene Captures in the correct priority order.  Complexity was added because main scene renderer construction had been inadvertently moved before deferred Scene Capture updates, and the constructor pulls CRPs from the FScene, which hadn't been added yet, resulting in CRPs being delayed by a frame.  The simplified solution is to move main scene renderer construction after deferred Scene Capture updates, and revert the previously added code.

#rnx

diff --git a/Engine/Source/Runtime/Engine/Internal/SceneCapture/SceneCaptureInternal.h b/Engine/Source/Runtime/Engine/Internal/SceneCapture/SceneCaptureInternal.h
--- a/Engine/Source/Runtime/Engine/Internal/SceneCapture/SceneCaptureInternal.h
+++ b/Engine/Source/Runtime/Engine/Internal/SceneCapture/SceneCaptureInternal.h
@@ -10,7 +10,6 @@ class ISceneRenderBuilder;
 
 /**
  * Internal function called from FRendererModule::BeginRenderingViewFamilies.  List of view families is required to handle captures where
- * USceneCaptureComponent2D::ShouldRenderWithMainViewFamily() is true.  Optionally can process only custom render passes that should render
- * with the main view, so they can be added to the FScene before the main view scene renderer is constructed.
+ * USceneCaptureComponent2D::ShouldRenderWithMainViewFamily() is true.
  */
-ENGINE_API void SceneCaptureUpdateDeferredCapturesInternal(FSceneInterface* Scene, TConstArrayView<const FSceneViewFamily*> ViewFamilies, ISceneRenderBuilder& SceneRenderBuilder, bool bMainCustomRenderPassOnly);
+ENGINE_API void SceneCaptureUpdateDeferredCapturesInternal(FSceneInterface* Scene, TConstArrayView<const FSceneViewFamily*> ViewFamilies, ISceneRenderBuilder& SceneRenderBuilder);
diff --git a/Engine/Source/Runtime/Engine/Private/Components/SceneCaptureComponent.cpp b/Engine/Source/Runtime/Engine/Private/Components/SceneCaptureComponent.cpp
--- a/Engine/Source/Runtime/Engine/Private/Components/SceneCaptureComponent.cpp
+++ b/Engine/Source/Runtime/Engine/Private/Components/SceneCaptureComponent.cpp
@@ -509,11 +509,11 @@ void USceneCaptureComponent::Serialize(FArchive& Ar)
 void USceneCaptureComponent::UpdateDeferredCaptures(FSceneInterface* Scene)
 {
 	TUniquePtr<ISceneRenderBuilder> SceneRenderBuilder = ISceneRenderBuilder::Create(Scene);
-	SceneCaptureUpdateDeferredCapturesInternal(Scene, {}, *SceneRenderBuilder, /*bMainCustomRenderPassOnly=*/ false);
+	SceneCaptureUpdateDeferredCapturesInternal(Scene, {}, *SceneRenderBuilder);
 	SceneRenderBuilder->Execute();
 }
 
-void SceneCaptureUpdateDeferredCapturesInternal(FSceneInterface* Scene, TConstArrayView<const FSceneViewFamily*> ViewFamilies, ISceneRenderBuilder& SceneRenderBuilder, bool bMainCustomRenderPassOnly)
+void SceneCaptureUpdateDeferredCapturesInternal(FSceneInterface* Scene, TConstArrayView<const FSceneViewFamily*> ViewFamilies, ISceneRenderBuilder& SceneRenderBuilder)
 {
 	FScopeLock ScopeLock(&SceneCapturesToUpdateMapCS);
 	UWorld* World = Scene->GetWorld();
@@ -522,16 +522,6 @@ void SceneCaptureUpdateDeferredCapturesInternal(FSceneInterface* Scene, TConstAr
 		return;
 	}
 
-	if (bMainCustomRenderPassOnly)
-	{
-		// If custom render passes are globally disabled, just return since that's all we are trying to process
-		IConsoleVariable* AllowRenderInMainRenderer = IConsoleManager::Get().FindConsoleVariable(TEXT("r.SceneCapture.AllowRenderInMainRenderer"));
-		if (!AllowRenderInMainRenderer || !AllowRenderInMainRenderer->GetBool())
-		{
-			return;
-		}
-	}
-
 	// Only update the scene captures associated with the current scene.
 	// Updating others not associated with the scene would cause invalid data to be rendered into the target
 	TArray< TWeakObjectPtr<USceneCaptureComponent> > SceneCapturesToUpdate;
@@ -547,42 +537,6 @@ void SceneCaptureUpdateDeferredCapturesInternal(FSceneInterface* Scene, TConstAr
 		}
 	}
 
-	SceneCapturesToUpdate.Sort([](const TWeakObjectPtr<USceneCaptureComponent>& A, const TWeakObjectPtr<USceneCaptureComponent>& B)
-	{
-		if (!A.IsValid())
-		{
-			return false;
-		}
-		else if (!B.IsValid())
-		{
-			return true;
-		}
-		return A->CaptureSortPriority > B->CaptureSortPriority;
-	});
-
-	// To ensure CRPs run in the correct sorted order with scene captures, we need to interleave the CRPs.  This is accomplished by
-	// processing CRPs and scene captures in order downstream, but the main scene renderer is allocated before the deferred scene
-	// capture update loop, so we need to specifically separate out the CRPs that should render with the main scene renderer.  Any
-	// CRP that occurs after the last non-CRP Scene Capture goes with the main view, so we need to track the last such scene capture,
-	// and whether we're past it in the array in the loop below.
-	const USceneCaptureComponent* LastSceneCaptureNonCRP = nullptr;
-	if (bMainCustomRenderPassOnly)
-	{
-		for (int32 CaptureIndex = 0; CaptureIndex < SceneCapturesToUpdate.Num(); CaptureIndex++)
-		{
-			if (SceneCapturesToUpdate[CaptureIndex].IsValid())
-			{
-				USceneCaptureComponent2D* Component = Cast<USceneCaptureComponent2D>(SceneCapturesToUpdate[CaptureIndex].Get());
-
-				// If it's a cube capture, it never renders in the main renderer, otherwise check if it renders in the main renderer (meaning it's a CRP)
-				if (!Component || !Component->ShouldRenderInMainRenderer())
-				{
-					LastSceneCaptureNonCRP = SceneCapturesToUpdate[CaptureIndex].Get();
-				}
-			}
-		}
-	}
-
 	int32 IncrementIfNotRemoved;
 	for (int32 CaptureIndex = 0; CaptureIndex < SceneCapturesToUpdate.Num(); CaptureIndex += IncrementIfNotRemoved)
 	{
@@ -590,23 +544,11 @@ void SceneCaptureUpdateDeferredCapturesInternal(FSceneInterface* Scene, TConstAr
 
 		if (SceneCapturesToUpdate[CaptureIndex].IsValid())
 		{
-			// Track that we've passed the last non CRP scene capture by setting this to NULL when we found it
-			if (LastSceneCaptureNonCRP == SceneCapturesToUpdate[CaptureIndex].Get())
-			{
-				LastSceneCaptureNonCRP = nullptr;
-			}
-
 			if (USceneCaptureComponent2D* Component = Cast<USceneCaptureComponent2D>(SceneCapturesToUpdate[CaptureIndex].Get()))
 			{
-				// If this is intended to render with the main view family, and we don't have a main view family, remove it from the array so
-				// it can render with the main view family later.  Also remove if we are only searching for main view Custom Render Passes, and:
-				//
-				//		1.  It's not a custom render pass:  !ShouldRenderInMainRenderer()    OR
-				//		2.  It's not a main view CRP, meaning there is another scene capture after it in the array (LastSceneCaptureNonCRP != nullptr),
-				//			and it's not explicitly flagged to render with the main view family.
-				//
-				if ((Component->ShouldRenderWithMainViewFamily() && !MainViewFamily) ||
-					(bMainCustomRenderPassOnly && (!Component->ShouldRenderInMainRenderer() || (LastSceneCaptureNonCRP && !Component->ShouldRenderWithMainViewFamily()))))
+				// If this is intended to render with the main view family, and we don't have a main view family,
+				// remove it from the array so it can render with the main view family later.
+				if (Component->ShouldRenderWithMainViewFamily() && !MainViewFamily)
 				{
 					SceneCapturesToUpdate.RemoveAtSwap(CaptureIndex);
 					IncrementIfNotRemoved = 0;
@@ -615,6 +557,19 @@ void SceneCaptureUpdateDeferredCapturesInternal(FSceneInterface* Scene, TConstAr
 		}
 	}
 
+	SceneCapturesToUpdate.Sort([](const TWeakObjectPtr<USceneCaptureComponent>& A, const TWeakObjectPtr<USceneCaptureComponent>& B)
+	{
+		if (!A.IsValid())
+		{
+			return false;
+		}
+		else if (!B.IsValid())
+		{
+			return true;
+		}
+		return A->CaptureSortPriority > B->CaptureSortPriority;
+	});
+
 	for (TWeakObjectPtr<USceneCaptureComponent> Component : SceneCapturesToUpdate)
 	{
 		if (Component.IsValid())
@@ -632,11 +587,10 @@ void SceneCaptureUpdateDeferredCapturesInternal(FSceneInterface* Scene, TConstAr
 				((USceneCaptureComponent2D*)Component.Get())->MainViewFamily = nullptr;
 			}
 		}
-
-		// Remove this component, as we processed it (or it was invalid).  We remove specific processed items,
-		// as there are cases where not all items are processed, and need to be handled later.
-		SceneCapturesToUpdateMap.Remove(World, Component);
 	}
+
+	// All scene captures for this world have been updated
+	SceneCapturesToUpdateMap.Remove(World);
 }
 
 void USceneCaptureComponent::OnUnregister()
diff --git a/Engine/Source/Runtime/Renderer/Private/SceneRendering.cpp b/Engine/Source/Runtime/Renderer/Private/SceneRendering.cpp
--- a/Engine/Source/Runtime/Renderer/Private/SceneRendering.cpp
+++ b/Engine/Source/Runtime/Renderer/Private/SceneRendering.cpp
@@ -5021,21 +5021,16 @@ void FRendererModule::BeginRenderingViewFamilies(FCanvas* Canvas, TConstArrayVie
 
 		FSceneRenderBuilder SceneRenderBuilder(Scene);
 
+		// Update deferred scene captures before creating the main view scene renderers, so custom render passes are available during scene renderer construction
 		bool bShowHitProxies = (Canvas->GetHitProxyConsumer() != nullptr);
 		if (!bShowHitProxies)
 		{
-			// Custom render passes need to be added before the main scene renderer is allocated, otherwise their update will be delayed by a frame
-			SceneCaptureUpdateDeferredCapturesInternal(Scene, ViewFamilies, SceneRenderBuilder, /*bCustomRenderPassOnly=*/ true);
+			SceneCaptureUpdateDeferredCapturesInternal(Scene, ViewFamilies, SceneRenderBuilder);
 		}
 
 		TArray<FSceneRenderer*, FConcurrentLinearArrayAllocator> SceneRenderers = SceneRenderBuilder.CreateLinkedSceneRenderers(ViewFamilies, Canvas->GetHitProxyConsumer());
 		SetupDebugViewModes(SceneRenderers);
 
-		if (!bShowHitProxies)
-		{
-			SceneCaptureUpdateDeferredCapturesInternal(Scene, ViewFamilies, SceneRenderBuilder, /*bCustomRenderPassOnly=*/ false);
-		}
-
 		if (!bShowHitProxies)
 		{
 			for (int32 ReflectionIndex = 0; ReflectionIndex < Scene->PlanarReflections_GameThread.Num(); ReflectionIndex++)
