Moved RVT CustomMaterialData from RVT component to RVT asset.
Feedback is that it will be easier to set up once on the asset instead of potentially multiple times on the component.
If we need both (with an override) it can added back on the component later.

diff --git a/Engine/Source/Runtime/Engine/Classes/Components/RuntimeVirtualTextureComponent.h b/Engine/Source/Runtime/Engine/Classes/Components/RuntimeVirtualTextureComponent.h
--- a/Engine/Source/Runtime/Engine/Classes/Components/RuntimeVirtualTextureComponent.h
+++ b/Engine/Source/Runtime/Engine/Classes/Components/RuntimeVirtualTextureComponent.h
@@ -56,10 +56,6 @@ protected:
 	UPROPERTY(EditAnywhere, BlueprintReadOnly, TextExportTransient, Category = RuntimeVirtualTexture)
 	TObjectPtr<URuntimeVirtualTexture> VirtualTexture = nullptr;
 
-	/** A float4 custom value that can be read in the material that writes this virtual texture. */
-	UPROPERTY(EditAnywhere, BlueprintReadOnly, AdvancedDisplay, Category = RuntimeVirtualTexture)
-	FVector4f CustomMaterialData;
-
 	/** Per platform overrides for enabling the virtual texture. Only affects In-Game and PIE. */
 	UPROPERTY(EditAnywhere, AdvancedDisplay, Category = RuntimeVirtualTexture)
 	FPerPlatformBool EnableInGamePerPlatform;
@@ -156,7 +152,7 @@ public:
 	ENGINE_API bool IsEnabledInScene() const;
 
 	/** Public getter for the custom material data on this component. */
-	ENGINE_API FVector4f GetCustomMaterialData() const { return CustomMaterialData; }
+	ENGINE_API FVector4f GetCustomMaterialData() const;
 
 	/** Get if scalability settings are enabled. */
 	bool IsScalable() const { return bEnableScalability; }
@@ -191,6 +187,7 @@ public:
 	/** Returns whether the StreamingTexure contents is valid for use for the specified shading path */
 	ENGINE_API bool IsStreamingTextureInvalid(EShadingPath ShadingPath) const;
 
+	/** Public getter for priority setting. */
 	ENGINE_API EVTProducerPriority GetPriority() const;
 
 	enum class EStreamingTextureStatusFlags : uint8
diff --git a/Engine/Source/Runtime/Engine/Classes/VT/RuntimeVirtualTexture.h b/Engine/Source/Runtime/Engine/Classes/VT/RuntimeVirtualTexture.h
--- a/Engine/Source/Runtime/Engine/Classes/VT/RuntimeVirtualTexture.h
+++ b/Engine/Source/Runtime/Engine/Classes/VT/RuntimeVirtualTexture.h
@@ -82,8 +82,12 @@ protected:
 	UPROPERTY(EditAnywhere, AdvancedDisplay, Category = Layout, meta = (UIMin = "0", UIMax = "5", DisplayName = "Number of low mips to remove from the virtual texture"))
 	int32 RemoveLowMips = 0;
 
+	/** A float4 custom value that can be read in the material that writes this virtual texture. */
+	UPROPERTY(EditAnywhere, BlueprintReadOnly, Category = Settings)
+	FVector4f CustomMaterialData;
+
 	/** Texture group this texture belongs to */
-	UPROPERTY(EditAnywhere, BlueprintReadWrite, Category = LevelOfDetail, meta = (DisplayName = "Texture Group"), AssetRegistrySearchable)
+	UPROPERTY(EditAnywhere, BlueprintReadWrite, Category = Settings, meta = (DisplayName = "Texture Group"), AssetRegistrySearchable)
 	TEnumAsByte<enum TextureGroup> LODGroup = TEXTUREGROUP_World;
 	
 	/** Deprecated size of virtual texture. */
@@ -134,6 +138,8 @@ public:
 	int32 GetRemoveLowMips() const { return FMath::Clamp(RemoveLowMips, 0, 5); }
 	/** Public getter for virtual texture using low quality compression flag. */
 	bool GetLQCompression() const { return bUseLowQualityCompression; }
+	/** Public getter for the custom material data. */
+	FVector4f GetCustomMaterialData() const { return CustomMaterialData; }
 	/** Public getter for texture LOD Group */
 	TEnumAsByte<enum TextureGroup> GetLODGroup() const { return LODGroup; }
 
diff --git a/Engine/Source/Runtime/Engine/Private/Components/RuntimeVirtualTextureComponent.cpp b/Engine/Source/Runtime/Engine/Private/Components/RuntimeVirtualTextureComponent.cpp
--- a/Engine/Source/Runtime/Engine/Private/Components/RuntimeVirtualTextureComponent.cpp
+++ b/Engine/Source/Runtime/Engine/Private/Components/RuntimeVirtualTextureComponent.cpp
@@ -65,7 +65,6 @@ static TAutoConsoleVariable<bool> CVarVTVisLogManualInvalidations(
 
 URuntimeVirtualTextureComponent::URuntimeVirtualTextureComponent(const FObjectInitializer& ObjectInitializer)
 	: Super(ObjectInitializer)
-	, CustomMaterialData(FVector4f::Zero())
 	, EnableInGamePerPlatform(true)
 	, SceneProxy(nullptr)
 {
@@ -386,6 +385,10 @@ EVTProducerPriority URuntimeVirtualTextureComponent::GetPriority() const
 	return (VirtualTexture != nullptr) ? VirtualTexture->GetPriority() : EVTProducerPriority::Normal;
 }
 
+FVector4f URuntimeVirtualTextureComponent::GetCustomMaterialData() const
+{
+	return (VirtualTexture != nullptr) ? VirtualTexture->GetCustomMaterialData() : FVector4f::Zero();
+}
 
 #if WITH_EDITOR
 
diff --git a/Engine/Source/Runtime/Engine/Private/VT/RuntimeVirtualTexture.cpp b/Engine/Source/Runtime/Engine/Private/VT/RuntimeVirtualTexture.cpp
--- a/Engine/Source/Runtime/Engine/Private/VT/RuntimeVirtualTexture.cpp
+++ b/Engine/Source/Runtime/Engine/Private/VT/RuntimeVirtualTexture.cpp
@@ -282,6 +282,7 @@ private:
 
 URuntimeVirtualTexture::URuntimeVirtualTexture(const FObjectInitializer& ObjectInitializer)
 	: Super(ObjectInitializer)
+	, CustomMaterialData(FVector4f::Zero())
 {
 	// Initialize the RHI resources with a null producer
 	if (!HasAnyFlags(RF_ClassDefaultObject) && FApp::CanEverRender())
