[Shaders] Don't spawn local shader compiler if the distributed controller already handled local shader compilation.

Some distributed shader controllers such as UBA and XGE support local workers alongside their remote workers.
Their use of local workers can result in oversubscription if the standard thread for local shader compilation is also running.
This change avoids launching the local shader compiling thread when the distributed controller supports local workers.
The distributed controller can also notify the shader compiling manager if this condition changes at runtime,
since CVar 'r.XGEController.AvoidUsingLocalMachine' for instance can be modified.

diff --git a/Engine/Plugins/UbaController/Source/Private/UbaControllerModule.cpp b/Engine/Plugins/UbaController/Source/Private/UbaControllerModule.cpp
--- a/Engine/Plugins/UbaController/Source/Private/UbaControllerModule.cpp
+++ b/Engine/Plugins/UbaController/Source/Private/UbaControllerModule.cpp
@@ -264,6 +264,13 @@ void FUbaControllerModule::InitializeController()
 	}
 }
 
+bool FUbaControllerModule::SupportsLocalWorkers()
+{
+	// UbaController supports local workers if the maximum number of local cores is greater than zero or -1 (special value to use all available cores)
+	const int32 MaxLocalCores = FUbaHordeConfig::Get().HordeMaxLocalCores;
+	return MaxLocalCores > 0 || MaxLocalCores == -1;
+}
+
 FString FUbaControllerModule::CreateUniqueFilePath()
 {
 	check(bSupported);
diff --git a/Engine/Plugins/UbaController/Source/Public/UbaControllerModule.h b/Engine/Plugins/UbaController/Source/Public/UbaControllerModule.h
--- a/Engine/Plugins/UbaController/Source/Public/UbaControllerModule.h
+++ b/Engine/Plugins/UbaController/Source/Public/UbaControllerModule.h
@@ -24,6 +24,7 @@ public:
 	virtual void StartupModule() override final;
 	virtual void ShutdownModule() override final;
 	virtual void InitializeController() override final;
+	virtual bool SupportsLocalWorkers() override final;
 	virtual FString CreateUniqueFilePath() override final;
 	virtual TFuture<FDistributedBuildTaskResult> EnqueueTask(const FTaskCommandData& CommandData) override final;
 	virtual bool PollStats(FDistributedBuildStats& OutStats) override;
diff --git a/Engine/Plugins/XGEController/Source/Private/XGEControllerModule.cpp b/Engine/Plugins/XGEController/Source/Private/XGEControllerModule.cpp
--- a/Engine/Plugins/XGEController/Source/Private/XGEControllerModule.cpp
+++ b/Engine/Plugins/XGEController/Source/Private/XGEControllerModule.cpp
@@ -71,6 +71,13 @@ namespace XGEControllerVariables
 		TEXT("0: Do not avoid. Distributed tasks will be spawned on all available XGE agents. Can cause oversubscription on the initiator machine. \n")
 		TEXT("1: Avoid spawning tasks on the local (initiator) machine except when running a commandlet or -buildmachine is passed (default).\n")
 		TEXT("2: Avoid spawning tasks on the local (initiator) machine."),
+		FConsoleVariableDelegate::CreateLambda(
+			[](IConsoleVariable* /*CVar*/) -> void
+			{
+				// If XGE has changed its support for local workers, notify shader compiling manager in case local shader compiling thread must be spun up.
+				GShaderCompilingManager->OnDistributedShaderCompilingChanged();
+			}
+		),
 		ECVF_Default); // This can be flipped any time XGEControlWorker is restarted
 
 }
@@ -354,6 +361,13 @@ void FXGEControllerModule::InitializeController()
 	}
 }
 
+bool FXGEControllerModule::SupportsLocalWorkers()
+{
+	// XGE supports local workers if it's not explicitly disabled
+	const bool bAvoidLocalWorkers = XGEController::AvoidUsingLocalMachine();
+	return !bAvoidLocalWorkers;
+}
+
 void FXGEControllerModule::WriteOutThreadProc()
 {
 	do
diff --git a/Engine/Plugins/XGEController/Source/Public/XGEControllerModule.h b/Engine/Plugins/XGEController/Source/Public/XGEControllerModule.h
--- a/Engine/Plugins/XGEController/Source/Public/XGEControllerModule.h
+++ b/Engine/Plugins/XGEController/Source/Public/XGEControllerModule.h
@@ -67,6 +67,7 @@ public:
 	virtual void ShutdownModule() override final;
 
 	virtual void InitializeController() override final;
+	virtual bool SupportsLocalWorkers() override final;
 
 	XGECONTROLLER_API virtual const FString GetName() override final { return FString("XGE Controller"); };
 
diff --git a/Engine/Source/Developer/DistributedBuildInterface/Public/DistributedBuildControllerInterface.h b/Engine/Source/Developer/DistributedBuildInterface/Public/DistributedBuildControllerInterface.h
--- a/Engine/Source/Developer/DistributedBuildInterface/Public/DistributedBuildControllerInterface.h
+++ b/Engine/Source/Developer/DistributedBuildInterface/Public/DistributedBuildControllerInterface.h
@@ -75,6 +75,9 @@ class IDistributedBuildController : public IModuleInterface, public IModularFeat
 {
 public:
 	virtual bool SupportsDynamicReloading() override { return false; }
+
+	// Returns true if this distributed controller also supports local workers alongside remote workers. By default false.
+	virtual bool SupportsLocalWorkers() { return false; }
 	
 	virtual bool RequiresRelativePaths() { return false; }
 
diff --git a/Engine/Source/Runtime/Engine/Private/ShaderCompiler/ShaderCompiler.cpp b/Engine/Source/Runtime/Engine/Private/ShaderCompiler/ShaderCompiler.cpp
--- a/Engine/Source/Runtime/Engine/Private/ShaderCompiler/ShaderCompiler.cpp
+++ b/Engine/Source/Runtime/Engine/Private/ShaderCompiler/ShaderCompiler.cpp
@@ -784,17 +784,6 @@ FShaderCompilingManager::FShaderCompilingManager() :
 		bAllowAsynchronousShaderCompiling = false;
 	}
 
-	// If compiling through workers is not allowed, shaders cannot be compiled exclusively throgh remote workers either
-	GConfig->GetBool(TEXT("DevOptions.Shaders"), TEXT("bExclusiveRemoteShaderCompiling"), bExclusiveRemoteShaderCompiling, GEngineIni);
-	if (!bAllowCompilingThroughWorkers)
-	{
-		bExclusiveRemoteShaderCompiling = false;
-	}
-	else if (FParse::Param(FCommandLine::Get(), TEXT("ExclusiveRemoteShaderCompiling")))
-	{
-		bExclusiveRemoteShaderCompiling = true;
-	}
-
 	verify(GConfig->GetInt( TEXT("DevOptions.Shaders"), TEXT("MaxShaderJobBatchSize"), MaxShaderJobBatchSize, GEngineIni ));
 	verify(GConfig->GetBool( TEXT("DevOptions.Shaders"), TEXT("bPromptToRetryFailedShaderCompiles"), bPromptToRetryFailedShaderCompiles, GEngineIni ));
 	verify(GConfig->GetBool(TEXT("DevOptions.Shaders"), TEXT("bDebugBreakOnPromptToRetryShaderCompile"), bDebugBreakOnPromptToRetryShaderCompile, GEngineIni));
@@ -840,62 +829,22 @@ FShaderCompilingManager::FShaderCompilingManager() :
 
 	CalculateNumberOfCompilingThreads(FPlatformMisc::NumberOfCores(), FPlatformMisc::NumberOfCoresIncludingHyperthreads());
 
-	TUniquePtr<FShaderCompileThreadRunnableBase> RemoteCompileThread;
-	const bool bCanUseRemoteCompiling = bAllowCompilingThroughWorkers && ShaderCompiler::IsRemoteCompilingAllowed() && AllTargetPlatformSupportsRemoteShaderCompiling();
-	BuildDistributionController = bCanUseRemoteCompiling ? FindRemoteCompilerController() : nullptr;
-	
-	if (BuildDistributionController)
-	{
-		UE_LOG(LogShaderCompilers, Display, TEXT("Using %s for shader compilation"), *BuildDistributionController->GetName());
-		RemoteCompileThread = MakeUnique<FShaderCompileDistributedThreadRunnable_Interface>(this, *BuildDistributionController);
-	}
-	else if (bExclusiveRemoteShaderCompiling)
+	// Launch local and remote shader compiling threads
 	{
-		UE_LOG(LogShaderCompilers, Display, TEXT("Exclusive remote shader compiling requested, but no distributed controller is available"));
-		bExclusiveRemoteShaderCompiling = false;
-	}
-
-	GConfig->SetBool(TEXT("/Script/UnrealEd.UnrealEdOptions"), TEXT("UsingXGE"), RemoteCompileThread.IsValid(), GEditorIni);
+		constexpr bool bDelayCompileThreadsExecution = true;
+		FShaderCompileThreadRunnableBase* RemoteCompileThread = LaunchRemoteShaderCompilingThread(bDelayCompileThreadsExecution);
 
-	if (bExclusiveRemoteShaderCompiling)
-	{
-		// Only add remote shader compiling thread
-		Threads.Add(MoveTemp(RemoteCompileThread));
-	}
-	else
-	{
-		// Add local and remote shader compiling threads and distribute job priorities among local and remote workers
-		UE_LOG(LogShaderCompilers, Display, TEXT("Using %d local workers for shader compilation"), NumShaderCompilingThreads);
+		GConfig->SetBool(TEXT("/Script/UnrealEd.UnrealEdOptions"), TEXT("UsingXGE"), RemoteCompileThread != nullptr, GEditorIni);
 
-		TUniquePtr<FShaderCompileThreadRunnableBase> LocalThread = MakeUnique<FShaderCompileThreadRunnable>(this);
-		if (RemoteCompileThread)
+		if (!bExclusiveDistributedShaderCompiling)
 		{
-			checkf(ShaderCompiler::IsRemoteCompilingAllowed(), TEXT("We have a remote compiling thread without the remote compilation being allowed"));
-
-			// Only force-local jobs are guaranteed to stay on the local machine. Going wide with High priority jobs is important for the startup times,
-			// since special materials use High priority. Possibly the partition by priority is too rigid in general.
-			RemoteCompileThread->SetPriorityRange(EShaderCompileJobPriority::Low, EShaderCompileJobPriority::ExtraHigh);
-			LocalThread->SetPriorityRange(EShaderCompileJobPriority::Normal, EShaderCompileJobPriority::ForceLocal);
-			Threads.Add(MoveTemp(RemoteCompileThread));
+			LaunchLocalShaderCompilingThread(bDelayCompileThreadsExecution);
 		}
-		else
+
+		for (const auto& Thread : Threads)
 		{
-			if (GIsBuildMachine)
-			{
-				int32 MinSCWsToSpawnBeforeWarning = 8; // optional, default to 8
-				GConfig->GetInt(TEXT("DevOptions.Shaders"), TEXT("MinSCWsToSpawnBeforeWarning"), MinSCWsToSpawnBeforeWarning, GEngineIni);
-				if (NumShaderCompilingThreads < static_cast<uint32>(MinSCWsToSpawnBeforeWarning))
-				{
-					UE_LOG(LogShaderCompilers, Warning, TEXT("Only %d SCWs will be spawned, which will result in longer shader compile times."), NumShaderCompilingThreads);
-				}
-			}
+			Thread->StartThread();
 		}
-		Threads.Add(MoveTemp(LocalThread));
-	}
-
-	for (const auto& Thread : Threads)
-	{
-		Thread->StartThread();
 	}
 
 	OutOfMemoryDelegateHandle = FCoreDelegates::GetOutOfMemoryDelegate().AddRaw(this, &FShaderCompilingManager::ReportMemoryUsage);
@@ -1014,6 +963,104 @@ FShaderCompilingManager::FShaderCompilingManager() :
 #endif // WITH_EDITOR
 }
 
+FShaderCompileThreadRunnableBase* FShaderCompilingManager::LaunchShaderCompilingThread(TUniquePtr<FShaderCompileThreadRunnableBase>&& NewShaderCompilingThread, bool bDelayThreadExecution)
+{
+	if (!bDelayThreadExecution)
+	{
+		NewShaderCompilingThread->StartThread();
+	}
+
+	// Take ownership of new shader compiling thread
+	FShaderCompileThreadRunnableBase* ThreadRef = NewShaderCompilingThread.Get();
+	Threads.Add(MoveTemp(NewShaderCompilingThread));
+
+	// If there is more than one thread for shader compilation, re-arrange distribution of job priorities, to avoid all threads picking up the same type of jobs
+	if (Threads.Num() >= 2)
+	{
+		// Only force-local jobs are guaranteed to stay on the local machine. Going wide with High priority jobs is important for the startup times,
+		// since special materials use High priority. Possibly the partition by priority is too rigid in general.
+		for (const auto& Thread : Threads)
+		{
+			FShaderCompileThreadRunnableBase* CompileThread = Thread.Get();
+			switch (Thread->GetWorkerType())
+			{
+			case EShaderCompilerWorkerType::None:
+				checkNoEntry();
+				break;
+			case EShaderCompilerWorkerType::LocalThread:
+				CompileThread->SetPriorityRange(EShaderCompileJobPriority::Normal, EShaderCompileJobPriority::ForceLocal);
+				break;
+			case EShaderCompilerWorkerType::Distributed:
+				CompileThread->SetPriorityRange(EShaderCompileJobPriority::Low, EShaderCompileJobPriority::ExtraHigh);
+				break;
+			}
+		}
+	}
+
+	return ThreadRef;
+}
+
+FShaderCompileThreadRunnableBase* FShaderCompilingManager::LaunchRemoteShaderCompilingThread(bool bDelayThreadExecution)
+{
+	// Check if there already is a local compile thread
+	if (FShaderCompileThreadRunnableBase* LocalCompileThread = FindShaderCompilingThread(EShaderCompilerWorkerType::LocalThread))
+	{
+		return LocalCompileThread;
+	}
+
+	// Check if remote compiling is allowed and find distributed controller
+	const bool bCanUseRemoteCompiling = bAllowCompilingThroughWorkers && ShaderCompiler::IsRemoteCompilingAllowed() && AllTargetPlatformSupportsRemoteShaderCompiling();
+	BuildDistributionController = bCanUseRemoteCompiling ? FindRemoteCompilerController() : nullptr;
+
+	if (BuildDistributionController)
+	{
+		// Allocate distributed shader compiling thread
+		UE_LOG(LogShaderCompilers, Display, TEXT("Using %s for shader compilation"), *BuildDistributionController->GetName());
+		TUniquePtr<FShaderCompileThreadRunnableBase> RemoteCompileThread = MakeUnique<FShaderCompileDistributedThreadRunnable_Interface>(this, *BuildDistributionController);
+		bExclusiveDistributedShaderCompiling = BuildDistributionController->SupportsLocalWorkers();
+		return LaunchShaderCompilingThread(MoveTemp(RemoteCompileThread), bDelayThreadExecution);
+	}
+
+	return nullptr;
+}
+
+FShaderCompileThreadRunnableBase* FShaderCompilingManager::LaunchLocalShaderCompilingThread(bool bDelayThreadExecution)
+{
+	// Check if there already is a local compile thread
+	if (FShaderCompileThreadRunnableBase* LocalCompileThread = FindShaderCompilingThread(EShaderCompilerWorkerType::LocalThread))
+	{
+		return LocalCompileThread;
+	}
+
+	// Allocate local shader compiling thread
+	UE_LOG(LogShaderCompilers, Display, TEXT("Using %d local workers for shader compilation"), NumShaderCompilingThreads);
+	TUniquePtr<FShaderCompileThreadRunnableBase> LocalThread = MakeUnique<FShaderCompileThreadRunnable>(this);
+
+	if (GIsBuildMachine)
+	{
+		int32 MinSCWsToSpawnBeforeWarning = 8; // optional, default to 8
+		GConfig->GetInt(TEXT("DevOptions.Shaders"), TEXT("MinSCWsToSpawnBeforeWarning"), MinSCWsToSpawnBeforeWarning, GEngineIni);
+		if (NumShaderCompilingThreads < static_cast<uint32>(MinSCWsToSpawnBeforeWarning))
+		{
+			UE_LOG(LogShaderCompilers, Warning, TEXT("Only %d SCWs will be spawned, which will result in longer shader compile times."), NumShaderCompilingThreads);
+		}
+	}
+
+	return LaunchShaderCompilingThread(MoveTemp(LocalThread), bDelayThreadExecution);
+}
+
+FShaderCompileThreadRunnableBase* FShaderCompilingManager::FindShaderCompilingThread(EShaderCompilerWorkerType InWorkerType)
+{
+	for (const auto& Thread : Threads)
+	{
+		if (Thread->GetWorkerType() == InWorkerType)
+		{
+			return Thread.Get();
+		}
+	}
+	return nullptr;
+}
+
 FShaderCompilingManager::~FShaderCompilingManager()
 {
 	// we never initialized, so nothing to do
@@ -1165,6 +1212,23 @@ void FShaderCompilingManager::OnMachineResourcesChanged(int32 NumberOfCores, int
 	}
 }
 
+void FShaderCompilingManager::OnDistributedShaderCompilingChanged()
+{
+	if (BuildDistributionController)
+	{
+		// Only update conditions if local shader compiling thread was initially disabled but is now required
+		const bool bDistributedControllerSupportsLocalWorkers = BuildDistributionController->SupportsLocalWorkers();
+		if (bExclusiveDistributedShaderCompiling && !bDistributedControllerSupportsLocalWorkers)
+		{
+			// Launch local shader compiling thread, since the distributed controller no longer supports local workers
+			if (LaunchLocalShaderCompilingThread() != nullptr)
+			{
+				bExclusiveDistributedShaderCompiling = false;
+			}
+		}
+	}
+}
+
 FName FShaderCompilingManager::GetStaticAssetTypeName()
 {
 	return TEXT("UE-Shader");
diff --git a/Engine/Source/Runtime/Engine/Private/ShaderCompiler/ShaderCompilerDistributed.cpp b/Engine/Source/Runtime/Engine/Private/ShaderCompiler/ShaderCompilerDistributed.cpp
--- a/Engine/Source/Runtime/Engine/Private/ShaderCompiler/ShaderCompilerDistributed.cpp
+++ b/Engine/Source/Runtime/Engine/Private/ShaderCompiler/ShaderCompilerDistributed.cpp
@@ -325,7 +325,7 @@ int32 FShaderCompileDistributedThreadRunnable_Interface::CompilingLoop()
 		{
 			// Grab as many jobs from the job queue as we can, unless there is no local shader compiling thread to pick up smaller batches
 			const EShaderCompileJobPriority Priority = (EShaderCompileJobPriority)PriorityIndex;
-			const int32 MinBatchSize = (Priority == EShaderCompileJobPriority::Low || Manager->IsExclusiveRemoteCompilingEnabled()) ? 1 : DistributedShaderCompilerVariables::MinBatchSize;
+			const int32 MinBatchSize = (Priority == EShaderCompileJobPriority::Low || Manager->IsExclusiveDistributedCompilingEnabled()) ? 1 : DistributedShaderCompilerVariables::MinBatchSize;
 			const int32 NumJobs = Manager->AllJobs.GetPendingJobs(EShaderCompilerWorkerType::Distributed, Priority, MinBatchSize, INT32_MAX, PendingJobs);
 			if (NumJobs > 0)
 			{
diff --git a/Engine/Source/Runtime/Engine/Private/ShaderCompiler/ShaderCompilerThreadRunnable.cpp b/Engine/Source/Runtime/Engine/Private/ShaderCompiler/ShaderCompilerThreadRunnable.cpp
--- a/Engine/Source/Runtime/Engine/Private/ShaderCompiler/ShaderCompilerThreadRunnable.cpp
+++ b/Engine/Source/Runtime/Engine/Private/ShaderCompiler/ShaderCompilerThreadRunnable.cpp
@@ -196,6 +196,13 @@ void FShaderCompileThreadRunnableBase::StartThread()
 	}
 }
 
+void FShaderCompileThreadRunnableBase::SetPriorityRange(EShaderCompileJobPriority MinPriority, EShaderCompileJobPriority MaxPriority)
+{
+	MinPriorityIndex = (int32)MinPriority;
+	MaxPriorityIndex = (int32)MaxPriority;
+	check(MaxPriorityIndex >= MinPriorityIndex);
+}
+
 FShaderCompileThreadRunnable::FShaderCompileThreadRunnable(FShaderCompilingManager* InManager)
 	: FShaderCompileThreadRunnableBase(InManager)
 #if PLATFORM_WINDOWS
diff --git a/Engine/Source/Runtime/Engine/Public/ShaderCompiler.h b/Engine/Source/Runtime/Engine/Public/ShaderCompiler.h
--- a/Engine/Source/Runtime/Engine/Public/ShaderCompiler.h
+++ b/Engine/Source/Runtime/Engine/Public/ShaderCompiler.h
@@ -181,12 +181,7 @@ public:
 	virtual ~FShaderCompileThreadRunnableBase()
 	{}
 
-	inline void SetPriorityRange(EShaderCompileJobPriority MinPriority, EShaderCompileJobPriority MaxPriority)
-	{
-		MinPriorityIndex = (int32)MinPriority;
-		MaxPriorityIndex = (int32)MaxPriority;
-		check(MaxPriorityIndex >= MinPriorityIndex);
-	}
+	void SetPriorityRange(EShaderCompileJobPriority MinPriority, EShaderCompileJobPriority MaxPriority);
 
 	void StartThread();
 
@@ -204,6 +199,9 @@ public:
 	/** Main work loop. */
 	virtual int32 CompilingLoop() = 0;
 
+	/** Returns the type of shader workers this thread represents. */
+	virtual EShaderCompilerWorkerType GetWorkerType() const = 0;
+
 	/** Events from the manager */
 	virtual void OnMachineResourcesChanged() {}
 };
@@ -246,6 +244,11 @@ public:
 	FShaderCompileThreadRunnable(class FShaderCompilingManager* InManager);
 	virtual ~FShaderCompileThreadRunnable();
 
+	virtual EShaderCompilerWorkerType GetWorkerType() const override final
+	{
+		return EShaderCompilerWorkerType::LocalThread;
+	}
+
 protected:
 	virtual bool PrintWorkerMemoryUsage(bool bAllowToWaitForLock = true) override;
 	virtual FShaderCompileMemoryUsage GetExternalWorkerMemoryUsage() override;
@@ -364,6 +367,11 @@ public:
 	FShaderCompileDistributedThreadRunnable_Interface(class FShaderCompilingManager* InManager, class IDistributedBuildController& InController);
 	virtual ~FShaderCompileDistributedThreadRunnable_Interface();
 
+	virtual EShaderCompilerWorkerType GetWorkerType() const override final
+	{
+		return EShaderCompilerWorkerType::Distributed;
+	}
+
 	/** Main work loop. */
 	virtual int32 CompilingLoop() override;
 
@@ -930,8 +938,8 @@ private:
 	bool bAllowCompilingThroughWorkers;
 	/** Whether to allow shaders to compile in the background or to block after each material. */
 	bool bAllowAsynchronousShaderCompiling;
-	/** Whether shaders are compiled exclusively through remote workers and no local shader compiling thread is allocated. Ignored if 'bAllowCompilingThroughWorkers' is false. */
-	bool bExclusiveRemoteShaderCompiling = false;
+	/** Whether shaders are compiled exclusively through the distributed shader controller and no local shader compiling thread is allocated. Ignored if 'bAllowCompilingThroughWorkers' is false. */
+	bool bExclusiveDistributedShaderCompiling = false;
 	/** Whether to ask to retry a failed shader compile error. */
 	bool bPromptToRetryFailedShaderCompiles;
 	/** If enabled when we enter the prompt to retry we will break in the debugger if one is attached rather than prompting. */
@@ -1025,6 +1033,18 @@ private:
 	/** Prints out the memory usage for shader compile worker processes, if they exist. */
 	void ReportMemoryUsage();
 
+	/** Takes the ownership of the new shader compiling thread and returns its non-owning pointer. If 'bDelayThreadExecution' is false, the thread is started immediately. */
+	FShaderCompileThreadRunnableBase* LaunchShaderCompilingThread(TUniquePtr<FShaderCompileThreadRunnableBase>&& NewShaderCompilingThread, bool bDelayThreadExecution = false);
+
+	/** Launches the thread for remote shader compilation. */
+	FShaderCompileThreadRunnableBase* LaunchRemoteShaderCompilingThread(bool bDelayThreadExecution = false);
+
+	/** Launches the thread for local shader compilation. If the distributed shader controller supports local workers, this *might not* be invoked until this condition changes. */
+	FShaderCompileThreadRunnableBase* LaunchLocalShaderCompilingThread(bool bDelayThreadExecution = false);
+
+	/** Returns the shader compiling thread of the specified kind or null if there is none. */
+	FShaderCompileThreadRunnableBase* FindShaderCompilingThread(EShaderCompilerWorkerType InWorkerType);
+
 public:
 	
 	ENGINE_API FShaderCompilingManager();
@@ -1033,6 +1053,9 @@ public:
 	/** Called by external systems that have updated the number of worker threads available. */
 	ENGINE_API void OnMachineResourcesChanged(int32 NumberOfCores, int32 NumberOfCoresIncludingHyperthreads);
 
+	/** Called when CVars are changed at runtime that determine whether or not the distributed shader compiler supports local workers (e.g. 'r.XGEController.AvoidUsingLocalMachine' can be changed at runtime). */
+	ENGINE_API void OnDistributedShaderCompilingChanged();
+
 	ENGINE_API int32 GetNumPendingJobs() const;
 	ENGINE_API int32 GetNumOutstandingJobs() const;
 
@@ -1071,12 +1094,12 @@ public:
 	}
 
 	/**
-	 * Returns whether shaders are exclusively compiled through remote workerd, i.e. no local shader compiling thread is allocated.
+	 * Returns whether shaders are exclusively compiled through distributed controller, i.e. no local shader compiling thread is allocated.
 	 * If this is true, the distributed controller must not refuse job batches below a minimum threshold, otherwise those jobs will never get picked up.
 	 */
-	bool IsExclusiveRemoteCompilingEnabled() const
+	bool IsExclusiveDistributedCompilingEnabled() const
 	{
-		return bExclusiveRemoteShaderCompiling;
+		return bExclusiveDistributedShaderCompiling;
 	}
 
 	/**
