Added MakeRandomStreamFromLocation to Kismet for easier generation of random-but-cook-deterministic values in levels/blueprints.

diff --git a/Engine/Source/Runtime/Engine/Classes/Kismet/KismetMathLibrary.h b/Engine/Source/Runtime/Engine/Classes/Kismet/KismetMathLibrary.h
--- a/Engine/Source/Runtime/Engine/Classes/Kismet/KismetMathLibrary.h
+++ b/Engine/Source/Runtime/Engine/Classes/Kismet/KismetMathLibrary.h
@@ -3809,6 +3809,16 @@ class UKismetMathLibrary : public UBlueprintFunctionLibrary
 	UFUNCTION(BlueprintPure, meta = (Keywords = "construct build", NativeMakeFunc), Category = "Math|Random")
 	static ENGINE_API FRandomStream MakeRandomStream(int32 InitialSeed);
 
+	/**
+	 * Makes a random stream that will generate deterministic values, optionally grouping outputs within cells on a grid
+	 * 
+	 * @param Location			The input location to return the stream for
+	 * @param DistanceInterval	If non-zero, the grid size for which all locations within a grid cell will return the same random stream
+	 * @param bIncludeZ			Whether to factor in the Z-coordinate of the input location
+	 */
+	UFUNCTION(BlueprintPure, Category = "Math|Random", meta = (NativeMakeFunc))
+	static ENGINE_API FRandomStream MakeRandomStreamFromLocation(const FVector& Location, float DistanceInterval = 200.f, bool bIncludeZ = false);
+
 	/** Breaks apart a random number generator */
 	UFUNCTION(BlueprintPure, Category = "Math|Random", meta = (NativeBreakFunc))
 	static ENGINE_API void BreakRandomStream(const FRandomStream& InRandomStream, int32& InitialSeed);
diff --git a/Engine/Source/Runtime/Engine/Private/KismetMathLibrary.cpp b/Engine/Source/Runtime/Engine/Private/KismetMathLibrary.cpp
--- a/Engine/Source/Runtime/Engine/Private/KismetMathLibrary.cpp
+++ b/Engine/Source/Runtime/Engine/Private/KismetMathLibrary.cpp
@@ -1417,6 +1417,35 @@ FRandomStream UKismetMathLibrary::MakeRandomStream(int32 InitialSeed)
 	return FRandomStream(InitialSeed);
 }
 
+FRandomStream UKismetMathLibrary::MakeRandomStreamFromLocation(const FVector& Location, float DistanceInterval, bool bIncludeZ)
+{
+	uint32 Seed = GetTypeHash(DistanceInterval);
+	if (FMath::IsNearlyZero(DistanceInterval))
+	{
+		Seed = HashCombine(Seed, GetTypeHash(Location.X));
+		Seed = HashCombine(Seed, GetTypeHash(Location.Y));
+		if (bIncludeZ)
+		{
+			Seed = HashCombine(Seed, GetTypeHash(Location.Z));
+		}
+	}
+	else
+	{
+		FVector GridLocation = Location / DistanceInterval;
+		const int32 X = static_cast<int32>(GridLocation.X);
+		const int32 Y = static_cast<int32>(GridLocation.Y);
+		Seed = HashCombine(Seed, GetTypeHash(X));
+		Seed = HashCombine(Seed, GetTypeHash(Y));
+		if (bIncludeZ)
+		{
+			const int32 Z = static_cast<int32>(GridLocation.Z);
+			Seed = HashCombine(Seed, GetTypeHash(Z));
+		}
+	}
+
+	return FRandomStream(Seed);
+}
+
 FVector UKismetMathLibrary::RandomUnitVectorInConeInRadiansFromStream(const FRandomStream& Stream, const FVector& ConeDir, float ConeHalfAngleInRadians)
 {
 	return Stream.VRandCone(ConeDir, ConeHalfAngleInRadians);
