- Added initial prototype support for "hybrid cooked editor" to load cooked-for-game assets from IoStore/ZenStore without staging a full cooked build.
- After cooking a game target, run the editor with -hybridcookededitor to activate it. It will look in a ZenStore for the cooked data, falling back to uncooked for anything it doesn't find in the Store

diff --git a/Engine/Config/BaseEngine.ini b/Engine/Config/BaseEngine.ini
--- a/Engine/Config/BaseEngine.ini
+++ b/Engine/Config/BaseEngine.ini
@@ -2576,6 +2576,15 @@ HighDetail=(NumLODs=6,MaxNumStreamedLODs=0,bSupportLODStreaming=0,LODPercentTria
 +BucketConfigs=(Name="DefaultBucket", Size=10000000)
 PathToRAC=RuntimeAssetCache
 
+[HybridCookedEditor]
+bIsSupported=True
+; because this is needed very early, before we've initialized the config system, this setting
+; MUST be in BaseEngine.ini or DefaultEngine.ini
+; {Platform} will use the Game target (default for standard cooked builds), if your
+; game uses Client builds, set this to {Platform}Client. Or if you use MakeCookedEditor to 
+; make full cooked editor builds, you could use {Platform}CookedEditor
+RuntimeTargetPlatform="{Platform}"
+
 [Zen]
 AutoLaunch=true
 
diff --git a/Engine/Source/Runtime/Core/Internal/CoreGlobalsInternal.h b/Engine/Source/Runtime/Core/Internal/CoreGlobalsInternal.h
new file mode 100644
--- /dev/null
+++ b/Engine/Source/Runtime/Core/Internal/CoreGlobalsInternal.h
@@ -0,0 +1,21 @@
+// Copyright Epic Games, Inc. All Rights Reserved.
+#pragma once
+
+#include "CoreTypes.h"
+#include "Misc/Build.h"
+
+#if WITH_EDITOR
+extern CORE_API bool PRIVATE_GIsRunningHybridCookedEditor;
+#endif
+
+/**
+ * Check if the editor is running in Hybrid Cooked Editor mode (where it can dynamically choose between cooked assets and uncooked files
+ */
+FORCEINLINE bool IsRunningHybridCookedEditor()
+{
+#if WITH_EDITOR
+	return PRIVATE_GIsRunningHybridCookedEditor;
+#else
+	return false;
+#endif
+}
diff --git a/Engine/Source/Runtime/Core/Private/Misc/CoreGlobals.cpp b/Engine/Source/Runtime/Core/Private/Misc/CoreGlobals.cpp
--- a/Engine/Source/Runtime/Core/Private/Misc/CoreGlobals.cpp
+++ b/Engine/Source/Runtime/Core/Private/Misc/CoreGlobals.cpp
@@ -15,6 +15,7 @@
 #include "Misc/CommandLine.h"
 #include "ProfilingDebugging/MiscTrace.h"
 #include "GenericPlatform/GenericPlatformCrashContext.h"
+#include "CoreGlobalsInternal.h"
 
 #ifndef FAST_PATH_UNIQUE_NAME_GENERATION
 #define FAST_PATH_UNIQUE_NAME_GENERATION (!WITH_EDITORONLY_DATA)
@@ -176,6 +177,7 @@ bool GIsRunningUnattendedScript = false;
 #if WITH_EDITOR
 bool					PRIVATE_GIsRunningCookCommandlet	= false;				/** Whether this executable is running the cook commandlet */
 bool					PRIVATE_GIsRunningDLCCookCommandlet = false;				/** Whether this executable is running the cook commandlet on a DLC plugin */
+bool					PRIVATE_GIsRunningHybridCookedEditor= false;				/** Whether the editor is running in hybrid cooked editor mode */
 namespace UE::Private
 {
 int32					GMultiprocessId = 0;
diff --git a/Engine/Source/Runtime/Core/Public/Misc/Build.h b/Engine/Source/Runtime/Core/Public/Misc/Build.h
--- a/Engine/Source/Runtime/Core/Public/Misc/Build.h
+++ b/Engine/Source/Runtime/Core/Public/Misc/Build.h
@@ -472,10 +472,10 @@
 #ifndef UE_IS_COOKED_EDITOR
 #define UE_IS_COOKED_EDITOR 0
 #endif
-// Controls whether to enable loading cooked packages from I/O store in editor builds. Defaults to UE_IS_COOKED_EDITOR
-// but can be defined separately if needed
+// Controls whether to enable loading cooked packages from I/O store in editor builds. Defaults to WITH_EDITOR as we can now
+// load from IoStore without needing to stage an entire cooked build
 #ifndef WITH_IOSTORE_IN_EDITOR
-#define WITH_IOSTORE_IN_EDITOR UE_IS_COOKED_EDITOR
+#define WITH_IOSTORE_IN_EDITOR WITH_EDITOR
 #endif
 
 // Controls if iostore will be forced on
diff --git a/Engine/Source/Runtime/CoreUObject/Private/Serialization/AsyncLoading2.cpp b/Engine/Source/Runtime/CoreUObject/Private/Serialization/AsyncLoading2.cpp
--- a/Engine/Source/Runtime/CoreUObject/Private/Serialization/AsyncLoading2.cpp
+++ b/Engine/Source/Runtime/CoreUObject/Private/Serialization/AsyncLoading2.cpp
@@ -90,6 +90,7 @@
 #include "Containers/Ticker.h"
 #include "Logging/MessageLog.h"
 #include "VerseVM/VVMVerse.h"
+#include "CoreGlobalsInternal.h"
 
 #include <atomic>
 
@@ -450,6 +451,17 @@ static FAutoConsoleVariableRef CVarAllowLoadingPackagesFromFileSystem(
 	ECVF_Default
 	);
 
+// In editor builds, for packages that are both local and in zenstore, do we prefer uncooked (local) or cooked (zenstore) packages
+static bool GPreferCookedPackagesInEditor = false;
+static FAutoConsoleVariableRef CVarPreferCookedPackagesInEditor(
+	TEXT("s.PreferCookedPackagesInEditor"),
+	GPreferCookedPackagesInEditor,
+	TEXT("In editor builds, for packages that are both local and in zenstore, do we prefer uncooked (local) or cooked (zenstore) packages"),
+	ECVF_Default
+);
+
+
+
 #if WITH_EDITOR
 // This is important for the editor because the linker can end up recreating missing exports from its CreateImport function.
 // If we don't put the package in the zenloader loading queue, we can be left with objects in a RF_NeedLoad state that will never be actually loaded.
@@ -4027,7 +4039,7 @@ public:
 
 private:
 #if ALT2_ENABLE_LINKERLOAD_SUPPORT || WITH_EDITOR
-	bool TryGetPackagePathFromFileSystem(FName& PackageNameToLoad, FName& UPackageName, FPackagePath& OutPackagePath)
+	bool TryGetPackagePathFromFileSystem(FName& InOutPackageNameToLoad, FName& InOutUPackageName, FPackagePath& InOutPackagePath)
 	{
 		// If disabling loading from the file system in a shipping/test build and we are not the Editor avoid doing this look-up
 		// As this can be expensive to call and can cause contention on the GameThread and ASyncLoadingThread
@@ -4036,6 +4048,17 @@ private:
 			return false;
 		}
 
+		// if we prefer cooked packages in the editor, then we don't want to use the linkerload version
+		// NOTE that we won't have fixed up the FNames, but if we aren't using linkerloads, it shouldn't matter
+		if (GPreferCookedPackagesInEditor)
+		{
+			bool bExistsInIoStore = FIoDispatcher::Get().DoesChunkExist(CreatePackageDataChunkId(FPackageId::FromName(InOutPackagePath.GetPackageFName())));
+			if (bExistsInIoStore)
+			{
+				return false;
+			}
+		}
+
 #if WITH_EDITORONLY_DATA
 		// In editor, set MatchCaseOnDisk=true so that we set the capitalization of the Package's FName to match the
 		// capitalization on disk. Different capitalizations can arise from imports of the package that were somehow
@@ -4047,16 +4070,18 @@ private:
 		constexpr bool bMatchCaseOnDisk = false;
 #endif
 
-		if (!PackageNameToLoad.IsNone() &&
-			FPackageName::DoesPackageExistEx(OutPackagePath, FPackageName::EPackageLocationFilter::FileSystem,
-				bMatchCaseOnDisk, &OutPackagePath) != FPackageName::EPackageLocationFilter::None)
+		// now, check to see if the package exists in the local file system, and if it does, correct the case as needed
+		if (!InOutPackageNameToLoad.IsNone() &&
+			FPackageName::DoesPackageExistEx(InOutPackagePath, FPackageName::EPackageLocationFilter::FileSystem,
+				bMatchCaseOnDisk, &InOutPackagePath) != FPackageName::EPackageLocationFilter::None)
 		{
-			FName CaseCorrectedPackageName = OutPackagePath.GetPackageFName();
-			if (PackageNameToLoad == UPackageName)
+			FName CaseCorrectedPackageName = InOutPackagePath.GetPackageFName();
+			if (InOutPackageNameToLoad == InOutUPackageName)
 			{
-				UPackageName = CaseCorrectedPackageName;
+				InOutUPackageName = CaseCorrectedPackageName;
 			}
-			PackageNameToLoad = CaseCorrectedPackageName;
+			InOutPackageNameToLoad = CaseCorrectedPackageName;
+
 			return true;
 		}
 		return false;
@@ -4449,6 +4474,12 @@ void FAsyncLoadingThread2::InitializeLoading()
 		ParsePackageNames(VerbosePackageNamesString, GAsyncLoading2_VerbosePackageIds);
 		ParsePackageNames(DebugPackageNamesString, GAsyncLoading2_VerbosePackageIds);
 		GAsyncLoading2_VerboseLogFilter = GAsyncLoading2_VerbosePackageIds.Num() > 0 ? 1 : 2;
+
+		// force allowing cooked packages in hybrid cooked editor mode
+		if (IsRunningHybridCookedEditor())
+		{
+			GPreferCookedPackagesInEditor = true;
+		}
 	}
 
 	FileOpenLogWrapper = (FPlatformFileOpenLog*)(FPlatformFileManager::Get().FindPlatformFile(FPlatformFileOpenLog::GetTypeName()));
diff --git a/Engine/Source/Runtime/CoreUObject/Private/UObject/LinkerLoad.cpp b/Engine/Source/Runtime/CoreUObject/Private/UObject/LinkerLoad.cpp
--- a/Engine/Source/Runtime/CoreUObject/Private/UObject/LinkerLoad.cpp
+++ b/Engine/Source/Runtime/CoreUObject/Private/UObject/LinkerLoad.cpp
@@ -70,6 +70,8 @@
 #include "Internationalization/GatherableTextData.h"
 #include "Async/MappedFileHandle.h"
 #include "Async/UniqueLock.h"
+#include "CoreGlobalsInternal.h"
+
 class FTexture2DResourceMem;
 
 #define LOCTEXT_NAMESPACE "LinkerLoad"
@@ -172,6 +174,17 @@ static FAutoConsoleVariableRef CVarAllowCookedDataInEditorBuilds(
 	ECVF_Default
 );
 
+#if WITH_EDITOR
+// force the cvar value in HybridCookerEditor mode
+FDelayedAutoRegisterHelper GInitDevice(EDelayedRegisterRunPhase::IniSystemReady, []
+	{
+		if (IsRunningHybridCookedEditor())
+		{
+			GAllowCookedDataInEditorBuilds = 1;
+		}
+	});
+#endif
+
 int32 GSkipAsyncLoaderForCookedData = 0;
 static FAutoConsoleVariableRef CVarSkipAsyncLoaderForCookedData(
 	TEXT("cook.SkipAsyncLoaderForCookedData"),
diff --git a/Engine/Source/Runtime/Engine/Private/Level.cpp b/Engine/Source/Runtime/Engine/Private/Level.cpp
--- a/Engine/Source/Runtime/Engine/Private/Level.cpp
+++ b/Engine/Source/Runtime/Engine/Private/Level.cpp
@@ -1313,114 +1313,119 @@ void ULevel::PostLoad()
 	}
 
 #if WITH_EDITOR
-	const FLinkerLoad* Linker = GetLinker();
-	check(Linker || bWasDuplicated);
-
-	if ((IsUsingExternalActors() || IsUsingExternalObjects())
-		&& OwningWorld
-		&& OwningWorld->WorldType == EWorldType::Editor)
+	// all of the code in this WITH_EDITOR check ultimately assumes the presence of a Linker, but loading a cooked map likely won't
+	// have a Linker (if coming from IoStore), and it's really not needed/supported with cooked maps anyway
+	if (!GetOutermost()->HasAllPackagesFlags(PKG_Cooked))
 	{
-		if (Linker && Linker->IsPackageRelocated())
-		{
-			const FString LevelPackageName = GetPackage()->GetName();
-			UE_LOG(LogLevel, Error, TEXT("The level %s was moved on disk without the editor knowing it. This might cause some loading issues and some saving issue like an actor potentialy stomping the file of another actor."), *LevelPackageName);
+		const FLinkerLoad* Linker = GetLinker();
+		check(Linker || bWasDuplicated);
 
-			if (GIsEditor)
+		if ((IsUsingExternalActors() || IsUsingExternalObjects())
+			&& OwningWorld
+			&& OwningWorld->WorldType == EWorldType::Editor)
+		{
+			if (Linker && Linker->IsPackageRelocated())
 			{
-				FFormatNamedArguments Args;
-				Args.Add(TEXT("PackageName"), FText::FromString(LevelPackageName));
-				FNotificationInfo Info(FText::Format(LOCTEXT("UnsupportedRelocatedLevel", "The level ({PackageName}) was relocated.\nLevels that use one file per actor cannot be moved on disk without using the editor or the migration tool."), Args));
-				Info.ExpireDuration = 10.0f;
+				const FString LevelPackageName = GetPackage()->GetName();
+				UE_LOG(LogLevel, Error, TEXT("The level %s was moved on disk without the editor knowing it. This might cause some loading issues and some saving issue like an actor potentialy stomping the file of another actor."), *LevelPackageName);
+
+				if (GIsEditor)
+				{
+					FFormatNamedArguments Args;
+					Args.Add(TEXT("PackageName"), FText::FromString(LevelPackageName));
+					FNotificationInfo Info(FText::Format(LOCTEXT("UnsupportedRelocatedLevel", "The level ({PackageName}) was relocated.\nLevels that use one file per actor cannot be moved on disk without using the editor or the migration tool."), Args));
+					Info.ExpireDuration = 10.0f;
 
-				FSlateNotificationManager::Get().AddNotification(Info);
+					FSlateNotificationManager::Get().AddNotification(Info);
+				}
 			}
 		}
-	}
 
-	if (IsUsingActorFolders() && IsUsingExternalObjects() && IsActorFolderObjectsFeatureAvailable())
-	{
-		if (!bWasDuplicated && !Linker->GetInstancingContext().HasTag(ULevel::DontLoadExternalFoldersTag))
+		if (IsUsingActorFolders() && IsUsingExternalObjects() && IsActorFolderObjectsFeatureAvailable())
 		{
-			// Load all folders for this level
-			FExternalPackageHelper::LoadObjectsFromExternalPackages<UActorFolder>(this, [this](UActorFolder* LoadedFolder)
+			if (!bWasDuplicated && !Linker->GetInstancingContext().HasTag(ULevel::DontLoadExternalFoldersTag))
 			{
-				check(IsValid(LoadedFolder));
-				LoadedExternalActorFolders.Add(LoadedFolder);
-			});
+				// Load all folders for this level
+				FExternalPackageHelper::LoadObjectsFromExternalPackages<UActorFolder>(this, [this](UActorFolder* LoadedFolder)
+					{
+						check(IsValid(LoadedFolder));
+						LoadedExternalActorFolders.Add(LoadedFolder);
+					});
+			}
 		}
-	}
 
-	// if we use external actors, load dynamic actors here
-	if (IsUsingExternalActors())
-	{
-		if (!bWasDuplicated &&
-			(!bIsPartitioned || Linker->GetInstancingContext().HasTag(ULevel::LoadAllExternalObjectsTag)) &&
-			!Linker->GetInstancingContext().HasTag(ULevel::DontLoadExternalObjectsTag))
+		// if we use external actors, load dynamic actors here
+		if (IsUsingExternalActors())
 		{
-			UPackage* LevelPackage = GetPackage();
-			bool bPackageForPIE = LevelPackage->HasAnyPackageFlags(PKG_PlayInEditor);
-			FName PackageResourceName = LevelPackage->GetLoadedPath().GetPackageFName();
-			bool bInstanced = !PackageResourceName.IsNone() && (PackageResourceName != LevelPackage->GetFName());
-
-			// if the level is instanced, create an instancing context for remapping the actor imports
-			FLinkerInstancingContext InstancingContext;
-			if (bInstanced)
+			if (!bWasDuplicated &&
+				(!bIsPartitioned || Linker->GetInstancingContext().HasTag(ULevel::LoadAllExternalObjectsTag)) &&
+				!Linker->GetInstancingContext().HasTag(ULevel::DontLoadExternalObjectsTag))
 			{
-				InstancingContext.AddPackageMapping(PackageResourceName, LevelPackage->GetFName());
-			}
+				UPackage* LevelPackage = GetPackage();
+				bool bPackageForPIE = LevelPackage->HasAnyPackageFlags(PKG_PlayInEditor);
+				FName PackageResourceName = LevelPackage->GetLoadedPath().GetPackageFName();
+				bool bInstanced = !PackageResourceName.IsNone() && (PackageResourceName != LevelPackage->GetFName());
 
-			TArray<FString> ActorPackageNames = GetOnDiskExternalActorPackages(/*bTryUsingPackageLoadedPath*/ true);
-			TArray<FString> InstancePackageNames;
-			for (const FString& ActorPackageName : ActorPackageNames)
-			{
+				// if the level is instanced, create an instancing context for remapping the actor imports
+				FLinkerInstancingContext InstancingContext;
 				if (bInstanced)
 				{
-					const FString InstancedName = GetExternalActorPackageInstanceName(LevelPackage->GetName(), ActorPackageName);
-					InstancePackageNames.Add(InstancedName);
-
-					InstancingContext.AddPackageMapping(FName(*ActorPackageName), FName(*InstancedName));
+					InstancingContext.AddPackageMapping(PackageResourceName, LevelPackage->GetFName());
 				}
-			}
-
-			TSet<AActor*> ActorsSet(ObjectPtrDecay(Actors));
-			for (int32 i=0; i < ActorPackageNames.Num(); i++)
-			{
-				const FString& ActorPackageName = ActorPackageNames[i];
 
-				// Propagate RF_Transient from the Level Package in case we are loading an instanced level
-				UPackage* ActorPackage = bInstanced ? CreatePackage( *InstancePackageNames[i]) : nullptr;
-				if (ActorPackage && LevelPackage->HasAnyFlags(RF_Transient))
+				TArray<FString> ActorPackageNames = GetOnDiskExternalActorPackages(/*bTryUsingPackageLoadedPath*/ true);
+				TArray<FString> InstancePackageNames;
+				for (const FString& ActorPackageName : ActorPackageNames)
 				{
-					ActorPackage->SetFlags(RF_Transient);
+					if (bInstanced)
+					{
+						const FString InstancedName = GetExternalActorPackageInstanceName(LevelPackage->GetName(), ActorPackageName);
+						InstancePackageNames.Add(InstancedName);
+
+						InstancingContext.AddPackageMapping(FName(*ActorPackageName), FName(*InstancedName));
+					}
 				}
 
-				bool bFoundActor = false;
-				ActorPackage = LoadPackage(ActorPackage, *ActorPackageName, bPackageForPIE ? LOAD_PackageForPIE : LOAD_None, nullptr, &InstancingContext);
-				
-				ForEachObjectWithPackage(ActorPackage, [this, &ActorsSet, &bFoundActor, bPackageForPIE](UObject* PackageObject)
+				TSet<AActor*> ActorsSet(ObjectPtrDecay(Actors));
+				for (int32 i=0; i < ActorPackageNames.Num(); i++)
 				{
-					// There might be multiple actors per package in the case where an actor as a child actor component as we put child actor in the same package as their parent
-					if (PackageObject->IsA<AActor>() && !PackageObject->IsTemplate())
-					{
-						AActor* Actor = (AActor*)PackageObject;
-						// Verity that the actor is not already in the array (this is valid if, during last save, the actor returned true in AActor::ShouldLevelKeepRefIfExternal)
-						if (!ActorsSet.Contains(Actor))
-						{
-							Actors.Add(Actor);
-						}
-						bFoundActor = true;
-					}
-					// In PIE, we make sure to clear RF_Standalone flag on objects in external packages (UMetaData) 
-					// This guarantees that external packages of actors that are destroyed during the PIE session will
-					// properly get GC'ed and will allow future edits/modifications of OFPA actors.
-					if (bPackageForPIE)
+					const FString& ActorPackageName = ActorPackageNames[i];
+
+					// Propagate RF_Transient from the Level Package in case we are loading an instanced level
+					UPackage* ActorPackage = bInstanced ? CreatePackage( *InstancePackageNames[i]) : nullptr;
+					if (ActorPackage && LevelPackage->HasAnyFlags(RF_Transient))
 					{
-						PackageObject->ClearFlags(RF_Standalone);
+						ActorPackage->SetFlags(RF_Transient);
 					}
-					return true;
-				}, false);
 
-				UE_CLOG(!bFoundActor, LogLevel, Error, TEXT("Failed to load Actor for External Actor Package %s"), *ActorPackageName);
+					bool bFoundActor = false;
+					ActorPackage = LoadPackage(ActorPackage, *ActorPackageName, bPackageForPIE ? LOAD_PackageForPIE : LOAD_None, nullptr, &InstancingContext);
+
+					ForEachObjectWithPackage(ActorPackage, [this, &ActorsSet, &bFoundActor, bPackageForPIE](UObject* PackageObject)
+						{
+							// There might be multiple actors per package in the case where an actor as a child actor component as we put child actor in the same package as their parent
+							if (PackageObject->IsA<AActor>() && !PackageObject->IsTemplate())
+							{
+								AActor* Actor = (AActor*)PackageObject;
+								// Verity that the actor is not already in the array (this is valid if, during last save, the actor returned true in AActor::ShouldLevelKeepRefIfExternal)
+								if (!ActorsSet.Contains(Actor))
+								{
+									Actors.Add(Actor);
+								}
+								bFoundActor = true;
+							}
+							// In PIE, we make sure to clear RF_Standalone flag on objects in external packages (UMetaData) 
+							// This guarantees that external packages of actors that are destroyed during the PIE session will
+							// properly get GC'ed and will allow future edits/modifications of OFPA actors.
+							if (bPackageForPIE)
+							{
+								PackageObject->ClearFlags(RF_Standalone);
+							}
+							return true;
+						}, false);
+
+					UE_CLOG(!bFoundActor, LogLevel, Error, TEXT("Failed to load Actor for External Actor Package %s"), *ActorPackageName);
+				}
 			}
 		}
 	}
diff --git a/Engine/Source/Runtime/Launch/Private/LaunchEngineLoop.cpp b/Engine/Source/Runtime/Launch/Private/LaunchEngineLoop.cpp
--- a/Engine/Source/Runtime/Launch/Private/LaunchEngineLoop.cpp
+++ b/Engine/Source/Runtime/Launch/Private/LaunchEngineLoop.cpp
@@ -133,6 +133,7 @@
 	#include "Virtualization/VirtualizationSystem.h"
 	#include "AnimationUtils.h"
 	#include "UObject/ICookInfo.h"
+	#include "CoreGlobalsInternal.h"
 
 	#if PLATFORM_WINDOWS
 		#include "Windows/AllowWindowsPlatformTypes.h"
@@ -1950,6 +1951,14 @@ int32 FEngineLoop::PreInitPreStartupScreen(const TCHAR* CmdLine)
 		GShouldEmitVerboseNamedEvents = true;
 	}
 
+#if WITH_EDITOR
+	// we are running in HybridCookedEditor if it's on the commandline
+	if (FParse::Param(CmdLine, TEXT("hybridcookededitor")))
+	{
+		PRIVATE_GIsRunningHybridCookedEditor = true;
+	}
+#endif
+
 	// Set the flag for whether we've build DebugGame instead of Development. The engine does not know this (whereas the launch module does) because it is always built in development.
 #if UE_BUILD_DEVELOPMENT && defined(UE_BUILD_DEVELOPMENT_WITH_DEBUGGAME) && UE_BUILD_DEVELOPMENT_WITH_DEBUGGAME
 	FApp::SetDebugGame(true);
diff --git a/Engine/Source/Runtime/StorageServerClient/Private/StorageServerPlatformFile.cpp b/Engine/Source/Runtime/StorageServerClient/Private/StorageServerPlatformFile.cpp
--- a/Engine/Source/Runtime/StorageServerClient/Private/StorageServerPlatformFile.cpp
+++ b/Engine/Source/Runtime/StorageServerClient/Private/StorageServerPlatformFile.cpp
@@ -11,6 +11,7 @@
 #include "Misc/App.h"
 #include "Misc/AssertionMacros.h"
 #include "Misc/CommandLine.h"
+#include "Misc/ConfigCacheIni.h"
 #include "Misc/CoreDelegates.h"
 #include "Misc/FileHelper.h"
 #include "Misc/Paths.h"
@@ -33,6 +34,7 @@
 #include "Containers/SpscQueue.h"
 #include <atomic>
 #include "Async/MappedFileHandle.h"
+#include "CoreGlobalsInternal.h"
 
 DEFINE_LOG_CATEGORY_STATIC(LogStorageServerPlatformFile, Log, All);
 
@@ -63,6 +65,17 @@ DEFINE_LOG_CATEGORY_STATIC(LogStorageServerPlatformFile, Log, All);
 
 static FDateTime GAssumedImmutableTimeStamp = FDateTime::Now();
 
+// If this is set, then StorageServer will not be used for non-assets (OpenRead, etc is used for non-assets; assets are handled via the IoDispatcher)
+static bool GPreferLocalForNonAssets = false;
+static FAutoConsoleVariableRef CVarPreferLocalForNonAssets(
+	TEXT("s.PreferLocalForNonAssets"),
+	GPreferLocalForNonAssets,
+	TEXT("Set to true to look at the local file sytem for files before loading from StorageServer"),
+	ECVF_Default
+);
+
+
+
 FStorageServerFileSystemTOC::~FStorageServerFileSystemTOC()
 {
 	FWriteScopeLock _(TocLock);
@@ -840,6 +853,31 @@ FStorageServerPlatformFile::~FStorageServerPlatformFile()
 {
 }
 
+static FString GetCookedPlatformName()
+{
+	FString PlatformName;
+	if (IsRunningHybridCookedEditor())
+	{
+		// manually look in DefaultEngine.ini, or BaseEngine.ini, for the name of the platform to use to find the ue.projectstore file
+		FConfigFile Default;
+		Default.Read(*(FPaths::ProjectConfigDir() / TEXT("DefaultEngine.ini")));
+		if (!Default.GetString(TEXT("HybridCookedEditor"), TEXT("RuntimeTargetPlatform"), PlatformName))
+		{
+			FConfigFile Base;
+			Base.Read(*(FPaths::EngineConfigDir() / TEXT("BaseEngine.ini")));
+			// we expect this to always be found
+			verify(Base.GetString(TEXT("HybridCookedEditor"), TEXT("RuntimeTargetPlatform"), PlatformName));
+		}
+		PlatformName.ReplaceInline(TEXT("{Platform}"), ANSI_TO_TCHAR(FPlatformProperties::IniPlatformName()));
+	}
+	else
+	{
+		PlatformName = FPlatformProperties::PlatformName();
+	}
+
+	return PlatformName;
+}
+
 TUniquePtr<FArchive> FStorageServerPlatformFile::TryFindProjectStoreMarkerFile(IPlatformFile* Inner) const
 {
 	if (Inner == nullptr)
@@ -852,7 +890,7 @@ TUniquePtr<FArchive> FStorageServerPlatformFile::TryFindProjectStoreMarkerFile(I
 	{
 		FString RelativeStagedPath = TEXT("../../../");
 		FString RootPath = FPaths::RootDir();
-		FString PlatformName = FPlatformProperties::PlatformName();
+		FString PlatformName = GetCookedPlatformName();
 		FString CookedOutputPath = FPaths::Combine(FPaths::ProjectDir(), TEXT("Saved"), TEXT("Cooked"), PlatformName);
 
 		PotentialProjectStorePaths.Add(RelativeStagedPath);
@@ -985,6 +1023,12 @@ bool FStorageServerPlatformFile::ShouldBeUsed(IPlatformFile* Inner, const TCHAR*
 
 bool FStorageServerPlatformFile::Initialize(IPlatformFile* Inner, const TCHAR* CmdLine)
 {
+	// hybrid cooked editor wants to load any local files when possible, instead of any potential non-assets that were imported with a zen store
+	if (IsRunningHybridCookedEditor())
+	{
+		GPreferLocalForNonAssets = 1;
+	}
+
 	LowerLevel = Inner;
 	if (HostAddrs.Num() > 0)
 	{
@@ -1270,6 +1314,16 @@ IFileHandle* FStorageServerPlatformFile::OpenRead(const TCHAR* Filename, bool bA
 {
 	TStringBuilder<1024> StorageServerFilename;
 
+	// if we prefer local files, look local before checking if the file is in the ServerToc
+	if (GPreferLocalForNonAssets && LowerLevel && IsNonServerFilenameAllowed(Filename))
+	{
+		IFileHandle* Handle = LowerLevel->OpenRead(Filename, bAllowWrite);
+		if (Handle != nullptr)
+		{
+			return Handle;
+		}
+	}	
+
 	if (MakeStorageServerPath(Filename, StorageServerFilename))
 	{
 		FIoChunkId FileChunkId;
@@ -1279,7 +1333,8 @@ IFileHandle* FStorageServerPlatformFile::OpenRead(const TCHAR* Filename, bool bA
 			return InternalOpenFile(FileChunkId, RawSize, Filename);
 		}
 	}
-	return (LowerLevel && IsNonServerFilenameAllowed(Filename)) ? LowerLevel->OpenRead(Filename, bAllowWrite) : nullptr;
+	// if we preferred Server over Local, look in local if Server failed
+	return (!GPreferLocalForNonAssets && LowerLevel && IsNonServerFilenameAllowed(Filename)) ? LowerLevel->OpenRead(Filename, bAllowWrite) : nullptr;
 }
 
 bool FStorageServerPlatformFile::IterateDirectory(const TCHAR* Directory, IPlatformFile::FDirectoryVisitor& Visitor)
@@ -1402,6 +1457,12 @@ FString FStorageServerPlatformFile::GetFilenameOnDisk(const TCHAR* Filename)
 
 bool FStorageServerPlatformFile::DeleteFile(const TCHAR* Filename)
 {
+	// if we prefer local files, we can delete them (without this, if the file is in the ServerToc, it will fail to delete)
+	if (GPreferLocalForNonAssets && LowerLevel && LowerLevel->DeleteFile(Filename))
+	{
+		return true;
+	}
+
 	TStringBuilder<1024> StorageServerFilename;
 	if (MakeStorageServerPath(Filename, StorageServerFilename) && ServerToc.FileExists(*StorageServerFilename))
 	{
@@ -1417,6 +1478,11 @@ bool FStorageServerPlatformFile::MoveFile(const TCHAR* To, const TCHAR* From)
 		return false;
 	}
 
+	if (GPreferLocalForNonAssets && LowerLevel->MoveFile(To, From))
+	{
+		return true;
+	}
+
 	TStringBuilder<1024> StorageServerTo;
 	if (MakeStorageServerPath(To, StorageServerTo) && ServerToc.FileExists(*StorageServerTo))
 	{
