Remove direct accesses of FBodyInstance::ActorHandle in preparation for its deprecation.

diff --git a/Engine/Plugins/Experimental/Buoyancy/Source/Runtime/Private/BuoyancySubsystem.cpp b/Engine/Plugins/Experimental/Buoyancy/Source/Runtime/Private/BuoyancySubsystem.cpp
--- a/Engine/Plugins/Experimental/Buoyancy/Source/Runtime/Private/BuoyancySubsystem.cpp
+++ b/Engine/Plugins/Experimental/Buoyancy/Source/Runtime/Private/BuoyancySubsystem.cpp
@@ -392,7 +392,7 @@ void UBuoyancySubsystem::UpdateSplineData()
 						FBodyInstance& BodyInstance = WaterPrimitiveComponent->BodyInstance;
 						if (BodyInstance.IsValidBodyInstance())
 						{
-							if (Chaos::FSingleParticlePhysicsProxy* WaterProxy = BodyInstance.ActorHandle)
+							if (Chaos::FSingleParticlePhysicsProxy* WaterProxy = BodyInstance.GetPhysicsActorHandle())
 							{
 								SplineData->SetData_GT(WaterProxy->GetGameThreadAPI(), WaterSplineData);
 							}
diff --git a/Engine/Plugins/Experimental/ChaosCaching/Source/ChaosCaching/Private/Chaos/Adapters/StaticMeshComponentCacheAdapter.cpp b/Engine/Plugins/Experimental/ChaosCaching/Source/ChaosCaching/Private/Chaos/Adapters/StaticMeshComponentCacheAdapter.cpp
--- a/Engine/Plugins/Experimental/ChaosCaching/Source/ChaosCaching/Private/Chaos/Adapters/StaticMeshComponentCacheAdapter.cpp
+++ b/Engine/Plugins/Experimental/ChaosCaching/Source/ChaosCaching/Private/Chaos/Adapters/StaticMeshComponentCacheAdapter.cpp
@@ -64,7 +64,7 @@ namespace Chaos
 	{
 		UStaticMeshComponent* MeshComp = CastChecked<UStaticMeshComponent>(InComponent);
 
-		FSingleParticlePhysicsProxy* PhysProxy = MeshComp->BodyInstance.ActorHandle;
+		FSingleParticlePhysicsProxy* PhysProxy = MeshComp->BodyInstance.GetPhysicsActorHandle();
 
 		RecordToCacheInternal(PhysProxy, InRootTransform, OutFrame, InTime);
 	}
@@ -106,7 +106,7 @@ namespace Chaos
 	{
 		UStaticMeshComponent* MeshComp = CastChecked<UStaticMeshComponent>(InComponent);
 
-		FSingleParticlePhysicsProxy* PhysProxy = MeshComp->BodyInstance.ActorHandle;
+		FSingleParticlePhysicsProxy* PhysProxy = MeshComp->BodyInstance.GetPhysicsActorHandle();
 
 		PlayFromCacheInternal(PhysProxy, InCache, TickRecord, OutUpdatedRigids);
 	}
@@ -179,9 +179,9 @@ namespace Chaos
 	{
 		EnsureIsInGameThreadContext();
 
-		if(Cast<UStaticMeshComponent>(InComponent) && InComponent->GetBodyInstance() && InComponent->GetBodyInstance()->ActorHandle)
+		if(Cast<UStaticMeshComponent>(InComponent) && InComponent->GetBodyInstance() && InComponent->GetBodyInstance()->GetPhysicsActorHandle())
 		{
-			FPhysInterface_Chaos::SetIsKinematic_AssumesLocked(InComponent->GetBodyInstance()->ActorHandle, true);
+			FPhysInterface_Chaos::SetIsKinematic_AssumesLocked(InComponent->GetBodyInstance()->GetPhysicsActorHandle(), true);
 		}
 
 		return true;
diff --git a/Engine/Plugins/Experimental/ChaosModularVehicle/Source/ChaosModularVehicleEngine/Private/ChaosModularVehicle/ModularVehicleBaseComponent.cpp b/Engine/Plugins/Experimental/ChaosModularVehicle/Source/ChaosModularVehicleEngine/Private/ChaosModularVehicle/ModularVehicleBaseComponent.cpp
--- a/Engine/Plugins/Experimental/ChaosModularVehicle/Source/ChaosModularVehicleEngine/Private/ChaosModularVehicle/ModularVehicleBaseComponent.cpp
+++ b/Engine/Plugins/Experimental/ChaosModularVehicle/Source/ChaosModularVehicleEngine/Private/ChaosModularVehicle/ModularVehicleBaseComponent.cpp
@@ -901,7 +901,7 @@ IPhysicsProxyBase* UModularVehicleBaseComponent::GetPhysicsProxy() const
 	// SingleParticlePhysicsProxy case
 	if (const FBodyInstance* BodyInstance = UpdatedPrimitive ? UpdatedPrimitive->GetBodyInstance() : nullptr)
 	{
-		return CachedPhysicsProxy = BodyInstance->ActorHandle;
+		return CachedPhysicsProxy = BodyInstance->GetPhysicsActorHandle();
 	}
 
 	return nullptr;
diff --git a/Engine/Plugins/Experimental/ChaosVehiclesPlugin/Source/ChaosVehicles/Private/ChaosVehicleMovementComponent.cpp b/Engine/Plugins/Experimental/ChaosVehiclesPlugin/Source/ChaosVehicles/Private/ChaosVehicleMovementComponent.cpp
--- a/Engine/Plugins/Experimental/ChaosVehiclesPlugin/Source/ChaosVehicles/Private/ChaosVehicleMovementComponent.cpp
+++ b/Engine/Plugins/Experimental/ChaosVehiclesPlugin/Source/ChaosVehicles/Private/ChaosVehicleMovementComponent.cpp
@@ -953,7 +953,7 @@ void UChaosVehicleMovementComponent::SetTargetGear(int32 GearNum, bool bImmediat
 
 		if (TargetInstance)
 		{
-			FPhysicsCommand::ExecuteWrite(TargetInstance->ActorHandle, [&](const FPhysicsActorHandle& Chassis)
+			FPhysicsCommand::ExecuteWrite(TargetInstance->GetPhysicsActorHandle(), [&](const FPhysicsActorHandle& Chassis)
 				{
 					if (VehicleSimulationPT && VehicleSimulationPT->PVehicle && VehicleSimulationPT->PVehicle->HasTransmission())
 					{
@@ -1502,9 +1502,9 @@ void UChaosVehicleMovementComponent::SetupVehicleMass()
 
 void UChaosVehicleMovementComponent::UpdateMassProperties(FBodyInstance* BodyInstance)
 {
-	if (BodyInstance && FPhysicsInterface::IsValid(BodyInstance->ActorHandle) && FPhysicsInterface::IsRigidBody(BodyInstance->ActorHandle))
+	if (BodyInstance && FPhysicsInterface::IsValid(BodyInstance->GetPhysicsActorHandle()) && FPhysicsInterface::IsRigidBody(BodyInstance->GetPhysicsActorHandle()))
 	{
-		FPhysicsCommand::ExecuteWrite(BodyInstance->ActorHandle, [&](FPhysicsActorHandle& Actor)
+		FPhysicsCommand::ExecuteWrite(BodyInstance->GetPhysicsActorHandle(), [&](FPhysicsActorHandle& Actor)
 			{
 				const float MassRatio = this->Mass > 0.0f ? this->Mass / BodyInstance->GetBodyMass() : 1.0f;
 
@@ -1753,7 +1753,7 @@ void UChaosVehicleMovementComponent::Update(float DeltaTime)
 	{
 		if (const FBodyInstance* BodyInstance = GetBodyInstance())
 		{
-			if (auto Handle = BodyInstance->ActorHandle)
+			if (auto Handle = BodyInstance->GetPhysicsActorHandle())
 			{
 				CurAsyncInput->Proxy = Handle;	// vehicles are never static
 				FChaosVehicleAsyncInput* AsyncInput = static_cast<FChaosVehicleAsyncInput*>(CurAsyncInput);
diff --git a/Engine/Plugins/Experimental/Mover/Source/Mover/Private/Backends/MoverNetworkPhysicsLiaison.cpp b/Engine/Plugins/Experimental/Mover/Source/Mover/Private/Backends/MoverNetworkPhysicsLiaison.cpp
--- a/Engine/Plugins/Experimental/Mover/Source/Mover/Private/Backends/MoverNetworkPhysicsLiaison.cpp
+++ b/Engine/Plugins/Experimental/Mover/Source/Mover/Private/Backends/MoverNetworkPhysicsLiaison.cpp
@@ -50,7 +50,7 @@ void UMoverNetworkPhysicsLiaisonComponent::SetupConstraint()
 		const UMoverComponent& MoverComp = GetMoverComponent();
 		if (FBodyInstance* BI = MoverComp.UpdatedCompAsPrimitive ? MoverComp.UpdatedCompAsPrimitive->GetBodyInstance() : nullptr)
 		{
-			if (Chaos::FSingleParticlePhysicsProxy* CharacterProxy = BI->ActorHandle)
+			if (Chaos::FSingleParticlePhysicsProxy* CharacterProxy = BI->GetPhysicsActorHandle())
 			{
 				// Create and register the constraint
 				Constraint = MakeUnique<Chaos::FCharacterGroundConstraint>();
diff --git a/Engine/Plugins/Experimental/Mover/Source/Mover/Private/Backends/MoverNetworkPhysicsLiaisonBase.cpp b/Engine/Plugins/Experimental/Mover/Source/Mover/Private/Backends/MoverNetworkPhysicsLiaisonBase.cpp
--- a/Engine/Plugins/Experimental/Mover/Source/Mover/Private/Backends/MoverNetworkPhysicsLiaisonBase.cpp
+++ b/Engine/Plugins/Experimental/Mover/Source/Mover/Private/Backends/MoverNetworkPhysicsLiaisonBase.cpp
@@ -425,7 +425,7 @@ Chaos::FUniqueIdx UMoverNetworkPhysicsLiaisonComponentBase::GetUniqueIdx() const
 	{
 		if (FBodyInstance* BI = MoverComp.UpdatedCompAsPrimitive->GetBodyInstance())
 		{
-			if (FPhysicsActorHandle ActorHandle = BI->ActorHandle)
+			if (FPhysicsActorHandle ActorHandle = BI->GetPhysicsActorHandle())
 			{
 				return ActorHandle->GetGameThreadAPI().UniqueIdx();
 			}
diff --git a/Engine/Plugins/Experimental/Water/Source/Runtime/Private/BuoyancyComponent.cpp b/Engine/Plugins/Experimental/Water/Source/Runtime/Private/BuoyancyComponent.cpp
--- a/Engine/Plugins/Experimental/Water/Source/Runtime/Private/BuoyancyComponent.cpp
+++ b/Engine/Plugins/Experimental/Water/Source/Runtime/Private/BuoyancyComponent.cpp
@@ -167,9 +167,9 @@ void UBuoyancyComponent::Update(float DeltaTime)
 	{
 		if (const FBodyInstance* BodyInstance = SimulatingComponent->GetBodyInstance())
 		{
-			if (auto Handle = BodyInstance->ActorHandle)
+			if (auto Handle = BodyInstance->GetPhysicsActorHandle())
 			{
-				CurAsyncInput->Proxy = BodyInstance->ActorHandle;
+				CurAsyncInput->Proxy = Handle;
 			}
 		}
 
diff --git a/Engine/Plugins/Experimental/Water/Source/Runtime/Private/BuoyancyManager.cpp b/Engine/Plugins/Experimental/Water/Source/Runtime/Private/BuoyancyManager.cpp
--- a/Engine/Plugins/Experimental/Water/Source/Runtime/Private/BuoyancyManager.cpp
+++ b/Engine/Plugins/Experimental/Water/Source/Runtime/Private/BuoyancyManager.cpp
@@ -62,7 +62,7 @@ void ABuoyancyManager::OnDestroyPhysics(UActorComponent* Component)
 					{
 						if (FBodyInstance* BI = SimulatingComp->GetBodyInstance())
 						{
-							if (auto ActorHandle = BI->ActorHandle)
+							if (auto ActorHandle = BI->GetPhysicsActorHandle())
 							{
 								AsyncCallback->ClearAsyncAux_External(ActorHandle->GetGameThreadAPI().UniqueIdx());
 							}
@@ -245,7 +245,7 @@ void ABuoyancyManager::InitializeAsyncAux(UBuoyancyComponent* Component)
 	{
 		if (FBodyInstance* BI = SimulatingComponent->GetBodyInstance())
 		{
-			if (auto ActorHandle = BI->ActorHandle)
+			if (auto ActorHandle = BI->GetPhysicsActorHandle())
 			{
 				AsyncCallback->CreateAsyncAux_External(ActorHandle->GetGameThreadAPI().UniqueIdx(), Component->CreateAsyncAux());
 				PhysicsInitializedSimulatingComponents.Add(SimulatingComponent);
diff --git a/Engine/Plugins/PCG/Source/PCG/Private/Data/PCGVolumeData.cpp b/Engine/Plugins/PCG/Source/PCG/Private/Data/PCGVolumeData.cpp
--- a/Engine/Plugins/PCG/Source/PCG/Private/Data/PCGVolumeData.cpp
+++ b/Engine/Plugins/PCG/Source/PCG/Private/Data/PCGVolumeData.cpp
@@ -69,7 +69,7 @@ void UPCGVolumeData::SetupVolumeBodyInstance()
 
 			// In some instances, non-collidable bodies will not be initialized, but it's not an issue for PCG so we can continue regardless.
 			// Otherwise, require that the body is not dynamic.
-			if (BodyInstance && BodySetup && (!FPhysicsInterface::IsValid(BodyInstance->ActorHandle) || !BodyInstance->IsDynamic()))
+			if (BodyInstance && BodySetup && (!FPhysicsInterface::IsValid(BodyInstance->GetPhysicsActorHandle()) || !BodyInstance->IsDynamic()))
 			{
 				ReleaseInternalBodyInstance();
 
diff --git a/Engine/Plugins/Runtime/NetworkPrediction/Source/NetworkPrediction/Private/NetworkPredictionUtil.cpp b/Engine/Plugins/Runtime/NetworkPrediction/Source/NetworkPrediction/Private/NetworkPredictionUtil.cpp
--- a/Engine/Plugins/Runtime/NetworkPrediction/Source/NetworkPrediction/Private/NetworkPredictionUtil.cpp
+++ b/Engine/Plugins/Runtime/NetworkPrediction/Source/NetworkPrediction/Private/NetworkPredictionUtil.cpp
@@ -15,7 +15,7 @@ namespace UE_NP
 			bool bValid = false;
 			if (UPrimitiveComponent* PrimComp = Cast<UPrimitiveComponent>(Component))
 			{
-				bValid = FPhysicsInterface::IsValid(PrimComp->BodyInstance.ActorHandle);
+				bValid = FPhysicsInterface::IsValid(PrimComp->BodyInstance.GetPhysicsActorHandle());
 			}
 			return bValid;
 		};
@@ -68,7 +68,7 @@ namespace UE_NP
 
 		if (PrimitiveComponent)
 		{
-			Proxy = PrimitiveComponent->BodyInstance.ActorHandle;
+			Proxy = PrimitiveComponent->BodyInstance.GetPhysicsActorHandle();
 		}
 		return Proxy;
 	}
diff --git a/Engine/Source/Runtime/Engine/Classes/PhysicsEngine/BodyInstance.h b/Engine/Source/Runtime/Engine/Classes/PhysicsEngine/BodyInstance.h
--- a/Engine/Source/Runtime/Engine/Classes/PhysicsEngine/BodyInstance.h
+++ b/Engine/Source/Runtime/Engine/Classes/PhysicsEngine/BodyInstance.h
@@ -662,10 +662,12 @@ public:
 	ENGINE_API FPhysicsActorHandle& GetPhysicsActorHandle();
 	ENGINE_API const FPhysicsActorHandle& GetPhysicsActorHandle() const;
 	ENGINE_API const FPhysicsActorHandle& GetActorReferenceWithWelding() const;
+	ENGINE_API void SetPhysicsActorHandle(FPhysicsActorHandle InHandle);
 
 	// Internal physics representation of our body instance
 	FPhysicsActorHandle ActorHandle;
 
+public:
 	FBodyInstanceAsyncPhysicsTickHandle GetBodyInstanceAsyncPhysicsTickHandle() const { return FBodyInstanceAsyncPhysicsTickHandle{ ActorHandle }; }
 
 #if USE_BODYINSTANCE_DEBUG_NAMES
diff --git a/Engine/Source/Runtime/Engine/Private/Collision/CollisionConversions.cpp b/Engine/Source/Runtime/Engine/Private/Collision/CollisionConversions.cpp
--- a/Engine/Source/Runtime/Engine/Private/Collision/CollisionConversions.cpp
+++ b/Engine/Source/Runtime/Engine/Private/Collision/CollisionConversions.cpp
@@ -151,7 +151,7 @@ static void SetHitResultFromShapeAndFaceIndex(const FPhysicsShape& Shape,  const
 		}
 
 		OwningComponent = BodyInst->OwnerComponent;
-		OutResult.PhysicsObject = BodyInst->ActorHandle ? BodyInst->ActorHandle->GetPhysicsObject() : nullptr;
+		OutResult.PhysicsObject = BodyInst->GetPhysicsActorHandle() ? BodyInst->GetPhysicsActorHandle()->GetPhysicsObject() : nullptr;
 	}
 	else
 	{
@@ -672,7 +672,7 @@ void ConvertQueryOverlap(const FPhysicsShape& Shape, const FPhysicsActor& Actor,
 			OutOverlap.Component = BodyInst->OwnerComponent; // Copying weak pointer is faster than assigning raw pointer.
 			OutOverlap.ItemIndex = OwnerComponent->bMultiBodyOverlap ? BodyInst->InstanceBodyIndex : INDEX_NONE;
 		}
-		OutOverlap.PhysicsObject = BodyInst->ActorHandle ? BodyInst->ActorHandle->GetPhysicsObject() : nullptr;
+		OutOverlap.PhysicsObject = BodyInst->GetPhysicsActorHandle() ? BodyInst->GetPhysicsActorHandle()->GetPhysicsObject() : nullptr;
 
 		if (OutOverlap.PhysicsObject != nullptr)
 		{
diff --git a/Engine/Source/Runtime/Engine/Private/Components/CharacterMovementComponent.cpp b/Engine/Source/Runtime/Engine/Private/Components/CharacterMovementComponent.cpp
--- a/Engine/Source/Runtime/Engine/Private/Components/CharacterMovementComponent.cpp
+++ b/Engine/Source/Runtime/Engine/Private/Components/CharacterMovementComponent.cpp
@@ -13363,7 +13363,7 @@ void UCharacterMovementComponent::FillAsyncInput(const FVector& InputVector, FCh
 			UpdatedComponentInput->bForceGatherOverlaps = !FUpdatedComponentAsyncInput::ShouldCheckOverlapFlagToQueueOverlaps(*PrimitiveComponent);
 			UpdatedComponentInput->bGatherOverlaps = PrimitiveComponent->GetGenerateOverlapEvents() || UpdatedComponentInput->bForceGatherOverlaps;
 			UpdatedComponentInput->UpdatedComponent = PrimitiveComponent;
-			UpdatedComponentInput->PhysicsHandle = PrimitiveComponent->BodyInstance.ActorHandle;
+			UpdatedComponentInput->PhysicsHandle = PrimitiveComponent->BodyInstance.GetPhysicsActorHandle();
 			ensure(UpdatedComponentInput->PhysicsHandle);
 
 			// TODO Double check we initialize all these SQ params correctly and use them in the right spots.
diff --git a/Engine/Source/Runtime/Engine/Private/InstancedStaticMesh.cpp b/Engine/Source/Runtime/Engine/Private/InstancedStaticMesh.cpp
--- a/Engine/Source/Runtime/Engine/Private/InstancedStaticMesh.cpp
+++ b/Engine/Source/Runtime/Engine/Private/InstancedStaticMesh.cpp
@@ -230,7 +230,7 @@ class FISMExecHelper : public FSelfRegisteringExec
 						}
 						++NumValidInstanceBodies;
 						
-						FPhysicsCommand::ExecuteRead(Body->ActorHandle, [&Body, &Shapes, &NumInstanceShapes](const FPhysicsActorHandle& Actor)
+						FPhysicsCommand::ExecuteRead(Body->GetPhysicsActorHandle(), [&Body, &Shapes, &NumInstanceShapes](const FPhysicsActorHandle& Actor)
 						{
 							Shapes.Reset();
 							Body->GetAllShapes_AssumesLocked(Shapes);
@@ -2788,11 +2788,11 @@ void UInstancedStaticMeshComponent::OnAsyncDestroyPhysicsState()
 
 Chaos::FPhysicsObject* UInstancedStaticMeshComponent::GetPhysicsObjectById(Chaos::FPhysicsObjectId Id) const
 {
-	if (!InstanceBodies.IsValidIndex(Id) || !InstanceBodies[Id] || !InstanceBodies[Id]->ActorHandle)
+	if (!InstanceBodies.IsValidIndex(Id) || !InstanceBodies[Id] || !InstanceBodies[Id]->GetPhysicsActorHandle())
 	{
 		return nullptr;
 	}
-	return InstanceBodies[Id]->ActorHandle->GetPhysicsObject();
+	return InstanceBodies[Id]->GetPhysicsActorHandle()->GetPhysicsObject();
 }
 
 TArray<Chaos::FPhysicsObject*> UInstancedStaticMeshComponent::GetAllPhysicsObjects() const
diff --git a/Engine/Source/Runtime/Engine/Private/PhysicsEngine/BodyInstance.cpp b/Engine/Source/Runtime/Engine/Private/PhysicsEngine/BodyInstance.cpp
--- a/Engine/Source/Runtime/Engine/Private/PhysicsEngine/BodyInstance.cpp
+++ b/Engine/Source/Runtime/Engine/Private/PhysicsEngine/BodyInstance.cpp
@@ -405,7 +405,7 @@ FBodyInstance::~FBodyInstance() = default;
 
 const FPhysicsActorHandle& FBodyInstance::GetActorReferenceWithWelding() const
 {
-	return WeldParent ? WeldParent->ActorHandle : ActorHandle;
+	return WeldParent ? WeldParent->GetPhysicsActorHandle() : GetPhysicsActorHandle();
 }
 
 FArchive& operator<<(FArchive& Ar,FBodyInstance& BodyInst)
@@ -473,9 +473,9 @@ const TMap<FPhysicsShapeHandle, FBodyInstance::FWeldInfo>* FBodyInstance::GetCur
 
 int32 FBodyInstance::GetAllShapes_AssumesLocked(TArray<FPhysicsShapeHandle>& OutShapes) const
 {
-	if(ActorHandle)
+	if(GetPhysicsActorHandle())
 	{
-		return FPhysicsInterface::GetAllShapes_AssumedLocked(ActorHandle, OutShapes);
+		return FPhysicsInterface::GetAllShapes_AssumedLocked(GetPhysicsActorHandle(), OutShapes);
 	}
 
 	return 0;
@@ -485,7 +485,7 @@ void FBodyInstance::UpdateTriMeshVertices(const TArray<FVector> & NewPositions)
 {
 	if (BodySetup.IsValid())
 	{
-		FPhysicsCommand::ExecuteWrite(ActorHandle, [&](const FPhysicsActorHandle& Actor)
+		FPhysicsCommand::ExecuteWrite(GetPhysicsActorHandle(), [&](const FPhysicsActorHandle& Actor)
 		{
 			//after updating the vertices we must call setGeometry again to update any shapes referencing the mesh
 			TArray<FPhysicsShapeHandle> Shapes;
@@ -966,7 +966,7 @@ void FBodyInstance::SetMaskFilter(FMaskFilter InMaskFilter)
 		return;
 	}
 
-	FPhysicsCommand::ExecuteWrite(ActorHandle, [&](const FPhysicsActorHandle& Actor)
+	FPhysicsCommand::ExecuteWrite(GetPhysicsActorHandle(), [&](const FPhysicsActorHandle& Actor)
 	{
 		TArray<FPhysicsShapeHandle> Shapes;
 		FPhysicsInterface::GetAllShapes_AssumedLocked(Actor, Shapes);
@@ -1014,7 +1014,7 @@ void FBodyInstance::UpdatePhysicsFilterData()
 	FPhysicsCommand::ExecuteWrite(GetActorReferenceWithWelding(), [&](const FPhysicsActorHandle& Actor)
 	{
 		TArray<FPhysicsShapeHandle> AllShapes;
-		const int32 NumSyncShapes = FPhysicsInterface::GetAllShapes_AssumedLocked(ActorHandle, AllShapes);
+		const int32 NumSyncShapes = FPhysicsInterface::GetAllShapes_AssumedLocked(GetPhysicsActorHandle(), AllShapes);
 		const int32 NumTotalShapes = AllShapes.Num();
 		// In skeletal case, collision enable/disable/movement should be overriden by mesh component
 		FBodyCollisionData BodyCollisionData;
@@ -1271,7 +1271,7 @@ template<typename InAllocatorType>
 void TInitBodiesHelperBase<InAllocatorType>::CreateActor_AssumesLocked(FBodyInstance* Instance, const FTransform& Transform) const
 {
 	SCOPE_CYCLE_COUNTER(STAT_CreatePhysicsActor);
-	checkSlow(!FPhysicsInterface::IsValid(Instance->ActorHandle));
+	checkSlow(!FPhysicsInterface::IsValid(Instance->GetPhysicsActorHandle()));
 	const ECollisionEnabled::Type CollisionType = Instance->GetCollisionEnabled();
 
 	FActorCreationParams ActorParams;
@@ -1290,25 +1290,25 @@ void TInitBodiesHelperBase<InAllocatorType>::CreateActor_AssumesLocked(FBodyInst
 	{
 		ActorParams.bStatic = true;
 
-		FPhysicsInterface::CreateActor(ActorParams, Instance->ActorHandle, PrimitiveComp);
+		FPhysicsInterface::CreateActor(ActorParams, Instance->GetPhysicsActorHandle(), PrimitiveComp);
 	}
 	else
 	{
-		FPhysicsInterface::CreateActor(ActorParams, Instance->ActorHandle, PrimitiveComp);
-		FPhysicsInterface::SetCcdEnabled_AssumesLocked(Instance->ActorHandle, Instance->bUseCCD);
-		FPhysicsInterface::SetMACDEnabled_AssumesLocked(Instance->ActorHandle, Instance->IsUsingMACD());
-		FPhysicsInterface::SetIsKinematic_AssumesLocked(Instance->ActorHandle, !Instance->ShouldInstanceSimulatingPhysics());
+		FPhysicsInterface::CreateActor(ActorParams, Instance->GetPhysicsActorHandle(), PrimitiveComp);
+		FPhysicsInterface::SetCcdEnabled_AssumesLocked(Instance->GetPhysicsActorHandle(), Instance->bUseCCD);
+		FPhysicsInterface::SetMACDEnabled_AssumesLocked(Instance->GetPhysicsActorHandle(), Instance->IsUsingMACD());
+		FPhysicsInterface::SetIsKinematic_AssumesLocked(Instance->GetPhysicsActorHandle(), !Instance->ShouldInstanceSimulatingPhysics());
 
-		FPhysicsInterface::SetMaxLinearVelocity_AssumesLocked(Instance->ActorHandle, TNumericLimits<float>::Max());
-		FPhysicsInterface::SetSmoothEdgeCollisionsEnabled_AssumesLocked(Instance->ActorHandle, Instance->bSmoothEdgeCollisions);
-		FPhysicsInterface::SetInertiaConditioningEnabled_AssumesLocked(Instance->ActorHandle, Instance->IsInertiaConditioningEnabled());
-		FPhysicsInterface::SetGyroscopicTorqueEnabled_AssumesLocked(Instance->ActorHandle, Instance->bGyroscopicTorqueEnabled);
-		FPhysicsInterface::SetPositionSolverIterationCount_AssumesLocked(Instance->ActorHandle, Instance->PositionSolverIterationCount);
-		FPhysicsInterface::SetVelocitySolverIterationCount_AssumesLocked(Instance->ActorHandle, Instance->VelocitySolverIterationCount);
-		FPhysicsInterface::SetProjectionSolverIterationCount_AssumesLocked(Instance->ActorHandle, Instance->ProjectionSolverIterationCount);
+		FPhysicsInterface::SetMaxLinearVelocity_AssumesLocked(Instance->GetPhysicsActorHandle(), TNumericLimits<float>::Max());
+		FPhysicsInterface::SetSmoothEdgeCollisionsEnabled_AssumesLocked(Instance->GetPhysicsActorHandle(), Instance->bSmoothEdgeCollisions);
+		FPhysicsInterface::SetInertiaConditioningEnabled_AssumesLocked(Instance->GetPhysicsActorHandle(), Instance->IsInertiaConditioningEnabled());
+		FPhysicsInterface::SetGyroscopicTorqueEnabled_AssumesLocked(Instance->GetPhysicsActorHandle(), Instance->bGyroscopicTorqueEnabled);
+		FPhysicsInterface::SetPositionSolverIterationCount_AssumesLocked(Instance->GetPhysicsActorHandle(), Instance->PositionSolverIterationCount);
+		FPhysicsInterface::SetVelocitySolverIterationCount_AssumesLocked(Instance->GetPhysicsActorHandle(), Instance->VelocitySolverIterationCount);
+		FPhysicsInterface::SetProjectionSolverIterationCount_AssumesLocked(Instance->GetPhysicsActorHandle(), Instance->ProjectionSolverIterationCount);
 
 		// Set sleep event notification
-		FPhysicsInterface::SetSendsSleepNotifies_AssumesLocked(Instance->ActorHandle, Instance->bGenerateWakeEvents);
+		FPhysicsInterface::SetSendsSleepNotifies_AssumesLocked(Instance->GetPhysicsActorHandle(), Instance->bGenerateWakeEvents);
 	}
 }
 
@@ -1331,9 +1331,9 @@ bool TInitBodiesHelperBase<InAllocatorType>::CreateShapes_AssumesLocked(FBodyIns
 	// #PHYS2 Call interface AddGeometry
 	BodySetup->AddShapesToRigidActor_AssumesLocked(Instance, Instance->Scale3D, SimplePhysMat, ComplexPhysMats, ComplexPhysMatMasks, BodyCollisionData, FTransform::Identity);
 
-	FPhysicsInterface::SetIgnoreAnalyticCollisions_AssumesLocked(Instance->ActorHandle, CVarIgnoreAnalyticCollisionsOverride.GetValueOnAnyThread() ? true : Instance->bIgnoreAnalyticCollisions);
+	FPhysicsInterface::SetIgnoreAnalyticCollisions_AssumesLocked(Instance->GetPhysicsActorHandle(), CVarIgnoreAnalyticCollisionsOverride.GetValueOnAnyThread() ? true : Instance->bIgnoreAnalyticCollisions);
 
-	const int32 NumShapes = FPhysicsInterface::GetNumShapes(Instance->ActorHandle);
+	const int32 NumShapes = FPhysicsInterface::GetNumShapes(Instance->GetPhysicsActorHandle());
 	bInitFail |= NumShapes == 0;
 
 	return bInitFail;
@@ -1538,7 +1538,7 @@ bool TInitBodiesHelperBase<InAllocatorType>::CreateShapesAndActors()
 				UE_LOG(LogPhysics, Log, TEXT("Init Instance %d of Primitive Component %s failed. Does it have collision data available?"), BodyIdx, *PrimitiveComp->GetReadableName());
 			}
 
-			FPhysicsInterface::ReleaseActor(Instance->ActorHandle, PhysScene);
+			FPhysicsInterface::ReleaseActor(Instance->GetPhysicsActorHandle(), PhysScene);
 
 			Instance->OwnerComponent = nullptr;
 			Instance->SourceObject = nullptr;
@@ -1548,10 +1548,10 @@ bool TInitBodiesHelperBase<InAllocatorType>::CreateShapesAndActors()
 			continue;
 		}
 
-		FPhysicsInterface::SetActorUserData_AssumesLocked(Instance->ActorHandle, &Instance->PhysicsUserData);
+		FPhysicsInterface::SetActorUserData_AssumesLocked(Instance->GetPhysicsActorHandle(), &Instance->PhysicsUserData);
 
 #if USE_BODYINSTANCE_DEBUG_NAMES
-		Instance->ActorHandle->GetParticle_LowLevel()->SetDebugName(DebugName);
+		Instance->GetPhysicsActorHandle()->GetParticle_LowLevel()->SetDebugName(DebugName);
 #endif
 	}
 	return true;
@@ -1771,10 +1771,10 @@ FBodyInstance::FAsyncTermBodyPayload FBodyInstance::StartAsyncTermBody_GameThrea
 
 	// Prepare payload for async task before resetting the ActorHandle/OwnerComponent
 	FAsyncTermBodyPayload Payload;
-	Payload.ActorHandle = ActorHandle;
+	Payload.ActorHandle = GetPhysicsActorHandle();
 	Payload.OwnerComponent = OwnerComponent;
 	Payload.bNotifyRigidBodyCollision = bNotifyRigidBodyCollision;
-	ActorHandle = nullptr;
+	SetPhysicsActorHandle(nullptr);
 	OwnerComponent = nullptr;
 
 	CurrentSceneState = BodyInstanceSceneState::NotAdded;
@@ -1829,9 +1829,9 @@ void FBodyInstance::TermBody(bool bNeverDeferRelease)
 	{
 		if (FPhysScene_Chaos* PhysScene = PrimComp->GetWorld()->GetPhysicsScene())
 		{
-			if (FPhysicsInterface::IsValid(ActorHandle))
+			if (FPhysicsInterface::IsValid(GetPhysicsActorHandle()))
 			{
-				PhysScene->RemoveFromComponentMaps(ActorHandle);
+				PhysScene->RemoveFromComponentMaps(GetPhysicsActorHandle());
 			}
 			if (bNotifyRigidBodyCollision)
 			{
@@ -1842,7 +1842,7 @@ void FBodyInstance::TermBody(bool bNeverDeferRelease)
 
 	if (IsValidBodyInstance())
 	{
-		FPhysicsInterface::ReleaseActor(ActorHandle, GetPhysicsScene(), bNeverDeferRelease);
+		FPhysicsInterface::ReleaseActor(GetPhysicsActorHandle(), GetPhysicsScene(), bNeverDeferRelease);
 	}
 
 	// @TODO: Release spring body here
@@ -1881,7 +1881,7 @@ bool FBodyInstance::Weld(FBodyInstance* TheirBody, const FTransform& TheirTM)
 
 	FTransform RelativeTM = TheirTM.GetRelativeTransform(MyTM);
 
-	FPhysicsCommand::ExecuteWrite(ActorHandle, [&](const FPhysicsActorHandle& Actor)
+	FPhysicsCommand::ExecuteWrite(GetPhysicsActorHandle(), [&](const FPhysicsActorHandle& Actor)
 	{
 		SCOPE_CYCLE_COUNTER(STAT_UpdatePhysMats);
 
@@ -1938,7 +1938,7 @@ bool FBodyInstance::Weld(FBodyInstance* TheirBody, const FTransform& TheirTM)
 		PostShapeChange();
 
 		// remove their body from scenes (don't call TermBody because we don't want to clear things like BodySetup)
-		FPhysicsInterface::ReleaseActor(TheirBody->ActorHandle, TheirBody->GetPhysicsScene());
+		FPhysicsInterface::ReleaseActor(TheirBody->GetPhysicsActorHandle(), TheirBody->GetPhysicsScene());
 	});
 	
 	UpdateInterpolateWhenSubStepping();
@@ -1958,7 +1958,7 @@ void FBodyInstance::UnWeld(FBodyInstance* TheirBI)
 
 	bool bShapesChanged = false;
 
-	FPhysicsCommand::ExecuteWrite(ActorHandle, [&](const FPhysicsActorHandle& Actor)
+	FPhysicsCommand::ExecuteWrite(GetPhysicsActorHandle(), [&](const FPhysicsActorHandle& Actor)
 				{
 		TArray<FPhysicsShapeHandle> Shapes;
 		const int32 NumSyncShapes = GetAllShapes_AssumesLocked(Shapes);
@@ -2152,7 +2152,7 @@ bool FBodyInstance::UpdateBodyScale(const FVector& InScale3D, bool bForceUpdate)
 	//Get all shapes
 	EScaleMode::Type ScaleMode = EScaleMode::Free;
 
-	FPhysicsCommand::ExecuteWrite(ActorHandle, [&](const FPhysicsActorHandle& Actor)
+	FPhysicsCommand::ExecuteWrite(GetPhysicsActorHandle(), [&](const FPhysicsActorHandle& Actor)
 	{
 		TArray<FPhysicsShapeHandle> Shapes;
 		GetAllShapes_AssumesLocked(Shapes);
@@ -2466,8 +2466,8 @@ bool FBodyInstance::UpdateBodyScale(const FVector& InScale3D, bool bForceUpdate)
 		// Only follow through with update if all shapes succeeded.
 		if (CHAOS_ENSURE(NewGeometry.Num() == Shapes.Num()))
 		{
-			ActorHandle->GetGameThreadAPI().SetGeometry(MakeImplicitObjectPtr<Chaos::FImplicitObjectUnion>(MoveTemp(NewGeometry)));
-			FPhysicsInterface::WakeUp_AssumesLocked(ActorHandle);
+			GetPhysicsActorHandle()->GetGameThreadAPI().SetGeometry(MakeImplicitObjectPtr<Chaos::FImplicitObjectUnion>(MoveTemp(NewGeometry)));
+			FPhysicsInterface::WakeUp_AssumesLocked(GetPhysicsActorHandle());
 		}
 		else
 		{
@@ -2481,7 +2481,7 @@ bool FBodyInstance::UpdateBodyScale(const FVector& InScale3D, bool bForceUpdate)
 
 		if (FPhysScene_Chaos* Scene = GetPhysicsScene())
 		{
-			Scene->UpdateActorInAccelerationStructure(ActorHandle);
+			Scene->UpdateActorInAccelerationStructure(GetPhysicsActorHandle());
 		}
 
 		// update mass if required
@@ -2507,7 +2507,7 @@ void FBodyInstance::UpdateInstanceSimulatePhysics()
 	const bool bUseSimulate = IsInstanceSimulatingPhysics();
 	bool bInitialized = false;
 
-	FPhysicsCommand::ExecuteWrite(ActorHandle, [&](const FPhysicsActorHandle& Actor)
+	FPhysicsCommand::ExecuteWrite(GetPhysicsActorHandle(), [&](const FPhysicsActorHandle& Actor)
 	{
 		bInitialized = true;
 		FPhysicsInterface::SetIsKinematic_AssumesLocked(Actor, !bUseSimulate);
@@ -2555,7 +2555,7 @@ bool FBodyInstance::IsNonKinematic() const
 
 bool FBodyInstance::IsDynamic() const
 {
-	return FPhysicsInterface::IsDynamic(ActorHandle);
+	return FPhysicsInterface::IsDynamic(GetPhysicsActorHandle());
 }
 
 void FBodyInstance::ApplyWeldOnChildren()
@@ -2614,7 +2614,7 @@ void FBodyInstance::SetInstanceSimulatePhysics(bool bSimulate, bool bMaintainPhy
 				FMessageLog("PIE").Warning(FText::Format(LOCTEXT("SimPhysNoBody", "Trying to simulate physics on ''{0}'' but no physics body."),
 					FText::FromString(GetPathNameSafe(OwnerComponentInst))));
 			}
-			else if (FPhysicsInterface::IsStatic(ActorHandle))
+			else if (FPhysicsInterface::IsStatic(GetPhysicsActorHandle()))
 			{
 				FMessageLog("PIE").Warning(FText::Format(LOCTEXT("SimPhysStatic", "Trying to simulate physics on ''{0}'' but it is static."),
 					FText::FromString(GetPathNameSafe(OwnerComponentInst))));
@@ -2646,7 +2646,7 @@ void FBodyInstance::SetInstanceSimulatePhysics(bool bSimulate, bool bMaintainPhy
 
 bool FBodyInstance::IsValidBodyInstance() const
 {
-	return FPhysicsInterface::IsValid(ActorHandle);
+	return FPhysicsInterface::IsValid(GetPhysicsActorHandle());
 }
 
 FTransform GetUnrealWorldTransformImp_AssumesLocked(const FBodyInstance* BodyInstance, bool bWithProjection, bool bGlobalPose)
@@ -2655,7 +2655,7 @@ FTransform GetUnrealWorldTransformImp_AssumesLocked(const FBodyInstance* BodyIns
 
 	if(BodyInstance && BodyInstance->IsValidBodyInstance())
 	{
-		WorldTM = FPhysicsInterface::GetTransform_AssumesLocked(BodyInstance->ActorHandle, bGlobalPose);
+		WorldTM = FPhysicsInterface::GetTransform_AssumesLocked(BodyInstance->GetPhysicsActorHandle(), bGlobalPose);
 
 		if(bWithProjection)
 		{
@@ -2670,7 +2670,7 @@ FTransform FBodyInstance::GetUnrealWorldTransform(bool bWithProjection /* = true
 {
 	FTransform OutTransform = FTransform::Identity;
 
-	FPhysicsCommand::ExecuteRead(ActorHandle, [&](const FPhysicsActorHandle& Actor)
+	FPhysicsCommand::ExecuteRead(GetPhysicsActorHandle(), [&](const FPhysicsActorHandle& Actor)
 	{
 		OutTransform = GetUnrealWorldTransformImp_AssumesLocked(this, bWithProjection, bForceGlobalPose);
 	});
@@ -2703,7 +2703,7 @@ void FBodyInstance::SetBodyTransform(const FTransform& NewTransform, ETeleportTy
 		return;
 	}
 
-	if(FPhysicsInterface::IsValid(ActorHandle))
+	if(FPhysicsInterface::IsValid(GetPhysicsActorHandle()))
 	{
 		if(!NewTransform.IsValid())
 		{
@@ -2721,9 +2721,9 @@ void FBodyInstance::SetBodyTransform(const FTransform& NewTransform, ETeleportTy
 
 		FPhysScene* Scene = GetPhysicsScene();
 
-		if(FPhysicsInterface::IsDynamic(ActorHandle) && !bEditorWorld && Scene)
+		if(FPhysicsInterface::IsDynamic(GetPhysicsActorHandle()) && !bEditorWorld && Scene)
 		{
-			FPhysicsCommand::ExecuteWrite(ActorHandle, [&](const FPhysicsActorHandle& Actor)
+			FPhysicsCommand::ExecuteWrite(GetPhysicsActorHandle(), [&](const FPhysicsActorHandle& Actor)
 			{
 				const bool bKinematic = FPhysicsInterface::IsKinematic_AssumesLocked(Actor);
 				const bool bSimulated = FPhysicsInterface::CanSimulate_AssumesLocked(Actor);
@@ -2747,7 +2747,7 @@ void FBodyInstance::SetBodyTransform(const FTransform& NewTransform, ETeleportTy
 		}
 		else if(Scene)
 		{
-			FPhysicsCommand::ExecuteWrite(ActorHandle, [&](const FPhysicsActorHandle& Actor)
+			FPhysicsCommand::ExecuteWrite(GetPhysicsActorHandle(), [&](const FPhysicsActorHandle& Actor)
 			{
 				FPhysicsInterface::SetGlobalPose_AssumesLocked(Actor, NewTransform);
 			});
@@ -2769,7 +2769,7 @@ FVector FBodyInstance::GetUnrealWorldVelocity() const
 {
 	FVector OutVelocity = FVector::ZeroVector;
 
-	FPhysicsCommand::ExecuteRead(ActorHandle, [&](const FPhysicsActorHandle& Actor)
+	FPhysicsCommand::ExecuteRead(GetPhysicsActorHandle(), [&](const FPhysicsActorHandle& Actor)
 	{
 		OutVelocity = GetUnrealWorldVelocity_AssumesLocked();
 	});
@@ -2780,9 +2780,9 @@ FVector FBodyInstance::GetUnrealWorldVelocity() const
 FTransform FBodyInstance::GetKinematicTarget_AssumesLocked() const
 {
 	FTransform TM;
-	if (FPhysicsInterface::IsValid(ActorHandle))
+	if (FPhysicsInterface::IsValid(GetPhysicsActorHandle()))
 	{
-		TM = FPhysicsInterface::GetKinematicTarget_AssumesLocked(ActorHandle);
+		TM = FPhysicsInterface::GetKinematicTarget_AssumesLocked(GetPhysicsActorHandle());
 	}
 	return TM;
 }
@@ -2790,7 +2790,7 @@ FTransform FBodyInstance::GetKinematicTarget_AssumesLocked() const
 FTransform FBodyInstance::GetKinematicTarget() const
 {
 	FTransform TM;
-	FPhysicsCommand::ExecuteRead(ActorHandle, [&](const FPhysicsActorHandle& Actor)
+	FPhysicsCommand::ExecuteRead(GetPhysicsActorHandle(), [&](const FPhysicsActorHandle& Actor)
 		{
 			TM = FPhysicsInterface::GetKinematicTarget_AssumesLocked(Actor);
 		});
@@ -2801,9 +2801,9 @@ FTransform FBodyInstance::GetKinematicTarget() const
 FVector FBodyInstance::GetUnrealWorldVelocity_AssumesLocked() const
 {
 	FVector LinVel(EForceInit::ForceInitToZero);
-	if (FPhysicsInterface::IsValid(ActorHandle))
+	if (FPhysicsInterface::IsValid(GetPhysicsActorHandle()))
 	{
-		LinVel = FPhysicsInterface::GetLinearVelocity_AssumesLocked(ActorHandle);
+		LinVel = FPhysicsInterface::GetLinearVelocity_AssumesLocked(GetPhysicsActorHandle());
 	}
 
 	return LinVel;
@@ -2814,7 +2814,7 @@ FVector FBodyInstance::GetUnrealWorldAngularVelocityInRadians() const
 {
 	FVector OutVelocity = FVector::ZeroVector;
 
-	FPhysicsCommand::ExecuteRead(ActorHandle, [&](const FPhysicsActorHandle& Actor)
+	FPhysicsCommand::ExecuteRead(GetPhysicsActorHandle(), [&](const FPhysicsActorHandle& Actor)
 	{
 		OutVelocity = FPhysicsInterface::GetAngularVelocity_AssumesLocked(Actor);
 	});
@@ -2825,13 +2825,13 @@ FVector FBodyInstance::GetUnrealWorldAngularVelocityInRadians() const
 /** Note: returns angular velocity in radians per second. */
 FVector FBodyInstance::GetUnrealWorldAngularVelocityInRadians_AssumesLocked() const
 {
-	return FPhysicsInterface::GetAngularVelocity_AssumesLocked(ActorHandle);
+	return FPhysicsInterface::GetAngularVelocity_AssumesLocked(GetPhysicsActorHandle());
 }
 
 FVector FBodyInstance::GetUnrealWorldVelocityAtPoint(const FVector& Point) const
 {
 	FVector OutVelocity = FVector::ZeroVector;
-	FPhysicsCommand::ExecuteRead(ActorHandle, [&](const FPhysicsActorHandle& Actor)
+	FPhysicsCommand::ExecuteRead(GetPhysicsActorHandle(), [&](const FPhysicsActorHandle& Actor)
 	{
 		OutVelocity = FPhysicsInterface::GetWorldVelocityAtPoint_AssumesLocked(Actor, Point);
 	});
@@ -2842,14 +2842,14 @@ FVector FBodyInstance::GetUnrealWorldVelocityAtPoint(const FVector& Point) const
 
 FVector FBodyInstance::GetUnrealWorldVelocityAtPoint_AssumesLocked(const FVector& Point) const
 {
-	return FPhysicsInterface::GetWorldVelocityAtPoint_AssumesLocked(ActorHandle, Point);
+	return FPhysicsInterface::GetWorldVelocityAtPoint_AssumesLocked(GetPhysicsActorHandle(), Point);
 }
 
 FTransform FBodyInstance::GetMassSpaceToWorldSpace() const
 {
 	FTransform MassSpaceToWorldSpace = FTransform::Identity;
 
-	FPhysicsCommand::ExecuteRead(ActorHandle, [&](const FPhysicsActorHandle& Actor)
+	FPhysicsCommand::ExecuteRead(GetPhysicsActorHandle(), [&](const FPhysicsActorHandle& Actor)
 		{
 		MassSpaceToWorldSpace = FPhysicsInterface::GetComTransform_AssumesLocked(Actor);
 	});
@@ -2861,7 +2861,7 @@ FTransform FBodyInstance::GetMassSpaceLocal() const
 {
 	FTransform MassSpaceLocal = FTransform::Identity;
 
-	FPhysicsCommand::ExecuteRead(ActorHandle, [&](const FPhysicsActorHandle& Actor)
+	FPhysicsCommand::ExecuteRead(GetPhysicsActorHandle(), [&](const FPhysicsActorHandle& Actor)
 	{
 		MassSpaceLocal = FPhysicsInterface::GetComTransformLocal_AssumesLocked(Actor);
 	});
@@ -2872,7 +2872,7 @@ FTransform FBodyInstance::GetMassSpaceLocal() const
 void FBodyInstance::SetMassSpaceLocal(const FTransform& NewMassSpaceLocalTM)
 {
 	//TODO: UE doesn't store this so any changes to mass properties will not remember about this properly
-	FPhysicsCommand::ExecuteWrite(ActorHandle, [&](const FPhysicsActorHandle& Actor)
+	FPhysicsCommand::ExecuteWrite(GetPhysicsActorHandle(), [&](const FPhysicsActorHandle& Actor)
 	{
 		FPhysicsInterface::SetComLocalPose_AssumesLocked(Actor, NewMassSpaceLocalTM);
 	});
@@ -2882,7 +2882,7 @@ float FBodyInstance::GetBodyMass() const
 {
 	float OutMass = 0.0f;
 
-	FPhysicsCommand::ExecuteRead(ActorHandle, [&](const FPhysicsActorHandle& Actor)
+	FPhysicsCommand::ExecuteRead(GetPhysicsActorHandle(), [&](const FPhysicsActorHandle& Actor)
 	{
 		OutMass = FPhysicsInterface::GetMass_AssumesLocked(Actor);
 	});
@@ -2895,7 +2895,7 @@ FVector FBodyInstance::GetBodyInertiaTensor() const
 {
 	FVector OutTensor = FVector::ZeroVector;
 
-	FPhysicsCommand::ExecuteRead(ActorHandle, [&](const FPhysicsActorHandle& Actor)
+	FPhysicsCommand::ExecuteRead(GetPhysicsActorHandle(), [&](const FPhysicsActorHandle& Actor)
 	{
 		OutTensor = FPhysicsInterface::GetLocalInertiaTensor_AssumesLocked(Actor);
 	});
@@ -2908,7 +2908,7 @@ void FBodyInstance::SetInertiaConditioningEnabled(bool bInEnabled)
 	if (bInEnabled != bInertiaConditioning)
 	{
 		bInertiaConditioning = bInEnabled;
-		FPhysicsCommand::ExecuteWrite(ActorHandle,
+		FPhysicsCommand::ExecuteWrite(GetPhysicsActorHandle(),
 			[bInEnabled](const FPhysicsActorHandle& Actor)
 			{
 				if (FPhysicsInterface::IsRigidBody(Actor))
@@ -2923,7 +2923,7 @@ FBox FBodyInstance::GetBodyBounds() const
 {
 	FBox OutBox(EForceInit::ForceInitToZero);
 
-	FPhysicsCommand::ExecuteRead(ActorHandle, [&](const FPhysicsActorHandle& Actor)
+	FPhysicsCommand::ExecuteRead(GetPhysicsActorHandle(), [&](const FPhysicsActorHandle& Actor)
 		{
 			OutBox = FPhysicsInterface::GetBounds_AssumesLocked(Actor);
 		});
@@ -2935,7 +2935,7 @@ FBox FBodyInstance::GetBodyBoundsLocal() const
 {
 	FBox OutBox(EForceInit::ForceInitToZero);
 
-	FPhysicsCommand::ExecuteRead(ActorHandle, [&](const FPhysicsActorHandle& Actor)
+	FPhysicsCommand::ExecuteRead(GetPhysicsActorHandle(), [&](const FPhysicsActorHandle& Actor)
 		{
 			OutBox = FPhysicsInterface::GetBoundsLocal_AssumesLocked(Actor);
 		});
@@ -2958,8 +2958,8 @@ void FBodyInstance::CopyBodyInstancePropertiesFrom(const FBodyInstance* FromInst
 	check(FromInst);
 	check(FromInst->OwnerComponent.Get() == NULL);
 	check(FromInst->BodySetup.Get() == NULL);
-	check(!FPhysicsInterface::IsValid(FromInst->ActorHandle));
-	check(!FPhysicsInterface::IsValid(ActorHandle));
+	check(!FPhysicsInterface::IsValid(FromInst->GetPhysicsActorHandle()));
+	check(!FPhysicsInterface::IsValid(GetPhysicsActorHandle()));
 
 	*this = *FromInst;
 }
@@ -2983,12 +2983,12 @@ void FBodyInstance::CopyRuntimeBodyInstancePropertiesFrom(const FBodyInstance* F
 
 const FPhysScene* FBodyInstance::GetPhysicsScene() const
 {
-	return static_cast<const FPhysScene*>(FPhysicsInterface::GetCurrentScene(ActorHandle));
+	return static_cast<const FPhysScene*>(FPhysicsInterface::GetCurrentScene(GetPhysicsActorHandle()));
 }
 
 FPhysScene* FBodyInstance::GetPhysicsScene()
 {
-	return static_cast<FPhysScene*>(FPhysicsInterface::GetCurrentScene(ActorHandle));
+	return static_cast<FPhysScene*>(FPhysicsInterface::GetCurrentScene(GetPhysicsActorHandle()));
 }
 
 FPhysicsActorHandle& FBodyInstance::GetPhysicsActorHandle()
@@ -2996,6 +2996,12 @@ FPhysicsActorHandle& FBodyInstance::GetPhysicsActorHandle()
 	return ActorHandle;
 }
 
+void FBodyInstance::SetPhysicsActorHandle(FPhysicsActorHandle InHandle)
+{
+	check(ActorHandle == nullptr);
+	ActorHandle = InHandle;
+}
+
 const FPhysicsActorHandle& FBodyInstance::GetPhysicsActorHandle() const
 {
 	return ActorHandle;
@@ -3243,9 +3249,9 @@ void FBodyInstance::UpdateMassProperties()
 
 	UPhysicalMaterial* PhysMat = GetSimplePhysicalMaterial();
 
-	if (FPhysicsInterface::IsValid(ActorHandle) && FPhysicsInterface::IsRigidBody(ActorHandle))
+	if (FPhysicsInterface::IsValid(GetPhysicsActorHandle()) && FPhysicsInterface::IsRigidBody(GetPhysicsActorHandle()))
 	{
-		FPhysicsCommand::ExecuteWrite(ActorHandle, [&](FPhysicsActorHandle& Actor)
+		FPhysicsCommand::ExecuteWrite(GetPhysicsActorHandle(), [&](FPhysicsActorHandle& Actor)
 		{
 			check(FPhysicsInterface::IsValid(Actor));
 
@@ -3384,7 +3390,7 @@ void FBodyInstance::UpdateDebugRendering()
 
 void FBodyInstance::UpdateDampingProperties()
 {
-	FPhysicsCommand::ExecuteWrite(ActorHandle, [&](const FPhysicsActorHandle& Actor)
+	FPhysicsCommand::ExecuteWrite(GetPhysicsActorHandle(), [&](const FPhysicsActorHandle& Actor)
 	{
 		if(FPhysicsInterface::IsDynamic(Actor))
 	{
@@ -3398,7 +3404,7 @@ bool FBodyInstance::IsInstanceAwake() const
 {
 	bool bIsAwake = false;
 
-	FPhysicsCommand::ExecuteRead(ActorHandle, [&](const FPhysicsActorHandle& Actor)
+	FPhysicsCommand::ExecuteRead(GetPhysicsActorHandle(), [&](const FPhysicsActorHandle& Actor)
 	{
 		if(FPhysicsInterface::IsDynamic(Actor))
 	{
@@ -3411,7 +3417,7 @@ bool FBodyInstance::IsInstanceAwake() const
 
 void FBodyInstance::WakeInstance()
 {
-	FPhysicsCommand::ExecuteWrite(ActorHandle, [&](const FPhysicsActorHandle& Actor)
+	FPhysicsCommand::ExecuteWrite(GetPhysicsActorHandle(), [&](const FPhysicsActorHandle& Actor)
 	{
 		if(FPhysicsInterface::IsDynamic(Actor) && FPhysicsInterface::IsInScene(Actor) && !FPhysicsInterface::IsKinematic_AssumesLocked(Actor))
 		{
@@ -3422,7 +3428,7 @@ void FBodyInstance::WakeInstance()
 
 void FBodyInstance::PutInstanceToSleep()
 {
-	FPhysicsCommand::ExecuteWrite(ActorHandle, [&](const FPhysicsActorHandle& Actor)
+	FPhysicsCommand::ExecuteWrite(GetPhysicsActorHandle(), [&](const FPhysicsActorHandle& Actor)
 	{
 		if(FPhysicsInterface::IsDynamic(Actor) && FPhysicsInterface::IsInScene(Actor) && !FPhysicsInterface::IsKinematic_AssumesLocked(Actor))
 		{
@@ -3448,7 +3454,7 @@ float FBodyInstance::GetSleepThresholdMultiplier() const
 void FBodyInstance::SetLinearVelocity(const FVector& NewVel, bool bAddToCurrent, bool bAutoWake)
 {
 
-	FPhysicsCommand::ExecuteWrite(ActorHandle, [&](const FPhysicsActorHandle& Actor)
+	FPhysicsCommand::ExecuteWrite(GetPhysicsActorHandle(), [&](const FPhysicsActorHandle& Actor)
 	{
 		if(FPhysicsInterface::IsRigidBody(Actor))
 	{
@@ -3466,7 +3472,7 @@ void FBodyInstance::SetLinearVelocity(const FVector& NewVel, bool bAddToCurrent,
 
 void FBodyInstance::SetAngularVelocityInRadians(const FVector& NewAngVel, bool bAddToCurrent, bool bAutoWake)
 {
-	FPhysicsCommand::ExecuteWrite(ActorHandle, [&](const FPhysicsActorHandle& Actor)
+	FPhysicsCommand::ExecuteWrite(GetPhysicsActorHandle(), [&](const FPhysicsActorHandle& Actor)
 	{
 		if(FPhysicsInterface::IsRigidBody(Actor))
 	{
@@ -3491,7 +3497,7 @@ void FBodyInstance::SetMaxAngularVelocityInRadians(float NewMaxAngVel, bool bAdd
 {
 	float NewMaxInDegrees = FMath::RadiansToDegrees(NewMaxAngVel);
 
-	FPhysicsCommand::ExecuteWrite(ActorHandle, [&](const FPhysicsActorHandle& Actor)
+	FPhysicsCommand::ExecuteWrite(GetPhysicsActorHandle(), [&](const FPhysicsActorHandle& Actor)
 	{
 		if (bAddToCurrent)
 		{
@@ -3534,7 +3540,7 @@ void FBodyInstance::UpdateMaxDepenetrationVelocity()
 
 	// NOTE: FBodyInstance::MaxDepentrationVelocity now means initial depenetration velocity, 
 	// and not the general solver depentration velocity limit (which will probably be removed)
-	FPhysicsCommand::ExecuteWrite(ActorHandle, [this, UsedMaxDepenetrationVelocity](const FPhysicsActorHandle& Actor)
+	FPhysicsCommand::ExecuteWrite(GetPhysicsActorHandle(), [this, UsedMaxDepenetrationVelocity](const FPhysicsActorHandle& Actor)
 	{
 		FPhysicsInterface::SetMaxDepenetrationVelocity_AssumesLocked(Actor, UsedMaxDepenetrationVelocity);
 	});
@@ -3542,7 +3548,7 @@ void FBodyInstance::UpdateMaxDepenetrationVelocity()
 
 void FBodyInstance::AddCustomPhysics(FCalculateCustomPhysics& CalculateCustomPhysics)
 {
-	FPhysicsCommand::ExecuteWrite(ActorHandle, [&](const FPhysicsActorHandle& Actor)
+	FPhysicsCommand::ExecuteWrite(GetPhysicsActorHandle(), [&](const FPhysicsActorHandle& Actor)
 	{
 		if(!IsRigidBodyKinematic_AssumesLocked(Actor))
 		{
@@ -3568,7 +3574,7 @@ void FBodyInstance::ApplyAsyncPhysicsCommand(FAsyncPhysicsTimestamp TimeStamp, c
 	}
 	else
 	{
-		FPhysicsCommand::ExecuteWrite(ActorHandle, [&](const FPhysicsActorHandle& Actor)
+		FPhysicsCommand::ExecuteWrite(GetPhysicsActorHandle(), [&](const FPhysicsActorHandle& Actor)
 			{
 				Command();
 			});
@@ -3588,9 +3594,9 @@ void FBodyInstance::AddForce(const FVector& Force, bool bAllowSubstepping, bool
 	const bool bIsInternal = TimeStamp.IsValid();
 	ApplyAsyncPhysicsCommand(TimeStamp, bIsInternal, PlayerController, [&, Force, bAllowSubstepping, bAccelChange, bIsInternal]()
 	{
-		if (FPhysicsInterface::IsInScene(ActorHandle) && IsBodyDynamic(ActorHandle, bIsInternal))
+		if (FPhysicsInterface::IsInScene(GetPhysicsActorHandle()) && IsBodyDynamic(GetPhysicsActorHandle(), bIsInternal))
 		{
-			FPhysicsInterface::AddForce_AssumesLocked(ActorHandle, Force, bAllowSubstepping, bAccelChange, bIsInternal);
+			FPhysicsInterface::AddForce_AssumesLocked(GetPhysicsActorHandle(), Force, bAllowSubstepping, bAccelChange, bIsInternal);
 		}
 	});
 }
@@ -3600,16 +3606,16 @@ void FBodyInstance::AddForceAtPosition(const FVector& Force, const FVector& Posi
 	const bool bIsInternal = TimeStamp.IsValid();
 	ApplyAsyncPhysicsCommand(TimeStamp, bIsInternal, PlayerController, [&, Force, Position, bAllowSubstepping, bIsLocalForce, bIsInternal]()
 	{
-		if (FPhysicsInterface::IsInScene(ActorHandle) && IsBodyDynamic(ActorHandle, bIsInternal))
+		if (FPhysicsInterface::IsInScene(GetPhysicsActorHandle()) && IsBodyDynamic(GetPhysicsActorHandle(), bIsInternal))
 		{
-			FPhysicsInterface::AddForceAtPosition_AssumesLocked(ActorHandle, Force, Position, bAllowSubstepping, bIsLocalForce, bIsInternal);
+			FPhysicsInterface::AddForceAtPosition_AssumesLocked(GetPhysicsActorHandle(), Force, Position, bAllowSubstepping, bIsLocalForce, bIsInternal);
 		}
 	});
 }
 
 void FBodyInstance::ClearForces(bool bAllowSubstepping)
 {
-	FPhysicsCommand::ExecuteWrite(ActorHandle, [&](const FPhysicsActorHandle& Actor)
+	FPhysicsCommand::ExecuteWrite(GetPhysicsActorHandle(), [&](const FPhysicsActorHandle& Actor)
 	{
 		if(!IsRigidBodyKinematic_AssumesLocked(Actor))
 		{
@@ -3635,7 +3641,7 @@ void FBodyInstance::UpdateOneWayInteraction()
 {
 	const bool bCurrentOneWayInteraction = bOneWayInteraction;
 
-	FPhysicsCommand::ExecuteWrite(ActorHandle, [bCurrentOneWayInteraction](const FPhysicsActorHandle& Actor)
+	FPhysicsCommand::ExecuteWrite(GetPhysicsActorHandle(), [bCurrentOneWayInteraction](const FPhysicsActorHandle& Actor)
 	{
 		if (FPhysicsInterface::IsRigidBody(Actor) && !IsRigidBodyKinematic_AssumesLocked(Actor))
 		{
@@ -3649,16 +3655,16 @@ void FBodyInstance::AddTorqueInRadians(const FVector& Torque, bool bAllowSubstep
 	const bool bIsInternal = TimeStamp.IsValid();
 	ApplyAsyncPhysicsCommand(TimeStamp, bIsInternal, PlayerController, [&, Torque, bAllowSubstepping, bAccelChange, bIsInternal]()
 	{
-		if (FPhysicsInterface::IsInScene(ActorHandle) && IsBodyDynamic(ActorHandle, bIsInternal))
+		if (FPhysicsInterface::IsInScene(GetPhysicsActorHandle()) && IsBodyDynamic(GetPhysicsActorHandle(), bIsInternal))
 		{
-			FPhysicsInterface::AddTorque_AssumesLocked(ActorHandle, Torque, bAllowSubstepping, bAccelChange, bIsInternal);
+			FPhysicsInterface::AddTorque_AssumesLocked(GetPhysicsActorHandle(), Torque, bAllowSubstepping, bAccelChange, bIsInternal);
 		}
 	});
 }
 
 void FBodyInstance::ClearTorques(bool bAllowSubstepping)
 {
-	FPhysicsCommand::ExecuteWrite(ActorHandle, [&](const FPhysicsActorHandle& Actor)
+	FPhysicsCommand::ExecuteWrite(GetPhysicsActorHandle(), [&](const FPhysicsActorHandle& Actor)
 	{
 		if(!IsRigidBodyKinematic_AssumesLocked(Actor))
 		{
@@ -3675,15 +3681,15 @@ void FBodyInstance::AddAngularImpulseInRadians(const FVector& AngularImpulse, bo
 	const bool bIsInternal = TimeStamp.IsValid();
 	ApplyAsyncPhysicsCommand(TimeStamp, bIsInternal, PlayerController, [&, AngularImpulse, bVelChange, bIsInternal]()
 	{
-		if (FPhysicsInterface::IsInScene(ActorHandle) && IsBodyDynamic(ActorHandle, bIsInternal))
+		if (FPhysicsInterface::IsInScene(GetPhysicsActorHandle()) && IsBodyDynamic(GetPhysicsActorHandle(), bIsInternal))
 		{
 			if (bVelChange)
 			{
-				FPhysicsInterface::AddAngularVelocityInRadians_AssumesLocked(ActorHandle, AngularImpulse, bIsInternal);
+				FPhysicsInterface::AddAngularVelocityInRadians_AssumesLocked(GetPhysicsActorHandle(), AngularImpulse, bIsInternal);
 			}
 			else
 			{
-				FPhysicsInterface::AddAngularImpulseInRadians_AssumesLocked(ActorHandle, AngularImpulse, bIsInternal);
+				FPhysicsInterface::AddAngularImpulseInRadians_AssumesLocked(GetPhysicsActorHandle(), AngularImpulse, bIsInternal);
 			}
 		}
 	});
@@ -3694,15 +3700,15 @@ void FBodyInstance::AddImpulse(const FVector& Impulse, bool bVelChange, const FA
 	const bool bIsInternal = TimeStamp.IsValid();
 	ApplyAsyncPhysicsCommand(TimeStamp, bIsInternal, PlayerController, [&, Impulse, bVelChange, bIsInternal]()
 	{		
-		if (FPhysicsInterface::IsInScene(ActorHandle) && IsBodyDynamic(ActorHandle, bIsInternal))
+		if (FPhysicsInterface::IsInScene(GetPhysicsActorHandle()) && IsBodyDynamic(GetPhysicsActorHandle(), bIsInternal))
 		{
 			if (bVelChange)
 			{
-				FPhysicsInterface::AddVelocity_AssumesLocked(ActorHandle, Impulse, bIsInternal);
+				FPhysicsInterface::AddVelocity_AssumesLocked(GetPhysicsActorHandle(), Impulse, bIsInternal);
 			}
 			else
 			{
-				FPhysicsInterface::AddImpulse_AssumesLocked(ActorHandle, Impulse, bIsInternal);
+				FPhysicsInterface::AddImpulse_AssumesLocked(GetPhysicsActorHandle(), Impulse, bIsInternal);
 			}
 		}
 	});
@@ -3713,9 +3719,9 @@ void FBodyInstance::AddImpulseAtPosition(const FVector& Impulse, const FVector&
 	const bool bIsInternal = TimeStamp.IsValid();
 	ApplyAsyncPhysicsCommand(TimeStamp, bIsInternal, PlayerController, [&, Impulse, Position, bIsInternal]()
 	{
-		if (FPhysicsInterface::IsInScene(ActorHandle) && IsBodyDynamic(ActorHandle, bIsInternal))
+		if (FPhysicsInterface::IsInScene(GetPhysicsActorHandle()) && IsBodyDynamic(GetPhysicsActorHandle(), bIsInternal))
 		{
-			FPhysicsInterface::AddImpulseAtLocation_AssumesLocked(ActorHandle, Impulse, Position, bIsInternal);
+			FPhysicsInterface::AddImpulseAtLocation_AssumesLocked(GetPhysicsActorHandle(), Impulse, Position, bIsInternal);
 		}
 	});
 }
@@ -3725,9 +3731,9 @@ void FBodyInstance::AddVelocityChangeImpulseAtLocation(const FVector& Impulse, c
 	const bool bIsInternal = TimeStamp.IsValid();
 	ApplyAsyncPhysicsCommand(TimeStamp, bIsInternal, PlayerController, [&, Impulse, Position, bIsInternal]()
 	{
-		if (FPhysicsInterface::IsInScene(ActorHandle) && IsBodyDynamic(ActorHandle, bIsInternal))
+		if (FPhysicsInterface::IsInScene(GetPhysicsActorHandle()) && IsBodyDynamic(GetPhysicsActorHandle(), bIsInternal))
 		{
-			FPhysicsInterface::AddVelocityChangeImpulseAtLocation_AssumesLocked(ActorHandle, Impulse, Position, bIsInternal);
+			FPhysicsInterface::AddVelocityChangeImpulseAtLocation_AssumesLocked(GetPhysicsActorHandle(), Impulse, Position, bIsInternal);
 		}
 	});
 }
@@ -3772,7 +3778,7 @@ void FBodyInstance::SetEnableGravity(bool bInGravityEnabled)
 		bEnableGravity = bInGravityEnabled;
 
 		{
-			FPhysicsCommand::ExecuteWrite(ActorHandle, [&](const FPhysicsActorHandle& Actor)
+			FPhysicsCommand::ExecuteWrite(GetPhysicsActorHandle(), [&](const FPhysicsActorHandle& Actor)
 			{
 				if(FPhysicsInterface::IsRigidBody(Actor))
 				{
@@ -3795,7 +3801,7 @@ void FBodyInstance::SetGravityGroupIndex(int32 NewGravityGroupIndex)
 		GravityGroupIndex = NewGravityGroupIndex;
 
 		{
-			FPhysicsCommand::ExecuteWrite(ActorHandle, [&](const FPhysicsActorHandle& Actor)
+			FPhysicsCommand::ExecuteWrite(GetPhysicsActorHandle(), [&](const FPhysicsActorHandle& Actor)
 				{
 					if (FPhysicsInterface::IsRigidBody(Actor))
 					{
@@ -3813,7 +3819,7 @@ void FBodyInstance::SetUpdateKinematicFromSimulation(bool bInUpdateKinematicFrom
 		bUpdateKinematicFromSimulation = bInUpdateKinematicFromSimulation;
 
 		{
-			FPhysicsCommand::ExecuteWrite(ActorHandle, [&](const FPhysicsActorHandle& Actor)
+			FPhysicsCommand::ExecuteWrite(GetPhysicsActorHandle(), [&](const FPhysicsActorHandle& Actor)
 				{
 					if (FPhysicsInterface::IsRigidBody(Actor))
 					{
@@ -3831,7 +3837,7 @@ void FBodyInstance::SetGyroscopicTorqueEnabled(bool bInGyroscopicTorqueEnabled)
 		bGyroscopicTorqueEnabled = bInGyroscopicTorqueEnabled;
 
 		{
-			FPhysicsCommand::ExecuteWrite(ActorHandle, [&](const FPhysicsActorHandle& Actor)
+			FPhysicsCommand::ExecuteWrite(GetPhysicsActorHandle(), [&](const FPhysicsActorHandle& Actor)
 				{
 					if (FPhysicsInterface::IsRigidBody(Actor))
 					{
@@ -3858,7 +3864,7 @@ void FBodyInstance::SetSmoothEdgeCollisionsEnabled(bool bNewSmoothEdgeCollisions
 	{
 		bSmoothEdgeCollisions = bNewSmoothEdgeCollisions;
 
-		FPhysicsCommand::ExecuteWrite(ActorHandle, 
+		FPhysicsCommand::ExecuteWrite(GetPhysicsActorHandle(),
 			[bNewSmoothEdgeCollisions](const FPhysicsActorHandle& Actor)
 			{
 				if (FPhysicsInterface::IsRigidBody(Actor))
@@ -3875,7 +3881,7 @@ void FBodyInstance::SetUseCCD(bool bInUseCCD)
 	{
 		bUseCCD = bInUseCCD;
 		// Need to set body flag
-		FPhysicsCommand::ExecuteWrite(ActorHandle, [&](const FPhysicsActorHandle& Actor)
+		FPhysicsCommand::ExecuteWrite(GetPhysicsActorHandle(), [&](const FPhysicsActorHandle& Actor)
 		{
 			if (FPhysicsInterface::IsValid(Actor) && FPhysicsInterface::IsRigidBody(Actor))
 			{
@@ -3893,7 +3899,7 @@ void FBodyInstance::SetUseMACD(bool bInUseMACD)
 	{
 		bUseMACD = bInUseMACD;
 
-		FPhysicsCommand::ExecuteWrite(ActorHandle, [this, bInUseMACD](const FPhysicsActorHandle& Actor)
+		FPhysicsCommand::ExecuteWrite(GetPhysicsActorHandle(), [this, bInUseMACD](const FPhysicsActorHandle& Actor)
 			{
 				if (FPhysicsInterface::IsValid(Actor) && FPhysicsInterface::IsRigidBody(Actor))
 				{
@@ -3909,7 +3915,7 @@ void FBodyInstance::SetPositionSolverIterationCount(uint8 PositionSolverIteratio
 	{
 		PositionSolverIterationCount = PositionSolverIterationCountIn;
 
-		FPhysicsCommand::ExecuteWrite(ActorHandle, [this, PositionSolverIterationCountIn](const FPhysicsActorHandle& Actor)
+		FPhysicsCommand::ExecuteWrite(GetPhysicsActorHandle(), [this, PositionSolverIterationCountIn](const FPhysicsActorHandle& Actor)
 			{
 				if (FPhysicsInterface::IsValid(Actor) && FPhysicsInterface::IsRigidBody(Actor))
 				{
@@ -3925,7 +3931,7 @@ void FBodyInstance::SetVelocitySolverIterationCount(uint8 VelocitySolverIteratio
 	{
 		VelocitySolverIterationCount = VelocitySolverIterationCountIn;
 
-		FPhysicsCommand::ExecuteWrite(ActorHandle, [this, VelocitySolverIterationCountIn](const FPhysicsActorHandle& Actor)
+		FPhysicsCommand::ExecuteWrite(GetPhysicsActorHandle(), [this, VelocitySolverIterationCountIn](const FPhysicsActorHandle& Actor)
 			{
 				if (FPhysicsInterface::IsValid(Actor) && FPhysicsInterface::IsRigidBody(Actor))
 				{
@@ -3941,7 +3947,7 @@ void FBodyInstance::SetProjectionSolverIterationCount(uint8 ProjectionSolverIter
 	{
 		ProjectionSolverIterationCount = ProjectionSolverIterationCountIn;
 
-		FPhysicsCommand::ExecuteWrite(ActorHandle, [this, ProjectionSolverIterationCountIn](const FPhysicsActorHandle& Actor)
+		FPhysicsCommand::ExecuteWrite(GetPhysicsActorHandle(), [this, ProjectionSolverIterationCountIn](const FPhysicsActorHandle& Actor)
 			{
 				if (FPhysicsInterface::IsValid(Actor) && FPhysicsInterface::IsRigidBody(Actor))
 				{
@@ -3953,7 +3959,7 @@ void FBodyInstance::SetProjectionSolverIterationCount(uint8 ProjectionSolverIter
 
 void FBodyInstance::SetPhysicsDisabled(bool bSetDisabled)
 {
-	FPhysicsCommand::ExecuteWrite(ActorHandle, [&](const FPhysicsActorHandle& Actor)
+	FPhysicsCommand::ExecuteWrite(GetPhysicsActorHandle(), [&](const FPhysicsActorHandle& Actor)
 	{
 		if (FPhysicsInterface::IsValid(Actor) && FPhysicsInterface::IsRigidBody(Actor))
 		{
@@ -3965,7 +3971,7 @@ void FBodyInstance::SetPhysicsDisabled(bool bSetDisabled)
 bool FBodyInstance::IsPhysicsDisabled() const
 {
 	bool bIsDisabled = false;
-	FPhysicsCommand::ExecuteRead(ActorHandle, [&](const FPhysicsActorHandle& Actor)
+	FPhysicsCommand::ExecuteRead(GetPhysicsActorHandle(), [&](const FPhysicsActorHandle& Actor)
 	{
 		bIsDisabled = FPhysicsInterface::IsDisabled(Actor);
 	});
@@ -3986,7 +3992,7 @@ EPhysicsReplicationMode FBodyInstance::GetPhysicsReplicationMode() const
 
 void FBodyInstance::AddRadialImpulseToBody(const FVector& Origin, float Radius, float Strength, uint8 Falloff, bool bVelChange)
 {
-	FPhysicsCommand::ExecuteWrite(ActorHandle, [&](const FPhysicsActorHandle& Actor)
+	FPhysicsCommand::ExecuteWrite(GetPhysicsActorHandle(), [&](const FPhysicsActorHandle& Actor)
 	{
 		if(FPhysicsInterface::IsRigidBody(Actor) && FPhysicsInterface::IsInScene(Actor) && !IsRigidBodyKinematic_AssumesLocked(Actor))
 		{
@@ -3997,7 +4003,7 @@ void FBodyInstance::AddRadialImpulseToBody(const FVector& Origin, float Radius,
 
 void FBodyInstance::AddRadialForceToBody(const FVector& Origin, float Radius, float Strength, uint8 Falloff, bool bAccelChange, bool bAllowSubstepping)
 {
-	FPhysicsCommand::ExecuteWrite(ActorHandle, [&](const FPhysicsActorHandle& Actor)
+	FPhysicsCommand::ExecuteWrite(GetPhysicsActorHandle(), [&](const FPhysicsActorHandle& Actor)
 	{
 		if(FPhysicsInterface::IsRigidBody(Actor) && FPhysicsInterface::IsInScene(Actor) && !IsRigidBodyKinematic_AssumesLocked(Actor))
 		{
@@ -4065,7 +4071,7 @@ bool FBodyInstance::OverlapTestForBodiesImpl(const FVector& Pos, const FQuat& Ro
 {
 	bool bHaveOverlap = false;
 
-	FPhysicsCommand::ExecuteRead(ActorHandle, [&](const FPhysicsActorHandle& Actor)
+	FPhysicsCommand::ExecuteRead(GetPhysicsActorHandle(), [&](const FPhysicsActorHandle& Actor)
 	{
 		// calculate the test global pose of the rigid body
 		FTransform PTestGlobalPose = FTransform(Rot, Pos);
@@ -4108,7 +4114,7 @@ bool FBodyInstance::OverlapTest(const FVector& Position, const FQuat& Rotation,
 
 	bool bHasOverlap = false;
 
-	FPhysicsCommand::ExecuteRead(ActorHandle, [&](const FPhysicsActorHandle& Actor)
+	FPhysicsCommand::ExecuteRead(GetPhysicsActorHandle(), [&](const FPhysicsActorHandle& Actor)
 	{
 		FTransform GeomTransform(Rotation, Position);
 
@@ -4188,7 +4194,7 @@ bool FBodyInstance::OverlapMulti(TArray<struct FOverlapResult>& InOutOverlaps, c
 
 	const FBodyInstance* TargetInstance = WeldParent ? WeldParent : this;
 
-	FPhysicsCommand::ExecuteRead(TargetInstance->ActorHandle, [&](const FPhysicsActorHandle& Actor)
+	FPhysicsCommand::ExecuteRead(TargetInstance->GetPhysicsActorHandle(), [&](const FPhysicsActorHandle& Actor)
 	{
 		if(FPhysicsInterface::IsValid(Actor))
 		{
@@ -4302,7 +4308,7 @@ void FBodyInstance::LoadProfileData(bool bVerifyProfile)
 
 void FBodyInstance::GetBodyInstanceResourceSizeEx(FResourceSizeEx& CumulativeResourceSize) const
 {
-	CumulativeResourceSize.AddDedicatedSystemMemoryBytes(FPhysicsInterface::GetResourceSizeEx(ActorHandle));
+	CumulativeResourceSize.AddDedicatedSystemMemoryBytes(FPhysicsInterface::GetResourceSizeEx(GetPhysicsActorHandle()));
 }
 
 void FBodyInstance::FixupData(class UObject* Loader)
@@ -4427,7 +4433,7 @@ void FBodyInstance::InitDynamicProperties_AssumesLocked()
 		return;
 	}
 	
-	if(FPhysicsInterface::IsDynamic(ActorHandle))
+	if(FPhysicsInterface::IsDynamic(GetPhysicsActorHandle()))
 	{
 		//A non simulated body may become simulated at runtime, so we need to compute its mass.
 		//However, this is not supported for complexAsSimple since a trimesh cannot itself be simulated, it can only be used for collision of other simple shapes.
@@ -4461,40 +4467,40 @@ void FBodyInstance::InitDynamicProperties_AssumesLocked()
 
 		if (ShouldInstanceSimulatingPhysics())
 		{
-			FPhysicsInterface::SetLinearVelocity_AssumesLocked(ActorHandle, InitialLinVel);
+			FPhysicsInterface::SetLinearVelocity_AssumesLocked(GetPhysicsActorHandle(), InitialLinVel);
 		}
 
-		FPhysicsInterface::SetSleepThresholdMultiplier_AssumesLocked(ActorHandle, GetSleepThresholdMultiplier());
+		FPhysicsInterface::SetSleepThresholdMultiplier_AssumesLocked(GetPhysicsActorHandle(), GetSleepThresholdMultiplier());
 
 		// @todo: implement sleep energy threshold
-		float SleepEnergyThresh = FPhysicsInterface::GetSleepEnergyThreshold_AssumesLocked(ActorHandle);
+		float SleepEnergyThresh = FPhysicsInterface::GetSleepEnergyThreshold_AssumesLocked(GetPhysicsActorHandle());
 		SleepEnergyThresh *= GetSleepThresholdMultiplier();
-		FPhysicsInterface::SetSleepEnergyThreshold_AssumesLocked(ActorHandle, SleepEnergyThresh);
+		FPhysicsInterface::SetSleepEnergyThreshold_AssumesLocked(GetPhysicsActorHandle(), SleepEnergyThresh);
 
 		// PhysX specific dynamic parameters, not generically exposed to physics interface
 
-		float StabilizationThreshold = FPhysicsInterface::GetStabilizationEnergyThreshold_AssumesLocked(ActorHandle);
+		float StabilizationThreshold = FPhysicsInterface::GetStabilizationEnergyThreshold_AssumesLocked(GetPhysicsActorHandle());
 		StabilizationThreshold *= StabilizationThresholdMultiplier;
-		FPhysicsInterface::SetStabilizationEnergyThreshold_AssumesLocked(ActorHandle, StabilizationThreshold);
+		FPhysicsInterface::SetStabilizationEnergyThreshold_AssumesLocked(GetPhysicsActorHandle(), StabilizationThreshold);
 
 		uint32 PositionIterCount = FMath::Clamp<uint8>(PositionSolverIterationCount, 1, 255);
 		uint32 VelocityIterCount = FMath::Clamp<uint8>(VelocitySolverIterationCount, 1, 255);
 		uint32 ProjectionIterCount = FMath::Clamp<uint8>(ProjectionSolverIterationCount, 1, 255);
-		FPhysicsInterface::SetSolverPositionIterationCount_AssumesLocked(ActorHandle, PositionIterCount);
-		FPhysicsInterface::SetSolverVelocityIterationCount_AssumesLocked(ActorHandle, VelocityIterCount);
-		FPhysicsInterface::SetSolverProjectionIterationCount_AssumesLocked(ActorHandle, ProjectionIterCount);
+		FPhysicsInterface::SetSolverPositionIterationCount_AssumesLocked(GetPhysicsActorHandle(), PositionIterCount);
+		FPhysicsInterface::SetSolverVelocityIterationCount_AssumesLocked(GetPhysicsActorHandle(), VelocityIterCount);
+		FPhysicsInterface::SetSolverProjectionIterationCount_AssumesLocked(GetPhysicsActorHandle(), ProjectionIterCount);
 
 		CreateDOFLock();
-		if(FPhysicsInterface::IsInScene(ActorHandle) && !IsRigidBodyKinematic_AssumesLocked(ActorHandle))
+		if(FPhysicsInterface::IsInScene(GetPhysicsActorHandle()) && !IsRigidBodyKinematic_AssumesLocked(GetPhysicsActorHandle()))
 		{
 			if(!bStartAwake && !bComponentAwake)
 			{
-				FPhysicsInterface::SetWakeCounter_AssumesLocked(ActorHandle, 0.0f);
-				FPhysicsInterface::PutToSleep_AssumesLocked(ActorHandle);
+				FPhysicsInterface::SetWakeCounter_AssumesLocked(GetPhysicsActorHandle(), 0.0f);
+				FPhysicsInterface::PutToSleep_AssumesLocked(GetPhysicsActorHandle());
 			}
 		}
 
-		FPhysicsInterface::SetInitialized_AssumesLocked(ActorHandle, true);
+		FPhysicsInterface::SetInitialized_AssumesLocked(GetPhysicsActorHandle(), true);
 	}
 }
 
@@ -4722,7 +4728,7 @@ void FBodyInstance::UpdateSolverAsyncDeltaTime()
 			return;
 		}
 
-		Chaos::FPhysicsSolverBase* Solver = ActorHandle->GetSolverBase();
+		Chaos::FPhysicsSolverBase* Solver = GetPhysicsActorHandle()->GetSolverBase();
 
 		if (!Solver->IsUsingAsyncResults())
 		{
diff --git a/Engine/Source/Runtime/Engine/Private/PhysicsEngine/BodySetup.cpp b/Engine/Source/Runtime/Engine/Private/PhysicsEngine/BodySetup.cpp
--- a/Engine/Source/Runtime/Engine/Private/PhysicsEngine/BodySetup.cpp
+++ b/Engine/Source/Runtime/Engine/Private/PhysicsEngine/BodySetup.cpp
@@ -722,7 +722,7 @@ void UBodySetup::AddShapesToRigidActor_AssumesLocked(
 
 	{
 		SCOPE_CYCLE_COUNTER(STAT_AddGeomToSolver);
-		FPhysicsInterface::AddGeometry(OwningInstance->ActorHandle, AddParams, NewShapes);
+		FPhysicsInterface::AddGeometry(OwningInstance->GetPhysicsActorHandle(), AddParams, NewShapes);
 	}
 }
 
diff --git a/Engine/Source/Runtime/Engine/Private/PhysicsEngine/ConstraintInstance.cpp b/Engine/Source/Runtime/Engine/Private/PhysicsEngine/ConstraintInstance.cpp
--- a/Engine/Source/Runtime/Engine/Private/PhysicsEngine/ConstraintInstance.cpp
+++ b/Engine/Source/Runtime/Engine/Private/PhysicsEngine/ConstraintInstance.cpp
@@ -455,8 +455,8 @@ float ComputeAverageMass_AssumesLocked(Chaos::FPhysicsObject* Body1, Chaos::FPhy
 
 bool GetActorRefs(FBodyInstance* Body1, FBodyInstance* Body2, FPhysicsActorHandle& OutActorRef1, FPhysicsActorHandle& OutActorRef2, UObject* DebugOwner)
 {
-	FPhysicsActorHandle ActorRef1 = Body1 ? Body1->ActorHandle : FPhysicsActorHandle();
-	FPhysicsActorHandle ActorRef2 = Body2 ? Body2->ActorHandle : FPhysicsActorHandle();
+	FPhysicsActorHandle ActorRef1 = Body1 ? Body1->GetPhysicsActorHandle() : FPhysicsActorHandle();
+	FPhysicsActorHandle ActorRef2 = Body2 ? Body2->GetPhysicsActorHandle() : FPhysicsActorHandle();
 
 	// Do not create joint unless you have two actors
 	// Do not create joint unless one of the actors is dynamic
diff --git a/Engine/Source/Runtime/Engine/Private/PhysicsEngine/Experimental/PhysInterface_Chaos.cpp b/Engine/Source/Runtime/Engine/Private/PhysicsEngine/Experimental/PhysInterface_Chaos.cpp
--- a/Engine/Source/Runtime/Engine/Private/PhysicsEngine/Experimental/PhysInterface_Chaos.cpp
+++ b/Engine/Source/Runtime/Engine/Private/PhysicsEngine/Experimental/PhysInterface_Chaos.cpp
@@ -845,11 +845,11 @@ bool FPhysInterface_Chaos::LineTrace_Geom(FHitResult& OutHit, const FBodyInstanc
 	{
 		{
 			// #PHYS2 Really need a concept for "multi" locks here - as we're locking ActorRef but not TargetInstance->ActorRef
-			FPhysicsCommand::ExecuteRead(InInstance->ActorHandle, [&](const FPhysicsActorHandle& Actor)
+			FPhysicsCommand::ExecuteRead(InInstance->GetPhysicsActorHandle(), [&](const FPhysicsActorHandle& Actor)
 			{
 				// If we're welded then the target instance is actually our parent
 				const FBodyInstance* TargetInstance = InInstance->WeldParent ? InInstance->WeldParent : InInstance;
-				if(const FPhysicsActorHandle RigidBody = TargetInstance->ActorHandle)
+				if(const FPhysicsActorHandle RigidBody = TargetInstance->GetPhysicsActorHandle())
 				{
 					FRaycastHit BestHit;
 					BestHit.Distance = FLT_MAX;
@@ -942,7 +942,7 @@ bool FPhysInterface_Chaos::Sweep_Geom(FHitResult& OutHit, const FBodyInstance* I
 
 		const FBodyInstance* TargetInstance = InInstance->WeldParent ? InInstance->WeldParent : InInstance;
 
-		FPhysicsCommand::ExecuteRead(TargetInstance->ActorHandle, [&](const FPhysicsActorHandle& Actor)
+		FPhysicsCommand::ExecuteRead(TargetInstance->GetPhysicsActorHandle(), [&](const FPhysicsActorHandle& Actor)
 		{
 			if (Actor && InInstance->OwnerComponent.Get())
 			{
@@ -1040,7 +1040,7 @@ bool FPhysInterface_Chaos::Sweep_Geom(FHitResult& OutHit, const FBodyInstance* I
 bool Overlap_GeomInternal(const FBodyInstance* InInstance, const Chaos::FImplicitObject& InGeom, const FTransform& GeomTransform, FMTDResult* OutOptResult, bool bTraceComplex)
 {
 	const FBodyInstance* TargetInstance = InInstance->WeldParent ? InInstance->WeldParent : InInstance;
-	FPhysicsActorHandle RigidBody = TargetInstance->ActorHandle;
+	FPhysicsActorHandle RigidBody = TargetInstance->GetPhysicsActorHandle();
 
 	if (!RigidBody)
 	{
@@ -1129,7 +1129,7 @@ bool FPhysInterface_Chaos::GetSquaredDistanceToBody(const FBodyInstance* InInsta
 	const FTransform BodyTM = UseBI->GetUnrealWorldTransform();
 	const FVector LocalPoint = BodyTM.InverseTransformPositionNoScale(InPoint);
 
-	FPhysicsCommand::ExecuteRead(UseBI->ActorHandle, [&](const FPhysicsActorHandle& Actor)
+	FPhysicsCommand::ExecuteRead(UseBI->GetPhysicsActorHandle(), [&](const FPhysicsActorHandle& Actor)
 	{
 		bEarlyOut = false;
 
diff --git a/Engine/Source/Runtime/Engine/Private/PhysicsEngine/Experimental/PhysScene_Chaos.cpp b/Engine/Source/Runtime/Engine/Private/PhysicsEngine/Experimental/PhysScene_Chaos.cpp
--- a/Engine/Source/Runtime/Engine/Private/PhysicsEngine/Experimental/PhysScene_Chaos.cpp
+++ b/Engine/Source/Runtime/Engine/Private/PhysicsEngine/Experimental/PhysScene_Chaos.cpp
@@ -1585,7 +1585,7 @@ void FPhysScene_Chaos::SetKinematicTarget_AssumesLocked(FBodyInstance* BodyInsta
 
 bool FPhysScene_Chaos::GetKinematicTarget_AssumesLocked(const FBodyInstance* BodyInstance, FTransform& OutTM) const
 {
-	OutTM = FPhysicsInterface::GetKinematicTarget_AssumesLocked(BodyInstance->ActorHandle);
+	OutTM = FPhysicsInterface::GetKinematicTarget_AssumesLocked(BodyInstance->GetPhysicsActorHandle());
 	return true;
 }
 
@@ -1673,7 +1673,7 @@ void GatherActorsAndTransforms(
 	for (int32 i = 0; i < NumBodies; i++)
 	{
 		FBodyInstance* BodyInst = SkelComp->Bodies[i];
-		FPhysicsActorHandle& ActorHandle = BodyInst->ActorHandle;
+		FPhysicsActorHandle& ActorHandle = BodyInst->GetPhysicsActorHandle();
 		if (bTeleport || !BodyInst->IsInstanceSimulatingPhysics())
 		{
 			const int32 BoneIndex = BodyInst->InstanceBoneIndex;
@@ -1940,7 +1940,7 @@ void FPhysScene_Chaos::UpdateKinematicsOnDeferredSkelMeshes()
 						continue;
 					}
 
-					FPhysicsActorHandle& ActorHandle = BodyInst->ActorHandle;
+					FPhysicsActorHandle& ActorHandle = BodyInst->GetPhysicsActorHandle();
 					if (GKinematicDeferralCheckValidBodies && (ActorHandle == nullptr || ActorHandle->GetSyncTimestamp() == nullptr || ActorHandle->GetMarkedDeleted()))
 					{
 						if (GKinematicDeferralLogInvalidBodies)
@@ -2025,7 +2025,7 @@ void FPhysScene_Chaos::UpdateKinematicsOnDeferredSkelMeshes()
 						continue;
 					}
 
-					FPhysicsActorHandle& ActorHandle = BodyInst->ActorHandle;
+					FPhysicsActorHandle& ActorHandle = BodyInst->GetPhysicsActorHandle();
 					if (GKinematicDeferralCheckValidBodies && (ActorHandle == nullptr || ActorHandle->GetSyncTimestamp() == nullptr || ActorHandle->GetMarkedDeleted()))
 					{
 						if (GKinematicDeferralLogInvalidBodies)
@@ -2358,7 +2358,7 @@ void FPhysScene_Chaos::OnSyncBodies(Chaos::FPhysicsSolverBase* Solver)
 
 						}
 
-						PhysScene->UpdateActorInAccelerationStructure(BodyInstance->ActorHandle);
+						PhysScene->UpdateActorInAccelerationStructure(BodyInstance->GetPhysicsActorHandle());
 					}
 
 					if (Proxy->GetWakeEvent() != Chaos::EWakeEventEntry::None && !bPendingMove)
diff --git a/Engine/Source/Runtime/Engine/Private/PhysicsEngine/ImmediatePhysics/ImmediatePhysicsChaos/ImmediatePhysicsActorHandle_Chaos.cpp b/Engine/Source/Runtime/Engine/Private/PhysicsEngine/ImmediatePhysics/ImmediatePhysicsChaos/ImmediatePhysicsActorHandle_Chaos.cpp
--- a/Engine/Source/Runtime/Engine/Private/PhysicsEngine/ImmediatePhysics/ImmediatePhysicsChaos/ImmediatePhysicsActorHandle_Chaos.cpp
+++ b/Engine/Source/Runtime/Engine/Private/PhysicsEngine/ImmediatePhysics/ImmediatePhysicsChaos/ImmediatePhysicsActorHandle_Chaos.cpp
@@ -84,12 +84,12 @@ namespace ImmediatePhysics_Chaos
 	bool CloneGeometry(FBodyInstance* BodyInstance, EActorType ActorType, const FVector& Scale, FReal& OutMass, Chaos::FVec3& OutInertia, Chaos::FRigidTransform3& OutCoMTransform, Chaos::FImplicitObjectPtr& OutGeom, TArray<TUniquePtr<Chaos::FPerShapeData>>& OutShapes)
 	{
 		// We should only get non-simulated objects through this path, but you never know...
-		if ((BodyInstance != nullptr) && !BodyInstance->bSimulatePhysics && BodyInstance->ActorHandle)
+		if ((BodyInstance != nullptr) && !BodyInstance->bSimulatePhysics && BodyInstance->GetPhysicsActorHandle())
 		{
 			OutMass = 0.0f;
 			OutInertia = FVector::ZeroVector;
 			OutCoMTransform = FTransform::Identity;
-			OutGeom = CloneGeometry(BodyInstance->ActorHandle->GetGameThreadAPI().GetGeometry(), OutShapes);
+			OutGeom = CloneGeometry(BodyInstance->GetPhysicsActorHandle()->GetGameThreadAPI().GetGeometry(), OutShapes);
 			if (OutGeom != nullptr)
 			{
 				return true;
diff --git a/Engine/Source/Runtime/Engine/Private/PhysicsEngine/PhysAnim.cpp b/Engine/Source/Runtime/Engine/Private/PhysicsEngine/PhysAnim.cpp
--- a/Engine/Source/Runtime/Engine/Private/PhysicsEngine/PhysAnim.cpp
+++ b/Engine/Source/Runtime/Engine/Private/PhysicsEngine/PhysAnim.cpp
@@ -645,7 +645,7 @@ void USkeletalMeshComponent::UpdateKinematicBonesToAnim(const TArray<FTransform>
 					{
 						continue;
 					}
-					FPhysicsActorHandle& ActorHandle = BodyInst->ActorHandle;
+					FPhysicsActorHandle& ActorHandle = BodyInst->GetPhysicsActorHandle();
 
 					if (FPhysicsInterface::IsValid(ActorHandle) &&
 						!FPhysicsInterface::IsStatic(ActorHandle) &&
diff --git a/Engine/Source/Runtime/Engine/Private/PhysicsEngine/PhysicalAnimationComponent.cpp b/Engine/Source/Runtime/Engine/Private/PhysicsEngine/PhysicalAnimationComponent.cpp
--- a/Engine/Source/Runtime/Engine/Private/PhysicsEngine/PhysicalAnimationComponent.cpp
+++ b/Engine/Source/Runtime/Engine/Private/PhysicsEngine/PhysicalAnimationComponent.cpp
@@ -399,7 +399,7 @@ void UPhysicalAnimationComponent::UpdatePhysicsEngineImp()
 					int32 ChildBodyIdx = PhysAsset->FindBodyIndex(PhysAnimData.BodyName);
 					if (FBodyInstance* ChildBody = (ChildBodyIdx == INDEX_NONE ? nullptr : SkeletalMeshComponent->Bodies[ChildBodyIdx]))
 					{
-						if (FPhysicsActorHandle ActorHandle = ChildBody->ActorHandle)
+						if (FPhysicsActorHandle ActorHandle = ChildBody->GetPhysicsActorHandle())
 						{
 							FPhysScene* Scene = ChildBody->GetPhysicsScene();
 
@@ -431,7 +431,7 @@ void UPhysicalAnimationComponent::UpdatePhysicsEngineImp()
 							// Save reference to the kinematic actor.
 							InstanceData.TargetActor = KineActor;
 
-							ConstraintInstance->InitConstraint_AssumesLocked(ChildBody->ActorHandle, InstanceData.TargetActor, 1.f);
+							ConstraintInstance->InitConstraint_AssumesLocked(ChildBody->GetPhysicsActorHandle(), InstanceData.TargetActor, 1.f);
 						}
 					}
 				}
diff --git a/Engine/Source/Runtime/Engine/Private/PhysicsEngine/PhysicsConstraintComponent.cpp b/Engine/Source/Runtime/Engine/Private/PhysicsEngine/PhysicsConstraintComponent.cpp
--- a/Engine/Source/Runtime/Engine/Private/PhysicsEngine/PhysicsConstraintComponent.cpp
+++ b/Engine/Source/Runtime/Engine/Private/PhysicsEngine/PhysicsConstraintComponent.cpp
@@ -294,8 +294,8 @@ void UPhysicsConstraintComponent::InitComponentConstraint()
 	FBodyInstance* Body1 = GetBodyInstance(EConstraintFrame::Frame1);
 	FBodyInstance* Body2 = GetBodyInstance(EConstraintFrame::Frame2);
 
-	Chaos::FPhysicsObject* Object1 = (Body1 && Body1->IsValidBodyInstance()) ? Body1->ActorHandle->GetPhysicsObject() : GetPhysicsObject(EConstraintFrame::Frame1);
-	Chaos::FPhysicsObject* Object2 = (Body2 && Body2->IsValidBodyInstance()) ? Body2->ActorHandle->GetPhysicsObject() : GetPhysicsObject(EConstraintFrame::Frame2);
+	Chaos::FPhysicsObject* Object1 = (Body1 && Body1->IsValidBodyInstance()) ? Body1->GetPhysicsActorHandle()->GetPhysicsObject() : GetPhysicsObject(EConstraintFrame::Frame1);
+	Chaos::FPhysicsObject* Object2 = (Body2 && Body2->IsValidBodyInstance()) ? Body2->GetPhysicsActorHandle()->GetPhysicsObject() : GetPhysicsObject(EConstraintFrame::Frame2);
 
 	if (Object1 || Object2)
 	{
diff --git a/Engine/Source/Runtime/Engine/Private/PhysicsEngine/PhysicsHandleComponent.cpp b/Engine/Source/Runtime/Engine/Private/PhysicsEngine/PhysicsHandleComponent.cpp
--- a/Engine/Source/Runtime/Engine/Private/PhysicsEngine/PhysicsHandleComponent.cpp
+++ b/Engine/Source/Runtime/Engine/Private/PhysicsEngine/PhysicsHandleComponent.cpp
@@ -60,9 +60,9 @@ void UPhysicsHandleComponent::GrabComponent(class UPrimitiveComponent* InCompone
 
 	FRotator GrabbedRotation = FRotator::ZeroRotator;
 
-	if(FPhysicsInterface::IsValid(BodyInstance->ActorHandle))
+	if(FPhysicsInterface::IsValid(BodyInstance->GetPhysicsActorHandle()))
 	{
-		FPhysicsCommand::ExecuteRead(BodyInstance->ActorHandle, [&](const FPhysicsActorHandle& Actor)
+		FPhysicsCommand::ExecuteRead(BodyInstance->GetPhysicsActorHandle(), [&](const FPhysicsActorHandle& Actor)
 		{
 			GrabbedRotation = FPhysicsInterface::GetGlobalPose_AssumesLocked(Actor).Rotator();
 		});
diff --git a/Engine/Source/Runtime/Engine/Private/PhysicsEngine/Tests/PhysicsBodyInstanceTests.cpp b/Engine/Source/Runtime/Engine/Private/PhysicsEngine/Tests/PhysicsBodyInstanceTests.cpp
--- a/Engine/Source/Runtime/Engine/Private/PhysicsEngine/Tests/PhysicsBodyInstanceTests.cpp
+++ b/Engine/Source/Runtime/Engine/Private/PhysicsEngine/Tests/PhysicsBodyInstanceTests.cpp
@@ -29,7 +29,7 @@ bool FPhysicsBodyInstanceTest::RunTest(const FString& Parameters)
 
 	// TODO: Physics actor handle will be changing to a ptr!
 	FBodyInstance& BI = SMC->BodyInstance;
-	FPhysicsActorHandle PH = BI.ActorHandle;
+	FPhysicsActorHandle PH = BI.GetPhysicsActorHandle();
 
 	check(FPhysicsInterface::IsValid(PH));
 	check(FPhysicsInterface::IsRigidBody(PH));
diff --git a/Engine/Source/Runtime/Engine/Private/PrimitiveComponentPhysics.cpp b/Engine/Source/Runtime/Engine/Private/PrimitiveComponentPhysics.cpp
--- a/Engine/Source/Runtime/Engine/Private/PrimitiveComponentPhysics.cpp
+++ b/Engine/Source/Runtime/Engine/Private/PrimitiveComponentPhysics.cpp
@@ -1535,7 +1535,7 @@ Chaos::FPhysicsObject* UPrimitiveComponent::GetPhysicsObjectById(Chaos::FPhysics
 		return nullptr;
 	}
 
-	return BodyInstance.ActorHandle->GetPhysicsObject();
+	return BodyInstance.GetPhysicsActorHandle()->GetPhysicsObject();
 }
 
 Chaos::FPhysicsObject* UPrimitiveComponent::GetPhysicsObjectByName(const FName& Name) const
diff --git a/Engine/Source/Runtime/Engine/Private/SkeletalMeshComponentPhysics.cpp b/Engine/Source/Runtime/Engine/Private/SkeletalMeshComponentPhysics.cpp
--- a/Engine/Source/Runtime/Engine/Private/SkeletalMeshComponentPhysics.cpp
+++ b/Engine/Source/Runtime/Engine/Private/SkeletalMeshComponentPhysics.cpp
@@ -733,12 +733,12 @@ void USkeletalMeshComponent::InitArticulated(FPhysScene* PhysScene)
 	{
 		if (FBodyInstance* Body = Bodies[BodyIndex])
 		{
-			if (!Body->ActorHandle)
+			if (!Body->GetPhysicsActorHandle())
 			{
 				continue;
 			}
 
-			Chaos::FPhysicsObject* PhysicsObject = Body->ActorHandle->GetPhysicsObject();
+			Chaos::FPhysicsObject* PhysicsObject = Body->GetPhysicsActorHandle()->GetPhysicsObject();
 			if (UBodySetup* Setup = Body->GetBodySetup())
 			{
 				Chaos::FPhysicsObjectInterface::SetName(PhysicsObject, Setup->BoneName);
@@ -3836,22 +3836,22 @@ FTransform USkeletalMeshComponent::GetComponentTransformFromBodyInstance(FBodyIn
 
 Chaos::FPhysicsObject* USkeletalMeshComponent::GetPhysicsObjectById(Chaos::FPhysicsObjectId Id) const
 {
-	if (!Bodies.IsValidIndex(Id) || !Bodies[Id] || !Bodies[Id]->ActorHandle)
+	if (!Bodies.IsValidIndex(Id) || !Bodies[Id] || !Bodies[Id]->GetPhysicsActorHandle())
 	{
 		return nullptr;
 	}
-	return Bodies[Id]->ActorHandle->GetPhysicsObject();
+	return Bodies[Id]->GetPhysicsActorHandle()->GetPhysicsObject();
 }
 
 Chaos::FPhysicsObject* USkeletalMeshComponent::GetPhysicsObjectByName(const FName& Name) const
 {
 	FBodyInstance* Body = GetBodyInstance(Name);
-	if (!Body || !Body->ActorHandle)
+	if (!Body || !Body->GetPhysicsActorHandle())
 	{
 		return nullptr;
 	}
 
-	return Body->ActorHandle->GetPhysicsObject();
+	return Body->GetPhysicsActorHandle()->GetPhysicsObject();
 }
 
 TArray<Chaos::FPhysicsObject*> USkeletalMeshComponent::GetAllPhysicsObjects() const
diff --git a/Engine/Source/Runtime/Engine/Private/Tests/AutoRTFM/AutoRTFMActorComponentTests.cpp b/Engine/Source/Runtime/Engine/Private/Tests/AutoRTFM/AutoRTFMActorComponentTests.cpp
--- a/Engine/Source/Runtime/Engine/Private/Tests/AutoRTFM/AutoRTFMActorComponentTests.cpp
+++ b/Engine/Source/Runtime/Engine/Private/Tests/AutoRTFM/AutoRTFMActorComponentTests.cpp
@@ -158,9 +158,9 @@ size_t CountAAutoRTFMTestActors()
 // See: SOL-6709
 AUTORTFM_ACTOR_COMPONENT_TEST(RegisterComponent_UnregisterComponent)
 {
-	Component->BodyInstance.ActorHandle = Chaos::FSingleParticlePhysicsProxy::Create(Chaos::FGeometryParticle::CreateParticle());
-	Component->BodyInstance.ActorHandle->GetParticle_LowLevel()->SetGeometry(MakeImplicitObjectPtr<Chaos::FSphere>(Chaos::FVec3(1, 2, 3), 1));
-	World->GetPhysicsScene()->GetSolver()->RegisterObject(Component->BodyInstance.ActorHandle);
+	Component->BodyInstance.SetPhysicsActorHandle(Chaos::FSingleParticlePhysicsProxy::Create(Chaos::FGeometryParticle::CreateParticle()));
+	Component->BodyInstance.GetPhysicsActorHandle()->GetParticle_LowLevel()->SetGeometry(MakeImplicitObjectPtr<Chaos::FSphere>(Chaos::FVec3(1, 2, 3), 1));
+	World->GetPhysicsScene()->GetSolver()->RegisterObject(Component->BodyInstance.GetPhysicsActorHandle());
 
 	AutoRTFM::ETransactionResult Result = AutoRTFM::Transact([&]
 		{
diff --git a/Engine/Source/Runtime/Landscape/Private/LandscapeCollision.cpp b/Engine/Source/Runtime/Landscape/Private/LandscapeCollision.cpp
--- a/Engine/Source/Runtime/Landscape/Private/LandscapeCollision.cpp
+++ b/Engine/Source/Runtime/Landscape/Private/LandscapeCollision.cpp
@@ -497,7 +497,7 @@ void ULandscapeHeightfieldCollisionComponent::OnCreatePhysicsState()
 				// Set body instance data
 				BodyInstance.PhysicsUserData = FPhysicsUserData(&BodyInstance);
 				BodyInstance.OwnerComponent = this;
-				BodyInstance.ActorHandle = PhysHandle;
+				BodyInstance.SetPhysicsActorHandle(PhysHandle);
 
 				Body_External.SetUserData(&BodyInstance.PhysicsUserData);
 
@@ -1898,7 +1898,7 @@ struct FMeshCollisionInitHelper
 
 		TargetInstance->PhysicsUserData = FPhysicsUserData(TargetInstance);
 		TargetInstance->OwnerComponent = Component;
-		TargetInstance->ActorHandle = ActorHandle;
+		TargetInstance->SetPhysicsActorHandle(ActorHandle);
 
 		ActorHandle->GetGameThreadAPI().SetUserData(&TargetInstance->PhysicsUserData);
 	}
@@ -2053,7 +2053,7 @@ void ULandscapeHeightfieldCollisionComponent::UpdateHeightfieldRegion(int32 Comp
 			return;
 		}
 
-		if(!BodyInstance.ActorHandle)
+		if(!BodyInstance.GetPhysicsActorHandle())
 		{
 			return;
 		}
