[Engine]
* Added initializer to statestream struct

diff --git a/Engine/Source/Runtime/Engine/Public/StateStream/ExponentialHeightFogStateStream.h b/Engine/Source/Runtime/Engine/Public/StateStream/ExponentialHeightFogStateStream.h
--- a/Engine/Source/Runtime/Engine/Public/StateStream/ExponentialHeightFogStateStream.h
+++ b/Engine/Source/Runtime/Engine/Public/StateStream/ExponentialHeightFogStateStream.h
@@ -33,69 +33,69 @@ struct FExponentialHeightFogDynamicState
 // private: // Should be private but complicates some of the code that needs support both this path and UComponent path
 
 	UPROPERTY()
-	float FogDensity;
+	float FogDensity = 0.0f;
 	UPROPERTY()
-	float FogHeightFalloff;
+	float FogHeightFalloff = 0.0f;
 	UPROPERTY()
 	FExponentialHeightFogData SecondFogData;
 	UPROPERTY()
-	FLinearColor FogInscatteringLuminance;
+	FLinearColor FogInscatteringLuminance = FLinearColor::White;
 	UPROPERTY()
-	FLinearColor SkyAtmosphereAmbientContributionColorScale;
+	FLinearColor SkyAtmosphereAmbientContributionColorScale = FLinearColor::White;
 	UPROPERTY()
 	TObjectPtr<UTextureCube> InscatteringColorCubemap;
 	UPROPERTY()
-	float InscatteringColorCubemapAngle;
+	float InscatteringColorCubemapAngle = 0.0f;
 	UPROPERTY()
-	FLinearColor InscatteringTextureTint;
+	FLinearColor InscatteringTextureTint = FLinearColor::White;
 	UPROPERTY()
-	float FullyDirectionalInscatteringColorDistance;
+	float FullyDirectionalInscatteringColorDistance = 0.0f;
 	UPROPERTY()
-	float NonDirectionalInscatteringColorDistance;
+	float NonDirectionalInscatteringColorDistance = 0.0f;
 	UPROPERTY()
-	float DirectionalInscatteringExponent;
+	float DirectionalInscatteringExponent = 0.0f;
 	UPROPERTY()
-	float DirectionalInscatteringStartDistance;
+	float DirectionalInscatteringStartDistance = 0.0f;
 	UPROPERTY()
-	FLinearColor DirectionalInscatteringLuminance;
+	FLinearColor DirectionalInscatteringLuminance = FLinearColor::White;
 	UPROPERTY()
-	float FogMaxOpacity;
+	float FogMaxOpacity = 0.0f;
 	UPROPERTY()
-	float StartDistance;
+	float StartDistance = 0.0f;
 	UPROPERTY()
-	float EndDistance;
+	float EndDistance = 0.0f;
 	UPROPERTY()
-	float FogCutoffDistance;
+	float FogCutoffDistance = 0.0f;
 	UPROPERTY()
-	bool bEnableVolumetricFog;
+	bool bEnableVolumetricFog = false;
 	UPROPERTY()
-	float VolumetricFogScatteringDistribution;
+	float VolumetricFogScatteringDistribution = 0.0f;
 	UPROPERTY()
-	FColor VolumetricFogAlbedo;
+	FColor VolumetricFogAlbedo = FColor::White;
 	UPROPERTY()
-	FLinearColor VolumetricFogEmissive;
+	FLinearColor VolumetricFogEmissive = FLinearColor::White;
 	UPROPERTY()
-	float VolumetricFogExtinctionScale;
+	float VolumetricFogExtinctionScale = 0.0f;
 	UPROPERTY()
-	float VolumetricFogDistance;
+	float VolumetricFogDistance = 0.0f;
 	UPROPERTY()
-	float VolumetricFogStartDistance;
+	float VolumetricFogStartDistance = 0.0f;
 	UPROPERTY()
-	float VolumetricFogNearFadeInDistance;
+	float VolumetricFogNearFadeInDistance = 0.0f;
 	UPROPERTY()
-	float VolumetricFogStaticLightingScatteringIntensity;
+	float VolumetricFogStaticLightingScatteringIntensity = 0.0f;
 	UPROPERTY()
-	bool bOverrideLightColorsWithFogInscatteringColors;
+	bool bOverrideLightColorsWithFogInscatteringColors = false;
 	UPROPERTY()
-	uint8 bHoldout : 1;
+	uint8 bHoldout : 1 = false;
 	UPROPERTY()
-	uint8 bRenderInMainPass : 1;
+	uint8 bRenderInMainPass : 1 = false;
 	UPROPERTY()
-	uint8 bVisibleInReflectionCaptures : 1;
+	uint8 bVisibleInReflectionCaptures : 1 = false;
 	UPROPERTY()
-	uint8 bVisibleInRealTimeSkyCaptures : 1;
+	uint8 bVisibleInRealTimeSkyCaptures : 1 = false;
 	UPROPERTY()
-	float Height;
+	float Height = 0.0f;
 };
 
 ////////////////////////////////////////////////////////////////////////////////////////////////////
