Iris - Engine changes
- Fixed Creativemovetool replication with Iris after recent refactoring in CL 39043048.AObjectInteractionMode
- Fixed AObjectInteractionMode not replicating
- Fixed protocolid mismatch due to conditional cdo modification
- UInteractionTargetDetection now properly supports being a replicated subobject

Preflight 679353ae730cc5b564a33225 succeded on retry

diff --git a/Engine/Source/Runtime/Engine/Private/Actor.cpp b/Engine/Source/Runtime/Engine/Private/Actor.cpp
--- a/Engine/Source/Runtime/Engine/Private/Actor.cpp
+++ b/Engine/Source/Runtime/Engine/Private/Actor.cpp
@@ -4479,10 +4479,12 @@ void AActor::DispatchBeginPlay(bool bFromLevelStreaming)
 		const uint32 CurrentCallDepth = BeginPlayCallDepth++;
 
 		bActorBeginningPlayFromLevelStreaming = bFromLevelStreaming;
-		ActorHasBegunPlay = EActorBeginPlayState::BeginningPlay;
 
+		// Call this before updating status of ActorHasBegunPlay
 		BuildReplicatedComponentsInfo();
 
+		ActorHasBegunPlay = EActorBeginPlayState::BeginningPlay;
+
 #if UE_WITH_IRIS
 		BeginReplication();
 #endif // UE_WITH_IRIS
diff --git a/Engine/Source/Runtime/Engine/Private/ActorReplication.cpp b/Engine/Source/Runtime/Engine/Private/ActorReplication.cpp
--- a/Engine/Source/Runtime/Engine/Private/ActorReplication.cpp
+++ b/Engine/Source/Runtime/Engine/Private/ActorReplication.cpp
@@ -718,7 +718,7 @@ void AActor::AddReplicatedSubObject(UObject* SubObject, ELifetimeCondition NetCo
 		*GetName(), this, *SubObject->GetName(), SubObject, *StaticEnum<ELifetimeCondition>()->GetValueAsString(ReplicatedSubObjects.GetNetCondition(SubObject)), *StaticEnum<ELifetimeCondition>()->GetValueAsString(NetCondition));
 
 #if UE_WITH_IRIS
-	if (GetIsReplicated() && HasActorBegunPlay())
+	if (GetIsReplicated() && (ActorHasBegunPlay != EActorBeginPlayState::HasNotBegunPlay))
 	{
 		UE::Net::FReplicationSystemUtil::BeginReplicationForActorSubObject(this, SubObject, NetCondition);
 	}
