Call uninitalize animation before deleting linked anim instances when calling ResetLinkedAnimInstances on the skeletal mesh

diff --git a/Engine/Source/Runtime/Engine/Private/Components/SkeletalMeshComponent.cpp b/Engine/Source/Runtime/Engine/Private/Components/SkeletalMeshComponent.cpp
--- a/Engine/Source/Runtime/Engine/Private/Components/SkeletalMeshComponent.cpp
+++ b/Engine/Source/Runtime/Engine/Private/Components/SkeletalMeshComponent.cpp
@@ -3370,6 +3370,7 @@ void USkeletalMeshComponent::ResetLinkedAnimInstances()
 	{
 		if(LinkedInstance && LinkedInstance->bCreatedByLinkedAnimGraph)
 		{
+			LinkedInstance->UninitializeAnimation();
 			LinkedInstance->EndNotifyStates();
 			LinkedInstance->MarkAsGarbage();
 			LinkedInstance = nullptr;
