Fix NMT_CloseReason control channel messages preventing other control messages from being parsed if sent in the same bunch.

diff --git a/Engine/Source/Runtime/Engine/Private/DataChannel.cpp b/Engine/Source/Runtime/Engine/Private/DataChannel.cpp
--- a/Engine/Source/Runtime/Engine/Private/DataChannel.cpp
+++ b/Engine/Source/Runtime/Engine/Private/DataChannel.cpp
@@ -1951,8 +1951,6 @@ void UControlChannel::ReceivedBunch( FInBunch& Bunch )
 			{
 				Connection->HandleReceiveCloseReason(CloseReasonList);
 			}
-
-			bStopReadingBunch = true;
 		}
 		else if (MessageType == NMT_NetPing)
 		{
