Removed FString / FName cost from updating a skeletal mesh every frame.

diff --git a/Engine/Source/Runtime/Engine/Private/SkeletalRender.cpp b/Engine/Source/Runtime/Engine/Private/SkeletalRender.cpp
--- a/Engine/Source/Runtime/Engine/Private/SkeletalRender.cpp
+++ b/Engine/Source/Runtime/Engine/Private/SkeletalRender.cpp
@@ -274,6 +274,10 @@ void FSkeletalMeshObject::InitLODInfos(const TArrayView<const FSkelMeshComponent
 			const FSkelMeshComponentLODInfo &Info = InComponentLODInfo[Idx];
 			MeshObjectLODInfo.HiddenMaterials = Info.HiddenMaterials;
 		}
+
+#if RHI_ENABLE_RESOURCE_INFO
+		MeshObjectLODInfo.AssetPathName = FName(FString::Printf(TEXT("%s [LOD%d]"), *AssetPathName.ToString(), Idx));
+#endif
 	}
 }
 
@@ -291,7 +295,7 @@ FName FSkeletalMeshObject::GetAssetPathName(int32 LODIndex) const
 #if RHI_ENABLE_RESOURCE_INFO
 	if (LODIndex > -1)
 	{
-		return FName(FString::Printf(TEXT("%s [LOD%d]"), *AssetPathName.ToString(), LODIndex));
+		return LODInfo[LODIndex].AssetPathName;
 	}
 	else
 	{
diff --git a/Engine/Source/Runtime/Engine/Public/SkeletalRenderPublic.h b/Engine/Source/Runtime/Engine/Public/SkeletalRenderPublic.h
--- a/Engine/Source/Runtime/Engine/Public/SkeletalRenderPublic.h
+++ b/Engine/Source/Runtime/Engine/Public/SkeletalRenderPublic.h
@@ -313,7 +313,11 @@ public:
 		 *  occur at when the object size is ScreenSize + LODHysteresis. This avoids flickering at the boundary.
 		 */
 		float LODHysteresis;
-	};	
+
+#if RHI_ENABLE_RESOURCE_INFO
+		FName AssetPathName;
+#endif
+	};
 
 	TArray<FSkelMeshObjectLODInfo> LODInfo;
 
