Removed vestigial GPU event scope.

diff --git a/Engine/Source/Runtime/Engine/Private/LevelTick.cpp b/Engine/Source/Runtime/Engine/Private/LevelTick.cpp
--- a/Engine/Source/Runtime/Engine/Private/LevelTick.cpp
+++ b/Engine/Source/Runtime/Engine/Private/LevelTick.cpp
@@ -956,21 +956,6 @@ bool UWorld::HasEndOfFrameUpdates() const
 	return ComponentsThatNeedEndOfFrameUpdate_OnGameThread.Num() > 0 || ComponentsThatNeedEndOfFrameUpdate.Num() > 0 || bMaterialParameterCollectionInstanceNeedsDeferredUpdate;
 }
 
-struct FSendAllEndOfFrameUpdates
-{
-	FSendAllEndOfFrameUpdates(FSceneInterface* InScene)
-	{
-		if (InScene != nullptr)
-		{
-			GPUSkinCache = InScene->GetGPUSkinCache();
-			FeatureLevel = InScene->GetFeatureLevel();
-		}
-	}
-	
-	FGPUSkinCache* GPUSkinCache = nullptr;
-	ERHIFeatureLevel::Type FeatureLevel = ERHIFeatureLevel::Num;
-};
-
 DECLARE_GPU_STAT(EndOfFrameUpdates);
 DECLARE_GPU_STAT(GPUSkinCacheRayTracingGeometry);
 
@@ -1042,9 +1027,6 @@ void UWorld::SendAllEndOfFrameUpdates()
 	RHI_BREADCRUMB_EVENT_GAMETHREAD("SendAllEndOfFrameUpdates");
 #endif
 
-	// Issue a GPU event to wrap GPU work done during SendAllEndOfFrameUpdates, like skin cache updates
-	FSendAllEndOfFrameUpdates SendAllEndOfFrameUpdates(Scene);
-
 	// update all dirty components. 
 	FGuardValue_Bitfield(bPostTickComponentUpdate, true); 
 
