Optimization - InstancedStaticMeshComponent: Made the cached bounds property persistent (was transient) - to avoid recomputing the bounds upon loading.
* In a cooked build of City Sample, 100% of the calls to GetOrCalcBounds() are now hitting the cache
* In that sample, greatly helps reducing some of the peaks caused by bounds computation for massive ISMC (some with more than 15k instances). Bounds computation for those could previously take over 0.8ms

diff --git a/Engine/Source/Runtime/Engine/Classes/Components/InstancedStaticMeshComponent.h b/Engine/Source/Runtime/Engine/Classes/Components/InstancedStaticMeshComponent.h
--- a/Engine/Source/Runtime/Engine/Classes/Components/InstancedStaticMeshComponent.h
+++ b/Engine/Source/Runtime/Engine/Classes/Components/InstancedStaticMeshComponent.h
@@ -94,6 +94,64 @@ struct FInstancedStaticMeshRandomSeed
 	}
 };
 
+USTRUCT()
+struct FBoundsCacheElement
+{
+	GENERATED_USTRUCT_BODY()
+
+	FBoundsCacheElement()
+		: bIsValid(false)
+		, Hash(0)
+		, Value(ForceInit)
+	{
+	}
+
+	void Invalidate()
+	{
+		bIsValid = false;
+	}
+
+	bool IsValid() const
+	{
+		return bIsValid;
+	}
+
+	void CacheValue(const FBoxSphereBounds& InValue, uint32 InHash = 0)
+	{
+		bIsValid = true;
+		Value = InValue;
+		Hash = InHash;
+	}
+
+	const FBoxSphereBounds& GetValue(bool bCheckValid = true) const
+	{
+		ensure(!bCheckValid || bIsValid);
+		return Value;
+	}
+
+	bool operator == (uint32 InHash) const
+	{
+		return bIsValid && (Hash == InHash);
+	}
+
+	bool operator != (uint32 InHash) const
+	{
+		return !(*this == InHash);
+	}
+
+	static uint32 GetHash(const FTransform& InBoundTransform, const FBox& InInstanceBounds);
+
+private:
+	UPROPERTY()
+	bool bIsValid;
+
+	UPROPERTY()
+	uint32 Hash;
+
+	UPROPERTY()
+	FBoxSphereBounds Value;
+};
+
 /** A component that efficiently renders multiple instances of the same StaticMesh. */
 UCLASS(ClassGroup = Rendering, meta = (BlueprintSpawnableComponent), Blueprintable, MinimalAPI)
 class UInstancedStaticMeshComponent : public UStaticMeshComponent, public ISMInstanceManager
@@ -509,7 +567,8 @@ public:
 	//~ Begin UObject Interface
 	ENGINE_API virtual void Serialize(FArchive& Ar) override;
 	ENGINE_API virtual void GetResourceSizeEx(FResourceSizeEx& CumulativeResourceSize) override;
-	ENGINE_API void BeginDestroy() override;
+	ENGINE_API virtual void BeginDestroy() override;
+	ENGINE_API virtual void PreSave(FObjectPreSaveContext ObjectSaveContext) override;
 #if WITH_EDITOR
 	ENGINE_API virtual void PostStaticMeshCompilation() override;
 	ENGINE_API virtual void PostEditChangeChainProperty(FPropertyChangedChainEvent& PropertyChangedEvent) override;
@@ -631,6 +690,15 @@ protected:
 	/** Request to navigation system to update for the bounds of the ISM. */
 	ENGINE_API virtual void FullNavigationUpdate();
 
+	enum class EBoundsCacheType : uint8
+	{
+		Component,
+		Navigation,
+		Local,
+		Count,
+		Other
+	};
+
 	UE_DEPRECATED(5.5, "Use version with FBox")
 	FBoxSphereBounds CalcBoundsImpl(const FTransform& BoundTransform, bool bForNavigation) const;
 
@@ -711,69 +779,9 @@ protected:
 
 private:
 
-	struct FBoundsCacheElement
-	{
-		FBoundsCacheElement() 
-			: bIsValid(false)
-			, Hash(0)
-			, Value(ForceInit)
-		{}
-
-		void Invalidate()
-		{
-			bIsValid = false;
-		}
-
-		bool IsValid() const
-		{
-			return bIsValid;
-		}
-
-		void CacheValue(const FBoxSphereBounds& InValue, uint32 InHash = 0)
-		{
-			bIsValid = true;
-			Value = InValue;
-			Hash = InHash;
-		}
-
-		const FBoxSphereBounds& GetValue(bool bCheckValid = true) const
-		{
-			ensure(!bCheckValid || bIsValid);
-			return Value;
-		}
-
-		bool operator == (uint32 InHash) const
-		{
-			return bIsValid && (Hash == InHash);
-		}
-
-		bool operator != (uint32 InHash) const
-		{
-			return !(*this == InHash);
-		}
-
-		static uint32 GetHash(const FTransform& InBoundTransform, const FBox& InInstanceBounds)
-		{
-			return HashCombineFast(HashCombineFast(GetTypeHash(InBoundTransform), GetTypeHash(InInstanceBounds.Min)), GetTypeHash(InInstanceBounds.Max));
-		}
-
-	private:
-		bool bIsValid;
-		uint32 Hash;
-		FBoxSphereBounds Value;
-	};
-
-	enum class EBoundsCacheType : uint8
-	{
-		Component,
-		Navigation,
-		Local,
-		Count,
-		Other
-	};
-
 	FBoxSphereBounds GetOrCalcBounds(const FTransform& InBoundTransform, EBoundsCacheType InType) const;
 
+	UPROPERTY()
 	mutable FBoundsCacheElement CachedBounds[(int)EBoundsCacheType::Count];
 };
 
diff --git a/Engine/Source/Runtime/Engine/Private/InstancedStaticMesh.cpp b/Engine/Source/Runtime/Engine/Private/InstancedStaticMesh.cpp
--- a/Engine/Source/Runtime/Engine/Private/InstancedStaticMesh.cpp
+++ b/Engine/Source/Runtime/Engine/Private/InstancedStaticMesh.cpp
@@ -45,6 +45,7 @@
 #include "Elements/Framework/EngineElementsLibrary.h"
 #include "Elements/Interfaces/TypedElementWorldInterface.h"
 #include "UObject/UObjectIterator.h"
+#include "UObject/ObjectSaveContext.h"
 #include "GenericPlatform/ICursor.h"
 #include "NaniteVertexFactory.h"
 #include "InstancedStaticMeshSceneProxyDesc.h"
@@ -2812,6 +2813,22 @@ bool UInstancedStaticMeshComponent::CanEditSimulatePhysics()
 	return false;
 }
 
+uint32 FBoundsCacheElement::GetHash(const FTransform& InBoundTransform, const FBox& InInstanceBounds)
+{
+	// Compute a persistent hash as this value may be saved
+	FVector Translation = InBoundTransform.GetTranslation();
+	FQuat Rotation = InBoundTransform.GetRotation();
+	FVector Scale3D = InBoundTransform.GetScale3D();
+
+	uint64 Hash = 0;
+	Hash = CityHash64WithSeed((const char*)&Translation, sizeof(Translation), Hash);
+	Hash = CityHash64WithSeed((const char*)&Rotation, sizeof(Rotation), Hash);
+	Hash = CityHash64WithSeed((const char*)&Scale3D, sizeof(Scale3D), Hash);
+	Hash = CityHash64WithSeed((const char*)&InInstanceBounds.Min, sizeof(InInstanceBounds.Min), Hash);
+	Hash = CityHash64WithSeed((const char*)&InInstanceBounds.Max, sizeof(InInstanceBounds.Max), Hash);
+	return static_cast<uint32>(Hash) ^ static_cast<uint32>(Hash >> 32);
+}
+
 void UInstancedStaticMeshComponent::InvalidateCachedBounds()
 {
 	static_assert(UE_ARRAY_COUNT(CachedBounds) == int32(EBoundsCacheType::Count));
@@ -2823,14 +2840,12 @@ void UInstancedStaticMeshComponent::InvalidateCachedBounds()
 
 FBoxSphereBounds UInstancedStaticMeshComponent::CalcBounds(const FTransform& InBoundTransform) const
 {
-	TRACE_CPUPROFILER_EVENT_SCOPE(UInstancedStaticMeshComponent::CalcBounds);
 	EBoundsCacheType Type = bIsUpdatingBounds ? EBoundsCacheType::Component : InBoundTransform.Equals(FTransform::Identity) ? EBoundsCacheType::Local : EBoundsCacheType::Other;
 	return GetOrCalcBounds(InBoundTransform, Type);
 }
 
 void UInstancedStaticMeshComponent::CalcAndCacheNavigationBounds()
 {
-	TRACE_CPUPROFILER_EVENT_SCOPE(UInstancedStaticMeshComponent::CalcAndCacheNavigationBounds);
 	GetOrCalcBounds(GetComponentTransform(), EBoundsCacheType::Navigation);
 }
 
@@ -2843,7 +2858,8 @@ FBoxSphereBounds UInstancedStaticMeshComponent::GetOrCalcBounds(const FTransform
 		{
 			const UWorld* const World = GetWorld();
 			const bool bIsGameWorld = World && World->IsGameWorld();
-			if (!bIsGameWorld || InType == EBoundsCacheType::Other)
+			const bool bShouldCache = InType != EBoundsCacheType::Other && (IsRegistered() || bIsGameWorld);
+			if (!bShouldCache)
 			{
 				FBoxSphereBounds Value = CalcBoundsImpl(InBoundTransform, InstanceBounds);
 				if (InType == EBoundsCacheType::Navigation)
@@ -2897,6 +2913,7 @@ FBoxSphereBounds UInstancedStaticMeshComponent::CalcBoundsImpl(const FTransform&
 
 FBoxSphereBounds UInstancedStaticMeshComponent::CalcBoundsImpl(const FTransform& BoundTransform, const FBox& InstanceBounds) const
 {
+	TRACE_CPUPROFILER_EVENT_SCOPE(UInstancedStaticMeshComponent::CalcBoundsImpl);
 	check(PerInstanceSMData.Num() > 0);
 	const FMatrix BoundTransformMatrix = BoundTransform.ToMatrixWithScale();
 	FBoxSphereBounds::Builder BoundsBuilder;
@@ -3264,6 +3281,22 @@ bool UInstancedStaticMeshComponent::ShouldInheritPerInstanceData(const UInstance
 	return (bInheritPerInstanceData || !bEditableWhenInherited) && InArchetype && InArchetype->IsInBlueprint() && !IsTemplate();
 }
 
+void UInstancedStaticMeshComponent::PreSave(FObjectPreSaveContext ObjectSaveContext)
+{
+	Super::PreSave(ObjectSaveContext);
+
+	if (!IsTemplate() && IsRegistered())
+	{
+		GetOrCalcBounds(FTransform::Identity, EBoundsCacheType::Local);
+		GetOrCalcBounds(GetComponentTransform(), EBoundsCacheType::Component);
+		GetOrCalcBounds(GetComponentTransform(), EBoundsCacheType::Navigation);
+	}
+	else
+	{
+		InvalidateCachedBounds();
+	}
+}
+
 void UInstancedStaticMeshComponent::Serialize(FArchive& Ar)
 {
 	LLM_SCOPE(ELLMTag::InstancedMesh);
@@ -4995,6 +5028,15 @@ void UInstancedStaticMeshComponent::PostLoad()
 {
 	Super::PostLoad();
 
+	// If world is instanced, clear our cached bounds as they will likely not be valid
+	if (ULevel* Level = GetComponentLevel())
+	{
+		if (Level->IsInstancedLevel())
+		{
+			InvalidateCachedBounds();
+		}
+	}
+
 	// Ensure we have the proper amount of per instance custom float data.
 	// We have instances, and we have custom float data per instance, but we don't have the right amount of custom float data allocated.
 	if ((PerInstanceSMData.Num() * NumCustomDataFloats) != PerInstanceSMCustomData.Num() && PerInstanceSMData.Num() > 0)
@@ -5775,4 +5817,4 @@ void FInstancedStaticMeshVertexFactoryShaderParameters::GetElementShaderBindings
 	ShaderBindings.Add(InstancingOffsetParameter, InstancingOffset);
 
 	ShaderBindings.Add(Shader->GetUniformBufferParameter<FInstancedStaticMeshVFLooseUniformShaderParameters>(), BatchElement.LooseParametersUniformBuffer);
-}
\ No newline at end of file
+}
