Texture: Fixed the debug name for VT chunk cache requests

#rnx

diff --git a/Engine/Source/Runtime/Engine/Private/TextureDerivedData.cpp b/Engine/Source/Runtime/Engine/Private/TextureDerivedData.cpp
--- a/Engine/Source/Runtime/Engine/Private/TextureDerivedData.cpp
+++ b/Engine/Source/Runtime/Engine/Private/TextureDerivedData.cpp
@@ -2006,7 +2006,7 @@ static bool LoadDerivedStreamingVTChunks(const TArray<FVirtualTextureDataChunk>&
 		if (!Chunk.DerivedDataKey.IsEmpty() && !Chunk.BulkData.IsBulkDataLoaded())
 		{
 			FCacheGetValueRequest& Request = Requests.AddDefaulted_GetRef();
-			Request.Name = FSharedString(WriteToString<256>(DebugContext, TEXT(" [Chunk ]"), ChunkIndex, TEXT("]")));
+			Request.Name = FSharedString(WriteToString<256>(DebugContext, TEXT(" [Chunk "), ChunkIndex, TEXT("]")));
 			Request.Key = ConvertLegacyCacheKey(Chunk.DerivedDataKey);
 			Request.UserData = ChunkIndex;
 		}
