Fix a check comparing input device names which would always return true to actually be correct, and just compare the device id.

#rnx

diff --git a/Engine/Source/Runtime/Engine/Private/GameFramework/InputDeviceSubsystem.cpp b/Engine/Source/Runtime/Engine/Private/GameFramework/InputDeviceSubsystem.cpp
--- a/Engine/Source/Runtime/Engine/Private/GameFramework/InputDeviceSubsystem.cpp
+++ b/Engine/Source/Runtime/Engine/Private/GameFramework/InputDeviceSubsystem.cpp
@@ -557,47 +557,32 @@ FHardwareDeviceIdentifier UInputDeviceSubsystem::GetInputDeviceHardwareIdentifie
 
 void UInputDeviceSubsystem::SetMostRecentlyUsedHardwareDevice(const FInputDeviceId InDeviceId, const FHardwareDeviceIdentifier& InHardwareId)
 {	
+	// Store or update the most recent hardware information about this FInputDeviceId.
+	// We always want this to be up to date no matter what. This data doesn't care about the owning platform user,
+	// it is only about the actual Input Device itself.
+	LatestInputDeviceIdentifiers.Add(InDeviceId, InHardwareId);
+
+	// Next, we will track information per platform user about which input devices that user has been
+	// We want to broadcast the OnInputHardwareDeviceChanged delegate if the user's most recent input device
+	// has changed to something else.
 	const FPlatformUserId OwningUserId = IPlatformInputDeviceMapper::Get().GetUserForInputDevice(InDeviceId);
 	
-	// Keep track of each input device's latest hardware id so that it can be queried at any time
-	LatestInputDeviceIdentifiers.Add(InDeviceId, InHardwareId);
+	// Keep a map to platform users so that we can easily get their most recent hardware
+	FLatestInputDeviceData& LatestUserData = LatestUserDeviceIdentifiers.FindOrAdd(OwningUserId);
 
-	// If this hardware is the same as what the platform user already has, then there is no need to fire this event
-	if (FLatestInputDeviceData* ExistingData = LatestUserDeviceIdentifiers.Find(OwningUserId))
+	// Check the last known input device ID. If it is the same as the one we are currently
+	// processing then there is no need to broadcast the "changed" delegate,  because it hasnt.
+	bool bRequiresDelegateUpdate = true;
+	if (LatestUserData.LatestDeviceId == InDeviceId)
 	{
-		// We always want to update the device ID and info about the device
-		ExistingData->LatestDeviceId = InDeviceId;
-		ExistingData->LatestDevices.Add(InHardwareId.PrimaryDeviceType, InDeviceId);
-
-		// Some conditions to exit early to not broadcast the OnInputHardwareDeviceChanged delegate
-		// based on the current state of the known input devices
-		if (FHardwareDeviceIdentifier* ExistingDeviceId = LatestInputDeviceIdentifiers.Find(InDeviceId))
-		{
-			// If the hardware devices are the same, then there is no need to broadcast a "changed" delegate, because
-			// they have not changed.
-			if (InHardwareId == *ExistingDeviceId)
-			{
-				return;
-			}
-			// If they are the same device but different input API's, just update the data but don't broadcast the delegate
-			// This can happen if there are multiple input interfaces sending messages for the same device (such as a keyboard)
-			// Which your gameplay code doesn't really care about switching which API the input came from
-			else if (InHardwareId.HardwareDeviceIdentifier == ExistingDeviceId->HardwareDeviceIdentifier)
-			{
-				*ExistingDeviceId = InHardwareId;
-				return;
-			}	
-		}
+		bRequiresDelegateUpdate = false;
 	}
 
-	// Keep a map to platform users so that we can easily get their most recent hardware
-	FLatestInputDeviceData& LatestData = LatestUserDeviceIdentifiers.Add(OwningUserId);
-
 	// Keep track of the latest hardware device and the latest of this type for the given user
-	LatestData.LatestDevices.FindOrAdd(InHardwareId.PrimaryDeviceType, InDeviceId);
-	LatestData.LatestDeviceId = InDeviceId;
+	LatestUserData.LatestDevices.FindOrAdd(InHardwareId.PrimaryDeviceType, InDeviceId);
+	LatestUserData.LatestDeviceId = InDeviceId;
 	
-	if (OnInputHardwareDeviceChanged.IsBound())
+	if (OnInputHardwareDeviceChanged.IsBound() && bRequiresDelegateUpdate)
 	{
 		OnInputHardwareDeviceChanged.Broadcast(OwningUserId, InDeviceId);
 	}
