[HWRT] Disable ray tracing LOD bias when the skeletal mesh is using mesh deformer instead of skin cache otherwise RayTracingGeometry LODIndex will not match RayTracingLODIndex and cause issues. Add an assert to catch issues like this in the future.

#rnx

diff --git a/Engine/Source/Runtime/Engine/Private/SkeletalRenderGPUSkin.cpp b/Engine/Source/Runtime/Engine/Private/SkeletalRenderGPUSkin.cpp
--- a/Engine/Source/Runtime/Engine/Private/SkeletalRenderGPUSkin.cpp
+++ b/Engine/Source/Runtime/Engine/Private/SkeletalRenderGPUSkin.cpp
@@ -866,6 +866,7 @@ void FSkeletalMeshObjectGPUSkin::UpdateRayTracingGeometry_Internal(
 {
 	if (IsRayTracingEnabled() && MeshObject->bSupportRayTracing)
 	{
+		check(LODIndex == MeshObject->GetRayTracingLOD());
 		bool bRequireRecreatingRayTracingGeometry = LODIndex != RayTracingGeometry.LODIndex
 			|| MeshObject->bHiddenMaterialVisibilityDirtyForRayTracing
 			|| RayTracingGeometry.Initializer.Segments.Num() == 0;
@@ -2465,9 +2466,15 @@ void FDynamicSkelMeshObjectDataGPUSkin::InitDynamicSkelMeshObjectDataGPUSkin(
 	}
 
 #if RHI_RAYTRACING
+	int32 RayTracingLODBias = GetRayTracingSkeletalMeshGlobalLODBias();
+	// TODO: MeshDeformer only supports using the same LOD as rendering so we have to disable ray tracing LOD bias
+	if (InMeshObject->GetGPUSkinTechnique(LODIndex) == ESkeletalMeshGPUSkinTechnique::MeshDeformer)
+	{
+		RayTracingLODBias = 0;
+	}
 	// If Proxy is not visible in raytracing scene, set RayTracingLODIndex to -1 which means that additional RT update will not be performed for that mesh object, probably RT proxy handles that 
 	const bool bVisibleInRayTracing = SkeletalMeshProxy ? SkeletalMeshProxy->IsVisibleInRayTracing() : true;
-	RayTracingLODIndex = bVisibleInRayTracing ? FMath::Clamp(FMath::Max(LODIndex + GetRayTracingSkeletalMeshGlobalLODBias(), InMeshObject->RayTracingMinLOD), LODIndex, InSkeletalMeshRenderData->LODRenderData.Num() - 1) : -1;
+	RayTracingLODIndex = bVisibleInRayTracing ? FMath::Clamp(FMath::Max(LODIndex + RayTracingLODBias, InMeshObject->RayTracingMinLOD), LODIndex, InSkeletalMeshRenderData->LODRenderData.Num() - 1) : -1;
 #endif
 
 #if !(UE_BUILD_SHIPPING || UE_BUILD_TEST)
