FBodyInstance::ActorHandle accessors without non-const references. Old API to be deprecated in later CL.

diff --git a/Engine/Source/Runtime/Engine/Classes/PhysicsEngine/BodyInstance.h b/Engine/Source/Runtime/Engine/Classes/PhysicsEngine/BodyInstance.h
--- a/Engine/Source/Runtime/Engine/Classes/PhysicsEngine/BodyInstance.h
+++ b/Engine/Source/Runtime/Engine/Classes/PhysicsEngine/BodyInstance.h
@@ -658,12 +658,19 @@ public:
 public:
 
 	ENGINE_API UBodySetup* GetBodySetup() const;
-	
+
+	ENGINE_API FPhysicsActorHandle GetPhysicsActor() const;
+	ENGINE_API FPhysicsActorHandle GetPhysicsActorWeldRoot() const;
+	ENGINE_API void SetPhysicsActor(FPhysicsActorHandle InHandle);
+	ENGINE_API void ReleasePhysicsActor();
+
+	// TODO_CHAOSAPI: Deprecate ActorHandle Ref API
 	ENGINE_API FPhysicsActorHandle& GetPhysicsActorHandle();
 	ENGINE_API const FPhysicsActorHandle& GetPhysicsActorHandle() const;
 	ENGINE_API const FPhysicsActorHandle& GetActorReferenceWithWelding() const;
-	ENGINE_API void SetPhysicsActorHandle(FPhysicsActorHandle InHandle);
+	inline void SetPhysicsActorHandle(FPhysicsActorHandle InHandle) { SetPhysicsActor(InHandle); }
 
+	// TODO_CHAOSAPI: Deprecate public ActorHandle
 	// Internal physics representation of our body instance
 	FPhysicsActorHandle ActorHandle;
 
diff --git a/Engine/Source/Runtime/Engine/Private/PhysicsEngine/BodyInstance.cpp b/Engine/Source/Runtime/Engine/Private/PhysicsEngine/BodyInstance.cpp
--- a/Engine/Source/Runtime/Engine/Private/PhysicsEngine/BodyInstance.cpp
+++ b/Engine/Source/Runtime/Engine/Private/PhysicsEngine/BodyInstance.cpp
@@ -403,11 +403,6 @@ FBodyInstance::FBodyInstance()
 
 FBodyInstance::~FBodyInstance() = default;
 
-const FPhysicsActorHandle& FBodyInstance::GetActorReferenceWithWelding() const
-{
-	return WeldParent ? WeldParent->GetPhysicsActorHandle() : GetPhysicsActorHandle();
-}
-
 FArchive& operator<<(FArchive& Ar,FBodyInstance& BodyInst)
 {
 	if (!Ar.IsLoading() && !Ar.IsSaving())
@@ -473,9 +468,9 @@ const TMap<FPhysicsShapeHandle, FBodyInstance::FWeldInfo>* FBodyInstance::GetCur
 
 int32 FBodyInstance::GetAllShapes_AssumesLocked(TArray<FPhysicsShapeHandle>& OutShapes) const
 {
-	if(GetPhysicsActorHandle())
+	if(GetPhysicsActor())
 	{
-		return FPhysicsInterface::GetAllShapes_AssumedLocked(GetPhysicsActorHandle(), OutShapes);
+		return FPhysicsInterface::GetAllShapes_AssumedLocked(GetPhysicsActor(), OutShapes);
 	}
 
 	return 0;
@@ -485,7 +480,7 @@ void FBodyInstance::UpdateTriMeshVertices(const TArray<FVector> & NewPositions)
 {
 	if (BodySetup.IsValid())
 	{
-		FPhysicsCommand::ExecuteWrite(GetPhysicsActorHandle(), [&](const FPhysicsActorHandle& Actor)
+		FPhysicsCommand::ExecuteWrite(GetPhysicsActor(), [&](FPhysicsActorHandle Actor)
 		{
 			//after updating the vertices we must call setGeometry again to update any shapes referencing the mesh
 			TArray<FPhysicsShapeHandle> Shapes;
@@ -544,7 +539,7 @@ void FBodyInstance::UpdatePhysicalMaterials()
 	TArray<FPhysicalMaterialMaskParams> ComplexPhysMatMasks;
 	TArray<UPhysicalMaterial*> ComplexPhysMats = GetComplexPhysicalMaterials(ComplexPhysMatMasks);
 
-	FPhysicsCommand::ExecuteWrite(GetActorReferenceWithWelding(), [&](const FPhysicsActorHandle& Actor)
+	FPhysicsCommand::ExecuteWrite(GetPhysicsActorWeldRoot(), [&](const FPhysicsActorHandle& Actor)
 	{
 		ApplyMaterialToInstanceShapes_AssumesLocked(SimplePhysMat, ComplexPhysMats, ComplexPhysMatMasks);
 	});
@@ -966,7 +961,7 @@ void FBodyInstance::SetMaskFilter(FMaskFilter InMaskFilter)
 		return;
 	}
 
-	FPhysicsCommand::ExecuteWrite(GetPhysicsActorHandle(), [&](const FPhysicsActorHandle& Actor)
+	FPhysicsCommand::ExecuteWrite(GetPhysicsActor(), [&](const FPhysicsActorHandle& Actor)
 	{
 		TArray<FPhysicsShapeHandle> Shapes;
 		FPhysicsInterface::GetAllShapes_AssumedLocked(Actor, Shapes);
@@ -1011,10 +1006,10 @@ void FBodyInstance::UpdatePhysicsFilterData()
 		return;
 	}
 
-	FPhysicsCommand::ExecuteWrite(GetActorReferenceWithWelding(), [&](const FPhysicsActorHandle& Actor)
+	FPhysicsCommand::ExecuteWrite(GetPhysicsActorWeldRoot(), [&](const FPhysicsActorHandle& Actor)
 	{
 		TArray<FPhysicsShapeHandle> AllShapes;
-		const int32 NumSyncShapes = FPhysicsInterface::GetAllShapes_AssumedLocked(GetPhysicsActorHandle(), AllShapes);
+		const int32 NumSyncShapes = FPhysicsInterface::GetAllShapes_AssumedLocked(GetPhysicsActor(), AllShapes);
 		const int32 NumTotalShapes = AllShapes.Num();
 		// In skeletal case, collision enable/disable/movement should be overriden by mesh component
 		FBodyCollisionData BodyCollisionData;
@@ -1271,7 +1266,7 @@ template<typename InAllocatorType>
 void TInitBodiesHelperBase<InAllocatorType>::CreateActor_AssumesLocked(FBodyInstance* Instance, const FTransform& Transform) const
 {
 	SCOPE_CYCLE_COUNTER(STAT_CreatePhysicsActor);
-	checkSlow(!FPhysicsInterface::IsValid(Instance->GetPhysicsActorHandle()));
+	checkSlow(!FPhysicsInterface::IsValid(Instance->GetPhysicsActor()));
 	const ECollisionEnabled::Type CollisionType = Instance->GetCollisionEnabled();
 
 	FActorCreationParams ActorParams;
@@ -1285,30 +1280,28 @@ void TInitBodiesHelperBase<InAllocatorType>::CreateActor_AssumesLocked(FBodyInst
 	ActorParams.bUpdateKinematicFromSimulation = Instance->bUpdateKinematicFromSimulation;
 	ActorParams.bQueryOnly = CollisionType == ECollisionEnabled::QueryOnly;
 	ActorParams.Scene = PhysScene;
+	ActorParams.bStatic = IsStatic();
 
-	if (IsStatic())
-	{
-		ActorParams.bStatic = true;
+	FPhysicsActorHandle ActorHandle = nullptr;
+	FPhysicsInterface::CreateActor(ActorParams, ActorHandle, PrimitiveComp);
+	Instance->SetPhysicsActor(ActorHandle);
 
-		FPhysicsInterface::CreateActor(ActorParams, Instance->GetPhysicsActorHandle(), PrimitiveComp);
-	}
-	else
+	if (!IsStatic())
 	{
-		FPhysicsInterface::CreateActor(ActorParams, Instance->GetPhysicsActorHandle(), PrimitiveComp);
-		FPhysicsInterface::SetCcdEnabled_AssumesLocked(Instance->GetPhysicsActorHandle(), Instance->bUseCCD);
-		FPhysicsInterface::SetMACDEnabled_AssumesLocked(Instance->GetPhysicsActorHandle(), Instance->IsUsingMACD());
-		FPhysicsInterface::SetIsKinematic_AssumesLocked(Instance->GetPhysicsActorHandle(), !Instance->ShouldInstanceSimulatingPhysics());
+		FPhysicsInterface::SetCcdEnabled_AssumesLocked(ActorHandle, Instance->bUseCCD);
+		FPhysicsInterface::SetMACDEnabled_AssumesLocked(ActorHandle, Instance->IsUsingMACD());
+		FPhysicsInterface::SetIsKinematic_AssumesLocked(ActorHandle, !Instance->ShouldInstanceSimulatingPhysics());
 
-		FPhysicsInterface::SetMaxLinearVelocity_AssumesLocked(Instance->GetPhysicsActorHandle(), TNumericLimits<float>::Max());
-		FPhysicsInterface::SetSmoothEdgeCollisionsEnabled_AssumesLocked(Instance->GetPhysicsActorHandle(), Instance->bSmoothEdgeCollisions);
-		FPhysicsInterface::SetInertiaConditioningEnabled_AssumesLocked(Instance->GetPhysicsActorHandle(), Instance->IsInertiaConditioningEnabled());
-		FPhysicsInterface::SetGyroscopicTorqueEnabled_AssumesLocked(Instance->GetPhysicsActorHandle(), Instance->bGyroscopicTorqueEnabled);
-		FPhysicsInterface::SetPositionSolverIterationCount_AssumesLocked(Instance->GetPhysicsActorHandle(), Instance->PositionSolverIterationCount);
-		FPhysicsInterface::SetVelocitySolverIterationCount_AssumesLocked(Instance->GetPhysicsActorHandle(), Instance->VelocitySolverIterationCount);
-		FPhysicsInterface::SetProjectionSolverIterationCount_AssumesLocked(Instance->GetPhysicsActorHandle(), Instance->ProjectionSolverIterationCount);
+		FPhysicsInterface::SetMaxLinearVelocity_AssumesLocked(ActorHandle, TNumericLimits<float>::Max());
+		FPhysicsInterface::SetSmoothEdgeCollisionsEnabled_AssumesLocked(ActorHandle, Instance->bSmoothEdgeCollisions);
+		FPhysicsInterface::SetInertiaConditioningEnabled_AssumesLocked(ActorHandle, Instance->IsInertiaConditioningEnabled());
+		FPhysicsInterface::SetGyroscopicTorqueEnabled_AssumesLocked(ActorHandle, Instance->bGyroscopicTorqueEnabled);
+		FPhysicsInterface::SetPositionSolverIterationCount_AssumesLocked(ActorHandle, Instance->PositionSolverIterationCount);
+		FPhysicsInterface::SetVelocitySolverIterationCount_AssumesLocked(ActorHandle, Instance->VelocitySolverIterationCount);
+		FPhysicsInterface::SetProjectionSolverIterationCount_AssumesLocked(ActorHandle, Instance->ProjectionSolverIterationCount);
 
 		// Set sleep event notification
-		FPhysicsInterface::SetSendsSleepNotifies_AssumesLocked(Instance->GetPhysicsActorHandle(), Instance->bGenerateWakeEvents);
+		FPhysicsInterface::SetSendsSleepNotifies_AssumesLocked(ActorHandle, Instance->bGenerateWakeEvents);
 	}
 }
 
@@ -1331,9 +1324,9 @@ bool TInitBodiesHelperBase<InAllocatorType>::CreateShapes_AssumesLocked(FBodyIns
 	// #PHYS2 Call interface AddGeometry
 	BodySetup->AddShapesToRigidActor_AssumesLocked(Instance, Instance->Scale3D, SimplePhysMat, ComplexPhysMats, ComplexPhysMatMasks, BodyCollisionData, FTransform::Identity);
 
-	FPhysicsInterface::SetIgnoreAnalyticCollisions_AssumesLocked(Instance->GetPhysicsActorHandle(), CVarIgnoreAnalyticCollisionsOverride.GetValueOnAnyThread() ? true : Instance->bIgnoreAnalyticCollisions);
+	FPhysicsInterface::SetIgnoreAnalyticCollisions_AssumesLocked(Instance->GetPhysicsActor(), CVarIgnoreAnalyticCollisionsOverride.GetValueOnAnyThread() ? true : Instance->bIgnoreAnalyticCollisions);
 
-	const int32 NumShapes = FPhysicsInterface::GetNumShapes(Instance->GetPhysicsActorHandle());
+	const int32 NumShapes = FPhysicsInterface::GetNumShapes(Instance->GetPhysicsActor());
 	bInitFail |= NumShapes == 0;
 
 	return bInitFail;
@@ -1502,7 +1495,7 @@ bool TInitBodiesHelperBase<InAllocatorType>::CreateShapesAndActors()
 
 		// Don't process if we've already got a body
 		// Just ask actorref
-		if (FPhysicsInterface::IsValid(Instance->GetPhysicsActorHandle()))
+		if (FPhysicsInterface::IsValid(Instance->GetPhysicsActor()))
 		{
 			Instance->OwnerComponent = nullptr;
 			Instance->SourceObject = nullptr;
@@ -1538,7 +1531,7 @@ bool TInitBodiesHelperBase<InAllocatorType>::CreateShapesAndActors()
 				UE_LOG(LogPhysics, Log, TEXT("Init Instance %d of Primitive Component %s failed. Does it have collision data available?"), BodyIdx, *PrimitiveComp->GetReadableName());
 			}
 
-			FPhysicsInterface::ReleaseActor(Instance->GetPhysicsActorHandle(), PhysScene);
+			Instance->ReleasePhysicsActor();
 
 			Instance->OwnerComponent = nullptr;
 			Instance->SourceObject = nullptr;
@@ -1548,10 +1541,10 @@ bool TInitBodiesHelperBase<InAllocatorType>::CreateShapesAndActors()
 			continue;
 		}
 
-		FPhysicsInterface::SetActorUserData_AssumesLocked(Instance->GetPhysicsActorHandle(), &Instance->PhysicsUserData);
+		FPhysicsInterface::SetActorUserData_AssumesLocked(Instance->GetPhysicsActor(), &Instance->PhysicsUserData);
 
 #if USE_BODYINSTANCE_DEBUG_NAMES
-		Instance->GetPhysicsActorHandle()->GetParticle_LowLevel()->SetDebugName(DebugName);
+		Instance->GetPhysicsActor()->GetParticle_LowLevel()->SetDebugName(DebugName);
 #endif
 	}
 	return true;
@@ -1578,7 +1571,7 @@ void TInitBodiesHelperBase<InAllocatorType>::InitBodies()
 				SCOPE_CYCLE_COUNTER(STAT_InitBodyAggregate);
 				for (FBodyInstance* BI : Bodies)
 				{
-					const FPhysicsActorHandle& ActorHandle = BI->GetPhysicsActorHandle();
+					const FPhysicsActorHandle& ActorHandle = BI->GetPhysicsActor();
 					if (FPhysicsInterface::IsValid(ActorHandle))
 					{
 						FPhysicsInterface::AddActorToAggregate_AssumesLocked(Aggregate, ActorHandle);
@@ -1593,7 +1586,7 @@ void TInitBodiesHelperBase<InAllocatorType>::InitBodies()
 
 				for (FBodyInstance* BI : Bodies)
 				{
-					FPhysicsActorHandle& ActorHandle = BI->GetPhysicsActorHandle();
+					FPhysicsActorHandle ActorHandle = BI->GetPhysicsActor();
 					if (FPhysicsInterface::IsValid(ActorHandle))
 					{
 						ActorHandles.Add(ActorHandle);
@@ -1624,7 +1617,7 @@ void TInitBodiesHelperBase<InAllocatorType>::InitBodies()
 
 				for (FBodyInstance* BI : Bodies)
 				{
-					FPhysicsActorHandle& ActorHandle = BI->GetPhysicsActorHandle();
+					FPhysicsActorHandle ActorHandle = BI->GetPhysicsActor();
 					if (FPhysicsInterface::IsValid(ActorHandle))
 					{
 						PhysScene->AddToComponentMaps(BI->OwnerComponent.Get(), ActorHandle);
@@ -1771,10 +1764,10 @@ FBodyInstance::FAsyncTermBodyPayload FBodyInstance::StartAsyncTermBody_GameThrea
 
 	// Prepare payload for async task before resetting the ActorHandle/OwnerComponent
 	FAsyncTermBodyPayload Payload;
-	Payload.ActorHandle = GetPhysicsActorHandle();
+	Payload.ActorHandle = GetPhysicsActor();
 	Payload.OwnerComponent = OwnerComponent;
 	Payload.bNotifyRigidBodyCollision = bNotifyRigidBodyCollision;
-	SetPhysicsActorHandle(nullptr);
+	SetPhysicsActor(nullptr);
 	OwnerComponent = nullptr;
 
 	CurrentSceneState = BodyInstanceSceneState::NotAdded;
@@ -1829,9 +1822,9 @@ void FBodyInstance::TermBody(bool bNeverDeferRelease)
 	{
 		if (FPhysScene_Chaos* PhysScene = PrimComp->GetWorld()->GetPhysicsScene())
 		{
-			if (FPhysicsInterface::IsValid(GetPhysicsActorHandle()))
+			if (FPhysicsInterface::IsValid(GetPhysicsActor()))
 			{
-				PhysScene->RemoveFromComponentMaps(GetPhysicsActorHandle());
+				PhysScene->RemoveFromComponentMaps(GetPhysicsActor());
 			}
 			if (bNotifyRigidBodyCollision)
 			{
@@ -1842,7 +1835,7 @@ void FBodyInstance::TermBody(bool bNeverDeferRelease)
 
 	if (IsValidBodyInstance())
 	{
-		FPhysicsInterface::ReleaseActor(GetPhysicsActorHandle(), GetPhysicsScene(), bNeverDeferRelease);
+		ReleasePhysicsActor();
 	}
 
 	// @TODO: Release spring body here
@@ -1881,7 +1874,7 @@ bool FBodyInstance::Weld(FBodyInstance* TheirBody, const FTransform& TheirTM)
 
 	FTransform RelativeTM = TheirTM.GetRelativeTransform(MyTM);
 
-	FPhysicsCommand::ExecuteWrite(GetPhysicsActorHandle(), [&](const FPhysicsActorHandle& Actor)
+	FPhysicsCommand::ExecuteWrite(GetPhysicsActor(), [&](const FPhysicsActorHandle& Actor)
 	{
 		SCOPE_CYCLE_COUNTER(STAT_UpdatePhysMats);
 
@@ -1938,7 +1931,7 @@ bool FBodyInstance::Weld(FBodyInstance* TheirBody, const FTransform& TheirTM)
 		PostShapeChange();
 
 		// remove their body from scenes (don't call TermBody because we don't want to clear things like BodySetup)
-		FPhysicsInterface::ReleaseActor(TheirBody->GetPhysicsActorHandle(), TheirBody->GetPhysicsScene());
+		FPhysicsInterface::ReleaseActor(TheirBody->ActorHandle, TheirBody->GetPhysicsScene());
 	});
 	
 	UpdateInterpolateWhenSubStepping();
@@ -1958,7 +1951,7 @@ void FBodyInstance::UnWeld(FBodyInstance* TheirBI)
 
 	bool bShapesChanged = false;
 
-	FPhysicsCommand::ExecuteWrite(GetPhysicsActorHandle(), [&](const FPhysicsActorHandle& Actor)
+	FPhysicsCommand::ExecuteWrite(GetPhysicsActor(), [&](const FPhysicsActorHandle& Actor)
 				{
 		TArray<FPhysicsShapeHandle> Shapes;
 		const int32 NumSyncShapes = GetAllShapes_AssumesLocked(Shapes);
@@ -2152,7 +2145,7 @@ bool FBodyInstance::UpdateBodyScale(const FVector& InScale3D, bool bForceUpdate)
 	//Get all shapes
 	EScaleMode::Type ScaleMode = EScaleMode::Free;
 
-	FPhysicsCommand::ExecuteWrite(GetPhysicsActorHandle(), [&](const FPhysicsActorHandle& Actor)
+	FPhysicsCommand::ExecuteWrite(GetPhysicsActor(), [&](const FPhysicsActorHandle& Actor)
 	{
 		TArray<FPhysicsShapeHandle> Shapes;
 		GetAllShapes_AssumesLocked(Shapes);
@@ -2466,8 +2459,8 @@ bool FBodyInstance::UpdateBodyScale(const FVector& InScale3D, bool bForceUpdate)
 		// Only follow through with update if all shapes succeeded.
 		if (CHAOS_ENSURE(NewGeometry.Num() == Shapes.Num()))
 		{
-			GetPhysicsActorHandle()->GetGameThreadAPI().SetGeometry(MakeImplicitObjectPtr<Chaos::FImplicitObjectUnion>(MoveTemp(NewGeometry)));
-			FPhysicsInterface::WakeUp_AssumesLocked(GetPhysicsActorHandle());
+			GetPhysicsActor()->GetGameThreadAPI().SetGeometry(MakeImplicitObjectPtr<Chaos::FImplicitObjectUnion>(MoveTemp(NewGeometry)));
+			FPhysicsInterface::WakeUp_AssumesLocked(GetPhysicsActor());
 		}
 		else
 		{
@@ -2481,7 +2474,7 @@ bool FBodyInstance::UpdateBodyScale(const FVector& InScale3D, bool bForceUpdate)
 
 		if (FPhysScene_Chaos* Scene = GetPhysicsScene())
 		{
-			Scene->UpdateActorInAccelerationStructure(GetPhysicsActorHandle());
+			Scene->UpdateActorInAccelerationStructure(GetPhysicsActor());
 		}
 
 		// update mass if required
@@ -2507,7 +2500,7 @@ void FBodyInstance::UpdateInstanceSimulatePhysics()
 	const bool bUseSimulate = IsInstanceSimulatingPhysics();
 	bool bInitialized = false;
 
-	FPhysicsCommand::ExecuteWrite(GetPhysicsActorHandle(), [&](const FPhysicsActorHandle& Actor)
+	FPhysicsCommand::ExecuteWrite(GetPhysicsActor(), [&](const FPhysicsActorHandle& Actor)
 	{
 		bInitialized = true;
 		FPhysicsInterface::SetIsKinematic_AssumesLocked(Actor, !bUseSimulate);
@@ -2555,7 +2548,7 @@ bool FBodyInstance::IsNonKinematic() const
 
 bool FBodyInstance::IsDynamic() const
 {
-	return FPhysicsInterface::IsDynamic(GetPhysicsActorHandle());
+	return FPhysicsInterface::IsDynamic(GetPhysicsActor());
 }
 
 void FBodyInstance::ApplyWeldOnChildren()
@@ -2614,7 +2607,7 @@ void FBodyInstance::SetInstanceSimulatePhysics(bool bSimulate, bool bMaintainPhy
 				FMessageLog("PIE").Warning(FText::Format(LOCTEXT("SimPhysNoBody", "Trying to simulate physics on ''{0}'' but no physics body."),
 					FText::FromString(GetPathNameSafe(OwnerComponentInst))));
 			}
-			else if (FPhysicsInterface::IsStatic(GetPhysicsActorHandle()))
+			else if (FPhysicsInterface::IsStatic(GetPhysicsActor()))
 			{
 				FMessageLog("PIE").Warning(FText::Format(LOCTEXT("SimPhysStatic", "Trying to simulate physics on ''{0}'' but it is static."),
 					FText::FromString(GetPathNameSafe(OwnerComponentInst))));
@@ -2646,7 +2639,7 @@ void FBodyInstance::SetInstanceSimulatePhysics(bool bSimulate, bool bMaintainPhy
 
 bool FBodyInstance::IsValidBodyInstance() const
 {
-	return FPhysicsInterface::IsValid(GetPhysicsActorHandle());
+	return FPhysicsInterface::IsValid(GetPhysicsActor());
 }
 
 FTransform GetUnrealWorldTransformImp_AssumesLocked(const FBodyInstance* BodyInstance, bool bWithProjection, bool bGlobalPose)
@@ -2655,7 +2648,7 @@ FTransform GetUnrealWorldTransformImp_AssumesLocked(const FBodyInstance* BodyIns
 
 	if(BodyInstance && BodyInstance->IsValidBodyInstance())
 	{
-		WorldTM = FPhysicsInterface::GetTransform_AssumesLocked(BodyInstance->GetPhysicsActorHandle(), bGlobalPose);
+		WorldTM = FPhysicsInterface::GetTransform_AssumesLocked(BodyInstance->GetPhysicsActor(), bGlobalPose);
 
 		if(bWithProjection)
 		{
@@ -2670,7 +2663,7 @@ FTransform FBodyInstance::GetUnrealWorldTransform(bool bWithProjection /* = true
 {
 	FTransform OutTransform = FTransform::Identity;
 
-	FPhysicsCommand::ExecuteRead(GetPhysicsActorHandle(), [&](const FPhysicsActorHandle& Actor)
+	FPhysicsCommand::ExecuteRead(GetPhysicsActor(), [&](const FPhysicsActorHandle& Actor)
 	{
 		OutTransform = GetUnrealWorldTransformImp_AssumesLocked(this, bWithProjection, bForceGlobalPose);
 	});
@@ -2703,7 +2696,7 @@ void FBodyInstance::SetBodyTransform(const FTransform& NewTransform, ETeleportTy
 		return;
 	}
 
-	if(FPhysicsInterface::IsValid(GetPhysicsActorHandle()))
+	if(FPhysicsInterface::IsValid(GetPhysicsActor()))
 	{
 		if(!NewTransform.IsValid())
 		{
@@ -2721,9 +2714,9 @@ void FBodyInstance::SetBodyTransform(const FTransform& NewTransform, ETeleportTy
 
 		FPhysScene* Scene = GetPhysicsScene();
 
-		if(FPhysicsInterface::IsDynamic(GetPhysicsActorHandle()) && !bEditorWorld && Scene)
+		if(FPhysicsInterface::IsDynamic(GetPhysicsActor()) && !bEditorWorld && Scene)
 		{
-			FPhysicsCommand::ExecuteWrite(GetPhysicsActorHandle(), [&](const FPhysicsActorHandle& Actor)
+			FPhysicsCommand::ExecuteWrite(GetPhysicsActor(), [&](const FPhysicsActorHandle& Actor)
 			{
 				const bool bKinematic = FPhysicsInterface::IsKinematic_AssumesLocked(Actor);
 				const bool bSimulated = FPhysicsInterface::CanSimulate_AssumesLocked(Actor);
@@ -2747,7 +2740,7 @@ void FBodyInstance::SetBodyTransform(const FTransform& NewTransform, ETeleportTy
 		}
 		else if(Scene)
 		{
-			FPhysicsCommand::ExecuteWrite(GetPhysicsActorHandle(), [&](const FPhysicsActorHandle& Actor)
+			FPhysicsCommand::ExecuteWrite(GetPhysicsActor(), [&](const FPhysicsActorHandle& Actor)
 			{
 				FPhysicsInterface::SetGlobalPose_AssumesLocked(Actor, NewTransform);
 			});
@@ -2769,7 +2762,7 @@ FVector FBodyInstance::GetUnrealWorldVelocity() const
 {
 	FVector OutVelocity = FVector::ZeroVector;
 
-	FPhysicsCommand::ExecuteRead(GetPhysicsActorHandle(), [&](const FPhysicsActorHandle& Actor)
+	FPhysicsCommand::ExecuteRead(GetPhysicsActor(), [&](const FPhysicsActorHandle& Actor)
 	{
 		OutVelocity = GetUnrealWorldVelocity_AssumesLocked();
 	});
@@ -2780,9 +2773,9 @@ FVector FBodyInstance::GetUnrealWorldVelocity() const
 FTransform FBodyInstance::GetKinematicTarget_AssumesLocked() const
 {
 	FTransform TM;
-	if (FPhysicsInterface::IsValid(GetPhysicsActorHandle()))
+	if (FPhysicsInterface::IsValid(GetPhysicsActor()))
 	{
-		TM = FPhysicsInterface::GetKinematicTarget_AssumesLocked(GetPhysicsActorHandle());
+		TM = FPhysicsInterface::GetKinematicTarget_AssumesLocked(GetPhysicsActor());
 	}
 	return TM;
 }
@@ -2790,7 +2783,7 @@ FTransform FBodyInstance::GetKinematicTarget_AssumesLocked() const
 FTransform FBodyInstance::GetKinematicTarget() const
 {
 	FTransform TM;
-	FPhysicsCommand::ExecuteRead(GetPhysicsActorHandle(), [&](const FPhysicsActorHandle& Actor)
+	FPhysicsCommand::ExecuteRead(GetPhysicsActor(), [&](const FPhysicsActorHandle& Actor)
 		{
 			TM = FPhysicsInterface::GetKinematicTarget_AssumesLocked(Actor);
 		});
@@ -2801,9 +2794,9 @@ FTransform FBodyInstance::GetKinematicTarget() const
 FVector FBodyInstance::GetUnrealWorldVelocity_AssumesLocked() const
 {
 	FVector LinVel(EForceInit::ForceInitToZero);
-	if (FPhysicsInterface::IsValid(GetPhysicsActorHandle()))
+	if (FPhysicsInterface::IsValid(GetPhysicsActor()))
 	{
-		LinVel = FPhysicsInterface::GetLinearVelocity_AssumesLocked(GetPhysicsActorHandle());
+		LinVel = FPhysicsInterface::GetLinearVelocity_AssumesLocked(GetPhysicsActor());
 	}
 
 	return LinVel;
@@ -2814,7 +2807,7 @@ FVector FBodyInstance::GetUnrealWorldAngularVelocityInRadians() const
 {
 	FVector OutVelocity = FVector::ZeroVector;
 
-	FPhysicsCommand::ExecuteRead(GetPhysicsActorHandle(), [&](const FPhysicsActorHandle& Actor)
+	FPhysicsCommand::ExecuteRead(GetPhysicsActor(), [&](const FPhysicsActorHandle& Actor)
 	{
 		OutVelocity = FPhysicsInterface::GetAngularVelocity_AssumesLocked(Actor);
 	});
@@ -2825,13 +2818,13 @@ FVector FBodyInstance::GetUnrealWorldAngularVelocityInRadians() const
 /** Note: returns angular velocity in radians per second. */
 FVector FBodyInstance::GetUnrealWorldAngularVelocityInRadians_AssumesLocked() const
 {
-	return FPhysicsInterface::GetAngularVelocity_AssumesLocked(GetPhysicsActorHandle());
+	return FPhysicsInterface::GetAngularVelocity_AssumesLocked(GetPhysicsActor());
 }
 
 FVector FBodyInstance::GetUnrealWorldVelocityAtPoint(const FVector& Point) const
 {
 	FVector OutVelocity = FVector::ZeroVector;
-	FPhysicsCommand::ExecuteRead(GetPhysicsActorHandle(), [&](const FPhysicsActorHandle& Actor)
+	FPhysicsCommand::ExecuteRead(GetPhysicsActor(), [&](const FPhysicsActorHandle& Actor)
 	{
 		OutVelocity = FPhysicsInterface::GetWorldVelocityAtPoint_AssumesLocked(Actor, Point);
 	});
@@ -2842,14 +2835,14 @@ FVector FBodyInstance::GetUnrealWorldVelocityAtPoint(const FVector& Point) const
 
 FVector FBodyInstance::GetUnrealWorldVelocityAtPoint_AssumesLocked(const FVector& Point) const
 {
-	return FPhysicsInterface::GetWorldVelocityAtPoint_AssumesLocked(GetPhysicsActorHandle(), Point);
+	return FPhysicsInterface::GetWorldVelocityAtPoint_AssumesLocked(GetPhysicsActor(), Point);
 }
 
 FTransform FBodyInstance::GetMassSpaceToWorldSpace() const
 {
 	FTransform MassSpaceToWorldSpace = FTransform::Identity;
 
-	FPhysicsCommand::ExecuteRead(GetPhysicsActorHandle(), [&](const FPhysicsActorHandle& Actor)
+	FPhysicsCommand::ExecuteRead(GetPhysicsActor(), [&](const FPhysicsActorHandle& Actor)
 		{
 		MassSpaceToWorldSpace = FPhysicsInterface::GetComTransform_AssumesLocked(Actor);
 	});
@@ -2861,7 +2854,7 @@ FTransform FBodyInstance::GetMassSpaceLocal() const
 {
 	FTransform MassSpaceLocal = FTransform::Identity;
 
-	FPhysicsCommand::ExecuteRead(GetPhysicsActorHandle(), [&](const FPhysicsActorHandle& Actor)
+	FPhysicsCommand::ExecuteRead(GetPhysicsActor(), [&](const FPhysicsActorHandle& Actor)
 	{
 		MassSpaceLocal = FPhysicsInterface::GetComTransformLocal_AssumesLocked(Actor);
 	});
@@ -2872,7 +2865,7 @@ FTransform FBodyInstance::GetMassSpaceLocal() const
 void FBodyInstance::SetMassSpaceLocal(const FTransform& NewMassSpaceLocalTM)
 {
 	//TODO: UE doesn't store this so any changes to mass properties will not remember about this properly
-	FPhysicsCommand::ExecuteWrite(GetPhysicsActorHandle(), [&](const FPhysicsActorHandle& Actor)
+	FPhysicsCommand::ExecuteWrite(GetPhysicsActor(), [&](const FPhysicsActorHandle& Actor)
 	{
 		FPhysicsInterface::SetComLocalPose_AssumesLocked(Actor, NewMassSpaceLocalTM);
 	});
@@ -2882,7 +2875,7 @@ float FBodyInstance::GetBodyMass() const
 {
 	float OutMass = 0.0f;
 
-	FPhysicsCommand::ExecuteRead(GetPhysicsActorHandle(), [&](const FPhysicsActorHandle& Actor)
+	FPhysicsCommand::ExecuteRead(GetPhysicsActor(), [&](const FPhysicsActorHandle& Actor)
 	{
 		OutMass = FPhysicsInterface::GetMass_AssumesLocked(Actor);
 	});
@@ -2895,7 +2888,7 @@ FVector FBodyInstance::GetBodyInertiaTensor() const
 {
 	FVector OutTensor = FVector::ZeroVector;
 
-	FPhysicsCommand::ExecuteRead(GetPhysicsActorHandle(), [&](const FPhysicsActorHandle& Actor)
+	FPhysicsCommand::ExecuteRead(GetPhysicsActor(), [&](const FPhysicsActorHandle& Actor)
 	{
 		OutTensor = FPhysicsInterface::GetLocalInertiaTensor_AssumesLocked(Actor);
 	});
@@ -2908,7 +2901,7 @@ void FBodyInstance::SetInertiaConditioningEnabled(bool bInEnabled)
 	if (bInEnabled != bInertiaConditioning)
 	{
 		bInertiaConditioning = bInEnabled;
-		FPhysicsCommand::ExecuteWrite(GetPhysicsActorHandle(),
+		FPhysicsCommand::ExecuteWrite(GetPhysicsActor(),
 			[bInEnabled](const FPhysicsActorHandle& Actor)
 			{
 				if (FPhysicsInterface::IsRigidBody(Actor))
@@ -2923,7 +2916,7 @@ FBox FBodyInstance::GetBodyBounds() const
 {
 	FBox OutBox(EForceInit::ForceInitToZero);
 
-	FPhysicsCommand::ExecuteRead(GetPhysicsActorHandle(), [&](const FPhysicsActorHandle& Actor)
+	FPhysicsCommand::ExecuteRead(GetPhysicsActor(), [&](const FPhysicsActorHandle& Actor)
 		{
 			OutBox = FPhysicsInterface::GetBounds_AssumesLocked(Actor);
 		});
@@ -2935,7 +2928,7 @@ FBox FBodyInstance::GetBodyBoundsLocal() const
 {
 	FBox OutBox(EForceInit::ForceInitToZero);
 
-	FPhysicsCommand::ExecuteRead(GetPhysicsActorHandle(), [&](const FPhysicsActorHandle& Actor)
+	FPhysicsCommand::ExecuteRead(GetPhysicsActor(), [&](const FPhysicsActorHandle& Actor)
 		{
 			OutBox = FPhysicsInterface::GetBoundsLocal_AssumesLocked(Actor);
 		});
@@ -2958,8 +2951,8 @@ void FBodyInstance::CopyBodyInstancePropertiesFrom(const FBodyInstance* FromInst
 	check(FromInst);
 	check(FromInst->OwnerComponent.Get() == NULL);
 	check(FromInst->BodySetup.Get() == NULL);
-	check(!FPhysicsInterface::IsValid(FromInst->GetPhysicsActorHandle()));
-	check(!FPhysicsInterface::IsValid(GetPhysicsActorHandle()));
+	check(!FPhysicsInterface::IsValid(FromInst->GetPhysicsActor()));
+	check(!FPhysicsInterface::IsValid(GetPhysicsActor()));
 
 	*this = *FromInst;
 }
@@ -2983,12 +2976,12 @@ void FBodyInstance::CopyRuntimeBodyInstancePropertiesFrom(const FBodyInstance* F
 
 const FPhysScene* FBodyInstance::GetPhysicsScene() const
 {
-	return static_cast<const FPhysScene*>(FPhysicsInterface::GetCurrentScene(GetPhysicsActorHandle()));
+	return static_cast<const FPhysScene*>(FPhysicsInterface::GetCurrentScene(GetPhysicsActor()));
 }
 
 FPhysScene* FBodyInstance::GetPhysicsScene()
 {
-	return static_cast<FPhysScene*>(FPhysicsInterface::GetCurrentScene(GetPhysicsActorHandle()));
+	return static_cast<FPhysScene*>(FPhysicsInterface::GetCurrentScene(GetPhysicsActor()));
 }
 
 FPhysicsActorHandle& FBodyInstance::GetPhysicsActorHandle()
@@ -2996,17 +2989,41 @@ FPhysicsActorHandle& FBodyInstance::GetPhysicsActorHandle()
 	return ActorHandle;
 }
 
-void FBodyInstance::SetPhysicsActorHandle(FPhysicsActorHandle InHandle)
+const FPhysicsActorHandle& FBodyInstance::GetPhysicsActorHandle() const
 {
-	check(ActorHandle == nullptr);
-	ActorHandle = InHandle;
+	return ActorHandle;
 }
 
-const FPhysicsActorHandle& FBodyInstance::GetPhysicsActorHandle() const
+const FPhysicsActorHandle& FBodyInstance::GetActorReferenceWithWelding() const
+{
+	return WeldParent ? WeldParent->GetPhysicsActorHandle() : GetPhysicsActorHandle();
+}
+
+FPhysicsActorHandle FBodyInstance::GetPhysicsActor() const
 {
 	return ActorHandle;
 }
 
+FPhysicsActorHandle FBodyInstance::GetPhysicsActorWeldRoot() const
+{
+	return (WeldParent != nullptr) ? WeldParent->GetPhysicsActor() : GetPhysicsActor();
+}
+
+void FBodyInstance::SetPhysicsActor(FPhysicsActorHandle InHandle)
+{
+	ActorHandle = InHandle;
+}
+
+void FBodyInstance::ReleasePhysicsActor()
+{
+	if (ActorHandle != nullptr)
+	{
+		FPhysicsInterface::ReleaseActor(ActorHandle, GetPhysicsScene());
+		check(ActorHandle == nullptr);
+	}
+}
+
+
 const FWalkableSlopeOverride& FBodyInstance::GetWalkableSlopeOverride() const
 {
 	if (bOverrideWalkableSlopeOnInstance || !BodySetup.IsValid())
@@ -3249,9 +3266,9 @@ void FBodyInstance::UpdateMassProperties()
 
 	UPhysicalMaterial* PhysMat = GetSimplePhysicalMaterial();
 
-	if (FPhysicsInterface::IsValid(GetPhysicsActorHandle()) && FPhysicsInterface::IsRigidBody(GetPhysicsActorHandle()))
+	if (FPhysicsInterface::IsValid(GetPhysicsActor()) && FPhysicsInterface::IsRigidBody(GetPhysicsActor()))
 	{
-		FPhysicsCommand::ExecuteWrite(GetPhysicsActorHandle(), [&](FPhysicsActorHandle& Actor)
+		FPhysicsCommand::ExecuteWrite(GetPhysicsActor(), [&](FPhysicsActorHandle Actor)
 		{
 			check(FPhysicsInterface::IsValid(Actor));
 
@@ -3390,7 +3407,7 @@ void FBodyInstance::UpdateDebugRendering()
 
 void FBodyInstance::UpdateDampingProperties()
 {
-	FPhysicsCommand::ExecuteWrite(GetPhysicsActorHandle(), [&](const FPhysicsActorHandle& Actor)
+	FPhysicsCommand::ExecuteWrite(GetPhysicsActor(), [&](FPhysicsActorHandle Actor)
 	{
 		if(FPhysicsInterface::IsDynamic(Actor))
 	{
@@ -3404,7 +3421,7 @@ bool FBodyInstance::IsInstanceAwake() const
 {
 	bool bIsAwake = false;
 
-	FPhysicsCommand::ExecuteRead(GetPhysicsActorHandle(), [&](const FPhysicsActorHandle& Actor)
+	FPhysicsCommand::ExecuteRead(GetPhysicsActor(), [&](const FPhysicsActorHandle& Actor)
 	{
 		if(FPhysicsInterface::IsDynamic(Actor))
 	{
@@ -3417,7 +3434,7 @@ bool FBodyInstance::IsInstanceAwake() const
 
 void FBodyInstance::WakeInstance()
 {
-	FPhysicsCommand::ExecuteWrite(GetPhysicsActorHandle(), [&](const FPhysicsActorHandle& Actor)
+	FPhysicsCommand::ExecuteWrite(GetPhysicsActor(), [&](FPhysicsActorHandle Actor)
 	{
 		if(FPhysicsInterface::IsDynamic(Actor) && FPhysicsInterface::IsInScene(Actor) && !FPhysicsInterface::IsKinematic_AssumesLocked(Actor))
 		{
@@ -3428,7 +3445,7 @@ void FBodyInstance::WakeInstance()
 
 void FBodyInstance::PutInstanceToSleep()
 {
-	FPhysicsCommand::ExecuteWrite(GetPhysicsActorHandle(), [&](const FPhysicsActorHandle& Actor)
+	FPhysicsCommand::ExecuteWrite(GetPhysicsActor(), [&](const FPhysicsActorHandle& Actor)
 	{
 		if(FPhysicsInterface::IsDynamic(Actor) && FPhysicsInterface::IsInScene(Actor) && !FPhysicsInterface::IsKinematic_AssumesLocked(Actor))
 		{
@@ -3454,7 +3471,7 @@ float FBodyInstance::GetSleepThresholdMultiplier() const
 void FBodyInstance::SetLinearVelocity(const FVector& NewVel, bool bAddToCurrent, bool bAutoWake)
 {
 
-	FPhysicsCommand::ExecuteWrite(GetPhysicsActorHandle(), [&](const FPhysicsActorHandle& Actor)
+	FPhysicsCommand::ExecuteWrite(GetPhysicsActor(), [&](const FPhysicsActorHandle& Actor)
 	{
 		if(FPhysicsInterface::IsRigidBody(Actor))
 	{
@@ -3472,7 +3489,7 @@ void FBodyInstance::SetLinearVelocity(const FVector& NewVel, bool bAddToCurrent,
 
 void FBodyInstance::SetAngularVelocityInRadians(const FVector& NewAngVel, bool bAddToCurrent, bool bAutoWake)
 {
-	FPhysicsCommand::ExecuteWrite(GetPhysicsActorHandle(), [&](const FPhysicsActorHandle& Actor)
+	FPhysicsCommand::ExecuteWrite(GetPhysicsActor(), [&](const FPhysicsActorHandle& Actor)
 	{
 		if(FPhysicsInterface::IsRigidBody(Actor))
 	{
@@ -3497,7 +3514,7 @@ void FBodyInstance::SetMaxAngularVelocityInRadians(float NewMaxAngVel, bool bAdd
 {
 	float NewMaxInDegrees = FMath::RadiansToDegrees(NewMaxAngVel);
 
-	FPhysicsCommand::ExecuteWrite(GetPhysicsActorHandle(), [&](const FPhysicsActorHandle& Actor)
+	FPhysicsCommand::ExecuteWrite(GetPhysicsActor(), [&](const FPhysicsActorHandle& Actor)
 	{
 		if (bAddToCurrent)
 		{
@@ -3540,7 +3557,7 @@ void FBodyInstance::UpdateMaxDepenetrationVelocity()
 
 	// NOTE: FBodyInstance::MaxDepentrationVelocity now means initial depenetration velocity, 
 	// and not the general solver depentration velocity limit (which will probably be removed)
-	FPhysicsCommand::ExecuteWrite(GetPhysicsActorHandle(), [this, UsedMaxDepenetrationVelocity](const FPhysicsActorHandle& Actor)
+	FPhysicsCommand::ExecuteWrite(GetPhysicsActor(), [this, UsedMaxDepenetrationVelocity](const FPhysicsActorHandle& Actor)
 	{
 		FPhysicsInterface::SetMaxDepenetrationVelocity_AssumesLocked(Actor, UsedMaxDepenetrationVelocity);
 	});
@@ -3548,7 +3565,7 @@ void FBodyInstance::UpdateMaxDepenetrationVelocity()
 
 void FBodyInstance::AddCustomPhysics(FCalculateCustomPhysics& CalculateCustomPhysics)
 {
-	FPhysicsCommand::ExecuteWrite(GetPhysicsActorHandle(), [&](const FPhysicsActorHandle& Actor)
+	FPhysicsCommand::ExecuteWrite(GetPhysicsActor(), [&](const FPhysicsActorHandle& Actor)
 	{
 		if(!IsRigidBodyKinematic_AssumesLocked(Actor))
 		{
@@ -3574,7 +3591,7 @@ void FBodyInstance::ApplyAsyncPhysicsCommand(FAsyncPhysicsTimestamp TimeStamp, c
 	}
 	else
 	{
-		FPhysicsCommand::ExecuteWrite(GetPhysicsActorHandle(), [&](const FPhysicsActorHandle& Actor)
+		FPhysicsCommand::ExecuteWrite(GetPhysicsActor(), [&](const FPhysicsActorHandle& Actor)
 			{
 				Command();
 			});
@@ -3594,9 +3611,9 @@ void FBodyInstance::AddForce(const FVector& Force, bool bAllowSubstepping, bool
 	const bool bIsInternal = TimeStamp.IsValid();
 	ApplyAsyncPhysicsCommand(TimeStamp, bIsInternal, PlayerController, [&, Force, bAllowSubstepping, bAccelChange, bIsInternal]()
 	{
-		if (FPhysicsInterface::IsInScene(GetPhysicsActorHandle()) && IsBodyDynamic(GetPhysicsActorHandle(), bIsInternal))
+		if (FPhysicsInterface::IsInScene(GetPhysicsActor()) && IsBodyDynamic(GetPhysicsActor(), bIsInternal))
 		{
-			FPhysicsInterface::AddForce_AssumesLocked(GetPhysicsActorHandle(), Force, bAllowSubstepping, bAccelChange, bIsInternal);
+			FPhysicsInterface::AddForce_AssumesLocked(GetPhysicsActor(), Force, bAllowSubstepping, bAccelChange, bIsInternal);
 		}
 	});
 }
@@ -3606,16 +3623,16 @@ void FBodyInstance::AddForceAtPosition(const FVector& Force, const FVector& Posi
 	const bool bIsInternal = TimeStamp.IsValid();
 	ApplyAsyncPhysicsCommand(TimeStamp, bIsInternal, PlayerController, [&, Force, Position, bAllowSubstepping, bIsLocalForce, bIsInternal]()
 	{
-		if (FPhysicsInterface::IsInScene(GetPhysicsActorHandle()) && IsBodyDynamic(GetPhysicsActorHandle(), bIsInternal))
+		if (FPhysicsInterface::IsInScene(GetPhysicsActor()) && IsBodyDynamic(GetPhysicsActor(), bIsInternal))
 		{
-			FPhysicsInterface::AddForceAtPosition_AssumesLocked(GetPhysicsActorHandle(), Force, Position, bAllowSubstepping, bIsLocalForce, bIsInternal);
+			FPhysicsInterface::AddForceAtPosition_AssumesLocked(GetPhysicsActor(), Force, Position, bAllowSubstepping, bIsLocalForce, bIsInternal);
 		}
 	});
 }
 
 void FBodyInstance::ClearForces(bool bAllowSubstepping)
 {
-	FPhysicsCommand::ExecuteWrite(GetPhysicsActorHandle(), [&](const FPhysicsActorHandle& Actor)
+	FPhysicsCommand::ExecuteWrite(GetPhysicsActor(), [&](const FPhysicsActorHandle& Actor)
 	{
 		if(!IsRigidBodyKinematic_AssumesLocked(Actor))
 		{
@@ -3641,7 +3658,7 @@ void FBodyInstance::UpdateOneWayInteraction()
 {
 	const bool bCurrentOneWayInteraction = bOneWayInteraction;
 
-	FPhysicsCommand::ExecuteWrite(GetPhysicsActorHandle(), [bCurrentOneWayInteraction](const FPhysicsActorHandle& Actor)
+	FPhysicsCommand::ExecuteWrite(GetPhysicsActor(), [bCurrentOneWayInteraction](const FPhysicsActorHandle& Actor)
 	{
 		if (FPhysicsInterface::IsRigidBody(Actor) && !IsRigidBodyKinematic_AssumesLocked(Actor))
 		{
@@ -3655,16 +3672,16 @@ void FBodyInstance::AddTorqueInRadians(const FVector& Torque, bool bAllowSubstep
 	const bool bIsInternal = TimeStamp.IsValid();
 	ApplyAsyncPhysicsCommand(TimeStamp, bIsInternal, PlayerController, [&, Torque, bAllowSubstepping, bAccelChange, bIsInternal]()
 	{
-		if (FPhysicsInterface::IsInScene(GetPhysicsActorHandle()) && IsBodyDynamic(GetPhysicsActorHandle(), bIsInternal))
+		if (FPhysicsInterface::IsInScene(GetPhysicsActor()) && IsBodyDynamic(GetPhysicsActor(), bIsInternal))
 		{
-			FPhysicsInterface::AddTorque_AssumesLocked(GetPhysicsActorHandle(), Torque, bAllowSubstepping, bAccelChange, bIsInternal);
+			FPhysicsInterface::AddTorque_AssumesLocked(GetPhysicsActor(), Torque, bAllowSubstepping, bAccelChange, bIsInternal);
 		}
 	});
 }
 
 void FBodyInstance::ClearTorques(bool bAllowSubstepping)
 {
-	FPhysicsCommand::ExecuteWrite(GetPhysicsActorHandle(), [&](const FPhysicsActorHandle& Actor)
+	FPhysicsCommand::ExecuteWrite(GetPhysicsActor(), [&](const FPhysicsActorHandle& Actor)
 	{
 		if(!IsRigidBodyKinematic_AssumesLocked(Actor))
 		{
@@ -3681,15 +3698,15 @@ void FBodyInstance::AddAngularImpulseInRadians(const FVector& AngularImpulse, bo
 	const bool bIsInternal = TimeStamp.IsValid();
 	ApplyAsyncPhysicsCommand(TimeStamp, bIsInternal, PlayerController, [&, AngularImpulse, bVelChange, bIsInternal]()
 	{
-		if (FPhysicsInterface::IsInScene(GetPhysicsActorHandle()) && IsBodyDynamic(GetPhysicsActorHandle(), bIsInternal))
+		if (FPhysicsInterface::IsInScene(GetPhysicsActor()) && IsBodyDynamic(GetPhysicsActor(), bIsInternal))
 		{
 			if (bVelChange)
 			{
-				FPhysicsInterface::AddAngularVelocityInRadians_AssumesLocked(GetPhysicsActorHandle(), AngularImpulse, bIsInternal);
+				FPhysicsInterface::AddAngularVelocityInRadians_AssumesLocked(GetPhysicsActor(), AngularImpulse, bIsInternal);
 			}
 			else
 			{
-				FPhysicsInterface::AddAngularImpulseInRadians_AssumesLocked(GetPhysicsActorHandle(), AngularImpulse, bIsInternal);
+				FPhysicsInterface::AddAngularImpulseInRadians_AssumesLocked(GetPhysicsActor(), AngularImpulse, bIsInternal);
 			}
 		}
 	});
@@ -3700,15 +3717,15 @@ void FBodyInstance::AddImpulse(const FVector& Impulse, bool bVelChange, const FA
 	const bool bIsInternal = TimeStamp.IsValid();
 	ApplyAsyncPhysicsCommand(TimeStamp, bIsInternal, PlayerController, [&, Impulse, bVelChange, bIsInternal]()
 	{		
-		if (FPhysicsInterface::IsInScene(GetPhysicsActorHandle()) && IsBodyDynamic(GetPhysicsActorHandle(), bIsInternal))
+		if (FPhysicsInterface::IsInScene(GetPhysicsActor()) && IsBodyDynamic(GetPhysicsActor(), bIsInternal))
 		{
 			if (bVelChange)
 			{
-				FPhysicsInterface::AddVelocity_AssumesLocked(GetPhysicsActorHandle(), Impulse, bIsInternal);
+				FPhysicsInterface::AddVelocity_AssumesLocked(GetPhysicsActor(), Impulse, bIsInternal);
 			}
 			else
 			{
-				FPhysicsInterface::AddImpulse_AssumesLocked(GetPhysicsActorHandle(), Impulse, bIsInternal);
+				FPhysicsInterface::AddImpulse_AssumesLocked(GetPhysicsActor(), Impulse, bIsInternal);
 			}
 		}
 	});
@@ -3719,9 +3736,9 @@ void FBodyInstance::AddImpulseAtPosition(const FVector& Impulse, const FVector&
 	const bool bIsInternal = TimeStamp.IsValid();
 	ApplyAsyncPhysicsCommand(TimeStamp, bIsInternal, PlayerController, [&, Impulse, Position, bIsInternal]()
 	{
-		if (FPhysicsInterface::IsInScene(GetPhysicsActorHandle()) && IsBodyDynamic(GetPhysicsActorHandle(), bIsInternal))
+		if (FPhysicsInterface::IsInScene(GetPhysicsActor()) && IsBodyDynamic(GetPhysicsActor(), bIsInternal))
 		{
-			FPhysicsInterface::AddImpulseAtLocation_AssumesLocked(GetPhysicsActorHandle(), Impulse, Position, bIsInternal);
+			FPhysicsInterface::AddImpulseAtLocation_AssumesLocked(GetPhysicsActor(), Impulse, Position, bIsInternal);
 		}
 	});
 }
@@ -3731,9 +3748,9 @@ void FBodyInstance::AddVelocityChangeImpulseAtLocation(const FVector& Impulse, c
 	const bool bIsInternal = TimeStamp.IsValid();
 	ApplyAsyncPhysicsCommand(TimeStamp, bIsInternal, PlayerController, [&, Impulse, Position, bIsInternal]()
 	{
-		if (FPhysicsInterface::IsInScene(GetPhysicsActorHandle()) && IsBodyDynamic(GetPhysicsActorHandle(), bIsInternal))
+		if (FPhysicsInterface::IsInScene(GetPhysicsActor()) && IsBodyDynamic(GetPhysicsActor(), bIsInternal))
 		{
-			FPhysicsInterface::AddVelocityChangeImpulseAtLocation_AssumesLocked(GetPhysicsActorHandle(), Impulse, Position, bIsInternal);
+			FPhysicsInterface::AddVelocityChangeImpulseAtLocation_AssumesLocked(GetPhysicsActor(), Impulse, Position, bIsInternal);
 		}
 	});
 }
@@ -3778,7 +3795,7 @@ void FBodyInstance::SetEnableGravity(bool bInGravityEnabled)
 		bEnableGravity = bInGravityEnabled;
 
 		{
-			FPhysicsCommand::ExecuteWrite(GetPhysicsActorHandle(), [&](const FPhysicsActorHandle& Actor)
+			FPhysicsCommand::ExecuteWrite(GetPhysicsActor(), [&](const FPhysicsActorHandle& Actor)
 			{
 				if(FPhysicsInterface::IsRigidBody(Actor))
 				{
@@ -3801,7 +3818,7 @@ void FBodyInstance::SetGravityGroupIndex(int32 NewGravityGroupIndex)
 		GravityGroupIndex = NewGravityGroupIndex;
 
 		{
-			FPhysicsCommand::ExecuteWrite(GetPhysicsActorHandle(), [&](const FPhysicsActorHandle& Actor)
+			FPhysicsCommand::ExecuteWrite(GetPhysicsActor(), [&](const FPhysicsActorHandle& Actor)
 				{
 					if (FPhysicsInterface::IsRigidBody(Actor))
 					{
@@ -3819,7 +3836,7 @@ void FBodyInstance::SetUpdateKinematicFromSimulation(bool bInUpdateKinematicFrom
 		bUpdateKinematicFromSimulation = bInUpdateKinematicFromSimulation;
 
 		{
-			FPhysicsCommand::ExecuteWrite(GetPhysicsActorHandle(), [&](const FPhysicsActorHandle& Actor)
+			FPhysicsCommand::ExecuteWrite(GetPhysicsActor(), [&](const FPhysicsActorHandle& Actor)
 				{
 					if (FPhysicsInterface::IsRigidBody(Actor))
 					{
@@ -3837,7 +3854,7 @@ void FBodyInstance::SetGyroscopicTorqueEnabled(bool bInGyroscopicTorqueEnabled)
 		bGyroscopicTorqueEnabled = bInGyroscopicTorqueEnabled;
 
 		{
-			FPhysicsCommand::ExecuteWrite(GetPhysicsActorHandle(), [&](const FPhysicsActorHandle& Actor)
+			FPhysicsCommand::ExecuteWrite(GetPhysicsActor(), [&](const FPhysicsActorHandle& Actor)
 				{
 					if (FPhysicsInterface::IsRigidBody(Actor))
 					{
@@ -3864,7 +3881,7 @@ void FBodyInstance::SetSmoothEdgeCollisionsEnabled(bool bNewSmoothEdgeCollisions
 	{
 		bSmoothEdgeCollisions = bNewSmoothEdgeCollisions;
 
-		FPhysicsCommand::ExecuteWrite(GetPhysicsActorHandle(),
+		FPhysicsCommand::ExecuteWrite(GetPhysicsActor(),
 			[bNewSmoothEdgeCollisions](const FPhysicsActorHandle& Actor)
 			{
 				if (FPhysicsInterface::IsRigidBody(Actor))
@@ -3881,7 +3898,7 @@ void FBodyInstance::SetUseCCD(bool bInUseCCD)
 	{
 		bUseCCD = bInUseCCD;
 		// Need to set body flag
-		FPhysicsCommand::ExecuteWrite(GetPhysicsActorHandle(), [&](const FPhysicsActorHandle& Actor)
+		FPhysicsCommand::ExecuteWrite(GetPhysicsActor(), [&](const FPhysicsActorHandle& Actor)
 		{
 			if (FPhysicsInterface::IsValid(Actor) && FPhysicsInterface::IsRigidBody(Actor))
 			{
@@ -3899,7 +3916,7 @@ void FBodyInstance::SetUseMACD(bool bInUseMACD)
 	{
 		bUseMACD = bInUseMACD;
 
-		FPhysicsCommand::ExecuteWrite(GetPhysicsActorHandle(), [this, bInUseMACD](const FPhysicsActorHandle& Actor)
+		FPhysicsCommand::ExecuteWrite(GetPhysicsActor(), [this, bInUseMACD](const FPhysicsActorHandle& Actor)
 			{
 				if (FPhysicsInterface::IsValid(Actor) && FPhysicsInterface::IsRigidBody(Actor))
 				{
@@ -3915,7 +3932,7 @@ void FBodyInstance::SetPositionSolverIterationCount(uint8 PositionSolverIteratio
 	{
 		PositionSolverIterationCount = PositionSolverIterationCountIn;
 
-		FPhysicsCommand::ExecuteWrite(GetPhysicsActorHandle(), [this, PositionSolverIterationCountIn](const FPhysicsActorHandle& Actor)
+		FPhysicsCommand::ExecuteWrite(GetPhysicsActor(), [this, PositionSolverIterationCountIn](const FPhysicsActorHandle& Actor)
 			{
 				if (FPhysicsInterface::IsValid(Actor) && FPhysicsInterface::IsRigidBody(Actor))
 				{
@@ -3931,7 +3948,7 @@ void FBodyInstance::SetVelocitySolverIterationCount(uint8 VelocitySolverIteratio
 	{
 		VelocitySolverIterationCount = VelocitySolverIterationCountIn;
 
-		FPhysicsCommand::ExecuteWrite(GetPhysicsActorHandle(), [this, VelocitySolverIterationCountIn](const FPhysicsActorHandle& Actor)
+		FPhysicsCommand::ExecuteWrite(GetPhysicsActor(), [this, VelocitySolverIterationCountIn](const FPhysicsActorHandle& Actor)
 			{
 				if (FPhysicsInterface::IsValid(Actor) && FPhysicsInterface::IsRigidBody(Actor))
 				{
@@ -3947,7 +3964,7 @@ void FBodyInstance::SetProjectionSolverIterationCount(uint8 ProjectionSolverIter
 	{
 		ProjectionSolverIterationCount = ProjectionSolverIterationCountIn;
 
-		FPhysicsCommand::ExecuteWrite(GetPhysicsActorHandle(), [this, ProjectionSolverIterationCountIn](const FPhysicsActorHandle& Actor)
+		FPhysicsCommand::ExecuteWrite(GetPhysicsActor(), [this, ProjectionSolverIterationCountIn](const FPhysicsActorHandle& Actor)
 			{
 				if (FPhysicsInterface::IsValid(Actor) && FPhysicsInterface::IsRigidBody(Actor))
 				{
@@ -3959,7 +3976,7 @@ void FBodyInstance::SetProjectionSolverIterationCount(uint8 ProjectionSolverIter
 
 void FBodyInstance::SetPhysicsDisabled(bool bSetDisabled)
 {
-	FPhysicsCommand::ExecuteWrite(GetPhysicsActorHandle(), [&](const FPhysicsActorHandle& Actor)
+	FPhysicsCommand::ExecuteWrite(GetPhysicsActor(), [&](const FPhysicsActorHandle& Actor)
 	{
 		if (FPhysicsInterface::IsValid(Actor) && FPhysicsInterface::IsRigidBody(Actor))
 		{
@@ -3971,7 +3988,7 @@ void FBodyInstance::SetPhysicsDisabled(bool bSetDisabled)
 bool FBodyInstance::IsPhysicsDisabled() const
 {
 	bool bIsDisabled = false;
-	FPhysicsCommand::ExecuteRead(GetPhysicsActorHandle(), [&](const FPhysicsActorHandle& Actor)
+	FPhysicsCommand::ExecuteRead(GetPhysicsActor(), [&](const FPhysicsActorHandle& Actor)
 	{
 		bIsDisabled = FPhysicsInterface::IsDisabled(Actor);
 	});
@@ -3992,7 +4009,7 @@ EPhysicsReplicationMode FBodyInstance::GetPhysicsReplicationMode() const
 
 void FBodyInstance::AddRadialImpulseToBody(const FVector& Origin, float Radius, float Strength, uint8 Falloff, bool bVelChange)
 {
-	FPhysicsCommand::ExecuteWrite(GetPhysicsActorHandle(), [&](const FPhysicsActorHandle& Actor)
+	FPhysicsCommand::ExecuteWrite(GetPhysicsActor(), [&](const FPhysicsActorHandle& Actor)
 	{
 		if(FPhysicsInterface::IsRigidBody(Actor) && FPhysicsInterface::IsInScene(Actor) && !IsRigidBodyKinematic_AssumesLocked(Actor))
 		{
@@ -4003,7 +4020,7 @@ void FBodyInstance::AddRadialImpulseToBody(const FVector& Origin, float Radius,
 
 void FBodyInstance::AddRadialForceToBody(const FVector& Origin, float Radius, float Strength, uint8 Falloff, bool bAccelChange, bool bAllowSubstepping)
 {
-	FPhysicsCommand::ExecuteWrite(GetPhysicsActorHandle(), [&](const FPhysicsActorHandle& Actor)
+	FPhysicsCommand::ExecuteWrite(GetPhysicsActor(), [&](const FPhysicsActorHandle& Actor)
 	{
 		if(FPhysicsInterface::IsRigidBody(Actor) && FPhysicsInterface::IsInScene(Actor) && !IsRigidBodyKinematic_AssumesLocked(Actor))
 		{
@@ -4071,7 +4088,7 @@ bool FBodyInstance::OverlapTestForBodiesImpl(const FVector& Pos, const FQuat& Ro
 {
 	bool bHaveOverlap = false;
 
-	FPhysicsCommand::ExecuteRead(GetPhysicsActorHandle(), [&](const FPhysicsActorHandle& Actor)
+	FPhysicsCommand::ExecuteRead(GetPhysicsActor(), [&](const FPhysicsActorHandle& Actor)
 	{
 		// calculate the test global pose of the rigid body
 		FTransform PTestGlobalPose = FTransform(Rot, Pos);
@@ -4114,7 +4131,7 @@ bool FBodyInstance::OverlapTest(const FVector& Position, const FQuat& Rotation,
 
 	bool bHasOverlap = false;
 
-	FPhysicsCommand::ExecuteRead(GetPhysicsActorHandle(), [&](const FPhysicsActorHandle& Actor)
+	FPhysicsCommand::ExecuteRead(GetPhysicsActor(), [&](const FPhysicsActorHandle& Actor)
 	{
 		FTransform GeomTransform(Rotation, Position);
 
@@ -4194,7 +4211,7 @@ bool FBodyInstance::OverlapMulti(TArray<struct FOverlapResult>& InOutOverlaps, c
 
 	const FBodyInstance* TargetInstance = WeldParent ? WeldParent : this;
 
-	FPhysicsCommand::ExecuteRead(TargetInstance->GetPhysicsActorHandle(), [&](const FPhysicsActorHandle& Actor)
+	FPhysicsCommand::ExecuteRead(TargetInstance->GetPhysicsActor(), [&](const FPhysicsActorHandle& Actor)
 	{
 		if(FPhysicsInterface::IsValid(Actor))
 		{
@@ -4308,7 +4325,7 @@ void FBodyInstance::LoadProfileData(bool bVerifyProfile)
 
 void FBodyInstance::GetBodyInstanceResourceSizeEx(FResourceSizeEx& CumulativeResourceSize) const
 {
-	CumulativeResourceSize.AddDedicatedSystemMemoryBytes(FPhysicsInterface::GetResourceSizeEx(GetPhysicsActorHandle()));
+	CumulativeResourceSize.AddDedicatedSystemMemoryBytes(FPhysicsInterface::GetResourceSizeEx(GetPhysicsActor()));
 }
 
 void FBodyInstance::FixupData(class UObject* Loader)
@@ -4433,7 +4450,7 @@ void FBodyInstance::InitDynamicProperties_AssumesLocked()
 		return;
 	}
 	
-	if(FPhysicsInterface::IsDynamic(GetPhysicsActorHandle()))
+	if(FPhysicsInterface::IsDynamic(GetPhysicsActor()))
 	{
 		//A non simulated body may become simulated at runtime, so we need to compute its mass.
 		//However, this is not supported for complexAsSimple since a trimesh cannot itself be simulated, it can only be used for collision of other simple shapes.
@@ -4467,40 +4484,40 @@ void FBodyInstance::InitDynamicProperties_AssumesLocked()
 
 		if (ShouldInstanceSimulatingPhysics())
 		{
-			FPhysicsInterface::SetLinearVelocity_AssumesLocked(GetPhysicsActorHandle(), InitialLinVel);
+			FPhysicsInterface::SetLinearVelocity_AssumesLocked(GetPhysicsActor(), InitialLinVel);
 		}
 
-		FPhysicsInterface::SetSleepThresholdMultiplier_AssumesLocked(GetPhysicsActorHandle(), GetSleepThresholdMultiplier());
+		FPhysicsInterface::SetSleepThresholdMultiplier_AssumesLocked(GetPhysicsActor(), GetSleepThresholdMultiplier());
 
 		// @todo: implement sleep energy threshold
-		float SleepEnergyThresh = FPhysicsInterface::GetSleepEnergyThreshold_AssumesLocked(GetPhysicsActorHandle());
+		float SleepEnergyThresh = FPhysicsInterface::GetSleepEnergyThreshold_AssumesLocked(GetPhysicsActor());
 		SleepEnergyThresh *= GetSleepThresholdMultiplier();
-		FPhysicsInterface::SetSleepEnergyThreshold_AssumesLocked(GetPhysicsActorHandle(), SleepEnergyThresh);
+		FPhysicsInterface::SetSleepEnergyThreshold_AssumesLocked(GetPhysicsActor(), SleepEnergyThresh);
 
 		// PhysX specific dynamic parameters, not generically exposed to physics interface
 
-		float StabilizationThreshold = FPhysicsInterface::GetStabilizationEnergyThreshold_AssumesLocked(GetPhysicsActorHandle());
+		float StabilizationThreshold = FPhysicsInterface::GetStabilizationEnergyThreshold_AssumesLocked(GetPhysicsActor());
 		StabilizationThreshold *= StabilizationThresholdMultiplier;
-		FPhysicsInterface::SetStabilizationEnergyThreshold_AssumesLocked(GetPhysicsActorHandle(), StabilizationThreshold);
+		FPhysicsInterface::SetStabilizationEnergyThreshold_AssumesLocked(GetPhysicsActor(), StabilizationThreshold);
 
 		uint32 PositionIterCount = FMath::Clamp<uint8>(PositionSolverIterationCount, 1, 255);
 		uint32 VelocityIterCount = FMath::Clamp<uint8>(VelocitySolverIterationCount, 1, 255);
 		uint32 ProjectionIterCount = FMath::Clamp<uint8>(ProjectionSolverIterationCount, 1, 255);
-		FPhysicsInterface::SetSolverPositionIterationCount_AssumesLocked(GetPhysicsActorHandle(), PositionIterCount);
-		FPhysicsInterface::SetSolverVelocityIterationCount_AssumesLocked(GetPhysicsActorHandle(), VelocityIterCount);
-		FPhysicsInterface::SetSolverProjectionIterationCount_AssumesLocked(GetPhysicsActorHandle(), ProjectionIterCount);
+		FPhysicsInterface::SetSolverPositionIterationCount_AssumesLocked(GetPhysicsActor(), PositionIterCount);
+		FPhysicsInterface::SetSolverVelocityIterationCount_AssumesLocked(GetPhysicsActor(), VelocityIterCount);
+		FPhysicsInterface::SetSolverProjectionIterationCount_AssumesLocked(GetPhysicsActor(), ProjectionIterCount);
 
 		CreateDOFLock();
-		if(FPhysicsInterface::IsInScene(GetPhysicsActorHandle()) && !IsRigidBodyKinematic_AssumesLocked(GetPhysicsActorHandle()))
+		if(FPhysicsInterface::IsInScene(GetPhysicsActor()) && !IsRigidBodyKinematic_AssumesLocked(GetPhysicsActor()))
 		{
 			if(!bStartAwake && !bComponentAwake)
 			{
-				FPhysicsInterface::SetWakeCounter_AssumesLocked(GetPhysicsActorHandle(), 0.0f);
-				FPhysicsInterface::PutToSleep_AssumesLocked(GetPhysicsActorHandle());
+				FPhysicsInterface::SetWakeCounter_AssumesLocked(GetPhysicsActor(), 0.0f);
+				FPhysicsInterface::PutToSleep_AssumesLocked(GetPhysicsActor());
 			}
 		}
 
-		FPhysicsInterface::SetInitialized_AssumesLocked(GetPhysicsActorHandle(), true);
+		FPhysicsInterface::SetInitialized_AssumesLocked(GetPhysicsActor(), true);
 	}
 }
 
@@ -4728,7 +4745,7 @@ void FBodyInstance::UpdateSolverAsyncDeltaTime()
 			return;
 		}
 
-		Chaos::FPhysicsSolverBase* Solver = GetPhysicsActorHandle()->GetSolverBase();
+		Chaos::FPhysicsSolverBase* Solver = GetPhysicsActor()->GetSolverBase();
 
 		if (!Solver->IsUsingAsyncResults())
 		{
diff --git a/Engine/Source/Runtime/Engine/Private/PhysicsEngine/Experimental/ChaosScopedSceneLock.cpp b/Engine/Source/Runtime/Engine/Private/PhysicsEngine/Experimental/ChaosScopedSceneLock.cpp
--- a/Engine/Source/Runtime/Engine/Private/PhysicsEngine/Experimental/ChaosScopedSceneLock.cpp
+++ b/Engine/Source/Runtime/Engine/Private/PhysicsEngine/Experimental/ChaosScopedSceneLock.cpp
@@ -9,7 +9,7 @@
 #include "Components/SkeletalMeshComponent.h"
 #include "PBDRigidsSolver.h"
 
-FScopedSceneLock_Chaos::FScopedSceneLock_Chaos(FPhysicsActorHandle const* InActorHandle, EPhysicsInterfaceScopedLockType InLockType)
+FScopedSceneLock_Chaos::FScopedSceneLock_Chaos(const FPhysicsActorHandle& InActorHandle, EPhysicsInterfaceScopedLockType InLockType)
 	: LockType(InLockType)
 {
 	auto Scene = GetSceneForActor(InActorHandle);
@@ -17,7 +17,7 @@ FScopedSceneLock_Chaos::FScopedSceneLock_Chaos(FPhysicsActorHandle const* InActo
 	LockScene();
 }
 
-FScopedSceneLock_Chaos::FScopedSceneLock_Chaos(FPhysicsActorHandle const* InActorHandleA, FPhysicsActorHandle const* InActorHandleB, EPhysicsInterfaceScopedLockType InLockType)
+FScopedSceneLock_Chaos::FScopedSceneLock_Chaos(const FPhysicsActorHandle& InActorHandleA, const FPhysicsActorHandle& InActorHandleB, EPhysicsInterfaceScopedLockType InLockType)
 	: LockType(InLockType)
 {
 	FChaosScene* SceneA = GetSceneForActor(InActorHandleA);
@@ -41,6 +41,38 @@ FScopedSceneLock_Chaos::FScopedSceneLock_Chaos(FPhysicsActorHandle const* InActo
 	LockScene();
 }
 
+FScopedSceneLock_Chaos::FScopedSceneLock_Chaos(FPhysicsActorHandle const* InActorHandle, EPhysicsInterfaceScopedLockType InLockType)
+	: LockType(InLockType)
+{
+	auto Scene = GetSceneForActor(*InActorHandle);
+	Solver = Scene ? Scene->GetSolver() : nullptr;
+	LockScene();
+}
+
+FScopedSceneLock_Chaos::FScopedSceneLock_Chaos(FPhysicsActorHandle const* InActorHandleA, FPhysicsActorHandle const* InActorHandleB, EPhysicsInterfaceScopedLockType InLockType)
+	: LockType(InLockType)
+{
+	FChaosScene* SceneA = (InActorHandleA != nullptr) ? GetSceneForActor(*InActorHandleA) : nullptr;
+	FChaosScene* SceneB = (InActorHandleB != nullptr) ? GetSceneForActor(*InActorHandleB) : nullptr;
+	FChaosScene* Scene = nullptr;
+
+	if (SceneA == SceneB)
+	{
+		Scene = SceneA;
+	}
+	else if (!SceneA || !SceneB)
+	{
+		Scene = SceneA ? SceneA : SceneB;
+	}
+	else
+	{
+		UE_LOG(LogPhysics, Warning, TEXT("Attempted to aquire a physics scene lock for two paired actors that were not in the same scene. Skipping lock"));
+	}
+
+	Solver = Scene ? Scene->GetSolver() : nullptr;
+	LockScene();
+}
+
 FScopedSceneLock_Chaos::FScopedSceneLock_Chaos(FPhysicsConstraintHandle const* InConstraintHandle, EPhysicsInterfaceScopedLockType InLockType)
 	: Solver(nullptr)
 	, LockType(InLockType)
@@ -68,7 +100,7 @@ FScopedSceneLock_Chaos::FScopedSceneLock_Chaos(USkeletalMeshComponent* InSkelMes
 	{
 		for (FBodyInstance* BI : InSkelMeshComp->Bodies)
 		{
-			auto Scene = GetSceneForActor(&BI->GetPhysicsActorHandle());
+			auto Scene = GetSceneForActor(BI->GetPhysicsActor());
 			if (Scene)
 			{
 				Solver = Scene->GetSolver();
@@ -180,11 +212,11 @@ void FScopedSceneLock_Chaos::UnlockScene()
 	bHasLock = false;
 }
 
-FChaosScene* FScopedSceneLock_Chaos::GetSceneForActor(FPhysicsActorHandle const* InActorHandle)
+FChaosScene* FScopedSceneLock_Chaos::GetSceneForActor(const FPhysicsActorHandle& InActorHandle)
 {
 	if (InActorHandle)
 	{
-		return static_cast<FPhysScene*>(FChaosEngineInterface::GetCurrentScene(*InActorHandle));
+		return static_cast<FPhysScene*>(FChaosEngineInterface::GetCurrentScene(InActorHandle));
 	}
 
 	return nullptr;
diff --git a/Engine/Source/Runtime/Engine/Private/PhysicsEngine/Experimental/PhysInterface_Chaos.cpp b/Engine/Source/Runtime/Engine/Private/PhysicsEngine/Experimental/PhysInterface_Chaos.cpp
--- a/Engine/Source/Runtime/Engine/Private/PhysicsEngine/Experimental/PhysInterface_Chaos.cpp
+++ b/Engine/Source/Runtime/Engine/Private/PhysicsEngine/Experimental/PhysInterface_Chaos.cpp
@@ -507,7 +507,7 @@ bool FPhysInterface_Chaos::ExecuteRead(const FPhysicsActorHandle& InActorReferen
 {
 	if(InActorReference)
 	{
-		FScopedSceneLock_Chaos SceneLock(&InActorReference, EPhysicsInterfaceScopedLockType::Read);
+		FScopedSceneLock_Chaos SceneLock(InActorReference, EPhysicsInterfaceScopedLockType::Read);
 
 		InCallable(InActorReference);
 		return true;
@@ -524,7 +524,7 @@ bool FPhysInterface_Chaos::ExecuteRead(USkeletalMeshComponent* InMeshComponent,
 
 bool FPhysInterface_Chaos::ExecuteRead(const FPhysicsActorHandle& InActorReferenceA, const FPhysicsActorHandle& InActorReferenceB, TFunctionRef<void(const FPhysicsActorHandle& ActorA, const FPhysicsActorHandle& ActorB)> InCallable)
 {
-	FScopedSceneLock_Chaos SceneLock(&InActorReferenceA, &InActorReferenceB, EPhysicsInterfaceScopedLockType::Read);
+	FScopedSceneLock_Chaos SceneLock(InActorReferenceA, InActorReferenceB, EPhysicsInterfaceScopedLockType::Read);
 	InCallable(InActorReferenceA, InActorReferenceB);
 	return true;
 }
@@ -562,21 +562,21 @@ bool FPhysInterface_Chaos::ExecuteRead(Chaos::FPhysicsObject* InObjectA, Chaos::
 
 bool FPhysInterface_Chaos::ExecuteWrite(const FPhysicsActorHandle& InActorReference, TFunctionRef<void(const FPhysicsActorHandle& Actor)> InCallable)
 {
-	//why do we have a write that takes in a const handle?
-	if(InActorReference)
+	if (InActorReference)
 	{
-		FScopedSceneLock_Chaos SceneLock(&InActorReference, EPhysicsInterfaceScopedLockType::Write);
+		FScopedSceneLock_Chaos SceneLock(InActorReference, EPhysicsInterfaceScopedLockType::Write);
 		InCallable(InActorReference);
 		return true;
 	}
+
 	return false;
 }
 
 bool FPhysInterface_Chaos::ExecuteWrite(FPhysicsActorHandle& InActorReference, TFunctionRef<void(FPhysicsActorHandle& Actor)> InCallable)
 {
-	if(InActorReference)
+	if (InActorReference)
 	{
-		FScopedSceneLock_Chaos SceneLock(&InActorReference, EPhysicsInterfaceScopedLockType::Write);
+		FScopedSceneLock_Chaos SceneLock(InActorReference, EPhysicsInterfaceScopedLockType::Write);
 		InCallable(InActorReference);
 		return true;
 	}
@@ -584,23 +584,16 @@ bool FPhysInterface_Chaos::ExecuteWrite(FPhysicsActorHandle& InActorReference, T
 	return false;
 }
 
-bool FPhysInterface_Chaos::ExecuteWrite(USkeletalMeshComponent* InMeshComponent, TFunctionRef<void()> InCallable)
-{
-	FScopedSceneLock_Chaos SceneLock(InMeshComponent, EPhysicsInterfaceScopedLockType::Write);
-	InCallable();
-	return true;
-}
-
 bool FPhysInterface_Chaos::ExecuteWrite(const FPhysicsActorHandle& InActorReferenceA, const FPhysicsActorHandle& InActorReferenceB, TFunctionRef<void(const FPhysicsActorHandle& ActorA, const FPhysicsActorHandle& ActorB)> InCallable)
 {
-	FScopedSceneLock_Chaos SceneLock(&InActorReferenceA, &InActorReferenceB, EPhysicsInterfaceScopedLockType::Write);
+	FScopedSceneLock_Chaos SceneLock(InActorReferenceA, InActorReferenceB, EPhysicsInterfaceScopedLockType::Write);
 	InCallable(InActorReferenceA, InActorReferenceB);
 	return true;
 }
 
 bool FPhysInterface_Chaos::ExecuteWrite(const FPhysicsConstraintHandle& InConstraintRef, TFunctionRef<void(const FPhysicsConstraintHandle& Constraint)> InCallable)
 {
-	if(InConstraintRef.IsValid())
+	if (InConstraintRef.IsValid())
 	{
 		FScopedSceneLock_Chaos SceneLock(&InConstraintRef, EPhysicsInterfaceScopedLockType::Write);
 		InCallable(InConstraintRef);
@@ -610,6 +603,13 @@ bool FPhysInterface_Chaos::ExecuteWrite(const FPhysicsConstraintHandle& InConstr
 	return false;
 }
 
+bool FPhysInterface_Chaos::ExecuteWrite(USkeletalMeshComponent* InMeshComponent, TFunctionRef<void()> InCallable)
+{
+	FScopedSceneLock_Chaos SceneLock(InMeshComponent, EPhysicsInterfaceScopedLockType::Write);
+	InCallable();
+	return true;
+}
+
 bool FPhysInterface_Chaos::ExecuteWrite(FPhysScene* InScene, TFunctionRef<void()> InCallable)
 {
 	if(InScene)
@@ -645,7 +645,7 @@ void FPhysInterface_Chaos::ExecuteShapeWrite(FBodyInstance* InInstance, FPhysics
 {
 	if(InInstance && InShape.IsValid())
 	{
-		FScopedSceneLock_Chaos SceneLock(&InInstance->GetPhysicsActorHandle(), EPhysicsInterfaceScopedLockType::Write);
+		FScopedSceneLock_Chaos SceneLock(InInstance->GetPhysicsActor(), EPhysicsInterfaceScopedLockType::Write);
 		InCallable(InShape);
 	}
 }
@@ -684,7 +684,7 @@ const FBodyInstance* FPhysInterface_Chaos::ShapeToOriginalBodyInstance(const FBo
 
 
 
-void FPhysInterface_Chaos::AddGeometry(FPhysicsActorHandle& InActor, const FGeometryAddParams& InParams, TArray<FPhysicsShapeHandle>* OutOptShapes)
+void FPhysInterface_Chaos::AddGeometry(const FPhysicsActorHandle& InActor, const FGeometryAddParams& InParams, TArray<FPhysicsShapeHandle>* OutOptShapes)
 {
 	TRACE_CPUPROFILER_EVENT_SCOPE(FPhysInterface_Chaos::AddGeometry);
 	LLM_SCOPE(ELLMTag::ChaosGeometry);
@@ -845,11 +845,11 @@ bool FPhysInterface_Chaos::LineTrace_Geom(FHitResult& OutHit, const FBodyInstanc
 	{
 		{
 			// #PHYS2 Really need a concept for "multi" locks here - as we're locking ActorRef but not TargetInstance->ActorRef
-			FPhysicsCommand::ExecuteRead(InInstance->GetPhysicsActorHandle(), [&](const FPhysicsActorHandle& Actor)
+			FPhysicsCommand::ExecuteRead(InInstance->GetPhysicsActor(), [&](const FPhysicsActorHandle& Actor)
 			{
 				// If we're welded then the target instance is actually our parent
 				const FBodyInstance* TargetInstance = InInstance->WeldParent ? InInstance->WeldParent : InInstance;
-				if(const FPhysicsActorHandle RigidBody = TargetInstance->GetPhysicsActorHandle())
+				if(FPhysicsActorHandle RigidBody = TargetInstance->GetPhysicsActor())
 				{
 					FRaycastHit BestHit;
 					BestHit.Distance = FLT_MAX;
@@ -942,7 +942,7 @@ bool FPhysInterface_Chaos::Sweep_Geom(FHitResult& OutHit, const FBodyInstance* I
 
 		const FBodyInstance* TargetInstance = InInstance->WeldParent ? InInstance->WeldParent : InInstance;
 
-		FPhysicsCommand::ExecuteRead(TargetInstance->GetPhysicsActorHandle(), [&](const FPhysicsActorHandle& Actor)
+		FPhysicsCommand::ExecuteRead(TargetInstance->GetPhysicsActor(), [&](const FPhysicsActorHandle& Actor)
 		{
 			if (Actor && InInstance->OwnerComponent.Get())
 			{
@@ -1040,7 +1040,7 @@ bool FPhysInterface_Chaos::Sweep_Geom(FHitResult& OutHit, const FBodyInstance* I
 bool Overlap_GeomInternal(const FBodyInstance* InInstance, const Chaos::FImplicitObject& InGeom, const FTransform& GeomTransform, FMTDResult* OutOptResult, bool bTraceComplex)
 {
 	const FBodyInstance* TargetInstance = InInstance->WeldParent ? InInstance->WeldParent : InInstance;
-	FPhysicsActorHandle RigidBody = TargetInstance->GetPhysicsActorHandle();
+	FPhysicsActorHandle RigidBody = TargetInstance->GetPhysicsActor();
 
 	if (!RigidBody)
 	{
@@ -1129,7 +1129,7 @@ bool FPhysInterface_Chaos::GetSquaredDistanceToBody(const FBodyInstance* InInsta
 	const FTransform BodyTM = UseBI->GetUnrealWorldTransform();
 	const FVector LocalPoint = BodyTM.InverseTransformPositionNoScale(InPoint);
 
-	FPhysicsCommand::ExecuteRead(UseBI->GetPhysicsActorHandle(), [&](const FPhysicsActorHandle& Actor)
+	FPhysicsCommand::ExecuteRead(UseBI->GetPhysicsActor(), [&](const FPhysicsActorHandle& Actor)
 	{
 		bEarlyOut = false;
 
diff --git a/Engine/Source/Runtime/Engine/Public/Physics/Experimental/ChaosScopedSceneLock.h b/Engine/Source/Runtime/Engine/Public/Physics/Experimental/ChaosScopedSceneLock.h
--- a/Engine/Source/Runtime/Engine/Public/Physics/Experimental/ChaosScopedSceneLock.h
+++ b/Engine/Source/Runtime/Engine/Public/Physics/Experimental/ChaosScopedSceneLock.h
@@ -21,6 +21,9 @@ namespace Chaos
 
 struct FScopedSceneLock_Chaos
 {
+	ENGINE_API FScopedSceneLock_Chaos(const FPhysicsActorHandle& InActorHandle, EPhysicsInterfaceScopedLockType InLockType);
+	ENGINE_API FScopedSceneLock_Chaos(const FPhysicsActorHandle& InActorHandleA, const FPhysicsActorHandle& InActorHandleB, EPhysicsInterfaceScopedLockType InLockType);
+	// TODO_CHAOSAPI: Deprecate pointer-to-handle API
 	ENGINE_API FScopedSceneLock_Chaos(FPhysicsActorHandle const* InActorHandle, EPhysicsInterfaceScopedLockType InLockType);
 	ENGINE_API FScopedSceneLock_Chaos(FPhysicsActorHandle const* InActorHandleA, FPhysicsActorHandle const* InActorHandleB, EPhysicsInterfaceScopedLockType InLockType);
 	ENGINE_API FScopedSceneLock_Chaos(FPhysicsConstraintHandle const* InConstraintHandle, EPhysicsInterfaceScopedLockType InLockType);
@@ -42,7 +45,7 @@ private:
 	ENGINE_API void LockScene();
 	ENGINE_API void UnlockScene();
 
-	ENGINE_API FChaosScene* GetSceneForActor(FPhysicsActorHandle const* InActorHandle);
+	ENGINE_API FChaosScene* GetSceneForActor(const FPhysicsActorHandle& InActorHandle);
 	ENGINE_API FChaosScene* GetSceneForActor(FPhysicsConstraintHandle const* InConstraintHandle);
 
 	bool bHasLock = false;
diff --git a/Engine/Source/Runtime/Engine/Public/Physics/Experimental/PhysInterface_Chaos.h b/Engine/Source/Runtime/Engine/Public/Physics/Experimental/PhysInterface_Chaos.h
--- a/Engine/Source/Runtime/Engine/Public/Physics/Experimental/PhysInterface_Chaos.h
+++ b/Engine/Source/Runtime/Engine/Public/Physics/Experimental/PhysInterface_Chaos.h
@@ -73,6 +73,7 @@ public:
 	static ENGINE_API bool ExecuteRead(Chaos::FPhysicsObject* InObjectA, Chaos::FPhysicsObject* InObjectB, TFunctionRef<void(Chaos::FPhysicsObject* InObjectA, Chaos::FPhysicsObject* InObjectB)> InCallable);
 
     static ENGINE_API bool ExecuteWrite(const FPhysicsActorHandle& InActorReference, TFunctionRef<void(const FPhysicsActorHandle& Actor)> InCallable);
+	// TODO_CHAOSAPI: Deprecate non-const ref API
 	static ENGINE_API bool ExecuteWrite(FPhysicsActorHandle& InActorReference, TFunctionRef<void(FPhysicsActorHandle& Actor)> InCallable);
     static ENGINE_API bool ExecuteWrite(USkeletalMeshComponent* InMeshComponent, TFunctionRef<void()> InCallable);
     static ENGINE_API bool ExecuteWrite(const FPhysicsActorHandle& InActorReferenceA, const FPhysicsActorHandle& InActorReferenceB, TFunctionRef<void(const FPhysicsActorHandle& ActorA, const FPhysicsActorHandle& ActorB)> InCallable);
@@ -92,7 +93,7 @@ public:
 	// Shape interface functions
 	static ENGINE_API FPhysicsShapeHandle CreateShape(physx::PxGeometry* InGeom, bool bSimulation = true, bool bQuery = true, UPhysicalMaterial* InSimpleMaterial = nullptr, TArray<UPhysicalMaterial*>* InComplexMaterials = nullptr);
 	
-	static ENGINE_API void AddGeometry(FPhysicsActorHandle& InActor, const FGeometryAddParams& InParams, TArray<FPhysicsShapeHandle>* OutOptShapes = nullptr);
+	static ENGINE_API void AddGeometry(const FPhysicsActorHandle& InActor, const FGeometryAddParams& InParams, TArray<FPhysicsShapeHandle>* OutOptShapes = nullptr);
 	static ENGINE_API void CreateGeometry(const FGeometryAddParams& InParams, TArray<Chaos::FImplicitObjectPtr>& OutGeoms, Chaos::FShapesArray& OutShapes, TArray<FPhysicsShapeHandle>* OutOptShapes);
 	
 	UE_DEPRECATED(5.4, "Use CreateGeometry with FImplicitObjectPtr instead")
diff --git a/Engine/Source/Runtime/PhysicsCore/Private/ChaosEngineInterface.cpp b/Engine/Source/Runtime/PhysicsCore/Private/ChaosEngineInterface.cpp
--- a/Engine/Source/Runtime/PhysicsCore/Private/ChaosEngineInterface.cpp
+++ b/Engine/Source/Runtime/PhysicsCore/Private/ChaosEngineInterface.cpp
@@ -175,13 +175,13 @@ const FQuat& FPhysicsShapeAdapter_Chaos::GetGeomOrientation() const
 	return GeometryRotation;
 }
 
-void FChaosEngineInterface::AddActorToSolver(FPhysicsActorHandle& Handle,Chaos::FPhysicsSolver* Solver)
+void FChaosEngineInterface::AddActorToSolver(const FPhysicsActorHandle& Handle,Chaos::FPhysicsSolver* Solver)
 {
 	Solver->RegisterObject(Handle);
 }
 
 
-void FChaosEngineInterface::RemoveActorFromSolver(FPhysicsActorHandle& Handle,Chaos::FPhysicsSolver* Solver)
+void FChaosEngineInterface::RemoveActorFromSolver(const FPhysicsActorHandle& Handle,Chaos::FPhysicsSolver* Solver)
 {
 	// Should we stop passing solver in? (need to check it's not null regardless in case proxy was never registered)
 	if(Solver && Handle && Handle->GetSolverBase() == Solver)
@@ -366,7 +366,7 @@ void FChaosEngineInterface::RemoveDisabledCollisionsFor_AssumesLocked(TArray< FP
 {
 	if (bEnableChaosCollisionManager)
 	{
-		for (FPhysicsActorHandle& ActorReference : InPhysicsActors)
+		for (const FPhysicsActorHandle& ActorReference : InPhysicsActors)
 		{
 			Chaos::FUniqueIdx ActorIndex = ActorReference->GetGameThreadAPI().UniqueIdx();
 
@@ -390,7 +390,7 @@ bool FChaosEngineInterface::IsDisabled(const FPhysicsActorHandle& InPhysicsActor
 	return InPhysicsActor->GetGameThreadAPI().Disabled();
 }
 
-void FChaosEngineInterface::SetActorUserData_AssumesLocked(FPhysicsActorHandle& InActorReference,FPhysicsUserData* InUserData)
+void FChaosEngineInterface::SetActorUserData_AssumesLocked(const FPhysicsActorHandle& InActorReference,FPhysicsUserData* InUserData)
 {
 	InActorReference->GetGameThreadAPI().SetUserData(InUserData);
 }
@@ -1226,7 +1226,7 @@ void FChaosEngineInterface::SetSleepThresholdMultiplier_AssumesLocked(const FPhy
 	return InActorReference->GetGameThreadAPI().SetSleepThresholdMultiplier(InThresholdMultiplier);
 }
 
-void FChaosEngineInterface::SetMass_AssumesLocked(FPhysicsActorHandle& InActorReference,float InMass)
+void FChaosEngineInterface::SetMass_AssumesLocked(const FPhysicsActorHandle& InActorReference,float InMass)
 {
 	Chaos::FRigidBodyHandle_External& Body_External = InActorReference->GetGameThreadAPI();
 	Body_External.SetM(InMass);
@@ -1239,7 +1239,7 @@ void FChaosEngineInterface::SetMass_AssumesLocked(FPhysicsActorHandle& InActorRe
 	}
 }
 
-void FChaosEngineInterface::SetMassSpaceInertiaTensor_AssumesLocked(FPhysicsActorHandle& InActorReference,const FVector& InTensor)
+void FChaosEngineInterface::SetMassSpaceInertiaTensor_AssumesLocked(const FPhysicsActorHandle& InActorReference,const FVector& InTensor)
 {
 	if(CHAOS_ENSURE(!FMath::IsNearlyZero(InTensor.X)) && CHAOS_ENSURE(!FMath::IsNearlyZero(InTensor.Y)) && CHAOS_ENSURE(!FMath::IsNearlyZero(InTensor.Z)))
 	{
diff --git a/Engine/Source/Runtime/PhysicsCore/Public/Chaos/ChaosEngineInterface.h b/Engine/Source/Runtime/PhysicsCore/Public/Chaos/ChaosEngineInterface.h
--- a/Engine/Source/Runtime/PhysicsCore/Public/Chaos/ChaosEngineInterface.h
+++ b/Engine/Source/Runtime/PhysicsCore/Public/Chaos/ChaosEngineInterface.h
@@ -347,8 +347,8 @@ public:
 	static PHYSICSCORE_API void ReleaseActor(FPhysicsActorHandle& InActorReference,FChaosScene* InScene = nullptr,bool bNeverDeferRelease=false);
 
 	static bool IsValid(const FPhysicsActorHandle& Handle) { return Handle != nullptr; }
-	static PHYSICSCORE_API void AddActorToSolver(FPhysicsActorHandle& Handle,Chaos::FPhysicsSolver* Solver);
-	static PHYSICSCORE_API void RemoveActorFromSolver(FPhysicsActorHandle& Handle,Chaos::FPhysicsSolver* Solver);
+	static PHYSICSCORE_API void AddActorToSolver(const FPhysicsActorHandle& Handle,Chaos::FPhysicsSolver* Solver);
+	static PHYSICSCORE_API void RemoveActorFromSolver(const FPhysicsActorHandle& Handle,Chaos::FPhysicsSolver* Solver);
 
 	static PHYSICSCORE_API void SetGlobalPose_AssumesLocked(const FPhysicsActorHandle& InActorReference,const FTransform& InNewPose,bool bAutoWake = true);
 	static PHYSICSCORE_API void SetKinematicTarget_AssumesLocked(const FPhysicsActorHandle& InActorReference,const FTransform& InNewTarget);
@@ -382,7 +382,7 @@ public:
 	static PHYSICSCORE_API void SetDisabled(const FPhysicsActorHandle& InPhysicsActor, bool bSetDisabled);
 	static PHYSICSCORE_API bool IsDisabled(const FPhysicsActorHandle& InPhysicsActor);
 
-	static PHYSICSCORE_API void SetActorUserData_AssumesLocked(FPhysicsActorHandle& InActorReference,FPhysicsUserData* InUserData);
+	static PHYSICSCORE_API void SetActorUserData_AssumesLocked(const FPhysicsActorHandle& InActorReference,FPhysicsUserData* InUserData);
 
 	static PHYSICSCORE_API bool IsRigidBody(const FPhysicsActorHandle& InActorReference);
 	static PHYSICSCORE_API bool IsDynamic(const FPhysicsActorHandle& InActorReference);
@@ -479,8 +479,8 @@ public:
 
 	static PHYSICSCORE_API void SetSleepThresholdMultiplier_AssumesLocked(const FPhysicsActorHandle& InActorReference, float ThresholdMultiplier);
 
-	static PHYSICSCORE_API void SetMass_AssumesLocked(FPhysicsActorHandle& InHandle,float InMass);
-	static PHYSICSCORE_API void SetMassSpaceInertiaTensor_AssumesLocked(FPhysicsActorHandle& InHandle,const FVector& InTensor);
+	static PHYSICSCORE_API void SetMass_AssumesLocked(const FPhysicsActorHandle& InHandle,float InMass);
+	static PHYSICSCORE_API void SetMassSpaceInertiaTensor_AssumesLocked(const FPhysicsActorHandle& InHandle,const FVector& InTensor);
 	static PHYSICSCORE_API void SetComLocalPose_AssumesLocked(const FPhysicsActorHandle& InHandle,const FTransform& InComLocalPose);
 
 	static PHYSICSCORE_API bool IsInertiaConditioningEnabled_AssumesLocked(const FPhysicsActorHandle& InActorReference);
