SkeletalMesh: Fix typo in simulation enum

#rnx

diff --git a/Engine/Source/Runtime/Engine/Classes/Components/SkeletalMeshComponent.h b/Engine/Source/Runtime/Engine/Classes/Components/SkeletalMeshComponent.h
--- a/Engine/Source/Runtime/Engine/Classes/Components/SkeletalMeshComponent.h
+++ b/Engine/Source/Runtime/Engine/Classes/Components/SkeletalMeshComponent.h
@@ -195,7 +195,7 @@ namespace EPhysicsTransformUpdateMode
 {
 	enum Type : int
 	{
-		SimulationUpatesComponentTransform,
+		SimulationUpatesComponentTransform UMETA(DisplayName="Simulation Updates Component Transform"), 
 		ComponentTransformIsKinematic
 	};
 }
@@ -1972,7 +1972,7 @@ public:
 #endif
 protected:
 	/**
-	 * If PhysicsTransformUpdateMode is set to SimulationUpatesComponentTransform, this returns the 
+	 * If PhysicsTransformUpdateMode is set to SimulationUpdatesComponentTransform, this returns the 
 	 * world-space transform of the body passed in (assumed to be the root body), after applying the 
 	 * TransformToRoot offset. Otherwise, it returns the component transform.
 	 */
