Add editor-only asset user data support for World Settings.

diff --git a/Engine/Source/Runtime/Engine/Classes/GameFramework/WorldSettings.h b/Engine/Source/Runtime/Engine/Classes/GameFramework/WorldSettings.h
--- a/Engine/Source/Runtime/Engine/Classes/GameFramework/WorldSettings.h
+++ b/Engine/Source/Runtime/Engine/Classes/GameFramework/WorldSettings.h
@@ -744,6 +744,12 @@ protected:
 	UPROPERTY(EditAnywhere, AdvancedDisplay, Instanced, Category = AssetUserData)
 	TArray<TObjectPtr<UAssetUserData>> AssetUserData;
 
+#if WITH_EDITORONLY_DATA
+	/** Array of editor-only user data stored with the asset */
+	UPROPERTY(EditAnywhere, AdvancedDisplay, Instanced, Category = AssetUserData)
+	TArray<TObjectPtr<UAssetUserData>> AssetUserDataEditorOnly;
+#endif
+
 	// If paused, PlayerState of person pausing the game.
 	UPROPERTY(transient, replicated)
 	TObjectPtr<class APlayerState> PauserPlayerState;
diff --git a/Engine/Source/Runtime/Engine/Private/WorldSettings.cpp b/Engine/Source/Runtime/Engine/Private/WorldSettings.cpp
--- a/Engine/Source/Runtime/Engine/Private/WorldSettings.cpp
+++ b/Engine/Source/Runtime/Engine/Private/WorldSettings.cpp
@@ -452,20 +452,17 @@ void AWorldSettings::AddAssetUserData(UAssetUserData* InUserData)
 {
 	if (InUserData != NULL)
 	{
-		UAssetUserData* ExistingData = GetAssetUserDataOfClass(InUserData->GetClass());
-		if (ExistingData != NULL)
-		{
-			AssetUserData.Remove(ExistingData);
-		}
+		RemoveUserDataOfClass(InUserData->GetClass());
 		AssetUserData.Add(InUserData);
 	}
 }
 
 UAssetUserData* AWorldSettings::GetAssetUserDataOfClass(TSubclassOf<UAssetUserData> InUserDataClass)
 {
-	for (int32 DataIdx = 0; DataIdx < AssetUserData.Num(); DataIdx++)
+	const TArray<UAssetUserData*>* ArrayPtr = GetAssetUserDataArray();
+	for (int32 DataIdx = 0; DataIdx < ArrayPtr->Num(); DataIdx++)
 	{
-		UAssetUserData* Datum = AssetUserData[DataIdx];
+		UAssetUserData* Datum = (*ArrayPtr)[DataIdx];
 		if (Datum != NULL && Datum->IsA(InUserDataClass))
 		{
 			return Datum;
@@ -476,7 +473,22 @@ UAssetUserData* AWorldSettings::GetAssetUserDataOfClass(TSubclassOf<UAssetUserDa
 
 const TArray<UAssetUserData*>* AWorldSettings::GetAssetUserDataArray() const
 {
+#if WITH_EDITOR
+	if (IsRunningCookCommandlet())
+	{
+		return &ToRawPtrTArrayUnsafe(AssetUserData);
+	}
+	else
+	{
+		static thread_local TArray<TObjectPtr<UAssetUserData>> CachedAssetUserData;
+		CachedAssetUserData.Reset();
+		CachedAssetUserData.Append(AssetUserData);
+		CachedAssetUserData.Append(AssetUserDataEditorOnly);
+		return &ToRawPtrTArrayUnsafe(CachedAssetUserData);
+	}
+#else
 	return &ToRawPtrTArrayUnsafe(AssetUserData);
+#endif
 }
 
 #if WITH_EDITOR
@@ -608,6 +620,17 @@ void AWorldSettings::RemoveUserDataOfClass(TSubclassOf<UAssetUserData> InUserDat
 			return;
 		}
 	}
+#if WITH_EDITOR
+	for (int32 DataIdx = 0; DataIdx < AssetUserDataEditorOnly.Num(); DataIdx++)
+	{
+		UAssetUserData* Datum = AssetUserDataEditorOnly[DataIdx];
+		if (Datum != NULL && Datum->IsA(InUserDataClass))
+		{
+			AssetUserDataEditorOnly.RemoveAt(DataIdx);
+			return;
+		}
+	}
+#endif
 }
 
 void AWorldSettings::PostLoad()
@@ -842,6 +865,14 @@ void AWorldSettings::PostEditChangeProperty(FPropertyChangedEvent& PropertyChang
 		}
 	}
 
+	for (UAssetUserData* Datum : AssetUserDataEditorOnly)
+	{
+		if (Datum != nullptr)
+		{
+			Datum->PostEditChangeOwner();
+		}
+	}
+
 	if (PropName == GET_MEMBER_NAME_CHECKED(AWorldSettings, bShowInstancedFoliageGrid))
 	{
 		if (bShowInstancedFoliageGrid)
