[UBA] Fix exceeding path limit in UBA input/output files.

#rnx

diff --git a/Engine/Source/Runtime/Engine/Private/ShaderCompiler/ShaderCompilerDistributed.cpp b/Engine/Source/Runtime/Engine/Private/ShaderCompiler/ShaderCompilerDistributed.cpp
--- a/Engine/Source/Runtime/Engine/Private/ShaderCompiler/ShaderCompilerDistributed.cpp
+++ b/Engine/Source/Runtime/Engine/Private/ShaderCompiler/ShaderCompilerDistributed.cpp
@@ -167,6 +167,12 @@ void FShaderCompileDistributedThreadRunnable_Interface::DispatchShaderCompileJob
 	}
 	BaseFilePath += BaseFilePathBuilder.ToView();
 
+	// Ensure input and output files don't exceed platform path limit
+	if (BaseFilePath.Len() + 4 > FPlatformMisc::GetMaxPathLength())
+	{
+		BaseFilePath.LeftChopInline(FPlatformMisc::GetMaxPathLength() - 4);
+	}
+
 	FString InputFilePath = BaseFilePath + TEXT(".in");
 	FString OutputFilePath = BaseFilePath + TEXT(".out");
 
