Fixes  non initialized property error in shader stats.

- LogClass: UInt32Property FShaderCompilerSinglePermutationStat::PermutationString is not initialized properly even though its struct probably has a custom default constructor. Module:Engine File:Public/ShaderCompiler.h

#rnx

diff --git a/Engine/Source/Runtime/Engine/Public/ShaderCompiler.h b/Engine/Source/Runtime/Engine/Public/ShaderCompiler.h
--- a/Engine/Source/Runtime/Engine/Public/ShaderCompiler.h
+++ b/Engine/Source/Runtime/Engine/Public/ShaderCompiler.h
@@ -707,19 +707,19 @@ struct FShaderCompilerSinglePermutationStat
 
 	UE_DEPRECATED(5.6, "PermutationString is no longer stored due to memory overhead; use PermutationHash to uniquely identify permutation stats")
 	UPROPERTY()
-	uint32 PermutationString;
+	uint32 PermutationString = 0u;
 
 	UPROPERTY()
-	uint32 Compiled;
+	uint32 Compiled = 0u;
 
 	UPROPERTY()
-	uint32 Cooked;
+	uint32 Cooked = 0u;
 
 	UPROPERTY()
-	uint32 CompiledDouble;
+	uint32 CompiledDouble = 0u;
 
 	UPROPERTY()
-	uint32 CookedDouble;
+	uint32 CookedDouble = 0u;
 };
 
 USTRUCT()
