StaticMesh Windows cook optimization.  Remove LOD streaming settings from the DDC cache key, allowing the same DDC data to be used for both Editor (streaming disabled) and Cooked (streaming enabled) versions of the StaticMesh render data.  This works because there is no difference in the DDC data between the two versions.  This cuts StaticMesh build time in half, potentially cutting hours off Windows cook time for a large project.  A simple change, but one that needs a verbose CL description to justify:

1.  UStaticMesh::GetEnableLODStreaming disables streaming for the editor target platform (where TargetPlatform->HasEditorOnlyData() returns true).
2.  If no target platform is set, UStaticMesh::GetEnableLODStreaming uses the editor target platform (as returned by GetTargetPlatformManagerRef().GetRunningTargetPlatform() during a cook).
3.  FStaticMeshLODResources::IsLODCookedOut and FStaticMeshLODResources::IsLODInlined are the only callers of UStaticMesh::GetEnableLODStreaming.
4.  FStaticMeshLODResources::Serialize is the only caller of the two functions mentioned in #3.
5.  Where FStaticMeshRenderData::Cache generates DDC data, it calls Serialize with no target platform set.
6.  Therefore, DDC data is always generated for the editor platform, meaning with streaming disabled.  (The DDC requires streaming to be disabled regardless, for the serialized cache to contain all the data!).
7.  All other cook logic that reads streaming LOD settings is also gated by LOD streaming being enabled (NumStreamedLODs and LODGroupSettings accessed by IsLODInlined, call to GetNumOptionalLODsAllowed only reachable when IsLODInlined is false).
8.  In conclusion, the DDC key doesn't need to take into account streaming settings!  Only serialization to packages is affected by streaming.

#rnx

diff --git a/Engine/Source/Runtime/Engine/Private/StaticMesh.cpp b/Engine/Source/Runtime/Engine/Private/StaticMesh.cpp
--- a/Engine/Source/Runtime/Engine/Private/StaticMesh.cpp
+++ b/Engine/Source/Runtime/Engine/Private/StaticMesh.cpp
@@ -3498,28 +3498,6 @@ static FString BuildStaticMeshDerivedDataKeySuffix(const ITargetPlatform* Target
 		KeySuffix += "_";
 	}
 
-	// Mesh LOD streaming settings that need to trigger recache when changed
-	const bool bAllowLODStreaming = Mesh->GetEnableLODStreaming(TargetPlatform);
-	KeySuffix += bAllowLODStreaming ? TEXT("LS1") : TEXT("LS0");
-	KeySuffix += TEXT("MNS");
-	if (bAllowLODStreaming)
-	{
-		int32 MaxNumStreamedLODs = Mesh->NumStreamedLODs.GetValueForPlatform(*TargetPlatform->IniPlatformName());
-		if (MaxNumStreamedLODs < 0)
-		{
-			MaxNumStreamedLODs = LODGroup.GetDefaultMaxNumStreamedLODs();
-		}
-		for (int32 Idx = 0; Idx < 4; ++Idx)
-		{
-			ByteToHex((MaxNumStreamedLODs & 0xff000000) >> 24, KeySuffix);
-			MaxNumStreamedLODs <<= 8;
-		}
-	}
-	else
-	{
-		KeySuffix += TEXT("zzzzzzzz");
-	}
-
 	if (Mesh->bSupportRayTracing && TargetPlatform->UsesRayTracing())
 	{
 		KeySuffix += TEXT("RT1");
