Mac build fix

diff --git a/Engine/Source/Runtime/Engine/Classes/Animation/AnimNodeBase.h b/Engine/Source/Runtime/Engine/Classes/Animation/AnimNodeBase.h
--- a/Engine/Source/Runtime/Engine/Classes/Animation/AnimNodeBase.h
+++ b/Engine/Source/Runtime/Engine/Classes/Animation/AnimNodeBase.h
@@ -162,7 +162,7 @@ public:
 
 	FAnimationUpdateSharedContext* SharedContext;
 
-	FAnimationBaseContext();
+	ENGINE_API FAnimationBaseContext();
 
 protected:
 	// DEPRECATED - Please use constructor that uses an FAnimInstanceProxy*
