Update texture compiler to force rebuild textures using the standard path rather than ForceRebuild.

#rnx

diff --git a/Engine/Source/Runtime/Engine/Classes/Engine/CanvasRenderTarget2D.h b/Engine/Source/Runtime/Engine/Classes/Engine/CanvasRenderTarget2D.h
--- a/Engine/Source/Runtime/Engine/Classes/Engine/CanvasRenderTarget2D.h
+++ b/Engine/Source/Runtime/Engine/Classes/Engine/CanvasRenderTarget2D.h
@@ -30,6 +30,7 @@ class UCanvasRenderTarget2D
 	GENERATED_UCLASS_BODY()
 
 public:
+	ENGINE_API virtual void UpdateResourceWithParams(EUpdateResourceFlags InFlags) override;
 
 	/**
 	 * Updates the the canvas render target texture's resource. This is where the render target will create or 
@@ -40,7 +41,7 @@ public:
 	 * texture to make it available for rendering.
 	 */
 	UFUNCTION(BlueprintCallable, Category="Canvas Render Target 2D")
-	ENGINE_API virtual void UpdateResource() override;
+	ENGINE_API void UpdateResource() { UpdateResourceWithParams(EUpdateResourceFlags::None); }
 
 	/**
 	 * Creates a new canvas render target and initializes it to the specified dimensions
diff --git a/Engine/Source/Runtime/Engine/Classes/Engine/Texture.h b/Engine/Source/Runtime/Engine/Classes/Engine/Texture.h
--- a/Engine/Source/Runtime/Engine/Classes/Engine/Texture.h
+++ b/Engine/Source/Runtime/Engine/Classes/Engine/Texture.h
@@ -1667,9 +1667,29 @@ public:
 	ENGINE_API void ReleaseResource();
 
 	/**
-	 * Creates a new resource for the texture, and updates any cached references to the resource.
-	 */
-	ENGINE_API virtual void UpdateResource();
+	* Creates a new resource for the texture, and updates any cached references to the resource.
+	* If you are deriving this function you should update your class to derive UpdateResourceWithParams as this
+	* function is intended to be a convenience function to access the ::None version of UpdateResourceWithParams and 
+	* only retains virtual to avoid breaking existing code.
+	*/
+	ENGINE_API virtual void UpdateResource() { UpdateResourceWithParams(EUpdateResourceFlags::None); }
+
+	enum class EUpdateResourceFlags : uint32
+	{
+		None			=0,
+		
+		/* Forces the derived data for the texture to completely rebuild. Only applies to editor builds */
+		ForceRebuild	=0x1,
+		/* Forces any derived data fetch/rebuild to run on the calling thread. */
+		Synchronous		=0x2
+	};
+
+	/**
+	* Creates a new resource for the texture, and updates any cached references to the resource. It is always
+	* best practice to call BlockOnAnyAsyncBuild() before assuming the resource is ready if you need it immediately,
+	* even if you use the Synchronous flag to ensure the texture compiler has finished with it.
+	*/
+	ENGINE_API virtual void UpdateResourceWithParams(EUpdateResourceFlags InFlags);
 
 	/**
 	 * Implemented by subclasses to create a new resource for the texture.
@@ -1776,7 +1796,7 @@ public:
 	 * 
 	 * This is called optionally from worker threads via the FAsyncEncode class (LightMaps, ShadowMaps)
 	 */
-	void CachePlatformData(bool bAsyncCache = false, bool bAllowAsyncBuild = false, bool bAllowAsyncLoading = false, class ITextureCompressorModule* Compressor = nullptr);
+	void CachePlatformData(bool bAsyncCache = false, bool bAllowAsyncBuild = false, bool bAllowAsyncLoading = false, class ITextureCompressorModule* Compressor = nullptr, bool bForceRebuild = false);
 
 	/* Returns if it is possible to build this texture for a given TargetPlatform (or any).
 	This supercedes checking Source.IsValid()
@@ -1836,12 +1856,10 @@ public:
 	ENGINE_API void FinishCachePlatformData();
 
 	/**
-	 * Forces platform data to be rebuilt.
-	 * @param InEncodeSpeedOverride		Optionally force a specific encode speed 
-	 * 									using the ETextureEncodeSpeedOverride enum.
-	 * 									Type hidden to keep out of Texture.h
+	 * Forces platform data to be rebuilt. This function is specifically meant for getting the platform data during resource
+	 * creation when it's missing and has subtle caveats. You probably want to use UpdateResourceWithParams(EUpdateResourceFlags::ForceRebuild)!
 	 */
-	ENGINE_API void ForceRebuildPlatformData(uint8 InEncodeSpeedOverride=255);
+	ENGINE_API void ForceRebuildPlatformData();
 
 	/**
 	 * Get an estimate of the peak amount of memory required to build this texture.
diff --git a/Engine/Source/Runtime/Engine/Classes/Engine/Texture2D.h b/Engine/Source/Runtime/Engine/Classes/Engine/Texture2D.h
--- a/Engine/Source/Runtime/Engine/Classes/Engine/Texture2D.h
+++ b/Engine/Source/Runtime/Engine/Classes/Engine/Texture2D.h
@@ -133,7 +133,7 @@ public:
 	virtual TextureAddress GetTextureAddressY() const override { return AddressY; }
 	virtual FTextureResource* CreateResource() override;
 	virtual EMaterialValueType GetMaterialType() const override;
-	virtual void UpdateResource() override;
+	virtual void UpdateResourceWithParams(EUpdateResourceFlags InFlags) override;
 	virtual float GetAverageBrightness(bool bIgnoreTrueBlack, bool bUseGrayscale) override;
 	virtual FTexturePlatformData** GetRunningPlatformData() final override;
 #if WITH_EDITOR
diff --git a/Engine/Source/Runtime/Engine/Classes/Engine/Texture2DArray.cpp b/Engine/Source/Runtime/Engine/Classes/Engine/Texture2DArray.cpp
--- a/Engine/Source/Runtime/Engine/Classes/Engine/Texture2DArray.cpp
+++ b/Engine/Source/Runtime/Engine/Classes/Engine/Texture2DArray.cpp
@@ -198,25 +198,25 @@ FTextureResource* UTexture2DArray::CreateResource()
 	return nullptr;
 }
 
-void UTexture2DArray::UpdateResource()
+void UTexture2DArray::UpdateResourceWithParams(EUpdateResourceFlags InFlags)
 {
 #if WITH_EDITOR
-	// Re-cache platform data if the source has changed.
-	if (FTextureCompilingManager::Get().IsAsyncCompilationAllowed(this))
+	bool bAsyncFlags = false;
+	// Recache platform data if the source has changed.
+	if (!EnumHasAnyFlags(InFlags, EUpdateResourceFlags::Synchronous) &&
+		FTextureCompilingManager::Get().IsAsyncCompilationAllowed(this))
 	{
-		BeginCachePlatformData();
-	}
-	else
-	{
-		CachePlatformData();
+		bAsyncFlags = true;
 	}
+
+	CachePlatformData(bAsyncFlags, bAsyncFlags, bAsyncFlags, nullptr, EnumHasAnyFlags(InFlags, EUpdateResourceFlags::ForceRebuild));
 #endif // #if WITH_EDITOR
 	
 #if WITH_EDITORONLY_DATA
 	bSourceGeneratedFromSourceTexturesArray = !Source.GetNumSlices() || SourceTextures.Num();
 #endif
 
-	Super::UpdateResource();
+	Super::UpdateResourceWithParams(InFlags);
 }
 
 uint32 UTexture2DArray::CalcTextureMemorySize(int32 MipCount) const
diff --git a/Engine/Source/Runtime/Engine/Classes/Engine/Texture2DArray.h b/Engine/Source/Runtime/Engine/Classes/Engine/Texture2DArray.h
--- a/Engine/Source/Runtime/Engine/Classes/Engine/Texture2DArray.h
+++ b/Engine/Source/Runtime/Engine/Classes/Engine/Texture2DArray.h
@@ -67,7 +67,7 @@ public:
 	ENGINE_API void InvalidateTextureSource();
 	ENGINE_API bool CheckArrayTexturesCompatibility();
 #endif // WITH_EDITOR
-	virtual void UpdateResource() override;
+	virtual void UpdateResourceWithParams(EUpdateResourceFlags InFlags) override;
 	virtual EMaterialValueType GetMaterialType() const override { return MCT_Texture2DArray; }
 	virtual FTexturePlatformData** GetRunningPlatformData() override;
 #if WITH_EDITOR
diff --git a/Engine/Source/Runtime/Engine/Classes/Engine/TextureCube.h b/Engine/Source/Runtime/Engine/Classes/Engine/TextureCube.h
--- a/Engine/Source/Runtime/Engine/Classes/Engine/TextureCube.h
+++ b/Engine/Source/Runtime/Engine/Classes/Engine/TextureCube.h
@@ -67,7 +67,7 @@ public:
 	virtual float GetSurfaceDepth() const override { return 0; }
 	virtual uint32 GetSurfaceArraySize() const override { return 6; }
 	ENGINE_API virtual FTextureResource* CreateResource() override;
-	ENGINE_API virtual void UpdateResource() override;
+	ENGINE_API virtual void UpdateResourceWithParams(EUpdateResourceFlags InFlags) override;
 	virtual EMaterialValueType GetMaterialType() const override { return MCT_TextureCube; }
 	ENGINE_API virtual FTexturePlatformData** GetRunningPlatformData() override;
 
diff --git a/Engine/Source/Runtime/Engine/Classes/Engine/TextureCubeArray.h b/Engine/Source/Runtime/Engine/Classes/Engine/TextureCubeArray.h
--- a/Engine/Source/Runtime/Engine/Classes/Engine/TextureCubeArray.h
+++ b/Engine/Source/Runtime/Engine/Classes/Engine/TextureCubeArray.h
@@ -53,7 +53,7 @@ public:
 	ENGINE_API void InvalidateTextureSource();
 	ENGINE_API bool CheckArrayTexturesCompatibility();
 #endif // WITH_EDITOR
-	virtual void UpdateResource() override;
+	virtual void UpdateResourceWithParams(EUpdateResourceFlags InFlags) override;
 	virtual EMaterialValueType GetMaterialType() const override { return MCT_TextureCubeArray; }
 	virtual FTexturePlatformData** GetRunningPlatformData() override;
 #if WITH_EDITOR
diff --git a/Engine/Source/Runtime/Engine/Classes/Engine/VolumeTexture.h b/Engine/Source/Runtime/Engine/Classes/Engine/VolumeTexture.h
--- a/Engine/Source/Runtime/Engine/Classes/Engine/VolumeTexture.h
+++ b/Engine/Source/Runtime/Engine/Classes/Engine/VolumeTexture.h
@@ -97,7 +97,7 @@ public:
 	ENGINE_API void SetDefaultSource2DTileSize();
 	ENGINE_API virtual void PostEditChangeProperty(FPropertyChangedEvent& PropertyChangedEvent) override;
 #endif // WITH_EDITOR
-	virtual void UpdateResource() override;
+	virtual void UpdateResourceWithParams(EUpdateResourceFlags InFlags) override;
 	virtual EMaterialValueType GetMaterialType() const override { return MCT_VolumeTexture; }
 	virtual FTexturePlatformData** GetRunningPlatformData() override;
 #if WITH_EDITOR
diff --git a/Engine/Source/Runtime/Engine/Classes/VT/MeshPaintVirtualTexture.h b/Engine/Source/Runtime/Engine/Classes/VT/MeshPaintVirtualTexture.h
--- a/Engine/Source/Runtime/Engine/Classes/VT/MeshPaintVirtualTexture.h
+++ b/Engine/Source/Runtime/Engine/Classes/VT/MeshPaintVirtualTexture.h
@@ -24,7 +24,7 @@ class UMeshPaintVirtualTexture : public UTexture2D
 
 	//~ Begin UTexture Interface.
 	virtual void GetVirtualTextureBuildSettings(FVirtualTextureBuildSettings& OutSettings) const override;
-	virtual void UpdateResource() override;
+	virtual void UpdateResourceWithParams(EUpdateResourceFlags InFlags) override;
 
 #if WITH_EDITOR
 	virtual void BeginCacheForCookedPlatformData(const ITargetPlatform* TargetPlatform) override;
diff --git a/Engine/Source/Runtime/Engine/Private/CanvasRenderTarget2D.cpp b/Engine/Source/Runtime/Engine/Private/CanvasRenderTarget2D.cpp
--- a/Engine/Source/Runtime/Engine/Private/CanvasRenderTarget2D.cpp
+++ b/Engine/Source/Runtime/Engine/Private/CanvasRenderTarget2D.cpp
@@ -24,7 +24,7 @@ UCanvasRenderTarget2D::UCanvasRenderTarget2D( const FObjectInitializer& ObjectIn
 }
 
 
-void UCanvasRenderTarget2D::UpdateResource()
+void UCanvasRenderTarget2D::UpdateResourceWithParams(EUpdateResourceFlags InFlags)
 {
 	if (SampleCount != ETextureRenderTargetSampleCount::RTSC_1)
 	{
@@ -32,7 +32,7 @@ void UCanvasRenderTarget2D::UpdateResource()
 	}
 
 	// Call parent implementation
-	Super::UpdateResource();
+	Super::UpdateResourceWithParams(InFlags);
 
 	// Don't allocate canvas object for CRT2D CDO; also, we can't update it during PostLoad!
 	if (IsTemplate() || FUObjectThreadContext::Get().IsRoutingPostLoad)
diff --git a/Engine/Source/Runtime/Engine/Private/Texture.cpp b/Engine/Source/Runtime/Engine/Private/Texture.cpp
--- a/Engine/Source/Runtime/Engine/Private/Texture.cpp
+++ b/Engine/Source/Runtime/Engine/Private/Texture.cpp
@@ -314,7 +314,7 @@ void UTexture::ReleaseResource()
 	}
 }
 
-void UTexture::UpdateResource()
+void UTexture::UpdateResourceWithParams(EUpdateResourceFlags InFlags)
 {
 	// Release the existing texture resource.
 	ReleaseResource();
diff --git a/Engine/Source/Runtime/Engine/Private/Texture2D.cpp b/Engine/Source/Runtime/Engine/Private/Texture2D.cpp
--- a/Engine/Source/Runtime/Engine/Private/Texture2D.cpp
+++ b/Engine/Source/Runtime/Engine/Private/Texture2D.cpp
@@ -722,7 +722,7 @@ void UTexture2D::GetAssetRegistryTags(FAssetRegistryTagsContext Context) const
 	Super::GetAssetRegistryTags(Context);
 }
 
-void UTexture2D::UpdateResource()
+void UTexture2D::UpdateResourceWithParams(EUpdateResourceFlags InFlags)
 {
 	TRACE_CPUPROFILER_EVENT_SCOPE(UTexture2D::UpdateResource);
 
@@ -733,15 +733,15 @@ void UTexture2D::UpdateResource()
 	// on access if needed.
 	CPUCopyTexture = nullptr;
 
+	bool bAsyncFlags = false;
 	// Recache platform data if the source has changed.
-	if (FTextureCompilingManager::Get().IsAsyncCompilationAllowed(this))
+	if (!EnumHasAnyFlags(InFlags, EUpdateResourceFlags::Synchronous) &&
+		FTextureCompilingManager::Get().IsAsyncCompilationAllowed(this))
 	{
-		BeginCachePlatformData();
-	}
-	else
-	{
-		CachePlatformData();
+		bAsyncFlags = true;
 	}
+
+	CachePlatformData(bAsyncFlags, bAsyncFlags, bAsyncFlags, nullptr, EnumHasAnyFlags(InFlags, EUpdateResourceFlags::ForceRebuild));
 	
 	// clear all the cooked cached platform data if the source could have changed... 
 	ClearAllCachedCookedPlatformData();
@@ -770,7 +770,7 @@ void UTexture2D::UpdateResource()
 #endif // #if WITH_EDITOR
 
 	// Route to super.
-	Super::UpdateResource();
+	Super::UpdateResourceWithParams(InFlags);
 }
 
 
diff --git a/Engine/Source/Runtime/Engine/Private/TextureCompiler.cpp b/Engine/Source/Runtime/Engine/Private/TextureCompiler.cpp
--- a/Engine/Source/Runtime/Engine/Private/TextureCompiler.cpp
+++ b/Engine/Source/Runtime/Engine/Private/TextureCompiler.cpp
@@ -848,7 +848,10 @@ void FTextureCompilingManager::ProcessDeferredRequests()
 	{
 		for (UTexture* DeferredTexture : DeferredTextures)
 		{
-			DeferredTexture->ForceRebuildPlatformData();
+			DeferredTexture->UpdateResourceWithParams(UTexture::EUpdateResourceFlags::ForceRebuild);
+
+			// The texture should now be in the normal texture compilation path. Since we don't need it immediately
+			// we don't FinishCompilation on it.
 		}
 	}
 }
diff --git a/Engine/Source/Runtime/Engine/Private/TextureCube.cpp b/Engine/Source/Runtime/Engine/Private/TextureCube.cpp
--- a/Engine/Source/Runtime/Engine/Private/TextureCube.cpp
+++ b/Engine/Source/Runtime/Engine/Private/TextureCube.cpp
@@ -201,22 +201,22 @@ void UTextureCube::GetAssetRegistryTags(FAssetRegistryTagsContext Context) const
 	Super::GetAssetRegistryTags(Context);
 }
 
-void UTextureCube::UpdateResource()
+void UTextureCube::UpdateResourceWithParams(EUpdateResourceFlags InFlags)
 {
 #if WITH_EDITOR
+	bool bAsyncFlags = false;
 	// Recache platform data if the source has changed.
-	if (FTextureCompilingManager::Get().IsAsyncCompilationAllowed(this))
+	if (!EnumHasAnyFlags(InFlags, EUpdateResourceFlags::Synchronous) &&
+		FTextureCompilingManager::Get().IsAsyncCompilationAllowed(this))
 	{
-		BeginCachePlatformData();
-	}
-	else
-	{
-		CachePlatformData();
+		bAsyncFlags = true;
 	}
+
+	CachePlatformData(bAsyncFlags, bAsyncFlags, bAsyncFlags, nullptr, EnumHasAnyFlags(InFlags, EUpdateResourceFlags::ForceRebuild));
 #endif // #if WITH_EDITOR
 
 	// Route to super.
-	Super::UpdateResource();
+	Super::UpdateResourceWithParams(InFlags);
 }
 
 FString UTextureCube::GetDesc()
diff --git a/Engine/Source/Runtime/Engine/Private/TextureCubeArray.cpp b/Engine/Source/Runtime/Engine/Private/TextureCubeArray.cpp
--- a/Engine/Source/Runtime/Engine/Private/TextureCubeArray.cpp
+++ b/Engine/Source/Runtime/Engine/Private/TextureCubeArray.cpp
@@ -460,25 +460,25 @@ FTextureResource* UTextureCubeArray::CreateResource()
 	return nullptr;
 }
 
-void UTextureCubeArray::UpdateResource()
+void UTextureCubeArray::UpdateResourceWithParams(EUpdateResourceFlags InFlags)
 {
 #if WITH_EDITOR
-	// Re-cache platform data if the source has changed.
-	if (FTextureCompilingManager::Get().IsAsyncCompilationAllowed(this))
+	bool bAsyncFlags = false;
+	// Recache platform data if the source has changed.
+	if (!EnumHasAnyFlags(InFlags, EUpdateResourceFlags::Synchronous) &&
+		FTextureCompilingManager::Get().IsAsyncCompilationAllowed(this))
 	{
-		BeginCachePlatformData();
-	}
-	else
-	{
-		CachePlatformData();
+		bAsyncFlags = true;
 	}
+
+	CachePlatformData(bAsyncFlags, bAsyncFlags, bAsyncFlags, nullptr, EnumHasAnyFlags(InFlags, EUpdateResourceFlags::ForceRebuild));
 #endif // #if WITH_EDITOR
 
 #if WITH_EDITORONLY_DATA
 	bSourceGeneratedFromSourceTexturesArray = !Source.GetNumSlices() || SourceTextures.Num();
 #endif
 
-	Super::UpdateResource();
+	Super::UpdateResourceWithParams(InFlags);
 }
 
 uint32 UTextureCubeArray::CalcTextureMemorySize(int32 MipCount) const
diff --git a/Engine/Source/Runtime/Engine/Private/TextureDerivedData.cpp b/Engine/Source/Runtime/Engine/Private/TextureDerivedData.cpp
--- a/Engine/Source/Runtime/Engine/Private/TextureDerivedData.cpp
+++ b/Engine/Source/Runtime/Engine/Private/TextureDerivedData.cpp
@@ -1773,7 +1773,7 @@ void FTexturePlatformData::Cache(
 		Flags |= ETextureCacheFlags::ForVirtualTextureStreamingBuild;
 	}
 
-	if (bAsync && !bForceRebuild)
+	if (bAsync)
 	{
 		FQueuedThreadPool*  TextureThreadPool = FTextureCompilingManager::Get().GetThreadPool();
 		EQueuedWorkPriority BasePriority      = FTextureCompilingManager::Get().GetBasePriority(&InTexture);
@@ -3521,7 +3521,7 @@ bool UTexture::CanBuildPlatformData(const ITargetPlatformSettings * TargetPlatfo
 	return true;
 }
 
-void UTexture::CachePlatformData(bool bAsyncCache, bool bAllowAsyncBuild, bool bAllowAsyncLoading, ITextureCompressorModule* Compressor)
+void UTexture::CachePlatformData(bool bAsyncCache, bool bAllowAsyncBuild, bool bAllowAsyncLoading, ITextureCompressorModule* Compressor, bool bForceRebuild)
 {
 	//
 	// NOTE this can be called off the main thread via FAsyncEncode<> for shadow/light maps!
@@ -3541,7 +3541,8 @@ void UTexture::CachePlatformData(bool bAsyncCache, bool bAllowAsyncBuild, bool b
 			ETextureCacheFlags CacheFlags =
 				(bAsyncCache ? ETextureCacheFlags::Async : ETextureCacheFlags::None) |
 				(bAllowAsyncBuild? ETextureCacheFlags::AllowAsyncBuild : ETextureCacheFlags::None) |
-				(bAllowAsyncLoading? ETextureCacheFlags::AllowAsyncLoading : ETextureCacheFlags::None);
+				(bAllowAsyncLoading? ETextureCacheFlags::AllowAsyncLoading : ETextureCacheFlags::None) |
+				(bForceRebuild ? ETextureCacheFlags::ForceRebuild : ETextureCacheFlags::None);
 
 			ETextureEncodeSpeed EncodeSpeed = GetDesiredEncodeSpeed();
 
@@ -3620,7 +3621,7 @@ void UTexture::CachePlatformData(bool bAsyncCache, bool bAllowAsyncBuild, bool b
 
 			// The only time we don't cache is if we a) have existing data and b) it matches what we want.
 			bPerformCache = true;
-			if (PlatformDataLink != nullptr)
+			if (PlatformDataLink != nullptr && !EnumHasAnyFlags(CacheFlags, ETextureCacheFlags::ForceRebuild))
 			{
 				bPerformCache = false;
 
@@ -3714,21 +3715,6 @@ void UTexture::CachePlatformData(bool bAsyncCache, bool bAllowAsyncBuild, bool b
 void UTexture::BeginCachePlatformData()
 {
 	CachePlatformData(true, true, true);
-
-#if 0 // don't cache in post load, this increases our peak memory usage, instead cache just before we save the package
-	// enable caching in postload for derived data cache commandlet and cook by the book
-	ITargetPlatformManagerModule* TPM = GetTargetPlatformManager();
-	if (TPM && (TPM->RestrictFormatsToRuntimeOnly() == false))
-	{
-		ITargetPlatformManagerModule* TPM = GetTargetPlatformManager();
-		TArray<ITargetPlatform*> Platforms = TPM->GetActiveTargetPlatforms();
-		// Cache for all the shader formats that the cooking target requires
-		for (int32 FormatIndex = 0; FormatIndex < Platforms.Num(); FormatIndex++)
-		{
-			BeginCacheForCookedPlatformData(Platforms[FormatIndex]);
-		}
-	}
-#endif
 }
 
 void UTexture::BeginCacheForCookedPlatformData( const ITargetPlatform *TargetPlatform )
@@ -4054,7 +4040,7 @@ void UTexture::FinishCachePlatformData()
 	//  this is not a good place to put code that finalizes caching
 }
 
-void UTexture::ForceRebuildPlatformData(uint8 InEncodeSpeedOverride /* =255 ETextureEncodeSpeedOverride::Disabled */)
+void UTexture::ForceRebuildPlatformData()
 {
 	TRACE_CPUPROFILER_EVENT_SCOPE(UTexture::ForceRebuildPlatformData)
 
@@ -4070,15 +4056,7 @@ void UTexture::ForceRebuildPlatformData(uint8 InEncodeSpeedOverride /* =255 ETex
 		FTexturePlatformData *&PlatformDataLink = *PlatformDataLinkPtr;
 		FlushRenderingCommands();
 
-		ETextureEncodeSpeed EncodeSpeed;
-		if (InEncodeSpeedOverride != (uint8)ETextureEncodeSpeedOverride::Disabled)
-		{
-			EncodeSpeed = (ETextureEncodeSpeed)InEncodeSpeedOverride;
-		}
-		else
-		{
-			EncodeSpeed = GetDesiredEncodeSpeed();
-		}
+		ETextureEncodeSpeed EncodeSpeed = GetDesiredEncodeSpeed();
 
 		// Since we are forcing a rebuild, build what is desired rather than what is available
 		if (EncodeSpeed == ETextureEncodeSpeed::FinalIfAvailable)
diff --git a/Engine/Source/Runtime/Engine/Private/TextureDerivedDataTask.cpp b/Engine/Source/Runtime/Engine/Private/TextureDerivedDataTask.cpp
--- a/Engine/Source/Runtime/Engine/Private/TextureDerivedDataTask.cpp
+++ b/Engine/Source/Runtime/Engine/Private/TextureDerivedDataTask.cpp
@@ -1323,7 +1323,6 @@ FTextureCacheDerivedDataWorker::FTextureCacheDerivedDataWorker(
 	DerivedData->Reset();
 	UTexture::GetPixelFormatEnum();
 		
-	const bool bAllowAsyncBuild = EnumHasAnyFlags(CacheFlags, ETextureCacheFlags::AllowAsyncBuild);
 	const bool bAllowAsyncLoading = EnumHasAnyFlags(CacheFlags, ETextureCacheFlags::AllowAsyncLoading);
 	const bool bForVirtualTextureStreamingBuild = EnumHasAnyFlags(CacheFlags, ETextureCacheFlags::ForVirtualTextureStreamingBuild);
 
@@ -1707,8 +1706,7 @@ static bool DDC1_LoadAndValidateTextureData(
 	UTexture& Texture,
 	FTextureSourceData& TextureData,
 	FTextureSourceData& CompositeTextureData,
-	IImageWrapperModule* ImageWrapper,
-	bool bAllowAsyncLoading
+	IImageWrapperModule* ImageWrapper
 )
 {
 	TRACE_CPUPROFILER_EVENT_SCOPE(Texture.DDC1_LoadAndValidateTextureData);
@@ -1803,7 +1801,7 @@ bool DDC1_BuildTiledClassicTexture(
 	if (bHasLinearDerivedData == false)
 	{
 		// Linear data didn't exist, need to build it.
-		bool bGotSourceTextureData = DDC1_LoadAndValidateTextureData(Texture, TextureData, CompositeTextureData, ImageWrapper, EnumHasAnyFlags(CacheFlags, ETextureCacheFlags::AllowAsyncLoading));
+		bool bGotSourceTextureData = DDC1_LoadAndValidateTextureData(Texture, TextureData, CompositeTextureData, ImageWrapper);
 		if (bGotSourceTextureData)
 		{
 			// We know we want all the mips for tiling, so force inline
@@ -2030,7 +2028,7 @@ void FTextureCacheDerivedDataWorker::DoWork()
 
 	if (!bSucceeded && bAllowAsyncBuild)
 	{
-		if (DDC1_LoadAndValidateTextureData(Texture, TextureData, CompositeTextureData, ImageWrapper, bAllowAsyncLoading))
+		if (DDC1_LoadAndValidateTextureData(Texture, TextureData, CompositeTextureData, ImageWrapper))
 		{
 			for (int32 LayerIndex = 0; LayerIndex < BuildSettingsPerLayerFetchOrBuild.Num(); LayerIndex++)
 			{
diff --git a/Engine/Source/Runtime/Engine/Private/VT/MeshPaintVirtualTexture.cpp b/Engine/Source/Runtime/Engine/Private/VT/MeshPaintVirtualTexture.cpp
--- a/Engine/Source/Runtime/Engine/Private/VT/MeshPaintVirtualTexture.cpp
+++ b/Engine/Source/Runtime/Engine/Private/VT/MeshPaintVirtualTexture.cpp
@@ -301,9 +301,9 @@ void UMeshPaintVirtualTexture::GetVirtualTextureBuildSettings(FVirtualTextureBui
 	OutSettings.TileBorderSize = MeshPaintVirtualTexture::GetTileBorderSize();
 }
 
-void UMeshPaintVirtualTexture::UpdateResource()
+void UMeshPaintVirtualTexture::UpdateResourceWithParams(EUpdateResourceFlags InFlags)
 {
-	Super::UpdateResource();
+	Super::UpdateResourceWithParams(InFlags);
 
 	// We get here on virtual texture pool recreation, and on texture compilation in editor.
 	// In those cases we need to reacquire the virtual texture, and notify our component.
diff --git a/Engine/Source/Runtime/Engine/Private/VolumeTexture.cpp b/Engine/Source/Runtime/Engine/Private/VolumeTexture.cpp
--- a/Engine/Source/Runtime/Engine/Private/VolumeTexture.cpp
+++ b/Engine/Source/Runtime/Engine/Private/VolumeTexture.cpp
@@ -361,22 +361,22 @@ void UVolumeTexture::GetAssetRegistryTags(FAssetRegistryTagsContext Context) con
 	Super::GetAssetRegistryTags(Context);
 }
 
-void UVolumeTexture::UpdateResource()
+void UVolumeTexture::UpdateResourceWithParams(EUpdateResourceFlags InFlags)
 {
 #if WITH_EDITOR
+	bool bAsyncFlags = false;
 	// Recache platform data if the source has changed.
-	if (FTextureCompilingManager::Get().IsAsyncCompilationAllowed(this))
+	if (!EnumHasAnyFlags(InFlags, EUpdateResourceFlags::Synchronous) &&
+		FTextureCompilingManager::Get().IsAsyncCompilationAllowed(this))
 	{
-		BeginCachePlatformData();
-	}
-	else
-	{
-		CachePlatformData();
+		bAsyncFlags = true;
 	}
+
+	CachePlatformData(bAsyncFlags, bAsyncFlags, bAsyncFlags, nullptr, EnumHasAnyFlags(InFlags, EUpdateResourceFlags::ForceRebuild));
 #endif // #if WITH_EDITOR
 
 	// Route to super.
-	Super::UpdateResource();
+	Super::UpdateResourceWithParams(InFlags);
 }
 
 FString UVolumeTexture::GetDesc()
diff --git a/Engine/Source/Runtime/MediaAssets/Public/MediaTexture.h b/Engine/Source/Runtime/MediaAssets/Public/MediaTexture.h
--- a/Engine/Source/Runtime/MediaAssets/Public/MediaTexture.h
+++ b/Engine/Source/Runtime/MediaAssets/Public/MediaTexture.h
@@ -162,10 +162,11 @@ public:
 	
 	/**
 	 * Creates a new resource for the texture, and updates any cached references to the resource.
-	 * This obviously is just an override to expose to blueprints.
+	 * This obviously is just an override to expose to blueprints. Since the base is deprecated in favor
+	 * of the new params version, we remove the virtual override and pass the flags we care about here.
 	 */
 	UFUNCTION(BlueprintCallable, Category = "Media|MediaTexture")
-	virtual void UpdateResource() override  { Super::UpdateResource(); }
+	void UpdateResource() { UpdateResourceWithParams(EUpdateResourceFlags::None); }
 
 	/**
 	 * Caches the next available sample time from the queue when last rendering was made
