[Visual Logger] Add support for Coordinate System entries for VisLog

diff --git a/Engine/Plugins/Animation/GameplayInsights/Source/RewindDebuggerVLog/Private/VisualLogEntryRenderer.cpp b/Engine/Plugins/Animation/GameplayInsights/Source/RewindDebuggerVLog/Private/VisualLogEntryRenderer.cpp
--- a/Engine/Plugins/Animation/GameplayInsights/Source/RewindDebuggerVLog/Private/VisualLogEntryRenderer.cpp
+++ b/Engine/Plugins/Animation/GameplayInsights/Source/RewindDebuggerVLog/Private/VisualLogEntryRenderer.cpp
@@ -422,6 +422,24 @@ void FVisualLogEntryRenderer::RenderLogEntry(class UWorld* World, const FVisualL
 				RenderDescription(World, Canvas, Font, ElementToDraw, Center, Color, Index/3, ElementToDraw->Points.Num()/3);
 			}
 		}
+			break;
+		case EVisualLoggerShapeElement::CoordinateSystem:
+		{
+			const float Thickness = float(ElementToDraw->Thickness);
+			const bool bDrawLabel = ElementToDraw->Description.IsEmpty() == false;
+			const int32 NumPoints = ElementToDraw->Points.Num();
+
+			for (int32 Index = 0; Index + 2 < NumPoints; Index += 3)
+			{
+				const FVector AxisLoc = ElementToDraw->Points[Index];
+				const FRotator AxisRot = ElementToDraw->Points[Index + 1].ToOrientationRotator();
+				const float Scale = ElementToDraw->Points[Index + 2].X;
+
+				DrawDebugCoordinateSystem(World, AxisLoc, AxisRot, Scale, /*bPersistentLines*/false, /*LifeTime*/-1, DepthPriority, Thickness);
+				RenderDescription(World, Canvas, Font, ElementToDraw, AxisLoc, Color, Index/3, ElementToDraw->Points.Num()/3);
+			}
+		}
+
 			break;
 		case EVisualLoggerShapeElement::Invalid:
 			UE_LOG(LogVisual, Warning, TEXT("Invalid element type"));
diff --git a/Engine/Source/Developer/LogVisualizer/Private/VisualLoggerRenderingActorBase.cpp b/Engine/Source/Developer/LogVisualizer/Private/VisualLoggerRenderingActorBase.cpp
--- a/Engine/Source/Developer/LogVisualizer/Private/VisualLoggerRenderingActorBase.cpp
+++ b/Engine/Source/Developer/LogVisualizer/Private/VisualLoggerRenderingActorBase.cpp
@@ -88,6 +88,7 @@ FDebugRenderSceneProxy* UVisualLoggerRenderingComponent::CreateDebugSceneProxy()
 			VLogSceneProxy->Cylinders.Append(Shapes.Cylinders);
 			VLogSceneProxy->ArrowLines.Append(Shapes.Arrows);
 			VLogSceneProxy->Capsules.Append(Shapes.Capsules);
+			VLogSceneProxy->CoordinateSystems.Append(Shapes.CoordinateSystems);
 		});
 
 	return VLogSceneProxy;
@@ -556,6 +557,28 @@ void AVisualLoggerRenderingActorBase::GetDebugShapes(const FVisualLogEntry& InEn
 			}
 		}
 			break;
+		case EVisualLoggerShapeElement::CoordinateSystem:
+		{
+			const float Thickness = float(ElementToDraw->Thickness);
+			const bool bDrawLabel = ElementToDraw->Description.IsEmpty() == false;
+			const int32 NumPoints = ElementToDraw->Points.Num();
+
+			for (int32 Index = 0; Index + 2 < NumPoints; Index += 3)
+			{
+				const FVector AxisLoc = ElementToDraw->Points[Index];
+				const FRotator AxisRot = ElementToDraw->Points[Index + 1].ToOrientationRotator();
+				const float Scale = ElementToDraw->Points[Index + 2].X;
+
+				DebugShapes.CoordinateSystems.Add(FDebugRenderSceneProxy::FCoordinateSystem(AxisLoc, AxisRot, Scale, Color, Thickness));
+
+				if (bDrawLabel)
+				{
+					const FString PrintString = NumPoints == 3 ? ElementToDraw->Description : FString::Printf(TEXT("%s_%d"), *ElementToDraw->Description, Index / 3);
+					DebugShapes.Texts.Add(FDebugRenderSceneProxy::FText3d(PrintString, AxisLoc, Color.WithAlpha(255)));
+				}
+			}
+		}
+			break;
 		case EVisualLoggerShapeElement::Invalid:
 			UE_LOG(LogVisual, Warning, TEXT("Invalid element type"));
 			break;
diff --git a/Engine/Source/Developer/LogVisualizer/Public/VisualLoggerRenderingActorBase.h b/Engine/Source/Developer/LogVisualizer/Public/VisualLoggerRenderingActorBase.h
--- a/Engine/Source/Developer/LogVisualizer/Public/VisualLoggerRenderingActorBase.h
+++ b/Engine/Source/Developer/LogVisualizer/Public/VisualLoggerRenderingActorBase.h
@@ -35,6 +35,7 @@ public:
 		TArray<FDebugRenderSceneProxy::FWireCylinder> Cylinders;
 		TArray<FDebugRenderSceneProxy::FCapsule> Capsules;
 		TArray<FDebugRenderSceneProxy::FArrowLine> Arrows;
+		TArray<FDebugRenderSceneProxy::FCoordinateSystem> CoordinateSystems;
 		TArray<FVector> LogEntriesPath;
 
 		void Reset()
@@ -49,6 +50,7 @@ public:
 			Cylinders.Reset();
 			Capsules.Reset();
 			Arrows.Reset();
+			CoordinateSystems.Reset();
 			LogEntriesPath.Reset();
 		}
 	};
diff --git a/Engine/Source/Runtime/Engine/Private/DebugRenderSceneProxy.cpp b/Engine/Source/Runtime/Engine/Private/DebugRenderSceneProxy.cpp
--- a/Engine/Source/Runtime/Engine/Private/DebugRenderSceneProxy.cpp
+++ b/Engine/Source/Runtime/Engine/Private/DebugRenderSceneProxy.cpp
@@ -106,7 +106,8 @@ uint32 FDebugRenderSceneProxy::GetAllocatedSize(void) const
 		Lines.GetAllocatedSize() + 
 		Boxes.GetAllocatedSize() + 
 		Spheres.GetAllocatedSize() +
-		Texts.GetAllocatedSize();
+		Texts.GetAllocatedSize() +
+		CoordinateSystems.GetAllocatedSize();
 }
 
 
@@ -320,6 +321,15 @@ void FDebugRenderSceneProxy::GetDynamicMeshElementsForView(const FSceneView* Vie
 		FMaterialCache& MeshMaterialCache = Mesh.Color.A == 255 ? SolidMeshMaterialCache : DefaultMaterialCache;
 		MeshBuilder.GetMesh(FMatrix::Identity, MeshMaterialCache[Mesh.Color], SDPG_World, false, false, ViewIndex, Collector);
 	}
+
+	// Draw Coordinate Systems
+	for (const FCoordinateSystem CoordinateSystem : CoordinateSystems)
+	{
+		if (!bAllowFrustumCulling || SphereInFrustum(CoordinateSystem.AxisLoc, CoordinateSystem.Scale, AdjustedFrustum))
+		{
+			CoordinateSystem.Draw(PDI);
+		}
+	}
 }
 
 FConvexVolume FDebugRenderSceneProxy::AdjustViewFrustumForFarClipping(const FSceneView* InView, double InFarClippingDistance)
@@ -520,4 +530,9 @@ void FDebugRenderSceneProxy::FCapsule::Draw(FPrimitiveDrawInterface* PDI, EDrawT
 	{
 		GetCapsuleMesh(Base, X, Y, Z, Color, Radius, HalfHeight, (InDrawType == SolidAndWireMeshes) ? 9 : 16, MaterialCache[Color.WithAlpha(InDrawAlpha * Color.A)], SDPG_World, false, ViewIndex, Collector);
 	}
+}
+
+void FDebugRenderSceneProxy::FCoordinateSystem::Draw(FPrimitiveDrawInterface* PDI) const
+{
+	DrawCoordinateSystem(PDI, AxisLoc, AxisRot, Scale, SDPG_World, Thickness);
 }
\ No newline at end of file
diff --git a/Engine/Source/Runtime/Engine/Private/VisualLogger/VisualLogger.cpp b/Engine/Source/Runtime/Engine/Private/VisualLogger/VisualLogger.cpp
--- a/Engine/Source/Runtime/Engine/Private/VisualLogger/VisualLogger.cpp
+++ b/Engine/Source/Runtime/Engine/Private/VisualLogger/VisualLogger.cpp
@@ -177,6 +177,13 @@ void FVisualLogger::HistogramDataLogfImpl(const UObject* Object, const FName& Ca
 	);
 }
 
+void FVisualLogger::CoordinateSystemLogfImpl(const UObject* Object, const FName& CategoryName, ELogVerbosity::Type Verbosity, const FVector& AxisLoc, const FRotator& AxisRot, const float Scale, const FColor& Color, const uint16 Thickness, const TCHAR* Fmt, ...)
+{
+	COLLAPSED_LOGF(
+		CurrentEntry->AddCoordinateSystem(AxisLoc, AxisRot, Scale, CategoryName, Verbosity, Color, Buffer, Thickness);
+	);
+}
+
 #undef COLLAPSED_LOGF
 
 
diff --git a/Engine/Source/Runtime/Engine/Private/VisualLogger/VisualLoggerTypes.cpp b/Engine/Source/Runtime/Engine/Private/VisualLogger/VisualLoggerTypes.cpp
--- a/Engine/Source/Runtime/Engine/Private/VisualLogger/VisualLoggerTypes.cpp
+++ b/Engine/Source/Runtime/Engine/Private/VisualLogger/VisualLoggerTypes.cpp
@@ -305,6 +305,19 @@ void FVisualLogEntry::AddHistogramData(const FVector2D& DataSample, const FName&
 	HistogramSamples.Add(Sample);
 }
 
+void FVisualLogEntry::AddCoordinateSystem(const FVector& AxisLoc, const FRotator& AxisRot, const float Scale, const FName& CategoryName, ELogVerbosity::Type Verbosity, const FColor& Color, const FString& Description, uint16 Thickness)
+{
+	FVisualLogShapeElement Element(Description, Color, Thickness, CategoryName);
+	Element.Type = EVisualLoggerShapeElement::CoordinateSystem;
+	Element.Verbosity = Verbosity;
+	Element.Points.Reserve(3);
+	Element.Points.Add(AxisLoc);
+	Element.Points.Add(AxisRot.Vector());
+	Element.Points.Add(FVector(Scale, 0., 0.));
+
+	ElementsToDraw.Add(Element);
+}
+
 FVisualLogDataBlock& FVisualLogEntry::AddDataBlock(const FString& TagName, const TArray<uint8>& BlobDataArray, const FName& CategoryName, ELogVerbosity::Type Verbosity)
 {
 	FVisualLogDataBlock DataBlock;
diff --git a/Engine/Source/Runtime/Engine/Public/DebugRenderSceneProxy.h b/Engine/Source/Runtime/Engine/Public/DebugRenderSceneProxy.h
--- a/Engine/Source/Runtime/Engine/Public/DebugRenderSceneProxy.h
+++ b/Engine/Source/Runtime/Engine/Public/DebugRenderSceneProxy.h
@@ -361,6 +361,27 @@ public:
 		EDrawType DrawTypeOverride = EDrawType::Invalid;
 	};
 
+	// struct to hold coordinate systems to render
+	struct FCoordinateSystem
+	{
+		FCoordinateSystem() = delete;
+		FCoordinateSystem(const FVector& InAxisLoc, const FRotator& InAxisRot, const float InScale, const FColor& InColor, const float InThickness)
+			: AxisLoc(InAxisLoc)
+			, AxisRot(InAxisRot)
+			, Scale(InScale)
+			, Color(InColor)
+			, Thickness(InThickness)
+		{}
+
+		ENGINE_API void Draw(FPrimitiveDrawInterface* PDI) const;
+
+		FVector AxisLoc;
+		FRotator AxisRot;
+		float Scale;
+		FColor Color;
+		float Thickness;
+	};
+
 	TArray<FDebugLine> Lines;
 	TArray<FDashedLine>	DashedLines;
 	TArray<FArrowLine> ArrowLines;
@@ -373,6 +394,7 @@ public:
 	TArray<FCone> Cones;
 	TArray<FMesh> Meshes;
 	TArray<FCapsule> Capsules;
+	TArray<FCoordinateSystem> CoordinateSystems;
 
 	uint32 ViewFlagIndex;
 	float TextWithoutShadowDistance;
diff --git a/Engine/Source/Runtime/Engine/Public/VisualLogger/VisualLogger.h b/Engine/Source/Runtime/Engine/Public/VisualLogger/VisualLogger.h
--- a/Engine/Source/Runtime/Engine/Public/VisualLogger/VisualLogger.h
+++ b/Engine/Source/Runtime/Engine/Public/VisualLogger/VisualLogger.h
@@ -107,6 +107,9 @@
 // Wire circle shape
 #define UE_VLOG_WIRECIRCLE_THICK(LogOwner, CategoryName, Verbosity, Center, UpAxis, Radius, Color, Thickness, Format, ...) if(FVisualLogger::IsRecording()) FVisualLogger::DiscLogf(LogOwner, CategoryName, ELogVerbosity::Verbosity, Center, UpAxis, Radius, Color, Thickness, /*bWireframe = */true, Format, ##__VA_ARGS__)
 #define UE_CVLOG_WIRECIRCLE_THICK(Condition, LogOwner, CategoryName, Verbosity, Center, UpAxis, Radius, Color, Thickness, Format, ...) if(FVisualLogger::IsRecording() && Condition) {UE_VLOG_WIRECIRCLE_THICK(LogOwner, CategoryName, Verbosity, Center, UpAxis, Radius, Color, Thickness, Format, ##__VA_ARGS__);} 
+// Coordinate system
+#define UE_VLOG_COORDINATESYSTEM(LogOwner, CategoryName, Verbosity, AxisLoc, AxisRot, Scale, Color, Thickness, Format, ...) if(FVisualLogger::IsRecording()) FVisualLogger::CoordinateSystemLogf(LogOwner, CategoryName, ELogVerbosity::Verbosity, AxisLoc, AxisRot, Scale, Color, Thickness, Format, ##__VA_ARGS__)
+#define UE_CVLOG_COORDINATESYSTEM(Condition, LogOwner, CategoryName, Verbosity, AxisLoc, AxisRot, Scale, Color, Thickness, Format, ...) if(FVisualLogger::IsRecording() && Condition) {UE_VLOG_COORDINATESYSTEM(LogOwner, CategoryName, Verbosity, AxisLoc, AxisRot, Scale, Color, Thickness, Format, ##__VA_ARGS__);}
 
 #define DECLARE_VLOG_EVENT(EventName) extern FVisualLogEventBase EventName;
 #define DEFINE_VLOG_EVENT(EventName, Verbosity, UserFriendlyDesc) FVisualLogEventBase EventName(TEXT(#EventName), TEXT(UserFriendlyDesc), ELogVerbosity::Verbosity); 
@@ -241,6 +244,7 @@ class FVisualLogger : public FOutputDevice
 	UE_DEPRECATED(5.6, "Use DiscLogfImpl")
 	static ENGINE_API void CircleLogfImpl(const UObject* LogOwner, const FName& CategoryName, ELogVerbosity::Type Verbosity, const FVector& Center, const FVector& UpAxis, const float Radius, const FColor& Color, const uint16 Thickness, const TCHAR* Fmt, ...);
 	static ENGINE_API void DiscLogfImpl(const UObject* LogOwner, const FName& CategoryName, ELogVerbosity::Type Verbosity, const FVector& Center, const FVector& UpAxis, const float Radius, const FColor& Color, const uint16 Thickness, bool bWireframe, const TCHAR* Fmt, ...);
+	static ENGINE_API void CoordinateSystemLogfImpl(const UObject* LogOwner, const FName& CategoryName, ELogVerbosity::Type Verbosity, const FVector& AxisLoc, const FRotator& AxisRot, const float Scale, const FColor& Color, const uint16 Thickness, const TCHAR* Fmt, ...);
 
 public:
 	// Regular text log
@@ -698,6 +702,24 @@ public:
 		HistogramDataLogfImpl(LogOwner, CategoryName, Verbosity, GraphName, DataName, Data, Color, (const TCHAR*)Fmt, Args...);
 	}
 
+	// Coordinate System
+	template <typename FmtType, typename... Types>
+	static void CoordinateSystemLogf(const UObject* LogOwner, const FLogCategoryBase& Category, ELogVerbosity::Type Verbosity, const FVector& AxisLoc, const FRotator& AxisRot, const float Scale, const FColor& Color, const uint16 Thickness, const FmtType& Fmt, Types... Args)
+	{
+		static_assert(TIsArrayOrRefOfTypeByPredicate<FmtType, TIsCharEncodingCompatibleWithTCHAR>::Value, "Formatting string must be a TCHAR array.");
+		static_assert((TIsValidVariadicFunctionArg<Types>::Value && ...), "Invalid argument(s) passed to FVisualLogger::CoordinateSystemDataLogf");
+
+		CoordinateSystemLogfImpl(LogOwner, Category.GetCategoryName(), Verbosity, AxisLoc, AxisRot, Scale, Color, Thickness, (const TCHAR*)Fmt, Args...);
+	}
+	template <typename FmtType, typename... Types>
+	static void CoordinateSystemLogf(const UObject* LogOwner, const FName& Category, ELogVerbosity::Type Verbosity, const FVector& AxisLoc, const FRotator& AxisRot, const float Scale, const FColor& Color, const uint16 Thickness, const FmtType& Fmt, Types... Args)
+	{
+		static_assert(TIsArrayOrRefOfTypeByPredicate<FmtType, TIsCharEncodingCompatibleWithTCHAR>::Value, "Formatting string must be a TCHAR array.");
+		static_assert((TIsValidVariadicFunctionArg<Types>::Value && ...), "Invalid argument(s) passed to FVisualLogger::CoordinateSystemDataLogf");
+
+		CoordinateSystemLogfImpl(LogOwner, Category, Verbosity, AxisLoc, AxisRot, Scale, Color, Thickness, (const TCHAR*)Fmt, Args...);
+	}
+
 	// Navigation data debug snapshot
 	static ENGINE_API void NavigationDataDump(const UObject* LogOwner, const FLogCategoryBase& Category, const ELogVerbosity::Type Verbosity, const FBox& Box);
 	static ENGINE_API void NavigationDataDump(const UObject* LogOwner, const FName& CategoryName, const ELogVerbosity::Type Verbosity, const FBox& Box);
diff --git a/Engine/Source/Runtime/Engine/Public/VisualLogger/VisualLoggerTypes.h b/Engine/Source/Runtime/Engine/Public/VisualLogger/VisualLoggerTypes.h
--- a/Engine/Source/Runtime/Engine/Public/VisualLogger/VisualLoggerTypes.h
+++ b/Engine/Source/Runtime/Engine/Public/VisualLogger/VisualLoggerTypes.h
@@ -62,6 +62,7 @@ enum class EVisualLoggerShapeElement : uint8
 	Arrow,
 	Circle,
 	WireCircle,
+	CoordinateSystem,
 	// note that in order to remain backward compatibility in terms of log
 	// serialization new enum values need to be added at the end
 };
@@ -255,6 +256,8 @@ struct FVisualLogEntry final
 	ENGINE_API void AddBoxes(const TArray<FBox>& Boxes, const FName& CategoryName, ELogVerbosity::Type Verbosity, const FColor& Color = FColor::White);
 	// circle
 	ENGINE_API void AddCircle(const FVector& Center, const FVector& UpAxis, const float Radius, const FName& CategoryName, ELogVerbosity::Type Verbosity, const FColor& Color, const FString& Description = TEXT(""), uint16 Thickness = 0, bool bInUseWires = false);
+	// coordinate system
+	ENGINE_API void AddCoordinateSystem(const FVector& AxisLoc, const FRotator& AxisRot, const float Scale, const FName& CategoryName, ELogVerbosity::Type Verbosity, const FColor& Color, const FString& Description = TEXT(""), uint16 Thickness = 0);
 
 	// Custom data block
 	ENGINE_API FVisualLogDataBlock& AddDataBlock(const FString& TagName, const TArray<uint8>& BlobDataArray, const FName& CategoryName, ELogVerbosity::Type Verbosity);
