Fix ODSC that would miss some shader compilation requests + add ODSC.usedefaultmaterialonrecompile

When dealing with multiple shader types, we could end in a situation where bShouldForceRecompile would be set to false if the request for the first shader was already sent. Because of that, we wouldn't try sending shader requests for the subsequent shaders
Fix bMissingShader such that is returns true as soon as a shader is missing. If the last shader was not null, it would become false again
Add the ability to prefer the default material when waiting for shaders to be compiled with the cvar ODSC.usedefaultmaterialonrecompile: this helps when changing shader parameters in a BEGIN_SHADER_PARAMETER_STRUCT/END_SHADER_PARAMETER_STRUCT and not having to recook because the parameter layout changed. This won't work when changing the layout of the default material

diff --git a/Engine/Source/Runtime/Engine/Private/Materials/MaterialShared.cpp b/Engine/Source/Runtime/Engine/Private/Materials/MaterialShared.cpp
--- a/Engine/Source/Runtime/Engine/Private/Materials/MaterialShared.cpp
+++ b/Engine/Source/Runtime/Engine/Private/Materials/MaterialShared.cpp
@@ -3588,9 +3588,11 @@ bool FMaterial::TryGetShaders(const FMaterialShaderTypes& InTypes, const FVertex
 #if WITH_ODSC
 	const bool bIsODSCActive = FODSCManager::IsODSCActive();
 	bool bShouldForceRecompile = bIsODSCActive && FODSCManager::ShouldForceRecompile(ShaderMap, this);
+	const bool bUseDefaultMaterialOnRecompile = FODSCManager::UseDefaultMaterialOnRecompile();
 #else
 	constexpr bool bIsODSCActive = false;
 	constexpr bool bShouldForceRecompile = false;
+	constexpr bool bUseDefaultMaterialOnRecompile = false;
 #endif
 
 	OutShaders.ShaderMap = ShaderMap;
@@ -3649,6 +3651,7 @@ bool FMaterial::TryGetShaders(const FMaterialShaderTypes& InTypes, const FVertex
 			}
 		}
 
+		bool bRequestNewCompilation = (Pipeline == nullptr) || bShouldForceRecompile;
 #if WITH_ODSC
 		TArray<FShaderId> RequestShaderIds;
 		TArray<FString> ShaderStageNamesToCompile;
@@ -3662,22 +3665,23 @@ bool FMaterial::TryGetShaders(const FMaterialShaderTypes& InTypes, const FVertex
 			}
 
 			bODSCRequestAlreadySent = GODSCManager->CheckIfRequestAlreadySent(RequestShaderIds, this);
-			if (bODSCRequestAlreadySent)
-			{
-				bShouldForceRecompile = false;
-			}
+		}
 
+		if (bODSCRequestAlreadySent)
+		{
+			bRequestNewCompilation = false;
+		    bMissingShader |= ((Pipeline == nullptr) || bUseDefaultMaterialOnRecompile);
 		}
 #endif
 
         // we don't do 'else' here because when bShouldForceRecompile is true, we still want to use the current 
         // pipeline until we have a new one ready. The ODSC server might fail to find the right shader, and this might
         // skew results when doing some A/B comparisons		
-		if (Pipeline == nullptr || bShouldForceRecompile)
+		if (bRequestNewCompilation)
 		{
 			if (InTypes.PipelineType->ShouldOptimizeUnusedOutputs(ShaderPlatform))
 			{
-				bMissingShader = (Pipeline == nullptr);
+    		    bMissingShader |= ((Pipeline == nullptr) || bUseDefaultMaterialOnRecompile);
 
 #if WITH_EDITOR || WITH_ODSC
 				for (const FShaderType* ShaderType : InTypes.PipelineType->GetStages())
@@ -3768,8 +3772,9 @@ bool FMaterial::TryGetShaders(const FMaterialShaderTypes& InTypes, const FVertex
 		        // we don't do 'else' here because when bShouldForceRecompile is true, we still want to use the current 
 		        // shader until we have a new one ready. The ODSC server might fail to find the right shader, and this might
 		        // skew results when doing some A/B comparisons		
-#if WITH_ODSC
 
+				bool bRequestNewCompilation = (Shader == nullptr) || bShouldForceRecompile;
+#if WITH_ODSC
 				TArray<FShaderId> RequestShaderIds;
                 bool bODSCRequestAlreadySent = false;
 				if (bIsODSCActive && (Shader == nullptr || bShouldForceRecompile))
@@ -3777,17 +3782,18 @@ bool FMaterial::TryGetShaders(const FMaterialShaderTypes& InTypes, const FVertex
 					RequestShaderIds.Add(FShaderId(ShaderType, ShaderMap->GetShaderMapId().CookedShaderMapIdHash, FHashedName(), InVertexFactoryType, PermutationId, ShaderPlatform));
 
 					bODSCRequestAlreadySent = GODSCManager->CheckIfRequestAlreadySent(RequestShaderIds, this);
-					if (bODSCRequestAlreadySent)
-					{
-						bShouldForceRecompile = false;
-					}
+				}
+
+				if (bODSCRequestAlreadySent)
+				{
+    				bMissingShader |= ((Shader == nullptr) || bUseDefaultMaterialOnRecompile);
+					bRequestNewCompilation = false;
 				}
 #endif
 				
-				if (Shader == nullptr || bShouldForceRecompile)
+				if (bRequestNewCompilation)
 				{
-					bMissingShader = (Shader == nullptr);
-
+    				bMissingShader |= ((Shader == nullptr) || bUseDefaultMaterialOnRecompile);
 #if WITH_EDITOR || WITH_ODSC
 					if (!ShouldCacheShaderType(ShaderType, InVertexFactoryType, PermutationId))
 					{
diff --git a/Engine/Source/Runtime/Engine/Private/ODSC/ODSCManager.cpp b/Engine/Source/Runtime/Engine/Private/ODSC/ODSCManager.cpp
--- a/Engine/Source/Runtime/Engine/Private/ODSC/ODSCManager.cpp
+++ b/Engine/Source/Runtime/Engine/Private/ODSC/ODSCManager.cpp
@@ -32,6 +32,18 @@ static TAutoConsoleVariable<int32> CVarODSCRecompileMode(
 	TEXT("1 (default): Compile only the permutations that are requested by the renderer. Faster iteration but more prone to hitching because of MDC recaching\n")
 	);
 
+static int32 GODSCUseDefaultMaterialOnRecompile = 0;
+static FAutoConsoleVariableRef CVarODSCUseDefaultMaterialOnRecompile(
+	TEXT("ODSC.usedefaultmaterialonrecompile"),
+	GODSCUseDefaultMaterialOnRecompile,
+	TEXT("Indicates if the default material should be used while waiting for a shader to be compiled by ODSC\n")
+	TEXT("0 (default) - use the default material only if the permutation is misssing\n")
+	TEXT("1 - use the default material even if a permutation exists\n")
+	TEXT("Setting to 1 can be useful when changing uniform buffer layout on some shaders (SHADER_PARAMETER_STRUCT for example) and avoid recooking.\n")
+	TEXT("This won't work if the default material's layout is changed\n"),
+	ECVF_RenderThreadSafe
+);
+
 FODSCManager* GODSCManager = nullptr;
 
 class FODSCManagerAccess
@@ -428,6 +440,11 @@ void FODSCManager::ReportODSCError(const FString& InErrorMessage)
 	}
 }
 
+bool FODSCManager::UseDefaultMaterialOnRecompile()
+{
+	return GODSCUseDefaultMaterialOnRecompile > 0;
+}
+
 void FODSCManager::RetrieveErrorMessage(FString& OutErrorMessage)
 {
 	FScopeLock Lock(&ErrorMessageCS);
diff --git a/Engine/Source/Runtime/Engine/Public/ODSC/ODSCManager.h b/Engine/Source/Runtime/Engine/Public/ODSC/ODSCManager.h
--- a/Engine/Source/Runtime/Engine/Public/ODSC/ODSCManager.h
+++ b/Engine/Source/Runtime/Engine/Public/ODSC/ODSCManager.h
@@ -109,6 +109,7 @@ public:
 	void TryLoadGlobalShaders(EShaderPlatform ShaderPlatform);
 
 	static void ReportODSCError(const FString& InErrorMessage);
+	static bool UseDefaultMaterialOnRecompile();
 
 	bool CheckIfRequestAlreadySent(const TArray<FShaderId>& RequestShaderIds, const FMaterial* Material) const;
 
