Prevents UMaterialExpressionPinBase from being deleted by the user. Before the user was able to accidentally delete a pin base in a material expression composite, permanently breaking the expression.

Also, turned off logging "editor only data" renamed warning when a material is duplicated, as the renaming will always occur and is normal behavior.

diff --git a/Engine/Source/Runtime/Engine/Private/Materials/MaterialExpressions.cpp b/Engine/Source/Runtime/Engine/Private/Materials/MaterialExpressions.cpp
--- a/Engine/Source/Runtime/Engine/Private/Materials/MaterialExpressions.cpp
+++ b/Engine/Source/Runtime/Engine/Private/Materials/MaterialExpressions.cpp
@@ -11461,7 +11461,6 @@ UMaterialExpressionComposite::UMaterialExpressionComposite(const FObjectInitiali
 
 #if WITH_EDITORONLY_DATA
 	bShowOutputNameOnPin = true;
-
 	Outputs.Reset();
 #endif
 }
@@ -11881,6 +11880,12 @@ bool UMaterialExpressionPinBase::Modify(bool bAlwaysMarkDirty)
 	return Super::Modify(bAlwaysMarkDirty);
 }
 
+bool UMaterialExpressionPinBase::CanUserDeleteExpression() const
+{
+	return false;
+}
+
+
 int32 UMaterialExpressionFresnel::Compile(class FMaterialCompiler* Compiler, int32 OutputIndex)
 {
 	// pow(1 - max(0,Normal dot Camera),Exponent) * (1 - BaseReflectFraction) + BaseReflectFraction
diff --git a/Engine/Source/Runtime/Engine/Private/Materials/MaterialInterface.cpp b/Engine/Source/Runtime/Engine/Private/Materials/MaterialInterface.cpp
--- a/Engine/Source/Runtime/Engine/Private/Materials/MaterialInterface.cpp
+++ b/Engine/Source/Runtime/Engine/Private/Materials/MaterialInterface.cpp
@@ -104,7 +104,7 @@ static EMaterialGetParameterValueFlags MakeParameterValueFlags(bool bOveriddenOn
 
 // Deletes any invalid EditorOnlyData object with specified name that isn't the one assigned to specified interface.
 // Returns whether the MI's EditorOnlyData was renamed.
-static bool SafeRenameEditorOnlyData(UMaterialInterface* MaterialInterface, const FString& EditorOnlyDataName)
+static bool SafeRenameEditorOnlyData(UMaterialInterface* MaterialInterface, const FString& EditorOnlyDataName, bool bLogMessageIfRenamed)
 {
 	UObject* EditorOnlyCurrent = MaterialInterface->GetEditorOnlyData();
 	UObject* EditorOnlyWithIncorrectName = StaticFindObject(/*Class=*/ nullptr, MaterialInterface, *EditorOnlyDataName, true);
@@ -122,7 +122,10 @@ static bool SafeRenameEditorOnlyData(UMaterialInterface* MaterialInterface, cons
 
 	if (EditorOnlyCurrent->GetName() != EditorOnlyDataName)
 	{
-		UE_LOG(LogMaterial, Display, TEXT("MaterialInterface %s has a incorrectly name EditorOnlyData '%s'. This may cause issues when running the editor on cooked data. Trying to rename it to the correct name '%s'."), *MaterialInterface->GetName(), *EditorOnlyCurrent->GetName(), *EditorOnlyDataName);
+		if (bLogMessageIfRenamed)
+		{
+			UE_LOG(LogMaterial, Display, TEXT("MaterialInterface %s has a incorrectly name EditorOnlyData '%s'. This may cause issues when running the editor on cooked data. Trying to rename it to the correct name '%s'."), *MaterialInterface->GetName(), *EditorOnlyCurrent->GetName(), *EditorOnlyDataName);
+		}
 
 		return EditorOnlyCurrent->Rename(*EditorOnlyDataName, nullptr, REN_NonTransactional | REN_DontCreateRedirectors);
 	}
@@ -746,7 +749,7 @@ void UMaterialInterface::PostDuplicate(bool bDuplicateForPIE)
 	// If this MI references its editor only data, make sure it has the correct name, or the editor
 	// won't run properly when running on cooked data.
 	FString EditorOnlyDataName = UE::MaterialInterface::GetEditorOnlyDataName(*GetName());
-	SafeRenameEditorOnlyData(this, EditorOnlyDataName);
+	SafeRenameEditorOnlyData(this, EditorOnlyDataName, false);
 #endif // WITH_EDITORONLY_DATA
 }
 
@@ -1903,7 +1906,7 @@ void UMaterialInterface::PreSave(FObjectPreSaveContext ObjectSaveContext)
 	FString EditorOnlyDataName = UE::MaterialInterface::GetEditorOnlyDataName(*GetName());
 	if (!ObjectSaveContext.IsProceduralSave())
 	{
-		SafeRenameEditorOnlyData(this, EditorOnlyDataName);
+		SafeRenameEditorOnlyData(this, EditorOnlyDataName, true);
 	}
 #endif // WITH_EDITORONLY_DATA
 
@@ -2033,7 +2036,7 @@ bool UMaterialInterface::Rename(const TCHAR* NewName, UObject* NewOuter, ERename
 	if (bRenamed && NewName && EditorOnlyData)
 	{
 		FString EditorOnlyDataName = UE::MaterialInterface::GetEditorOnlyDataName(NewName);
-		bRenamed = SafeRenameEditorOnlyData(this, EditorOnlyDataName);
+		bRenamed = SafeRenameEditorOnlyData(this, EditorOnlyDataName, true);
 	}
 #endif
 	return bRenamed;
diff --git a/Engine/Source/Runtime/Engine/Public/Materials/MaterialExpressionPinBase.h b/Engine/Source/Runtime/Engine/Public/Materials/MaterialExpressionPinBase.h
--- a/Engine/Source/Runtime/Engine/Public/Materials/MaterialExpressionPinBase.h
+++ b/Engine/Source/Runtime/Engine/Public/Materials/MaterialExpressionPinBase.h
@@ -92,6 +92,8 @@ class UMaterialExpressionPinBase : public UMaterialExpression
 	virtual bool Modify(bool bAlwaysMarkDirty = true) override;
 	//~ End UObject Interface.
 
+	virtual bool CanUserDeleteExpression() const override;
+
 private:
 	/** Used to remove reroute expressions from the material after we remove them. */
 	TArray<UMaterialExpressionReroute*> PreEditRereouteExpresions;
