[Shaders] Fix filepath trimming when UBA I/O base path is too long.

#rnx

diff --git a/Engine/Source/Runtime/Engine/Private/ShaderCompiler/ShaderCompilerDistributed.cpp b/Engine/Source/Runtime/Engine/Private/ShaderCompiler/ShaderCompilerDistributed.cpp
--- a/Engine/Source/Runtime/Engine/Private/ShaderCompiler/ShaderCompilerDistributed.cpp
+++ b/Engine/Source/Runtime/Engine/Private/ShaderCompiler/ShaderCompilerDistributed.cpp
@@ -168,9 +168,10 @@ void FShaderCompileDistributedThreadRunnable_Interface::DispatchShaderCompileJob
 	BaseFilePath += BaseFilePathBuilder.ToView();
 
 	// Ensure input and output files don't exceed platform path limit
-	if (BaseFilePath.Len() + 4 > FPlatformMisc::GetMaxPathLength())
+	const int32 MaxBaseFilePathLen = FPlatformMisc::GetMaxPathLength() - 4;
+	if (BaseFilePath.Len() > MaxBaseFilePathLen)
 	{
-		BaseFilePath.LeftChopInline(FPlatformMisc::GetMaxPathLength() - 4);
+		BaseFilePath.LeftChopInline(BaseFilePath.Len() - MaxBaseFilePathLen);
 	}
 
 	FString InputFilePath = BaseFilePath + TEXT(".in");
