NetDriver
* In generic replication, fix IsActorRelevantToConnection called only every 1sec and causing issues if RelevantTimeout is configured to be lower then 1 sec
* Misc comments and improved some logs

#rnx

diff --git a/Engine/Source/Runtime/Engine/Classes/Engine/NetDriver.h b/Engine/Source/Runtime/Engine/Classes/Engine/NetDriver.h
--- a/Engine/Source/Runtime/Engine/Classes/Engine/NetDriver.h
+++ b/Engine/Source/Runtime/Engine/Classes/Engine/NetDriver.h
@@ -905,7 +905,11 @@ public:
 	UPROPERTY(Config)
 	float SpawnPrioritySeconds;
 
-	/** @todo document */
+	/** 
+	 * Amount of time (in seconds) before we close the channel of a replicated actor when it goes out of relevancy. Ex: If the client moves farther then the actor's netculldistance, the actor will still be replicated to that client for the duration of the RelevantTimeout.
+	 * Default value is 5 seconds
+	 * Note: ReplicationGraph and Iris do not use this value directly and instead have their unique config.
+	 */
 	UPROPERTY(Config)
 	float RelevantTimeout;
 
@@ -1146,7 +1150,9 @@ public:
 	}
 
 private:
-	double						ElapsedTime;
+
+	/** Accumulated delta time since the creation of the NetDriver */
+	double ElapsedTime;
 
 	/** Whether or not the NetDriver is ticking */
 	bool bInTick;
diff --git a/Engine/Source/Runtime/Engine/Private/DataChannel.cpp b/Engine/Source/Runtime/Engine/Private/DataChannel.cpp
--- a/Engine/Source/Runtime/Engine/Private/DataChannel.cpp
+++ b/Engine/Source/Runtime/Engine/Private/DataChannel.cpp
@@ -2581,7 +2581,7 @@ bool UActorChannel::CleanUp(const bool bForDestroy, EChannelCloseReason CloseRea
 
 	const bool bIsServer = Connection->Driver->IsServer();
 
-	UE_LOG( LogNetTraffic, Log, TEXT( "UActorChannel::CleanUp: %s" ), *Describe() );
+	UE_LOG( LogNetTraffic, Log, TEXT( "UActorChannel::CleanUp (ChIndex: %i): %s | NetGuid: %s | Reason: %s" ), ChIndex, *GetFullNameSafe(Actor), *ActorNetGUID.ToString(), LexToString(CloseReason));
 
 	if (!bIsServer && QueuedBunches.Num() > 0 && ChIndex >= 0 && !bForDestroy)
 	{
@@ -2676,7 +2676,7 @@ bool UActorChannel::CleanUp(const bool bForDestroy, EChannelCloseReason CloseRea
 			}
 			else if (!Actor->bNetTemporary && Actor->GetWorld() != NULL && !IsEngineExitRequested() && Connection->Driver->ShouldClientDestroyActor(Actor))
 			{
-				UE_LOG(LogNetDormancy, Verbose, TEXT("UActorChannel::CleanUp: Destroying Actor. %s"), *Describe() );
+				UE_LOG(LogNetDormancy, Verbose, TEXT("UActorChannel::CleanUp: Destroying Actor (ChIndex: %i): %s"), ChIndex, *GetFullNameSafe(Actor));
 
 				DestroyActorAndComponents();
 			}
@@ -2788,7 +2788,7 @@ void UActorChannel::SetChannelActor(AActor* InActor, ESetChannelActorFlags Flags
 	// Further, expect calling code to do these checks.
 	bActorIsPendingKill = false;
 
-	UE_LOG(LogNetTraffic, VeryVerbose, TEXT("SetChannelActor: ChIndex: %i, Actor: %s, NetGUID: %s"), ChIndex, Actor ? *Actor->GetFullName() : TEXT("NULL"), *ActorNetGUID.ToString() );
+	UE_LOG(LogNetTraffic, VeryVerbose, TEXT("SetChannelActor (ChIndex: %i): Actor: %s | NetGUID: %s"), ChIndex, *GetFullNameSafe(Actor), *ActorNetGUID.ToString() );
 
 	// Internal ack connections never add to PendingOutRec
 	if (Connection->PendingOutRec.IsValidIndex(ChIndex) && Connection->PendingOutRec[ChIndex] > 0)
diff --git a/Engine/Source/Runtime/Engine/Private/NetDriver.cpp b/Engine/Source/Runtime/Engine/Private/NetDriver.cpp
--- a/Engine/Source/Runtime/Engine/Private/NetDriver.cpp
+++ b/Engine/Source/Runtime/Engine/Private/NetDriver.cpp
@@ -5545,11 +5545,12 @@ int32 UNetDriver::ServerReplicateActors_ProcessPrioritizedActorsRange( UNetConne
 
 			const bool bLevelInitializedForActor = IsLevelInitializedForActor( Actor, Connection );
 
-			// only check visibility on already visible actors every 1.0 + 0.5R seconds
+			// only check visibility on already visible actors every 1.0 + 0.5R seconds or every RelevantTimeout if it's lower then 1sec.
 			// bTearOff actors should never be checked
 			if ( bLevelInitializedForActor )
 			{
-				if ( !Actor->GetTearOff() && ( !Channel || ElapsedTime - Channel->RelevantTime > 1.0 ) )
+				const float MinVisibilityTimeout = FMath::Min(RelevantTimeout, 1.0f);
+				if ( !Actor->GetTearOff() && ( !Channel || ElapsedTime - Channel->RelevantTime >= MinVisibilityTimeout) )
 				{
 					if ( IsActorRelevantToConnection( Actor, ConnectionViewers ) )
 					{
