Fix fetching level info from Asset Registry could potentially result in a huge load.

UActorDescContainer::Initialize during Data Validation in a licensee project: 109s -> 1ms.

diff --git a/Engine/Source/Runtime/Engine/Private/Level.cpp b/Engine/Source/Runtime/Engine/Private/Level.cpp
--- a/Engine/Source/Runtime/Engine/Private/Level.cpp
+++ b/Engine/Source/Runtime/Engine/Private/Level.cpp
@@ -2807,7 +2807,7 @@ namespace LevelAssetRegistryHelper
 		IAssetRegistry& AssetRegistry = FModuleManager::LoadModuleChecked<FAssetRegistryModule>(TEXT("AssetRegistry")).Get();
 
 		// Always ask for scan in case no initial scan was done (Commandlets, IsRunningGame) or if AssetRegistry isn't done scanning
-		AssetRegistry.ScanFilesSynchronous({ LevelPackage.ToString()});
+		AssetRegistry.ScanSynchronous({ }, { LevelPackage.ToString()});
 		
 		TArray<FAssetData> LevelPackageAssets;
 		AssetRegistry.GetAssetsByPackageName(LevelPackage, LevelPackageAssets, true);
