Add back in the default for channel mask parameter.

#rnx

diff --git a/Engine/Source/Runtime/Engine/Private/Materials/MaterialExpressions.cpp b/Engine/Source/Runtime/Engine/Private/Materials/MaterialExpressions.cpp
--- a/Engine/Source/Runtime/Engine/Private/Materials/MaterialExpressions.cpp
+++ b/Engine/Source/Runtime/Engine/Private/Materials/MaterialExpressions.cpp
@@ -8411,6 +8411,9 @@ UMaterialExpressionChannelMaskParameter::UMaterialExpressionChannelMaskParameter
 	Outputs.Add(FExpressionOutput(TEXT(""), 1, 1, 0, 0, 0));
 	bShowMaskColorsOnPin = false;
 #endif
+
+	// Default mask to red channel
+	DefaultValue = FLinearColor(1.0f, 0.0f, 0.0f, 0.0f);
 }
 
 #if WITH_EDITOR
