IncrementalCook: StaticMesh: Add the DDC keys for the cookeddata for DistanceFields and MeshCards to all staticmesh keys.
#rnx
#incrementalcook

diff --git a/Engine/Source/Runtime/Engine/Classes/Engine/StaticMesh.h b/Engine/Source/Runtime/Engine/Classes/Engine/StaticMesh.h
--- a/Engine/Source/Runtime/Engine/Classes/Engine/StaticMesh.h
+++ b/Engine/Source/Runtime/Engine/Classes/Engine/StaticMesh.h
@@ -1754,6 +1754,7 @@ public:
 	ENGINE_API virtual void PostLoad() override;
 #if WITH_EDITORONLY_DATA
 	ENGINE_API static void DeclareConstructClasses(TArray<FTopLevelAssetPath>& OutConstructClasses, const UClass* SpecificSubclass);
+	ENGINE_API static void AppendToClassSchema(FAppendToClassSchemaContext& Context);
 #endif
 	ENGINE_API virtual bool IsPostLoadThreadSafe() const override;
 	ENGINE_API virtual void BeginDestroy() override;
diff --git a/Engine/Source/Runtime/Engine/Private/DistanceFieldAtlas.cpp b/Engine/Source/Runtime/Engine/Private/DistanceFieldAtlas.cpp
--- a/Engine/Source/Runtime/Engine/Private/DistanceFieldAtlas.cpp
+++ b/Engine/Source/Runtime/Engine/Private/DistanceFieldAtlas.cpp
@@ -130,6 +130,20 @@ FDistanceFieldAsyncQueue* GDistanceFieldAsyncQueue = NULL;
 // DDC key for distance field data, must be changed when modifying the generation code or data format
 #define DISTANCEFIELD_DERIVEDDATA_VER TEXT("9137A53B-09F6-4C2F-86E1-29FE6834D74B")
 
+void AppendDistanceFieldConstantsToClassSchema(FAppendToClassSchemaContext& Context)
+{
+	IConsoleManager& CM = IConsoleManager::Get();
+	const int32 PerMeshMax
+		= CM.FindTConsoleVariableDataInt(TEXT("r.DistanceFields.MaxPerMeshResolution"))->GetValueOnAnyThread();
+	const float VoxelDensity
+		= CM.FindTConsoleVariableDataFloat(TEXT("r.DistanceFields.DefaultVoxelDensity"))->GetValueOnAnyThread();
+	FStringView VersionString(DISTANCEFIELD_DERIVEDDATA_VER);
+
+	Context.Update(&PerMeshMax, sizeof(PerMeshMax));
+	Context.Update(&VoxelDensity, sizeof(VoxelDensity));
+	Context.Update(VersionString.GetData(), VersionString.Len() * sizeof(VersionString[0]));
+}
+
 FString BuildDistanceFieldDerivedDataKey(const FString& InMeshKey)
 {
 	static const auto CVar = IConsoleManager::Get().FindTConsoleVariableDataInt(TEXT("r.DistanceFields.MaxPerMeshResolution"));
diff --git a/Engine/Source/Runtime/Engine/Private/MeshCardRepresentation.cpp b/Engine/Source/Runtime/Engine/Private/MeshCardRepresentation.cpp
--- a/Engine/Source/Runtime/Engine/Private/MeshCardRepresentation.cpp
+++ b/Engine/Source/Runtime/Engine/Private/MeshCardRepresentation.cpp
@@ -98,11 +98,24 @@ int32 MeshCardRepresentation::GetDebugSurfelDirection()
 
 FCardRepresentationAsyncQueue* GCardRepresentationAsyncQueue = NULL;
 
-#if WITH_EDITOR
+#if WITH_EDITORONLY_DATA
 
 // DDC key for card representation data, must be changed when modifying the generation code or data format
 #define CARDREPRESENTATION_DERIVEDDATA_VER TEXT("59BE661D-D92C-4158-B14D-3C384D5F5FEF")
 
+void AppendMeshCardConstantsToClassSchema(FAppendToClassSchemaContext& Context)
+{
+	const float MinDensity = MeshCardRepresentation::GetMinDensity();
+	const float NormalTreshold = MeshCardRepresentation::GetNormalTreshold();
+	uint8 DebugMode = MeshCardRepresentation::IsDebugMode() ? 1 : 0;
+	FStringView VersionString(CARDREPRESENTATION_DERIVEDDATA_VER);
+
+	Context.Update(&MinDensity, sizeof(MinDensity));
+	Context.Update(&NormalTreshold, sizeof(NormalTreshold));
+	Context.Update(&DebugMode, sizeof(DebugMode));
+	Context.Update(VersionString.GetData(), VersionString.Len() * sizeof(VersionString[0]));
+}
+
 FString BuildCardRepresentationDerivedDataKey(const FString& InMeshKey, int32 MaxLumenMeshCards)
 {
 	const float MinDensity = MeshCardRepresentation::GetMinDensity();
@@ -117,10 +130,6 @@ FString BuildCardRepresentationDerivedDataKey(const FString& InMeshKey, int32 Ma
 		TEXT(""));
 }
 
-#endif
-
-#if WITH_EDITORONLY_DATA
-
 extern void BuildSignedDistanceFieldBuildSectionData(UStaticMesh* Mesh, uint32 LODIndex, TArray<FSignedDistanceFieldBuildSectionData>& OutData);
 
 void BeginCacheMeshCardRepresentation(const ITargetPlatform* TargetPlatform, UStaticMesh* StaticMeshAsset, FStaticMeshRenderData& RenderData, const FString& DistanceFieldKey, FSourceMeshDataForDerivedDataTask* OptionalSourceMeshData)
diff --git a/Engine/Source/Runtime/Engine/Private/StaticMesh.cpp b/Engine/Source/Runtime/Engine/Private/StaticMesh.cpp
--- a/Engine/Source/Runtime/Engine/Private/StaticMesh.cpp
+++ b/Engine/Source/Runtime/Engine/Private/StaticMesh.cpp
@@ -7066,6 +7066,12 @@ void UStaticMesh::DeclareConstructClasses(TArray<FTopLevelAssetPath>& OutConstru
 	OutConstructClasses.Add(FTopLevelAssetPath(UStaticMeshDescriptionBulkData::StaticClass()));
 	OutConstructClasses.Add(FTopLevelAssetPath(TEXT("/Script/NavigationSystem.NavCollision")));
 }
+void UStaticMesh::AppendToClassSchema(FAppendToClassSchemaContext& Context)
+{
+	Super::AppendToClassSchema(Context);
+	AppendMeshCardConstantsToClassSchema(Context);
+	AppendDistanceFieldConstantsToClassSchema(Context);
+}
 #endif
 
 void UStaticMesh::BeginPostLoadInternal(FStaticMeshPostLoadContext& Context)
diff --git a/Engine/Source/Runtime/Engine/Public/DistanceFieldAtlas.h b/Engine/Source/Runtime/Engine/Public/DistanceFieldAtlas.h
--- a/Engine/Source/Runtime/Engine/Public/DistanceFieldAtlas.h
+++ b/Engine/Source/Runtime/Engine/Public/DistanceFieldAtlas.h
@@ -22,6 +22,7 @@
 #include "MeshUtilities.h"
 #endif
 
+struct FAppendToClassSchemaContext;
 struct FAssetCompileData;
 class FAsyncCompilationNotification;
 class FDistanceFieldVolumeData;
@@ -440,7 +441,16 @@ private:
 /** Global build queue. */
 extern ENGINE_API FDistanceFieldAsyncQueue* GDistanceFieldAsyncQueue;
 
+#if WITH_EDITOR
+/**
+ * IncrementalCook: Any UClass that serializes cooked data with data loaded from DDC with key created from
+ * BuildDistanceFieldDerivedDataKey must add the version numbers for that calculation to their ClassSchema,
+ * to invalidate their cooked data when those constants change. This function handles that appending.
+ */
+extern ENGINE_API void AppendDistanceFieldConstantsToClassSchema(FAppendToClassSchemaContext& Context);
+
 extern ENGINE_API FString BuildDistanceFieldDerivedDataKey(const FString& InMeshKey);
+#endif
 
 #if UE_ENABLE_INCLUDE_ORDER_DEPRECATED_IN_5_3
 #include "CoreMinimal.h"
diff --git a/Engine/Source/Runtime/Engine/Public/MeshCardBuild.h b/Engine/Source/Runtime/Engine/Public/MeshCardBuild.h
--- a/Engine/Source/Runtime/Engine/Public/MeshCardBuild.h
+++ b/Engine/Source/Runtime/Engine/Public/MeshCardBuild.h
@@ -12,6 +12,7 @@
 
 class FAsyncCompilationNotification;
 class UStaticMesh;
+struct FAppendToClassSchemaContext;
 struct FSignedDistanceFieldBuildSectionData;
 
 class FLumenCardBuildData
@@ -318,9 +319,17 @@ private:
 /** Global build queue. */
 extern ENGINE_API FCardRepresentationAsyncQueue* GCardRepresentationAsyncQueue;
 
-extern ENGINE_API FString BuildCardRepresentationDerivedDataKey(const FString& InMeshKey);
+#if WITH_EDITORONLY_DATA
+/**
+ * IncrementalCook: Any UClass that uses data created by BeginCacheMeshCardRepresentation must add the version
+ * numbers for that calculation to their ClassSchema, to invalidate their cooked data when those constants change.
+ * This function handles that appending.
+ */
+extern ENGINE_API void AppendMeshCardConstantsToClassSchema(FAppendToClassSchemaContext& Context);
 
+extern ENGINE_API FString BuildCardRepresentationDerivedDataKey(const FString& InMeshKey);
 extern ENGINE_API void BeginCacheMeshCardRepresentation(const ITargetPlatform* TargetPlatform, UStaticMesh* StaticMeshAsset, class FStaticMeshRenderData& RenderData, const FString& DistanceFieldKey, FSourceMeshDataForDerivedDataTask* OptionalSourceMeshData);
+#endif
 
 #if UE_ENABLE_INCLUDE_ORDER_DEPRECATED_IN_5_4
 #include "AssetCompilingManager.h"
