Skeletal Mesh: Fix a crash in FSkeletalMeshLODModel::GetMeshDescription when there are broken morph target deltas

#rnx

diff --git a/Engine/Source/Runtime/Engine/Private/SkeletalMeshLODModel.cpp b/Engine/Source/Runtime/Engine/Private/SkeletalMeshLODModel.cpp
--- a/Engine/Source/Runtime/Engine/Private/SkeletalMeshLODModel.cpp
+++ b/Engine/Source/Runtime/Engine/Private/SkeletalMeshLODModel.cpp
@@ -1338,6 +1338,7 @@ void FSkeletalMeshLODModel::CopyStructure(FSkeletalMeshLODModel* Destination, co
 	Destination->BulkDataReadMutex = DestinationBulkDataReadMutex;
 }
 
+
 void FSkeletalMeshLODModel::GetMeshDescription(const USkeletalMesh *InSkeletalMesh, const int32 InLODIndex, FMeshDescription& OutMeshDescription) const
 {
 	using UE::AnimationCore::FBoneWeights;
@@ -1558,19 +1559,22 @@ void FSkeletalMeshLODModel::GetMeshDescription(const USkeletalMesh *InSkeletalMe
 
 		for (const FMorphTargetDelta& Delta: MorphSource.Value->Vertices)
 		{
-			const int32 TargetVertexIndex = SourceToTargetVertexMap[Delta.SourceIdx];
-			const FVertexID VertexID = VertexIDs[TargetVertexIndex];
-
-			PositionDelta.Set(VertexID, Delta.PositionDelta);
-
-			if (bMorphTargetIncludeNormals)
+			if (SourceToTargetVertexMap.IsValidIndex(Delta.SourceIdx))
 			{
-				VertexInstanceIDs.Reset();
-				SourceVertexToVertexInstanceMap.MultiFind(Delta.SourceIdx, VertexInstanceIDs);
+				const int32 TargetVertexIndex = SourceToTargetVertexMap[Delta.SourceIdx];
+				const FVertexID VertexID = VertexIDs[TargetVertexIndex];
 
-				for (FVertexInstanceID VertexInstanceID: VertexInstanceIDs)
+				PositionDelta.Set(VertexID, Delta.PositionDelta);
+
+				if (bMorphTargetIncludeNormals)
 				{
-					NormalDelta.Set(VertexInstanceID, Delta.TangentZDelta); 
+					VertexInstanceIDs.Reset();
+					SourceVertexToVertexInstanceMap.MultiFind(Delta.SourceIdx, VertexInstanceIDs);
+
+					for (FVertexInstanceID VertexInstanceID: VertexInstanceIDs)
+					{
+						NormalDelta.Set(VertexInstanceID, Delta.TangentZDelta); 
+					}
 				}
 			}
 		}
