Fix DO_GUARD_SLOW assert when deleting a Bluprint class.

diff --git a/Engine/Source/Runtime/Engine/Private/WorldPartition/WorldPartitionClassDescRegistry.cpp b/Engine/Source/Runtime/Engine/Private/WorldPartition/WorldPartitionClassDescRegistry.cpp
--- a/Engine/Source/Runtime/Engine/Private/WorldPartition/WorldPartitionClassDescRegistry.cpp
+++ b/Engine/Source/Runtime/Engine/Private/WorldPartition/WorldPartitionClassDescRegistry.cpp
@@ -586,6 +586,7 @@ void FWorldPartitionClassDescRegistry::OnAssetRemoved(const FAssetData& InAssetD
 			{
 				ClassByPath.Remove(ClassPath);
 			}
+			ParentClassMap.Remove(ClassPath);
 			ValidateInternalState();
 		}
 	}
