Ensure that StaticMesh is build when fetching or building FDistanceFieldVolumeData.  We will now uses FDistanceFieldAsyncQueue to try fetching data from DDC then if we failed to retrieve from DDC we will kick the task to build it

diff --git a/Engine/Source/Runtime/Engine/Private/DistanceFieldAtlas.cpp b/Engine/Source/Runtime/Engine/Private/DistanceFieldAtlas.cpp
--- a/Engine/Source/Runtime/Engine/Private/DistanceFieldAtlas.cpp
+++ b/Engine/Source/Runtime/Engine/Private/DistanceFieldAtlas.cpp
@@ -199,39 +199,14 @@ void FDistanceFieldVolumeData::CacheDerivedData(const FString& InStaticMeshDeriv
 	TRACE_CPUPROFILER_EVENT_SCOPE(FDistanceFieldVolumeData::CacheDerivedData);
 	COOK_STAT(auto Timer = DistanceFieldCookStats::UsageStats.TimeSyncWork());
 
-	TArray<FSignedDistanceFieldBuildSectionData> BuildSectionData;
-
-	const uint32 LODIndex = 0;
-	BuildSignedDistanceFieldBuildSectionData(Mesh, LODIndex, BuildSectionData);
-
-	FString DistanceFieldKey = BuildDistanceFieldDerivedDataKey(InStaticMeshDerivedDataKey);
-
-	for (int32 SectionIndex = 0; SectionIndex < BuildSectionData.Num(); SectionIndex++)
-	{
-		DistanceFieldKey += FString::Printf(TEXT("_M%u_%u_%u"), 
-			(uint32)BuildSectionData[SectionIndex].BlendMode,
-			BuildSectionData[SectionIndex].bTwoSided ? 1 : 0,
-			BuildSectionData[SectionIndex].bAffectDistanceFieldLighting ? 1 : 0);
-	}
-
-	TArray<uint8> DerivedData;
-
-	if (GetDerivedDataCacheRef().GetSynchronous(*DistanceFieldKey, DerivedData, Mesh->GetPathName()))
-	{
-		COOK_STAT(Timer.AddHit(DerivedData.Num()));
-		FMemoryReader Ar(DerivedData, /*bIsPersistent=*/ true);
-		Serialize(Ar, Mesh);
-
-		BeginCacheMeshCardRepresentation(TargetPlatform, Mesh, RenderData, DistanceFieldKey, /*OptionalSourceMeshData*/ nullptr);
-	}
-	else if (GDistanceFieldAsyncQueue)
+	if (GDistanceFieldAsyncQueue)
 	{
 		check(Mesh && GenerateSource);
 
 		// We don't actually build the resource until later, so only track the cycles used here.
 		COOK_STAT(Timer.TrackCyclesOnly());
 		FAsyncDistanceFieldTask* NewTask = new FAsyncDistanceFieldTask;
-		NewTask->DDCKey = DistanceFieldKey;
+		NewTask->DDCKey = InStaticMeshDerivedDataKey;
 		NewTask->TargetPlatform = TargetPlatform;
 		NewTask->StaticMesh = Mesh;
 		NewTask->GenerateSource = GenerateSource;
@@ -240,7 +215,6 @@ void FDistanceFieldVolumeData::CacheDerivedData(const FString& InStaticMeshDeriv
 		NewTask->GeneratedVolumeData = new FDistanceFieldVolumeData();
 		NewTask->GeneratedVolumeData->AssetName = Mesh->GetFName();
 		NewTask->GeneratedVolumeData->bAsyncBuilding = true;
-		NewTask->SectionData = MoveTemp(BuildSectionData);
 
 		// Nanite overrides source static mesh with a coarse representation. Need to load original data before we build the mesh SDF.
 		if (Mesh->IsNaniteEnabled())
@@ -518,6 +492,44 @@ void FDistanceFieldAsyncQueue::ProcessPendingTasks()
 		if ((Task->GenerateSource == nullptr || !Task->GenerateSource->IsCompiling()) && (Task->StaticMesh == nullptr || !Task->StaticMesh ->IsCompiling()))
 		{
 			StartBackgroundTask(Task);
+#if WITH_EDITORONLY_DATA			
+			COOK_STAT(auto Timer = DistanceFieldCookStats::UsageStats.TimeSyncWork());
+			TArray<FSignedDistanceFieldBuildSectionData> BuildSectionData;
+			const uint32 LODIndex = 0;
+			BuildSignedDistanceFieldBuildSectionData(Task->StaticMesh, LODIndex, BuildSectionData);
+
+			FString DistanceFieldKey = BuildDistanceFieldDerivedDataKey(Task->DDCKey);
+
+			for (int32 SectionIndex = 0; SectionIndex < BuildSectionData.Num(); SectionIndex++)
+			{
+				DistanceFieldKey += FString::Printf(TEXT("_M%u_%u_%u"), 
+					(uint32)BuildSectionData[SectionIndex].BlendMode,
+					BuildSectionData[SectionIndex].bTwoSided ? 1 : 0,
+					BuildSectionData[SectionIndex].bAffectDistanceFieldLighting ? 1 : 0);
+			}
+
+			TArray<uint8> DerivedData;
+			if (GetDerivedDataCacheRef().GetSynchronous(*DistanceFieldKey, DerivedData, Task->StaticMesh->GetPathName()))
+			{
+				COOK_STAT(Timer.AddHit(DerivedData.Num()));
+				FMemoryReader Ar(DerivedData, /*bIsPersistent=*/ true);
+				Task->GeneratedVolumeData->Serialize(Ar, Task->StaticMesh);
+				CompletedTasks.Add(Task);
+
+				BeginCacheMeshCardRepresentation(
+					Task->TargetPlatform, 
+					Task->StaticMesh, 
+					Task->StaticMesh->GetPlatformStaticMeshRenderData(Task->StaticMesh, Task->TargetPlatform), 
+					DistanceFieldKey, 
+					/*OptionalSourceMeshData*/ nullptr);
+			}
+			else
+			{
+				Task->DDCKey = DistanceFieldKey;
+				Task->SectionData = MoveTemp(BuildSectionData);
+				StartBackgroundTask(Task);
+			}
+#endif			
 			It.RemoveCurrent();
 		}
 	}
