[NetPhysics] Delta Serialization support for NetworkPhysicsComponent

diff --git a/Engine/Source/Runtime/Engine/Private/PhysicsEngine/NetworkPhysicsComponent.cpp b/Engine/Source/Runtime/Engine/Private/PhysicsEngine/NetworkPhysicsComponent.cpp
--- a/Engine/Source/Runtime/Engine/Private/PhysicsEngine/NetworkPhysicsComponent.cpp
+++ b/Engine/Source/Runtime/Engine/Private/PhysicsEngine/NetworkPhysicsComponent.cpp
@@ -66,10 +66,17 @@ namespace PhysicsReplicationCVars
 		static FAutoConsoleVariableRef CVarResimApplySimProxyStateAtRuntime(TEXT("np2.Resim.ApplySimProxyStateAtRuntime"), bApplySimProxyStateAtRuntime, TEXT("When true, call ApplyData on received states for simulated proxies at runtime."));
 		bool bApplySimProxyInputAtRuntime = true;
 		static FAutoConsoleVariableRef CVarResimApplySimProxyInputAtRuntime(TEXT("np2.Resim.ApplySimProxyInputAtRuntime"), bApplySimProxyInputAtRuntime, TEXT("When true, call ApplyData on received inputs for simulated proxies at runtime."));
+
+		bool bEnableStatefulDeltaSerialization = true;
+		static FAutoConsoleVariableRef CVarResimEnableStatefulDeltaSerialization(TEXT("np2.Resim.StatefulDeltaSerialization.Enable"), bEnableStatefulDeltaSerialization, TEXT("Enables stateful delta serialization for FNetworkPhysicsData derived inputs and states. During FNetworkPhysicsData::NetSerialize there will be a valid pointer to a previous FNetworkPhysicsData which can be used for delta serialization, FNetworkPhysicsData::DeltaSourceData. NOTE: Switching this during gameplay might cause disconnections."));
+		bool bUseDefaultDeltaForDeltaSourceReplication = true;
+		static FAutoConsoleVariableRef CVarResimUseDefaultForDeltaSourceReplication(TEXT("np2.Resim.StatefulDeltaSerialization.UseDefaultForDeltaSourceReplication"), bUseDefaultDeltaForDeltaSourceReplication, TEXT("When false delta sources will use standard serialization when being replicated. When true there will be a valid delta source pointer to default data which can be used for delta serialization when replicating delta sources."));
+		float TimeToSyncStatefulDeltaSource = 5.0f;
+		static FAutoConsoleVariableRef CVarResimTimeToSyncStatefulDeltaSource(TEXT("np2.Resim.StatefulDeltaSerialization.TimeToSyncStatefulDeltaSource"), TimeToSyncStatefulDeltaSource, TEXT("The time in seconds between synchronizing the stateful delta source from server to clients."));
 	}
 }
 
-bool FNetworkPhysicsRewindDataProxy::NetSerializeBase(FArchive& Ar, class UPackageMap* Map, bool& bOutSuccess, TUniqueFunction<TUniquePtr<Chaos::FBaseRewindHistory>()> CreateHistoryFunction)
+bool FNetworkPhysicsRewindDataProxy::NetSerializeBase(FArchive& Ar, class UPackageMap* Map, bool& bOutSuccess, TUniqueFunction<TUniquePtr<Chaos::FBaseRewindHistory> ()> CreateHistoryFunction , TUniqueFunction<FNetworkPhysicsData* (const int32)> GetDeltaSourceData)
 {
 	Ar << Owner;
 
@@ -98,7 +105,52 @@ bool FNetworkPhysicsRewindDataProxy::NetSerializeBase(FArchive& Ar, class UPacka
 			}
 		}
 
-		History->NetSerialize(Ar, Map);
+		History->NetSerialize(Ar, Map, [&](void* Data, const int32 DataIndex)
+			{
+				if (FNetworkPhysicsData* NetData = static_cast<FNetworkPhysicsData*>(Data))
+				{
+					if (Owner)
+					{
+						// Set the component pointer to the implementation that uses this data
+						NetData->SetImplementationComponent(Owner.Get()->ActorComponent.Get());
+
+						// Only use stateful delta source for the first entry in history, the following entries will use the previous entry as delta source
+						if (PhysicsReplicationCVars::ResimulationCVars::bEnableStatefulDeltaSerialization && DataIndex == 0 && GetDeltaSourceData)
+						{
+							// Stateful Delta Serialization
+							uint32 DeltaSourceFrame = 0;
+							if (Ar.IsLoading())
+							{
+								Ar.SerializeIntPacked(DeltaSourceFrame);
+								
+								if (FNetworkPhysicsData* DeltaSourceData = GetDeltaSourceData(static_cast<int32>(DeltaSourceFrame) - 1))
+								{
+									NetData->SetDeltaSourceData(DeltaSourceData);
+								}
+								else
+								{
+#if DEBUG_NETWORK_PHYSICS
+									UE_LOG(LogChaos, Warning, TEXT("DEBUGGING - Delta Serialization ISSUE, did not find delta source with frame: %d"), DeltaSourceFrame);
+#endif
+									bDeltaSerializationIssue = true;
+									NetData->SetDeltaSourceData(GetDeltaSourceData(/*Default*/ -2));
+								}
+							}
+							else
+							{
+								if (FNetworkPhysicsData* DeltaSourceData = GetDeltaSourceData(/*Latest*/ -1))
+								{
+									NetData->SetDeltaSourceData(DeltaSourceData);
+
+									ensure((DeltaSourceData->ServerFrame + 1) >= 0);
+									DeltaSourceFrame = static_cast<uint32>(DeltaSourceData->ServerFrame + 1); // +1 since ServerFrame has a default value of -1 and it needs to be serialized unsigned
+								}
+								Ar.SerializeIntPacked(DeltaSourceFrame);
+							}
+						}
+					}
+				}
+			});
 	}
 
 	return true;
@@ -121,33 +173,62 @@ UE_NET_IMPLEMENT_NAMED_STRUCT_LASTRESORT_NETSERIALIZER_AND_REGISTRY_DELEGATES(Ne
 UE_NET_IMPLEMENT_NAMED_STRUCT_LASTRESORT_NETSERIALIZER_AND_REGISTRY_DELEGATES(NetworkPhysicsRewindDataStateProxy);
 UE_NET_IMPLEMENT_NAMED_STRUCT_LASTRESORT_NETSERIALIZER_AND_REGISTRY_DELEGATES(NetworkPhysicsRewindDataImportantInputProxy);
 UE_NET_IMPLEMENT_NAMED_STRUCT_LASTRESORT_NETSERIALIZER_AND_REGISTRY_DELEGATES(NetworkPhysicsRewindDataImportantStateProxy);
+UE_NET_IMPLEMENT_NAMED_STRUCT_LASTRESORT_NETSERIALIZER_AND_REGISTRY_DELEGATES(NetworkPhysicsRewindDataDeltaSourceInputProxy);
+UE_NET_IMPLEMENT_NAMED_STRUCT_LASTRESORT_NETSERIALIZER_AND_REGISTRY_DELEGATES(NetworkPhysicsRewindDataDeltaSourceStateProxy);
 #endif // UE_WITH_IRIS
 
 bool FNetworkPhysicsRewindDataInputProxy::NetSerialize(FArchive& Ar, class UPackageMap* Map, bool& bOutSuccess)
 {
-	return NetSerializeBase(Ar, Map, bOutSuccess, [this]() { return Owner->ReplicatedInputs.History->CreateNew(); });
+	return NetSerializeBase(Ar, Map, bOutSuccess, [this]() { return Owner->InputHelper->CreateUniqueRewindHistory(0); }, [this](const int32 Value) -> FNetworkPhysicsData* { return Owner->GetDeltaSourceInput(Value, /*bValueIsIndex*/ false); });
 }
 
 bool FNetworkPhysicsRewindDataRemoteInputProxy::NetSerialize(FArchive& Ar, class UPackageMap* Map, bool& bOutSuccess)
 {
-	return NetSerializeBase(Ar, Map, bOutSuccess, [this]() { return Owner->ReplicatedRemoteInputs.History->CreateNew(); });
+	return NetSerializeBase(Ar, Map, bOutSuccess, [this]() { return Owner->InputHelper->CreateUniqueRewindHistory(0); }, [this](const int32 Value) -> FNetworkPhysicsData* { return Owner->GetDeltaSourceInput(Value, /*bValueIsIndex*/ false); });
 }
 
 bool FNetworkPhysicsRewindDataStateProxy::NetSerialize(FArchive& Ar, class UPackageMap* Map, bool& bOutSuccess)
 {
-	return NetSerializeBase(Ar, Map, bOutSuccess, [this]() { return Owner->ReplicatedStates.History->CreateNew(); });
+	return NetSerializeBase(Ar, Map, bOutSuccess, [this]() { return Owner->StateHelper->CreateUniqueRewindHistory(0); }, [this](const int32 Value) -> FNetworkPhysicsData* { return Owner->GetDeltaSourceState(Value, /*bValueIsIndex*/ false); });
 }
 
 bool FNetworkPhysicsRewindDataImportantInputProxy::NetSerialize(FArchive& Ar, class UPackageMap* Map, bool& bOutSuccess)
 {
-	return NetSerializeBase(Ar, Map, bOutSuccess, [this]() { return Owner->ReplicatedImportantInput.History->CreateNew(); });
+	return NetSerializeBase(Ar, Map, bOutSuccess, [this]() { return Owner->InputHelper->CreateUniqueRewindHistory(0); }, [this](const int32 Value) -> FNetworkPhysicsData* { return Owner->GetDeltaSourceInput(Value, /*bValueIsIndex*/ false); });
 }
 
 bool FNetworkPhysicsRewindDataImportantStateProxy::NetSerialize(FArchive& Ar, class UPackageMap* Map, bool& bOutSuccess)
 {
-	return NetSerializeBase(Ar, Map, bOutSuccess, [this]() { return Owner->ReplicatedImportantState.History->CreateNew(); });
+	return NetSerializeBase(Ar, Map, bOutSuccess, [this]() { return Owner->StateHelper->CreateUniqueRewindHistory(0); }, [this](const int32 Value) -> FNetworkPhysicsData* { return Owner->GetDeltaSourceState(Value, /*bValueIsIndex*/ false); });
 }
 
+bool FNetworkPhysicsRewindDataDeltaSourceInputProxy::NetSerialize(FArchive& Ar, class UPackageMap* Map, bool& bOutSuccess)
+{
+	if (PhysicsReplicationCVars::ResimulationCVars::bUseDefaultDeltaForDeltaSourceReplication)
+	{
+		// Use default as base for delta serialization when sending delta source
+		return NetSerializeBase(Ar, Map, bOutSuccess, [this]() { return Owner->InputHelper->CreateUniqueRewindHistory(0); }, [this](const int32 Value) -> FNetworkPhysicsData* { return Owner->GetDeltaSourceInput(/*Default*/ -2, /*bValueIsIndex*/ false); });
+	}
+	else
+	{
+		// Standard serialization for delta source
+		return NetSerializeBase(Ar, Map, bOutSuccess, [this]() { return Owner->InputHelper->CreateUniqueRewindHistory(0); }, nullptr);
+	}
+}
+
+bool FNetworkPhysicsRewindDataDeltaSourceStateProxy::NetSerialize(FArchive& Ar, class UPackageMap* Map, bool& bOutSuccess)
+{
+	if (PhysicsReplicationCVars::ResimulationCVars::bUseDefaultDeltaForDeltaSourceReplication)
+	{
+		// Use default as base for delta serialization when sending delta source
+		return NetSerializeBase(Ar, Map, bOutSuccess, [this]() { return Owner->StateHelper->CreateUniqueRewindHistory(0); }, [this](const int32 Value) -> FNetworkPhysicsData* { return Owner->GetDeltaSourceState(/*Default*/ -2, /*bValueIsIndex*/ false); });
+	}
+	else
+	{
+		// Standard serialization for delta source
+		return NetSerializeBase(Ar, Map, bOutSuccess, [this]() { return Owner->StateHelper->CreateUniqueRewindHistory(0); }, nullptr);
+	}
+}
 
 // --------------------------- Network Physics Callback ---------------------------
 
@@ -444,6 +525,8 @@ void UNetworkPhysicsComponent::GetLifetimeReplicatedProps(TArray< FLifetimePrope
 	ReplicatedParamsAll.RepNotifyCondition = REPNOTIFY_Always;
 	ReplicatedParamsAll.bIsPushBased = true;
 
+	DOREPLIFETIME_WITH_PARAMS_FAST(UNetworkPhysicsComponent, ReplicatedDeltaSourceInput, ReplicatedParamsAll);
+	DOREPLIFETIME_WITH_PARAMS_FAST(UNetworkPhysicsComponent, ReplicatedDeltaSourceState, ReplicatedParamsAll);
 	DOREPLIFETIME_CONDITION(UNetworkPhysicsComponent, InputsToNetwork_Owner, COND_OwnerOnly);
 	DOREPLIFETIME_WITH_PARAMS_FAST(UNetworkPhysicsComponent, ReplicatedInputs, ReplicatedParamsOwner);
 	DOREPLIFETIME_WITH_PARAMS_FAST(UNetworkPhysicsComponent, ReplicatedRemoteInputs, ReplicatedParamsRemote);
@@ -454,12 +537,20 @@ void UNetworkPhysicsComponent::GetLifetimeReplicatedProps(TArray< FLifetimePrope
 void UNetworkPhysicsComponent::TickComponent(float DeltaTime, ELevelTick TickType, FActorComponentTickFunction* ThisTickFunction)
 {
 	Super::TickComponent(DeltaTime, TickType, ThisTickFunction);
+
 	UpdateAsyncComponent(false);
 	NetworkMarshaledData();
 }
 
 void UNetworkPhysicsComponent::NetworkMarshaledData()
 {
+	UWorld* World = GetWorld();
+	if (!World)
+	{
+		return;
+	}
+
+
 	if (NetworkPhysicsComponent_Internal)
 	{
 		const bool bIsServer = HasServerWorld();
@@ -469,6 +560,43 @@ void UNetworkPhysicsComponent::NetworkMarshaledData()
 			return;
 		}
 
+		const bool bShouldSyncDeltaSourceInput = bIsServer && PhysicsReplicationCVars::ResimulationCVars::bEnableStatefulDeltaSerialization && World->GetRealTimeSeconds() > TimeToSyncDeltaSourceInput;
+		const bool bShouldSyncDeltaSourceState = bIsServer && PhysicsReplicationCVars::ResimulationCVars::bEnableStatefulDeltaSerialization && World->GetRealTimeSeconds() > TimeToSyncDeltaSourceState;
+		bool bHasSyncedDeltaSourceInput = false;
+		bool bHasSyncedDeltaSourceState = false;
+
+		// Replicate source data for input delta serialization
+		auto DeltaSourceInputSyncHelper = [&](const TUniquePtr<Chaos::FBaseRewindHistory>& InputData)
+		{
+			if (bShouldSyncDeltaSourceInput && !bHasSyncedDeltaSourceInput && IsValidNextDeltaSourceInput(InputData->GetLatestFrame()))
+			{
+				ReplicatedDeltaSourceInput.History->SetRecordDataIncremental(true); // Only record data that is newer than already cached data
+				if (InputData->CopyAllData(*ReplicatedDeltaSourceInput.History, /*bIncludeUnimportant*/ true, /*bIncludeImportant*/ true))
+				{
+					MARK_PROPERTY_DIRTY_FROM_NAME(UNetworkPhysicsComponent, ReplicatedDeltaSourceInput, this);
+					bHasSyncedDeltaSourceInput = true;
+					TimeToSyncDeltaSourceInput = World->GetRealTimeSeconds() + PhysicsReplicationCVars::ResimulationCVars::TimeToSyncStatefulDeltaSource;
+					AddDeltaSourceInput();
+				}
+			}
+		};
+
+		// Replicate source data for state delta serialization
+		auto DeltaSourceStateSyncHelper = [&](const TUniquePtr<Chaos::FBaseRewindHistory>& StateData)
+		{
+			if (bShouldSyncDeltaSourceState && !bHasSyncedDeltaSourceState && IsValidNextDeltaSourceState(StateData->GetLatestFrame()))
+			{
+				ReplicatedDeltaSourceState.History->SetRecordDataIncremental(true); // Only record data that is newer than already cached data
+				if (StateData->CopyAllData(*ReplicatedDeltaSourceState.History, /*bIncludeUnimportant*/ true, /*bIncludeImportant*/ true))
+				{
+					MARK_PROPERTY_DIRTY_FROM_NAME(UNetworkPhysicsComponent, ReplicatedDeltaSourceState, this);
+					bHasSyncedDeltaSourceState = true;
+					TimeToSyncDeltaSourceState = World->GetRealTimeSeconds() + PhysicsReplicationCVars::ResimulationCVars::TimeToSyncStatefulDeltaSource;
+					AddDeltaSourceState();
+				}
+			}
+		};
+
 		while (Chaos::TSimCallbackOutputHandle<FAsyncNetworkPhysicsComponentOutput> AsyncOutput = NetworkPhysicsComponent_Internal->PopFutureOutputData_External())
 		{
 			if (AsyncOutput->InputsToNetwork_Owner.IsSet())
@@ -484,6 +612,9 @@ void UNetworkPhysicsComponent::NetworkMarshaledData()
 			{
 				if (bIsServer)
 				{
+					// Replicate source data for delta serialization
+					DeltaSourceInputSyncHelper(AsyncOutput->InputData);
+
 					if (IsLocallyControlled())
 					{
 						// Send inputs to remote clients after getting marshaled from PT if server is the one controlling the component
@@ -526,6 +657,9 @@ void UNetworkPhysicsComponent::NetworkMarshaledData()
 						continue;
 					}
 
+					// Replicate source data for delta serialization
+					DeltaSourceInputSyncHelper(InputImportant);
+
 					ReplicatedImportantInput.History->ResizeDataHistory(InputImportant->GetHistorySize(), EAllowShrinking::Yes);
 					if (InputImportant->CopyAllData(*ReplicatedImportantInput.History, /*bIncludeUnimportant*/ false, /*bIncludeImportant*/ true))
 					{
@@ -546,11 +680,16 @@ void UNetworkPhysicsComponent::NetworkMarshaledData()
 				// Unimportant / Unreliable
 				if (bEnableUnreliableFlow
 					&& AsyncOutput->StateData
-					&& AsyncOutput->StateData->HasDataInHistory()
-					&& AsyncOutput->StateData->CopyAllData(*ReplicatedStates.History, /*bIncludeUnimportant*/ true, /*bIncludeImportant*/ bEnableReliableFlow == false))
+					&& AsyncOutput->StateData->HasDataInHistory())
 				{
-					// If on server we should send the states onto all the clients through repnotify
-					MARK_PROPERTY_DIRTY_FROM_NAME(UNetworkPhysicsComponent, ReplicatedStates, this);
+					// Replicate source data for delta serialization
+					DeltaSourceStateSyncHelper(AsyncOutput->StateData);
+
+					if (AsyncOutput->StateData->CopyAllData(*ReplicatedStates.History, /*bIncludeUnimportant*/ true, /*bIncludeImportant*/ bEnableReliableFlow == false))
+					{
+						// If on server we should send the states onto all the clients through repnotify
+						MARK_PROPERTY_DIRTY_FROM_NAME(UNetworkPhysicsComponent, ReplicatedStates, this);
+					}
 				}
 
 				// Important / Reliable
@@ -563,6 +702,9 @@ void UNetworkPhysicsComponent::NetworkMarshaledData()
 							continue;
 						}
 
+						// Replicate source data for delta serialization
+						DeltaSourceStateSyncHelper(StateImportant);
+
 						ReplicatedImportantState.History->ResizeDataHistory(StateImportant->GetHistorySize(), EAllowShrinking::Yes);
 						if (StateImportant->CopyAllData(*ReplicatedImportantState.History, /*bIncludeUnimportant*/ false, /*bIncludeImportant*/ true))
 						{
@@ -581,9 +723,187 @@ void UNetworkPhysicsComponent::NetworkMarshaledData()
 	}
 }
 
+void UNetworkPhysicsComponent::OnRep_SetReplicatedDeltaSourceInput()
+{
+	if (ReplicatedDeltaSourceInput.bDeltaSerializationIssue)
+	{
+		return;
+	}
+	AddDeltaSourceInput();
+}
+
+void UNetworkPhysicsComponent::ServerReceiveDeltaSourceInputFrame_Implementation(const int32 Frame)
+{
+	for (int32 I = 0; I < DeltaSourceInputs.Num(); I++)
+	{
+		TUniquePtr<FNetworkPhysicsData>& Data = DeltaSourceInputs[I];
+		if (Data->ServerFrame == Frame)
+		{
+			// Set latest delta source index acknowledged by the client so that we can start using this delta source
+			LatestAcknowledgedDeltaSourceInputIndex = I;
+			break;
+		}
+	}
+}
+
+void UNetworkPhysicsComponent::AddDeltaSourceInput()
+{
+	// Get the data entry for the correct index in the data sources array
+	const int32 Index = GetDeltaSourceIndexForFrame(ReplicatedDeltaSourceInput.History->GetLatestFrame());
+	check(Index <= DeltaSourceInputs.Num());
+	FNetworkPhysicsData* PhysicsData = DeltaSourceInputs[Index].Get();
+
+	// Extract the data from the replicated DeltaSources property
+	if (ReplicatedDeltaSourceInput.History->ExtractData(ReplicatedDeltaSourceInput.History->GetLatestFrame(), /*bResetSolver*/false, PhysicsData, /*bExactFrame*/true))
+	{
+		// The data is now extracted via PhysicsData and stored inside DeltaSourceInputs
+		
+		if (!HasServerWorld())
+		{
+			// On the client, set the latest index, to be used when sending inputs to the server
+			LatestAcknowledgedDeltaSourceInputIndex = Index;
+
+			if (IsLocallyControlled())
+			{
+				// If this client is the one controlling this entity, send back an acknowledgment to the server we have received the delta source for ServerFrame
+				ServerReceiveDeltaSourceInputFrame(PhysicsData->ServerFrame);
+			}
+		}
+		else if (IsLocallyControlled())
+		{
+			// If server is locally controlled, set the latest index directly, else wait for the owning client to send back ServerReceiveDeltaSourceInputFrame before the server starts to use this 
+			LatestAcknowledgedDeltaSourceInputIndex = Index;
+		}
+
+		LatestCachedDeltaSourceInputIndex = Index;
+	}
+}
+
+void UNetworkPhysicsComponent::OnRep_SetReplicatedDeltaSourceState()
+{
+	if (ReplicatedDeltaSourceState.bDeltaSerializationIssue)
+	{
+		return;
+	}
+	AddDeltaSourceState();
+}
+
+void UNetworkPhysicsComponent::ServerReceiveDeltaSourceStateFrame_Implementation(const int32 Frame)
+{
+	for (int32 I = 0; I < DeltaSourceStates.Num(); I++)
+	{
+		TUniquePtr<FNetworkPhysicsData>& Data = DeltaSourceStates[I];
+		if (Data->ServerFrame == Frame)
+		{
+			// Set latest delta source index acknowledged by the client so that we can start using this delta source
+			LatestAcknowledgedDeltaSourceStateIndex = I;
+			break;
+		}
+	}
+}
+
+void UNetworkPhysicsComponent::AddDeltaSourceState()
+{
+	// Get the data entry for the correct index in the data sources array
+	const int32 Index = GetDeltaSourceIndexForFrame(ReplicatedDeltaSourceState.History->GetLatestFrame());
+	check(Index <= DeltaSourceStates.Num());
+	FNetworkPhysicsData* PhysicsData = DeltaSourceStates[Index].Get();
+
+	// Extract the data from the replicated DeltaSources property
+	if (ReplicatedDeltaSourceState.History->ExtractData(ReplicatedDeltaSourceState.History->GetLatestFrame(), /*bResetSolver*/false, PhysicsData, /*bExactFrame*/true))
+	{
+		// The data is now extracted via PhysicsData and stored inside DeltaSourceStates
+
+		if (!HasServerWorld())
+		{
+			// On the client, set the latest index (unlike for DeltaSourceInput this latest index is not used on the client since the client doesn't send states towards the server, but set the value for the future)
+			LatestAcknowledgedDeltaSourceStateIndex = Index;
+
+			if (IsLocallyControlled())
+			{
+				// If this client is the one controlling this entity, send back an acknowledgment to the server we have received the delta source for ServerFrame
+				ServerReceiveDeltaSourceStateFrame(PhysicsData->ServerFrame);
+			}
+		}
+		else if (IsLocallyControlled())
+		{
+			// If server is locally controlled, set the latest index directly, else wait for the owning client to send back ServerReceiveDeltaSourceStateFrame before the server starts to use this 
+			LatestAcknowledgedDeltaSourceStateIndex = Index;
+		}
+
+		LatestCachedDeltaSourceStateIndex = Index;
+	}
+}
+
+FNetworkPhysicsData* UNetworkPhysicsComponent::GetDeltaSourceInput(const int32 Value, const bool bValueIsIndexElseFrame)
+{
+	TUniquePtr<FNetworkPhysicsData>* DataPtr = nullptr;
+	if (Value == -1) // Latest
+	{
+		DataPtr = &DeltaSourceInputs[LatestAcknowledgedDeltaSourceInputIndex];
+	}
+	else if (Value == -2) // Default
+	{
+		DataPtr = &InputDataDefault;
+	}
+	else if(bValueIsIndexElseFrame)
+	{
+		if (Value < DeltaSourceInputs.Num())
+		{
+			DataPtr = &DeltaSourceInputs[Value];
+		}
+	}
+	else // Value is Frame
+	{
+		int32 Index = GetDeltaSourceIndexForFrame(Value);
+		if (Index < DeltaSourceInputs.Num())
+		{
+			if (DeltaSourceInputs[Index]->ServerFrame == Value)
+			{
+				DataPtr = &DeltaSourceInputs[Index];
+			}
+		}
+	}
+
+	return DataPtr ? DataPtr->Get() : nullptr;
+}
+
+FNetworkPhysicsData* UNetworkPhysicsComponent::GetDeltaSourceState(const int32 Value, const bool bValueIsIndexElseFrame)
+{
+	TUniquePtr<FNetworkPhysicsData>* DataPtr = nullptr;
+	if (Value == -1) // Latest
+	{
+		DataPtr = &DeltaSourceStates[LatestAcknowledgedDeltaSourceStateIndex];
+	}
+	else if (Value == -2) // Default
+	{
+		DataPtr = &StateDataDefault;
+	}
+	else if (bValueIsIndexElseFrame)
+	{
+		if (Value < DeltaSourceStates.Num())
+		{
+			DataPtr = &DeltaSourceStates[Value];
+		}
+	}
+	else // Value is Frame
+	{
+		int32 Index = GetDeltaSourceIndexForFrame(Value);
+		if (Index < DeltaSourceStates.Num())
+		{
+			if (DeltaSourceStates[Index]->ServerFrame == Value)
+			{
+				DataPtr = &DeltaSourceStates[Index];
+			}
+		}
+	}
+
+	return DataPtr ? DataPtr->Get() : nullptr;
+}
+
 void UNetworkPhysicsComponent::OnRep_SetReplicatedStates()
 {
-	if (!NetworkPhysicsComponent_Internal || !StateHelper || !ReplicatedStates.History)
+	if (!NetworkPhysicsComponent_Internal || !StateHelper || !ReplicatedStates.History || ReplicatedStates.bDeltaSerializationIssue)
 	{
 		return;
 	}
@@ -603,7 +923,7 @@ void UNetworkPhysicsComponent::OnRep_SetReplicatedStates()
 
 void UNetworkPhysicsComponent::OnRep_SetReplicatedInputs()
 {
-	if (!NetworkPhysicsComponent_Internal || !InputHelper || !ReplicatedInputs.History)
+	if (!NetworkPhysicsComponent_Internal || !InputHelper || !ReplicatedInputs.History || ReplicatedInputs.bDeltaSerializationIssue)
 	{
 		return ;
 	}
@@ -623,7 +943,7 @@ void UNetworkPhysicsComponent::OnRep_SetReplicatedInputs()
 
 void UNetworkPhysicsComponent::OnRep_SetReplicatedRemoteInputs()
 {
-	if (!NetworkPhysicsComponent_Internal || !InputHelper || !ReplicatedRemoteInputs.History)
+	if (!NetworkPhysicsComponent_Internal || !InputHelper || !ReplicatedRemoteInputs.History || ReplicatedRemoteInputs.bDeltaSerializationIssue)
 	{
 		return ;
 	}
@@ -643,7 +963,7 @@ void UNetworkPhysicsComponent::OnRep_SetReplicatedRemoteInputs()
 
 void UNetworkPhysicsComponent::ServerReceiveInputData_Implementation(const FNetworkPhysicsRewindDataInputProxy& ClientInputs)
 {
-	if (!NetworkPhysicsComponent_Internal || !InputHelper || !ClientInputs.History)
+	if (!NetworkPhysicsComponent_Internal || !InputHelper || !ClientInputs.History || ClientInputs.bDeltaSerializationIssue)
 	{
 		return;
 	}
@@ -656,9 +976,9 @@ void UNetworkPhysicsComponent::ServerReceiveInputData_Implementation(const FNetw
 		}
 
 		// Validate data in the received inputs
-		if (bValidateDataOnGameThread && ActorComponent)
+		if (bValidateDataOnGameThread && ActorComponent.IsValid())
 		{
-			ClientInputs.History->ValidateDataInHistory(ActorComponent);
+			ClientInputs.History->ValidateDataInHistory(ActorComponent.Get());
 		}
 
 		AsyncInput->InputData->ResetFast();
@@ -675,7 +995,7 @@ void UNetworkPhysicsComponent::ServerReceiveInputData_Implementation(const FNetw
 
 void UNetworkPhysicsComponent::ServerReceiveImportantInputData_Implementation(const FNetworkPhysicsRewindDataImportantInputProxy& ClientInputs)
 {
-	if (!NetworkPhysicsComponent_Internal || !ClientInputs.History)
+	if (!NetworkPhysicsComponent_Internal || !ClientInputs.History || ClientInputs.bDeltaSerializationIssue)
 	{
 		return;
 	}
@@ -686,9 +1006,9 @@ void UNetworkPhysicsComponent::ServerReceiveImportantInputData_Implementation(co
 		ClientInputs.History->Initialize();
 
 		// Validate data in the received inputs
-		if (bValidateDataOnGameThread && ActorComponent)
+		if (bValidateDataOnGameThread && ActorComponent.IsValid())
 		{
-			ClientInputs.History->ValidateDataInHistory(ActorComponent);
+			ClientInputs.History->ValidateDataInHistory(ActorComponent.Get());
 		}
 
 		// Create new data collection for marshaling
@@ -704,7 +1024,7 @@ void UNetworkPhysicsComponent::MulticastReceiveImportantInputData_Implementation
 		return;
 	}
 
-	if (!NetworkPhysicsComponent_Internal || !ServerInputs.History)
+	if (!NetworkPhysicsComponent_Internal || !ServerInputs.History || ServerInputs.bDeltaSerializationIssue)
 	{
 		return;
 	}
@@ -727,7 +1047,7 @@ void UNetworkPhysicsComponent::MulticastReceiveImportantStateData_Implementation
 		return;
 	}
 
-	if (!NetworkPhysicsComponent_Internal || !ServerStates.History)
+	if (!NetworkPhysicsComponent_Internal || !ServerStates.History || ServerStates.bDeltaSerializationIssue)
 	{
 		return;
 	}
@@ -897,9 +1217,9 @@ void UNetworkPhysicsComponent::UpdateAsyncComponent(const bool bFullUpdate)
 					}
 				}
 				
-				if (ActorComponent)
+				if (ActorComponent.IsValid())
 				{
-					AsyncInput->ActorComponent = ActorComponent;
+					AsyncInput->ActorComponent = ActorComponent.Get();
 				}
 			}
 		}
@@ -912,7 +1232,7 @@ void UNetworkPhysicsComponent::CreateAsyncDataHistory()
 	{
 		if (FAsyncNetworkPhysicsComponentInput* AsyncInput = NetworkPhysicsComponent_Internal->GetProducerInputData_External())
 		{
-			AsyncInput->ActorComponent = ActorComponent;
+			AsyncInput->ActorComponent = ActorComponent.Get();
 
 			if (InputHelper)
 			{
diff --git a/Engine/Source/Runtime/Engine/Public/Physics/NetworkPhysicsComponent.h b/Engine/Source/Runtime/Engine/Public/Physics/NetworkPhysicsComponent.h
--- a/Engine/Source/Runtime/Engine/Public/Physics/NetworkPhysicsComponent.h
+++ b/Engine/Source/Runtime/Engine/Public/Physics/NetworkPhysicsComponent.h
@@ -275,11 +275,27 @@ struct TNetRewindHistory : public Chaos::TDataRewindHistory<DataType>
 		return false;
 	}
 
+	UE_DEPRECATED(5.6, "Deprecated, use the NetSerialize call with parameter that takes a function DataSetupFunction, pass in nullptr to opt out of implementing a function.")
 	virtual void NetSerialize(FArchive& Ar, UPackageMap* InPackageMap) override
 	{
-		uint32 NumFramesUnsigned = static_cast<uint32>(Super::NumFrames);
-		Ar.SerializeIntPacked(NumFramesUnsigned);
-		Super::NumFrames = static_cast<int32>(NumFramesUnsigned);
+		NetSerialize(Ar, InPackageMap, [](void* Data, const int32 DataIndex) {});
+	}
+
+	virtual void NetSerialize(FArchive& Ar, UPackageMap* InPackageMap, TUniqueFunction<void(void* Data, const int32 DataIndex)> DataSetupFunction) override
+	{
+		bool bOneEntry = Super::NumFrames == 1;
+		Ar.SerializeBits(&bOneEntry, 1);
+
+		if (!bOneEntry)
+		{
+			uint32 NumFramesUnsigned = static_cast<uint32>(Super::NumFrames);
+			Ar.SerializeIntPacked(NumFramesUnsigned);
+			Super::NumFrames = static_cast<int32>(NumFramesUnsigned);
+		}
+		else
+		{
+			Super::NumFrames = 1;
+		}
 
 		if (Super::NumFrames > GetMaxArraySize())
 		{
@@ -293,10 +309,33 @@ struct TNetRewindHistory : public Chaos::TDataRewindHistory<DataType>
 			Super::DataHistory.SetNum(Super::NumFrames);
 		}
 
+		for (int32 I = 0; I < Super::DataHistory.Num(); I++)
+		{
+			DataType& Data = Super::DataHistory[I];
+
+			// Set the implementation component pointer and stateful delta serialization source
+			if (DataSetupFunction)
+			{
+				DataSetupFunction(&Data, I);
+			}
+
+			// Set the internal delta serialization source (between data entries in the collection)
+			if (I > 0)
+			{
+				Data.SetDeltaSourceData(&Super::DataHistory[I - 1]);
+			}
+		}
+
 		for (DataType& Data : Super::DataHistory)
 		{
 			NetSerializeData(Data, Ar, InPackageMap);
+
+			// Clear delta source and strong pointer to implementation component after serialization
+			Data.ClearImplementationComponent();
+			Data.ClearDeltaSourceData();
 		}
+
+		Super::Initialize();
 	}
 
 	/** Debug the data from the archive */
@@ -396,7 +435,10 @@ struct FNetworkPhysicsRewindDataProxy
 	bool operator==(const FNetworkPhysicsRewindDataProxy& Other) const { return false; }
 
 protected:
-	bool NetSerializeBase(FArchive& Ar, class UPackageMap* Map, bool& bOutSuccess, TUniqueFunction<TUniquePtr<Chaos::FBaseRewindHistory>()> CreateHistoryFunction);
+	UE_DEPRECATED(5.6, "Deprecated, use the NetSerializeBase call with parameter that takes a function GetDeltaSourceData, pass in nullptr to opt out of implementing a source for delta compression")
+	bool NetSerializeBase(FArchive& Ar, class UPackageMap* Map, bool& bOutSuccess, TUniqueFunction<TUniquePtr<Chaos::FBaseRewindHistory>()> CreateHistoryFunction) { return NetSerializeBase(Ar, Map, bOutSuccess, [&](){ return CreateHistoryFunction(); }, nullptr); };
+
+	bool NetSerializeBase(FArchive& Ar, class UPackageMap* Map, bool& bOutSuccess, TUniqueFunction<TUniquePtr<Chaos::FBaseRewindHistory>()> CreateHistoryFunction, TUniqueFunction<FNetworkPhysicsData*(const int32)> GetDeltaSourceData);
 
 public:
 	/** The history to be serialized */
@@ -405,6 +447,9 @@ public:
 	/** Component that utilizes this data */
 	UPROPERTY()
 	TObjectPtr<UNetworkPhysicsComponent> Owner = nullptr;
+
+	/** If a delta serialization issue was detected, i.e. the data might be corrupt if this is true*/
+	bool bDeltaSerializationIssue = false;
 };
 
 /**
@@ -512,6 +557,49 @@ struct TStructOpsTypeTraits<FNetworkPhysicsRewindDataImportantStateProxy> : publ
 };
 
 
+/**
+ * Struct suitable for use as a replicated property to replicate state rewind history
+ */
+USTRUCT()
+struct FNetworkPhysicsRewindDataDeltaSourceStateProxy : public FNetworkPhysicsRewindDataProxy
+{
+	GENERATED_BODY()
+
+	bool NetSerialize(FArchive& Ar, class UPackageMap* Map, bool& bOutSuccess);
+};
+
+template<>
+struct TStructOpsTypeTraits<FNetworkPhysicsRewindDataDeltaSourceStateProxy> : public TStructOpsTypeTraitsBase2<FNetworkPhysicsRewindDataDeltaSourceStateProxy>
+{
+	enum
+	{
+		WithNetSerializer = true,
+		WithIdenticalViaEquality = true
+	};
+};
+
+/**
+ * Struct suitable for use as a replicated property to replicate state rewind history
+ */
+USTRUCT()
+struct FNetworkPhysicsRewindDataDeltaSourceInputProxy : public FNetworkPhysicsRewindDataProxy
+{
+	GENERATED_BODY()
+
+	bool NetSerialize(FArchive& Ar, class UPackageMap* Map, bool& bOutSuccess);
+};
+
+template<>
+struct TStructOpsTypeTraits<FNetworkPhysicsRewindDataDeltaSourceInputProxy> : public TStructOpsTypeTraitsBase2<FNetworkPhysicsRewindDataDeltaSourceInputProxy>
+{
+	enum
+	{
+		WithNetSerializer = true,
+		WithIdenticalViaEquality = true
+	};
+};
+
+
 /**
  * Network physics rewind callback to manage all the sim callbacks rewind functionalities
  */
@@ -617,22 +705,84 @@ struct FNetworkPhysicsData
 	// If this data is marked as important (replicated reliably)
 	bool bImportant = false;
 
+	/** Thread safe pointer to the UActorComponent that implements the derived type of this data
+	* Note: This pointer can be accessed on both Game Thread and Physics Thread but you still need to ensure the read and write calls inside the UActorComponent are thread safe. */
+	TStrongObjectPtr<UActorComponent> ImplementationComponent = nullptr;
+	
+	void SetImplementationComponent(UActorComponent* InImplementationComponent)
+	{
+		ImplementationComponent = TStrongObjectPtr<UActorComponent>(InImplementationComponent);
+	}
+	void ClearImplementationComponent() { ImplementationComponent = nullptr; }
+
+	/** Pointer to a previous FNetworkPhysicsData which is valid during NetSerialize() to be used for delta serialization */
+	FNetworkPhysicsData* DeltaSourceData = nullptr;
+
+	void SetDeltaSourceData(FNetworkPhysicsData* IntDeltaSourceData)
+	{
+		DeltaSourceData = IntDeltaSourceData ? IntDeltaSourceData : nullptr;
+	}
+	void ClearDeltaSourceData() { DeltaSourceData = nullptr; }
+
 	// Serialize the data into/from the archive
 	void SerializeFrames(FArchive& Ar)
 	{
-		uint32 ServerFrameUnsigned = 0;
-
-		if (Ar.IsLoading()) // Deserializing
+		// Delta Serialization
+		if (DeltaSourceData) 
 		{
-			Ar.SerializeIntPacked(ServerFrameUnsigned);
-			ServerFrame = static_cast<int32>(ServerFrameUnsigned) - 1;
-			LocalFrame = ServerFrame; // Temporarily set LocalFrame to ServerFrame, it will get recalculated later in TRewindHistory::ReceiveNewData
+			uint32 ServerFrameUnsigned = 0;
+
+			bool bIncrememtalFrame = false;
+			if (Ar.IsLoading())
+			{
+				Ar.SerializeBits(&bIncrememtalFrame, 1);
+				if (!bIncrememtalFrame)
+				{
+					bool bFrameDeltaNegative = false;
+					Ar.SerializeBits(&bFrameDeltaNegative, 1); // Get if the delta is negative
+					uint32 FrameDeltaUnsigned = 0;
+					Ar.SerializeIntPacked(FrameDeltaUnsigned); // Get the frame delta
+
+					// Apply the frame delta to the delta source to get the ServerFrame value
+					ServerFrame = bFrameDeltaNegative ? (DeltaSourceData->ServerFrame - FrameDeltaUnsigned) : (DeltaSourceData->ServerFrame + FrameDeltaUnsigned);
+				}
+				else
+				{
+					// Increment the delta source ServerFrame once to get the ServerFrame value 
+					ServerFrame = DeltaSourceData->ServerFrame + 1;
+				}
+
+				LocalFrame = ServerFrame; // Temporarily set LocalFrame to ServerFrame, it will get recalculated later in TRewindHistory::ReceiveNewData
+			}
+			else
+			{
+				bIncrememtalFrame = ServerFrame == (DeltaSourceData->ServerFrame + 1);
+				Ar.SerializeBits(&bIncrememtalFrame, 1); // Write if the frame delta is just +1, which is most common for internal deltas
+				if (!bIncrememtalFrame)
+				{
+					int32 FrameDelta = ServerFrame - DeltaSourceData->ServerFrame; // Get the frame delta
+					bool bFrameDeltaNegative = FrameDelta < 0;
+					Ar.SerializeBits(&bFrameDeltaNegative, 1); // Write if the delta is negative
+					uint32 FrameDeltaUnsigned = FMath::Abs(FrameDelta);
+					Ar.SerializeIntPacked(FrameDeltaUnsigned); // Write the unsigned delta frame
+				}
+			}
 		}
-		else // Serializing
+		else // Standard Serialization
 		{
-			check((ServerFrame + 1) >= 0);
-			ServerFrameUnsigned = static_cast<uint32>(ServerFrame + 1);
-			Ar.SerializeIntPacked(ServerFrameUnsigned);
+			uint32 ServerFrameUnsigned = 0;
+			if (Ar.IsLoading()) // Deserializing
+			{
+				Ar.SerializeIntPacked(ServerFrameUnsigned);
+				ServerFrame = static_cast<int32>(ServerFrameUnsigned) - 1;
+				LocalFrame = ServerFrame; // Temporarily set LocalFrame to ServerFrame, it will get recalculated later in TRewindHistory::ReceiveNewData
+			}
+			else // Serializing
+			{
+				check((ServerFrame + 1) >= 0);
+				ServerFrameUnsigned = static_cast<uint32>(ServerFrame + 1);
+				Ar.SerializeIntPacked(ServerFrameUnsigned);
+			}
 		}
 	}
 
@@ -710,6 +860,7 @@ struct FNetworkPhysicsDataHelper
 	
 	virtual TUniquePtr<FNetworkPhysicsDataHelper> Clone() const = 0;
 	virtual TUniquePtr<FNetworkPhysicsData> CreateUniqueData() const = 0;
+	virtual void CopyData(FNetworkPhysicsData* To, FNetworkPhysicsData* From) = 0;
 	virtual TUniquePtr<Chaos::FBaseRewindHistory> CreateUniqueRewindHistory(const int32 Size) const = 0;
 };
 
@@ -721,6 +872,10 @@ struct TNetworkPhysicsDataHelper : FNetworkPhysicsDataHelper
 
 	virtual TUniquePtr<FNetworkPhysicsDataHelper> Clone() const { return MakeUnique<TNetworkPhysicsDataHelper>(*this); }
 	virtual TUniquePtr<FNetworkPhysicsData> CreateUniqueData() const { return MakeUnique<DataType>(); }
+	virtual void CopyData(FNetworkPhysicsData* To, FNetworkPhysicsData* From)
+	{
+		*static_cast<DataType*>(To) = *static_cast<DataType*>(From);
+	};
 	virtual TUniquePtr<Chaos::FBaseRewindHistory> CreateUniqueRewindHistory(const int32 Size) const { return MakeUnique<TNetRewindHistory<DataType>>(Size); }
 };
 
@@ -815,7 +970,32 @@ public:
 	/** Get FAsyncNetworkPhysicsComponent on the Physics Thread */
 	FAsyncNetworkPhysicsComponent* GetNetworkPhysicsComponent_Internal() { return NetworkPhysicsComponent_Internal; }
 
-protected : 
+protected: 
+
+	// repnotify for input, used for delta serialization
+	UFUNCTION()
+	ENGINE_API void OnRep_SetReplicatedDeltaSourceInput();
+
+	// repnotify for state, used for delta serialization
+	UFUNCTION()
+	ENGINE_API void OnRep_SetReplicatedDeltaSourceState();
+
+	// replicated physics input, used for delta serialization
+	UPROPERTY(Transient, ReplicatedUsing = OnRep_SetReplicatedDeltaSourceInput)
+	FNetworkPhysicsRewindDataDeltaSourceInputProxy ReplicatedDeltaSourceInput;
+
+	// replicated physics states, used for delta serialization
+	UPROPERTY(Transient, ReplicatedUsing = OnRep_SetReplicatedDeltaSourceState)
+	FNetworkPhysicsRewindDataDeltaSourceStateProxy ReplicatedDeltaSourceState;
+
+	/** Server RPC to acknowledge which Delta Source Input Frame the owning client has received */
+	UFUNCTION(Server, reliable)
+	ENGINE_API void ServerReceiveDeltaSourceInputFrame(const int32 Frame);
+	
+	/** Server RPC to acknowledge which Delta Source State Frame the owning client has received */
+	UFUNCTION(Server, reliable)
+	ENGINE_API void ServerReceiveDeltaSourceStateFrame(const int32 Frame);
+
 
 	// Server RPC to receive inputs from client
 	UFUNCTION(Server, unreliable)
@@ -897,7 +1077,41 @@ private:
 		StatesToNetwork = FMath::Max(NumInputs, static_cast<uint16>(1));
 	}
 
+	/** Inject the current delta source data into array */
+	void AddDeltaSourceInput();
+
+	/** Inject the current delta source data into array */
+	void AddDeltaSourceState();
+
+	/** Returns the next index in the delta source input array */
+	const int32 GetNextDeltaSourceInputIndex() const { return GetDeltaSourceIndexForFrame(LatestCachedDeltaSourceInputIndex + 1); }
+
+	/** Returns the next index in the delta source state array */
+	const int32 GetNextDeltaSourceStateIndex() const { return GetDeltaSourceIndexForFrame(LatestCachedDeltaSourceStateIndex + 1); }
+
+	/** Returns if this frame is valid to store at the next index in the delta source input array, frame value needs to match with the modulo of the array size */
+	const bool IsValidNextDeltaSourceInput(const int32 Frame) const { return GetDeltaSourceIndexForFrame(Frame) == GetNextDeltaSourceInputIndex(); }
+
+	/** Returns if this frame is valid to store at the next index in the delta source state array, frame value needs to match with the modulo of the array size */
+	const bool IsValidNextDeltaSourceState(const int32 Frame) const { return GetDeltaSourceIndexForFrame(Frame) == GetNextDeltaSourceStateIndex(); }
+
 public:
+	/** Get the delta source stored at index or with frame, @param bValueIsIndex switches between grabbing frame vs index
+	* Note, passing in @param Value as -1 returns the latest data
+	* Note, passing in @param Value as -2 returns default data */
+	FNetworkPhysicsData* GetDeltaSourceInput(const int32 Value, const bool bValueIsIndexElseFrame);
+
+	/** Get the delta source stored at index or with frame, @param bValueIsIndex switches between grabbing frame vs index
+	* Note, passing in @param Value as - 1 returns the latest data
+	* Note, passing in @param Value as - 2 returns default data */
+	FNetworkPhysicsData* GetDeltaSourceState(const int32 Value, const bool bValueIsIndexElseFrame);
+
+	/** Size of the array caching Delta Sources for delta serialization. */
+	static const int32 DeltaSourceBufferSize = 10;
+
+	/** Convert frame number to its corresponding index it would hold in the delta sources array */
+	static const int32 GetDeltaSourceIndexForFrame(const int32 Frame) { return FMath::Abs(Frame % DeltaSourceBufferSize); }
+
 	/** Register and create both state and input to be both networked and cached in history */
 	template<typename PhysicsTraits>
 	void CreateDataHistory(UActorComponent* HistoryComponent);
@@ -916,6 +1130,8 @@ private:
 	friend struct FNetworkPhysicsRewindDataStateProxy;
 	friend struct FNetworkPhysicsRewindDataImportantInputProxy;
 	friend struct FNetworkPhysicsRewindDataImportantStateProxy;
+	friend struct FNetworkPhysicsRewindDataDeltaSourceInputProxy;
+	friend struct FNetworkPhysicsRewindDataDeltaSourceStateProxy;
 
 	// Network Physics Component data internal to the physics thread
 	FAsyncNetworkPhysicsComponent* NetworkPhysicsComponent_Internal;
@@ -932,6 +1148,12 @@ private:
 	// Helper for the creation of state data and history with correct derived type
 	TUniquePtr<FNetworkPhysicsDataHelper> StateHelper;
 
+	// Local temporary input data
+	TUniquePtr<FNetworkPhysicsData> InputDataDefault;
+
+	// Local temporary state data
+	TUniquePtr<FNetworkPhysicsData> StateDataDefault;
+
 	// The number of inputs the owning client should send to the server with each RPC, replicated from the server. This is dynamically scaled based on when there are holes in the inputs buffer if np2.Resim.DynamicInputScaling.Enabled is enabled
 	UPROPERTY( Replicated )
 	uint16 InputsToNetwork_Owner = 3;
@@ -945,9 +1167,23 @@ private:
 	// Send last N number of states each replication call from server to remote clients
 	uint16 StatesToNetwork = 1;
 
+	// Array of delta sources, used as a base for delta serialization
+	TArray<TUniquePtr<FNetworkPhysicsData>> DeltaSourceInputs;
+	int32 LatestAcknowledgedDeltaSourceInputIndex = 0;
+	int32 LatestCachedDeltaSourceInputIndex = 0;
+	double TimeToSyncDeltaSourceInput = 0;
+	
+	// Array of delta sources, used as a base for delta serialization
+	TArray<TUniquePtr<FNetworkPhysicsData>> DeltaSourceStates;
+	int32 LatestAcknowledgedDeltaSourceStateIndex = 0;
+	int32 LatestCachedDeltaSourceStateIndex = 0;
+	double TimeToSyncDeltaSourceState = 0;
+
+public:
 	// Actor component that will be used to fill the histories
-	TObjectPtr<UActorComponent> ActorComponent;
+	TWeakObjectPtr<UActorComponent> ActorComponent;
 
+private:
 	// Root components physics object
 	Chaos::FPhysicsObjectHandle RootPhysicsObject;
 
@@ -1006,23 +1242,39 @@ FORCEINLINE void UNetworkPhysicsComponent::CreateDataHistory(UActorComponent* Hi
 {
 	InputHelper = MakeUnique<TNetworkPhysicsDataHelper<typename PhysicsTraits::InputsType>>();
 	StateHelper = MakeUnique<TNetworkPhysicsDataHelper<typename PhysicsTraits::StatesType>>();
+	
+	InputDataDefault = InputHelper->CreateUniqueData();
+	StateDataDefault = StateHelper->CreateUniqueData();
+
+	// Initialize delta source arrays
+	for (int32 I = 0; I < DeltaSourceBufferSize; I++)
+	{
+		DeltaSourceInputs.Add(InputHelper->CreateUniqueData());
+		DeltaSourceStates.Add(StateHelper->CreateUniqueData());
+	}
 
-	ReplicatedInputs.History = MakeUnique<TNetRewindHistory<typename PhysicsTraits::InputsType>>(InputsToNetwork_OwnerDefault);
+	ReplicatedInputs.History = InputHelper->CreateUniqueRewindHistory(InputsToNetwork_OwnerDefault);
 	ReplicatedInputs.Owner = this;
 
-	ReplicatedRemoteInputs.History = MakeUnique<TNetRewindHistory<typename PhysicsTraits::InputsType>>(InputsToNetwork_Simulated);
+	ReplicatedRemoteInputs.History = InputHelper->CreateUniqueRewindHistory(InputsToNetwork_Simulated);
 	ReplicatedRemoteInputs.Owner = this;
 
-	ReplicatedStates.History = MakeUnique<TNetRewindHistory<typename PhysicsTraits::StatesType>>(StatesToNetwork);
+	ReplicatedStates.History = StateHelper->CreateUniqueRewindHistory(StatesToNetwork);
 	ReplicatedStates.Owner = this;
 
-	ReplicatedImportantInput.History = MakeUnique<TNetRewindHistory<typename PhysicsTraits::InputsType>>(1);
+	ReplicatedImportantInput.History = InputHelper->CreateUniqueRewindHistory(1);
 	ReplicatedImportantInput.Owner = this;
 
-	ReplicatedImportantState.History = MakeUnique<TNetRewindHistory<typename PhysicsTraits::StatesType>>(1);
+	ReplicatedImportantState.History = StateHelper->CreateUniqueRewindHistory(1);
 	ReplicatedImportantState.Owner = this;
 
-	ActorComponent = HistoryComponent;
+	ReplicatedDeltaSourceInput.History = InputHelper->CreateUniqueRewindHistory(1);
+	ReplicatedDeltaSourceInput.Owner = this;
+
+	ReplicatedDeltaSourceState.History = StateHelper->CreateUniqueRewindHistory(1);
+	ReplicatedDeltaSourceState.Owner = this;
+
+	ActorComponent = TWeakObjectPtr<UActorComponent>(HistoryComponent);
 
 	CreateAsyncDataHistory();
 }
@@ -1031,17 +1283,28 @@ template<class InputsType>
 FORCEINLINE void UNetworkPhysicsComponent::CreateInputHistory(UActorComponent* HistoryComponent)
 {
 	InputHelper = MakeUnique<TNetworkPhysicsDataHelper<InputsType>>();
+	
+	InputDataDefault = InputHelper->CreateUniqueData();
 
-	ReplicatedInputs.History = MakeUnique<TNetRewindHistory<InputsType>>(InputsToNetwork_OwnerDefault);
+	// Initialize delta source array
+	for (int32 I = 0; I < DeltaSourceBufferSize; I++)
+	{
+		DeltaSourceInputs.Add(InputHelper->CreateUniqueData());
+	}
+
+	ReplicatedInputs.History = InputHelper->CreateUniqueRewindHistory(InputsToNetwork_OwnerDefault);
 	ReplicatedInputs.Owner = this;
 
-	ReplicatedRemoteInputs.History = MakeUnique<TNetRewindHistory<InputsType>>(InputsToNetwork_Simulated);
+	ReplicatedRemoteInputs.History = InputHelper->CreateUniqueRewindHistory(InputsToNetwork_Simulated);
 	ReplicatedRemoteInputs.Owner = this;
 
-	ReplicatedImportantInput.History = MakeUnique<TNetRewindHistory<InputsType>>(1);
+	ReplicatedImportantInput.History = InputHelper->CreateUniqueRewindHistory(1);
 	ReplicatedImportantInput.Owner = this;
 
-	ActorComponent = HistoryComponent;
+	ReplicatedDeltaSourceInput.History = InputHelper->CreateUniqueRewindHistory(1);
+	ReplicatedDeltaSourceInput.Owner = this;
+
+	ActorComponent = TWeakObjectPtr<UActorComponent>(HistoryComponent);
 
 	CreateAsyncDataHistory();
 }
@@ -1057,7 +1320,7 @@ struct FAsyncNetworkPhysicsComponentInput : public Chaos::FSimCallbackInput
 	TOptional<int32> NetworkPhysicsTickOffset;
 	TOptional<uint16> InputsToNetwork_Owner;
 	TOptional<EPhysicsReplicationMode> PhysicsReplicationMode;
-	TOptional<TObjectPtr<UActorComponent>> ActorComponent;
+	TOptional<TWeakObjectPtr<UActorComponent>> ActorComponent;
 	TOptional<FString> ActorName;
 	TOptional<TUniquePtr<FNetworkPhysicsDataHelper>> InputHelper;
 	TOptional<TUniquePtr<FNetworkPhysicsDataHelper>> StateHelper;
diff --git a/Engine/Source/Runtime/Experimental/Chaos/Public/RewindData.h b/Engine/Source/Runtime/Experimental/Chaos/Public/RewindData.h
--- a/Engine/Source/Runtime/Experimental/Chaos/Public/RewindData.h
+++ b/Engine/Source/Runtime/Experimental/Chaos/Public/RewindData.h
@@ -102,7 +102,11 @@ struct FBaseRewindHistory
 	virtual int32 ReceiveNewData(FBaseRewindHistory& NewData, const int32 FrameOffset, const bool CompareDataForRewind = false, const bool bImportant = false) { return INDEX_NONE; }
 
 	/** Serialize the data to or from a network archive */
-	virtual void NetSerialize(FArchive& Ar, UPackageMap* PackageMap) {}
+	UE_DEPRECATED(5.6, "Deprecated, use the NetSerialize call with parameter that takes a function DataSetupFunction, pass in nullptr to opt out of implementing a function.")
+	virtual void NetSerialize(FArchive& Ar, UPackageMap* PackageMap) { NetSerialize(Ar, PackageMap, [](void* Data, const int32 DataIndex){}); }
+
+	/** Serialize the data to or from a network archive */
+	virtual void NetSerialize(FArchive& Ar, UPackageMap* PackageMap, TUniqueFunction<void(void* Data, const int32 DataIndex)> DataSetupFunction) {}
 	
 	/** Validate data in history buffer received from clients on the server */
 	virtual void ValidateDataInHistory(const void* ActorComponent) {}
