[NetPhysics] Fix where physics replication would not iterate TickCount correctly resulting in incorrect target alignment.

diff --git a/Engine/Source/Runtime/Engine/Private/PhysicsEngine/PhysicsReplication.cpp b/Engine/Source/Runtime/Engine/Private/PhysicsEngine/PhysicsReplication.cpp
--- a/Engine/Source/Runtime/Engine/Private/PhysicsEngine/PhysicsReplication.cpp
+++ b/Engine/Source/Runtime/Engine/Private/PhysicsEngine/PhysicsReplication.cpp
@@ -901,34 +901,33 @@ void FPhysicsReplicationAsync::OnPreSimulate_Internal()
 		return;
 	}
 
-	if (const FPhysicsReplicationAsyncInput* AsyncInput = GetConsumerInput_Internal())
-	{
-		Chaos::FPBDRigidsSolver* RigidsSolver = static_cast<Chaos::FPBDRigidsSolver*>(GetSolver());
-		check(RigidsSolver);
+	Chaos::FPBDRigidsSolver* RigidsSolver = static_cast<Chaos::FPBDRigidsSolver*>(GetSolver());
+	check(RigidsSolver);
 
-		// Early out if this is a resim frame
-		Chaos::FRewindData* RewindData = RigidsSolver->GetRewindData();
-		const bool bRewindDataExist = RewindData != nullptr;
-		if (bRewindDataExist && RewindData->IsResim())
+	// Early out if this is a resim frame
+	Chaos::FRewindData* RewindData = RigidsSolver->GetRewindData();
+	const bool bRewindDataExist = RewindData != nullptr;
+	if (bRewindDataExist && RewindData->IsResim())
+	{
+		// TODO, Handle the transition from post-resim to interpolation better (disabled by default, resim vs replication interaction is handled via FPhysicsReplicationAsync::CacheResimInteractions)
+		if (SettingsCurrent.PredictiveInterpolationSettings.GetPostResimWaitForUpdate() && RewindData->IsFinalResim())
 		{
-			// TODO, Handle the transition from post-resim to interpolation better (disabled by default, resim vs replication interaction is handled via FPhysicsReplicationAsync::CacheResimInteractions)
-			if (SettingsCurrent.PredictiveInterpolationSettings.GetPostResimWaitForUpdate() && RewindData->IsFinalResim())
+			for (auto Itr = ObjectToTarget.CreateIterator(); Itr; ++Itr)
 			{
-				for (auto Itr = ObjectToTarget.CreateIterator(); Itr; ++Itr)
-				{
-					FReplicatedPhysicsTargetAsync& Target = Itr.Value();
+				FReplicatedPhysicsTargetAsync& Target = Itr.Value();
 
-					// If final resim frame, mark interpolated targets as waiting for up to date data from the server.
-					if (Target.RepMode == EPhysicsReplicationMode::PredictiveInterpolation)
-					{
-						Target.SetWaiting(RigidsSolver->GetCurrentFrame() + Target.FrameOffset, Target.RepModeOverride);
-					}
+				// If final resim frame, mark interpolated targets as waiting for up to date data from the server.
+				if (Target.RepMode == EPhysicsReplicationMode::PredictiveInterpolation)
+				{
+					Target.SetWaiting(RigidsSolver->GetCurrentFrame() + Target.FrameOffset, Target.RepModeOverride);
 				}
 			}
-
-			return;
 		}
+		return;
+	}
 
+	if (const FPhysicsReplicationAsyncInput* AsyncInput = GetConsumerInput_Internal())
+	{
 		// Update async targets with target input
 		for (const FPhysicsRepAsyncInputData& Input : AsyncInput->InputData)
 		{
@@ -950,15 +949,25 @@ void FPhysicsReplicationAsync::OnPreSimulate_Internal()
 
 			UpdateRewindDataTarget(Input);
 			UpdateAsyncTarget(Input, RigidsSolver);
-		}
+			
+			// Deprecated, legacy BodyInstance flow for Default Replication
+			if (Input.Proxy != nullptr)
+			{
+				Chaos::FSingleParticlePhysicsProxy* Proxy = Input.Proxy;
+				Chaos::FRigidBodyHandle_Internal* Handle = Proxy->GetPhysicsThreadAPI();
 
-		if (Chaos::FPBDRigidsSolver::IsNetworkPhysicsPredictionEnabled())
-		{
-			CacheResimInteractions();
+				const FPhysicsRepErrorCorrectionData& UsedErrorCorrection = Input.ErrorCorrection.IsSet() ? Input.ErrorCorrection.GetValue() : AsyncInput->ErrorCorrection;
+				DefaultReplication_DEPRECATED(Handle, Input, GetDeltaTime_Internal(), UsedErrorCorrection);
+			}
 		}
+	}
 
-		ApplyTargetStatesAsync(GetDeltaTime_Internal(), AsyncInput->ErrorCorrection, AsyncInput->InputData);
+	if (Chaos::FPBDRigidsSolver::IsNetworkPhysicsPredictionEnabled())
+	{
+		CacheResimInteractions();
 	}
+
+	ApplyTargetStatesAsync(GetDeltaTime_Internal());
 }
 
 FReplicatedPhysicsTargetAsync* FPhysicsReplicationAsync::AddObjectToReplication(Chaos::FConstPhysicsObjectHandle PhysicsObject)
@@ -1253,23 +1262,10 @@ void FPhysicsReplicationAsync::CacheResimInteractions()
 	}
 }
 
-void FPhysicsReplicationAsync::ApplyTargetStatesAsync(const float DeltaSeconds, const FPhysicsRepErrorCorrectionData& ErrorCorrection, const TArray<FPhysicsRepAsyncInputData>& InputData)
+void FPhysicsReplicationAsync::ApplyTargetStatesAsync(const float DeltaSeconds)
 {
 	using namespace Chaos;
 
-	// Deprecated, legacy BodyInstance flow
-	for (const FPhysicsRepAsyncInputData& Input : InputData)
-	{
-		if (Input.Proxy != nullptr)
-		{
-			Chaos::FSingleParticlePhysicsProxy* Proxy = Input.Proxy;
-			Chaos::FRigidBodyHandle_Internal* Handle = Proxy->GetPhysicsThreadAPI();
-
-			const FPhysicsRepErrorCorrectionData& UsedErrorCorrection = Input.ErrorCorrection.IsSet() ? Input.ErrorCorrection.GetValue() : ErrorCorrection;
-			DefaultReplication_DEPRECATED(Handle, Input, DeltaSeconds, UsedErrorCorrection);
-		}
-	}
-
 	/** Helper function to remove replicated target*/
 	auto RemoveTargetHelper = [this](TMap<Chaos::FConstPhysicsObjectHandle, FReplicatedPhysicsTargetAsync>::TIterator Itr, FGeometryParticleHandle* Handle)
 	{
diff --git a/Engine/Source/Runtime/Engine/Public/PhysicsReplication.h b/Engine/Source/Runtime/Engine/Public/PhysicsReplication.h
--- a/Engine/Source/Runtime/Engine/Public/PhysicsReplication.h
+++ b/Engine/Source/Runtime/Engine/Public/PhysicsReplication.h
@@ -196,7 +196,9 @@ class FPhysicsReplicationAsync : public IPhysicsReplicationAsync,
 	virtual void OnPreSimulate_Internal() override;
 	virtual void OnPhysicsObjectUnregistered_Internal(Chaos::FConstPhysicsObjectHandle PhysicsObject) override;
 
-	virtual void ApplyTargetStatesAsync(const float DeltaSeconds, const FPhysicsRepErrorCorrectionData& ErrorCorrection, const TArray<FPhysicsRepAsyncInputData>& TargetStates);
+	virtual void ApplyTargetStatesAsync(const float DeltaSeconds);
+	UE_DEPRECATED(5.6, "Deprecated, call the function with just @param DeltaSeconds instead.")
+	virtual void ApplyTargetStatesAsync(const float DeltaSeconds, const FPhysicsRepErrorCorrectionData& ErrorCorrection, const TArray<FPhysicsRepAsyncInputData>& TargetStates) { ApplyTargetStatesAsync(DeltaSeconds); };
 
 	// Replication functions
 	virtual void DefaultReplication_DEPRECATED(Chaos::FRigidBodyHandle_Internal* Handle, const FPhysicsRepAsyncInputData& State, const float DeltaSeconds, const FPhysicsRepErrorCorrectionData& ErrorCorrection);
