[AutoRTFM] Wrap ISPC calls to allow for transactional safety.

Verse script can trigger animation playback, so this code can be hit from within a transaction and needs to allow for rollback.

diff --git a/Engine/Source/Runtime/Engine/Private/Animation/AnimEncoding_ConstantKeyLerp.cpp b/Engine/Source/Runtime/Engine/Private/Animation/AnimEncoding_ConstantKeyLerp.cpp
--- a/Engine/Source/Runtime/Engine/Private/Animation/AnimEncoding_ConstantKeyLerp.cpp
+++ b/Engine/Source/Runtime/Engine/Private/Animation/AnimEncoding_ConstantKeyLerp.cpp
@@ -7,6 +7,7 @@
 #include "AnimEncoding_ConstantKeyLerp.h"
 #if INTEL_ISPC
 #include "AnimEncoding_ConstantKeyLerp.ispc.generated.h"
+#include "AutoRTFM/AutoRTFM.h"
 
 static_assert(sizeof(ispc::FTransform) == sizeof(FTransform), "sizeof(ispc::FTransform) != sizeof(FTransform)");
 static_assert(sizeof(ispc::BoneTrackPair) == sizeof(BoneTrackPair), "sizeof(ispc::BoneTrackPair) != sizeof(BoneTrackPair)");
@@ -277,18 +278,24 @@ inline void AEFConstantKeyLerp<FORMAT>::GetPoseRotations(
 	if (bAnim_ConstantKeyLerp_ISPC_Enabled)
 	{
 #if INTEL_ISPC
-		const FUECompressedAnimData& AnimData = static_cast<const FUECompressedAnimData&>(DecompContext.CompressedAnimData);
-
-		ispc::GetConstantKeyLerpPoseRotations(
-			(ispc::FTransform*)&Atoms[0],
-			(ispc::BoneTrackPair*)&DesiredPairs[0],
-			AnimData.CompressedTrackOffsets.GetData(),
-			AnimData.CompressedByteStream.GetData(),
-			DecompContext.GetPlayableLength(),
-			DecompContext.GetRelativePosition(),
-			(uint8)DecompContext.Interpolation,
-			FORMAT,
-			PairCount);
+		UE_AUTORTFM_OPEN_NO_VALIDATION // Open writes are recorded
+		{
+			const FUECompressedAnimData& AnimData = static_cast<const FUECompressedAnimData&>(DecompContext.CompressedAnimData);
+
+			AutoRTFM::RecordOpenWrite(
+				Atoms.GetData(),
+				Atoms.NumBytes());
+			ispc::GetConstantKeyLerpPoseRotations(
+				reinterpret_cast<ispc::FTransform*>(Atoms.GetData()),
+				reinterpret_cast<const ispc::BoneTrackPair*>(DesiredPairs.GetData()),
+				AnimData.CompressedTrackOffsets.GetData(),
+				AnimData.CompressedByteStream.GetData(),
+				DecompContext.GetPlayableLength(),
+				DecompContext.GetRelativePosition(),
+				(uint8)DecompContext.Interpolation,
+				FORMAT,
+				PairCount);
+		};
 #endif
 	}
 	else
@@ -329,18 +336,24 @@ inline void AEFConstantKeyLerp<FORMAT>::GetPoseTranslations(
 	if (bAnim_ConstantKeyLerp_ISPC_Enabled)
 	{
 #if INTEL_ISPC
-		const FUECompressedAnimData& AnimData = static_cast<const FUECompressedAnimData&>(DecompContext.CompressedAnimData);
-
-		ispc::GetConstantKeyLerpPoseTranslations(
-			(ispc::FTransform*)&Atoms[0],
-			(ispc::BoneTrackPair*)&DesiredPairs[0],
-			AnimData.CompressedTrackOffsets.GetData(),
-			AnimData.CompressedByteStream.GetData(),
-			DecompContext.GetPlayableLength(),
-			DecompContext.GetRelativePosition(),
-			(uint8)DecompContext.Interpolation,
-			FORMAT,
-			PairCount);
+		UE_AUTORTFM_OPEN_NO_VALIDATION // Open writes are recorded
+		{
+			const FUECompressedAnimData& AnimData = static_cast<const FUECompressedAnimData&>(DecompContext.CompressedAnimData);
+
+			AutoRTFM::RecordOpenWrite(
+				Atoms.GetData(),
+				Atoms.NumBytes());
+			ispc::GetConstantKeyLerpPoseTranslations(
+				reinterpret_cast<ispc::FTransform*>(Atoms.GetData()),
+				reinterpret_cast<const ispc::BoneTrackPair*>(DesiredPairs.GetData()),
+				AnimData.CompressedTrackOffsets.GetData(),
+				AnimData.CompressedByteStream.GetData(),
+				DecompContext.GetPlayableLength(),
+				DecompContext.GetRelativePosition(),
+				(uint8)DecompContext.Interpolation,
+				FORMAT,
+				PairCount);
+		};
 #endif
 	}
 	else
@@ -387,7 +400,7 @@ inline void AEFConstantKeyLerp<FORMAT>::GetPoseScales(
 	const BoneTrackArray& DesiredPairs,
 	FAnimSequenceDecompressionContext& DecompContext)
 {
-	const int32 PairCount= DesiredPairs.Num();
+	const int32 PairCount = DesiredPairs.Num();
 
 	if (PairCount == 0)
 	{
@@ -397,22 +410,28 @@ inline void AEFConstantKeyLerp<FORMAT>::GetPoseScales(
 	if (bAnim_ConstantKeyLerp_ISPC_Enabled)
 	{
 #if INTEL_ISPC
-		const FUECompressedAnimData& AnimData = static_cast<const FUECompressedAnimData&>(DecompContext.CompressedAnimData);
-
-		const TArrayView<int32> ScaleOffsets = AnimData.CompressedScaleOffsets.OffsetData;
-		const int32 StripSize = AnimData.CompressedScaleOffsets.StripSize;
-
-		ispc::GetConstantKeyLerpPoseScales(
-			(ispc::FTransform*)&Atoms[0],
-			(ispc::BoneTrackPair*)&DesiredPairs[0],
-			ScaleOffsets.GetData(),
-			StripSize,
-			AnimData.CompressedByteStream.GetData(),
-			DecompContext.GetPlayableLength(),
-			DecompContext.GetRelativePosition(),
-			(uint8)DecompContext.Interpolation,
-			FORMAT,
-			PairCount);
+		UE_AUTORTFM_OPEN_NO_VALIDATION // Open writes are recorded
+		{
+			const FUECompressedAnimData& AnimData = static_cast<const FUECompressedAnimData&>(DecompContext.CompressedAnimData);
+
+			const TArrayView<int32> ScaleOffsets = AnimData.CompressedScaleOffsets.OffsetData;
+			const int32 StripSize = AnimData.CompressedScaleOffsets.StripSize;
+
+			AutoRTFM::RecordOpenWrite(
+				Atoms.GetData(),
+				Atoms.NumBytes());
+			ispc::GetConstantKeyLerpPoseScales(
+				reinterpret_cast<ispc::FTransform*>(Atoms.GetData()),
+				reinterpret_cast<const ispc::BoneTrackPair*>(DesiredPairs.GetData()),
+				ScaleOffsets.GetData(),
+				StripSize,
+				AnimData.CompressedByteStream.GetData(),
+				DecompContext.GetPlayableLength(),
+				DecompContext.GetRelativePosition(),
+				(uint8)DecompContext.Interpolation,
+				FORMAT,
+				PairCount);
+		};
 #endif
 	}
 	else
diff --git a/Engine/Source/Runtime/Engine/Private/Animation/AnimEncoding_PerTrackCompression.cpp b/Engine/Source/Runtime/Engine/Private/Animation/AnimEncoding_PerTrackCompression.cpp
--- a/Engine/Source/Runtime/Engine/Private/Animation/AnimEncoding_PerTrackCompression.cpp
+++ b/Engine/Source/Runtime/Engine/Private/Animation/AnimEncoding_PerTrackCompression.cpp
@@ -10,6 +10,7 @@
 #include "Animation/AnimSequenceDecompressionContext.h"
 #if INTEL_ISPC
 #include "AnimEncoding_PerTrackCompression.ispc.generated.h"
+#include "AutoRTFM/AutoRTFM.h"
 
 static_assert(sizeof(ispc::FTransform) == sizeof(FTransform), "sizeof(ispc::FTransform) != sizeof(FTransform)");
 static_assert(sizeof(ispc::BoneTrackPair) == sizeof(BoneTrackPair), "sizeof(ispc::BoneTrackPair) != sizeof(BoneTrackPair)");
@@ -750,18 +751,24 @@ void AEFPerTrackCompressionCodec::GetPoseRotations(
 	if (bAnim_PerTrackCompression_ISPC_Enabled)
 	{
 #if INTEL_ISPC
-		const FUECompressedAnimData& AnimData = static_cast<const FUECompressedAnimData&>(DecompContext.CompressedAnimData);
-
-		ispc::GetPerTrackCompressionPoseRotations(
-			(ispc::FTransform*)&Atoms[0],
-			(ispc::BoneTrackPair*)&DesiredPairs[0],
-			AnimData.CompressedTrackOffsets.GetData(),
-			AnimData.CompressedByteStream.GetData(),
-			AnimData.CompressedNumberOfKeys,
-			DecompContext.GetPlayableLength(),
-			DecompContext.GetRelativePosition(),
-			(uint8)DecompContext.Interpolation,
-			PairCount);
+		UE_AUTORTFM_OPEN_NO_VALIDATION // Open writes are recorded
+		{
+			const FUECompressedAnimData& AnimData = static_cast<const FUECompressedAnimData&>(DecompContext.CompressedAnimData);
+
+			AutoRTFM::RecordOpenWrite(
+				Atoms.GetData(),
+				Atoms.NumBytes());
+			ispc::GetPerTrackCompressionPoseRotations(
+				reinterpret_cast<ispc::FTransform*>(Atoms.GetData()),
+				reinterpret_cast<const ispc::BoneTrackPair*>(DesiredPairs.GetData()),
+				AnimData.CompressedTrackOffsets.GetData(),
+				AnimData.CompressedByteStream.GetData(),
+				AnimData.CompressedNumberOfKeys,
+				DecompContext.GetPlayableLength(),
+				DecompContext.GetRelativePosition(),
+				(uint8)DecompContext.Interpolation,
+				PairCount);
+		};
 #endif
 	}
 	else
@@ -798,18 +805,24 @@ void AEFPerTrackCompressionCodec::GetPoseTranslations(
 	if (bAnim_PerTrackCompression_ISPC_Enabled)
 	{
 #if INTEL_ISPC
-		const FUECompressedAnimData& AnimData = static_cast<const FUECompressedAnimData&>(DecompContext.CompressedAnimData);
-
-		ispc::GetPerTrackCompressionPoseTranslations(
-			(ispc::FTransform*)&Atoms[0],
-			(ispc::BoneTrackPair*)&DesiredPairs[0],
-			AnimData.CompressedTrackOffsets.GetData(),
-			AnimData.CompressedByteStream.GetData(),
-			AnimData.CompressedNumberOfKeys,
-			DecompContext.GetPlayableLength(),
-			DecompContext.GetRelativePosition(),
-			(uint8)DecompContext.Interpolation,
-			PairCount);
+		UE_AUTORTFM_OPEN_NO_VALIDATION // Open writes are recorded
+		{
+			const FUECompressedAnimData& AnimData = static_cast<const FUECompressedAnimData&>(DecompContext.CompressedAnimData);
+
+			AutoRTFM::RecordOpenWrite(
+				Atoms.GetData(),
+				Atoms.NumBytes());
+			ispc::GetPerTrackCompressionPoseTranslations(
+				reinterpret_cast<ispc::FTransform*>(Atoms.GetData()),
+				reinterpret_cast<const ispc::BoneTrackPair*>(DesiredPairs.GetData()),
+				AnimData.CompressedTrackOffsets.GetData(),
+				AnimData.CompressedByteStream.GetData(),
+				AnimData.CompressedNumberOfKeys,
+				DecompContext.GetPlayableLength(),
+				DecompContext.GetRelativePosition(),
+				(uint8)DecompContext.Interpolation,
+				PairCount);
+		};
 #endif
 	}
 	else
@@ -848,22 +861,28 @@ void AEFPerTrackCompressionCodec::GetPoseScales(
 	if (bAnim_PerTrackCompression_ISPC_Enabled)
 	{
 #if INTEL_ISPC
-		const FUECompressedAnimData& AnimData = static_cast<const FUECompressedAnimData&>(DecompContext.CompressedAnimData);
-
-		const TArrayView<int32> ScaleOffsets = AnimData.CompressedScaleOffsets.OffsetData;
-		const int32 StripSize = AnimData.CompressedScaleOffsets.StripSize;
-
-		ispc::GetPerTrackCompressionPoseScales(
-			(ispc::FTransform*)&Atoms[0],
-			(ispc::BoneTrackPair*)&DesiredPairs[0],
-			ScaleOffsets.GetData(),
-			StripSize,
-			AnimData.CompressedByteStream.GetData(),
-			AnimData.CompressedNumberOfKeys,
-			DecompContext.GetPlayableLength(),
-			DecompContext.GetRelativePosition(),
-			(uint8)DecompContext.Interpolation,
-			PairCount);
+		UE_AUTORTFM_OPEN_NO_VALIDATION // Open writes are recorded
+		{
+			const FUECompressedAnimData& AnimData = static_cast<const FUECompressedAnimData&>(DecompContext.CompressedAnimData);
+
+			const TArrayView<int32> ScaleOffsets = AnimData.CompressedScaleOffsets.OffsetData;
+			const int32 StripSize = AnimData.CompressedScaleOffsets.StripSize;
+
+			AutoRTFM::RecordOpenWrite(
+				Atoms.GetData(),
+				Atoms.NumBytes());
+			ispc::GetPerTrackCompressionPoseScales(
+				reinterpret_cast<ispc::FTransform*>(Atoms.GetData()),
+				reinterpret_cast<const ispc::BoneTrackPair*>(DesiredPairs.GetData()),
+				ScaleOffsets.GetData(),
+				StripSize,
+				AnimData.CompressedByteStream.GetData(),
+				AnimData.CompressedNumberOfKeys,
+				DecompContext.GetPlayableLength(),
+				DecompContext.GetRelativePosition(),
+				(uint8)DecompContext.Interpolation,
+				PairCount);
+		};
 #endif
 	}
 	else
diff --git a/Engine/Source/Runtime/Engine/Private/Animation/AnimEncoding_VariableKeyLerp.cpp b/Engine/Source/Runtime/Engine/Private/Animation/AnimEncoding_VariableKeyLerp.cpp
--- a/Engine/Source/Runtime/Engine/Private/Animation/AnimEncoding_VariableKeyLerp.cpp
+++ b/Engine/Source/Runtime/Engine/Private/Animation/AnimEncoding_VariableKeyLerp.cpp
@@ -7,6 +7,7 @@
 #include "AnimEncoding_VariableKeyLerp.h"
 #if INTEL_ISPC
 #include "AnimEncoding_VariableKeyLerp.ispc.generated.h"
+#include "AutoRTFM/AutoRTFM.h"
 
 static_assert(sizeof(ispc::FTransform) == sizeof(FTransform), "sizeof(ispc::FTransform) != sizeof(FTransform)");
 static_assert(sizeof(ispc::BoneTrackPair) == sizeof(BoneTrackPair), "sizeof(ispc::BoneTrackPair) != sizeof(BoneTrackPair)");
@@ -242,18 +243,24 @@ void AEFVariableKeyLerp<FORMAT>::GetPoseRotations(
 	if (bAnim_VariableKeyLerp_ISPC_Enabled)
 	{
 #if INTEL_ISPC
-		const FUECompressedAnimData& AnimData = static_cast<const FUECompressedAnimData&>(DecompContext.CompressedAnimData);
-
-		ispc::GetVariableKeyLerpPoseRotations(
-			(ispc::FTransform*)&Atoms[0],
-			(ispc::BoneTrackPair*)&DesiredPairs[0],
-			AnimData.CompressedTrackOffsets.GetData(),
-			AnimData.CompressedByteStream.GetData(),
-			AnimData.CompressedNumberOfKeys,
-			DecompContext.GetRelativePosition(),
-			(uint8)DecompContext.Interpolation,
-			FORMAT,
-			PairCount);
+		UE_AUTORTFM_OPEN_NO_VALIDATION // Open writes are recorded
+		{
+			const FUECompressedAnimData& AnimData = static_cast<const FUECompressedAnimData&>(DecompContext.CompressedAnimData);
+
+			AutoRTFM::RecordOpenWrite(
+				Atoms.GetData(),
+				Atoms.NumBytes());
+			ispc::GetVariableKeyLerpPoseRotations(
+				reinterpret_cast<ispc::FTransform*>(Atoms.GetData()),
+				reinterpret_cast<const ispc::BoneTrackPair*>(DesiredPairs.GetData()),
+				AnimData.CompressedTrackOffsets.GetData(),
+				AnimData.CompressedByteStream.GetData(),
+				AnimData.CompressedNumberOfKeys,
+				DecompContext.GetRelativePosition(),
+				(uint8)DecompContext.Interpolation,
+				FORMAT,
+				PairCount);
+		};
 #endif
 	}
 	else
@@ -294,18 +301,24 @@ void AEFVariableKeyLerp<FORMAT>::GetPoseTranslations(
 	if (bAnim_VariableKeyLerp_ISPC_Enabled)
 	{
 #if INTEL_ISPC
-		const FUECompressedAnimData& AnimData = static_cast<const FUECompressedAnimData&>(DecompContext.CompressedAnimData);
-
-		ispc::GetVariableKeyLerpPoseTranslations(
-			(ispc::FTransform*)&Atoms[0],
-			(ispc::BoneTrackPair*)&DesiredPairs[0],
-			AnimData.CompressedTrackOffsets.GetData(),
-			AnimData.CompressedByteStream.GetData(),
-			AnimData.CompressedNumberOfKeys,
-			DecompContext.GetRelativePosition(),
-			(uint8)DecompContext.Interpolation,
-			FORMAT,
-			PairCount);
+		UE_AUTORTFM_OPEN_NO_VALIDATION // Open writes are recorded
+		{
+			const FUECompressedAnimData& AnimData = static_cast<const FUECompressedAnimData&>(DecompContext.CompressedAnimData);
+
+			AutoRTFM::RecordOpenWrite(
+				Atoms.GetData(),
+				Atoms.NumBytes());
+			ispc::GetVariableKeyLerpPoseTranslations(
+				reinterpret_cast<ispc::FTransform*>(Atoms.GetData()),
+				reinterpret_cast<const ispc::BoneTrackPair*>(DesiredPairs.GetData()),
+				AnimData.CompressedTrackOffsets.GetData(),
+				AnimData.CompressedByteStream.GetData(),
+				AnimData.CompressedNumberOfKeys,
+				DecompContext.GetRelativePosition(),
+				(uint8)DecompContext.Interpolation,
+				FORMAT,
+				PairCount);
+		};
 #endif
 	}
 	else
@@ -346,22 +359,28 @@ void AEFVariableKeyLerp<FORMAT>::GetPoseScales(
 	if (bAnim_VariableKeyLerp_ISPC_Enabled)
 	{
 #if INTEL_ISPC
-		const FUECompressedAnimData& AnimData = static_cast<const FUECompressedAnimData&>(DecompContext.CompressedAnimData);
-
-		const TArrayView<int32> ScaleOffsets = AnimData.CompressedScaleOffsets.OffsetData;
-		const int32 StripSize = AnimData.CompressedScaleOffsets.StripSize;
-
-		ispc::GetVariableKeyLerpPoseScales(
-			(ispc::FTransform*)&Atoms[0],
-			(ispc::BoneTrackPair*)&DesiredPairs[0],
-			ScaleOffsets.GetData(),
-			StripSize,
-			AnimData.CompressedByteStream.GetData(),
-			AnimData.CompressedNumberOfKeys,
-			DecompContext.GetRelativePosition(),
-			(uint8)DecompContext.Interpolation,
-			FORMAT,
-			PairCount);
+		UE_AUTORTFM_OPEN_NO_VALIDATION // Open writes are recorded
+		{
+			const FUECompressedAnimData& AnimData = static_cast<const FUECompressedAnimData&>(DecompContext.CompressedAnimData);
+
+			const TArrayView<int32> ScaleOffsets = AnimData.CompressedScaleOffsets.OffsetData;
+			const int32 StripSize = AnimData.CompressedScaleOffsets.StripSize;
+
+			AutoRTFM::RecordOpenWrite(
+				Atoms.GetData(),
+				Atoms.NumBytes());
+			ispc::GetVariableKeyLerpPoseScales(
+				reinterpret_cast<ispc::FTransform*>(Atoms.GetData()),
+				reinterpret_cast<const ispc::BoneTrackPair*>(DesiredPairs.GetData()),
+				ScaleOffsets.GetData(),
+				StripSize,
+				AnimData.CompressedByteStream.GetData(),
+				AnimData.CompressedNumberOfKeys,
+				DecompContext.GetRelativePosition(),
+				(uint8)DecompContext.Interpolation,
+				FORMAT,
+				PairCount);
+		};
 #endif
 	}
 	else
