Apply Compression on existing Anim Sequences hangs the editor
#fix use lockless api to prevent recursive RWlock entry

diff --git a/Engine/Source/Runtime/Engine/Private/Animation/AnimationCompressionDerivedData.cpp b/Engine/Source/Runtime/Engine/Private/Animation/AnimationCompressionDerivedData.cpp
--- a/Engine/Source/Runtime/Engine/Private/Animation/AnimationCompressionDerivedData.cpp
+++ b/Engine/Source/Runtime/Engine/Private/Animation/AnimationCompressionDerivedData.cpp
@@ -307,7 +307,7 @@ void FAnimationSequenceAsyncCacheTask::CalculateRequiredMemoryEstimate()
 		RequiredMemory = AnimSequence->GetApproxRawSize() + AdditiveAnimSize;
 		if (Compression::FAnimationCompressionMemorySummaryScope::ShouldStoreCompressionResults())
 		{
-			Compression::FAnimationCompressionMemorySummaryScope::CompressionResultSummary().GatherPreCompressionStats(AnimSequence->GetApproxRawSize(), AnimSequence->GetApproxCompressedSize());
+			Compression::FAnimationCompressionMemorySummaryScope::CompressionResultSummary().GatherPreCompressionStats(AnimSequence->GetApproxRawSize(), AnimSequence->GetApproxCompressedSize_Lockless());
 		}
 
 		if (const UAnimBoneCompressionSettings* BoneCompressionSettings = AnimSequence->BoneCompressionSettings.Get())
