Added support for Ignore Objects with Remote Physics Queries.

diff --git a/Engine/Source/Runtime/CoreUObject/Public/UObject/RemoteExecutor.h b/Engine/Source/Runtime/CoreUObject/Public/UObject/RemoteExecutor.h
--- a/Engine/Source/Runtime/CoreUObject/Public/UObject/RemoteExecutor.h
+++ b/Engine/Source/Runtime/CoreUObject/Public/UObject/RemoteExecutor.h
@@ -23,8 +23,8 @@ struct FRemoteTransactionId
 		return FRemoteTransactionId();
 	}
 
-	bool operator==(const FRemoteTransactionId& Rhs) const = default;
-	bool operator!=(const FRemoteTransactionId& Rhs) const = default;
+	bool operator==(const FRemoteTransactionId& Rhs) const { return Id == Rhs.Id; }
+	bool operator!=(const FRemoteTransactionId& Rhs) const { return Id != Rhs.Id; }
 
 	bool IsValid() const
 	{
diff --git a/Engine/Source/Runtime/CoreUObject/Public/UObject/RemoteObjectTransfer.h b/Engine/Source/Runtime/CoreUObject/Public/UObject/RemoteObjectTransfer.h
--- a/Engine/Source/Runtime/CoreUObject/Public/UObject/RemoteObjectTransfer.h
+++ b/Engine/Source/Runtime/CoreUObject/Public/UObject/RemoteObjectTransfer.h
@@ -193,6 +193,11 @@ public:
 	{
 	}
 
+	FORCEINLINE bool operator==(const FRemoteObjectReference& Other) const
+	{
+		return ObjectId == Other.ObjectId;
+	}
+
 	COREUOBJECT_API FObjectPtr ToObjectPtr() const;
 	COREUOBJECT_API FWeakObjectPtr ToWeakPtr() const;
 	COREUOBJECT_API UObject* Resolve() const;
diff --git a/Engine/Source/Runtime/Engine/Private/Collision/Collision.cpp b/Engine/Source/Runtime/Engine/Private/Collision/Collision.cpp
--- a/Engine/Source/Runtime/Engine/Private/Collision/Collision.cpp
+++ b/Engine/Source/Runtime/Engine/Private/Collision/Collision.cpp
@@ -104,6 +104,25 @@ static FORCEINLINE_DEBUGGABLE bool CheckForCollision(const UPrimitiveComponent*
 	return bCheckForCollision;
 }
 
+void FCollisionQueryParams::Internal_AddIgnoredSourceObject(const FWeakObjectPtr& WeakPtr)
+{
+	if (WeakPtr.IsValid())
+	{
+		if (!WeakPtr.IsRemote())
+		{
+			UObject* Obj = WeakPtr.Get();
+			// On rare occasions, this is null even with the above `IsValid` check
+			if (Obj != nullptr)
+			{
+				AddIgnoredSourceObject(WeakPtr.Get()->GetUniqueID());
+			}
+		}
+
+#if UE_WITH_REMOTE_OBJECT_HANDLE
+		IgnoreSourceObjectReferences.Add(FRemoteObjectReference(WeakPtr));
+#endif
+	}
+}
 
 void FCollisionQueryParams::AddIgnoredActor(const AActor* InIgnoreActor)
 {
@@ -135,21 +154,39 @@ void FCollisionQueryParams::AddIgnoredActors(const TArray<TWeakObjectPtr<const A
 {
 	for (int32 Idx = 0; Idx < InIgnoreActors.Num(); ++Idx)
 	{
-		AddIgnoredActor(InIgnoreActors[Idx].Get());
+		Internal_AddIgnoredSourceObject(static_cast<FWeakObjectPtr>(InIgnoreActors[Idx]));
 	}
 }
 
+void FCollisionQueryParams::AddIgnoredSourceObject(const uint32 InIgnoreActorID)
+{
+	IgnoreSourceObjects.Add(InIgnoreActorID);
+
+#if UE_WITH_REMOTE_OBJECT_HANDLE
+	UObject* Obj = (UObject*)GUObjectArray.IndexToObject(InIgnoreActorID)->Object;
+	IgnoreSourceObjectReferences.Add(FRemoteObjectReference(Obj));
+#endif
+}
+
 void FCollisionQueryParams::AddIgnoredSourceObject(const UObject* InIgnoreObject)
 {
 	if (InIgnoreObject)
 	{
-		IgnoreSourceObjects.Add(InIgnoreObject->GetUniqueID());
+		AddIgnoredSourceObject(InIgnoreObject->GetUniqueID());
+#if UE_WITH_REMOTE_OBJECT_HANDLE
+		IgnoreSourceObjectReferences.Add(FRemoteObjectReference(InIgnoreObject));
+#endif
 	}
 }
 
-void FCollisionQueryParams::AddIgnoredSourceObject(const uint32 InIgnoreActorID)
+void FCollisionQueryParams::AddIgnoredSourceObject(const TWeakObjectPtr<const UObject>& InIgnoreObject)
 {
-	IgnoreSourceObjects.Add(InIgnoreActorID);
+	Internal_AddIgnoredSourceObject(static_cast<FWeakObjectPtr>(InIgnoreObject));
+}
+
+void FCollisionQueryParams::AddIgnoredSourceObject(const FRemoteObjectReference& ObjectRef)
+{
+	Internal_AddIgnoredSourceObject(ObjectRef.ToWeakPtr());
 }
 
 void FCollisionQueryParams::AddIgnoredSourceObjects(const TArray<UObject*>& InIgnoreObjects)
@@ -172,7 +209,7 @@ void FCollisionQueryParams::AddIgnoredSourceObjects(const TArray<TWeakObjectPtr<
 {
 	for (int32 Idx = 0; Idx < InIgnoreObjects.Num(); ++Idx)
 	{
-		AddIgnoredSourceObject(InIgnoreObjects[Idx].Get());
+		AddIgnoredSourceObject(InIgnoreObjects[Idx]);
 	}
 }
 
@@ -187,13 +224,60 @@ FORCEINLINE_DEBUGGABLE void FCollisionQueryParams::Internal_AddIgnoredComponent(
 void FCollisionQueryParams::AddIgnoredComponent(const UPrimitiveComponent* InIgnoreComponent)
 {
 	Internal_AddIgnoredComponent(InIgnoreComponent);
+
+#if UE_WITH_REMOTE_OBJECT_HANDLE
+	if (InIgnoreComponent != nullptr)
+	{
+		IgnoreComponentReferences.Add(FRemoteObjectReference(InIgnoreComponent));
+	}
+#endif
+}
+
+void FCollisionQueryParams::AddIgnoredComponent(const TWeakObjectPtr<UPrimitiveComponent>& InIgnoreComponent)
+{
+	if (InIgnoreComponent.IsValid())
+	{
+		// Only add valid local objects, otherwise migration could happen
+		if (!InIgnoreComponent.IsRemote())
+		{
+			Internal_AddIgnoredComponent(InIgnoreComponent.Get());
+		}
+
+#if UE_WITH_REMOTE_OBJECT_HANDLE
+		IgnoreComponentReferences.Add(FRemoteObjectReference(InIgnoreComponent));
+#endif
+	}
+}
+
+void FCollisionQueryParams::AddIgnoredComponent(const FRemoteObjectReference& InComponentRef)
+{
+	FWeakObjectPtr WeakPtr = InComponentRef.ToWeakPtr();
+	if (WeakPtr.IsValid())
+	{
+		if (!WeakPtr.IsRemote())
+		{
+			UObject* Obj = WeakPtr.Get();
+			// On rare occasions, this is null even with the above `IsValid` check
+			if (Obj != nullptr)
+			{
+				UPrimitiveComponent* PrimitiveComponent = Cast<UPrimitiveComponent>(Obj);
+				check(PrimitiveComponent);
+				Internal_AddIgnoredComponent(PrimitiveComponent);
+			}
+			
+		}
+
+#if UE_WITH_REMOTE_OBJECT_HANDLE
+		IgnoreComponentReferences.Add(InComponentRef);
+#endif
+	}
 }
 
 void FCollisionQueryParams::AddIgnoredComponents(const TArray<UPrimitiveComponent*>& InIgnoreComponents)
 {
 	for (const UPrimitiveComponent* IgnoreComponent : InIgnoreComponents)
 	{
-		Internal_AddIgnoredComponent(IgnoreComponent);
+		AddIgnoredComponent(IgnoreComponent);
 	}
 }
 
@@ -201,7 +285,7 @@ void FCollisionQueryParams::AddIgnoredComponents(const TArray<TWeakObjectPtr<UPr
 {
 	for (TWeakObjectPtr<UPrimitiveComponent> IgnoreComponent : InIgnoreComponents)
 	{
-		Internal_AddIgnoredComponent(IgnoreComponent.Get());
+		AddIgnoredComponent(IgnoreComponent);
 	}
 }
 
@@ -216,6 +300,9 @@ void FCollisionQueryParams::AddIgnoredComponent_LikelyDuplicatedRoot(const UPrim
 		{
 			IgnoreComponents.Add(ComponentID);
 			bComponentListUnique = false;
+#if UE_WITH_REMOTE_OBJECT_HANDLE
+			IgnoreComponentReferences.Add(FRemoteObjectReference(InIgnoreComponent));
+#endif
 		}
 	}
 }
diff --git a/Engine/Source/Runtime/Engine/Public/CollisionQueryParams.h b/Engine/Source/Runtime/Engine/Public/CollisionQueryParams.h
--- a/Engine/Source/Runtime/Engine/Public/CollisionQueryParams.h
+++ b/Engine/Source/Runtime/Engine/Public/CollisionQueryParams.h
@@ -7,6 +7,7 @@
 #include "CoreMinimal.h"
 #include "Engine/EngineTypes.h"
 #include "Chaos/ChaosEngineInterface.h"
+#include "UObject/RemoteObjectTransfer.h"
 
 /** Macro to convert ECollisionChannels to bit flag **/
 #define ECC_TO_BITFIELD(x)	(1<<(x))
@@ -82,6 +83,11 @@ struct FCollisionQueryParams
 	/** TArray typedef of actors to ignore. */
 	typedef TArray<uint32, TInlineAllocator<4>> IgnoreSourceObjectsArrayType;
 
+#if UE_WITH_REMOTE_OBJECT_HANDLE
+	typedef TArray<FRemoteObjectReference, TInlineAllocator<4>> IgnoreSourceObjectReferencesArrayType;
+	typedef TArray<FRemoteObjectReference, TInlineAllocator<8>> IgnoreComponentReferencesArrayType;
+#endif
+
 	typedef IgnoreSourceObjectsArrayType IgnoreActorsArrayType;
 
 	/** Extra filtering done on the query. See declaration for filtering logic */
@@ -110,6 +116,18 @@ private:
 	/** Set of source objects to ignore during the trace. In actor workflows, these are actors. */
 	IgnoreSourceObjectsArrayType IgnoreSourceObjects;
 
+#if UE_WITH_REMOTE_OBJECT_HANDLE
+	/** 
+	 * When running with auto-rtfm transactions, these reference arrays are the source of truth for all ignore sources.
+	 * The pre-existing ignore lists will only contain the local objects. 
+	 * It should always be possible to reconstruct the ignore object lists from the ignore references lists.
+	 */
+	mutable IgnoreComponentReferencesArrayType IgnoreComponentReferences;
+	IgnoreSourceObjectReferencesArrayType IgnoreSourceObjectReferences;
+#endif
+
+	void Internal_AddIgnoredSourceObject(const FWeakObjectPtr& WeakPtr);
+
 	ENGINE_API void Internal_AddIgnoredComponent(const UPrimitiveComponent* InIgnoreComponent);
 
 public:
@@ -117,6 +135,13 @@ public:
 	/** Returns set of unique components to ignore during the trace. Elements are guaranteed to be unique (they are made so internally if they are not already). */
 	ENGINE_API const IgnoreComponentsArrayType& GetIgnoredComponents() const;
 
+#if UE_WITH_REMOTE_OBJECT_HANDLE
+	ENGINE_API const IgnoreComponentReferencesArrayType& GetIgnoredComponentReferences() const
+	{
+		return IgnoreComponentReferences;
+	}
+#endif
+
 	UE_DEPRECATED(5.5, "Use GetIgnoredSourceObjects instead.")
 	/** Returns set of source objects (including actors) to ignore during the trace. Note that elements are NOT guaranteed to be unique. This is less important for actors since it's less likely that duplicates are added.*/
 	const IgnoreActorsArrayType& GetIgnoredActors() const
@@ -134,11 +159,22 @@ public:
 		return IgnoreSourceObjects;
 	}
 
+#if UE_WITH_REMOTE_OBJECT_HANDLE
+	const IgnoreSourceObjectReferencesArrayType& GetIgnoredSourceObjectReferences() const
+	{
+		return IgnoreSourceObjectReferences;
+	}
+#endif
+
 	/** Clears the set of components to ignore during the trace. */
 	void ClearIgnoredComponents()
 	{
 		IgnoreComponents.Reset();
 		bComponentListUnique = true;
+
+#if UE_WITH_REMOTE_OBJECT_HANDLE
+		IgnoreComponentReferences.Reset();
+#endif
 	}
 
 	UE_DEPRECATED(5.5, "Use ClearIgnoredSourceObjects instead.")
@@ -146,12 +182,20 @@ public:
 	void ClearIgnoredActors()
 	{
 		IgnoreSourceObjects.Reset();
+
+#if UE_WITH_REMOTE_OBJECT_HANDLE
+		IgnoreSourceObjectReferences.Reset();
+#endif
 	}
 
 	/** Clears the set of source objects (such as actors) to ignore during the trace. */
 	void ClearIgnoredSourceObjects()
 	{
 		IgnoreSourceObjects.Reset();
+
+#if UE_WITH_REMOTE_OBJECT_HANDLE
+		IgnoreSourceObjectReferences.Reset();
+#endif
 	}
 
 	/**
@@ -203,6 +247,8 @@ public:
 
 	/** Add a source object for this trace to ignore */
 	ENGINE_API void AddIgnoredSourceObject(const UObject* InIgnoreActor);
+	ENGINE_API void AddIgnoredSourceObject(const TWeakObjectPtr<const UObject>& InIgnoreObject);
+	ENGINE_API void AddIgnoredSourceObject(const FRemoteObjectReference& ObjectRef);
 
 	/** Add a source object for this trace to ignore */
 	ENGINE_API void AddIgnoredSourceObject(const uint32 InIgnoreActorID);
@@ -219,6 +265,8 @@ public:
 
 	/** Add a component for this trace to ignore */
 	ENGINE_API void AddIgnoredComponent(const UPrimitiveComponent* InIgnoreComponent);
+	ENGINE_API void AddIgnoredComponent(const TWeakObjectPtr<UPrimitiveComponent>& InIgnoreComponent);
+	ENGINE_API void AddIgnoredComponent(const FRemoteObjectReference& InComponentRef);
 
 	/** Add a collection of components for this trace to ignore */
 	ENGINE_API void AddIgnoredComponents(const TArray<UPrimitiveComponent*>& InIgnoreComponents);
