[shaders] add job status tracking mechanism and hook up for material shader maps; inject this information when the cooker encounters a material/material instance that it's blocked on saving. this will hopefully help us track down the rare material cook hangs.

example output:
[[REDACTED_PHONE]:451][  0]LogCook: Display:     Material /ControlRig/Controls/ControlRigGizmoMaterial.ControlRigGizmoMaterial
[[REDACTED_PHONE]:452][  0]LogCook: Display:         Waiting on job [ControlRigGizmoMaterial_98da241e3c8b5695/Default/FLocalVertexFactory/FLumenCardPS/0] (input hash cbf2794e56e542da00f0e639cc4d3c417a4e4cc38ed9e0f6296a7264bbe35185); current status is PendingDistributedExecution
[[REDACTED_PHONE]:453][  0]LogCook: Display:         Waiting on job [ControlRigGizmoMaterial_98da241e3c8b5695/Default/FLocalVertexFactory/TBasePassPSFNoLightMapPolicy/0] (input hash 0ea5fe1e314ff2ece74e1a6dad759361f6db83de5be62a7f1f7d224d46b08fa8); current status is PendingDistributedExecution
[[REDACTED_PHONE]:454][  0]LogCook: Display:         Waiting on job [ControlRigGizmoMaterial_98da241e3c8b5695/Default/FLocalVertexFactory/F128BitRTBasePassPS/0] (input hash adc05b4d801b69a27e62c72ebdac4bb19043b4b1d2a7973b1e400a6784285664); current status is CompleteDistributedExecution
[[REDACTED_PHONE]:455][  0]LogCook: Display:         Waiting on job [ControlRigGizmoMaterial_98da241e3c8b5695/Default/FLocalVertexFactory/F128BitRTBasePassPS/1] (input hash adc05b4d801b69a27e62c72ebdac4bb19043b4b1d2a7973b1e400a6784285664); current status is CompleteDistributedExecution (duplicate)
[[REDACTED_PHONE]:455][  0]LogCook: Display:     Material /Engine/EngineDebugMaterials/DebugEditorMaterial.DebugEditorMaterial
[[REDACTED_PHONE]:456][  0]LogCook: Display:         Waiting on job [DebugEditorMaterial_477ae1f8294debb7/Default/FLocalVertexFactory/TBasePassPSFNoLightMapPolicy/0] (input hash 53ebdf038466d21921c5b17cca6125de57067fe681e48ec15c4c13b1c2a20a53); current status is PendingDistributedExecution
[[REDACTED_PHONE]:457][  0]LogCook: Display:         Waiting on job [DebugEditorMaterial_477ae1f8294debb7/Default/FLocalVertexFactory/F128BitRTBasePassPS/0] (input hash bf7edb9343a31474673c70d056da588723ee2ecabe6fed5f69a56fc25988920f); current status is Queued (duplicate)
[[REDACTED_PHONE]:457][  0]LogCook: Display:         Waiting on job [DebugEditorMaterial_477ae1f8294debb7/Default/FLocalVertexFactory/F128BitRTBasePassPS/1] (input hash bf7edb9343a31474673c70d056da588723ee2ecabe6fed5f69a56fc25988920f); current status is PendingDistributedExecution
[[REDACTED_PHONE]:458][  0]LogCook: Display:         Waiting on job [Pipeline Job
[[REDACTED_PHONE]:458][  0]LogCook: Display:             Stage 0: DebugEditorMaterial_477ae1f8294debb7/Default/VelocityPipeline/FLocalVertexFactory/FVelocityPS/0
[[REDACTED_PHONE]:458][  0]LogCook: Display:             Stage 1: DebugEditorMaterial_477ae1f8294debb7/Default/VelocityPipeline/FLocalVertexFactory/FVelocityVS/0
[[REDACTED_PHONE]:458][  0]LogCook: Display:         ] (input hash 1303e6a50dd2fb43f442081c07286303453167d37ffa886fd260a0ed1ea98565); current status is PendingDistributedExecution

diff --git a/Engine/Source/Runtime/Engine/Private/Materials/Material.cpp b/Engine/Source/Runtime/Engine/Private/Materials/Material.cpp
--- a/Engine/Source/Runtime/Engine/Private/Materials/Material.cpp
+++ b/Engine/Source/Runtime/Engine/Private/Materials/Material.cpp
@@ -4691,6 +4691,17 @@ void UMaterial::GetShaderTypes(EShaderPlatform ShaderPlatform, const ITargetPlat
 
 	NewResourcesToCache.Empty();
 }
+
+void UMaterial::AppendCompileStateDebugInfo(FStringBuilderBase& OutDebugInfo) const
+{
+	for (auto Iter = CachedMaterialResourcesForCooking.CreateConstIterator(); Iter; ++Iter)
+	{
+		for (const FMaterialResourceForCooking& Resource : Iter.Value())
+		{
+			Resource.Resource->AppendCompileStateDebugInfo(OutDebugInfo);
+		}
+	}
+}
 #endif // WITH_EDITOR
 
 bool UMaterial::IsPropertyConnected(EMaterialProperty Property) const
diff --git a/Engine/Source/Runtime/Engine/Private/Materials/MaterialShader.cpp b/Engine/Source/Runtime/Engine/Private/Materials/MaterialShader.cpp
--- a/Engine/Source/Runtime/Engine/Private/Materials/MaterialShader.cpp
+++ b/Engine/Source/Runtime/Engine/Private/Materials/MaterialShader.cpp
@@ -2228,6 +2228,7 @@ void FMaterialShaderMap::ReleaseCompilingId()
 		FWriteScopeLock Locker(GCompilingShaderMapLock);
 		check(CompilingMaterialDependencies.Num() == 0);
 		check(CompilingMaterialNumExternalDependencies == 0);
+		check(InFlightJobs.Num() == 0);
 		verify(GetCompilingShaderMapLookup().Remove(CompilingId) == 1);
 		CompilingId = 0;
 	}
@@ -2498,6 +2499,15 @@ int32 FMaterialShaderMap::SubmitCompileJobs(uint32 CompilingShaderMapId,
 
 	UE_LOG(LogShaders, Verbose, TEXT("		%u Shaders among %u VertexFactories"), NumShaders, NumVertexFactories);
 
+#if ALLOW_SHADERMAP_DEBUG_DATA
+	for (const FShaderCommonCompileJobPtr& Job : CompileJobs)
+	{
+		TStringBuilder<2048> JobName;
+		Job->AppendDebugName(JobName);
+		InFlightJobs.Add(Job->Hash, { JobName.ToString(), Job->JobStatusPtr });
+	}
+#endif
+
 	GShaderCompilingManager->SubmitJobs(CompileJobs, Material->GetBaseMaterialPathName(), GetDebugDescription());
 
 	return CompileJobs.Num();
@@ -2768,6 +2778,20 @@ FShader* FMaterialShaderMap::ProcessCompilationResultsForSingleJob(FShaderCompil
 	return Shader;
 }
 
+void FMaterialShaderMap::AppendCompileStateDebugInfo(FStringBuilderBase& OutDebugInfo) const
+{
+	for (const TPair<uint32, TPair<FString, TSharedPtr<FShaderCompileJobStatus>>>& StatusPair : InFlightJobs)
+	{
+		TSharedPtr<FShaderCompileJobStatus> Status = StatusPair.Value.Value;
+		OutDebugInfo << "Waiting on job [" << StatusPair.Value.Key << "]";
+		if (!Status->GetInputHash().IsZero())
+		{
+			OutDebugInfo << " (input hash " << LexToString(Status->GetInputHash()) << ")";
+		}
+		OutDebugInfo << "; current status is " << LexToString(Status->GetStatus()) << (Status->GetIsDuplicate() ? " (duplicate)" : "") << "\n";
+	}
+}
+
 void FMaterialShaderMap::ProcessCompilationResults(const TArray<FShaderCommonCompileJobPtr>& InCompilationResults, int32& InOutJobIndex, float& TimeBudget)
 {
 	check(!bCompilationFinalized);
@@ -2780,6 +2804,8 @@ void FMaterialShaderMap::ProcessCompilationResults(const TArray<FShaderCommonCom
 	do
 	{
 		auto& BaseJob = InCompilationResults[InOutJobIndex++];
+		InFlightJobs.Remove(BaseJob->Hash);
+
 		FShaderCompileJob* SingleJob = BaseJob->GetSingleShaderJob();
 		if (SingleJob)
 		{
diff --git a/Engine/Source/Runtime/Engine/Private/Materials/MaterialShared.cpp b/Engine/Source/Runtime/Engine/Private/Materials/MaterialShared.cpp
--- a/Engine/Source/Runtime/Engine/Private/Materials/MaterialShared.cpp
+++ b/Engine/Source/Runtime/Engine/Private/Materials/MaterialShared.cpp
@@ -46,13 +46,14 @@
 #include "MaterialCachedData.h"
 #include "VT/RuntimeVirtualTexture.h"
 #if WITH_EDITOR
-#include "Serialization/CompactBinary.h"
-#include "Serialization/CompactBinaryWriter.h"
-#include "Shader/PreshaderEvaluate.h"
+#include "MaterialKeyGeneratorContext.h"
 #include "Misc/AsciiSet.h"
 #include "PostProcess/PostProcessMaterialInputs.h"
-#include "MaterialKeyGeneratorContext.h"
+#include "Serialization/CompactBinary.h"
 #include "Serialization/CompactBinarySerialization.h"
+#include "Serialization/CompactBinaryWriter.h"
+#include "Shader/PreshaderEvaluate.h"
+#include "UObject/ICookInfo.h"
 #endif
 #if WITH_ODSC
 #include "ODSC/ODSCManager.h"
@@ -277,6 +278,20 @@ int32 GetCmdLineMaterialQualityToCache()
 #endif
 
 
+#if WITH_EDITOR
+static FDelayedAutoRegisterHelper GBlockedMaterialDebugDelegateRegister(EDelayedRegisterRunPhase::StartOfEnginePreInit, []
+	{
+		UE::Cook::FDelegates::PackageBlocked.AddLambda([](const UObject* Obj, FStringBuilderBase& OutDebugInfo)
+			{
+				const UMaterialInterface* MaterialInterface = Cast<UMaterialInterface>(Obj);
+				const UMaterial* Material = MaterialInterface ? MaterialInterface->GetMaterial() : nullptr;
+				if (Material)
+				{
+					Material->AppendCompileStateDebugInfo(OutDebugInfo);
+				}
+			});
+	});
+#endif
 
 int32 FMaterialCompiler::Errorf(const TCHAR* Format,...)
 {
@@ -2241,6 +2256,19 @@ bool FMaterialResource::CheckInValidStateForCompilation(FMaterialCompiler* Compi
 	return Material && Material->CheckInValidStateForCompilation(Compiler);
 }
 
+void FMaterialResource::AppendCompileStateDebugInfo(FStringBuilderBase& OutDebugInfo) const
+{
+	check(IsInGameThread());
+	if (GetGameThreadCompilingShaderMapId() != 0u)
+	{
+		FMaterialShaderMap* CompilingShaderMap = FMaterialShaderMap::FindCompilingShaderMap(GetGameThreadCompilingShaderMapId());
+		if (CompilingShaderMap)
+		{
+			CompilingShaderMap->AppendCompileStateDebugInfo(OutDebugInfo);
+		}
+	}
+}
+
 #endif
 
 FString FMaterialResource::GetFullPath() const
diff --git a/Engine/Source/Runtime/Engine/Private/ShaderCompiler/ShaderCompiler.cpp b/Engine/Source/Runtime/Engine/Private/ShaderCompiler/ShaderCompiler.cpp
--- a/Engine/Source/Runtime/Engine/Private/ShaderCompiler/ShaderCompiler.cpp
+++ b/Engine/Source/Runtime/Engine/Private/ShaderCompiler/ShaderCompiler.cpp
@@ -1539,7 +1539,7 @@ FShaderPipelineCompileJob* FShaderCompilingManager::PreparePipelineCompileJob(ui
 	return AllJobs.PrepareJob(Id, Key, Priority);
 }
 
-void FShaderCompilingManager::ProcessFinishedJob(FShaderCommonCompileJob* FinishedJob)
+void FShaderCompilingManager::ProcessFinishedJob(FShaderCommonCompileJob* FinishedJob, EShaderCompileJobStatus Status)
 {
 	bool bIsPipelineJob = FinishedJob->Type == EShaderCompileJobType::Pipeline;
 	FinishedJob->ForEachSingleShaderJob([this](FShaderCompileJob& SingleJob)
@@ -1561,7 +1561,7 @@ void FShaderCompilingManager::ProcessFinishedJob(FShaderCommonCompileJob* Finish
 			}
 		});
 
-	AllJobs.ProcessFinishedJob(FinishedJob);
+	AllJobs.ProcessFinishedJob(FinishedJob, Status);
 }
 
 /** Launches the worker, returns the launched process handle. */
diff --git a/Engine/Source/Runtime/Engine/Private/ShaderCompiler/ShaderCompilerDistributed.cpp b/Engine/Source/Runtime/Engine/Private/ShaderCompiler/ShaderCompilerDistributed.cpp
--- a/Engine/Source/Runtime/Engine/Private/ShaderCompiler/ShaderCompilerDistributed.cpp
+++ b/Engine/Source/Runtime/Engine/Private/ShaderCompiler/ShaderCompilerDistributed.cpp
@@ -5,6 +5,7 @@
 #include "HAL/FileManager.h"
 #include "Misc/CommandLine.h"
 #include "ShaderCompiler.h"
+#include "ShaderCompilerJobTypes.h"
 
 namespace DistributedShaderCompilerVariables
 {
@@ -328,6 +329,7 @@ int32 FShaderCompileDistributedThreadRunnable_Interface::CompilingLoop()
 	TArray<FShaderCommonCompileJobPtr> PendingJobs;
 	//if (LIKELY(!bIsHung))	// stop accepting jobs if we're hung - TODO: re-enable this after lockup detection logic is proved reliable and/or we have job resubmission in place
 	{
+		FScopeLock Lock(&Manager->CompileQueueSection);
 		for (int32 PriorityIndex = MaxPriorityIndex; PriorityIndex >= MinPriorityIndex; --PriorityIndex)
 		{
 			// Grab as many jobs from the job queue as we can, unless there is no local shader compiling thread to pick up smaller batches
@@ -523,7 +525,7 @@ int32 FShaderCompileDistributedThreadRunnable_Interface::CompilingLoop()
 				FScopeLock Lock(&Manager->CompileQueueSection);
 				for (const auto& Job : Task->ShaderJobs)
 				{
-					Manager->ProcessFinishedJob(Job);
+					Manager->ProcessFinishedJob(Job, EShaderCompileJobStatus::CompleteDistributedExecution);
 				}
 			}
 		}
diff --git a/Engine/Source/Runtime/Engine/Private/ShaderCompiler/ShaderCompilerJobCache.cpp b/Engine/Source/Runtime/Engine/Private/ShaderCompiler/ShaderCompilerJobCache.cpp
--- a/Engine/Source/Runtime/Engine/Private/ShaderCompiler/ShaderCompilerJobCache.cpp
+++ b/Engine/Source/Runtime/Engine/Private/ShaderCompiler/ShaderCompilerJobCache.cpp
@@ -582,7 +582,7 @@ public:
 	void SubmitJobs(const TArray<FShaderCommonCompileJobPtr>& InJobs);
 
 	/** This is an entry point for all jobs that have finished the compilation (whether real or cached). Can be called from multiple threads. Returns mutex stall time. */
-	double ProcessFinishedJob(FShaderCommonCompileJob* FinishedJob, bool bCompilationSkipped);
+	double ProcessFinishedJob(FShaderCommonCompileJob* FinishedJob, EShaderCompileJobStatus Status);
 
 	/** Adds the job to cache. */
 	void AddToCacheAndProcessPending(FShaderCommonCompileJob* FinishedJob);
@@ -1043,9 +1043,9 @@ void FShaderCompileJobCollection::SubmitJobs(const TArray<FShaderCommonCompileJo
 	JobsCache->SubmitJobs(InJobs);
 }
 
-void FShaderCompileJobCollection::ProcessFinishedJob(FShaderCommonCompileJob* FinishedJob, bool bCompilationSkipped)
+void FShaderCompileJobCollection::ProcessFinishedJob(FShaderCommonCompileJob* FinishedJob, EShaderCompileJobStatus Status)
 {
-	JobsCache->ProcessFinishedJob(FinishedJob, bCompilationSkipped);
+	JobsCache->ProcessFinishedJob(FinishedJob, Status);
 }
 
 void FShaderCompileJobCollection::AddToCacheAndProcessPending(FShaderCommonCompileJob* FinishedJob)
@@ -1358,7 +1358,7 @@ void FShaderJobCache::SubmitJob(FShaderCommonCompileJob* Job)
 		if (!Job->PrevLink)
 		{
 			UE_LOG(LogShaderCompilers, Display, TEXT("Cancelled job 0x%p with pending SubmitJob call."), Job);
-
+			Job->UpdateStatus(EShaderCompileJobStatus::Cancelled);
 			JobLock.WriteUnlock();
 			return;
 		}
@@ -1382,11 +1382,12 @@ void FShaderJobCache::SubmitJob(FShaderCommonCompileJob* Job)
 			Job->SerializeOutput(LoadContext);
 
 			// finish the job instantly
-			Job->TimeTaskSubmitJobsStall += ProcessFinishedJob(Job, true);
+			Job->TimeTaskSubmitJobsStall += ProcessFinishedJob(Job, EShaderCompileJobStatus::CompleteFoundInCache);
 		}
 		else
 		{
 			FShaderJobData& JobData = GetShaderJobData(Job->JobCacheRef);
+			Job->UpdateStatus(EShaderCompileJobStatus::Queued);
 			// see if another job with the same input hash is being worked on
 			if (JobData.JobInFlight)
 			{
@@ -1406,6 +1407,7 @@ void FShaderJobCache::SubmitJob(FShaderCommonCompileJob* Job)
 				}
 				++TotalCacheDuplicates;
 
+				Job->JobStatusPtr->SetIsDuplicate();
 				AddDuplicateJob(Job);
 				JobLock.WriteUnlock();
 				bNewJob = false;
@@ -1444,7 +1446,7 @@ void FShaderJobCache::SubmitJob(FShaderCommonCompileJob* Job)
 			if (!Job->PrevLink)
 			{
 				UE_LOG(LogShaderCompilers, Log, TEXT("Cancelled job 0x%p with pending SubmitJob call."), Job);
-
+				Job->UpdateStatus(EShaderCompileJobStatus::Cancelled);
 				JobLock.WriteUnlock();
 				return;
 			}
@@ -1500,6 +1502,7 @@ void FShaderJobCache::SubmitJobs(const TArray<FShaderCommonCompileJobPtr>& InJob
 
 					const bool bSubmitJob = PreprocessShader(Job);
 					Job->GetInputHash();
+					Job->UpdateStatus(EShaderCompileJobStatus::Ready);
 
 					if (bSubmitJob)
 					{
@@ -1507,18 +1510,28 @@ void FShaderJobCache::SubmitJobs(const TArray<FShaderCommonCompileJobPtr>& InJob
 					}
 					else // if preprocessing ran and failed, finish the job immediately
 					{
-						ProcessFinishedJob(Job, /* bCompilationSkipped = */true);
+						ProcessFinishedJob(Job, EShaderCompileJobStatus::Skipped);
 					}
 
 					Job->TimeTaskSubmitJobs = FPlatformTime::Seconds() - TimeStart;
+
 				}, Prio);
 		}
 	}
 }
 
-double FShaderJobCache::ProcessFinishedJob(FShaderCommonCompileJob* FinishedJob, bool bCompilationSkipped)
+double FShaderJobCache::ProcessFinishedJob(FShaderCommonCompileJob* FinishedJob, EShaderCompileJobStatus Status)
 {
 	double StallTime;
+
+	FinishedJob->UpdateStatus(Status);
+
+	bool bCompilationSkipped =
+		FinishedJob->JobStatusPtr->GetIsDuplicate() ||
+		(Status == EShaderCompileJobStatus::Skipped) ||
+		(Status == EShaderCompileJobStatus::CompleteFoundInCache) ||
+		(Status == EShaderCompileJobStatus::CompleteFoundInDDC);
+								
 	if (!bCompilationSkipped)
 	{
 		AddToCacheAndProcessPending(FinishedJob);
@@ -1628,7 +1641,7 @@ void FShaderJobCache::AddToCacheAndProcessPending(FShaderCommonCompileJob* Finis
 			checkf(DuplicateJob->bSucceeded == FinishedJob->bSucceeded, TEXT("Different success status for the job with the same ihash"));
 
 			// finish the job instantly
-			ProcessFinishedJob(DuplicateJob, true);
+			ProcessFinishedJob(DuplicateJob, FinishedJob->JobStatusPtr->GetStatus());
 		}
 	}
 }
@@ -1703,6 +1716,7 @@ int32 FShaderJobCache::GetPendingJobs(EShaderCompilerWorkerType InWorkerType, ES
 		UnlinkJobWithPriority(Job);
 
 		Job.CurrentWorker = InWorkerType;
+		Job.UpdateStatus(InWorkerType == EShaderCompilerWorkerType::Distributed ? EShaderCompileJobStatus::PendingDistributedExecution : EShaderCompileJobStatus::PendingLocalExecution);
 		OutJobs.Add(&Job);
 
 		// get a random number of jobs to skip (if we can). We're skipping after taking the first job so we can ensure that we always take the latest job into the batch
@@ -1769,6 +1783,8 @@ FShaderCacheLoadContext FShaderJobCache::FindOrAdd(const FShaderCompilerInputHas
 			UE::DerivedData::EPriority DerivedDataPriority;
 			UE::DerivedData::FRequestOwner* RequestOwner;
 
+			Job->UpdateStatus(EShaderCompileJobStatus::PendingDDC);
+
 			static const bool PerShaderDDCAsync = CVarShaderCompilerPerShaderDDCAsync.GetValueOnAnyThread();
 			if (PerShaderDDCAsync && FGenericPlatformProcess::SupportsMultithreading())
 			{
@@ -1852,6 +1868,7 @@ FShaderCacheLoadContext FShaderJobCache::FindOrAdd(const FShaderCompilerInputHas
 								UE_LOG(LogShaderCompilers, Display, TEXT("Cancelled job 0x%p (data 0x%p) with pending DDC hit."), JobDataPtr->JobInFlight.GetReference(), JobDataPtr);
 								if (JobDataPtr->JobInFlight)
 								{
+									JobDataPtr->JobInFlight->UpdateStatus(EShaderCompileJobStatus::Cancelled);
 #if WITH_EDITOR
 									if (JobDataPtr->JobInFlight->RequestOwner.IsValid())
 									{
@@ -1934,7 +1951,7 @@ FShaderCacheLoadContext FShaderJobCache::FindOrAdd(const FShaderCompilerInputHas
 							{
 								LoadContext.Reuse();
 								FinishedJob->SerializeOutput(LoadContext);
-								ProcessFinishedJob(FinishedJob, true);
+								ProcessFinishedJob(FinishedJob, EShaderCompileJobStatus::CompleteFoundInDDC);
 							}
 
 							if (FinishedJobs.Num() > 1)
@@ -1961,6 +1978,7 @@ FShaderCacheLoadContext FShaderJobCache::FindOrAdd(const FShaderCompilerInputHas
 							if (!Job->PrevLink)
 							{
 								UE_LOG(LogShaderCompilers, Display, TEXT("Cancelled job 0x%p (data 0x%p) with pending DDC miss."), Job, JobDataPtr);
+								Job->UpdateStatus(EShaderCompileJobStatus::Cancelled);
 
 								if (JobDataPtr->JobInFlight)
 								{
diff --git a/Engine/Source/Runtime/Engine/Private/ShaderCompiler/ShaderCompilerThreadRunnable.cpp b/Engine/Source/Runtime/Engine/Private/ShaderCompiler/ShaderCompilerThreadRunnable.cpp
--- a/Engine/Source/Runtime/Engine/Private/ShaderCompiler/ShaderCompilerThreadRunnable.cpp
+++ b/Engine/Source/Runtime/Engine/Private/ShaderCompiler/ShaderCompilerThreadRunnable.cpp
@@ -478,7 +478,7 @@ void FShaderCompileThreadRunnable::PushCompletedJobsToManager()
 			{
 				auto& Job = CurrentWorkerInfo.QueuedJobs[JobIndex];
 
-				Manager->ProcessFinishedJob(Job.GetReference());
+				Manager->ProcessFinishedJob(Job.GetReference(), EShaderCompileJobStatus::CompleteLocalExecution);
 			}
 
 			const double ElapsedTime = FPlatformTime::Seconds() - CurrentWorkerInfo.StartTime;
diff --git a/Engine/Source/Runtime/Engine/Public/MaterialShared.h b/Engine/Source/Runtime/Engine/Public/MaterialShared.h
--- a/Engine/Source/Runtime/Engine/Public/MaterialShared.h
+++ b/Engine/Source/Runtime/Engine/Public/MaterialShared.h
@@ -79,6 +79,7 @@ class FMeshMaterialShaderMapLayout;
 struct FMaterialShaderTypes;
 struct FMaterialShaders;
 struct FMaterialCachedExpressionData;
+class FShaderCompileJobStatus;
 class FShaderMapLayout;
 #if WITH_EDITOR
 class FMaterialKeyGeneratorContext;
@@ -1695,6 +1696,8 @@ public:
 
 #if WITH_EDITOR
 	void InitalizeForODSC(EShaderPlatform TargetShaderPlatform, const FMaterialCompilationOutput& NewCompilationOutput);
+
+	void AppendCompileStateDebugInfo(FStringBuilderBase& OutDebugInfo) const;
 #endif
 
 protected:
@@ -1776,6 +1779,10 @@ private:
 	FShader* ProcessCompilationResultsForSingleJob(class FShaderCompileJob* SingleJob, const FShaderPipelineType* ShaderPipeline, const FSHAHash& MaterialShaderMapHash);
 #endif
 
+#if WITH_EDITOR && ALLOW_SHADERMAP_DEBUG_DATA
+	mutable TMap<uint32, TPair<FString, TSharedPtr<FShaderCompileJobStatus>>> InFlightJobs;
+#endif
+
 	/** Iterate over global MaterialShaderMap cache and remove MaterialShaderMap which reference FShaderMapResource argument */
 	static void ShaderMapResourceExplicitRelease(const FShaderMapResource* ShaderMapResource);
 
@@ -3055,6 +3062,7 @@ public:
 	ENGINE_API virtual bool ShouldInlineShaderCode() const override { return false; }
 	ENGINE_API virtual bool IsUsingNewHLSLGenerator() const override;
 	ENGINE_API virtual bool CheckInValidStateForCompilation(FMaterialCompiler* Compiler) const override;
+	void AppendCompileStateDebugInfo(FStringBuilderBase& OutDebugInfo) const;
 #endif // WITH_EDITOR
 
 	ENGINE_API virtual FString GetFullPath() const override;
diff --git a/Engine/Source/Runtime/Engine/Public/Materials/Material.h b/Engine/Source/Runtime/Engine/Public/Materials/Material.h
--- a/Engine/Source/Runtime/Engine/Public/Materials/Material.h
+++ b/Engine/Source/Runtime/Engine/Public/Materials/Material.h
@@ -1875,6 +1875,8 @@ public:
 
 #if WITH_EDITOR
 	ENGINE_API virtual void GetShaderTypes(EShaderPlatform Platform, const ITargetPlatform* TargetPlatform, TArray<FDebugShaderTypeInfo>& OutShaderInfo) override;
+
+	void AppendCompileStateDebugInfo(FStringBuilderBase& OutCompileStateStr) const;
 #endif // WITH_EDITOR
 
 	ENGINE_API bool IsPropertyConnected(EMaterialProperty Property) const;
diff --git a/Engine/Source/Runtime/Engine/Public/ShaderCompiler.h b/Engine/Source/Runtime/Engine/Public/ShaderCompiler.h
--- a/Engine/Source/Runtime/Engine/Public/ShaderCompiler.h
+++ b/Engine/Source/Runtime/Engine/Public/ShaderCompiler.h
@@ -76,7 +76,7 @@ public:
 	/** Called for all completed jobs, including those that were cache hits, duplicates of other in flight jobs, or skipped due to failed preprocessing.
 	 * Can be called from multiple threads.
 	 */
-	void ProcessFinishedJob(FShaderCommonCompileJob* FinishedJob, bool bCompilationSkipped = false);
+	void ProcessFinishedJob(FShaderCommonCompileJob* FinishedJob, EShaderCompileJobStatus Status);
 
 	/** Adds the job to cache. */
 	void AddToCacheAndProcessPending(FShaderCommonCompileJob* FinishedJob);
@@ -1204,8 +1204,12 @@ public:
 	ENGINE_API FShaderCompileJob* PrepareShaderCompileJob(uint32 Id, const FShaderCompileJobKey& Key, EShaderCompileJobPriority Priority);
 	ENGINE_API FShaderPipelineCompileJob* PreparePipelineCompileJob(uint32 Id, const FShaderPipelineCompileJobKey& Key, EShaderCompileJobPriority Priority);
 
+
+	UE_DEPRECATED(5.6, "ProcessFinishedJob must now be passed an EShaderCompileJobStatus")
+	void ProcessFinishedJob(FShaderCommonCompileJob* FinishedJob) {}
+
 	/** This is an entry point for all jobs that have finished the compilation. Can be called from multiple threads.*/
-	ENGINE_API void ProcessFinishedJob(FShaderCommonCompileJob* FinishedJob);
+	ENGINE_API void ProcessFinishedJob(FShaderCommonCompileJob* FinishedJob, EShaderCompileJobStatus Status);
 
 	/** 
 	 * Adds shader jobs to be asynchronously compiled. 
diff --git a/Engine/Source/Runtime/RenderCore/Private/ShaderCore.cpp b/Engine/Source/Runtime/RenderCore/Private/ShaderCore.cpp
--- a/Engine/Source/Runtime/RenderCore/Private/ShaderCore.cpp
+++ b/Engine/Source/Runtime/RenderCore/Private/ShaderCore.cpp
@@ -3973,6 +3973,26 @@ FArchive& operator<<(FArchive& Ar, FShaderCompilerInput& Input)
 	return Ar;
 }
 
+const TCHAR* LexToString(EShaderCompileJobStatus Status)
+{
+	switch (Status)
+	{
+	case EShaderCompileJobStatus::Unset: return TEXT("Unset");
+	case EShaderCompileJobStatus::Ready: return TEXT("Ready");
+	case EShaderCompileJobStatus::Skipped: return TEXT("Skipped");
+	case EShaderCompileJobStatus::Cancelled: return TEXT("Cancelled");
+	case EShaderCompileJobStatus::PendingDDC: return TEXT("PendingDDC");
+	case EShaderCompileJobStatus::Queued: return TEXT("Queued");
+	case EShaderCompileJobStatus::PendingDistributedExecution: return TEXT("PendingDistributedExecution");
+	case EShaderCompileJobStatus::PendingLocalExecution: return TEXT("PendingLocalExecution");
+	case EShaderCompileJobStatus::CompleteDistributedExecution: return TEXT("CompleteDistributedExecution");
+	case EShaderCompileJobStatus::CompleteFoundInCache: return TEXT("CompleteFoundInCache");
+	case EShaderCompileJobStatus::CompleteFoundInDDC: return TEXT("CompleteFoundInDDC");
+	case EShaderCompileJobStatus::CompleteLocalExecution: return TEXT("CompleteLocalExecution");
+	default: return TEXT("(unknown)");
+	}
+}
+
 FShaderCompilerInputHash FShaderPipelineCompileJob::GetInputHash()
 {
 	if (bInputHashSet)
@@ -4335,8 +4355,10 @@ void FShaderPipelineCompileJob::OnComplete()
 
 void FShaderPipelineCompileJob::AppendDebugName(FStringBuilderBase& OutName) const
 {
+	OutName << "Pipeline Job\n";
 	for (int32 Index = 0, Num = StageJobs.Num(); Index < Num; ++Index)
 	{
+		OutName << "    Stage " << Index << ": ";
 		StageJobs[Index]->AppendDebugName(OutName);
 		OutName << TEXT("\n");
 	}
diff --git a/Engine/Source/Runtime/RenderCore/Public/ShaderCompilerJobTypes.h b/Engine/Source/Runtime/RenderCore/Public/ShaderCompilerJobTypes.h
--- a/Engine/Source/Runtime/RenderCore/Public/ShaderCompilerJobTypes.h
+++ b/Engine/Source/Runtime/RenderCore/Public/ShaderCompilerJobTypes.h
@@ -2,6 +2,7 @@
 
 #pragma once
 
+#include <atomic>
 #include "HAL/PlatformTime.h"
 #include "Hash/Blake3.h"
 #include "Shader.h"
@@ -72,6 +73,71 @@ struct FShaderJobCacheRef
 	}
 };
 
+enum class EShaderCompileJobStatus : uint8
+{
+	Unset,
+	Ready,
+	Skipped,
+	Cancelled,
+	PendingDDC,
+	Queued,
+	PendingDistributedExecution,
+	PendingLocalExecution,
+	CompleteDistributedExecution,
+	CompleteFoundInCache,
+	CompleteFoundInDDC,
+	CompleteLocalExecution,
+};
+RENDERCORE_API const TCHAR* LexToString(EShaderCompileJobStatus Status);
+
+class FShaderCompileJobStatus
+{
+public:
+	void SetStatus(EShaderCompileJobStatus NewStatus)
+	{
+		Status.store(NewStatus);
+	}
+
+	EShaderCompileJobStatus GetStatus()
+	{
+		return Status.load();
+	}
+
+	void SetIsDuplicate()
+	{
+		IsDuplicate.store(true);
+	}
+
+	bool GetIsDuplicate() const
+	{
+		return IsDuplicate.load();
+	}
+
+	void SetInputHash(const FShaderCompilerInputHash* InInputHash)
+	{
+		InputHash = InInputHash;
+	}
+
+	const FShaderCompilerInputHash& GetInputHash()
+	{
+		// input hash is guaranteed to be set once status is set to Submitted (which is the first status that is always set by the system)
+		// if we somehow query the hash prior to this occurring, just return zero (as it's possible the hash could be in an invalid state
+		// if the submit job is in the process of setting it). This prevents the need for any other thread safety since the status is set
+		// atomically.
+		if (GetStatus() == EShaderCompileJobStatus::Unset)
+		{
+			return FShaderCompilerInputHash::Zero;
+		}
+		check(InputHash);
+		return *InputHash;
+	}
+
+private:
+	const FShaderCompilerInputHash* InputHash = nullptr; // set when input hash is computed during submission
+	std::atomic<EShaderCompileJobStatus> Status{ EShaderCompileJobStatus::Unset };
+	std::atomic<bool> IsDuplicate{ false };
+};
+
 /** Stores all of the common information used to compile a shader or pipeline. */
 class FShaderCommonCompileJob
 {
@@ -128,6 +194,19 @@ public:
 
 	FShaderJobCacheRef JobCacheRef;
 
+	TSharedPtr<FShaderCompileJobStatus> JobStatusPtr;
+
+	void UpdateStatus(EShaderCompileJobStatus NewStatus)
+	{
+		if (NewStatus == EShaderCompileJobStatus::Ready)
+		{
+			check(bInputHashSet);
+			JobStatusPtr->SetInputHash(&InputHash);
+		}
+
+		JobStatusPtr->SetStatus(NewStatus);
+	}
+	
 	uint32 AddRef() const
 	{
 		return uint32(NumRefs.Increment());
@@ -168,7 +247,7 @@ public:
 	virtual void OnComplete() = 0;
 
 	virtual void AppendDebugName(FStringBuilderBase& OutName) const = 0;
-	
+
 	bool Equals(const FShaderCommonCompileJob& Rhs) const;
 
 	/** Calls the specified predicate for each single compile job, i.e. FShaderCompileJob and each stage of FShaderPipelineCompileJob. */
@@ -201,6 +280,8 @@ protected:
 
 	{
 		check(InPriority != EShaderCompileJobPriority::None);
+
+		JobStatusPtr = MakeShared<FShaderCompileJobStatus>();
 	}
 
 	virtual ~FShaderCommonCompileJob() {}
