[API] Fix linker issue for clients that intend to use the const-version of GetAnimCompositeSection().

The mutable version of GetAnimCompositeSection() is DLL exported (CL 2790182) but not the const-version.
This can cause unintuitive linker errors when clients attempt to access the const-version and there is no abvious reason to export the mutabel but not the const version.

#rnx

diff --git a/Engine/Source/Runtime/Engine/Classes/Animation/AnimMontage.h b/Engine/Source/Runtime/Engine/Classes/Animation/AnimMontage.h
--- a/Engine/Source/Runtime/Engine/Classes/Animation/AnimMontage.h
+++ b/Engine/Source/Runtime/Engine/Classes/Animation/AnimMontage.h
@@ -783,7 +783,7 @@ public:
 
 	/** Get FCompositeSection with InSectionName */
 	ENGINE_API FCompositeSection& GetAnimCompositeSection(int32 SectionIndex);
-	const FCompositeSection& GetAnimCompositeSection(int32 SectionIndex) const;
+	ENGINE_API const FCompositeSection& GetAnimCompositeSection(int32 SectionIndex) const;
 
 	// @todo document
 	ENGINE_API void GetSectionStartAndEndTime(int32 SectionIndex, float& OutStartTime, float& OutEndTime) const;
