fix for cooked editor build break

#rnx

diff --git a/Engine/Source/Runtime/Engine/Private/Materials/MaterialShader.cpp b/Engine/Source/Runtime/Engine/Private/Materials/MaterialShader.cpp
--- a/Engine/Source/Runtime/Engine/Private/Materials/MaterialShader.cpp
+++ b/Engine/Source/Runtime/Engine/Private/Materials/MaterialShader.cpp
@@ -2228,7 +2228,9 @@ void FMaterialShaderMap::ReleaseCompilingId()
 		FWriteScopeLock Locker(GCompilingShaderMapLock);
 		check(CompilingMaterialDependencies.Num() == 0);
 		check(CompilingMaterialNumExternalDependencies == 0);
+#if ALLOW_SHADERMAP_DEBUG_DATA
 		check(InFlightJobs.Num() == 0);
+#endif
 		verify(GetCompilingShaderMapLookup().Remove(CompilingId) == 1);
 		CompilingId = 0;
 	}
@@ -2780,6 +2782,7 @@ FShader* FMaterialShaderMap::ProcessCompilationResultsForSingleJob(FShaderCompil
 
 void FMaterialShaderMap::AppendCompileStateDebugInfo(FStringBuilderBase& OutDebugInfo) const
 {
+#if ALLOW_SHADERMAP_DEBUG_DATA
 	for (const TPair<uint32, TPair<FString, TSharedPtr<FShaderCompileJobStatus>>>& StatusPair : InFlightJobs)
 	{
 		TSharedPtr<FShaderCompileJobStatus> Status = StatusPair.Value.Value;
@@ -2790,6 +2793,7 @@ void FMaterialShaderMap::AppendCompileStateDebugInfo(FStringBuilderBase& OutDebu
 		}
 		OutDebugInfo << "; current status is " << LexToString(Status->GetStatus()) << (Status->GetIsDuplicate() ? " (duplicate)" : "") << "\n";
 	}
+#endif
 }
 
 void FMaterialShaderMap::ProcessCompilationResults(const TArray<FShaderCommonCompileJobPtr>& InCompilationResults, int32& InOutJobIndex, float& TimeBudget)
@@ -2804,7 +2808,9 @@ void FMaterialShaderMap::ProcessCompilationResults(const TArray<FShaderCommonCom
 	do
 	{
 		auto& BaseJob = InCompilationResults[InOutJobIndex++];
+#if ALLOW_SHADERMAP_DEBUG_DATA
 		InFlightJobs.Remove(BaseJob->Hash);
+#endif
 
 		FShaderCompileJob* SingleJob = BaseJob->GetSingleShaderJob();
 		if (SingleJob)
