Several optimizations to reduce the critical path cost of GPU skin cache.
- Frontloaded setup work into a setup task that gets synced by the renderer during the dynamic mesh gather. Also removed external waits in favor of syncing when the user explicitly requests skin cache buffers from the Render Thread. The actual RHI dispatch and VF update work does not need to be synced by the render thread as it is all recorded into a command list. The dispatch work is issued in a chained background task that doesn't sync until RDG execution, which is much later.
- Optimized the transition logic so that it no longer needs to sort to remove duplicate transitions, and merged the UAV transitions together to reduce the number of loops.

diff --git a/Engine/Source/Runtime/Engine/Private/GPUSkinCache.cpp b/Engine/Source/Runtime/Engine/Private/GPUSkinCache.cpp
--- a/Engine/Source/Runtime/Engine/Private/GPUSkinCache.cpp
+++ b/Engine/Source/Runtime/Engine/Private/GPUSkinCache.cpp
@@ -27,6 +27,7 @@ GPUSkinCache.cpp: Performs skinning on a compute shader into a buffer to avoid v
 #include "UObject/UObjectIterator.h"
 #include "Algo/Sort.h"
 #include "ComponentRecreateRenderStateContext.h"
+#include "RenderGraphBuilder.h"
 
 DEFINE_STAT(STAT_GPUSkinCache_TotalNumChunks);
 DEFINE_STAT(STAT_GPUSkinCache_TotalNumVertices);
@@ -847,27 +848,6 @@ void FGPUSkinCache::Cleanup()
 	ensure(Allocations.Num() == 0);
 }
 
-void FGPUSkinCache::TransitionAllToReadable(FRHICommandList& RHICmdList, const TArray<FSkinCacheRWBuffer*>& BuffersToTransitionToRead)
-{
-	TRACE_CPUPROFILER_EVENT_SCOPE(FGPUSkinCache::TransitionAllToReadable);
-
-	if (BuffersToTransitionToRead.Num() > 0)
-	{
-		TArray<FRHITransitionInfo, SceneRenderingAllocator> UAVs;
-		UAVs.Reserve(BuffersToTransitionToRead.Num());
-		for (TArray<FSkinCacheRWBuffer*>::TConstIterator SetIt(BuffersToTransitionToRead); SetIt; ++SetIt)
-		{
-			FSkinCacheRWBuffer* Buffer = *SetIt;
-			constexpr ERHIAccess ToState = ERHIAccess::VertexOrIndexBuffer | ERHIAccess::SRVMask;
-			if (Buffer->AccessState != ToState)
-			{
-				UAVs.Add(Buffer->UpdateAccessState(ToState));
-			}
-		}
-		RHICmdList.Transition(UAVs, ERHIPipeline::Graphics, IsGPUSkinCacheRayTracingSupported() ? ERHIPipeline::All : ERHIPipeline::Graphics);
-	}
-}
-
 /** base of the FRecomputeTangentsPerTrianglePassCS class */
 class FBaseRecomputeTangentsPerTriangleShader : public FGlobalShader
 {
@@ -1265,73 +1245,177 @@ FGPUSkinCache::FRWBuffersAllocation* FGPUSkinCache::TryAllocBuffer(uint32 NumVer
 
 DECLARE_GPU_STAT(GPUSkinCache);
 
-void FGPUSkinCache::MakeBufferTransitions(FRHICommandList& RHICmdList, TArray<FSkinCacheRWBuffer*>& Buffers, ERHIAccess ToState)
+void FGPUSkinCache::TransitionAllToReadable(FRHICommandList& RHICmdList, TConstArrayView<FSkinCacheRWBuffer*> Buffers)
 {
-	if (Buffers.Num() > 0)
+	constexpr ERHIAccess ReadState = ERHIAccess::VertexOrIndexBuffer | ERHIAccess::SRVMask;
+
+	if (!Buffers.IsEmpty())
 	{
-		// The tangent accumulation buffers are shared between sections so they can end up in the list more than once. We
-		// need to make sure we don't issue multiple transitions for the same resource in a single call, the RHIs can't deal with that.
-		Algo::Sort(Buffers);
+		TRACE_CPUPROFILER_EVENT_SCOPE(TransitionAllToReadable);
+		const uint32 NextTransitionFence = GetNextTransitionFence();
 
 		TArray<FRHITransitionInfo, SceneRenderingAllocator> UAVs;
 		UAVs.Reserve(Buffers.Num());
-		
+		for (FSkinCacheRWBuffer* Buffer : Buffers)
+		{
+			if (!Buffer->UpdateFence(NextTransitionFence))
+			{
+				continue;
+			}
+
+			if (Buffer->AccessState != ReadState)
+			{
+				UAVs.Add(Buffer->UpdateAccessState(ReadState));
+			}
+		}
+		RHICmdList.Transition(UAVs, ERHIPipeline::Graphics, IsGPUSkinCacheRayTracingSupported() ? ERHIPipeline::All : ERHIPipeline::Graphics);
+	}
+}
+
+void FGPUSkinCache::TransitionBuffers(FRHICommandList& RHICmdList, TArray<FSkinCacheRWBuffer*>& Buffers, ERHIAccess ToState)
+{
+	if (!Buffers.IsEmpty())
+	{
+		TRACE_CPUPROFILER_EVENT_SCOPE(TransitionBuffers);
+		const uint32 NextTransitionFence = GetNextTransitionFence();
+
+		TArray<FRHITransitionInfo, SceneRenderingAllocator> UAVs;
+		UAVs.Reserve(Buffers.Num());
+
 		FSkinCacheRWBuffer* LastBuffer = nullptr;
 		for (FSkinCacheRWBuffer* Buffer : Buffers)
 		{
-			if (Buffer == LastBuffer)
+			if (!Buffer->UpdateFence(NextTransitionFence))
 			{
 				continue;
 			}
 
 			LastBuffer = Buffer;
-			if ( EnumHasAnyFlags(ToState, ERHIAccess::UAVMask) || Buffer->AccessState != ToState)
+			if (EnumHasAnyFlags(ToState, ERHIAccess::UAVMask) || Buffer->AccessState != ToState)
 			{
 				UAVs.Add(Buffer->UpdateAccessState(ToState));
 			}
 		}
+
+		// The NoFence flag is necessary to silence the validator for transitioning from All pipes to Graphics.
+		RHICmdList.Transition(MakeArrayView(UAVs.GetData(), UAVs.Num()), ERHITransitionCreateFlags::NoFence);
+	}
+}
+
+void FGPUSkinCache::TransitionBufferUAVs(FRHICommandList& RHICmdList, const TArray<FSkinCacheRWBuffer*>& Buffers, TArray<FRHIUnorderedAccessView*>& OutUAVs)
+{
+	if (!Buffers.IsEmpty())
+	{
+		TRACE_CPUPROFILER_EVENT_SCOPE(TransitionBufferUAVs);
+		const uint32 NextTransitionFence = GetNextTransitionFence();
+
+		TArray<FRHITransitionInfo, SceneRenderingAllocator> UAVs;
+		UAVs.Reserve(Buffers.Num());
+		OutUAVs.Reset(Buffers.Num());
+
+		for (FSkinCacheRWBuffer* Buffer : Buffers)
+		{
+			if (!Buffer->UpdateFence(NextTransitionFence))
+			{
+				continue;
+			}
+
+			UAVs.Add(Buffer->UpdateAccessState(ERHIAccess::UAVCompute));
+			OutUAVs.Add(Buffer->Buffer.UAV);
+		}
+
+		// The NoFence flag is necessary to silence the validator for transitioning from All pipes to Graphics.
 		RHICmdList.Transition(MakeArrayView(UAVs.GetData(), UAVs.Num()), ERHITransitionCreateFlags::NoFence);
 	}
 }
 
-void FGPUSkinCache::GetBufferUAVs(const TArray<FSkinCacheRWBuffer*>& InBuffers, TArray<FRHIUnorderedAccessView*>& OutUAVs)
+UE::Tasks::FTask FGPUSkinCache::DoDispatch(FRDGBuilder& GraphBuilder, const UE::Tasks::FTask& PrerequisitesTask)
 {
-	OutUAVs.Reset(InBuffers.Num());
+	check(!DispatchSetupTask.IsValid());
 
-	// It looks like BeginUAVOverlap wants to get in unique buffers, without any overlaps. Previously that worked out
-	// because InBuffers was being filled with AddUnique, but that can add quite a bit of cost when there are many
-	// skinned meshes to process here.
-	//
-	// Using a TSet instead of TArray with AddUnique improves things, but it still came up slower than just pushing
-	// everything and dealing with duplicates separately.
-	//
-	// So here I've added a member to FSkinCacheRWBuffer that we can just use to check if it's already been visited. On
-	// 64-bit, it's not increasing the struct size (since there's a trailing ERHIAccess/uint32 already, it gets padded
-	// up to an 8 byte multiple).
-	for (const FSkinCacheRWBuffer* Buffer : InBuffers)
+	if (!GraphBuilder.IsParallelSetupEnabled())
 	{
-		Buffer->UniqueOpToken = 0;
+		PrerequisitesTask.Wait();
+		DoDispatch(GraphBuilder.RHICmdList);
+		return {};
 	}
-	for (const FSkinCacheRWBuffer* Buffer : InBuffers)
+
+	FRHICommandList* RHICmdList = new FRHICommandList();
+	GraphBuilder.RHICmdList.QueueAsyncCommandListSubmit(RHICmdList);
+
+	FDispatchSetupContext* SetupContext = GraphBuilder.AllocObject<FDispatchSetupContext>();
+
+	DispatchSetupTask = GraphBuilder.AddSetupTask([this, RHICmdList, SetupContext]
 	{
-		if (Buffer->UniqueOpToken == 0)
+		RHICmdList->SwitchPipeline(ERHIPipeline::Graphics);
+		SetupDispatch(*RHICmdList, *SetupContext);
+
+	}, PrerequisitesTask, UE::Tasks::ETaskPriority::High);
+
+	GraphBuilder.AddSetupTask([this, RHICmdList, SetupContext]
+	{
+		FTaskTagScope TagScope(ETaskTag::EParallelRenderingThread);
+		DoDispatch(*RHICmdList, *SetupContext);
+		RHICmdList->FinishRecording();
+
+	}, DispatchSetupTask, UE::Tasks::ETaskPriority::BackgroundHigh);
+
+	GraphBuilder.AddPostExecuteCallback([this] { DispatchSetupTask = {}; });
+
+	return DispatchSetupTask;
+}
+
+void FGPUSkinCache::DoDispatch(FRHICommandList& RHICmdList)
+{
+	check(!DispatchSetupTask.IsValid());
+
+	if (!BatchDispatches.IsEmpty())
+	{
+		FGPUSkinCache::FDispatchSetupContext Context;
+		SetupDispatch(RHICmdList, Context);
+		DoDispatch(RHICmdList, Context);
+		BatchDispatches.Reset();
+	}
+}
+
+void FGPUSkinCache::SetupDispatch(FRHICommandList& RHICmdList, FDispatchSetupContext& Context)
+{
+	TRACE_CPUPROFILER_EVENT_SCOPE(GPUSkinCache_SetupDispatch);
+
+	Context.BuffersToTransition.Reserve(BatchDispatches.Num() * 2);
+
+	for (FDispatchEntry& DispatchItem : BatchDispatches)
+	{
+		PrepareUpdateSkinning(DispatchItem.SkinCacheEntry, DispatchItem.Section, DispatchItem.SkinCacheEntry->DispatchData[DispatchItem.Section].RevisionNumber, &Context.BuffersToTransition);
+
+		// Clear the flag that this is queued for dispatch.
+		DispatchItem.SkinCacheEntry->bQueuedForDispatch = false;
+		DispatchItem.SkinCacheEntry->DispatchData[DispatchItem.Section].RevisionNumber = 0;
+	}
+
+#if RHI_RAYTRACING
+	if (IsGPUSkinCacheRayTracingSupported())
+	{
+		for (FGPUSkinCacheEntry* SkinCacheEntry : PendingProcessRTGeometryEntries)
 		{
-			Buffer->UniqueOpToken = 1;
-			OutUAVs.Add(Buffer->Buffer.UAV);
+			ProcessRayTracingGeometryToUpdate(RHICmdList, SkinCacheEntry);
 		}
+
+		PendingProcessRTGeometryEntries.Reset();
 	}
+#endif
 }
 
-void FGPUSkinCache::DoDispatch(FRHICommandList& RHICmdList)
+void FGPUSkinCache::DoDispatch(FRHICommandList& RHICmdList, FDispatchSetupContext& Context)
 {
 	int32 BatchCount = BatchDispatches.Num();
-	INC_DWORD_STAT_BY(STAT_GPUSkinCache_TotalNumChunks, BatchCount);
 
 	if (!BatchCount)
 	{
 		return;
 	}
 
+	INC_DWORD_STAT_BY(STAT_GPUSkinCache_TotalNumChunks, BatchCount);
 	TRACE_CPUPROFILER_EVENT_SCOPE(FGPUSkinCache::DoDispatch);
 
 	bool bCapture = BatchCount > 0 && GNumDispatchesToCapture > 0;
@@ -1341,31 +1425,10 @@ void FGPUSkinCache::DoDispatch(FRHICommandList& RHICmdList)
 	RHI_BREADCRUMB_EVENT_STAT(RHICmdList, GPUSkinCache, "GPUSkinCache");
 	SCOPED_GPU_STAT(RHICmdList, GPUSkinCache);
 
-	TArray<FSkinCacheRWBuffer*> BuffersToTransitionForSkinning;
-	BuffersToTransitionForSkinning.Reserve(BatchCount * 2);
-	{
-		TRACE_CPUPROFILER_EVENT_SCOPE(GPUSkinCache_PrepareUpdateSkinning);
-
-		for (int32 i = 0; i < BatchCount; ++i)
-		{
-			FDispatchEntry& DispatchItem = BatchDispatches[i];
-			PrepareUpdateSkinning(DispatchItem.SkinCacheEntry, DispatchItem.Section, DispatchItem.SkinCacheEntry->DispatchData[DispatchItem.Section].RevisionNumber, &BuffersToTransitionForSkinning);
-
-			// Clear the flag that this is queued for dispatch.
-			DispatchItem.SkinCacheEntry->bQueuedForDispatch = false;
-			DispatchItem.SkinCacheEntry->DispatchData[DispatchItem.Section].RevisionNumber = 0;
-		}
-
-		{
-			TRACE_CPUPROFILER_EVENT_SCOPE(MakeBufferTransitions);
-			MakeBufferTransitions(RHICmdList, BuffersToTransitionForSkinning, ERHIAccess::UAVCompute);
-		}
-	}
-
 	TArray<FSkinCacheRWBuffer*> BuffersToTransitionToRead;
 
 	TArray<FRHIUnorderedAccessView*> SkinningBuffersToOverlap;
-	GetBufferUAVs(BuffersToTransitionForSkinning, SkinningBuffersToOverlap);
+	TransitionBufferUAVs(RHICmdList, Context.BuffersToTransition, SkinningBuffersToOverlap);
 	RHICmdList.BeginUAVOverlap(SkinningBuffersToOverlap);
 	{
 		TRACE_CPUPROFILER_EVENT_SCOPE(GPUSkinCache_UpdateSkinningBatches);
@@ -1548,13 +1611,12 @@ void FGPUSkinCache::DoDispatch(FRHICommandList& RHICmdList)
 				IntermediateAccumulatedTangentBuffers.Add(DispatchData.GetIntermediateAccumulatedTangentBuffer());
 			}
 			BuffersToTransitionToRead.Add(DispatchData.GetPositionRWBuffer());
-		}	
+		}
 	}
-	MakeBufferTransitions(RHICmdList, BuffersToSRVForRecomputeTangents, ERHIAccess::SRVCompute);
-	MakeBufferTransitions(RHICmdList, IntermediateAccumulatedTangentBuffers, ERHIAccess::UAVCompute);
+	TransitionBuffers(RHICmdList, BuffersToSRVForRecomputeTangents, ERHIAccess::SRVCompute);
 
 	TArray<FRHIUnorderedAccessView*> IntermediateAccumulatedTangentBuffersToOverlap;
-	GetBufferUAVs(IntermediateAccumulatedTangentBuffers, IntermediateAccumulatedTangentBuffersToOverlap);
+	TransitionBufferUAVs(RHICmdList, IntermediateAccumulatedTangentBuffers, IntermediateAccumulatedTangentBuffersToOverlap);
 	RHICmdList.BeginUAVOverlap(IntermediateAccumulatedTangentBuffersToOverlap);
 	{
 		SCOPED_DRAW_EVENT(RHICmdList, GPUSkinCache_RecomputeTangentsBatches);
@@ -1583,11 +1645,10 @@ void FGPUSkinCache::DoDispatch(FRHICommandList& RHICmdList)
 				TangentBuffers.Add(DispatchData.GetTangentRWBuffer());
 				BuffersToTransitionToRead.Add(DispatchData.GetTangentRWBuffer());
 			}
-			MakeBufferTransitions(RHICmdList, TangentBuffers, ERHIAccess::UAVCompute);
-			MakeBufferTransitions(RHICmdList, IntermediateAccumulatedTangentBuffers, ERHIAccess::UAVCompute);
+			TransitionBuffers(RHICmdList, IntermediateAccumulatedTangentBuffers, ERHIAccess::UAVCompute);
 		
 			TArray<FRHIUnorderedAccessView*> TangentBuffersToOverlap;
-			GetBufferUAVs(TangentBuffers, TangentBuffersToOverlap);
+			TransitionBufferUAVs(RHICmdList, TangentBuffers, TangentBuffersToOverlap);
 			RHICmdList.BeginUAVOverlap(TangentBuffersToOverlap);
 			for (int32 i = 0; i < BatchCount; ++i)
 			{
@@ -1617,18 +1678,6 @@ void FGPUSkinCache::DoDispatch(FRHICommandList& RHICmdList)
 		TransitionAllToReadable(RHICmdList, BuffersToTransitionToRead);
 	}
 
-#if RHI_RAYTRACING
-	if (IsGPUSkinCacheRayTracingSupported())
-	{
-		for (FGPUSkinCacheEntry* SkinCacheEntry : PendingProcessRTGeometryEntries)
-		{
-			ProcessRayTracingGeometryToUpdate(RHICmdList, SkinCacheEntry);
-		}
-
-		PendingProcessRTGeometryEntries.Reset();
-	}
-#endif
-
 	BatchDispatches.Reset();
 }
 
@@ -2011,7 +2060,7 @@ void FGPUSkinCache::PrepareUpdateSkinning(FGPUSkinCacheEntry* Entry, int32 Secti
 	check(DispatchData.PreviousPositionBuffer != DispatchData.PositionBuffer);
 }
 
-void FGPUSkinCache::DispatchUpdateSkinning(FRHICommandList& RHICmdList, FGPUSkinCacheEntry* Entry, int32 Section, uint32 RevisionNumber, TArray<FSkinCacheRWBuffer*>& BuffersToTransitionToRead)
+void FGPUSkinCache::DispatchUpdateSkinning(FRHICommandList& RHICmdList, FGPUSkinCacheEntry* Entry, int32 Section, TArray<FSkinCacheRWBuffer*>& BuffersToTransitionToRead)
 {
 	FGPUSkinCacheEntry::FSectionDispatchData& DispatchData = Entry->DispatchData[Section];
 	FGPUBaseSkinVertexFactory::FShaderDataType& ShaderData = DispatchData.SourceVertexFactory->GetShaderData();
@@ -2196,6 +2245,7 @@ FRWBuffer* FGPUSkinCache::GetPositionBuffer(FGPUSkinCacheEntry const* Entry, uin
 {
 	if (Entry)
 	{
+		Entry->SkinCache->WaitForDispatchSetupTask();
 		FGPUSkinCacheEntry::FSectionDispatchData const& DispatchData = Entry->GetDispatchData()[SectionIndex];
 		FSkinCacheRWBuffer* SkinCacheRWBuffer = DispatchData.PositionBuffer;
 		return SkinCacheRWBuffer != nullptr ? &SkinCacheRWBuffer->Buffer : nullptr;
@@ -2207,6 +2257,7 @@ FRWBuffer* FGPUSkinCache::GetPreviousPositionBuffer(FGPUSkinCacheEntry const* En
 {
 	if (Entry)
 	{
+		Entry->SkinCache->WaitForDispatchSetupTask();
 		FGPUSkinCacheEntry::FSectionDispatchData const& DispatchData = Entry->GetDispatchData()[SectionIndex];
 		FSkinCacheRWBuffer* SkinCacheRWBuffer = DispatchData.PreviousPositionBuffer;
 		return SkinCacheRWBuffer != nullptr ? &SkinCacheRWBuffer->Buffer : nullptr;
@@ -2218,6 +2269,7 @@ FRWBuffer* FGPUSkinCache::GetTangentBuffer(FGPUSkinCacheEntry const* Entry, uint
 {
 	if (Entry)
 	{
+		Entry->SkinCache->WaitForDispatchSetupTask();
 		FGPUSkinCacheEntry::FSectionDispatchData const& DispatchData = Entry->GetDispatchData()[SectionIndex];
 		FSkinCacheRWBuffer* SkinCacheRWBuffer = DispatchData.TangentBuffer;
 		return SkinCacheRWBuffer != nullptr ? &SkinCacheRWBuffer->Buffer : nullptr;
@@ -2225,11 +2277,6 @@ FRWBuffer* FGPUSkinCache::GetTangentBuffer(FGPUSkinCacheEntry const* Entry, uint
 	return nullptr;
 }
 
-uint32 FGPUSkinCache::GetUpdatedFrame(FGPUSkinCacheEntry const* Entry, uint32 SectionIndex)
-{
-	return Entry != nullptr ? Entry->GetDispatchData()[SectionIndex].UpdatedFrameNumber : 0;
-}
-
 void FGPUSkinCache::UpdateSkinWeightBuffer(FGPUSkinCacheEntry* Entry)
 {
 	if (Entry)
diff --git a/Engine/Source/Runtime/Engine/Public/GPUSkinCache.h b/Engine/Source/Runtime/Engine/Public/GPUSkinCache.h
--- a/Engine/Source/Runtime/Engine/Public/GPUSkinCache.h
+++ b/Engine/Source/Runtime/Engine/Public/GPUSkinCache.h
@@ -176,8 +176,7 @@ public:
 	{
 		FRWBuffer	Buffer;
 		ERHIAccess	AccessState = ERHIAccess::Unknown;	// Keep track of current access state
-		// See note in FGPUSkinCache::GetBufferUAVs()
-		mutable uint32	UniqueOpToken = 0;
+		mutable uint32	LastTransitionFence = 0;
 
 		void Release()
 		{
@@ -185,6 +184,13 @@ public:
 			AccessState = ERHIAccess::Unknown;
 		}
 
+		bool UpdateFence(uint32 NextTransitionFence)
+		{
+			const bool bUpdateRequired = LastTransitionFence != NextTransitionFence;
+			LastTransitionFence = NextTransitionFence;
+			return bUpdateRequired;
+		}
+
 		// Update the access state and return transition info
 		FRHITransitionInfo UpdateAccessState(ERHIAccess NewState)
 		{
@@ -333,25 +339,44 @@ public:
 	static FRWBuffer* GetPositionBuffer(FGPUSkinCacheEntry const* Entry, uint32 SectionIndex);
 	static FRWBuffer* GetPreviousPositionBuffer(FGPUSkinCacheEntry const* Entry, uint32 SectionIndex);
 	static FRWBuffer* GetTangentBuffer(FGPUSkinCacheEntry const* Entry, uint32 SectionIndex);
-	static uint32 GetUpdatedFrame(FGPUSkinCacheEntry const* Entry, uint32 SectionIndex);
 
-	// Deprecated function. Can remove include of CachedGeometry.h when this is removed.
-	UE_DEPRECATED(5.1, "Use GetPositionBuffer() or similar instead.")
-	FCachedGeometry::Section GetCachedGeometry(FGPUSkinCacheEntry* InOutEntry, uint32 SectionId);
+	UE_DEPRECATED(5.6, "GetUpdatedFrame is no longer used")
+	static uint32 GetUpdatedFrame(FGPUSkinCacheEntry const*, uint32) { return 0; }
 
 #if RHI_RAYTRACING
 	void ProcessRayTracingGeometryToUpdate(FRHICommandList& RHICmdList, FGPUSkinCacheEntry* SkinCacheEntry);
 #endif // RHI_RAYTRACING
 
-	void ENGINE_API DoDispatch(FRHICommandList& RHICmdList);
+	ENGINE_API UE::Tasks::FTask DoDispatch(FRDGBuilder& GraphBuilder, const UE::Tasks::FTask& PrerequisitesTask = {});
+	ENGINE_API void DoDispatch(FRHICommandList& RHICmdList);
 
 	inline ERHIFeatureLevel::Type GetFeatureLevel() const { return FeatureLevel; }
 
 	inline bool HasWork() const { return !BatchDispatches.IsEmpty(); }
 
 protected:
-	void MakeBufferTransitions(FRHICommandList& RHICmdList, TArray<FSkinCacheRWBuffer*>& Buffers, ERHIAccess ToState);
-	void GetBufferUAVs(const TArray<FSkinCacheRWBuffer*>& InBuffers, TArray<FRHIUnorderedAccessView*>& OutUAVs);
+	struct FDispatchSetupContext
+	{
+		TArray<FSkinCacheRWBuffer*> BuffersToTransition;
+	};
+
+	void SetupDispatch(FRHICommandList& RHICmdList, FDispatchSetupContext& SetupContext);
+	void DoDispatch(FRHICommandList& RHICmdList, FDispatchSetupContext& SetupContext);
+
+	void WaitForDispatchSetupTask()
+	{
+		check(IsInRenderingThread());
+		if (DispatchSetupTask.IsValid())
+		{
+			TRACE_CPUPROFILER_EVENT_SCOPE(WaitForGPUSkinCache);
+			DispatchSetupTask.Wait();
+			DispatchSetupTask = {};
+		}
+	}
+
+	void TransitionBuffers(FRHICommandList& RHICmdList, TArray<FSkinCacheRWBuffer*>& Buffers, ERHIAccess ToState);
+	void TransitionBufferUAVs(FRHICommandList& RHICmdList, const TArray<FSkinCacheRWBuffer*>& InBuffers, TArray<FRHIUnorderedAccessView*>& OutUAVs);
+	void TransitionAllToReadable(FRHICommandList& RHICmdList, TConstArrayView<FSkinCacheRWBuffer*> BuffersToTransitionToRead);
 
 	TArray<FRWBuffersAllocation*> Allocations;
 	TArray<FGPUSkinCacheEntry*> Entries;
@@ -371,13 +396,11 @@ protected:
 	void DispatchUpdateSkinning(
 		FRHICommandList& RHICmdList, 
 		FGPUSkinCacheEntry* Entry, 
-		int32 Section, 
-		uint32 RevisionNumber,
+		int32 Section,
 		TArray<FSkinCacheRWBuffer*>& BuffersToTransitionToRead
 		);
 
 	void Cleanup();
-	static void TransitionAllToReadable(FRHICommandList& RHICmdList, const TArray<FSkinCacheRWBuffer*>& BuffersToTransitionToRead);
 	static void ReleaseSkinCacheEntry(FGPUSkinCacheEntry* SkinCacheEntry);
 	void InvalidateAllEntries();
 	uint64 UsedMemoryInBytes;
@@ -392,12 +415,17 @@ protected:
 	ERHIFeatureLevel::Type FeatureLevel;
 	UWorld* World;
 
+	UE::Tasks::FTask DispatchSetupTask;
+
 	static void CVarSinkFunction();
 	static FAutoConsoleVariableSink CVarSink;
 
 	bool IncrementDispatchCounter(FRHICommandList& RHICmdList);
 	int32 DispatchCounter = 0;
 
+	uint32 GetNextTransitionFence() { return ++TransitionFence; }
+	uint32 TransitionFence = 0;
+
 	void PrintMemorySummary() const;
 	FString GetSkeletalMeshObjectName(const FSkeletalMeshObject* GPUSkin) const;
 	FDebugName GetSkeletalMeshObjectDebugName(const FSkeletalMeshObject* GPUSkin) const;
diff --git a/Engine/Source/Runtime/Renderer/Private/RendererScene.cpp b/Engine/Source/Runtime/Renderer/Private/RendererScene.cpp
--- a/Engine/Source/Runtime/Renderer/Private/RendererScene.cpp
+++ b/Engine/Source/Runtime/Renderer/Private/RendererScene.cpp
@@ -5558,10 +5558,7 @@ void FScene::Update(FRDGBuilder& GraphBuilder, const FUpdateParameters& Paramete
 
 	if (GPUSkinCache && GPUSkinCache->HasWork())
 	{
-		GPUSkinCacheTask = GraphBuilder.AddCommandListSetupTask([this] (FRHICommandList& RHICmdList)
-		{
-			GPUSkinCache->DoDispatch(RHICmdList);
-		});
+		GPUSkinCacheTask = GPUSkinCache->DoDispatch(GraphBuilder);
 	}
 
 	// Add interface to blackboard to receive scene updates from compute.
diff --git a/Engine/Source/Runtime/Renderer/Private/ScenePrivate.h b/Engine/Source/Runtime/Renderer/Private/ScenePrivate.h
--- a/Engine/Source/Runtime/Renderer/Private/ScenePrivate.h
+++ b/Engine/Source/Runtime/Renderer/Private/ScenePrivate.h
@@ -3767,12 +3767,6 @@ public:
 		return CreateLightPrimitiveInteractionsTask;
 	}
 
-	void WaitForGPUSkinCacheTask()
-	{
-		CSV_SCOPED_SET_WAIT_STAT(GPUSkinCache);
-		GPUSkinCacheTask.Wait();
-	}
-
 	UE::Tasks::FTask GetGPUSkinCacheTask() const
 	{
 		return GPUSkinCacheTask;
diff --git a/Engine/Source/Runtime/Renderer/Private/SceneRendering.cpp b/Engine/Source/Runtime/Renderer/Private/SceneRendering.cpp
--- a/Engine/Source/Runtime/Renderer/Private/SceneRendering.cpp
+++ b/Engine/Source/Runtime/Renderer/Private/SceneRendering.cpp
@@ -3920,8 +3920,6 @@ IVisibilityTaskData* FSceneRenderer::OnRenderBegin(FRDGBuilder& GraphBuilder, co
 				CreateHairStrandsBookmarkParameters(Scene, Views, AllViews, Parameters, false/*bComputeVisibleInstances*/);
 				if (Parameters.HasInstances())
 				{
-					Scene->WaitForGPUSkinCacheTask();
-	
 					// 1. Select appropriate LOD & geometry type
 					RunHairStrandsBookmark(GraphBuilder, EHairStrandsBookmark::ProcessLODSelection, Parameters);
 				}
diff --git a/Engine/Source/Runtime/Renderer/Private/SceneVisibility.cpp b/Engine/Source/Runtime/Renderer/Private/SceneVisibility.cpp
--- a/Engine/Source/Runtime/Renderer/Private/SceneVisibility.cpp
+++ b/Engine/Source/Runtime/Renderer/Private/SceneVisibility.cpp
@@ -4747,7 +4747,6 @@ void FSceneRenderer::PreVisibilityFrameSetup(FRDGBuilder& GraphBuilder)
 				{
 					// 1.Update binding surfaces
 					// Prepare surface data (MeshLODData)
-					Scene->WaitForGPUSkinCacheTask();
 					RunHairStrandsBookmark(GraphBuilder, EHairStrandsBookmark::ProcessBindingSurfaceUpdate, Parameters);
 
 					// 2. Prepare surface data for guides
