- Extends the EngineTest CoreAssetLoading tests to allow a blueprint "LoadValidator" to be executed to confirm the load was performed correctly. The existing test would simply load and unload the object however for UE-236998 the issue is the load occurrs but the wrong data is shown. Thus we need a better mechanism for validating loads in these tests which is now possible. Care must be taken to ensure no hard references are made when validating since the test will unload the object and a hard reference will prevent unloading causing test failures
- Adds testcase for UE-236998 to ensure that in cooked zenloader builds not using the zenstore we are properly preloading InheritableComponentHandlers when loading to ensure we override our parent components appropriately.
- Changes BlueprintGeneratedClasses to always preload InheritableComponentHandles if they haven't been preloaded yet. The legacy loader handle this through load type sorting which zenloader doesn't do completely. The uncooked zenloader editor preloads but when cooked this preload path isn't taken. Fetching the archetype is expected to grab fully preloaded archetypes in Linkerload so this change upholds that guarantee.

#tests Added new Tests.Tests.Core.CoreAssetLoading.LoadChild_InheritableComponentHandlerTest for the bug in question being fixed
#rnx

#virtualized

diff --git a/Engine/Source/Runtime/Engine/Private/BlueprintGeneratedClass.cpp b/Engine/Source/Runtime/Engine/Private/BlueprintGeneratedClass.cpp
--- a/Engine/Source/Runtime/Engine/Private/BlueprintGeneratedClass.cpp
+++ b/Engine/Source/Runtime/Engine/Private/BlueprintGeneratedClass.cpp
@@ -1492,13 +1492,20 @@ UObject* UBlueprintGeneratedClass::FindArchetype(const UClass* ArchetypeClass, c
 			}
 			else if (UInheritableComponentHandler* ICH = Class->GetInheritableComponentHandler())
 			{
-				if (GEventDrivenLoaderEnabled && EVENT_DRIVEN_ASYNC_LOAD_ACTIVE_AT_RUNTIME)
+				if (ICH->HasAnyFlags(RF_NeedLoad))
 				{
-					if (ICH->HasAnyFlags(RF_NeedLoad))
+					if (GEventDrivenLoaderEnabled && EVENT_DRIVEN_ASYNC_LOAD_ACTIVE_AT_RUNTIME)
 					{
 						UE_LOG(LogClass, Fatal, TEXT("%s had RF_NeedLoad when searching for an archetype of %s named %s"), *GetFullNameSafe(ICH), *GetFullNameSafe(ArchetypeClass), *ArchetypeName.ToString());
 					}
+					else
+					{
+						// EDL will ensure we have preloaded all subobjects for our class before loading other exports use this archetype however that isn't guaranteed with
+						// the other loader paths so we must preload here to avoid creating exports with an uninitialized archetype.
+						ICH->PreloadAll();
+					}
 				}
+
 				// This would find either an SCS component template override (for which the archetype
 				// name will match the SCS variable name), or an old AddComponent node template override
 				// (for which the archetype name will match the override record's component template name).
