Expose per static mesh component MeshPaint overrides to Blueprint.

diff --git a/Engine/Source/Runtime/Engine/Classes/Components/StaticMeshComponent.h b/Engine/Source/Runtime/Engine/Classes/Components/StaticMeshComponent.h
--- a/Engine/Source/Runtime/Engine/Classes/Components/StaticMeshComponent.h
+++ b/Engine/Source/Runtime/Engine/Classes/Components/StaticMeshComponent.h
@@ -333,19 +333,19 @@ public:
 	uint8 bEnableTextureColorMeshPainting : 1;
 
 	/** Whether to override the MeshPaintTextureCoordinateIndex set on the static mesh. */
-	UPROPERTY(EditAnywhere, Category = "Mesh Painting", meta=(InlineEditConditionToggle))
+	UPROPERTY(EditAnywhere, BlueprintReadWrite, Category = "Mesh Painting", meta=(InlineEditConditionToggle))
 	uint8 bOverrideMeshPaintTextureCoordinateIndex : 1;
 
 	/** Whether to override the MeshPaintTextureCoordinateIndex set on the static mesh. */
-	UPROPERTY(EditAnywhere, Category = "Mesh Painting", meta=(InlineEditConditionToggle))
+	UPROPERTY(EditAnywhere, BlueprintReadWrite, Category = "Mesh Painting", meta=(InlineEditConditionToggle))
 	uint8 bOverrideMeshPaintTextureResolution : 1;
 
 	/** The overriden coordinate index to use when texture color painting on this component. */
-	UPROPERTY(EditAnywhere, AdvancedDisplay, Category = "Mesh Painting", meta=(UIMin = "0", UIMax = "3", editcondition = "bOverrideMeshPaintTextureCoordinateIndex"))
+	UPROPERTY(EditAnywhere, BlueprintReadWrite, AdvancedDisplay, Category = "Mesh Painting", meta=(UIMin = "0", UIMax = "3", editcondition = "bOverrideMeshPaintTextureCoordinateIndex"))
 	int32 OverriddenMeshPaintTextureCoordinateIndex;
 
 	/** The overriden resolution of texture color mesh paint textures on this component. */
-	UPROPERTY(EditAnywhere, AdvancedDisplay, Category = "Mesh Painting", meta = (UIMin = "0", UIMax = "4096", ClampMax = "4096", editcondition = "bOverrideMeshPaintTextureResolution"))
+	UPROPERTY(EditAnywhere, BlueprintReadWrite, AdvancedDisplay, Category = "Mesh Painting", meta = (UIMin = "0", UIMax = "4096", ClampMax = "4096", editcondition = "bOverrideMeshPaintTextureResolution"))
 	int32 OverriddenMeshPaintTextureResolution;
 
 	/** Light map resolution to use on this component, used if bOverrideLightMapRes is true and there is a valid StaticMesh. */
