UE-239744 - Bitfield enums don't cast correctly in Blueprints

Casting bytes to an enum wasn't working correctly if the enum was a bitfield and the byte value was a combination of valid flags.

diff --git a/Engine/Source/Runtime/Engine/Private/KismetNodeHelperLibrary.cpp b/Engine/Source/Runtime/Engine/Private/KismetNodeHelperLibrary.cpp
--- a/Engine/Source/Runtime/Engine/Private/KismetNodeHelperLibrary.cpp
+++ b/Engine/Source/Runtime/Engine/Private/KismetNodeHelperLibrary.cpp
@@ -153,7 +153,7 @@ FString UKismetNodeHelperLibrary::GetEnumeratorUserFriendlyName(const UEnum* Enu
 {
 	if (nullptr != Enum)
 	{
-		return Enum->GetDisplayNameTextByValue(EnumeratorValue).ToString();
+		return Enum->GetValueOrBitfieldAsDisplayNameText(EnumeratorValue).ToString();
 	}
 
 	return FName().ToString();
@@ -163,7 +163,7 @@ uint8 UKismetNodeHelperLibrary::GetValidValue(const UEnum* Enum, uint8 Enumerato
 {
 	if (ensureAsRuntimeWarning(Enum != nullptr))
 	{
-		if (Enum->IsValidEnumValue(EnumeratorValue))
+		if (Enum->IsValidEnumValueOrBitfield(EnumeratorValue))
 		{
 			return EnumeratorValue;
 		}
