Added a test to EngineTest that validates the behavior of setting AActor::NetCullDistanceSquared on relevancy checks in legacy and Iris replication systems.

diff --git a/Engine/Source/Runtime/Engine/Private/Tests/NetTestHelpers.cpp b/Engine/Source/Runtime/Engine/Private/Tests/NetTestHelpers.cpp
--- a/Engine/Source/Runtime/Engine/Private/Tests/NetTestHelpers.cpp
+++ b/Engine/Source/Runtime/Engine/Private/Tests/NetTestHelpers.cpp
@@ -33,8 +33,14 @@
 namespace UE::Net
 {
 
-FTestWorldInstance FTestWorldInstance::CreateServer(const TCHAR* InURL)
+FTestWorldInstance FTestWorldInstance::CreateServer(const TCHAR* InURL, const EReplicationSystem ReplicationSystem)
 {
+	const bool bPrevShouldUseIrisReplication = ShouldUseIrisReplication();
+	if (ReplicationSystem == EReplicationSystem::Iris)
+	{
+		SetUseIrisReplication(true);
+	}
+
 	FGameInstancePIEParameters ServerParams;
 	ServerParams.bSimulateInEditor = false;
 	ServerParams.bAnyBlueprintErrors = false;
@@ -52,6 +58,8 @@ FTestWorldInstance FTestWorldInstance::CreateServer(const TCHAR* InURL)
 	FString BrowseError;
 	GEngine->Browse(*NewInstance.GameInstance->GetWorldContext(), LocalURL, BrowseError);
 
+	SetUseIrisReplication(bPrevShouldUseIrisReplication);
+
 	return NewInstance;
 }
 
@@ -264,30 +272,30 @@ void FTestWorldInstance::UnloadStreamingLevel(FName LevelName)
 // FTestWorlds
 //============================================================================
 
-FTestWorlds::FTestWorlds()
+FTestWorlds::FTestWorlds(const EReplicationSystem ReplicationSystem)
 	: Server(true)
 {
 	InitDelegates();
 
-	Server = FTestWorldInstance::CreateServer(TEXT("/Engine/Maps/Entry?listen?Game=/Script/Engine.GameMode"));
+	Server = FTestWorldInstance::CreateServer(TEXT("/Engine/Maps/Entry?listen?Game=/Script/Engine.GameMode"), ReplicationSystem);
 }
 
-FTestWorlds::FTestWorlds(const FString& MapName, const FString& GameModeName)
+FTestWorlds::FTestWorlds(const FString& MapName, const FString& GameModeName, const EReplicationSystem ReplicationSystem)
 	: Server(true)
 {
 	InitDelegates();
 
 	FString URL = FString::Printf(TEXT("%s?listen?Game=%s"), *MapName, *GameModeName);
 
-	Server = FTestWorldInstance::CreateServer(*URL);
+	Server = FTestWorldInstance::CreateServer(*URL, ReplicationSystem);
 }
 
-FTestWorlds::FTestWorlds(const TCHAR* ServerURL)
+FTestWorlds::FTestWorlds(const TCHAR* ServerURL, const EReplicationSystem ReplicationSystem)
 	: Server(true)
 {
 	InitDelegates();
 
-	Server = FTestWorldInstance::CreateServer(ServerURL);
+	Server = FTestWorldInstance::CreateServer(ServerURL, ReplicationSystem);
 }
 
 FTestWorlds::~FTestWorlds()
diff --git a/Engine/Source/Runtime/Engine/Public/Tests/NetTestHelpers.h b/Engine/Source/Runtime/Engine/Public/Tests/NetTestHelpers.h
--- a/Engine/Source/Runtime/Engine/Public/Tests/NetTestHelpers.h
+++ b/Engine/Source/Runtime/Engine/Public/Tests/NetTestHelpers.h
@@ -94,7 +94,7 @@ struct FTestWorldInstance
 	struct FContext;
 
 public:
-	ENGINE_API static FTestWorldInstance CreateServer(const TCHAR* InURL);
+	ENGINE_API static FTestWorldInstance CreateServer(const TCHAR* InURL, const EReplicationSystem ReplicationSystem = EReplicationSystem::Generic);
 	ENGINE_API static FTestWorldInstance CreateClient(int32 ServerPort);
 	ENGINE_API static FTestWorldInstance CreateProxy();
 
@@ -155,17 +155,17 @@ private:
 struct FTestWorlds
 {
 	/** Creates a server world using the default map (Entry) and default GameMode class (GameMode) */
-	ENGINE_API FTestWorlds();
+	ENGINE_API FTestWorlds(const EReplicationSystem ReplicationSystem = EReplicationSystem::Generic);
 
 	/** 
 	 * Creates a server world using a specific map name and GameMode class
 	 * @param MapName Full path name of the map. ex: /Engine/Maps/Entry
 	 * @param GameModeName Full path name of the game mode. ex: /Script/EngineTest.MyGameMode
 	 */
-	ENGINE_API FTestWorlds(const FString& MapName, const FString& GameModeName);
+	ENGINE_API FTestWorlds(const FString& MapName, const FString& GameModeName, const EReplicationSystem ReplicationSystem = EReplicationSystem::Generic);
 
 	/** Creates a server world using the given URL. */
-	ENGINE_API explicit FTestWorlds(const TCHAR* ServerURL);
+	ENGINE_API explicit FTestWorlds(const TCHAR* ServerURL, const EReplicationSystem ReplicationSystem = EReplicationSystem::Generic);
 
 	ENGINE_API ~FTestWorlds();
 
