speculative fix for haptics soundwave crash
#swarm https://p4-swarm.epicgames.net/reviews/39458845

diff --git a/Engine/Source/Runtime/Engine/Private/AudioDecompress.cpp b/Engine/Source/Runtime/Engine/Private/AudioDecompress.cpp
--- a/Engine/Source/Runtime/Engine/Private/AudioDecompress.cpp
+++ b/Engine/Source/Runtime/Engine/Private/AudioDecompress.cpp
@@ -782,7 +782,6 @@ void FAsyncAudioDecompressWorker::DoWork()
 			else
 			{
 				LLM_SCOPE(ELLMTag::AudioFullDecompress);
-				check(Wave->DecompressionType == DTYPE_Procedural);
 
 				Wave->RawPCMDataSize = QualityInfo.SampleDataSize;
 				check(Wave->RawPCMData == nullptr);
diff --git a/Engine/Source/Runtime/Engine/Private/AudioDevice.cpp b/Engine/Source/Runtime/Engine/Private/AudioDevice.cpp
--- a/Engine/Source/Runtime/Engine/Private/AudioDevice.cpp
+++ b/Engine/Source/Runtime/Engine/Private/AudioDevice.cpp
@@ -6757,6 +6757,9 @@ void FAudioDevice::Precache(USoundWave* SoundWave, bool bSynchronous, bool bTrac
 
 			return;
 		}
+
+
+		
 	}
 
 	// calculate the decompression type
@@ -6766,11 +6769,6 @@ void FAudioDevice::Precache(USoundWave* SoundWave, bool bSynchronous, bool bTrac
 		// No channels - no way of knowing what to play back
 		SoundWave->DecompressionType = DTYPE_Invalid;
 	}
-	else if (SoundWave->RawPCMData)
-	{
-		// DTYPE_Preview has been removed based on the assumption that this path is unreachable
-		checkNoEntry();
-	}
 	else if (SoundWave->bProcedural)
 	{
 		// Procedurally created audio
@@ -6796,6 +6794,7 @@ void FAudioDevice::Precache(USoundWave* SoundWave, bool bSynchronous, bool bTrac
 
 		static FName NAME_OGG(TEXT("OGG"));
 		SoundWave->bDecompressedFromOgg = SoundWave->GetRuntimeFormat() == NAME_OGG;
+		bool bDecompressToPCM = false;
 
 		// handle audio decompression
 		if (FPlatformProperties::SupportsAudioStreaming() && SoundWave->IsStreaming(nullptr))
@@ -6803,21 +6802,28 @@ void FAudioDevice::Precache(USoundWave* SoundWave, bool bSynchronous, bool bTrac
 			SoundWave->DecompressionType = DTYPE_Streaming;
 			SoundWave->bCanProcessAsync = false;
 		}
-		else
+		else if (ShouldUseRealtimeDecompression(bForceFullDecompression, SoundGroup, SoundWave, CompressedDurationThreshold)) 
 		{
-			check(ShouldUseRealtimeDecompression(bForceFullDecompression, SoundGroup, SoundWave, CompressedDurationThreshold));
-			
 			// Store as compressed data and decompress in realtime
 			SoundWave->DecompressionType = DTYPE_RealTime;
 #if !(UE_BUILD_SHIPPING || UE_BUILD_TEST)
 			++PrecachedRealtime;
 #endif // !(UE_BUILD_SHIPPING || UE_BUILD_TEST)
 		}
+		else
+		{
+			// Fully expand loaded audio data into PCM
+			bDecompressToPCM = true;
+
+			// if the decompression type is already DTYPE_RealTime from a previous call
+			// the task will early exit and we won't decompress
+			SoundWave->DecompressionType = DTYPE_Setup;
+		}
 
 		// Grab the compressed audio data
 		SoundWave->InitAudioResource(SoundWave->GetRuntimeFormat());
 
-		if (SoundWave->AudioDecompressor == nullptr && SoundWave->DecompressionType == DTYPE_RealTime)
+		if (SoundWave->AudioDecompressor == nullptr && (bDecompressToPCM || SoundWave->DecompressionType == DTYPE_RealTime))
 		{
 			// Create a worker to decompress the audio data
 			if (bSynchronous)
