Made FIREmitter::EmitSetMaterialOutput in the new translator allocate the instruction the same way as the other values, fixing a (correctly spotted) ASAN issue.

diff --git a/Engine/Source/Runtime/Engine/Private/Materials/MaterialIREmitter.cpp b/Engine/Source/Runtime/Engine/Private/Materials/MaterialIREmitter.cpp
--- a/Engine/Source/Runtime/Engine/Private/Materials/MaterialIREmitter.cpp
+++ b/Engine/Source/Runtime/Engine/Private/Materials/MaterialIREmitter.cpp
@@ -724,14 +724,14 @@ FValue* FEmitter::EmitCustomPrimitiveData(uint32 PrimitiveDataIndex)
 
 FSetMaterialOutput* FEmitter::EmitSetMaterialOutput(EMaterialProperty InProperty, FValue* InArgValue)
 {
-	FSetMaterialOutput* Instr = new FSetMaterialOutput;
-	Instr->Kind 	= VK_SetMaterialOutput;
-	Instr->Block 	= Module->RootBlock;
-	Instr->Property = InProperty;
-	Instr->Arg 		= InArgValue;
+	FSetMaterialOutput Proto = MakePrototype<FSetMaterialOutput>(nullptr);
+	Proto.Block 	= Module->RootBlock;
+	Proto.Property  = InProperty;
+	Proto.Arg 		= InArgValue;
 
-	Module->Values.Add(Instr);
+	FSetMaterialOutput* Instr = static_cast<FSetMaterialOutput*>(EmitPrototype(this, Proto));
 	Module->Outputs.Add(Instr);
+
 	return Instr;
 }
 
