Fixed issue where meshedges and debug primitives could not be visible at the same time.
The DebugPrimitives PP pass now no longer runs in the MeshEdges prepass, but only in the main pass. This fixes an issue where the DebugPrimitives pass would prevent the EditorPrimitives pass from copying the wireframe to the MSAA prepass output buffer.

diff --git a/Engine/Source/Runtime/Engine/Public/ShowFlags.h b/Engine/Source/Runtime/Engine/Public/ShowFlags.h
--- a/Engine/Source/Runtime/Engine/Public/ShowFlags.h
+++ b/Engine/Source/Runtime/Engine/Public/ShowFlags.h
@@ -8,6 +8,7 @@
 
 #include "CoreMinimal.h"
 #include "Engine/EngineBaseTypes.h"
+#include "EngineDefines.h"
 
 /** for FShowFlagData, order is not important, adding new elements requires updating some editor code */
 enum EShowFlagGroup
@@ -408,7 +409,11 @@ private:
 		SetVisualizeVirtualShadowMap(false);
 		SetVectorFields(false);
 		SetGBufferHints(false);
-		SetCompositeEditorPrimitives(InitMode == ESFIM_Editor || InitMode == ESFIM_VREditing);
+#if UE_ENABLE_DEBUG_DRAWING
+		SetCompositeDebugPrimitives(true);
+#else
+		SetCompositeDebugPrimitives(false);
+#endif
 		SetOpaqueCompositeEditorPrimitives(false);
 		SetTestImage(false);
 		SetVisualizeDOF(false);
diff --git a/Engine/Source/Runtime/Engine/Public/ShowFlagsValues.inl b/Engine/Source/Runtime/Engine/Public/ShowFlagsValues.inl
--- a/Engine/Source/Runtime/Engine/Public/ShowFlagsValues.inl
+++ b/Engine/Source/Runtime/Engine/Public/ShowFlagsValues.inl
@@ -96,6 +96,8 @@ SHOWFLAG_ALWAYS_ACCESSIBLE(DepthOfField, SFG_PostProcess, NSLOCTEXT("UnrealEd",
 SHOWFLAG_FIXED_IN_SHIPPING(0, GBufferHints, SFG_Developer, NSLOCTEXT("UnrealEd", "GBufferHintsSF", "GBuffer Hints (material attributes)"))
 /** MotionBlur, for now only camera motion blur, for now SHOWFLAG_ALWAYS_ACCESSIBLE because it's exposed in SceneCapture */
 SHOWFLAG_ALWAYS_ACCESSIBLE(MotionBlur, SFG_PostProcess, NSLOCTEXT("UnrealEd", "MotionBlurSF", "Motion Blur"))
+/** Whether to render elements drawn through the debug PDI interface. */
+SHOWFLAG_ALWAYS_ACCESSIBLE(CompositeDebugPrimitives, SFG_Developer, NSLOCTEXT("UnrealEd", "CompositeDebugPrimitivesSF", "Composite Debug Primitives"))
 /** Whether to render the editor gizmos and other foreground editor widgets off screen and apply them after post process, only needed for the editor */
 SHOWFLAG_FIXED_IN_SHIPPING(0, CompositeEditorPrimitives, SFG_Developer, NSLOCTEXT("UnrealEd", "CompositeEditorPrimitivesSF", "Composite Editor Primitives"))
 /** When CompositeEditorPrimitives is set, determines whether occluded portions are drawn opaquely on top, or dithered when they are occluded */
diff --git a/Engine/Source/Runtime/Renderer/Private/MeshEdgesRendering.cpp b/Engine/Source/Runtime/Renderer/Private/MeshEdgesRendering.cpp
--- a/Engine/Source/Runtime/Renderer/Private/MeshEdgesRendering.cpp
+++ b/Engine/Source/Runtime/Renderer/Private/MeshEdgesRendering.cpp
@@ -301,6 +301,7 @@ void RenderMeshEdges(FSceneViewFamily& ViewFamily)
 		WireframeShowFlags.SetDecals(false);
 		WireframeShowFlags.SetAtmosphere(false);
 		WireframeShowFlags.SetPostProcessing(false);
+		WireframeShowFlags.SetCompositeDebugPrimitives(false);
 		WireframeShowFlags.SetCompositeEditorPrimitives(false);
 		WireframeShowFlags.SetGrid(false);
 		WireframeShowFlags.SetShaderPrint(false);
diff --git a/Engine/Source/Runtime/Renderer/Private/SceneRendering.cpp b/Engine/Source/Runtime/Renderer/Private/SceneRendering.cpp
--- a/Engine/Source/Runtime/Renderer/Private/SceneRendering.cpp
+++ b/Engine/Source/Runtime/Renderer/Private/SceneRendering.cpp
@@ -4677,7 +4677,7 @@ bool FSceneRenderer::ShouldCompositeDebugPrimitivesInPostProcess(const FViewInfo
 	{
 		//If we have primitives to draw, check we aren't forcing debug view modes off.
 		static bool bIsForceDisabled = IConsoleManager::Get().FindTConsoleVariableDataInt(TEXT("r.ForceDebugViewModes"))->GetValueOnAnyThread() == 2;
-		return !bIsForceDisabled;
+		return !bIsForceDisabled && View.Family->EngineShowFlags.CompositeDebugPrimitives;
 	}
 	return false;
 }
