[SceneCaptureComponent]
- Skip asset resolving for default objects to load them on-demand instead
- Saves 1.5s of editor boot times

diff --git a/Engine/Source/Runtime/Engine/Private/Components/SceneCaptureComponent.cpp b/Engine/Source/Runtime/Engine/Private/Components/SceneCaptureComponent.cpp
--- a/Engine/Source/Runtime/Engine/Private/Components/SceneCaptureComponent.cpp
+++ b/Engine/Source/Runtime/Engine/Private/Components/SceneCaptureComponent.cpp
@@ -669,7 +669,7 @@ USceneCaptureComponent2D::USceneCaptureComponent2D(const FObjectInitializer& Obj
 		ShowFlags.MotionBlur = false;
 
 #if WITH_EDITORONLY_DATA
-		if (!IsRunningCommandlet())
+		if (!IsRunningCommandlet() && !IsTemplate())
 		{
 			static ConstructorHelpers::FObjectFinder<UStaticMesh> EditorMesh(TEXT("/Engine/EditorMeshes/MatineeCam_SM"));
 			CaptureMesh = EditorMesh.Object;
@@ -1125,7 +1125,7 @@ UPlanarReflectionComponent::UPlanarReflectionComponent(const FObjectInitializer&
 	PlanarReflectionId = NextPlanarReflectionId;
 
 #if WITH_EDITORONLY_DATA
-	if (!IsRunningCommandlet())
+	if (!IsRunningCommandlet() && !IsTemplate())
 	{
 		static ConstructorHelpers::FObjectFinder<UStaticMesh> EditorMesh(TEXT("/Engine/EditorMeshes/PlanarReflectionPlane.PlanarReflectionPlane"));
 		CaptureMesh = EditorMesh.Object;
@@ -1286,7 +1286,7 @@ USceneCaptureComponentCube::USceneCaptureComponentCube(const FObjectInitializer&
 	bCaptureRotation = false;
 
 #if WITH_EDITORONLY_DATA
-	if (!IsRunningCommandlet())
+	if (!IsRunningCommandlet() && !IsTemplate())
 	{
 		static ConstructorHelpers::FObjectFinder<UStaticMesh> EditorMesh(TEXT("/Engine/EditorMeshes/MatineeCam_SM"));
 		CaptureMesh = EditorMesh.Object;
