[shaders] clear the inflight jobs map on releasing compiling id rather than asserting that it's empty. various paths to cancellation of in-flight jobs mean this sanity check is only valid for cooking.

diff --git a/Engine/Source/Runtime/Engine/Private/Materials/MaterialShader.cpp b/Engine/Source/Runtime/Engine/Private/Materials/MaterialShader.cpp
--- a/Engine/Source/Runtime/Engine/Private/Materials/MaterialShader.cpp
+++ b/Engine/Source/Runtime/Engine/Private/Materials/MaterialShader.cpp
@@ -2229,7 +2229,7 @@ void FMaterialShaderMap::ReleaseCompilingId()
 		check(CompilingMaterialDependencies.Num() == 0);
 		check(CompilingMaterialNumExternalDependencies == 0);
 #if ALLOW_SHADERMAP_DEBUG_DATA
-		check(InFlightJobs.Num() == 0);
+		InFlightJobs.Empty();
 #endif
 		verify(GetCompilingShaderMapLookup().Remove(CompilingId) == 1);
 		CompilingId = 0;
