[NewMaterialTranslator] - Implement length operator.

- Add missing op -> string for dot and cross

[RN] minor | Rendering - Materials and Shaders

diff --git a/Engine/Source/Runtime/Engine/Private/Materials/MaterialExpressionsIR.cpp b/Engine/Source/Runtime/Engine/Private/Materials/MaterialExpressionsIR.cpp
--- a/Engine/Source/Runtime/Engine/Private/Materials/MaterialExpressionsIR.cpp
+++ b/Engine/Source/Runtime/Engine/Private/Materials/MaterialExpressionsIR.cpp
@@ -482,6 +482,21 @@ void UMaterialExpressionFloor::Build(UE::MIR::FEmitter& Emitter)
 	BuildUnaryArithmeticOperatorFromInput(Emitter, MIR::UO_Floor, &Input);
 }
 
+void UMaterialExpressionLength::Build(UE::MIR::FEmitter& Emitter)
+{
+	MIR::FValue* Value = Emitter.Input(&Input);
+	UE_MIR_CHECKPOINT();
+
+	if (Value->Type->AsPrimitive()->IsScalar())
+	{
+		BuildUnaryArithmeticOperatorFromValue(Emitter, MIR::UO_Abs, Value);
+	}
+	else
+	{
+		BuildUnaryArithmeticOperatorFromValue(Emitter, MIR::UO_Length, Value);
+	}
+}
+
 void UMaterialExpressionRound::Build(UE::MIR::FEmitter& Emitter)
 {
 	BuildUnaryArithmeticOperatorFromInput(Emitter, MIR::UO_Round, &Input);
diff --git a/Engine/Source/Runtime/Engine/Private/Materials/MaterialIR.cpp b/Engine/Source/Runtime/Engine/Private/Materials/MaterialIR.cpp
--- a/Engine/Source/Runtime/Engine/Private/Materials/MaterialIR.cpp
+++ b/Engine/Source/Runtime/Engine/Private/Materials/MaterialIR.cpp
@@ -417,6 +417,7 @@ const TCHAR* OperatorToString(EOperator Op)
 		case UO_Exponential: return TEXT("Exponential");
 		case UO_Exponential2: return TEXT("Exponential2");
 		case UO_Floor: return TEXT("Floor");
+		case UO_Length: return TEXT("Length");
 		case UO_Logarithm: return TEXT("Logarithm");
 		case UO_Logarithm2: return TEXT("Logarithm2");
 		case UO_Logarithm10: return TEXT("Logarithm10");
@@ -441,6 +442,10 @@ const TCHAR* OperatorToString(EOperator Op)
 		case BO_Divide: return TEXT("Divide");
 		case BO_Min: return TEXT("Min");
 		case BO_Max: return TEXT("Max");
+
+		/* Binary Intrinsic Operators */
+		case BO_Dot: return TEXT("DotProduct");
+		case BO_Cross: return TEXT("CrossProduct");
 		
 		/* Ternary Arithmetic */
 		case TO_Lerp: return TEXT("Lerp");
diff --git a/Engine/Source/Runtime/Engine/Private/Materials/MaterialIREmitter.cpp b/Engine/Source/Runtime/Engine/Private/Materials/MaterialIREmitter.cpp
--- a/Engine/Source/Runtime/Engine/Private/Materials/MaterialIREmitter.cpp
+++ b/Engine/Source/Runtime/Engine/Private/Materials/MaterialIREmitter.cpp
@@ -1078,6 +1078,7 @@ static const FOperatorSignature* GetOperatorSignature(EOperator Op)
 		S[UO_Exponential] 			= UnaryArithmetic;
 		S[UO_Exponential2] 			= UnaryArithmetic;
 		S[UO_Floor] 				= UnaryArithmetic;
+		S[UO_Length]				= { { OT_Arithmetic | OT_NotMatrix }, ORT_FirstArgumentTypeToScalar };
 		S[UO_Logarithm] 			= UnaryArithmetic;
 		S[UO_Logarithm2] 			= UnaryArithmetic;
 		S[UO_Logarithm10] 			= UnaryArithmetic;
@@ -1231,7 +1232,7 @@ static FPrimitiveTypePtr ValidateOperatorAndGetResultType(FEmitter* Emitter, EOp
 /// 	op(v, w) == [op(v_x, w_x), ..., op(v_n, w_n)]
 static bool IsComponentwiseOperator(EOperator Op)
 {
-	return Op != BO_Dot && Op != BO_Cross;
+	return Op != BO_Dot && Op != BO_Cross && Op != UO_Length;
 }
 
 /// Tries to fold the operator by applying the operator componentwise on arguments components.
@@ -1409,6 +1410,33 @@ static FValue* TryFoldOperator(FEmitter* Emitter, EOperator Op, FValue* A, FValu
 			}
 		}
 	}
+	else if (Op == UO_Length)
+	{
+		if (A->Type->AsPrimitive()->IsVector())
+		{
+			FDimensional* Dimensional = A->As<FDimensional>();
+			if (Dimensional && Dimensional->AreComponentsConstant())
+			{
+				check(Dimensional->GetComponents().Num() <= 4);
+
+				FVector4f Input(0.0f, 0.0f, 0.0f, 0.0f);
+				int i = 0;
+				for (FValue const* Component : Dimensional->GetComponents())
+				{
+					Input[i++] = Component->As<FConstant>()->Float;
+				}
+
+				if (ResultType->ScalarKind == ScalarKind_Float)
+				{
+					return Emitter->EmitConstantFloatScalar(Input.Size());
+				}
+				else
+				{
+					UE_MIR_UNREACHABLE();
+				}
+			}
+		}
+	}
 
 	// If the operation supports componentwise application, try folding the operator componentwise.
 	if (IsComponentwiseOperator(Op))
diff --git a/Engine/Source/Runtime/Engine/Private/Materials/MaterialIRToHLSLTranslator.cpp b/Engine/Source/Runtime/Engine/Private/Materials/MaterialIRToHLSLTranslator.cpp
--- a/Engine/Source/Runtime/Engine/Private/Materials/MaterialIRToHLSLTranslator.cpp
+++ b/Engine/Source/Runtime/Engine/Private/Materials/MaterialIRToHLSLTranslator.cpp
@@ -492,6 +492,7 @@ struct FTranslator : FMaterialIRToHLSLTranslation
 				case MIR::UO_Ceil: OpString = TEXT("ceil"); break;
 				case MIR::UO_Cos: OpString = TEXT("cos"); break;
 				case MIR::UO_Floor: OpString = TEXT("floor"); break;
+				case MIR::UO_Length: OpString = TEXT("length"); break;
 				case MIR::UO_Sin: OpString = TEXT("sin"); break;
 				case MIR::UO_Tan: OpString = TEXT("tan"); break;
 				case MIR::UO_Truncate: OpString = TEXT("trunc"); break;
diff --git a/Engine/Source/Runtime/Engine/Public/Materials/MaterialExpressionLength.h b/Engine/Source/Runtime/Engine/Public/Materials/MaterialExpressionLength.h
--- a/Engine/Source/Runtime/Engine/Public/Materials/MaterialExpressionLength.h
+++ b/Engine/Source/Runtime/Engine/Public/Materials/MaterialExpressionLength.h
@@ -16,6 +16,7 @@ public:
 
 	//~ Begin UMaterialExpressionMaterialX Interface
 #if WITH_EDITOR
+	virtual void Build(UE::MIR::FEmitter& Emitter) override;
 	virtual int32 Compile(class FMaterialCompiler* Compiler, int32 OutputIndex) override;
 	virtual void GetCaption(TArray<FString>& OutCaptions) const override;
 	virtual void GetExpressionToolTip(TArray<FString>& OutToolTip) override;
diff --git a/Engine/Source/Runtime/Engine/Public/Materials/MaterialIR.h b/Engine/Source/Runtime/Engine/Public/Materials/MaterialIR.h
--- a/Engine/Source/Runtime/Engine/Public/Materials/MaterialIR.h
+++ b/Engine/Source/Runtime/Engine/Public/Materials/MaterialIR.h
@@ -322,6 +322,7 @@ enum EOperator
 	UO_Exponential,
 	UO_Exponential2,
 	UO_Floor,
+	UO_Length,
 	UO_Logarithm,
 	UO_Logarithm2,
 	UO_Logarithm10,
