World Partition - HLOD: Level instances can now be excluded from HLOD
* Their "Include in HLOD" property was previously ignored
* Added versionning to avoid breaking projects that could have modified the "Include in HLOD" property on existing actors.
* Actor Coloration: Modified the "HLOD Relevant" color handler to ensure it will refresh colors on LI's property changes
* Updated WP unit test map (+log)

#virtualized

diff --git a/Engine/Source/Editor/WorldPartitionEditor/Private/WorldPartition/HLOD/HLODEditorSubsystem.cpp b/Engine/Source/Editor/WorldPartitionEditor/Private/WorldPartition/HLOD/HLODEditorSubsystem.cpp
--- a/Engine/Source/Editor/WorldPartitionEditor/Private/WorldPartition/HLOD/HLODEditorSubsystem.cpp
+++ b/Engine/Source/Editor/WorldPartitionEditor/Private/WorldPartition/HLOD/HLODEditorSubsystem.cpp
@@ -7,6 +7,8 @@
 #include "Engine/World.h"
 #include "GameFramework/ActorPrimitiveColorHandler.h"
 #include "GameFramework/WorldSettings.h"
+#include "LevelInstance/LevelInstanceInterface.h"
+#include "LevelInstance/LevelInstanceSubsystem.h"
 #include "Misc/FileHelper.h"
 #include "StaticMeshResources.h"
 #include "Subsystems/UnrealEditorSubsystem.h"
@@ -43,7 +45,11 @@ UWorldPartitionHLODEditorSubsystem::UWorldPartitionHLODEditorSubsystem()
 #if ENABLE_ACTOR_PRIMITIVE_COLOR_HANDLER
 	if (HasAnyFlags(RF_ClassDefaultObject) && ExactCast<UWorldPartitionHLODEditorSubsystem>(this))
 	{
-		FActorPrimitiveColorHandler::Get().RegisterPrimitiveColorHandler(NAME_HLODRelevantColorHandler, LOCTEXT("HLODRelevantColor", "HLOD Relevant Color"), [](const UPrimitiveComponent* InPrimitiveComponent) -> FLinearColor
+		FActorPrimitiveColorHandler::FPrimitiveColorHandler HLODRelevantColorHandler;
+		HLODRelevantColorHandler.HandlerName = NAME_HLODRelevantColorHandler;
+		HLODRelevantColorHandler.HandlerText = LOCTEXT("HLODRelevantColor", "HLOD Relevant Color");
+		HLODRelevantColorHandler.HandlerToolTipText = LOCTEXT("HLODRelevantColor_ToolTip", "Colorize actor if relevant to the HLOD system. Green means relevant, otherwise the color is Red.");
+		HLODRelevantColorHandler.GetColorFunc = [](const UPrimitiveComponent* InPrimitiveComponent) -> FLinearColor
 		{
 			if (AActor* Actor = InPrimitiveComponent->GetOwner())
 			{
@@ -53,9 +59,19 @@ UWorldPartitionHLODEditorSubsystem::UWorldPartitionHLODEditorSubsystem()
 				}
 			}
 			return FLinearColor::Red;
-		},
-		[]() {},
-		LOCTEXT("HLODRelevantColor_ToolTip", "Colorize actor if relevant to the HLOD system. Green means relevant, otherwise the color is Red."));
+		};
+		
+		HLODRelevantColorHandler.ActivateFunc = [this]()
+		{
+			FCoreUObjectDelegates::OnObjectPropertyChanged.AddUObject(this, &UWorldPartitionHLODEditorSubsystem::OnColorHandlerPropertyChangedEvent);
+		};
+
+		HLODRelevantColorHandler.DeactivateFunc = [this]()
+		{
+			FCoreUObjectDelegates::OnObjectPropertyChanged.RemoveAll(this);
+		};
+
+		FActorPrimitiveColorHandler::Get().RegisterPrimitiveColorHandler(HLODRelevantColorHandler);
 	}
 #endif
 
@@ -75,6 +91,36 @@ UWorldPartitionHLODEditorSubsystem::~UWorldPartitionHLODEditorSubsystem()
 #endif
 }
 
+void UWorldPartitionHLODEditorSubsystem::OnColorHandlerPropertyChangedEvent(UObject* InObject, FPropertyChangedEvent& InPropertyChangedEvent)
+{
+	// When dealing with an LI, make sure to refresh the primitive color of all sub actors
+	if (ILevelInstanceInterface* LevelInstanceInterface = Cast<ILevelInstanceInterface>(InObject))
+	{
+		if (UWorld* World = InObject->GetWorld())
+		{
+			ULevelInstanceSubsystem* const LevelInstanceSubsystem = World->GetSubsystem<ULevelInstanceSubsystem>();
+			auto RefreshForLI = [LevelInstanceSubsystem](const ILevelInstanceInterface* LI)
+			{
+				ULevel* Level = LevelInstanceSubsystem->GetLevelInstanceLevel(LI);
+				if (Level)
+				{
+					FActorPrimitiveColorHandler::Get().RefreshPrimitiveColorHandler(NAME_HLODRelevantColorHandler, Level->Actors);
+				}
+			};
+
+			// Refresh LI actors
+			RefreshForLI(LevelInstanceInterface);
+
+			// Refresh child LIs actors
+			LevelInstanceSubsystem->ForEachLevelInstanceChild(LevelInstanceInterface, /*bRecursive=*/true, [&RefreshForLI](ILevelInstanceInterface* ChildLevelInstance)
+			{
+				RefreshForLI(ChildLevelInstance);
+				return true;
+			});
+		}
+	}
+}
+
 bool UWorldPartitionHLODEditorSubsystem::IsHLODInEditorEnabled()
 {
 	if (IsRunningCommandlet())
diff --git a/Engine/Source/Editor/WorldPartitionEditor/Public/WorldPartition/HLOD/HLODEditorSubsystem.h b/Engine/Source/Editor/WorldPartitionEditor/Public/WorldPartition/HLOD/HLODEditorSubsystem.h
--- a/Engine/Source/Editor/WorldPartitionEditor/Public/WorldPartition/HLOD/HLODEditorSubsystem.h
+++ b/Engine/Source/Editor/WorldPartitionEditor/Public/WorldPartition/HLOD/HLODEditorSubsystem.h
@@ -76,6 +76,8 @@ private:
 	void OnWorldPartitionEditorSettingsChanged(const FName& PropertyName, const UWorldPartitionEditorSettings& WorldPartitionEditorSettings);
 	void ApplyHLODSettingsFiltering();
 
+	void OnColorHandlerPropertyChangedEvent(UObject* InObject, FPropertyChangedEvent& InPropertyChangedEvent);
+
 private:
 	FVector CachedCameraLocation;
 	double CachedHLODMinDrawDistance;
diff --git a/Engine/Source/Runtime/Core/Public/UObject/FortniteMainBranchObjectVersions.inl b/Engine/Source/Runtime/Core/Public/UObject/FortniteMainBranchObjectVersions.inl
--- a/Engine/Source/Runtime/Core/Public/UObject/FortniteMainBranchObjectVersions.inl
+++ b/Engine/Source/Runtime/Core/Public/UObject/FortniteMainBranchObjectVersions.inl
@@ -562,5 +562,8 @@ FORTNITE_MAIN_VERSION(SpatialHashRuntimeGridInfoSpriteFixup, 182)
 // Removed UAnimSequence::bUseRawDataOnly flag alongside compression refactor
 FORTNITE_MAIN_VERSION(AnimSequenceRawDataOnlyFlagRemoval, 183)
 
+// HLOD relevancy of Level Instances was previously ignored, now taken into account. Reset to the default behavior.
+FORTNITE_MAIN_VERSION(ResetLevelInstanceHLODRelevancy, 184)
+
 // -----<new versions can be added above this line>-------------------------------------------------
 #undef FORTNITE_MAIN_VERSION
\ No newline at end of file
diff --git a/Engine/Source/Runtime/Engine/Classes/GameFramework/ActorPrimitiveColorHandler.h b/Engine/Source/Runtime/Engine/Classes/GameFramework/ActorPrimitiveColorHandler.h
--- a/Engine/Source/Runtime/Engine/Classes/GameFramework/ActorPrimitiveColorHandler.h
+++ b/Engine/Source/Runtime/Engine/Classes/GameFramework/ActorPrimitiveColorHandler.h
@@ -18,32 +18,26 @@ class ENGINE_API FActorPrimitiveColorHandler
 {
 	using FGetColorFunc = TFunction<FLinearColor(const UPrimitiveComponent*)>;
 	using FActivateFunc = TFunction<void(void)>;
+	using FDeactivateFunc = TFunction<void(void)>;
 
 public:
 	struct FPrimitiveColorHandler
 	{
-		FPrimitiveColorHandler(FName InHandlerName, FText InHandlerText, bool bInAvailalbleInEditor, const FGetColorFunc& InGetColorFunc, const FActivateFunc& InActivateFunc, FText InHandlerToolTipText = FText())
-			: HandlerName(InHandlerName)
-			, HandlerText(InHandlerText)
-			, HandlerToolTipText(InHandlerToolTipText)
-			, bAvailalbleInEditor(bInAvailalbleInEditor)
-			, GetColorFunc(InGetColorFunc)
-			, ActivateFunc(InActivateFunc)
-		{}
-
 		FName HandlerName;
 		FText HandlerText;
 		FText HandlerToolTipText;
 		bool bAvailalbleInEditor;
 		FGetColorFunc GetColorFunc;
 		FActivateFunc ActivateFunc;
+		FDeactivateFunc DeactivateFunc;
 	};	
 
 	FActorPrimitiveColorHandler();
 	static FActorPrimitiveColorHandler& Get();
 
-	void RegisterPrimitiveColorHandler(FName InHandlerName, FText InHandlerText, const FGetColorFunc& InHandlerFunc, const FActivateFunc& InActivateFunc = []() {}, FText InHandlerToolTipText = FText());
-	void RegisterPrimitiveColorHandler(FName InHandlerName, FText InHandlerText, bool bInAvailalbleInEditor, const FGetColorFunc& InHandlerFunc, const FActivateFunc& InActivateFunc = []() {}, FText InHandlerToolTipText = FText());
+	void RegisterPrimitiveColorHandler(FPrimitiveColorHandler& PrimitiveColorHandler);
+	void RegisterPrimitiveColorHandler(FName InHandlerName, const FText& InHandlerText, const FGetColorFunc& InHandlerFunc, const FActivateFunc& InActivateFunc = []() {}, const FText& InHandlerToolTipText = FText());
+	void RegisterPrimitiveColorHandler(FName InHandlerName, const FText& InHandlerText, bool bInAvailalbleInEditor, const FGetColorFunc& InHandlerFunc, const FActivateFunc& InActivateFunc = []() {}, const FText& InHandlerToolTipText = FText());
 	void UnregisterPrimitiveColorHandler(FName InHandlerName);
 	void GetRegisteredPrimitiveColorHandlers(TArray<FPrimitiveColorHandler>& OutPrimitiveColorHandlers) const;
 
diff --git a/Engine/Source/Runtime/Engine/Private/Actor.cpp b/Engine/Source/Runtime/Engine/Private/Actor.cpp
--- a/Engine/Source/Runtime/Engine/Private/Actor.cpp
+++ b/Engine/Source/Runtime/Engine/Private/Actor.cpp
@@ -6528,13 +6528,27 @@ bool AActor::IsHLODRelevant() const
 		return false;
 	}
 
-	if (HasAnyFlags(RF_Transient)
-#if WITH_EDITOR 
-		&& !IsInLevelInstance()			// Treat actors in LI as HLOD relevant for the sake of visualisation modes
-#endif
-	   )
+	if (HasAnyFlags(RF_Transient))
 	{
+#if WITH_EDITOR
+		// Treat actors in LI as HLOD relevant for the sake of visualisation modes
+		if (!IsInLevelInstance())
+		{
+			return false;
+		}
+		else if (ULevelInstanceSubsystem* LevelInstanceSubsystem = UWorld::GetSubsystem<ULevelInstanceSubsystem>(GetWorld()))
+		{
+			if (const ILevelInstanceInterface* LevelInstance = LevelInstanceSubsystem->GetOwningLevelInstance(GetLevel()))
+			{
+				if (!LevelInstance->IsLevelInstanceHLODRelevant())
+				{
+					return false;
+				}
+			}
+		}
+#else
 		return false;
+#endif
 	}
 
 	if (IsHidden())
diff --git a/Engine/Source/Runtime/Engine/Private/ActorPrimitiveColorHandler.cpp b/Engine/Source/Runtime/Engine/Private/ActorPrimitiveColorHandler.cpp
--- a/Engine/Source/Runtime/Engine/Private/ActorPrimitiveColorHandler.cpp
+++ b/Engine/Source/Runtime/Engine/Private/ActorPrimitiveColorHandler.cpp
@@ -27,17 +27,24 @@ FActorPrimitiveColorHandler& FActorPrimitiveColorHandler::Get()
 	return TLazySingleton<FActorPrimitiveColorHandler>::Get();
 }
 
-void FActorPrimitiveColorHandler::RegisterPrimitiveColorHandler(FName InHandlerName, FText InHandlerText, const FGetColorFunc& InHandlerFunc, const FActivateFunc& InActivateFunc, FText InHandlerToolTipText)
+void FActorPrimitiveColorHandler::RegisterPrimitiveColorHandler(FActorPrimitiveColorHandler::FPrimitiveColorHandler& PrimitiveColorHandler)
 {
-	RegisterPrimitiveColorHandler(InHandlerName, InHandlerText, true, InHandlerFunc, InActivateFunc, InHandlerToolTipText);
+#if ENABLE_ACTOR_PRIMITIVE_COLOR_HANDLER
+	check(!Handlers.Contains(PrimitiveColorHandler.HandlerName));
+	Handlers.Add(PrimitiveColorHandler.HandlerName, PrimitiveColorHandler);
+#endif
 }
 
-void FActorPrimitiveColorHandler::RegisterPrimitiveColorHandler(FName InHandlerName, FText InHandlerText, bool bInAvailalbleInEditor, const FGetColorFunc& InHandlerFunc, const FActivateFunc& InActivateFunc, FText InHandlerToolTipText)
+void FActorPrimitiveColorHandler::RegisterPrimitiveColorHandler(FName InHandlerName, const FText& InHandlerText, const FGetColorFunc& InHandlerFunc, const FActivateFunc& InActivateFunc, const FText& InHandlerToolTipText)
 {
-#if ENABLE_ACTOR_PRIMITIVE_COLOR_HANDLER
-	check(!Handlers.Contains(InHandlerName));
-	Handlers.Add(InHandlerName, { InHandlerName, InHandlerText, bInAvailalbleInEditor, InHandlerFunc, InActivateFunc, InHandlerToolTipText });
-#endif
+	FPrimitiveColorHandler Handler{ InHandlerName, InHandlerText, InHandlerToolTipText, true, InHandlerFunc, InActivateFunc };
+	RegisterPrimitiveColorHandler(Handler);
+}
+
+void FActorPrimitiveColorHandler::RegisterPrimitiveColorHandler(FName InHandlerName, const FText& InHandlerText, bool bInAvailalbleInEditor, const FGetColorFunc& InHandlerFunc, const FActivateFunc& InActivateFunc, const FText& InHandlerToolTipText)
+{
+	FPrimitiveColorHandler Handler{ InHandlerName, InHandlerText, InHandlerToolTipText, bInAvailalbleInEditor, InHandlerFunc, InActivateFunc };
+	RegisterPrimitiveColorHandler(Handler);
 }
 
 void FActorPrimitiveColorHandler::UnregisterPrimitiveColorHandler(FName InHandlerName)
@@ -49,6 +56,11 @@ void FActorPrimitiveColorHandler::UnregisterPrimitiveColorHandler(FName InHandle
 	
 	if (InHandlerName == ActivePrimitiveColorHandlerName)
 	{
+		if (ActivePrimitiveColorHandler)
+		{
+			ActivePrimitiveColorHandler->DeactivateFunc();
+		}
+
 		InitActivePrimitiveColorHandler();
 	}
 
@@ -60,6 +72,11 @@ bool FActorPrimitiveColorHandler::SetActivePrimitiveColorHandler(FName InHandler
 #if ENABLE_ACTOR_PRIMITIVE_COLOR_HANDLER
 	if (FPrimitiveColorHandler* NewActivePrimitiveColorHandler = Handlers.Find(InHandlerName); NewActivePrimitiveColorHandler && (NewActivePrimitiveColorHandler != ActivePrimitiveColorHandler))
 	{
+		if (ActivePrimitiveColorHandler)
+		{
+			ActivePrimitiveColorHandler->DeactivateFunc();
+		}
+
 		ActivePrimitiveColorHandlerName = InHandlerName;
 		ActivePrimitiveColorHandler = NewActivePrimitiveColorHandler;
 		NewActivePrimitiveColorHandler->ActivateFunc();
diff --git a/Engine/Source/Runtime/Engine/Private/LevelInstance/LevelInstanceActor.cpp b/Engine/Source/Runtime/Engine/Private/LevelInstance/LevelInstanceActor.cpp
--- a/Engine/Source/Runtime/Engine/Private/LevelInstance/LevelInstanceActor.cpp
+++ b/Engine/Source/Runtime/Engine/Private/LevelInstance/LevelInstanceActor.cpp
@@ -7,6 +7,7 @@
 #include "Engine/World.h"
 #include "LevelInstancePrivate.h"
 #include "Net/UnrealNetwork.h"
+#include "UObject/FortniteMainBranchObjectVersion.h"
 
 #include UE_INLINE_GENERATED_CPP_BY_NAME(LevelInstanceActor)
 
@@ -34,6 +35,8 @@ ALevelInstance::ALevelInstance()
 
 void ALevelInstance::Serialize(FArchive& Ar)
 {
+	Ar.UsingCustomVersion(FFortniteMainBranchObjectVersion::GUID);
+
 	Super::Serialize(Ar);
 	Ar << LevelInstanceActorGuid;
 
@@ -46,6 +49,12 @@ void ALevelInstance::Serialize(FArchive& Ar)
 		// force this softobjectpath to be processed by FSoftObjectPath::FixupForPIE (even one that comes from a level instance).
 		WorldAsset = FSoftObjectPath(UWorld::RemovePIEPrefix(WorldAsset.ToString()));
 	}
+
+	// Level instances were previously always considered HLOD relevant.
+	if (Ar.IsLoading() && Ar.CustomVer(FFortniteMainBranchObjectVersion::GUID) < FFortniteMainBranchObjectVersion::ResetLevelInstanceHLODRelevancy)
+	{
+		bEnableAutoLODGeneration = true;
+	}
 #endif
 }
 
@@ -404,6 +413,22 @@ void ALevelInstance::PushLevelInstanceEditingStateToProxies(bool bInEditingState
 	LevelInstanceActorImpl.PushLevelInstanceEditingStateToProxies(bInEditingState);
 }
 
+bool ALevelInstance::IsHLODRelevant() const
+{
+	if (!IsWorldAssetValid())
+	{
+		return false;
+	}
+
+	return Super::IsHLODRelevant();
+}
+
+bool ALevelInstance::HasHLODRelevantComponents() const
+{
+	// Assume LIs can provide HLOD relevant components
+	return true;
+}
+
 #endif // WITH_EDITOR
 
 #undef LOCTEXT_NAMESPACE
diff --git a/Engine/Source/Runtime/Engine/Private/LevelInstance/LevelInstanceInterface.cpp b/Engine/Source/Runtime/Engine/Private/LevelInstance/LevelInstanceInterface.cpp
--- a/Engine/Source/Runtime/Engine/Private/LevelInstance/LevelInstanceInterface.cpp
+++ b/Engine/Source/Runtime/Engine/Private/LevelInstance/LevelInstanceInterface.cpp
@@ -181,6 +181,18 @@ ULevel* ILevelInstanceInterface::GetLoadedLevel() const
 
 #if WITH_EDITOR
 
+bool ILevelInstanceInterface::IsLevelInstanceHLODRelevant() const
+{
+	if (const ULevelInstanceComponent* LevelInstanceComponent = GetLevelInstanceComponent())
+	{
+		if (const AActor* Owner = LevelInstanceComponent->GetOwner())
+		{
+			return Owner->IsHLODRelevant();
+		}
+	}
+	return false;
+}
+
 bool ILevelInstanceInterface::IsEditing() const
 {
 	// No need to check LevelInstanceID here since Editing Level Instance is referenced directly:
diff --git a/Engine/Source/Runtime/Engine/Private/PackedLevelActor/PackedLevelActor.cpp b/Engine/Source/Runtime/Engine/Private/PackedLevelActor/PackedLevelActor.cpp
--- a/Engine/Source/Runtime/Engine/Private/PackedLevelActor/PackedLevelActor.cpp
+++ b/Engine/Source/Runtime/Engine/Private/PackedLevelActor/PackedLevelActor.cpp
@@ -281,12 +281,6 @@ bool APackedLevelActor::IsHiddenEd() const
 	return Super::IsHiddenEd() || IsEditing() || HasChildEdit() || ShouldLoadForPacking();
 }
 
-bool APackedLevelActor::IsHLODRelevant() const
-{
-	// Bypass base class ALevelInstance (because it always returns true). We want the same implementation as AActor.
-	return AActor::IsHLODRelevant();
-}
-
 bool APackedLevelActor::CanEditChange(const FProperty* InProperty) const
 {
 	if (!Super::CanEditChange(InProperty))
diff --git a/Engine/Source/Runtime/Engine/Private/WorldPartition/LevelInstance/LevelInstanceActorDesc.cpp b/Engine/Source/Runtime/Engine/Private/WorldPartition/LevelInstance/LevelInstanceActorDesc.cpp
--- a/Engine/Source/Runtime/Engine/Private/WorldPartition/LevelInstance/LevelInstanceActorDesc.cpp
+++ b/Engine/Source/Runtime/Engine/Private/WorldPartition/LevelInstance/LevelInstanceActorDesc.cpp
@@ -449,5 +449,11 @@ void FLevelInstanceActorDesc::Serialize(FArchive& Ar)
 			PropertyOverrideDescPtr->PackageName = ActorPackage;
 		}
 	}
+
+	// As this property was previously ignored for Level Instances, make sure not to break existing content that may have modified that flag.
+	if (Ar.CustomVer(FFortniteMainBranchObjectVersion::GUID) < FFortniteMainBranchObjectVersion::ResetLevelInstanceHLODRelevancy)
+	{
+		bActorIsHLODRelevant = true;
+	}
 }
 #endif
diff --git a/Engine/Source/Runtime/Engine/Private/WorldPartition/WorldPartitionStreamingGeneration.cpp b/Engine/Source/Runtime/Engine/Private/WorldPartition/WorldPartitionStreamingGeneration.cpp
--- a/Engine/Source/Runtime/Engine/Private/WorldPartition/WorldPartitionStreamingGeneration.cpp
+++ b/Engine/Source/Runtime/Engine/Private/WorldPartition/WorldPartitionStreamingGeneration.cpp
@@ -219,6 +219,16 @@ FSoftObjectPath FStreamingGenerationActorDescView::GetHLODLayer() const
 	return Super::GetHLODLayer();
 }
 
+bool FStreamingGenerationActorDescView::GetActorIsHLODRelevant() const
+{
+	if (bIsForcedHLODIrrelevant)
+	{
+		return false;
+	}
+
+	return Super::GetActorIsHLODRelevant();
+}
+
 const FDataLayerInstanceNames& FStreamingGenerationActorDescView::GetDataLayerInstanceNames() const
 {
 	if (!bIsForcedNoDataLayers && ParentView)
@@ -259,12 +269,13 @@ FString FStreamingGenerationActorDescView::ToString(FWorldPartitionActorDesc::ET
 	FString Result = Super::ToString(Mode);
 	if (Mode >= FWorldPartitionActorDesc::EToStringMode::Compact)
 	{
-		Result += FString::Printf(TEXT(" View=(RuntimeHLODLayer=%s bIsForcedNonSpatiallyLoaded=%s bIsForcedNoRuntimeGrid=%s bIsForcedNoDataLayers=%s bIsForceNoHLODLayer=%s)"), 
+		Result += FString::Printf(TEXT(" View=(RuntimeHLODLayer=%s bIsForcedNonSpatiallyLoaded=%s bIsForcedNoRuntimeGrid=%s bIsForcedNoDataLayers=%s bIsForceNoHLODLayer=%s bIsForcedHLODIrrelevant=%s)"), 
 			RuntimeHLODLayer.IsSet() ? *RuntimeHLODLayer->ToString() : TEXT("none"),
 			*LexToString(bIsForcedNonSpatiallyLoaded),
 			*LexToString(bIsForcedNoRuntimeGrid),
 			*LexToString(bIsForcedNoDataLayers),
-			*LexToString(bIsForceNoHLODLayer));
+			*LexToString(bIsForceNoHLODLayer),
+			*LexToString(bIsForcedHLODIrrelevant));
 	}
 	return Result;
 }
@@ -392,6 +403,11 @@ void FStreamingGenerationActorDescView::SetRuntimeHLODLayer(const FSoftObjectPat
 	RuntimeHLODLayer = InHLODLayer;
 }
 
+void FStreamingGenerationActorDescView::SetForcedHLODIrrelevant()
+{
+	bIsForcedHLODIrrelevant = true;
+}
+
 const FDataLayerInstanceNames& FStreamingGenerationActorDescView::GetRuntimeDataLayerInstanceNames() const
 {
 	if (!RuntimeDataLayerInstanceNames.IsSet())
@@ -589,6 +605,7 @@ class FWorldPartitionStreamingGenerator
 	{
 		FContainerCollectionInstanceDescriptor()
 			: Bounds(ForceInit)
+			, bIsHLODRelevant(true)
 		{}
 
 		FBox Bounds;
@@ -600,6 +617,7 @@ class FWorldPartitionStreamingGenerator
 		FActorContainerID ParentID;
 		FGuid ContentBundleID;
 		FSoftObjectPath HLODLayer;
+		bool bIsHLODRelevant;
 
 		FContainerCollectionInstanceDescriptor(const FContainerCollectionInstanceDescriptor& Other) = delete;
 		FContainerCollectionInstanceDescriptor(FContainerCollectionInstanceDescriptor&& Other) = default;
@@ -742,6 +760,14 @@ class FWorldPartitionStreamingGenerator
 		}
 	}
 
+	void ResolveHLODRelevant(FStreamingGenerationActorDescView& ActorDescView, bool bIsHLODRelevant)
+	{
+		if (!bIsHLODRelevant)
+		{
+			ActorDescView.SetForcedHLODIrrelevant();
+		}
+	}
+	
 	void ResolveParentView(FStreamingGenerationActorDescView& ActorDescView, const FStreamingGenerationActorDescViewMap& ActorDescViewMap, const TSet<FGuid>& EditorOnlyActorDescSet, const TMap<FGuid, FTransform>& EditorOnlyParentActorTransforms)
 	{
 		if (FGuid ParentGuid = ActorDescView.GetParentActor(); ParentGuid.IsValid())
@@ -1021,6 +1047,7 @@ class FWorldPartitionStreamingGenerator
 				SubContainerInstanceDescriptor.ContentBundleID = ContainerCollectionInstanceDescriptor.ContentBundleID;
 				SubContainerInstanceDescriptor.InstanceData = InheritParentContainerPerInstanceData(ContainerCollectionInstanceDescriptor, ContainerCollectionInstanceView);
 				SubContainerInstanceDescriptor.HLODLayer = ContainerCollectionInstanceView.GetHLODLayer().IsValid() ? ContainerCollectionInstanceView.GetHLODLayer() : ContainerCollectionInstanceDescriptor.HLODLayer;
+				SubContainerInstanceDescriptor.bIsHLODRelevant = ContainerCollectionInstanceView.GetActorIsHLODRelevant() && ContainerCollectionInstanceDescriptor.bIsHLODRelevant;
 
 				if (WorldPartitionSubsystem && ContainerID.IsMainContainer() && ContainerCollectionInstanceView.GetChildContainerFilterType() == EWorldPartitionActorFilterType::Loading)
 				{
@@ -1066,6 +1093,7 @@ class FWorldPartitionStreamingGenerator
 		MainContainerCollectionInstance.ContentBundleID = InContainerInstanceCollection.GetContentBundleGuid();
 		MainContainerCollectionInstance.InstanceData.bIsSpatiallyLoaded = true; // Since we apply AND logic on spatially loaded flag recursively, startup value must be true
 		MainContainerCollectionInstance.HLODLayer = WorldPartitionContext ? FSoftObjectPath(WorldPartitionContext->GetDefaultHLODLayer()) : FSoftObjectPath();
+		MainContainerCollectionInstance.bIsHLODRelevant = true;
 
 		// Create Child Containers
 		CreateActorDescriptorViewsRecursive(MoveTemp(MainContainerCollectionInstance));
@@ -1106,6 +1134,7 @@ class FWorldPartitionStreamingGenerator
 			ResolveRuntimeGrid(ActorDescView);
 			ResolveRuntimeDataLayers(ActorDescView, ContainerCollectionInstanceDescriptor.DataLayerResolvers);
 			ResolveHLODLayer(ActorDescView, ContainerCollectionInstanceDescriptor.HLODLayer);
+			ResolveHLODRelevant(ActorDescView, ContainerCollectionInstanceDescriptor.bIsHLODRelevant);
 			ResolveParentView(ActorDescView, ContainerCollectionInstanceDescriptor.ActorDescViewMap.Get(), ContainerCollectionInstanceDescriptor.EditorOnlyActorDescSet, ContainerCollectionInstanceDescriptor.EditorOnlyParentActorTransforms);
 		};
 
@@ -1853,11 +1882,12 @@ public:
 				{
 					UE_SCOPED_INDENT_LOG_ARCHIVE(Ar.PrintfIndent(TEXT("%s:"), *ContainerInstanceDescriptor.OwnerName));
 
-					Ar.Printf(TEXT("       ID: %s"), *ContainerInstanceDescriptor.ID.ToString());
-					Ar.Printf(TEXT(" ParentID: %s"), *ContainerInstanceDescriptor.ParentID.ToString());
-					Ar.Printf(TEXT("   Bounds: %s"), *ContainerInstanceDescriptor.Bounds.ToString());
-					Ar.Printf(TEXT("Transform: %s"), *ContainerInstanceDescriptor.Transform.ToString());
-					Ar.Printf(TEXT("Container: %s"), *ContainerInstanceDescriptor.ContainerInstanceCollection->GetBaseContainerInstancePackageName().ToString());
+					Ar.Printf(TEXT("          ID: %s"), *ContainerInstanceDescriptor.ID.ToString());
+					Ar.Printf(TEXT("    ParentID: %s"), *ContainerInstanceDescriptor.ParentID.ToString());
+					Ar.Printf(TEXT("      Bounds: %s"), *ContainerInstanceDescriptor.Bounds.ToString());
+					Ar.Printf(TEXT("   Transform: %s"), *ContainerInstanceDescriptor.Transform.ToString());
+					Ar.Printf(TEXT("   Container: %s"), *ContainerInstanceDescriptor.ContainerInstanceCollection->GetBaseContainerInstancePackageName().ToString());
+					Ar.Printf(TEXT("HLODRelevant: %s"), ContainerInstanceDescriptor.bIsHLODRelevant ? TEXT("true") : TEXT("false"));
 				}
 
 				TArray<FActorContainerID> ChildContainersIDs;
diff --git a/Engine/Source/Runtime/Engine/Public/LevelInstance/LevelInstanceActor.h b/Engine/Source/Runtime/Engine/Public/LevelInstance/LevelInstanceActor.h
--- a/Engine/Source/Runtime/Engine/Public/LevelInstance/LevelInstanceActor.h
+++ b/Engine/Source/Runtime/Engine/Public/LevelInstance/LevelInstanceActor.h
@@ -112,6 +112,8 @@ public:
 	ENGINE_API virtual bool IsUserManaged() const override;
 	ENGINE_API virtual bool ShouldExport() override;
 	ENGINE_API virtual bool SupportsSubRootSelection() const override { return true; }
+	ENGINE_API virtual bool IsHLODRelevant() const override;
+	ENGINE_API virtual bool HasHLODRelevantComponents() const override;
 	// End of AActor interface
 
 	DECLARE_MULTICAST_DELEGATE_OneParam(FOnLevelInstanceActorPostLoad, ALevelInstance*);
diff --git a/Engine/Source/Runtime/Engine/Public/LevelInstance/LevelInstanceInterface.h b/Engine/Source/Runtime/Engine/Public/LevelInstance/LevelInstanceInterface.h
--- a/Engine/Source/Runtime/Engine/Public/LevelInstance/LevelInstanceInterface.h
+++ b/Engine/Source/Runtime/Engine/Public/LevelInstance/LevelInstanceInterface.h
@@ -149,6 +149,9 @@ class ILevelInstanceInterface
 
 	// Return supported filter types when setting filter through details panel
 	ENGINE_API virtual EWorldPartitionActorFilterType GetDetailsFilterTypes() const;
+
+	// Whether the level instance is HLOD relevant. If not, all it's actors will be excluded from HLOD.
+	ENGINE_API virtual bool IsLevelInstanceHLODRelevant() const;
 #endif
 
 	/** Sets the UWorld asset reference when loading a LevelInstance */
diff --git a/Engine/Source/Runtime/Engine/Public/PackedLevelActor/PackedLevelActor.h b/Engine/Source/Runtime/Engine/Public/PackedLevelActor/PackedLevelActor.h
--- a/Engine/Source/Runtime/Engine/Public/PackedLevelActor/PackedLevelActor.h
+++ b/Engine/Source/Runtime/Engine/Public/PackedLevelActor/PackedLevelActor.h
@@ -68,7 +68,6 @@ public:
 	void SetPackedHash(uint32 InHash) { PackedHash = InHash; }
 
 	ENGINE_API virtual bool IsHiddenEd() const override;
-	ENGINE_API virtual bool IsHLODRelevant() const override;
 
 	ENGINE_API void DestroyPackedComponents();
 	ENGINE_API void GetPackedComponents(TArray<UActorComponent*>& OutPackedComponents) const;
diff --git a/Engine/Source/Runtime/Engine/Public/WorldPartition/WorldPartitionStreamingGeneration.h b/Engine/Source/Runtime/Engine/Public/WorldPartition/WorldPartitionStreamingGeneration.h
--- a/Engine/Source/Runtime/Engine/Public/WorldPartition/WorldPartitionStreamingGeneration.h
+++ b/Engine/Source/Runtime/Engine/Public/WorldPartition/WorldPartitionStreamingGeneration.h
@@ -31,6 +31,7 @@ protected:
 		, bIsForcedNoRuntimeGrid(false)
 		, bIsForcedNoDataLayers(false)
 		, bIsForceNoHLODLayer(false)
+		, bIsForcedHLODIrrelevant(false)
 		, bIsUnsaved(false)
 	{}
 
@@ -56,6 +57,7 @@ public:
 	ENGINE_API virtual FName GetRuntimeGrid() const override;
 	ENGINE_API virtual bool GetIsSpatiallyLoaded() const override;
 	ENGINE_API virtual FSoftObjectPath GetHLODLayer() const override;
+	ENGINE_API virtual bool GetActorIsHLODRelevant() const override;
 	ENGINE_API virtual const TArray<FGuid>& GetReferences() const override;
 	ENGINE_API virtual const FDataLayerInstanceNames& GetDataLayerInstanceNames() const override;
 	ENGINE_API virtual bool IsEditorOnlyReference(const FGuid& ReferenceGuid) const override;
@@ -83,6 +85,7 @@ public:
 
 	ENGINE_API void SetForcedNoHLODLayer();
 	ENGINE_API void SetRuntimeHLODLayer(const FSoftObjectPath& InHLODLayer);
+	ENGINE_API void SetForcedHLODIrrelevant();
 
 	bool operator==(const FStreamingGenerationActorDescView& Other) const
 	{
@@ -102,6 +105,7 @@ private:
 	bool bIsForcedNoRuntimeGrid;
 	bool bIsForcedNoDataLayers;
 	bool bIsForceNoHLODLayer;
+	bool bIsForcedHLODIrrelevant;
 	bool bIsUnsaved;
 
 	// Stores invalid actor references
