Restoring MouseCaptureMode to its correct previous state when the console is closed

[RN] Restoring MouseCaptureMode to its correct previous state when the console is closed

diff --git a/Engine/Source/Runtime/Engine/Classes/Engine/Console.h b/Engine/Source/Runtime/Engine/Classes/Engine/Console.h
--- a/Engine/Source/Runtime/Engine/Classes/Engine/Console.h
+++ b/Engine/Source/Runtime/Engine/Classes/Engine/Console.h
@@ -374,5 +374,8 @@ private:
 	// Widget that was focused before the console was opened (focus will be restored to this if it's valid after the console closes)
 	TWeakPtr<SWidget> PreviousFocusedWidget;
 
+	// Previous mouse capture mode of the viewport before the console was opened (mouse capture mode must be restored to its previous state)
+	EMouseCaptureMode PreviousMouseCaptureMode = EMouseCaptureMode::CapturePermanently;
+
 	const bool bDPIAwareStringMeasurement = true;
 };
diff --git a/Engine/Source/Runtime/Engine/Private/UserInterface/Console.cpp b/Engine/Source/Runtime/Engine/Private/UserInterface/Console.cpp
--- a/Engine/Source/Runtime/Engine/Private/UserInterface/Console.cpp
+++ b/Engine/Source/Runtime/Engine/Private/UserInterface/Console.cpp
@@ -1750,11 +1750,23 @@ void UConsole::BeginState_Open(FName PreviousStateName)
 	{
 		FlushPlayerInput();
 	}
-}
 
+	// Ensure we receive IE_Pressed with LeftMouseButton
+	PreviousMouseCaptureMode = GetOuterUGameViewportClient()->GetMouseCaptureMode();
+	GetOuterUGameViewportClient()->SetMouseCaptureMode(EMouseCaptureMode::CaptureDuringMouseDown);
+	Selection.bActive = false;
+	Selection.bCapture = false;
+	Selection.bMade = false;
+	Selection.Offset = 0.0f;
+}
 
 void UConsole::EndState_Open(FName NextStateName)
 {
+	// Restore mouse capture for non-open states (only if nothing else modified it while the console was open otherwise, let's keep the new mode)
+	if (GetOuterUGameViewportClient()->GetMouseCaptureMode() == EMouseCaptureMode::CaptureDuringMouseDown)
+	{
+		GetOuterUGameViewportClient()->SetMouseCaptureMode(PreviousMouseCaptureMode);
+	}
 }
 
 PRAGMA_DISABLE_DEPRECATION_WARNINGS
@@ -2074,9 +2086,6 @@ void UConsole::FakeGotoState(FName NextStateName)
 	else if (ConsoleState == NAME_Open)
 	{
 		EndState_Open(NextStateName);
-
-		// Restore mouse capture for non-open states
-		GetOuterUGameViewportClient()->SetMouseCaptureMode(EMouseCaptureMode::CapturePermanently);
 	}
 	if (NextStateName == NAME_Typing)
 	{
@@ -2098,13 +2107,6 @@ void UConsole::FakeGotoState(FName NextStateName)
 
 		// Console has opened
 		OnConsoleActivationStateChanged.Broadcast(true);
-
-		// Ensure we receive IE_Pressed with LeftMouseButton
-		GetOuterUGameViewportClient()->SetMouseCaptureMode(EMouseCaptureMode::CaptureDuringMouseDown);
-		Selection.bActive = false;
-		Selection.bCapture = false;
-		Selection.bMade = false;
-		Selection.Offset = 0.0f;
 	}
 	else if (NextStateName == NAME_None)
 	{
