[DistanceFields] Avoid copying OptionalSourceMeshData when passing data to Mesh Card build.

diff --git a/Engine/Source/Runtime/Engine/Private/DistanceFieldAtlas.cpp b/Engine/Source/Runtime/Engine/Private/DistanceFieldAtlas.cpp
--- a/Engine/Source/Runtime/Engine/Private/DistanceFieldAtlas.cpp
+++ b/Engine/Source/Runtime/Engine/Private/DistanceFieldAtlas.cpp
@@ -222,7 +222,7 @@ void FDistanceFieldVolumeData::CacheDerivedData(const FString& InStaticMeshDeriv
 		FMemoryReader Ar(DerivedData, /*bIsPersistent=*/ true);
 		Serialize(Ar, Mesh);
 
-		BeginCacheMeshCardRepresentation(TargetPlatform, Mesh, RenderData, DistanceFieldKey, /*OptionalSourceMeshData*/ nullptr);
+		BeginCacheMeshCardRepresentation(TargetPlatform, Mesh, RenderData, DistanceFieldKey);
 	}
 	else if (GDistanceFieldAsyncQueue)
 	{
@@ -864,7 +864,7 @@ void FDistanceFieldAsyncQueue::ProcessAsyncTasks(bool bLimitExecutionTime)
 				Task->StaticMesh,
 				Task->StaticMesh->GetPlatformStaticMeshRenderData(Task->StaticMesh, Task->TargetPlatform),
 				Task->DDCKey,
-				&Task->SourceMeshData);
+				MoveTemp(Task->SourceMeshData));
 		}
 
 		delete Task;
diff --git a/Engine/Source/Runtime/Engine/Private/MeshCardRepresentation.cpp b/Engine/Source/Runtime/Engine/Private/MeshCardRepresentation.cpp
--- a/Engine/Source/Runtime/Engine/Private/MeshCardRepresentation.cpp
+++ b/Engine/Source/Runtime/Engine/Private/MeshCardRepresentation.cpp
@@ -132,7 +132,12 @@ FString BuildCardRepresentationDerivedDataKey(const FString& InMeshKey, int32 Ma
 
 extern void BuildSignedDistanceFieldBuildSectionData(UStaticMesh* Mesh, uint32 LODIndex, TArray<FSignedDistanceFieldBuildSectionData>& OutData);
 
-void BeginCacheMeshCardRepresentation(const ITargetPlatform* TargetPlatform, UStaticMesh* StaticMeshAsset, FStaticMeshRenderData& RenderData, const FString& DistanceFieldKey, FSourceMeshDataForDerivedDataTask* OptionalSourceMeshData)
+void BeginCacheMeshCardRepresentation(
+	const ITargetPlatform* TargetPlatform,
+	UStaticMesh* StaticMeshAsset,
+	FStaticMeshRenderData& RenderData,
+	const FString& DistanceFieldKey,
+	FSourceMeshDataForDerivedDataTask&& OptionalSourceMeshData)
 {
 	static const auto CVarCards = IConsoleManager::Get().FindTConsoleVariableDataInt(TEXT("r.MeshCardRepresentation"));
 
@@ -156,12 +161,19 @@ void BeginCacheMeshCardRepresentation(const ITargetPlatform* TargetPlatform, USt
 				MeshToGenerateFrom,
 				BuildSettings.MaxLumenMeshCards,
 				BuildSettings.bGenerateDistanceFieldAsIfTwoSided,
-				OptionalSourceMeshData);
+				MoveTemp(OptionalSourceMeshData));
 		}
 	}
 }
 
-void FCardRepresentationData::CacheDerivedData(const FString& InDDCKey, const ITargetPlatform* TargetPlatform, UStaticMesh* Mesh, UStaticMesh* GenerateSource, int32 MaxLumenMeshCards, bool bGenerateDistanceFieldAsIfTwoSided, FSourceMeshDataForDerivedDataTask* OptionalSourceMeshData)
+void FCardRepresentationData::CacheDerivedData(
+	const FString& InDDCKey,
+	const ITargetPlatform* TargetPlatform,
+	UStaticMesh* Mesh,
+	UStaticMesh* GenerateSource,
+	int32 MaxLumenMeshCards,
+	bool bGenerateDistanceFieldAsIfTwoSided,
+	FSourceMeshDataForDerivedDataTask&& OptionalSourceMeshData)
 {
 	TRACE_CPUPROFILER_EVENT_SCOPE(FCardRepresentationData::CacheDerivedData);
 
@@ -198,9 +210,9 @@ void FCardRepresentationData::CacheDerivedData(const FString& InDDCKey, const IT
 		BuildSignedDistanceFieldBuildSectionData(Mesh, LODIndex, NewTask->SectionData);
 
 		// Nanite overrides source static mesh with a coarse representation. Need to load original data before we build the mesh SDF.
-		if (OptionalSourceMeshData)
+		if (OptionalSourceMeshData.IsValid())
 		{
-			NewTask->SourceMeshData = *OptionalSourceMeshData;
+			NewTask->SourceMeshData = MoveTemp(OptionalSourceMeshData);
 		}
 		else if (Mesh->IsNaniteEnabled())
 		{
diff --git a/Engine/Source/Runtime/Engine/Private/StaticMesh.cpp b/Engine/Source/Runtime/Engine/Private/StaticMesh.cpp
--- a/Engine/Source/Runtime/Engine/Private/StaticMesh.cpp
+++ b/Engine/Source/Runtime/Engine/Private/StaticMesh.cpp
@@ -4024,7 +4024,7 @@ void FStaticMeshRenderData::Cache(const ITargetPlatform* TargetPlatform, UStatic
 	}
 	else
 	{
-		BeginCacheMeshCardRepresentation(TargetPlatform, Owner, *this, DerivedDataKey, /* OptionalSourceMeshData */ nullptr);
+		BeginCacheMeshCardRepresentation(TargetPlatform, Owner, *this, DerivedDataKey);
 	}
 }
 #endif // #if WITH_EDITOR
diff --git a/Engine/Source/Runtime/Engine/Public/MeshCardBuild.h b/Engine/Source/Runtime/Engine/Public/MeshCardBuild.h
--- a/Engine/Source/Runtime/Engine/Public/MeshCardBuild.h
+++ b/Engine/Source/Runtime/Engine/Public/MeshCardBuild.h
@@ -167,7 +167,27 @@ public:
 
 #if WITH_EDITORONLY_DATA
 
-	void CacheDerivedData(const FString& InDDCKey, const ITargetPlatform* TargetPlatform, UStaticMesh* Mesh, UStaticMesh* GenerateSource, int32 MaxLumenMeshCards, bool bGenerateDistanceFieldAsIfTwoSided, FSourceMeshDataForDerivedDataTask* OptionalSourceMeshData);
+	void CacheDerivedData(
+		const FString& InDDCKey,
+		const ITargetPlatform* TargetPlatform,
+		UStaticMesh* Mesh,
+		UStaticMesh* GenerateSource,
+		int32 MaxLumenMeshCards,
+		bool bGenerateDistanceFieldAsIfTwoSided,
+		FSourceMeshDataForDerivedDataTask&& OptionalSourceMeshData = FSourceMeshDataForDerivedDataTask());
+
+	UE_DEPRECATED(5.6, "Provide OptionalSourceMeshData by value (ideally using MoveTemp(...) to avoid copying the data)")
+	void CacheDerivedData(
+		const FString& InDDCKey,
+		const ITargetPlatform* TargetPlatform,
+		UStaticMesh* Mesh,
+		UStaticMesh* GenerateSource,
+		int32 MaxLumenMeshCards,
+		bool bGenerateDistanceFieldAsIfTwoSided,
+		FSourceMeshDataForDerivedDataTask* OptionalSourceMeshData)
+	{
+		CacheDerivedData(InDDCKey, TargetPlatform, Mesh, GenerateSource, MaxLumenMeshCards, bGenerateDistanceFieldAsIfTwoSided, OptionalSourceMeshData ? *OptionalSourceMeshData : FSourceMeshDataForDerivedDataTask());
+	}
 
 #endif
 
@@ -328,7 +348,24 @@ extern ENGINE_API FCardRepresentationAsyncQueue* GCardRepresentationAsyncQueue;
 extern ENGINE_API void AppendMeshCardConstantsToClassSchema(FAppendToClassSchemaContext& Context);
 
 extern ENGINE_API FString BuildCardRepresentationDerivedDataKey(const FString& InMeshKey);
-extern ENGINE_API void BeginCacheMeshCardRepresentation(const ITargetPlatform* TargetPlatform, UStaticMesh* StaticMeshAsset, class FStaticMeshRenderData& RenderData, const FString& DistanceFieldKey, FSourceMeshDataForDerivedDataTask* OptionalSourceMeshData);
+
+extern ENGINE_API void BeginCacheMeshCardRepresentation(
+	const ITargetPlatform* TargetPlatform,
+	UStaticMesh* StaticMeshAsset,
+	class FStaticMeshRenderData& RenderData,
+	const FString& DistanceFieldKey,
+	FSourceMeshDataForDerivedDataTask&& OptionalSourceMeshData = FSourceMeshDataForDerivedDataTask());
+
+UE_DEPRECATED(5.6, "Provide OptionalSourceMeshData by value (ideally using MoveTemp(...) to avoid copying the data)")
+inline void BeginCacheMeshCardRepresentation(
+	const ITargetPlatform* TargetPlatform,
+	UStaticMesh* StaticMeshAsset,
+	class FStaticMeshRenderData& RenderData,
+	const FString& DistanceFieldKey,
+	FSourceMeshDataForDerivedDataTask* OptionalSourceMeshData)
+{
+	BeginCacheMeshCardRepresentation(TargetPlatform, StaticMeshAsset, RenderData, DistanceFieldKey, OptionalSourceMeshData ? *OptionalSourceMeshData : FSourceMeshDataForDerivedDataTask());
+}
 #endif
 
 #if UE_ENABLE_INCLUDE_ORDER_DEPRECATED_IN_5_4
