Add Source Volume Modulation to Concurrency Evaluation

diff --git a/Engine/Source/Runtime/AudioMixer/Private/AudioMixerSource.cpp b/Engine/Source/Runtime/AudioMixer/Private/AudioMixerSource.cpp
--- a/Engine/Source/Runtime/AudioMixer/Private/AudioMixerSource.cpp
+++ b/Engine/Source/Runtime/AudioMixer/Private/AudioMixerSource.cpp
@@ -1995,6 +1995,9 @@ namespace Audio
 		}
 
 		ActiveSound->bModulationRoutingUpdated = false;
+
+		// Query a modulation value for the active sound to use during concurrency evaluation
+		ActiveSound->SourceModulationVolume = MixerSourceVoice->GetVolumeModulationValue();
 	}
 
 	void FMixerSource::UpdateSourceBusSends()
diff --git a/Engine/Source/Runtime/AudioMixer/Private/AudioMixerSourceManager.cpp b/Engine/Source/Runtime/AudioMixer/Private/AudioMixerSourceManager.cpp
--- a/Engine/Source/Runtime/AudioMixer/Private/AudioMixerSourceManager.cpp
+++ b/Engine/Source/Runtime/AudioMixer/Private/AudioMixerSourceManager.cpp
@@ -398,10 +398,12 @@ namespace Audio
 
 		GameThreadInfo.RelativeRenderCost.Reset(NumTotalSources);
 		GameThreadInfo.FreeSourceIndices.Reset(NumTotalSources);
+		GameThreadInfo.ModulationVolume.Reset(NumTotalSources);
 		for (int32 i = NumTotalSources - 1; i >= 0; --i)
 		{
 			GameThreadInfo.RelativeRenderCost.Add(1.0f);
 			GameThreadInfo.FreeSourceIndices.Add(i);
+			GameThreadInfo.ModulationVolume.Add(1.0f);
 		}
 
 		// Initialize the source buffer memory usage to max source scratch buffers (num frames times max source channels)
@@ -1959,6 +1961,12 @@ namespace Audio
 		return SourceInfos[SourceId].SourceEnvelopeValue;
 	}
 
+	float FMixerSourceManager::GetVolumeModulationValue(const int32 SourceId) const
+	{
+		AUDIO_MIXER_CHECK_GAME_THREAD(MixerDevice);
+		return GameThreadInfo.ModulationVolume[SourceId];
+	}
+
 #if ENABLE_AUDIO_DEBUG
 	double FMixerSourceManager::GetCPUCoreUtilization(const int32 SourceId) const
 	{
@@ -2732,6 +2740,7 @@ namespace Audio
 							VolumeStart *= ModVolumeEnd;
 						}
 						VolumeDestination *= ModVolumeEnd;
+						GameThreadInfo.ModulationVolume[SourceId] = ModVolumeEnd;
 					}
 
 					TArrayView<float> PreDistanceAttenBufferFadeSamplesView(PreDistanceAttenBufferPtr, NumFadeSamples);
@@ -2768,6 +2777,7 @@ namespace Audio
 							VolumeStart *= ModVolumeEnd;
 						}
 						VolumeDestination *= ModVolumeEnd;
+						GameThreadInfo.ModulationVolume[SourceId] = ModVolumeEnd;
 					}
 
 					Audio::ArrayFade(PreDistanceAttenBufferView, VolumeStart, VolumeDestination);
diff --git a/Engine/Source/Runtime/AudioMixer/Private/AudioMixerSourceManager.h b/Engine/Source/Runtime/AudioMixer/Private/AudioMixerSourceManager.h
--- a/Engine/Source/Runtime/AudioMixer/Private/AudioMixerSourceManager.h
+++ b/Engine/Source/Runtime/AudioMixer/Private/AudioMixerSourceManager.h
@@ -259,6 +259,9 @@ namespace Audio
 
 		int64 GetNumFramesPlayed(const int32 SourceId) const;
 		float GetEnvelopeValue(const int32 SourceId) const;
+
+		float GetVolumeModulationValue(const int32 SourceId) const;
+
 #if ENABLE_AUDIO_DEBUG
 		double GetCPUCoreUtilization(const int32 SourceId) const;
 #endif // ENABLE_AUDIO_DEBUG
@@ -646,6 +649,7 @@ namespace Audio
 			TArray<double> CPUCoreUtilization;
 #endif // #if ENABLE_AUDIO_DEBUG
 			TArray<float> RelativeRenderCost;
+			TArray<float> ModulationVolume;
 		} GameThreadInfo;
 
 		int32 NumActiveSources;
diff --git a/Engine/Source/Runtime/AudioMixer/Private/AudioMixerSourceVoice.cpp b/Engine/Source/Runtime/AudioMixer/Private/AudioMixerSourceVoice.cpp
--- a/Engine/Source/Runtime/AudioMixer/Private/AudioMixerSourceVoice.cpp
+++ b/Engine/Source/Runtime/AudioMixer/Private/AudioMixerSourceVoice.cpp
@@ -333,6 +333,13 @@ namespace Audio
 		return SourceManager->GetEnvelopeValue(SourceId);
 	}
 
+	float FMixerSourceVoice::GetVolumeModulationValue() const
+	{
+		AUDIO_MIXER_CHECK_GAME_THREAD(MixerDevice);
+
+		return SourceManager->GetVolumeModulationValue(SourceId);
+	}
+
 #if ENABLE_AUDIO_DEBUG
 	double FMixerSourceVoice::GetCPUCoreUtilization() const
 	{
diff --git a/Engine/Source/Runtime/AudioMixer/Private/AudioMixerSourceVoice.h b/Engine/Source/Runtime/AudioMixer/Private/AudioMixerSourceVoice.h
--- a/Engine/Source/Runtime/AudioMixer/Private/AudioMixerSourceVoice.h
+++ b/Engine/Source/Runtime/AudioMixer/Private/AudioMixerSourceVoice.h
@@ -117,6 +117,9 @@ namespace Audio
 		// Retrieves the envelope value of the source.
 		float GetEnvelopeValue() const;
 
+		// Retrieves final mixed Modulation value of the source.
+		float GetVolumeModulationValue() const;
+
 #if ENABLE_AUDIO_DEBUG
 		double GetCPUCoreUtilization() const;
 #endif // ENABLE_AUDIO_DEBUG
diff --git a/Engine/Source/Runtime/Engine/Private/ActiveSound.cpp b/Engine/Source/Runtime/Engine/Private/ActiveSound.cpp
--- a/Engine/Source/Runtime/Engine/Private/ActiveSound.cpp
+++ b/Engine/Source/Runtime/Engine/Private/ActiveSound.cpp
@@ -61,6 +61,13 @@ FAutoConsoleVariableRef CVarFocusDependsOnSpatialization(
 	TEXT("0: Focus and Spatialization are orthogonal, 1: Focus will be disabled if spatialization is not also enabled"),
 	ECVF_Default);
 
+static bool bIncludeModulationInConcurrencyEvaluationCVar = false;
+FAutoConsoleVariableRef CVarIncludeModulationInConcurrencyEvaluation(
+	TEXT("au.Concurrency.IncludeVolumeModulation"),
+	bIncludeModulationInConcurrencyEvaluationCVar,
+	TEXT("0: Ignore Audio Modulation Volume when evaluating concurrency (Default), 1: Include a sound's Volume Modulation Destination when determining whether to virtualize a sound based on volume."),
+	ECVF_Default);
+
 FTraceDelegate FActiveSound::ActiveSoundTraceDelegate;
 TMap<FTraceHandle, FActiveSound::FAsyncTraceDetails> FActiveSound::TraceToActiveSoundMap;
 
@@ -135,6 +142,7 @@ FActiveSound::FActiveSound()
 	, SubtitlePriority(DEFAULT_SUBTITLE_PRIORITY)
 	, Priority(1.0f)
 	, VolumeConcurrency(0.0f)
+	, SourceModulationVolume(0.0f)
 	, OcclusionCheckInterval(0.0f)
 	, LastOcclusionCheckTime(TNumericLimits<float>::Lowest())
 	, MaxDistance(FAudioDevice::GetMaxWorldDistance())
@@ -1056,6 +1064,11 @@ void FActiveSound::UpdateWaveInstances(TArray<FWaveInstance*> &InWaveInstances,
 				}
 			}
 
+			if (bIncludeModulationInConcurrencyEvaluationCVar)
+			{
+				VolumeConcurrency *= SourceModulationVolume;
+			}
+
 			// Remove concurrency volume scalars as this can cause ping-ponging to occur with virtualization and loops
 			// utilizing concurrency with rules that don't support eviction (removal from concurrency system prior to playback).
 			const float VolumeScale = GetTotalConcurrencyVolumeScale();
diff --git a/Engine/Source/Runtime/Engine/Private/Audio/AudioModulationInheritance.cpp b/Engine/Source/Runtime/Engine/Private/Audio/AudioModulationInheritance.cpp
new file mode 100644
--- /dev/null
+++ b/Engine/Source/Runtime/Engine/Private/Audio/AudioModulationInheritance.cpp
@@ -0,0 +1,146 @@
+// Copyright Epic Games, Inc. All Rights Reserved.
+#include "AudioModulationInheritance.h"
+
+#include "Sound/SoundModulationDestination.h"
+#include "Sound/SoundClass.h"
+
+namespace AudioModulationInheritance
+{
+	float MixInRoutedValue(const Audio::FModulationParameter& InParam, float InValueA, float InValueB)
+	{
+		if (InParam.bRequiresConversion)
+		{
+			InParam.NormalizedFunction(InValueA);
+			InParam.NormalizedFunction(InValueB);
+		}
+		InParam.MixFunction(InValueA, InValueB);
+		if (InParam.bRequiresConversion)
+		{
+			InParam.UnitFunction(InValueA);
+		}
+
+		return InValueA;
+	}
+
+	// This function gets all the Modulators which apply to a sound, factoring in the inheritance path
+	// ie. Audio Component inherits modulation settings from Sound, and sounds inherit from sound class.
+	TSet<TObjectPtr<USoundModulatorBase>> GetEffectiveVolumeModulators(const FActiveSound& InActiveSound)
+	{
+		FWaveInstance* WaveInstance = nullptr;
+		// For the "Sound" level of the inheritance path, we just grab the first WaveInstance we find
+		// Todo: How to resolve multiple wave instances with different modulation data? eg. Sound Cues
+		for (const TPair<UPTRINT, FWaveInstance*>& Pair : InActiveSound.GetWaveInstances())
+		{
+			WaveInstance = Pair.Value;
+			break;
+		}
+
+		EModulationRouting InWaveRouting = EModulationRouting::Inherit;
+		FSoundModulationDestinationSettings InWaveSettings = { };
+
+		if (WaveInstance)
+		{
+			UE_LOG(LogAudio, Verbose, TEXT("Using the first WaveInstance ('%s') on the ActiveSound to determine list of Modulators."), *WaveInstance->GetName());
+			InWaveRouting = WaveInstance->WaveData->ModulationSettings.VolumeRouting;
+			InWaveSettings = WaveInstance->WaveData->ModulationSettings.VolumeModulationDestination;
+		}
+
+		const USoundClass* InSoundClass = InActiveSound.GetSoundClass();
+		const Audio::FModulationParameter& InParam = Audio::GetModulationParameter("Volume");
+
+		auto UnionSoundClassSettings = [&](FSoundModulationDestinationSettings& InOutSettings)
+			{
+				if (InSoundClass)
+				{
+					const FSoundModulationDestinationSettings& ClassSettings = InSoundClass->Properties.ModulationSettings.VolumeModulationDestination;
+					InOutSettings.Value = MixInRoutedValue(InParam, InOutSettings.Value, ClassSettings.Value);
+					InOutSettings.Modulators = InOutSettings.Modulators.Union(ClassSettings.Modulators);
+				}
+			};
+
+		switch (InActiveSound.ModulationRouting.VolumeRouting)
+		{
+		case EModulationRouting::Union:
+		{
+			FSoundModulationDestinationSettings UnionSettings = InActiveSound.ModulationRouting.VolumeModulationDestination;
+			switch (InWaveRouting)
+			{
+			case EModulationRouting::Union:
+			{
+				UnionSettings.Value = MixInRoutedValue(InParam, UnionSettings.Value, InWaveSettings.Value);
+				UnionSettings.Modulators = UnionSettings.Modulators.Union(InWaveSettings.Modulators);
+				UnionSoundClassSettings(UnionSettings);
+				return UnionSettings.Modulators;
+			}
+			break;
+
+			case EModulationRouting::Inherit:
+			{
+				UnionSoundClassSettings(UnionSettings);
+			}
+			break;
+
+			case EModulationRouting::Override:
+			{
+				UnionSettings.Value = MixInRoutedValue(InParam, UnionSettings.Value, InWaveSettings.Value);
+				UnionSettings.Modulators = UnionSettings.Modulators.Union(InWaveSettings.Modulators);
+			}
+			break;
+
+			case EModulationRouting::Disable:
+			default:
+				break;
+			}
+
+			return UnionSettings.Modulators;
+		}
+		break;
+
+		case EModulationRouting::Inherit:
+		{
+			switch (InWaveRouting)
+			{
+			case EModulationRouting::Union:
+			{
+				FSoundModulationDestinationSettings UnionSettings = InWaveSettings;
+				UnionSoundClassSettings(UnionSettings);
+				return UnionSettings.Modulators;
+			}
+			break;
+
+			case EModulationRouting::Inherit:
+			{
+				if (InSoundClass)
+				{
+					return InSoundClass->Properties.ModulationSettings.VolumeModulationDestination.Modulators;
+				}
+			}
+			break;
+
+			case EModulationRouting::Override:
+			{
+				return InWaveSettings.Modulators;
+			}
+			break;
+
+			case EModulationRouting::Disable:
+			default:
+				break;
+			}
+		}
+		break;
+
+		case EModulationRouting::Override:
+		{
+			return InActiveSound.ModulationRouting.VolumeModulationDestination.Modulators;
+		}
+		break;
+
+		case EModulationRouting::Disable:
+		default:
+			break;
+		}
+
+		return TSet<TObjectPtr<USoundModulatorBase>>();
+	}
+}
diff --git a/Engine/Source/Runtime/Engine/Private/Audio/AudioModulationInheritance.h b/Engine/Source/Runtime/Engine/Private/Audio/AudioModulationInheritance.h
new file mode 100644
--- /dev/null
+++ b/Engine/Source/Runtime/Engine/Private/Audio/AudioModulationInheritance.h
@@ -0,0 +1,11 @@
+// Copyright Epic Games, Inc. All Rights Reserved.
+#include "ActiveSound.h"
+#include "IAudioModulation.h"
+
+namespace AudioModulationInheritance
+{
+	float MixInRoutedValue(const Audio::FModulationParameter& InParam, float InValueA, float InValueB);
+
+	TSet<TObjectPtr<USoundModulatorBase>> GetEffectiveVolumeModulators(const FActiveSound& InActiveSound);
+}
+
diff --git a/Engine/Source/Runtime/Engine/Private/AudioVirtualLoop.cpp b/Engine/Source/Runtime/Engine/Private/AudioVirtualLoop.cpp
--- a/Engine/Source/Runtime/Engine/Private/AudioVirtualLoop.cpp
+++ b/Engine/Source/Runtime/Engine/Private/AudioVirtualLoop.cpp
@@ -4,6 +4,7 @@
 #include "Audio/AudioDebug.h"
 #include "AudioDevice.h"
 #include "AudioMixerTrace.h"
+#include "Audio/AudioModulationInheritance.h"
 
 
 static int32 bVirtualLoopsEnabledCVar = 1;
@@ -74,6 +75,10 @@ UE_TRACE_EVENT_BEGIN(Audio, VirtualLoopUpdate)
 UE_TRACE_EVENT_END()
 #endif // UE_AUDIO_PROFILERTRACE_ENABLED
 
+namespace AudioVirtualLoopPrivate
+{
+	static const FLazyName VolumeModulationParameterName("Volume");
+}
 
 FAudioVirtualLoop::FAudioVirtualLoop()
 	: TimeSinceLastUpdate(0.0f)
@@ -125,6 +130,11 @@ bool FAudioVirtualLoop::Virtualize(const FActiveSound& InActiveSound, FAudioDevi
 	OutVirtualLoop.ActiveSound = ActiveSound;
 	OutVirtualLoop.CalculateUpdateInterval();
 
+	// Once virtualized, Modulation stops getting updated at the Mixer Source level. Because of this, we need to grab
+	// Modulation data now so that Volume modulation can be used to determine if a sound is loud enough to realize.
+	OutVirtualLoop.VolumeConcurrencyDestination.Init(InAudioDevice.DeviceID, AudioVirtualLoopPrivate::VolumeModulationParameterName, false /* bInIsBuffered */, true /* bInValueNormalized */);
+	OutVirtualLoop.VolumeConcurrencyDestination.UpdateModulators(AudioModulationInheritance::GetEffectiveVolumeModulators(InActiveSound));
+
 #if UE_AUDIO_PROFILERTRACE_ENABLED
 	const bool bChannelEnabled = UE_TRACE_CHANNELEXPR_IS_ENABLED(AudioChannel);
 	if (bChannelEnabled)
@@ -267,6 +277,10 @@ bool FAudioVirtualLoop::Update(float DeltaTime, bool bForceUpdate)
 		TimeSinceLastUpdate = 0.0f;
 	}
 
+	// ProcessControl only takes values in the Parameter's unit space, so we pass in 0dB as the default value.
+	VolumeConcurrencyDestination.ProcessControl(0.0f);
+	ActiveSound->SourceModulationVolume = VolumeConcurrencyDestination.GetValue();
+
 #if ENABLE_AUDIO_DEBUG
 	Audio::FAudioDebugger::DrawDebugInfo(*this);
 #endif // ENABLE_AUDIO_DEBUG
diff --git a/Engine/Source/Runtime/Engine/Public/ActiveSound.h b/Engine/Source/Runtime/Engine/Public/ActiveSound.h
--- a/Engine/Source/Runtime/Engine/Public/ActiveSound.h
+++ b/Engine/Source/Runtime/Engine/Public/ActiveSound.h
@@ -596,6 +596,9 @@ public:
 	// If negative, tracking is disabled for lifetime of ActiveSound */
 	float VolumeConcurrency;
 
+	/** The source's final mixed Modulation value, used to evaluate VolumeConcurrency */
+	float SourceModulationVolume;
+
 	/** The time in seconds with which to check for occlusion from its closest listener */
 	float OcclusionCheckInterval;
 
diff --git a/Engine/Source/Runtime/Engine/Public/AudioVirtualLoop.h b/Engine/Source/Runtime/Engine/Public/AudioVirtualLoop.h
--- a/Engine/Source/Runtime/Engine/Public/AudioVirtualLoop.h
+++ b/Engine/Source/Runtime/Engine/Public/AudioVirtualLoop.h
@@ -2,7 +2,7 @@
 #pragma once
 
 #include "ActiveSound.h"
-
+#include "Sound/SoundModulationDestination.h"
 
 /**
  * Class that tracks virtualized looping active sounds that are eligible to revive re-trigger
@@ -17,6 +17,9 @@ private:
 
 	FActiveSound* ActiveSound;
 
+	/** Modulation Destination for virtualized ActiveSounds to use when evaluating volume for concurrency */
+	Audio::FModulationDestination VolumeConcurrencyDestination;
+
 	/**
 	  * Check if provided active sound is in audible range.
 	  */
