Support BoundsScale on ISM.
This applies to the component, with no change to the per instance bounds.
One use is that setting a large BoundScale can prevent CPU side culling of the component which can be useful for some special cases where WPO is heavily used.

diff --git a/Engine/Source/Runtime/Engine/Private/HierarchicalInstancedStaticMesh.cpp b/Engine/Source/Runtime/Engine/Private/HierarchicalInstancedStaticMesh.cpp
--- a/Engine/Source/Runtime/Engine/Private/HierarchicalInstancedStaticMesh.cpp
+++ b/Engine/Source/Runtime/Engine/Private/HierarchicalInstancedStaticMesh.cpp
@@ -2023,8 +2023,10 @@ FBoxSphereBounds UHierarchicalInstancedStaticMeshComponent::CalcBounds(const FTr
 
 	if (BuiltInstanceBounds.IsValid || UnbuiltInstanceBounds.IsValid)
 	{
-		FBoxSphereBounds Result = BuiltInstanceBounds + UnbuiltInstanceBounds;
-		return Result.TransformBy(BoundTransform);
+		FBoxSphereBounds Result = (BuiltInstanceBounds + UnbuiltInstanceBounds).TransformBy(BoundTransform);
+		Result.BoxExtent *= BoundsScale;
+		Result.SphereRadius *= BoundsScale;
+		return Result;
 	}
 	else
 	{
diff --git a/Engine/Source/Runtime/Engine/Private/InstancedStaticMesh.cpp b/Engine/Source/Runtime/Engine/Private/InstancedStaticMesh.cpp
--- a/Engine/Source/Runtime/Engine/Private/InstancedStaticMesh.cpp
+++ b/Engine/Source/Runtime/Engine/Private/InstancedStaticMesh.cpp
@@ -2841,7 +2841,10 @@ void UInstancedStaticMeshComponent::InvalidateCachedBounds()
 FBoxSphereBounds UInstancedStaticMeshComponent::CalcBounds(const FTransform& InBoundTransform) const
 {
 	EBoundsCacheType Type = bIsUpdatingBounds ? EBoundsCacheType::Component : InBoundTransform.Equals(FTransform::Identity) ? EBoundsCacheType::Local : EBoundsCacheType::Other;
-	return GetOrCalcBounds(InBoundTransform, Type);
+	FBoxSphereBounds NewBounds = GetOrCalcBounds(InBoundTransform, Type);
+	NewBounds.BoxExtent *= BoundsScale;
+	NewBounds.SphereRadius *= BoundsScale;
+	return NewBounds;
 }
 
 void UInstancedStaticMeshComponent::CalcAndCacheNavigationBounds()
