Add Editor BlueprintCallable function for set whether texture is virtual texture.
Useful for batch texture processing.

diff --git a/Engine/Source/Runtime/Engine/Classes/Engine/Texture.h b/Engine/Source/Runtime/Engine/Classes/Engine/Texture.h
--- a/Engine/Source/Runtime/Engine/Classes/Engine/Texture.h
+++ b/Engine/Source/Runtime/Engine/Classes/Engine/Texture.h
@@ -1558,6 +1558,11 @@ PRAGMA_ENABLE_DEPRECATION_WARNINGS
 	UPROPERTY(EditAnywhere, BlueprintReadOnly, Category = Texture, AssetRegistrySearchable, AdvancedDisplay)
 	uint8 VirtualTextureStreaming : 1;
 
+#if WITH_EDITOR
+	UFUNCTION(BlueprintCallable, Category = Texture)
+	void SetVirtualTextureStreaming(bool bVirtualTextureStreaming);
+#endif
+
 	/** If true the texture stores YCoCg. Blue channel will be filled with a precision scale during compression. */
 	UPROPERTY()
 	uint8 CompressionYCoCg : 1;
diff --git a/Engine/Source/Runtime/Engine/Private/Texture.cpp b/Engine/Source/Runtime/Engine/Private/Texture.cpp
--- a/Engine/Source/Runtime/Engine/Private/Texture.cpp
+++ b/Engine/Source/Runtime/Engine/Private/Texture.cpp
@@ -997,6 +997,17 @@ PRAGMA_ENABLE_DEPRECATION_WARNINGS
 	}
 }
 
+void UTexture::SetVirtualTextureStreaming(bool bVirtualTextureStreaming)
+{
+	if (VirtualTextureStreaming != bVirtualTextureStreaming)
+	{
+		VirtualTextureStreaming = bVirtualTextureStreaming;
+
+		FPropertyChangedEvent PropertyChangedEvent(UTexture::StaticClass()->FindPropertyByName(GET_MEMBER_NAME_CHECKED(UTexture, VirtualTextureStreaming)));
+		PostEditChangeProperty(PropertyChangedEvent);
+	}
+}
+
 #endif // WITH_EDITOR
 
 #if WITH_EDITORONLY_DATA
diff --git a/Engine/Source/Runtime/Engine/Private/Texture2D.cpp b/Engine/Source/Runtime/Engine/Private/Texture2D.cpp
--- a/Engine/Source/Runtime/Engine/Private/Texture2D.cpp
+++ b/Engine/Source/Runtime/Engine/Private/Texture2D.cpp
@@ -542,7 +542,8 @@ void UTexture2D::PostEditUndo()
 void UTexture2D::PostEditChangeProperty(FPropertyChangedEvent& PropertyChangedEvent)
 {
 	if (PropertyChangedEvent.GetPropertyName() == GET_MEMBER_NAME_CHECKED(UTexture2D, AddressX)
-		|| PropertyChangedEvent.GetPropertyName() == GET_MEMBER_NAME_CHECKED(UTexture2D, AddressY))
+		|| PropertyChangedEvent.GetPropertyName() == GET_MEMBER_NAME_CHECKED(UTexture2D, AddressY)
+		|| PropertyChangedEvent.GetPropertyName() == GET_MEMBER_NAME_CHECKED(UTexture2D, VirtualTextureStreaming))
 	{
 		// VT need to recompile shaders when address mode changes
 		// Non-VT still needs to potentially update sampler state in the materials
