Backout of 38341618.  The original backout introduced an upgrade bug where the skeletal mesh asset property, on blueprint-added skel mesh components, that were overriden on actors within the level viewport would be reset to their default value when upgrading to 5.5.2.  This change reintroduces a bug where the skeletal mesh asset can't be set on a blueprint added mesh component via construction script.

diff --git a/Engine/Source/Runtime/Engine/Private/Components/SkeletalMeshComponentInstanceData.cpp b/Engine/Source/Runtime/Engine/Private/Components/SkeletalMeshComponentInstanceData.cpp
--- a/Engine/Source/Runtime/Engine/Private/Components/SkeletalMeshComponentInstanceData.cpp
+++ b/Engine/Source/Runtime/Engine/Private/Components/SkeletalMeshComponentInstanceData.cpp
@@ -7,6 +7,7 @@
 
 FSkeletalMeshComponentInstanceData::FSkeletalMeshComponentInstanceData(const USkeletalMeshComponent* SourceComponent)
 	: FSceneComponentInstanceData(SourceComponent)
+	, SkeletalMeshAsset(nullptr)
 	, bUpdateAnimationInEditor(0)
 	, bUpdateClothInEditor(0)
 {
@@ -17,6 +18,7 @@ FSkeletalMeshComponentInstanceData::FSkeletalMeshComponentInstanceData(const USk
 		|| SourceComponent->CreationMethod == EComponentCreationMethod::UserConstructionScript;
 	if (bIsBlueprintCreatedComponent)
 	{
+		SkeletalMeshAsset = SourceComponent->GetSkeletalMeshAsset();
 		bUpdateAnimationInEditor = SourceComponent->GetUpdateAnimationInEditor();
 		bUpdateClothInEditor = SourceComponent->GetUpdateClothInEditor();
 	}
@@ -25,7 +27,7 @@ FSkeletalMeshComponentInstanceData::FSkeletalMeshComponentInstanceData(const USk
 
 bool FSkeletalMeshComponentInstanceData::ContainsData() const
 {
-	return bUpdateAnimationInEditor || bUpdateClothInEditor;
+	return bUpdateAnimationInEditor || bUpdateClothInEditor || SkeletalMeshAsset;
 }
 
 void FSkeletalMeshComponentInstanceData::ApplyToComponent(UActorComponent* Component, const ECacheApplyPhase CacheApplyPhase)
@@ -33,9 +35,16 @@ void FSkeletalMeshComponentInstanceData::ApplyToComponent(UActorComponent* Compo
 	Super::ApplyToComponent(Component, CacheApplyPhase);
 
 #if WITH_EDITOR
-	if (USkeletalMeshComponent * SkeletalMesh = Cast<USkeletalMeshComponent>(Component))
+	const bool bIsBlueprintCreatedComponent = Component->CreationMethod == EComponentCreationMethod::SimpleConstructionScript
+		|| Component->CreationMethod == EComponentCreationMethod::UserConstructionScript;
+	USkeletalMeshComponent* SkeletalMesh = Cast<USkeletalMeshComponent>(Component);
+
+	if (SkeletalMesh && bIsBlueprintCreatedComponent)
 	{
-		SkeletalMesh->RefreshSkeletalMeshAsset();
+		if (SkeletalMeshAsset)
+		{
+			SkeletalMesh->SetSkeletalMeshAsset(SkeletalMeshAsset);
+		}
 		SkeletalMesh->SetUpdateAnimationInEditor(bUpdateAnimationInEditor);
 		SkeletalMesh->SetUpdateClothInEditor(bUpdateClothInEditor);
 	}
diff --git a/Engine/Source/Runtime/Engine/Private/Components/SkeletalMeshComponentInstanceData.h b/Engine/Source/Runtime/Engine/Private/Components/SkeletalMeshComponentInstanceData.h
--- a/Engine/Source/Runtime/Engine/Private/Components/SkeletalMeshComponentInstanceData.h
+++ b/Engine/Source/Runtime/Engine/Private/Components/SkeletalMeshComponentInstanceData.h
@@ -6,6 +6,7 @@
 #include "SkeletalMeshComponentInstanceData.generated.h"
 
 class USkeletalMeshComponent;
+class USkeletalMesh;
 
 /** Saves internal SkeletalMesh (transient) state that gets lost at reconstruction for Blueprint created components. */
 USTRUCT()
@@ -19,6 +20,7 @@ struct FSkeletalMeshComponentInstanceData : public FSceneComponentInstanceData
 	virtual bool ContainsData() const override;
 	virtual void ApplyToComponent(UActorComponent* Component, const ECacheApplyPhase CacheApplyPhase) override;
 
+	USkeletalMesh* SkeletalMeshAsset;
 	uint8 bUpdateAnimationInEditor : 1;
 	uint8 bUpdateClothInEditor : 1;
 };
