Chaos - GravityGroupIndex on FBodyInstance is now being used

diff --git a/Engine/Source/Runtime/Engine/Private/PhysicsEngine/BodyInstance.cpp b/Engine/Source/Runtime/Engine/Private/PhysicsEngine/BodyInstance.cpp
--- a/Engine/Source/Runtime/Engine/Private/PhysicsEngine/BodyInstance.cpp
+++ b/Engine/Source/Runtime/Engine/Private/PhysicsEngine/BodyInstance.cpp
@@ -1299,6 +1299,7 @@ void TInitBodiesHelperBase<InAllocatorType>::CreateActor_AssumesLocked(FBodyInst
 		FPhysicsInterface::SetPositionSolverIterationCount_AssumesLocked(ActorHandle, Instance->PositionSolverIterationCount);
 		FPhysicsInterface::SetVelocitySolverIterationCount_AssumesLocked(ActorHandle, Instance->VelocitySolverIterationCount);
 		FPhysicsInterface::SetProjectionSolverIterationCount_AssumesLocked(ActorHandle, Instance->ProjectionSolverIterationCount);
+		FPhysicsInterface::SetGravityGroupIndex_AssumesLocked(ActorHandle, Instance->GravityGroupIndex);
 
 		// Set sleep event notification
 		FPhysicsInterface::SetSendsSleepNotifies_AssumesLocked(ActorHandle, Instance->bGenerateWakeEvents);
