Fixed shader compilation error when blue noise is used in vertex shaders.
Also added +1 texture sample cost info.

#tests PC editor win64
#rnx

diff --git a/Engine/Source/Runtime/Engine/Private/Materials/HLSLMaterialTranslator.cpp b/Engine/Source/Runtime/Engine/Private/Materials/HLSLMaterialTranslator.cpp
--- a/Engine/Source/Runtime/Engine/Private/Materials/HLSLMaterialTranslator.cpp
+++ b/Engine/Source/Runtime/Engine/Private/Materials/HLSLMaterialTranslator.cpp
@@ -11944,8 +11944,10 @@ int32 FHLSLMaterialTranslator::ScalarBlueNoise()
 	{
 		Errorf(TEXT("ScalarBlueNoise node can only be used within Surface domain materials %s (asset: %s).\r\n"), *Material->GetDebugName(), *Material->GetAssetPath().ToString());
 	}
-	
-	return AddCodeChunk(MCT_Float1, TEXT("ViewScalarBlueNoise(Parameters.SvPosition.xy, View.StateFrameIndex)"));
+
+	AddEstimatedTextureSample(1);
+
+	return AddCodeChunk(MCT_Float1, TEXT("ViewScalarBlueNoise(GetPixelPosition(Parameters), View.StateFrameIndex)"));
 }
 
 int32 FHLSLMaterialTranslator::BlackBody( int32 Temp )
