SceneGraph
* Fix EntityReplicator filter parameters not set correctly in BeginReplication
* Fix entities forced to be always relevant if their transform component is added before the replication component is.

EntityReplicationTests
* Improve handling of temp CVars
* Added SceneGraph.Replication(EngineOnly).SpatialEntities unit test

#rnx

diff --git a/Engine/Source/Runtime/Engine/Private/Tests/NetTestHelpers.cpp b/Engine/Source/Runtime/Engine/Private/Tests/NetTestHelpers.cpp
--- a/Engine/Source/Runtime/Engine/Private/Tests/NetTestHelpers.cpp
+++ b/Engine/Source/Runtime/Engine/Private/Tests/NetTestHelpers.cpp
@@ -487,9 +487,9 @@ UObject* FTestWorlds::FindReplicatedObjectOnClient(UObject* ServerObject, uint32
 	const FTestWorldInstance::FContext ServerContext = Server.GetTestContext();
 	const FTestWorldInstance::FContext ClientContext = Clients[ClientIndex].GetTestContext();
 
-#if UE_WITH_IRIS
 	if (ServerContext.NetDriver->IsUsingIrisReplication())
 	{
+#if UE_WITH_IRIS
 		const FNetRefHandle NetHandle = ServerContext.IrisRepBridge->GetReplicatedRefHandle(ServerObject);
 
 		if (!NetHandle.IsValid())
@@ -499,21 +499,30 @@ UObject* FTestWorlds::FindReplicatedObjectOnClient(UObject* ServerObject, uint32
 		}
 
 		return ClientContext.IrisRepBridge->GetReplicatedObject(NetHandle);
+
+#else
+		return nullptr;
+#endif
 	}
 	else
-#endif // UE_WITH_IRIS
 	{
 		const FNetworkGUID NetGUID = ServerContext.NetDriver->GetNetGuidCache()->GetNetGUID(ServerObject);
 
 		if (!NetGUID.IsValid())
 		{
 			ensureMsgf(false, TEXT("FTestWorlds::FindReplicatedObjectOnClient ServerObject: %s is not replicated."), *GetFullNameSafe(ServerObject));
+			return nullptr;
 		}
 
 		return ClientContext.NetDriver->GetNetGuidCache()->GetObjectFromNetGUID(NetGUID, false);
 	}
 }
 
+bool FTestWorlds::DoesReplicatedObjectExistOnClient(UObject* ServerObject, uint32 ClientIndex) const
+{
+	return IsValid(FindReplicatedObjectOnClient(ServerObject, ClientIndex));
+}
+
 //------------------------------------------------------------------------
 
 FScopedCVarOverrideInt::FScopedCVarOverrideInt(const TCHAR* VariableName, int32 Value)
diff --git a/Engine/Source/Runtime/Engine/Public/Tests/NetTestHelpers.h b/Engine/Source/Runtime/Engine/Public/Tests/NetTestHelpers.h
--- a/Engine/Source/Runtime/Engine/Public/Tests/NetTestHelpers.h
+++ b/Engine/Source/Runtime/Engine/Public/Tests/NetTestHelpers.h
@@ -170,6 +170,7 @@ struct FTestWorlds
 	ENGINE_API ~FTestWorlds();
 
 	ENGINE_API void SetTickInSeconds(float TickInSeconds);
+	float GetTickDeltaSeconds() const { return TickDeltaSeconds; }
 	
 	/** Create a client and fully connect it to the server before returning */
 	ENGINE_API bool CreateAndConnectClient();
@@ -201,13 +202,20 @@ struct FTestWorlds
 	ENGINE_API APlayerController* GetServerPlayerControllerOfClient(uint32 ClientIndex);
 
 	/**
-	 * Find the remote instance of a replicated object.
+	 * Find the remote instance of a replicated object and ensures if it can't find it
 	 * @param ServerObject The server object that you want the remote version of.
 	 * @param ClientIndex The client you want a remote instance of
 	 * @return Return the remote (client) instance of the same object if it exists.
 	 */
 	ENGINE_API UObject* FindReplicatedObjectOnClient(UObject* ServerObject, uint32 ClientIndex) const;
 
+	/**
+	 * Returns if the the remote instance of a replicated object exists on the client or not.
+	 * @param ServerObject The server object you want to see if it is replicated
+	 * @param ClientIndex The client you want to test on
+	 */
+	ENGINE_API bool DoesReplicatedObjectExistOnClient(UObject* ServerObject, uint32 ClientIndex) const;
+
 public:
 
 	/** Server and Client Worlds */
