[NewMaterialTranslator] - Implement frac

#rnx

diff --git a/Engine/Source/Runtime/Engine/Private/Materials/MaterialExpressionsIR.cpp b/Engine/Source/Runtime/Engine/Private/Materials/MaterialExpressionsIR.cpp
--- a/Engine/Source/Runtime/Engine/Private/Materials/MaterialExpressionsIR.cpp
+++ b/Engine/Source/Runtime/Engine/Private/Materials/MaterialExpressionsIR.cpp
@@ -482,6 +482,11 @@ void UMaterialExpressionFloor::Build(UE::MIR::FEmitter& Emitter)
 	BuildUnaryArithmeticOperatorFromInput(Emitter, MIR::UO_Floor, &Input);
 }
 
+void UMaterialExpressionFrac::Build(UE::MIR::FEmitter& Emitter)
+{
+	BuildUnaryArithmeticOperatorFromInput(Emitter, MIR::UO_Frac, &Input);
+}
+
 void UMaterialExpressionLength::Build(UE::MIR::FEmitter& Emitter)
 {
 	MIR::FValue* Value = Emitter.Input(&Input);
diff --git a/Engine/Source/Runtime/Engine/Private/Materials/MaterialIR.cpp b/Engine/Source/Runtime/Engine/Private/Materials/MaterialIR.cpp
--- a/Engine/Source/Runtime/Engine/Private/Materials/MaterialIR.cpp
+++ b/Engine/Source/Runtime/Engine/Private/Materials/MaterialIR.cpp
@@ -417,6 +417,7 @@ const TCHAR* OperatorToString(EOperator Op)
 		case UO_Exponential: return TEXT("Exponential");
 		case UO_Exponential2: return TEXT("Exponential2");
 		case UO_Floor: return TEXT("Floor");
+		case UO_Frac: return TEXT("Frac");
 		case UO_Length: return TEXT("Length");
 		case UO_Logarithm: return TEXT("Logarithm");
 		case UO_Logarithm2: return TEXT("Logarithm2");
diff --git a/Engine/Source/Runtime/Engine/Private/Materials/MaterialIREmitter.cpp b/Engine/Source/Runtime/Engine/Private/Materials/MaterialIREmitter.cpp
--- a/Engine/Source/Runtime/Engine/Private/Materials/MaterialIREmitter.cpp
+++ b/Engine/Source/Runtime/Engine/Private/Materials/MaterialIREmitter.cpp
@@ -794,6 +794,10 @@ static TOptional<T> FoldScalarArithmeticOperator(FEmitter* Emitter, EOperator Op
 				{
 					return FMath::FloorToDouble(A);
 				}
+			case UO_Frac:
+			{
+				return FGenericPlatformMath::Fractional(A);
+			}
 			case UO_Logarithm:
 			case UO_Logarithm2:
 			case UO_Logarithm10:
@@ -936,6 +940,7 @@ static FValue* TryFoldOperatorScalar(FEmitter* Emitter, EOperator Op, FValue* A,
 		case UO_Exponential:
 		case UO_Exponential2:
 		case UO_Floor:
+		case UO_Frac:
 		case UO_Logarithm:
 		case UO_Logarithm10:
 		case UO_Logarithm2:
@@ -1078,6 +1083,7 @@ static const FOperatorSignature* GetOperatorSignature(EOperator Op)
 		S[UO_Exponential] 			= UnaryArithmetic;
 		S[UO_Exponential2] 			= UnaryArithmetic;
 		S[UO_Floor] 				= UnaryArithmetic;
+		S[UO_Frac]					= UnaryArithmetic;
 		S[UO_Length]				= { { OT_Arithmetic | OT_NotMatrix }, ORT_FirstArgumentTypeToScalar };
 		S[UO_Logarithm] 			= UnaryArithmetic;
 		S[UO_Logarithm2] 			= UnaryArithmetic;
diff --git a/Engine/Source/Runtime/Engine/Private/Materials/MaterialIRToHLSLTranslator.cpp b/Engine/Source/Runtime/Engine/Private/Materials/MaterialIRToHLSLTranslator.cpp
--- a/Engine/Source/Runtime/Engine/Private/Materials/MaterialIRToHLSLTranslator.cpp
+++ b/Engine/Source/Runtime/Engine/Private/Materials/MaterialIRToHLSLTranslator.cpp
@@ -492,6 +492,7 @@ struct FTranslator : FMaterialIRToHLSLTranslation
 				case MIR::UO_Ceil: OpString = TEXT("ceil"); break;
 				case MIR::UO_Cos: OpString = TEXT("cos"); break;
 				case MIR::UO_Floor: OpString = TEXT("floor"); break;
+				case MIR::UO_Frac: OpString = TEXT("frac"); break;
 				case MIR::UO_Length: OpString = TEXT("length"); break;
 				case MIR::UO_Sin: OpString = TEXT("sin"); break;
 				case MIR::UO_Tan: OpString = TEXT("tan"); break;
diff --git a/Engine/Source/Runtime/Engine/Public/Materials/MaterialExpressionFrac.h b/Engine/Source/Runtime/Engine/Public/Materials/MaterialExpressionFrac.h
--- a/Engine/Source/Runtime/Engine/Public/Materials/MaterialExpressionFrac.h
+++ b/Engine/Source/Runtime/Engine/Public/Materials/MaterialExpressionFrac.h
@@ -22,6 +22,7 @@ public:
 
 	//~ Begin UMaterialExpression Interface
 #if WITH_EDITOR
+	virtual void Build(UE::MIR::FEmitter& Emitter) override;
 	virtual int32 Compile(class FMaterialCompiler* Compiler, int32 OutputIndex) override;
 	virtual void GetCaption(TArray<FString>& OutCaptions) const override;
 #endif
diff --git a/Engine/Source/Runtime/Engine/Public/Materials/MaterialIR.h b/Engine/Source/Runtime/Engine/Public/Materials/MaterialIR.h
--- a/Engine/Source/Runtime/Engine/Public/Materials/MaterialIR.h
+++ b/Engine/Source/Runtime/Engine/Public/Materials/MaterialIR.h
@@ -322,6 +322,7 @@ enum EOperator
 	UO_Exponential,
 	UO_Exponential2,
 	UO_Floor,
+	UO_Frac,
 	UO_Length,
 	UO_Logarithm,
 	UO_Logarithm2,
