Fix for different replication support in NetTestHelpers.

diff --git a/Engine/Source/Runtime/Engine/Private/Tests/NetTestHelpers.cpp b/Engine/Source/Runtime/Engine/Private/Tests/NetTestHelpers.cpp
--- a/Engine/Source/Runtime/Engine/Private/Tests/NetTestHelpers.cpp
+++ b/Engine/Source/Runtime/Engine/Private/Tests/NetTestHelpers.cpp
@@ -40,6 +40,10 @@ FTestWorldInstance FTestWorldInstance::CreateServer(const TCHAR* InURL, const ER
 	{
 		SetUseIrisReplication(true);
 	}
+	else if (ReplicationSystem == EReplicationSystem::Generic)
+	{
+		SetUseIrisReplication(false);
+	}
 
 	FGameInstancePIEParameters ServerParams;
 	ServerParams.bSimulateInEditor = false;
@@ -58,7 +62,10 @@ FTestWorldInstance FTestWorldInstance::CreateServer(const TCHAR* InURL, const ER
 	FString BrowseError;
 	GEngine->Browse(*NewInstance.GameInstance->GetWorldContext(), LocalURL, BrowseError);
 
-	SetUseIrisReplication(bPrevShouldUseIrisReplication);
+	if (ReplicationSystem != EReplicationSystem::Default)
+	{
+		SetUseIrisReplication(bPrevShouldUseIrisReplication);
+	}
 
 	return NewInstance;
 }
diff --git a/Engine/Source/Runtime/Engine/Public/Tests/NetTestHelpers.h b/Engine/Source/Runtime/Engine/Public/Tests/NetTestHelpers.h
--- a/Engine/Source/Runtime/Engine/Public/Tests/NetTestHelpers.h
+++ b/Engine/Source/Runtime/Engine/Public/Tests/NetTestHelpers.h
@@ -94,7 +94,7 @@ struct FTestWorldInstance
 	struct FContext;
 
 public:
-	ENGINE_API static FTestWorldInstance CreateServer(const TCHAR* InURL, const EReplicationSystem ReplicationSystem = EReplicationSystem::Generic);
+	ENGINE_API static FTestWorldInstance CreateServer(const TCHAR* InURL, const EReplicationSystem ReplicationSystem = EReplicationSystem::Default);
 	ENGINE_API static FTestWorldInstance CreateClient(int32 ServerPort);
 	ENGINE_API static FTestWorldInstance CreateProxy();
 
@@ -155,17 +155,17 @@ private:
 struct FTestWorlds
 {
 	/** Creates a server world using the default map (Entry) and default GameMode class (GameMode) */
-	ENGINE_API FTestWorlds(const EReplicationSystem ReplicationSystem = EReplicationSystem::Generic);
+	ENGINE_API FTestWorlds(const EReplicationSystem ReplicationSystem = EReplicationSystem::Default);
 
 	/** 
 	 * Creates a server world using a specific map name and GameMode class
 	 * @param MapName Full path name of the map. ex: /Engine/Maps/Entry
 	 * @param GameModeName Full path name of the game mode. ex: /Script/EngineTest.MyGameMode
 	 */
-	ENGINE_API FTestWorlds(const FString& MapName, const FString& GameModeName, const EReplicationSystem ReplicationSystem = EReplicationSystem::Generic);
+	ENGINE_API FTestWorlds(const FString& MapName, const FString& GameModeName, const EReplicationSystem ReplicationSystem = EReplicationSystem::Default);
 
 	/** Creates a server world using the given URL. */
-	ENGINE_API explicit FTestWorlds(const TCHAR* ServerURL, const EReplicationSystem ReplicationSystem = EReplicationSystem::Generic);
+	ENGINE_API explicit FTestWorlds(const TCHAR* ServerURL, const EReplicationSystem ReplicationSystem = EReplicationSystem::Default);
 
 	ENGINE_API ~FTestWorlds();
 
