[Backout] - CL39609082

This reverts a backout for a hotfix so we can enable a new fix in its place

#rnx
Original CL Desc
-----------------------------------------------------------------
Backout of 38341618.  The original backout introduced an upgrade bug where the skeletal mesh asset property, on blueprint-added skel mesh components, that were overriden on actors within the level viewport would be reset to their default value when upgrading to 5.5.2.  This change reintroduces a bug where the skeletal mesh asset can't be set on a blueprint added mesh component via construction script.

diff --git a/Engine/Source/Runtime/Engine/Private/Components/SkeletalMeshComponentInstanceData.cpp b/Engine/Source/Runtime/Engine/Private/Components/SkeletalMeshComponentInstanceData.cpp
--- a/Engine/Source/Runtime/Engine/Private/Components/SkeletalMeshComponentInstanceData.cpp
+++ b/Engine/Source/Runtime/Engine/Private/Components/SkeletalMeshComponentInstanceData.cpp
@@ -7,18 +7,17 @@
 
 FSkeletalMeshComponentInstanceData::FSkeletalMeshComponentInstanceData(const USkeletalMeshComponent* SourceComponent)
 	: FSceneComponentInstanceData(SourceComponent)
-	, SkeletalMeshAsset(nullptr)
 	, bUpdateAnimationInEditor(0)
 	, bUpdateClothInEditor(0)
+
 {
 #if WITH_EDITOR
-	// Only Blueprint components would reset transient values when the construction script is re-run.
+	// Only Blueprint components would reset tansient values when the construction script is re-run.
 	// Hence, we only need to apply instance cache for Blueprint created components.
 	const bool bIsBlueprintCreatedComponent = SourceComponent->CreationMethod == EComponentCreationMethod::SimpleConstructionScript
 		|| SourceComponent->CreationMethod == EComponentCreationMethod::UserConstructionScript;
 	if (bIsBlueprintCreatedComponent)
 	{
-		SkeletalMeshAsset = SourceComponent->GetSkeletalMeshAsset();
 		bUpdateAnimationInEditor = SourceComponent->GetUpdateAnimationInEditor();
 		bUpdateClothInEditor = SourceComponent->GetUpdateClothInEditor();
 	}
@@ -27,7 +26,7 @@ FSkeletalMeshComponentInstanceData::FSkeletalMeshComponentInstanceData(const USk
 
 bool FSkeletalMeshComponentInstanceData::ContainsData() const
 {
-	return bUpdateAnimationInEditor || bUpdateClothInEditor || SkeletalMeshAsset;
+	return bUpdateAnimationInEditor || bUpdateClothInEditor;
 }
 
 void FSkeletalMeshComponentInstanceData::ApplyToComponent(UActorComponent* Component, const ECacheApplyPhase CacheApplyPhase)
@@ -35,16 +34,8 @@ void FSkeletalMeshComponentInstanceData::ApplyToComponent(UActorComponent* Compo
 	Super::ApplyToComponent(Component, CacheApplyPhase);
 
 #if WITH_EDITOR
-	const bool bIsBlueprintCreatedComponent = Component->CreationMethod == EComponentCreationMethod::SimpleConstructionScript
-		|| Component->CreationMethod == EComponentCreationMethod::UserConstructionScript;
-	USkeletalMeshComponent* SkeletalMesh = Cast<USkeletalMeshComponent>(Component);
-
-	if (SkeletalMesh && bIsBlueprintCreatedComponent)
+	if (USkeletalMeshComponent * SkeletalMesh = Cast<USkeletalMeshComponent>(Component))
 	{
-		if (SkeletalMeshAsset)
-		{
-			SkeletalMesh->SetSkeletalMeshAsset(SkeletalMeshAsset);
-		}
 		SkeletalMesh->SetUpdateAnimationInEditor(bUpdateAnimationInEditor);
 		SkeletalMesh->SetUpdateClothInEditor(bUpdateClothInEditor);
 	}
diff --git a/Engine/Source/Runtime/Engine/Private/Components/SkeletalMeshComponentInstanceData.h b/Engine/Source/Runtime/Engine/Private/Components/SkeletalMeshComponentInstanceData.h
--- a/Engine/Source/Runtime/Engine/Private/Components/SkeletalMeshComponentInstanceData.h
+++ b/Engine/Source/Runtime/Engine/Private/Components/SkeletalMeshComponentInstanceData.h
@@ -6,7 +6,6 @@
 #include "SkeletalMeshComponentInstanceData.generated.h"
 
 class USkeletalMeshComponent;
-class USkeletalMesh;
 
 /** Saves internal SkeletalMesh (transient) state that gets lost at reconstruction for Blueprint created components. */
 USTRUCT()
@@ -20,7 +19,6 @@ struct FSkeletalMeshComponentInstanceData : public FSceneComponentInstanceData
 	virtual bool ContainsData() const override;
 	virtual void ApplyToComponent(UActorComponent* Component, const ECacheApplyPhase CacheApplyPhase) override;
 
-	USkeletalMesh* SkeletalMeshAsset;
 	uint8 bUpdateAnimationInEditor : 1;
 	uint8 bUpdateClothInEditor : 1;
 };
