Only use no-overwrite lock mode if it's supported by the RHI. Also added validation to check that we support it.

diff --git a/Engine/Source/Runtime/Engine/Private/GPUSkinVertexFactory.cpp b/Engine/Source/Runtime/Engine/Private/GPUSkinVertexFactory.cpp
--- a/Engine/Source/Runtime/Engine/Private/GPUSkinVertexFactory.cpp
+++ b/Engine/Source/Runtime/Engine/Private/GPUSkinVertexFactory.cpp
@@ -317,7 +317,7 @@ void FGPUBaseSkinVertexFactory::FShaderDataType::UpdateBoneData(
 	const uint32 VectorArraySize = NumBones * sizeof(FVector4f) * 3;
 
 	// RLM_WriteOnly_NoOverwrite is safe here because the buffer pool waits N frames before allowing new allocations.
-	FMatrix3x4* ChunkMatrices = (FMatrix3x4*)RHICmdList.LockBuffer(VertexBufferRHI, 0, VectorArraySize, RLM_WriteOnly_NoOverwrite);
+	FMatrix3x4* ChunkMatrices = (FMatrix3x4*)RHICmdList.LockBuffer(VertexBufferRHI, 0, VectorArraySize, GRHIGlobals.SupportsMapWriteNoOverwrite ? RLM_WriteOnly_NoOverwrite : RLM_WriteOnly);
 
 	{
 		QUICK_SCOPE_CYCLE_COUNTER(STAT_FGPUBaseSkinVertexFactory_ShaderDataType_UpdateBoneData_CopyBones);
