[shaders] bugfixes:
- reset job status object in retry path
- initialize debug info path in shader compiling manager constructor (called in engine loop long before anything else uses it) to avoid race conditions initializing the local static

diff --git a/Engine/Source/Runtime/Engine/Private/ShaderCompiler/ShaderCompiler.cpp b/Engine/Source/Runtime/Engine/Private/ShaderCompiler/ShaderCompiler.cpp
--- a/Engine/Source/Runtime/Engine/Private/ShaderCompiler/ShaderCompiler.cpp
+++ b/Engine/Source/Runtime/Engine/Private/ShaderCompiler/ShaderCompiler.cpp
@@ -826,6 +826,8 @@ FShaderCompilingManager::FShaderCompilingManager() :
 	FPaths::NormalizeDirectoryName(AbsoluteBaseDirectory);
 	AbsoluteShaderBaseWorkingDirectory = AbsoluteBaseDirectory + TEXT("/");
 
+	// Initialize the shader debug info path; this internally uses a local static var so we create it as early as possible in the init loop to avoid thread safety issues 
+	GetShaderDebugInfoPath();
 
 	CalculateNumberOfCompilingThreads(FPlatformMisc::NumberOfCores(), FPlatformMisc::NumberOfCoresIncludingHyperthreads());
 
@@ -2512,6 +2514,7 @@ bool FShaderCompilingManager::HandlePotentialRetryOnError(TMap<int32, FShaderMap
 						}
 						SingleJob->Output = FShaderCompilerOutput();
 						SingleJob->PreprocessOutput = FShaderPreprocessOutput();
+						SingleJob->JobStatusPtr->Reset();
 					}
 					else
 					{
@@ -2520,6 +2523,7 @@ bool FShaderCompilingManager::HandlePotentialRetryOnError(TMap<int32, FShaderMap
 						{
 							CommonJob->Output = FShaderCompilerOutput();
 							CommonJob->PreprocessOutput = FShaderPreprocessOutput();
+							CommonJob->JobStatusPtr->Reset();
 							CommonJob->bFinalized = false;
 						}
 					}
diff --git a/Engine/Source/Runtime/RenderCore/Public/ShaderCompilerJobTypes.h b/Engine/Source/Runtime/RenderCore/Public/ShaderCompilerJobTypes.h
--- a/Engine/Source/Runtime/RenderCore/Public/ShaderCompilerJobTypes.h
+++ b/Engine/Source/Runtime/RenderCore/Public/ShaderCompilerJobTypes.h
@@ -93,6 +93,13 @@ RENDERCORE_API const TCHAR* LexToString(EShaderCompileJobStatus Status);
 class FShaderCompileJobStatus
 {
 public:
+	void Reset()
+	{
+		InputHash = nullptr;
+		IsDuplicate.store(false);
+		Status.store(EShaderCompileJobStatus::Unset);
+	}
+
 	void SetStatus(EShaderCompileJobStatus NewStatus)
 	{
 		Status.store(NewStatus);
