Change UWorld::IsDataValid to only validate the provided AssociatedExternalObjects from the validation context if provided.

diff --git a/Engine/Source/Runtime/Engine/Private/World.cpp b/Engine/Source/Runtime/Engine/Private/World.cpp
--- a/Engine/Source/Runtime/Engine/Private/World.cpp
+++ b/Engine/Source/Runtime/Engine/Private/World.cpp
@@ -27,6 +27,7 @@
 #include "UObject/ObjectSaveContext.h"
 #include "UObject/UObjectAnnotation.h"
 #include "Misc/PackageName.h"
+#include "Misc/DataValidation.h"
 #include "GameMapsSettings.h"
 #include "TimerManager.h"
 #include "AI/NavigationSystemBase.h"
@@ -887,13 +888,32 @@ void UWorld::AddReferencedObjects(UObject* InThis, FReferenceCollector& Collecto
 EDataValidationResult UWorld::IsDataValid(FDataValidationContext& Context) const 
 {
 	EDataValidationResult Result = EDataValidationResult::NotValidated;
-	// Pass flags explicitly so we iterate non-active levels for data validation but continue to skip destroyed/deleted actors.
-	for (FActorIterator It(this, EActorIteratorFlags::SkipPendingKill); It; ++It)
+	// If the validation system wishes to validate specific actors, it will pass them in the context as associated external objects
+	// Otherwise we should assume that we came from a function such as UEditorValidatorSubsystem::IsObjectValid and we should validate whatever is loaded in the user's session
+	if (Context.GetAssociatedExternalObjects().Num() != 0)
 	{
-		AActor* Actor = *It;
-		if (Actor)
+		 for (const FAssetData& Asset : Context.GetAssociatedExternalObjects())
+		 {
+			if (AActor* Actor = Cast<AActor>(Asset.FastGetAsset(false)))
+			{
+				ensureMsgf(Actor->GetWorld() == this, TEXT("Attempting to validate actor %s whose owning world is %s and not this (%s)"), 
+					*GetPathNameSafe(Actor),
+					*GetPathNameSafe(Actor->GetWorld()),
+					*GetPathName());
+				Result = CombineDataValidationResults(Result, Actor->IsDataValid(Context));
+			}
+		 }
+	}
+	else
+	{
+		// Pass flags explicitly so we iterate non-active levels for data validation but continue to skip destroyed/deleted actors.
+		for (FActorIterator It(this, EActorIteratorFlags::SkipPendingKill); It; ++It)
 		{
-			Result = CombineDataValidationResults(Result, Actor->IsDataValid(Context));
+			AActor* Actor = *It;
+			if (Actor)
+			{
+				Result = CombineDataValidationResults(Result, Actor->IsDataValid(Context));
+			}
 		}
 	}
 	return CombineDataValidationResults(Result, Super::IsDataValid(Context));
