Changed primitive component to serialize some non-UProperty members (for remote object migration)

FOverlapInfo is now a USTRUCT so that reflection is available for it.

Manual serialization added for some non-UProperty members of UPrimitiveComponent: OverlappingComponents and MoveIgnoreMask

#rnx

diff --git a/Engine/Source/Runtime/Engine/Private/Components/PrimitiveComponent.cpp b/Engine/Source/Runtime/Engine/Private/Components/PrimitiveComponent.cpp
--- a/Engine/Source/Runtime/Engine/Private/Components/PrimitiveComponent.cpp
+++ b/Engine/Source/Runtime/Engine/Private/Components/PrimitiveComponent.cpp
@@ -1215,6 +1215,28 @@ void UPrimitiveComponent::Serialize(FArchive& Ar)
 		}
 	}
 
+#if UE_WITH_REMOTE_OBJECT_HANDLE
+	// For remote objects we need non-UProperties to be transferred
+	if (Ar.IsMigratingRemoteObjects())
+	{
+		{
+	    	int32 NumOverlappingComponents = OverlappingComponents.Num();
+    		Ar << NumOverlappingComponents;
+    		if (Ar.IsLoading())
+		    {
+        		OverlappingComponents.SetNum(NumOverlappingComponents);
+    		}
+
+    		for (FOverlapInfo& OverlapInfo : OverlappingComponents)
+    		{
+				FOverlapInfo::StaticStruct()->SerializeTaggedProperties(Ar, (uint8*)&OverlapInfo, FOverlapInfo::StaticStruct(), nullptr);
+    		}
+		}
+		
+		Ar << MoveIgnoreMask;
+	}
+#endif
+
 #if WITH_EDITORONLY_DATA
 	if (Ar.CustomVer(FUE5PrivateFrostyStreamObjectVersion::GUID) < FUE5PrivateFrostyStreamObjectVersion::HLODBatchingPolicy)
 	{
diff --git a/Engine/Source/Runtime/Engine/Public/Engine/OverlapInfo.h b/Engine/Source/Runtime/Engine/Public/Engine/OverlapInfo.h
--- a/Engine/Source/Runtime/Engine/Public/Engine/OverlapInfo.h
+++ b/Engine/Source/Runtime/Engine/Public/Engine/OverlapInfo.h
@@ -4,12 +4,16 @@
 
 #include "CoreMinimal.h"
 #include "Engine/HitResult.h"
+#include "OverlapInfo.generated.h"
 
 class UPrimitiveComponent;
 
 /** Overlap info consisting of the primitive and the body that is overlapping */
+USTRUCT()
 struct FOverlapInfo
 {
+	GENERATED_BODY()
+
 	FOverlapInfo()
 	{}
 
@@ -24,11 +28,14 @@ struct FOverlapInfo
 
 	//This function completely ignores SweepResult information. It seems that places that use this function do not care, but it still seems risky
 	friend bool operator == (const FOverlapInfo& LHS, const FOverlapInfo& RHS) { return LHS.OverlapInfo.Component == RHS.OverlapInfo.Component && LHS.OverlapInfo.Item == RHS.OverlapInfo.Item; }
-	bool bFromSweep;
+
+	UPROPERTY(Transient)
+	bool bFromSweep = false;
 
 	/** Information for both sweep and overlap queries. Different parts are valid depending on bFromSweep.
 	  * If bFromSweep is true then FHitResult is completely valid just like a regular sweep result.
 	  * If bFromSweep is false only FHitResult::Component, FHitResult::Actor, FHitResult::Item are valid as this is really just an FOverlapResult*/
+	UPROPERTY(Transient)
 	FHitResult OverlapInfo;
 };
 
