Make -emulatestereo have asymmetric FOV by default.

- Rationale: off-center position mimics real headsets better and helps uncover more bugs.

diff --git a/Engine/Source/Runtime/Engine/Private/UnrealEngine.cpp b/Engine/Source/Runtime/Engine/Private/UnrealEngine.cpp
--- a/Engine/Source/Runtime/Engine/Private/UnrealEngine.cpp
+++ b/Engine/Source/Runtime/Engine/Private/UnrealEngine.cpp
@@ -3871,6 +3871,8 @@ class FFakeStereoRenderingDevice : public IStereoRendering, public IStereoRender
 public:
 	FFakeStereoRenderingDevice(int ViewportWidth = 640, int ViewportHeight = 480, int RequestedNumViews = 2) 
 		: FOVInDegrees(GetDefault<UGeneralProjectSettings>()->FOVForFakeStereoRenderingDevice)
+		, TopFOVRatio(GetDefault<UGeneralProjectSettings>()->TopFOVRatioForFakeStereoRenderingDevice)
+		, DifferenceBetweenEyes(GetDefault<UGeneralProjectSettings>()->DifferenceBetweenEyesForFakeStereoRenderingDevice)
 		, Width(ViewportWidth)
 		, Height(ViewportHeight)
 		, NumViews(RequestedNumViews)
@@ -3985,18 +3987,46 @@ public:
 
 	virtual FMatrix GetStereoProjectionMatrix(const int32 ViewIndex) const override
 	{
-		const float HalfFov = FMath::DegreesToRadians(FOVInDegrees) / 2.f;
-		const float InWidth = Width;
-		const float InHeight = Height;
-		const float XS = 1.0f / FMath::Tan(HalfFov);
-		const float YS = InWidth / FMath::Tan(HalfFov) / InHeight;
-		const float NearZ = GNearClippingPlane;
+		float AngleUp = TopFOVRatio * FOVInDegrees;
+		float AngleDown = -(1.0f - TopFOVRatio) * FOVInDegrees;
+		// setup for monoscopic by default
+		float AngleRight = 0.5f * FOVInDegrees;
+		float AngleLeft = -0.5f * FOVInDegrees;
 
-		return FMatrix(
-			FPlane(XS, 0.0f, 0.0f, 0.0f),
-			FPlane(0.0f, YS, 0.0f, 0.0f),
-			FPlane(0.0f, 0.0f, 0.0f, 1.0f),
-			FPlane(0.0f, 0.0f, NearZ, 0.0f));
+		// When difference is 0, the left and right angles are equal for each eye. When 1, each eye is only looking one way (i.e. left eye has AngleRight == 0 and vice versa)
+		float PrimaryEyeCoeff = (1.0f + DifferenceBetweenEyes) / 2.0f;
+		float SecondaryEyeCoeff = (1.0f - PrimaryEyeCoeff);
+		if (ViewIndex == eSSE_LEFT_EYE)
+		{
+			AngleLeft = -(PrimaryEyeCoeff) * FOVInDegrees;
+			AngleRight = SecondaryEyeCoeff * FOVInDegrees;
+		}
+		else if (ViewIndex == eSSE_RIGHT_EYE)
+		{
+			AngleLeft = -(SecondaryEyeCoeff) * FOVInDegrees;
+			AngleRight = PrimaryEyeCoeff * FOVInDegrees;
+		}
+
+		AngleUp = tan(FMath::DegreesToRadians(AngleUp));
+		AngleDown = tan(FMath::DegreesToRadians(AngleDown));
+		AngleLeft = tan(FMath::DegreesToRadians(AngleLeft));
+		AngleRight = tan(FMath::DegreesToRadians(AngleRight));
+
+		float ZNear = GNearClippingPlane_RenderThread;
+
+		float SumRL = (AngleRight + AngleLeft);
+		float SumTB = (AngleUp + AngleDown);
+		float InvRL = (1.0f / (AngleRight - AngleLeft));
+		float InvTB = (1.0f / (AngleUp - AngleDown));
+
+		FMatrix Mat = FMatrix(
+			FPlane((2.0f * InvRL), 0.0f, 0.0f, 0.0f),
+			FPlane(0.0f, (2.0f * InvTB), 0.0f, 0.0f),
+			FPlane((SumRL * -InvRL), (SumTB * -InvTB), 0.0f, 1.0f),
+			FPlane(0.0f, 0.0f, ZNear, 0.0f)
+		);
+
+		return Mat;
 	}
 
 	virtual void InitCanvasFromView(FSceneView* InView, UCanvas* Canvas) override
@@ -4128,6 +4158,11 @@ public:
 
 	/** max(HFOV, VFOV) in degrees of imaginable HMD */
 	float FOVInDegrees;
+	/** Proportion between Top FOV and Bottom FOV */
+	float TopFOVRatio;
+	/** How dissimilar are the eye frustums. If set to 0, they will be totally symmetric, if set to 1, totally different. */
+	float DifferenceBetweenEyes;
+
 	/** resolution of imaginable HMD (per eye) */
 	int32 Width, Height;
 	/** views of imaginable HMD */
diff --git a/Engine/Source/Runtime/EngineSettings/Classes/GeneralProjectSettings.h b/Engine/Source/Runtime/EngineSettings/Classes/GeneralProjectSettings.h
--- a/Engine/Source/Runtime/EngineSettings/Classes/GeneralProjectSettings.h
+++ b/Engine/Source/Runtime/EngineSettings/Classes/GeneralProjectSettings.h
@@ -102,4 +102,12 @@ class UGeneralProjectSettings
 	/*Determines the Field Of View of the virtual stereo device created when " -emulatestereo" command line arg is detected*/
 	UPROPERTY(config, EditAnywhere, Category = Settings)
 	float FOVForFakeStereoRenderingDevice;
+
+	/* Determines the how much of -emulatestereo FOV is above the horizontal plane (0-1) - usually helmets have a larger FOV down than up, so this value is less than 0.5 */
+	UPROPERTY(config, EditAnywhere, Category = Settings)
+	float TopFOVRatioForFakeStereoRenderingDevice;
+
+	/* Determines difference between left and eyes -emulatestereo FOVs. If 0, they are the same (their FOVs are symmetric), the closer to 1, the more each eye is looking towards its side. Usable values are in 0 - 0.4 range. */
+	UPROPERTY(config, EditAnywhere, Category = Settings)
+	float DifferenceBetweenEyesForFakeStereoRenderingDevice;
 };
diff --git a/Engine/Source/Runtime/EngineSettings/Private/EngineSettingsModule.cpp b/Engine/Source/Runtime/EngineSettings/Private/EngineSettingsModule.cpp
--- a/Engine/Source/Runtime/EngineSettings/Private/EngineSettingsModule.cpp
+++ b/Engine/Source/Runtime/EngineSettings/Private/EngineSettingsModule.cpp
@@ -75,6 +75,8 @@ UGeneralProjectSettings::UGeneralProjectSettings( const FObjectInitializer& Obje
 	, bAllowMinimize(true)
 	, EyeOffsetForFakeStereoRenderingDevice(0.0320000005f)
 	, FOVForFakeStereoRenderingDevice(90)
+	, TopFOVRatioForFakeStereoRenderingDevice(0.47f)
+	, DifferenceBetweenEyesForFakeStereoRenderingDevice(0.2f)
 { }
 
 
