Adjusting some pawn nav data code to use the NavMovementInterface

diff --git a/Engine/Source/Runtime/Engine/Private/Pawn.cpp b/Engine/Source/Runtime/Engine/Private/Pawn.cpp
--- a/Engine/Source/Runtime/Engine/Private/Pawn.cpp
+++ b/Engine/Source/Runtime/Engine/Private/Pawn.cpp
@@ -191,12 +191,17 @@ UPawnMovementComponent* APawn::GetMovementComponent() const
 
 void APawn::UpdateNavAgent()
 {
-	UPawnMovementComponent* MovementComponent = GetMovementComponent();
+	INavMovementInterface* NavMovementInterface = FindComponentByInterface<INavMovementInterface>();
+	
 	//// Update Nav Agent props with collision component's setup if it's not set yet
-	if (RootComponent != nullptr && MovementComponent != nullptr && MovementComponent->ShouldUpdateNavAgentWithOwnersCollision())
+	if (RootComponent != nullptr && NavMovementInterface != nullptr)
 	{
-		RootComponent->UpdateBounds();
-		MovementComponent->UpdateNavAgent(*this);
+		const FNavMovementProperties* NavMovementProperties = NavMovementInterface->GetNavMovementProperties();
+		if (NavMovementProperties && NavMovementProperties->bUpdateNavAgentWithOwnersCollision != 0)
+		{
+			RootComponent->UpdateBounds();
+			NavMovementInterface->UpdateNavAgent(*this);
+		}
 	}
 }
 
@@ -1332,7 +1337,7 @@ void APawn::PawnMakeNoise(float Loudness, FVector NoiseLocation, bool bUseNoiseM
 
 const FNavAgentProperties& APawn::GetNavAgentPropertiesRef() const
 {
-	UPawnMovementComponent* MovementComponent = GetMovementComponent();
-	return MovementComponent ? MovementComponent->GetNavAgentPropertiesRef() : FNavAgentProperties::DefaultProperties;
+	INavMovementInterface* NavMovementInterface = FindComponentByInterface<INavMovementInterface>();
+	return NavMovementInterface ? NavMovementInterface->GetNavAgentPropertiesRef() : FNavAgentProperties::DefaultProperties;
 }
 
