Enable HLSL2021 as well for inline ray tracing (was already forced enabled for ray gen and hit shaders)

diff --git a/Engine/Source/Runtime/Engine/Private/ShaderCompiler/ShaderCompiler.cpp b/Engine/Source/Runtime/Engine/Private/ShaderCompiler/ShaderCompiler.cpp
--- a/Engine/Source/Runtime/Engine/Private/ShaderCompiler/ShaderCompiler.cpp
+++ b/Engine/Source/Runtime/Engine/Private/ShaderCompiler/ShaderCompiler.cpp
@@ -3291,10 +3291,12 @@ void GlobalBeginCompileShader(
 
 	// Enables HLSL 2021
 	uint32 EnablesHLSL2021ByDefault = FDataDrivenShaderPlatformInfo::GetEnablesHLSL2021ByDefault(EShaderPlatform(Target.Platform));
+	bool bInlineRayTracing = Input.Environment.CompilerFlags.Contains(CFLAG_InlineRayTracing);
 	if((EnablesHLSL2021ByDefault == uint32(1) && DebugGroupName == TEXT("Global")) ||
 		EnablesHLSL2021ByDefault == uint32(2) ||
 		Target.Frequency == SF_RayGen ||			// We want to make sure that function overloads follow c++ rules for FRayDesc.
-		Target.Frequency == SF_RayHitGroup)
+		Target.Frequency == SF_RayHitGroup ||
+		bInlineRayTracing)
 	{
 		Input.Environment.CompilerFlags.Add(CFLAG_HLSL2021);
 	}
