Removed hack from FMaterialResource::IsTangentSpaceNormal() that would return true depending on whether material attributes or something was connected to the Normal attribute. Fix is hidden behind the r.Material.TentativeNormalFix cvar for now as it requires scrupulous testing.

Instead, it now returns the true bTangentSpaceNormal as set in the material. Then, fixed translation to emit vertex normal if material is using world space normals and the normal attribute is unset (and thus [0, 0, 1] would have emitted insted).

#rnx

diff --git a/Engine/Source/Runtime/Engine/Private/Materials/Material.cpp b/Engine/Source/Runtime/Engine/Private/Materials/Material.cpp
--- a/Engine/Source/Runtime/Engine/Private/Materials/Material.cpp
+++ b/Engine/Source/Runtime/Engine/Private/Materials/Material.cpp
@@ -127,6 +127,13 @@ static TAutoConsoleVariable<bool> CVarMaterialsDuplicateVerbatim(
 	TEXT("Default: false"),
 	ECVF_Default);
 
+TAutoConsoleVariable<bool> CVarMaterialsTentativeNormalFix(
+	TEXT("r.Material.TentativeNormalFix"),
+	false,
+	TEXT("It enables a fix that removes a hack regarding tangent normal space handling when the Normal material attribute is connected.\n")
+	TEXT("Default: false"),
+	ECVF_Default);
+
 namespace MaterialImpl
 {
 	// Filtered list of properties that we follow on mobile platforms.
@@ -6937,7 +6944,6 @@ int32 UMaterial::CompilePropertyEx( FMaterialCompiler* Compiler, const FGuid& At
 		case MP_EmissiveColor:			return EditorOnly->EmissiveColor.CompileWithDefault(Compiler, Property);
 		case MP_BaseColor:				return EditorOnly->BaseColor.CompileWithDefault(Compiler, Property);
 		case MP_SubsurfaceColor:		return EditorOnly->SubsurfaceColor.CompileWithDefault(Compiler, Property);
-		case MP_Normal:					return EditorOnly->Normal.CompileWithDefault(Compiler, Property);
 		case MP_Tangent:				return EditorOnly->Tangent.CompileWithDefault(Compiler, Property);
 		case MP_WorldPositionOffset:	return EditorOnly->WorldPositionOffset.CompileWithDefault(Compiler, Property);
 		case MP_Displacement:			return EditorOnly->Displacement.CompileWithDefault(Compiler, Property);
@@ -6946,6 +6952,15 @@ int32 UMaterial::CompilePropertyEx( FMaterialCompiler* Compiler, const FGuid& At
 		case MP_SurfaceThickness:		return EditorOnly->SurfaceThickness.CompileWithDefault(Compiler, Property);
 		case MP_FrontMaterial:			return EditorOnly->FrontMaterial.CompileWithDefault(Compiler, Property);
 
+		case MP_Normal:
+			// If material is not using tangent space normals, and nothing is connected to the normal pin,
+			// instead of setting the normal to the constant use the vertex normal.
+			if (CVarMaterialsTentativeNormalFix.GetValueOnAnyThread() && !this->bTangentSpaceNormal && EditorOnly->Normal.UseConstant)
+			{
+				return Compiler->VertexNormal();
+			}
+			return EditorOnly->Normal.CompileWithDefault(Compiler, Property);
+
 		default:
 			if (Property >= MP_CustomizedUVs0 && Property <= MP_CustomizedUVs7)
 			{
diff --git a/Engine/Source/Runtime/Engine/Private/Materials/MaterialShared.cpp b/Engine/Source/Runtime/Engine/Private/Materials/MaterialShared.cpp
--- a/Engine/Source/Runtime/Engine/Private/Materials/MaterialShared.cpp
+++ b/Engine/Source/Runtime/Engine/Private/Materials/MaterialShared.cpp
@@ -111,6 +111,8 @@ static TAutoConsoleVariable<bool> CVarDisableVRSForCustomPrimitiveData(
 	TEXT("Controls whether or not to disable VRS for materials using Custom Primitive Data."),
 	ECVF_RenderThreadSafe | ECVF_ReadOnly);
 
+extern TAutoConsoleVariable<bool> CVarMaterialsTentativeNormalFix;
+
 bool IsUsingNewMaterialTranslatorPrototype()
 {
 	return CVarUsingUseNewMaterialTranslatorPrototype.GetValueOnAnyThread();
@@ -1584,7 +1586,7 @@ void FMaterial::DiscardShaderMap()
 EMaterialDomain FMaterialResource::GetMaterialDomain() const { return Material->MaterialDomain; }
 bool FMaterialResource::IsTangentSpaceNormal() const 
 { 
-	if (IsSubstrateMaterial())
+	if (IsSubstrateMaterial() || CVarMaterialsTentativeNormalFix.GetValueOnAnyThread())
 	{
 		return Material->bTangentSpaceNormal; // We do not need to check MP_Normal with Substrate as this cannot be specified on the root node anymore.
 	}
