Re-instate fix from CL 37655049

Robomerged backout from 5.5 branch removed this fix, this puts it back.

#rnx

diff --git a/Engine/Source/Runtime/Engine/Private/Components/SkeletalMeshComponentInstanceData.cpp b/Engine/Source/Runtime/Engine/Private/Components/SkeletalMeshComponentInstanceData.cpp
--- a/Engine/Source/Runtime/Engine/Private/Components/SkeletalMeshComponentInstanceData.cpp
+++ b/Engine/Source/Runtime/Engine/Private/Components/SkeletalMeshComponentInstanceData.cpp
@@ -36,6 +36,7 @@ void FSkeletalMeshComponentInstanceData::ApplyToComponent(UActorComponent* Compo
 #if WITH_EDITOR
 	if (USkeletalMeshComponent * SkeletalMesh = Cast<USkeletalMeshComponent>(Component))
 	{
+		SkeletalMesh->RefreshSkeletalMeshAsset();
 		SkeletalMesh->SetUpdateAnimationInEditor(bUpdateAnimationInEditor);
 		SkeletalMesh->SetUpdateClothInEditor(bUpdateClothInEditor);
 	}
