[DemoRecording] Use WeakObjectPtr instead of raw object pointer to determine if the actor is valid or not. Use WeakObjectPtr in calls that expect a WeakObjectPtr instead of converting a raw, potentially invalid, pointer to a WeakObjectPtr. Ensure when trying to serialize an invalid actor rather than cause a crash.

diff --git a/Engine/Source/Runtime/Engine/Private/DemoNetDriver.cpp b/Engine/Source/Runtime/Engine/Private/DemoNetDriver.cpp
--- a/Engine/Source/Runtime/Engine/Private/DemoNetDriver.cpp
+++ b/Engine/Source/Runtime/Engine/Private/DemoNetDriver.cpp
@@ -1989,7 +1989,7 @@ void UDemoNetDriver::TickDemoRecordFrame(float DeltaSeconds)
 				{
 					AActor* Actor = ActorInfo->Actor;
 
-					if (!IsValid(Actor))
+					if (!ActorInfo->WeakActor.IsValid())
 					{
 						ActorsToRemove.Add(Actor);
 						continue;
@@ -2049,7 +2049,7 @@ void UDemoNetDriver::TickDemoRecordFrame(float DeltaSeconds)
 					UActorChannel* Channel = nullptr;
 					if (bDoFindActorChannelEarly)
 					{
-						Channel = ClientConnection->FindActorChannelRef(Actor);
+						Channel = ClientConnection->FindActorChannelRef(ActorInfo->WeakActor);
 
 						// Check dormancy
 						if (bDoCheckDormancyEarly && Channel && ShouldActorGoDormantForDemo(Actor, Channel))
@@ -2265,65 +2265,69 @@ bool UDemoNetDriver::ReplicatePrioritizedActor(const FActorPriority& ActorPriori
 	}
 	else if (bActorInfo)
 	{
-		++Params.NumActorsReplicated;
-
-		AActor* Actor = ActorInfo->Actor;
-		
-		if (bDoCheckDormancy)
+		constexpr bool bGetObjectEvenIfPendingKill = true;
+		if (ensureMsgf(ActorInfo->WeakActor.IsValid(), TEXT("Trying to record an invalid actor `%s`. Skipping."), ToCStr(GetNameSafe(ActorInfo->WeakActor.Get(bGetObjectEvenIfPendingKill)))))
 		{
-			UActorChannel* Channel = (bDoFindActorChannel ? Params.Connection->FindActorChannelRef(Actor) : ActorPriority.Channel);
-			if (Channel && ShouldActorGoDormantForDemo(Actor, Channel))
+			++Params.NumActorsReplicated;
+
+			AActor* Actor = ActorInfo->Actor;
+
+			if (bDoCheckDormancy)
 			{
-				// Either shouldn't go dormant, or is already dormant
-				Channel->StartBecomingDormant();
+				UActorChannel* Channel = (bDoFindActorChannel ? Params.Connection->FindActorChannelRef(ActorInfo->WeakActor) : ActorPriority.Channel);
+				if (Channel && ShouldActorGoDormantForDemo(Actor, Channel))
+				{
+					// Either shouldn't go dormant, or is already dormant
+					Channel->StartBecomingDormant();
+				}
 			}
-		}
 
-		// Use NetUpdateFrequency for this actor, but clamp it to RECORD_HZ.
-		const float ClampedNetUpdateFrequency = FMath::Clamp(Actor->GetNetUpdateFrequency(), Params.MinRecordHz, Params.MaxRecordHz);
-		const double NetUpdateDelay = 1.0 / ClampedNetUpdateFrequency;
+			// Use NetUpdateFrequency for this actor, but clamp it to RECORD_HZ.
+			const float ClampedNetUpdateFrequency = FMath::Clamp(Actor->GetNetUpdateFrequency(), Params.MinRecordHz, Params.MaxRecordHz);
+			const double NetUpdateDelay = 1.0 / ClampedNetUpdateFrequency;
 
-		// Set defaults if this actor is replicating for first time
-		if (ActorInfo->LastNetReplicateTime == 0)
-		{
-			ActorInfo->LastNetReplicateTime = GetDemoCurrentTime();
-			ActorInfo->OptimalNetUpdateDelta = NetUpdateDelay;
-		}
+			// Set defaults if this actor is replicating for first time
+			if (ActorInfo->LastNetReplicateTime == 0)
+			{
+				ActorInfo->LastNetReplicateTime = GetDemoCurrentTime();
+				ActorInfo->OptimalNetUpdateDelta = NetUpdateDelay;
+			}
 
-		const float LastReplicateDelta = static_cast<float>(GetDemoCurrentTime() - ActorInfo->LastNetReplicateTime);
+			const float LastReplicateDelta = static_cast<float>(GetDemoCurrentTime() - ActorInfo->LastNetReplicateTime);
 
-		// Calculate min delta (max rate actor will update), and max delta (slowest rate actor will update)
-		const float MinOptimalDelta = NetUpdateDelay;										// Don't go faster than NetUpdateFrequency
-		const float MinNetUpdateFrequency = (Actor->GetMinNetUpdateFrequency() == 0.0f) ? 2.0f : Actor->GetMinNetUpdateFrequency();
-		const float MaxOptimalDelta = FMath::Max(1.0f / MinNetUpdateFrequency, MinOptimalDelta);	// Don't go slower than MinNetUpdateFrequency (or NetUpdateFrequency if it's slower)
+			// Calculate min delta (max rate actor will update), and max delta (slowest rate actor will update)
+			const float MinOptimalDelta = NetUpdateDelay;										// Don't go faster than NetUpdateFrequency
+			const float MinNetUpdateFrequency = (Actor->GetMinNetUpdateFrequency() == 0.0f) ? 2.0f : Actor->GetMinNetUpdateFrequency();
+			const float MaxOptimalDelta = FMath::Max(1.0f / MinNetUpdateFrequency, MinOptimalDelta);	// Don't go slower than MinNetUpdateFrequency (or NetUpdateFrequency if it's slower)
 
-		const float ScaleDownStartTime = 2.0f;
-		const float ScaleDownTimeRange = 5.0f;
+			const float ScaleDownStartTime = 2.0f;
+			const float ScaleDownTimeRange = 5.0f;
 
-		if (LastReplicateDelta > ScaleDownStartTime)
-		{
-			// Interpolate between MinOptimalDelta/MaxOptimalDelta based on how long it's been since this actor actually sent anything
-			const float Alpha = FMath::Clamp((LastReplicateDelta - ScaleDownStartTime) / ScaleDownTimeRange, 0.0f, 1.0f);
-			ActorInfo->OptimalNetUpdateDelta = FMath::Lerp(MinOptimalDelta, MaxOptimalDelta, Alpha);
-		}
+			if (LastReplicateDelta > ScaleDownStartTime)
+			{
+				// Interpolate between MinOptimalDelta/MaxOptimalDelta based on how long it's been since this actor actually sent anything
+				const float Alpha = FMath::Clamp((LastReplicateDelta - ScaleDownStartTime) / ScaleDownTimeRange, 0.0f, 1.0f);
+				ActorInfo->OptimalNetUpdateDelta = FMath::Lerp(MinOptimalDelta, MaxOptimalDelta, Alpha);
+			}
 
-		const double NextUpdateDelta = Params.bUseAdapativeNetFrequency ? ActorInfo->OptimalNetUpdateDelta : NetUpdateDelay;
+			const double NextUpdateDelta = Params.bUseAdapativeNetFrequency ? ActorInfo->OptimalNetUpdateDelta : NetUpdateDelay;
 
-		// Account for being fractionally into the next frame
-		// But don't be more than a fraction of a frame behind either (we don't want to do catch-up frames when there is a long delay)
-		const double ExtraTime = GetDemoCurrentTime() - ActorInfo->NextUpdateTime;
-		const double ClampedExtraTime = FMath::Clamp(ExtraTime, 0.0, NetUpdateDelay);
+			// Account for being fractionally into the next frame
+			// But don't be more than a fraction of a frame behind either (we don't want to do catch-up frames when there is a long delay)
+			const double ExtraTime = GetDemoCurrentTime() - ActorInfo->NextUpdateTime;
+			const double ClampedExtraTime = FMath::Clamp(ExtraTime, 0.0, NetUpdateDelay);
 
-		// Try to spread the updates across multiple frames to smooth out spikes.
-		ActorInfo->NextUpdateTime = (GetDemoCurrentTime() + NextUpdateDelta - ClampedExtraTime + ((UpdateDelayRandomStream.FRand() - 0.5) * Params.ServerTickTime));
+			// Try to spread the updates across multiple frames to smooth out spikes.
+			ActorInfo->NextUpdateTime = (GetDemoCurrentTime() + NextUpdateDelta - ClampedExtraTime + ((UpdateDelayRandomStream.FRand() - 0.5) * Params.ServerTickTime));
 
-		const bool bDidReplicateActor = DemoReplicateActor(Actor, Params.Connection, false);
+			const bool bDidReplicateActor = DemoReplicateActor(Actor, Params.Connection, false);
 
-		if (bDidReplicateActor)
-		{
-			// Choose an optimal time, we choose 70% of the actual rate to allow frequency to go up if needed
-			ActorInfo->OptimalNetUpdateDelta = FMath::Clamp(LastReplicateDelta * 0.7f, MinOptimalDelta, MaxOptimalDelta);
-			ActorInfo->LastNetReplicateTime = GetDemoCurrentTime();
+			if (bDidReplicateActor)
+			{
+				// Choose an optimal time, we choose 70% of the actual rate to allow frequency to go up if needed
+				ActorInfo->OptimalNetUpdateDelta = FMath::Clamp(LastReplicateDelta * 0.7f, MinOptimalDelta, MaxOptimalDelta);
+				ActorInfo->LastNetReplicateTime = GetDemoCurrentTime();
+			}
 		}
 	}
 	else
diff --git a/Engine/Source/Runtime/Engine/Private/ReplayHelper.cpp b/Engine/Source/Runtime/Engine/Private/ReplayHelper.cpp
--- a/Engine/Source/Runtime/Engine/Private/ReplayHelper.cpp
+++ b/Engine/Source/Runtime/Engine/Private/ReplayHelper.cpp
@@ -629,23 +629,24 @@ void FReplayHelper::SaveCheckpoint(UNetConnection* Connection)
 				{
 					if (!bDeltaCheckpoint || NetworkObjectInfo->bDirtyForReplay)
 					{
+						const TWeakObjectPtr<AActor>& WeakActor = NetworkObjectInfo->WeakActor;
 						AActor* Actor = NetworkObjectInfo->Actor;
 
 						// check to see if it should replicate at all
-						bool bCheckpointActor = IsValid(Actor) && ((Actor->GetRemoteRole() != ROLE_None || Actor->GetTearOff()) && (Actor == Connection->PlayerController || Cast<APlayerController>(Actor) == nullptr));
+						bool bCheckpointActor = WeakActor.IsValid() && ((Actor->GetRemoteRole() != ROLE_None || Actor->GetTearOff()) && (Actor == Connection->PlayerController || Cast<APlayerController>(Actor) == nullptr));
 						
 						// now look for an open channel
-						bCheckpointActor = bCheckpointActor && ActorChannelMap.Contains(NetworkObjectInfo->Actor);
+						bCheckpointActor = bCheckpointActor && ActorChannelMap.Contains(WeakActor);
 
 						if (!bCheckpointActor)
 						{
 							// has it gone dormant?
-							bCheckpointActor = (NetworkObjectInfo->Actor->NetDormancy != DORM_Initial) &&  (NetworkObjectInfo->DormantConnections.Contains(Connection) || NetworkObjectInfo->RecentlyDormantConnections.Contains(Connection));
+							bCheckpointActor = (Actor->NetDormancy != DORM_Initial) &&  (NetworkObjectInfo->DormantConnections.Contains(Connection) || NetworkObjectInfo->RecentlyDormantConnections.Contains(Connection));
 						}
 						
 						if (bCheckpointActor)
 						{
-							ActorArray.Add({ NetworkObjectInfo->Actor, NetworkObjectInfo->Actor->GetOuter() });
+							ActorArray.Add({ WeakActor, Actor->GetOuter() });
 
 							NetworkObjectInfo->bDirtyForReplay = false;
 						}
