Check if the actor weak pointer is valid rather than checking if the actor pointer is null.

diff --git a/Engine/Source/Runtime/Engine/Private/DataReplication.cpp b/Engine/Source/Runtime/Engine/Private/DataReplication.cpp
--- a/Engine/Source/Runtime/Engine/Private/DataReplication.cpp
+++ b/Engine/Source/Runtime/Engine/Private/DataReplication.cpp
@@ -516,11 +516,11 @@ void FObjectReplicator::InitWithObject(UObject* InObject, UNetConnection* InConn
 
 	SetObject(InObject);
 
-	if (GetObject() == nullptr)
+	if (!WeakObjectPtr.IsValid())
 	{
 		// This may seem weird that we're checking for nullptr, but the SetObject above will wrap this object with TWeakObjectPtr
 		// If the object is pending kill, it will switch to nullptr, we're just making sure we handle this invalid edge case
-		UE_LOG(LogRep, Error, TEXT("InitWithObject: Object == nullptr"));
+		UE_LOG(LogRep, Error, TEXT("InitWithObject: Invalid Object pointer %p with name `%s`"), InObject, ToCStr(GetNameSafe(InObject)));
 		return;
 	}
 
