Fix static analysis issue.

diff --git a/Engine/Source/Runtime/Engine/Private/DataReplication.cpp b/Engine/Source/Runtime/Engine/Private/DataReplication.cpp
--- a/Engine/Source/Runtime/Engine/Private/DataReplication.cpp
+++ b/Engine/Source/Runtime/Engine/Private/DataReplication.cpp
@@ -536,7 +536,7 @@ void FObjectReplicator::InitWithObject(UObject* InObject, UNetConnection* InConn
 	RepLayout = Connection->Driver->GetObjectClassRepLayout(ObjectClass);
 
 	// Make a copy of the net properties
-	uint8* Source = bUseDefaultState ? (uint8*)GetObject()->GetArchetype() : (uint8*)InObject;
+	uint8* Source = bUseDefaultState ? (uint8*)InObject->GetArchetype() : (uint8*)InObject;
 
 	if ((Source == nullptr) && bUseDefaultState)
 	{
