New material translator push.

- Changed the 'dot' operator folding logic, so that it folds to a constant only if both arguments are fully constant (resulting in a constant scalar). If the operation cannot be fully folded, it will simply emit a 'dot' operator instruction, rather than attempting to fold individual components separately.
- Refactored the TryFoldOperator function to be cleaner and consistent with the folding logic for Cross and Length operators.
- Removed the unnecessary ".xyzw" in the generated HLSL for material parameters (e.g., "Parameters[0].xyzw.x").

diff --git a/Engine/Source/Runtime/Engine/Private/Materials/MaterialIREmitter.cpp b/Engine/Source/Runtime/Engine/Private/Materials/MaterialIREmitter.cpp
--- a/Engine/Source/Runtime/Engine/Private/Materials/MaterialIREmitter.cpp
+++ b/Engine/Source/Runtime/Engine/Private/Materials/MaterialIREmitter.cpp
@@ -1235,7 +1235,7 @@ static FPrimitiveTypePtr ValidateOperatorAndGetResultType(FEmitter* Emitter, EOp
 }
 
 /// Returns whether the operator supports componentwise application. In other words, if the following is true:
-/// 	op(v, w) == [op(v_x, w_x), ..., op(v_n, w_n)]
+/// 	op(v, w) == [op(v_0, w_0), ..., op(v_n, w_n)]
 static bool IsComponentwiseOperator(EOperator Op)
 {
 	return Op != BO_Dot && Op != BO_Cross && Op != UO_Length;
@@ -1344,17 +1344,17 @@ static FValue* TryFoldComponentwiseOperator(FEmitter* Emitter, EOperator Op, FVa
 
 /// If V is a dimensional and all its components are constants, it unpacks the components into OutComponents and returns true.
 /// If this is not possible for any reason, it returns false.
-static bool TryUnpackConstantVector3(FValue* V, TArrayView<FConstant*> OutComponents)
+static bool TryUnpackConstantScalarOrVector(FValue* V, TArrayView<FConstant*> OutComponents, int32& OutNumComponents)
 {
-	check(V && V->Type->AsPrimitive()->IsVector() && V->Type->AsPrimitive()->GetNumComponents() == 3);
+	// V not specified? Or not a scalar/vector?
 	FDimensional* Dimensional = V->As<FDimensional>();
-	if (!V)
+	if (!Dimensional || V->Type->AsPrimitive()->IsMatrix())
 	{
 		return false;
 	}
 
 	TArrayView<FValue*> Components = Dimensional->GetComponents();
-	for (int32 i = 0; i < 3; ++i)
+	for (int32 i = 0; i < Components.Num(); ++i)
 	{
 		OutComponents[i] = Components[i]->As<FConstant>();
 		if (!OutComponents[i])
@@ -1363,78 +1363,72 @@ static bool TryUnpackConstantVector3(FValue* V, TArrayView<FConstant*> OutCompon
 		}
 	}
 
+	OutNumComponents = Components.Num();
 	return true;
 }
 
-template <typename T>
-static Math::TVector<T> Cross(Math::TVector<T> A, Math::TVector<T> B)
+/// Computes the dot product on two arrays of constant float components.
+static MIR::TFloat ConstantDotFloat(TArrayView<FConstant*> AComponents, TArrayView<FConstant*> BComponents, int32 NumComponents)
 {
-	return Math::TVector<T>(
-		A.Y * B.Z - A.Z * B.Y,
-		A.Z * B.X - A.X * B.Z,
-		A.X * B.Y - A.Y * B.X);
+	float Result = 0.0f;
+	for (int32 i = 0; i < NumComponents; ++i)
+	{
+		Result += AComponents[i]->Float * BComponents[i]->Float;
+	}
+	return Result;
 }
 
 /// Tries to fold the operator, that is to evaluate its result now at translation time if its arguments are constant.
 /// If the operator could not be folded in any way, it returns nullptr.
 static FValue* TryFoldOperator(FEmitter* Emitter, EOperator Op, FValue* A, FValue* B, FValue* C, FPrimitiveTypePtr ResultType)
 {
-	if (Op == BO_Dot)
-	{
-		FPrimitiveTypePtr ArgumentsType = A->Type->AsPrimitive();
-
-		// Componentwise multiply the two vectors.
-		FValue* MultipliedComponents = TryFoldComponentwiseOperator(Emitter, BO_Multiply, A, B, nullptr, ArgumentsType);
-
-		// If the operation was folded (i.e. some componentwise multiplication was folded to a constant),
-		// simply add the components together to compute the final dot result.
-		if (MultipliedComponents)
-		{
-			FValue* Accumulator = Emitter->EmitSubscript(MultipliedComponents, 0);
-			for (int32 i = 1; i < ArgumentsType->NumRows; ++i)
-			{
-				Accumulator = Emitter->EmitOperator(BO_Add, Accumulator, Emitter->EmitSubscript(MultipliedComponents, i));
-			}
-			return Accumulator;
-		}
-	}
-	else if (Op == BO_Cross)
+	FConstant* AComponents[4];
+	int32      ANumComponents;
+	
+	// Try to unpack the first argument to an array of constants.
+	if (TryUnpackConstantScalarOrVector(A, AComponents, ANumComponents))
 	{
-		// Fold the cross operator only if all arguments components are constant.
-		FConstant* AVector[3];
-		FConstant* BVector[3];
-		if (TryUnpackConstantVector3(A, AVector) && TryUnpackConstantVector3(B, BVector))
+		FConstant* BComponents[4];
+		int32      BNumComponents;
+	
+		if (Op == UO_Length)
 		{
 			if (ResultType->ScalarKind == ScalarKind_Float)
 			{
-				FVector3f Result = Cross<float>({ AVector[0]->Float, AVector[1]->Float, AVector[2]->Float }, { BVector[0]->Float, BVector[1]->Float, BVector[2]->Float });
-				return Emitter->EmitConstantFloat3(Result);
+				float Result = FMath::Sqrt(ConstantDotFloat(AComponents, AComponents, ANumComponents));
+				return Emitter->EmitConstantFloatScalar(Result);
 			}
 			else
 			{
 				UE_MIR_UNREACHABLE();
 			}
 		}
-	}
-	else if (Op == UO_Length)
-	{
-		if (A->Type->AsPrimitive()->IsVector())
+		else if ((Op == BO_Dot || Op == BO_Cross) && TryUnpackConstantScalarOrVector(B, BComponents, BNumComponents))
 		{
-			FDimensional* Dimensional = A->As<FDimensional>();
-			if (Dimensional && Dimensional->AreComponentsConstant())
-			{
-				check(Dimensional->GetComponents().Num() <= 4);
+			// Verified before the operation is folded, here as a safety check.
+			check(ANumComponents == BNumComponents);
 
-				FVector4f Input(0.0f, 0.0f, 0.0f, 0.0f);
-				int i = 0;
-				for (FValue const* Component : Dimensional->GetComponents())
+			if (Op == BO_Dot)
+			{
+				if (ResultType->ScalarKind == ScalarKind_Float)
 				{
-					Input[i++] = Component->As<FConstant>()->Float;
+					float Result = ConstantDotFloat(AComponents, BComponents, ANumComponents);
+					return Emitter->EmitConstantFloatScalar(Result);
 				}
-
+				else
+				{
+					UE_MIR_UNREACHABLE();
+				}
+			}
+			else
+			{
+				check(Op == BO_Cross);
 				if (ResultType->ScalarKind == ScalarKind_Float)
 				{
-					return Emitter->EmitConstantFloatScalar(Input.Size());
+					FVector3f AVector{ AComponents[0]->Float, AComponents[1]->Float, AComponents[2]->Float };
+					FVector3f BVector{ BComponents[0]->Float, BComponents[1]->Float, BComponents[2]->Float };
+					FVector3f Result = AVector.Cross(BVector);
+					return Emitter->EmitConstantFloat3(Result);
 				}
 				else
 				{
diff --git a/Engine/Source/Runtime/Engine/Private/Materials/MaterialIRToHLSLTranslator.cpp b/Engine/Source/Runtime/Engine/Private/Materials/MaterialIRToHLSLTranslator.cpp
--- a/Engine/Source/Runtime/Engine/Private/Materials/MaterialIRToHLSLTranslator.cpp
+++ b/Engine/Source/Runtime/Engine/Private/Materials/MaterialIRToHLSLTranslator.cpp
@@ -419,15 +419,12 @@ struct FTranslator : FMaterialIRToHLSLTranslation
 			Printer << "asint(";
 		}
 
-		Printer << TEXT("Material.PreshaderBuffer[") << UniformParameter->Analysis_UniformParameterIndex << TEXT("].");
+		Printer << TEXT("Material.PreshaderBuffer[") << UniformParameter->Analysis_UniformParameterIndex << TEXT("]");
 		check(PrimitiveType->IsScalar() || PrimitiveType->IsVector()); // no matrices yet
-		
-		static const TCHAR* Suffixes[] = { TEXT("x"), TEXT("xy"), TEXT("xyz"), TEXT("xyzw") };
-		Printer << Suffixes[PrimitiveType->NumRows - 1];
 
 		if (PrimitiveType->ScalarKind != MIR::EScalarKind::ScalarKind_Float)
 		{
-			Printer << ")";
+			Printer << ")"; // close the "asint(" bracket
 		}
 	}
 
