Fixed crash at startup due to static symbols being initialized out of order. The skeletal mesh global linked list needs to be initialized prior to being used.

diff --git a/Engine/Source/Runtime/Engine/Internal/SkeletalMeshUpdater.h b/Engine/Source/Runtime/Engine/Internal/SkeletalMeshUpdater.h
--- a/Engine/Source/Runtime/Engine/Internal/SkeletalMeshUpdater.h
+++ b/Engine/Source/Runtime/Engine/Internal/SkeletalMeshUpdater.h
@@ -385,7 +385,7 @@ class FSkeletalMeshUpdateChannel
 			int32 NumPipeRefs = 0;
 		};
 
-		static FGlobalList GlobalList;
+		static FGlobalList& GetGlobalList();
 
 		ENGINE_API FBackend();
 		ENGINE_API virtual ~FBackend();
diff --git a/Engine/Source/Runtime/Engine/Private/SkeletalMeshUpdater.cpp b/Engine/Source/Runtime/Engine/Private/SkeletalMeshUpdater.cpp
--- a/Engine/Source/Runtime/Engine/Private/SkeletalMeshUpdater.cpp
+++ b/Engine/Source/Runtime/Engine/Private/SkeletalMeshUpdater.cpp
@@ -32,10 +32,15 @@ void FSkeletalMeshUpdatePacket::Finalize()
 
 ///////////////////////////////////////////////////////////////////////////////
 
-FSkeletalMeshUpdateChannel::FBackend::FGlobalList FSkeletalMeshUpdateChannel::FBackend::GlobalList;
+FSkeletalMeshUpdateChannel::FBackend::FGlobalList& FSkeletalMeshUpdateChannel::FBackend::GetGlobalList()
+{
+	static FGlobalList GlobalList;
+	return GlobalList;
+}
 
 FSkeletalMeshUpdateChannel::FBackend::FBackend()
 {
+	FGlobalList& GlobalList = GetGlobalList();
 	checkf(!GlobalList.NumPipeRefs, TEXT("A backend is being registered while a SkeletalMeshUpdateChannel is active. Only use DEFINE_SKELETAL_MESH_UPDATE_BACKEND to create backends statically."));
 	GlobalList.List.AddTail(this);
 	GlobalListIndex = GlobalList.Num++;
@@ -43,7 +48,7 @@ FSkeletalMeshUpdateChannel::FBackend::FBackend()
 
 FSkeletalMeshUpdateChannel::FBackend::~FBackend()
 {
-	checkf(!GlobalList.NumPipeRefs, TEXT("A backend is being unregistered while a SkeletalMeshUpdateChannel is still alive."));
+	checkf(!GetGlobalList().NumPipeRefs, TEXT("A backend is being unregistered while a SkeletalMeshUpdateChannel is still alive."));
 	Reset();
 }
 
@@ -69,11 +74,13 @@ void FSkeletalMeshUpdateChannel::FIndexAllocator::Free(int32 Index)
 
 TArray<FSkeletalMeshUpdateChannel> FSkeletalMeshUpdateChannel::GetChannels()
 {
+	FBackend::FGlobalList& GlobalList = FBackend::GetGlobalList();
+
 	TArray<FSkeletalMeshUpdateChannel> Channels;
-	Channels.Reserve(FBackend::GlobalList.Num);
+	Channels.Reserve(GlobalList.Num);
 
 	int32 ChannelIndex = 0;
-	for (FBackend& Backend : FBackend::GlobalList.List)
+	for (FBackend& Backend : GlobalList.List)
 	{
 		// Hitting this means something went wrong with static initialization.
 		check(ChannelIndex++ == Backend.GlobalListIndex);
@@ -88,12 +95,12 @@ FSkeletalMeshUpdateChannel::FSkeletalMeshUpdateChannel(const FBackend* InBackend
 	: OpQueue(new FOpQueue)
 	, Backend(InBackend)
 {
-	++Backend->GlobalList.NumPipeRefs;
+	++FBackend::GetGlobalList().NumPipeRefs;
 }
 
 FSkeletalMeshUpdateChannel::~FSkeletalMeshUpdateChannel()
 {
-	--Backend->GlobalList.NumPipeRefs;
+	--FBackend::GetGlobalList().NumPipeRefs;
 	check(!OpQueue);
 	check(OpStream.Ops.IsEmpty());
 
