Fixing an issue where the Visual Logger does not work with Immersive mode.  The underlying issue is that EditorWorld's time does not get ticked when in Immersive Mode.  The solution is to use the FApp time rather than EditorWorld's time.  A bonus issue fixed:  Subscribe to an even earlier event when starting PIE to reset the EditorBaseTimestamp to 0, so that early start-up Logs do not appear at the end of the timeline, but rather the beginning.

diff --git a/Engine/Source/Runtime/Engine/Private/VisualLogger/VisualLogger.cpp b/Engine/Source/Runtime/Engine/Private/VisualLogger/VisualLogger.cpp
--- a/Engine/Source/Runtime/Engine/Private/VisualLogger/VisualLogger.cpp
+++ b/Engine/Source/Runtime/Engine/Private/VisualLogger/VisualLogger.cpp
@@ -16,7 +16,9 @@
 #include "UnrealEngine.h"
 
 #if WITH_EDITOR
+#include "Editor.h"
 #include "Editor/EditorEngine.h"
+#include "Misc/App.h"
 #else
 #include "Engine/Engine.h"
 #include "Serialization/CustomVersion.h"
@@ -276,35 +278,24 @@ double FVisualLogger::GetTimeStampForObject(const UObject* Object) const
 		return GetTimeStampFunc(Object);
 	}
 
-	const UWorld* WorldForTimeStamp = nullptr;
 #if WITH_EDITOR
 	// When we're in the Editor, use a Global Engine TimeStamp so that we can synchronize between Client and Server instances.
-	UEditorEngine* EditorEngine = GIsEditor ? Cast<UEditorEngine>(GEngine) : nullptr;
-	if (EditorEngine)
+	if (GIsEditor)
 	{
-		// Using the Editor world to ensure a consistent clock since it does not reset when more clients are added or removed
-		// and can exist before a PIE session is started.
+		using namespace UE::VisLog::Private;
 
-		// We should always have the Editor world to use, but in some edge cases it can no longer be valid (e.g., world getting GCed).
-		WorldForTimeStamp = EditorEngine->GetEditorWorldContext().World();
-		if (WorldForTimeStamp)
-		{
-			using namespace UE::VisLog::Private;
+		double Expected = 0.0;
+		EditorOnly::EditorBaseTimeStamp.compare_exchange_strong(Expected, FApp::GetCurrentTime());
 
-			double Expected = 0;
-			EditorOnly::EditorBaseTimeStamp.compare_exchange_strong(Expected, WorldForTimeStamp->TimeSeconds);
-
-			return WorldForTimeStamp->TimeSeconds - EditorOnly::EditorBaseTimeStamp;
+		if (EditorOnly::EditorBaseTimeStamp > 0.0)
+		{
+			return FApp::GetCurrentTime() - EditorOnly::EditorBaseTimeStamp;
 		}
 	}
-
 #endif
 
 	// This will be the fallback mode in standalone.  We do not have a synchronized clock.
-	if (!WorldForTimeStamp)
-	{
-		WorldForTimeStamp = GetWorldForVisualLogger(Object);
-	}
+	const UWorld* WorldForTimeStamp = GetWorldForVisualLogger(Object);
 
 	return WorldForTimeStamp ? WorldForTimeStamp->TimeSeconds : 0.0;
 }
@@ -835,19 +826,24 @@ FVisualLogger::FVisualLogger()
 	});
 
 #if WITH_EDITOR
-	PIEStartedHandle = FWorldDelegates::OnPIEStarted.AddLambda(
-		[this](UGameInstance*) {
-			// We only want to do this if we are using unique log names, otherwise it makes more sense
-			// to just continue logging at later timestamps on the same timeline as the previous one.
-			if (bForceUniqueLogNames)
+	PIEStartedHandle = FEditorDelegates::StartPIE.AddLambda([this](bool bSimulating)
+		{
+			// Don't reset if we're not going to PIE in the same process
+			if (TOptional<FRequestPlaySessionParams> PlaySessionRequest = GEditor->GetPlaySessionRequest())
 			{
-				if (GEngine->IsSettingUpPlayWorld())
+				if (PlaySessionRequest->SessionDestination != EPlaySessionDestinationType::InProcess)
 				{
-					int32 OldUniqueNumber = UE::VisLog::Private::UniqueRunNumber;
-					OnDataReset();
-					UE::VisLog::Private::UniqueRunNumber = OldUniqueNumber + 1;
+					return;
 				}
+			}
 
+			// We only want to do this if we are using unique log names, otherwise it makes more sense
+			// to just continue logging at later timestamps on the same timeline as the previous one.
+			if (bForceUniqueLogNames)
+			{
+				int32 OldUniqueNumber = UE::VisLog::Private::UniqueRunNumber;
+				OnDataReset();
+				UE::VisLog::Private::UniqueRunNumber = OldUniqueNumber + 1;
 			}
 		});
 #endif
@@ -856,7 +852,7 @@ FVisualLogger::FVisualLogger()
 void FVisualLogger::TearDown()
 {
 #if WITH_EDITOR
-	FWorldDelegates::OnPIEStarted.Remove(PIEStartedHandle);
+	FEditorDelegates::StartPIE.Remove(PIEStartedHandle);
 #endif
 
 	SetIsRecording(false);
