Also enable bindless samplers when inline raytracing requires bindless (not just resources)

diff --git a/Engine/Source/Runtime/Engine/Private/ShaderCompiler/ShaderCompiler.cpp b/Engine/Source/Runtime/Engine/Private/ShaderCompiler/ShaderCompiler.cpp
--- a/Engine/Source/Runtime/Engine/Private/ShaderCompiler/ShaderCompiler.cpp
+++ b/Engine/Source/Runtime/Engine/Private/ShaderCompiler/ShaderCompiler.cpp
@@ -3987,9 +3987,10 @@ void GlobalBeginCompileShader(
 		if (bIsInlineRaytracingShader && bInlineRayTracingRequiresBindless)
 		{
 			check(ResourcesConfig != ERHIBindlessConfiguration::Disabled);
+			check(SamplersConfig != ERHIBindlessConfiguration::Disabled);
 		}
 
-		bool EnableBindlessForInlineRaytracing = (bIsInlineRaytracingShader && bInlineRayTracingRequiresBindless && ResourcesConfig != ERHIBindlessConfiguration::Disabled);
+		bool EnableBindlessForInlineRaytracing = (bIsInlineRaytracingShader && bInlineRayTracingRequiresBindless && ResourcesConfig != ERHIBindlessConfiguration::Disabled && SamplersConfig != ERHIBindlessConfiguration::Disabled);
 		if (ResourcesConfig == ERHIBindlessConfiguration::AllShaders || (ResourcesConfig == ERHIBindlessConfiguration::RayTracingShaders && bIsRaytracingShader) || EnableBindlessForInlineRaytracing)
 		{
 			bEnableBindlessMacro = true;
@@ -3997,7 +3998,7 @@ void GlobalBeginCompileShader(
 			SET_SHADER_DEFINE(Input.Environment, ENABLE_BINDLESS_RESOURCES, true);
 		}
 
-		if (SamplersConfig == ERHIBindlessConfiguration::AllShaders || (SamplersConfig == ERHIBindlessConfiguration::RayTracingShaders && bIsRaytracingShader))
+		if (SamplersConfig == ERHIBindlessConfiguration::AllShaders || (SamplersConfig == ERHIBindlessConfiguration::RayTracingShaders && bIsRaytracingShader) || EnableBindlessForInlineRaytracing)
 		{
 			bEnableBindlessMacro = true;
 			Input.Environment.CompilerFlags.Add(CFLAG_BindlessSamplers);
