Fixed body instance not cleaing out per-shape runtime overrides on body termination.
If you override a shape property at runtime we populate an array with an override per shape. TermBody was not clearing these overrides. If you then call SetStaticMesh for example on a static mesh component with a mesh containing more shapes - any update to its filters will cause it to iterate its shapes and read the array of overrides causing it to read out of bounds.

diff --git a/Engine/Source/Runtime/Engine/Private/PhysicsEngine/BodyInstance.cpp b/Engine/Source/Runtime/Engine/Private/PhysicsEngine/BodyInstance.cpp
--- a/Engine/Source/Runtime/Engine/Private/PhysicsEngine/BodyInstance.cpp
+++ b/Engine/Source/Runtime/Engine/Private/PhysicsEngine/BodyInstance.cpp
@@ -104,6 +104,12 @@ FAutoConsoleVariableRef CVarSensitiveSleepThresholdMultiplier(
 	SensitiveSleepThresholdMultiplier,
 	TEXT("The sleep threshold multiplier to use for bodies using the Sensitive sleep family."));
 
+bool bClearPerShapeOverridesOnTerminate = true;
+FAutoConsoleVariableRef CVarClearPerShapeOverridesOnTerminate(
+	TEXT("p.ClearPerShapeOverridesOnTerminate"),
+	bClearPerShapeOverridesOnTerminate,
+	TEXT("Whether shape overrides are cleared in between recreations of a body instance"));
+
 using namespace PhysicsInterfaceTypes;
 
 bool IsRigidBodyKinematic_AssumesLocked(const FPhysicsActorHandle& InActorRef)
@@ -1846,6 +1852,15 @@ void FBodyInstance::TermBody(bool bNeverDeferRelease)
 	OwnerComponent = NULL;
 	ExternalCollisionProfileBodySetup = nullptr;
 
+	// Reset per-shape runtime overrides
+	if(bClearPerShapeOverridesOnTerminate)
+	{
+		bShapeCollisionEnabledIsSet = false;
+		bShapeCollisionResponsesIsSet = false;
+		ShapeCollisionEnabled.Reset();
+		ShapeCollisionResponses.Reset();
+	}
+
 	if (DOFConstraint)
 	{
 		DOFConstraint->TermConstraint();
