Streamable Manager error handling
- Add an Error field to FStreamableHandle.
- Add FStreamableDelegateWithHandle to more easily find and inspect handles when a load completes.
- MakeShareable -> MakeShared
- Remove unecessary dll imports/exports

diff --git a/Engine/Source/Runtime/Engine/Classes/Engine/StreamableManager.h b/Engine/Source/Runtime/Engine/Classes/Engine/StreamableManager.h
--- a/Engine/Source/Runtime/Engine/Classes/Engine/StreamableManager.h
+++ b/Engine/Source/Runtime/Engine/Classes/Engine/StreamableManager.h
@@ -6,18 +6,36 @@
 #include "Containers/ArrayView.h"
 #include "Math/NumericLimits.h"
 #include "Misc/PackageAccessTracking.h"
+#include "Misc/SourceLocation.h"
+#include "Misc/SourceLocationUtils.h"
 #include "Templates/Casts.h"
 #include "Templates/Function.h"
 #include "UObject/Class.h"
 #include "UObject/GCObject.h"
 #include "UObject/ObjectMacros.h"
 #include "UObject/SoftObjectPtr.h"
-#include "Misc/SourceLocation.h"
-#include "Misc/SourceLocationUtils.h"
+
+enum class EStreamableErrorType : uint8
+{
+	Unknown = 0,
+	/** Low-level package loading error */
+	PackageLoadError,
+	/** Low-level package loading was canceled. This is different than canceling a FStreamableHandle which just detaches from the low level request. */
+	PackageLoadCanceled,
+	/** Error downloading assets, not implemented yet */
+	DownloadError,
+};
+
+struct FStreamableError
+{
+	EStreamableErrorType ErrorType = EStreamableErrorType::Unknown;
+	FString ErrorMessage;
+};
 
 /** Defines FStreamableDelegate delegate interface */
-DECLARE_DELEGATE(FStreamableDelegate);
-DECLARE_DELEGATE_OneParam(FStreamableUpdateDelegate, TSharedRef<struct FStreamableHandle>);
+using FStreamableDelegate = TDelegate<void()>;
+using FStreamableDelegateWithHandle = TDelegate<void(TSharedPtr<struct FStreamableHandle>)>;
+using FStreamableUpdateDelegate = TDelegate<void(TSharedRef<struct FStreamableHandle>)>;
 
 #if WITH_EDITOR
 enum class ECookLoadType : uint8;
@@ -129,6 +147,16 @@ struct FStreamableHandle : public TSharedFromThis<FStreamableHandle>
 		return bIsCombinedHandle;
 	}
 
+	bool HasError() const
+	{
+		return Error.IsSet();
+	}
+
+	const TOptional<FStreamableError>& GetError() const
+	{
+		return Error;
+	}
+
 	/** Returns true if we've done all the loading we can now, ie all handles are either completed or stalled */
 	ENGINE_API bool HasLoadCompletedOrStalled() const;
 
@@ -191,9 +219,11 @@ struct FStreamableHandle : public TSharedFromThis<FStreamableHandle>
 	ENGINE_API bool HasUpdateDelegate() const;
 
 	/** Bind delegate that is called when load completes, only works if loading is in progress. This will overwrite any already bound delegate! */
+	ENGINE_API bool BindCompleteDelegate(FStreamableDelegateWithHandle NewDelegate);
 	ENGINE_API bool BindCompleteDelegate(FStreamableDelegate NewDelegate);
 
 	/** Bind delegate that is called if handle is canceled, only works if loading is in progress. This will overwrite any already bound delegate! */
+	ENGINE_API bool BindCancelDelegate(FStreamableDelegateWithHandle NewDelegate);
 	ENGINE_API bool BindCancelDelegate(FStreamableDelegate NewDelegate);
 
 	/** Bind delegate that is called periodically as delegate updates, only works if loading is in progress. This will overwrite any already bound delegate! */
@@ -264,6 +294,17 @@ struct FStreamableHandle : public TSharedFromThis<FStreamableHandle>
 	 * @param AssociatedHandle	Streamable handle associated with this delegate, may be null
 	 * @param CancelDelegate	If handle gets cancelled before primary delegate executes, this delegate will be called instead
 	 */
+	static ENGINE_API void ExecuteDelegate(const FStreamableDelegateWithHandle& Delegate, TSharedPtr<FStreamableHandle> AssociatedHandle = nullptr, const FStreamableDelegateWithHandle& CancelDelegate = FStreamableDelegateWithHandle());
+	/** Overload taking the Delegate by rvalue. Using this will prevent any copies of FStreamableDelegate from being made if the callback is deferred. */
+	static ENGINE_API void ExecuteDelegate(FStreamableDelegateWithHandle&& Delegate, TSharedPtr<FStreamableHandle> AssociatedHandle = nullptr, FStreamableDelegateWithHandle&& CancelDelegate = FStreamableDelegateWithHandle());
+
+	/**
+	 * Calls a StreamableDelegate, this will add to the delayed callback queue depending on s.StreamableDelegateDelayFrames
+	 *
+	 * @param Delegate			Primary delegate to execute
+	 * @param AssociatedHandle	Streamable handle associated with this delegate, may be null
+	 * @param CancelDelegate	If handle gets cancelled before primary delegate executes, this delegate will be called instead
+	 */
 	static ENGINE_API void ExecuteDelegate(const FStreamableDelegate& Delegate, TSharedPtr<FStreamableHandle> AssociatedHandle = nullptr, const FStreamableDelegate& CancelDelegate = FStreamableDelegate());
 	/** Overload taking the Delegate by rvalue. Using this will prevent any copies of FStreamableDelegate from being made if the callback is deferred. */
 	static ENGINE_API void ExecuteDelegate(FStreamableDelegate&& Delegate, TSharedPtr<FStreamableHandle> AssociatedHandle = nullptr, FStreamableDelegate&& CancelDelegate = FStreamableDelegate());
@@ -368,6 +409,13 @@ struct FStreamableHandle : public TSharedFromThis<FStreamableHandle>
 		return OfType;
 	}
 
+	/** Constructor is private, but this allows use by MakeShared */
+private:
+	struct FPrivateToken { explicit FPrivateToken() = default; };
+	FStreamableHandle() = delete;
+public:
+	explicit FStreamableHandle(FPrivateToken);
+
 	/** Destructor */
 	ENGINE_API ~FStreamableHandle();
 
@@ -384,28 +432,34 @@ private:
 	friend struct FStreamable;
 
 	/** Called from self as part of SetDebugNameIfEmptyOrDefault */
-	ENGINE_API bool IsHandleNameEmptyOrDefault() const;
+	bool IsHandleNameEmptyOrDefault() const;
 
 	/** Called from manager to complete the request */
-	ENGINE_API void CompleteLoad();
+	void CompleteLoad();
 
 	/** Callback when async load finishes, it's here so we can use a shared pointer for callback safety */
-	ENGINE_API void AsyncLoadCallbackWrapper(const FName& PackageName, UPackage* LevelPackage, EAsyncLoadingResult::Type Result, FSoftObjectPath TargetName);
+	void AsyncLoadCallbackWrapper(const FName& PackageName, UPackage* LevelPackage, EAsyncLoadingResult::Type Result, FSoftObjectPath TargetName);
 
 	/** Notify all parents that a child completed loading */
-	ENGINE_API void NotifyParentsOfCompletion();
+	void NotifyParentsOfCompletion();
 
 	/** Notify all parents that a child was canceled */
-	ENGINE_API void NotifyParentsOfCancellation();
+	void NotifyParentsOfCancellation();
 
 	/** Called on meta handle when a child handle has completed/canceled */
-	ENGINE_API void UpdateCombinedHandle();
+	void UpdateCombinedHandle();
 
 	/** Call to call the update delegate if bound, will propagate to parents */
-	ENGINE_API void CallUpdateDelegate();
+	void CallUpdateDelegate();
 
 	/** Unbind all of our delegates, called after a cancel/load */
-	ENGINE_API void UnbindDelegates();
+	void UnbindDelegates();
+
+	/** Record any errors that occured with this handle */
+	void SetError(FStreamableError&& Error);
+	void SetErrorOnParents(const FStreamableError& Error);
+
+	TOptional<FStreamableError> Error;
 
 	/** True if this request has finished loading. It may still be active, or it may have been released */
 	bool bLoadCompleted;
@@ -426,10 +480,10 @@ private:
 	bool bIsCombinedHandle;
 
 	/** Delegate to call when streaming is completed */
-	FStreamableDelegate CompleteDelegate;
+	FStreamableDelegateWithHandle CompleteDelegate;
 
 	/** Delegate to call when streaming is canceled */
-	FStreamableDelegate CancelDelegate;
+	FStreamableDelegateWithHandle CancelDelegate;
 
 	/** Called periodically during streaming to update progress UI */
 	FStreamableUpdateDelegate UpdateDelegate;
@@ -472,8 +526,6 @@ private:
 #if WITH_EDITOR
 	ECookLoadType CookLoadType;
 #endif
-
-	ENGINE_API FStreamableHandle();
 };
 
 enum class EStreamableManagerCombinedHandleOptions : uint8 
@@ -507,10 +559,10 @@ struct FStreamableManager : public FGCObject
 	 * @param bStartStalled			[optional] If true, the handle will start in a stalled state and will not attempt to actually async load until StartStalledHandle is called on it
 	 * @param DebugName				[optional] Name of this handle, either FString or anything that can construct FString, will be reported in debug tools, will report Source Location if DebugName is not specified explicitly
 	 */
-	template< typename PathContainerType = TArray<FSoftObjectPath>, typename FuncType = FStreamableDelegate, typename DebugNameType = UE::FSourceLocation >
+	template< typename PathContainerType = TArray<FSoftObjectPath>, typename FuncType = FStreamableDelegateWithHandle, typename DebugNameType = UE::FSourceLocation >
 	TSharedPtr<FStreamableHandle> RequestAsyncLoad(
 		PathContainerType&& TargetsToStream,
-		FuncType&& DelegateToCall = FStreamableDelegate(),
+		FuncType&& DelegateToCall = FStreamableDelegateWithHandle(),
 		TAsyncLoadPriority Priority = DefaultAsyncLoadPriority,
 		bool bManageActiveHandle = false,
 		bool bStartStalled = false,
@@ -612,21 +664,21 @@ struct FStreamableManager : public FGCObject
 private:
 	friend FStreamableHandle;
 
-	ENGINE_API void RemoveReferencedAsset(const FSoftObjectPath& Target, TSharedRef<FStreamableHandle> Handle);
-	ENGINE_API void StartHandleRequests(TSharedRef<FStreamableHandle> Handle);
-	ENGINE_API TArray<int32> GetAsyncLoadRequestIds(TSharedRef<FStreamableHandle> Handle);
-	ENGINE_API void FindInMemory(FSoftObjectPath& InOutTarget, struct FStreamable* Existing, UPackage* Package = nullptr);
-	ENGINE_API FSoftObjectPath HandleLoadedRedirector(UObjectRedirector* LoadedRedirector, FSoftObjectPath RequestedPath, struct FStreamable* RequestedStreamable);
-	ENGINE_API struct FStreamable* FindStreamable(const FSoftObjectPath& Target) const;
-	ENGINE_API struct FStreamable* StreamInternal(const FSoftObjectPath& Target, TAsyncLoadPriority Priority, TSharedRef<FStreamableHandle> Handle);
+	void RemoveReferencedAsset(const FSoftObjectPath& Target, TSharedRef<FStreamableHandle> Handle);
+	void StartHandleRequests(TSharedRef<FStreamableHandle> Handle);
+	TArray<int32> GetAsyncLoadRequestIds(TSharedRef<FStreamableHandle> Handle);
+	void FindInMemory(FSoftObjectPath& InOutTarget, struct FStreamable* Existing, UPackage* Package = nullptr);
+	FSoftObjectPath HandleLoadedRedirector(UObjectRedirector* LoadedRedirector, FSoftObjectPath RequestedPath, struct FStreamable* RequestedStreamable);
+	struct FStreamable* FindStreamable(const FSoftObjectPath& Target) const;
+	struct FStreamable* StreamInternal(const FSoftObjectPath& Target, TAsyncLoadPriority Priority, TSharedRef<FStreamableHandle> Handle);
 	ENGINE_API UObject* GetStreamed(const FSoftObjectPath& Target) const;
-	ENGINE_API void CheckCompletedRequests(const FSoftObjectPath& Target, struct FStreamable* Existing);
+	void CheckCompletedRequests(const FSoftObjectPath& Target, struct FStreamable* Existing);
 
-	ENGINE_API void OnPreGarbageCollect();
-	ENGINE_API void AsyncLoadCallback(FSoftObjectPath Request, UPackage* Package);
+	void OnPreGarbageCollect();
+	void AsyncLoadCallback(FSoftObjectPath Request, UPackage* Package, EAsyncLoadingResult::Type Result);
 
 	ENGINE_API static bool ShouldStripDebugName();
-	ENGINE_API TSharedPtr<FStreamableHandle> RequestAsyncLoadInternal(TArray<FSoftObjectPath>&& TargetsToStream, FStreamableDelegate&& DelegateToCall, TAsyncLoadPriority Priority, bool bManageActiveHandle, bool bStartStalled, FString&& DebugName);
+	ENGINE_API TSharedPtr<FStreamableHandle> RequestAsyncLoadInternal(TArray<FSoftObjectPath>&& TargetsToStream, FStreamableDelegateWithHandle&& DelegateToCall, TAsyncLoadPriority Priority, bool bManageActiveHandle, bool bStartStalled, FString&& DebugName);
 	ENGINE_API TSharedPtr<FStreamableHandle> RequestSyncLoadInternal(TArray<FSoftObjectPath>&& TargetsToStream, bool bManageActiveHandle, FString&& DebugName);
 
 	/** Map of paths to streamable objects, this will be the post-redirector name */
@@ -698,6 +750,41 @@ void FStreamableHandle::ForEachLoadedAsset(CallableT Callable) const
 	}
 }
 
+namespace UE::StreamableManager::Private
+{
+	inline FStreamableDelegateWithHandle WrapDelegate(const FStreamableDelegate& InDelegate)
+	{
+		if (InDelegate.IsBound())
+		{
+			return FStreamableDelegateWithHandle::CreateLambda(
+				[InDelegate](TSharedPtr<FStreamableHandle>)
+				{
+					// Delegates may get tick delayed so may not be safe to call if bound to a GC'd object
+					InDelegate.ExecuteIfBound();
+				});
+		}
+
+		return FStreamableDelegateWithHandle();
+	}
+
+	inline FStreamableDelegateWithHandle WrapDelegate(FStreamableDelegate&& InDelegate)
+	{
+		if (InDelegate.IsBound())
+		{
+			return FStreamableDelegateWithHandle::CreateLambda(
+				[InDelegate = MoveTemp(InDelegate)](TSharedPtr<FStreamableHandle>)
+				{
+					// Delegates may get tick delayed so may not be safe to call if bound to a GC'd object
+					InDelegate.ExecuteIfBound();
+				});
+		}
+
+		return FStreamableDelegateWithHandle();
+	}
+
+	template<class> inline constexpr bool always_false_v = false;
+}
+
 template< typename PathContainerType, typename FuncType, typename DebugNameType >
 TSharedPtr<FStreamableHandle> FStreamableManager::RequestAsyncLoad(
 	PathContainerType&& TargetsToStream,
@@ -707,14 +794,29 @@ TSharedPtr<FStreamableHandle> FStreamableManager::RequestAsyncLoad(
 	bool bStartStalled,
 	DebugNameType&& DebugNameOrLocation)
 {
-	FStreamableDelegate DelegateToCall;
-	if constexpr (std::is_constructible_v<FStreamableDelegate, FuncType>)
+	FStreamableDelegateWithHandle DelegateToCall;
+	if constexpr (std::is_constructible_v<FStreamableDelegateWithHandle, FuncType>)
+	{
+		DelegateToCall = FStreamableDelegateWithHandle(Forward<FuncType>(Callback));
+	}
+	else if constexpr (std::is_constructible_v<FStreamableDelegate, FuncType>)
+	{
+		DelegateToCall = UE::StreamableManager::Private::WrapDelegate(
+			FStreamableDelegate(Forward<FuncType>(Callback)));
+	}
+	else if constexpr (std::is_invocable_v<FuncType, TSharedPtr<FStreamableHandle>>)
+	{
+		DelegateToCall = FStreamableDelegateWithHandle::CreateLambda(Forward<FuncType>(Callback));
+	}
+	else if constexpr (std::is_invocable_v<FuncType>)
 	{
-		DelegateToCall = FStreamableDelegate(Forward<FuncType>(Callback));
+		DelegateToCall = UE::StreamableManager::Private::WrapDelegate(
+			FStreamableDelegate::CreateLambda(Forward<FuncType>(Callback)));
 	}
 	else
 	{
-		DelegateToCall = FStreamableDelegate::CreateLambda(Forward<FuncType>(Callback));
+		// Static assertions are evaluated even in non-instantiable static branches unless a dependent type is used. C++23 fixes this.
+		static_assert(UE::StreamableManager::Private::always_false_v<FuncType>, "Cannot bind callback");
 	}
 
 	if constexpr (std::is_same_v<std::decay_t<DebugNameType>, UE::FSourceLocation>)
diff --git a/Engine/Source/Runtime/Engine/Private/StreamableManager.cpp b/Engine/Source/Runtime/Engine/Private/StreamableManager.cpp
--- a/Engine/Source/Runtime/Engine/Private/StreamableManager.cpp
+++ b/Engine/Source/Runtime/Engine/Private/StreamableManager.cpp
@@ -178,7 +178,7 @@ public:
 		while (DelegateToCall)
 		{
 			FPendingDelegate* NextDelegateToCall = DelegateToCall->Next;
-			DelegateToCall->Delegate.ExecuteIfBound();
+			DelegateToCall->Delegate.ExecuteIfBound(DelegateToCall->RelatedHandle);
 			delete DelegateToCall;
 			DelegateToCall = NextDelegateToCall;
 		}
@@ -215,10 +215,10 @@ private:
 		FPendingDelegate* NextByHandle;
 
 		/** Delegate to call when frames are up */
-		FStreamableDelegate Delegate;
+		FStreamableDelegateWithHandle Delegate;
 
 		/** Delegate to call if this gets cancelled early */
-		FStreamableDelegate CancelDelegate;
+		FStreamableDelegateWithHandle CancelDelegate;
 
 		/** Handle related to delegates, needs to keep these around to avoid things GCing before the user callback goes off. This may be null */
 		TSharedPtr<FStreamableHandle> RelatedHandle;
@@ -346,7 +346,7 @@ TStreamableHandleContextDataTypeID FStreamableHandleContextDataBase::AllocateCla
 }
 
 
-FStreamableHandle::FStreamableHandle()
+FStreamableHandle::FStreamableHandle(FPrivateToken)
 	: bLoadCompleted(false)
 	, bReleased(false)
 	, bCanceled(false)
@@ -362,7 +362,7 @@ FStreamableHandle::FStreamableHandle()
 {
 }
 
-bool FStreamableHandle::BindCompleteDelegate(FStreamableDelegate NewDelegate)
+bool FStreamableHandle::BindCompleteDelegate(FStreamableDelegateWithHandle NewDelegate)
 {
 	if (!IsLoadingInProgress())
 	{
@@ -374,7 +374,19 @@ bool FStreamableHandle::BindCompleteDelegate(FStreamableDelegate NewDelegate)
 	return true;
 }
 
-bool FStreamableHandle::BindCancelDelegate(FStreamableDelegate NewDelegate)
+bool FStreamableHandle::BindCompleteDelegate(FStreamableDelegate NewDelegate)
+{
+	if (!IsLoadingInProgress())
+	{
+		// Too late!
+		return false;
+	}
+
+	return BindCompleteDelegate(
+		UE::StreamableManager::Private::WrapDelegate(MoveTemp(NewDelegate)));
+}
+
+bool FStreamableHandle::BindCancelDelegate(FStreamableDelegateWithHandle NewDelegate)
 {
 	if (!IsLoadingInProgress())
 	{
@@ -386,6 +398,18 @@ bool FStreamableHandle::BindCancelDelegate(FStreamableDelegate NewDelegate)
 	return true;
 }
 
+bool FStreamableHandle::BindCancelDelegate(FStreamableDelegate NewDelegate)
+{
+	if (!IsLoadingInProgress())
+	{
+		// Too late!
+		return false;
+	}
+
+	return BindCancelDelegate(
+		UE::StreamableManager::Private::WrapDelegate(MoveTemp(NewDelegate)));
+}
+
 bool FStreamableHandle::BindUpdateDelegate(FStreamableUpdateDelegate NewDelegate)
 {
 	if (!IsLoadingInProgress())
@@ -800,36 +824,28 @@ void FStreamableHandle::CompleteLoad()
 
 void FStreamableHandle::NotifyParentsOfCompletion()
 {
-	if (ParentHandles.Num() > 0)
+	// Update any meta handles that are still active.
+	for (TWeakPtr<FStreamableHandle>& WeakHandle : ParentHandles)
 	{
-		// Update any meta handles that are still active. Copy the array first as elements may be removed from original while iterating
-		TArray<TWeakPtr<FStreamableHandle>> ParentHandlesCopy = ParentHandles;
-		for (TWeakPtr<FStreamableHandle>& WeakHandle : ParentHandlesCopy)
-		{
-			TSharedPtr<FStreamableHandle> Handle = WeakHandle.Pin();
+		TSharedPtr<FStreamableHandle> Handle = WeakHandle.Pin();
 
-			if (Handle.IsValid())
-			{
-				++Handle->CompletedChildCount;
-			}
+		if (Handle.IsValid())
+		{
+			++Handle->CompletedChildCount;
 		}
 	}
 }
 
 void FStreamableHandle::NotifyParentsOfCancellation()
 {
-	if (ParentHandles.Num() > 0)
+	// Update any meta handles that are still active.
+	for (TWeakPtr<FStreamableHandle>& WeakHandle : ParentHandles)
 	{
-		// Update any meta handles that are still active. Copy the array first as elements may be removed from original while iterating
-		TArray<TWeakPtr<FStreamableHandle>> ParentHandlesCopy = ParentHandles;
-		for (TWeakPtr<FStreamableHandle>& WeakHandle : ParentHandlesCopy)
-		{
-			TSharedPtr<FStreamableHandle> Handle = WeakHandle.Pin();
+		TSharedPtr<FStreamableHandle> Handle = WeakHandle.Pin();
 
-			if (Handle.IsValid())
-			{
-				++Handle->CanceledChildCount;
-			}
+		if (Handle.IsValid())
+		{
+			++Handle->CanceledChildCount;
 		}
 	}
 }
@@ -914,6 +930,29 @@ void FStreamableHandle::UnbindDelegates()
 	CompleteDelegate.Unbind();
 }
 
+void FStreamableHandle::SetError(FStreamableError&& InError)
+{
+	if (!Error)
+	{
+		SetErrorOnParents(InError);
+		Error = MoveTemp(InError);
+	}
+}
+
+void FStreamableHandle::SetErrorOnParents(const FStreamableError& InError)
+{
+	// Update any meta handles that are still active
+	for (TWeakPtr<FStreamableHandle>& WeakHandle : ParentHandles)
+	{
+		TSharedPtr<FStreamableHandle> Handle = WeakHandle.Pin();
+
+		if (Handle.IsValid())
+		{
+			Handle->SetError(CopyTemp(InError));
+		}
+	}
+}
+
 void FStreamableHandle::AsyncLoadCallbackWrapper(const FName& PackageName, UPackage* Package, EAsyncLoadingResult::Type Result, FSoftObjectPath TargetName)
 {
 	check(IsInGameThread());
@@ -921,7 +960,19 @@ void FStreamableHandle::AsyncLoadCallbackWrapper(const FName& PackageName, UPack
 	// Needed so we can bind with a shared pointer for safety
 	if (OwningManager)
 	{
-		OwningManager->AsyncLoadCallback(TargetName, Package);
+		switch (Result)
+		{
+		case EAsyncLoadingResult::Failed:
+			// We could supply an error string if async loading bubbled one up...
+			// Possibly GetExplanationForUnavailablePackage?
+			SetError(FStreamableError{ EStreamableErrorType::PackageLoadError, TEXT("Async load failed")});
+			break;
+		case EAsyncLoadingResult::Canceled:
+			SetError(FStreamableError{ EStreamableErrorType::PackageLoadCanceled, TEXT("Async load canceled")});
+			break;
+		}
+
+		OwningManager->AsyncLoadCallback(MoveTemp(TargetName), Package, Result);
 
 		if (!HasLoadCompleted())
 		{
@@ -941,14 +992,14 @@ namespace UE::StreamableManager::Private
 	void ExecuteDelegateInternal(InStreamableDelegate&& Delegate, TSharedPtr<FStreamableHandle> AssociatedHandle, InStreamableDelegate&& CancelDelegate)
 	{
 		using UnRefCvStreamableDelegate = std::remove_cv_t<std::remove_reference_t<InStreamableDelegate>>;
-		static_assert(std::is_same_v<UnRefCvStreamableDelegate, FStreamableDelegate>, "ExecuteDelegateInternal is only valid to be called with FStreamableDelegate.");
+		static_assert(std::is_same_v<UnRefCvStreamableDelegate, FStreamableDelegateWithHandle>, "ExecuteDelegateInternal is only valid to be called with FStreamableDelegateWithHandle.");
 
-		if (Delegate.IsBound())
+		if (Delegate.IsBound()) // checks for safety as well
 		{
 			if (GStreamableDelegateDelayFrames == 0)
 			{
 				// Execute it immediately
-				Delegate.Execute();
+				Delegate.Execute(AssociatedHandle);
 			}
 			else
 			{
@@ -964,16 +1015,28 @@ namespace UE::StreamableManager::Private
 	}
 }
 
-void FStreamableHandle::ExecuteDelegate(const FStreamableDelegate& Delegate, TSharedPtr<FStreamableHandle> AssociatedHandle, const FStreamableDelegate& CancelDelegate)
+void FStreamableHandle::ExecuteDelegate(const FStreamableDelegateWithHandle& Delegate, TSharedPtr<FStreamableHandle> AssociatedHandle, const FStreamableDelegateWithHandle& CancelDelegate)
 {
 	UE::StreamableManager::Private::ExecuteDelegateInternal(Delegate, MoveTemp(AssociatedHandle), CancelDelegate);
 }
 
-void FStreamableHandle::ExecuteDelegate(FStreamableDelegate&& Delegate, TSharedPtr<FStreamableHandle> AssociatedHandle, FStreamableDelegate&& CancelDelegate)
+void FStreamableHandle::ExecuteDelegate(FStreamableDelegateWithHandle&& Delegate, TSharedPtr<FStreamableHandle> AssociatedHandle, FStreamableDelegateWithHandle&& CancelDelegate)
 {
 	UE::StreamableManager::Private::ExecuteDelegateInternal(MoveTemp(Delegate), MoveTemp(AssociatedHandle), MoveTemp(CancelDelegate));
 }
 
+void FStreamableHandle::ExecuteDelegate(const FStreamableDelegate& Delegate, TSharedPtr<FStreamableHandle> AssociatedHandle, const FStreamableDelegate& CancelDelegate)
+{
+	using namespace UE::StreamableManager::Private;
+	ExecuteDelegateInternal(WrapDelegate(Delegate), MoveTemp(AssociatedHandle), WrapDelegate(CancelDelegate));
+}
+
+void FStreamableHandle::ExecuteDelegate(FStreamableDelegate&& Delegate, TSharedPtr<FStreamableHandle> AssociatedHandle, FStreamableDelegate&& CancelDelegate)
+{
+	using namespace UE::StreamableManager::Private;
+	ExecuteDelegateInternal(WrapDelegate(MoveTemp(Delegate)), MoveTemp(AssociatedHandle), WrapDelegate(MoveTemp(CancelDelegate)));
+}
+
 TSharedPtr<FStreamableHandle> FStreamableHandle::FindMatchingHandle(TFunction<bool(const FStreamableHandle&)> Predicate) const
 {
 	if (Predicate(*this))
@@ -1330,7 +1393,6 @@ FStreamable* FStreamableManager::StreamInternal(const FSoftObjectPath& InTargetN
 		// If async loading isn't safe or it's forced on, we have to do a sync load which will flush all async loading
 		if (GIsInitialLoad || ThreadContext.IsInConstructor > 0 || bForceSynchronousLoads)
 		{
-			FRedirectedPath RedirectedPath;
 			UE_LOG(LogStreamableManager, Verbose, TEXT("     Static loading %s"), *TargetName.ToString());
 			Existing->Target = StaticLoadObject(UObject::StaticClass(), nullptr, *TargetName.ToString());
 
@@ -1349,6 +1411,8 @@ FStreamable* FStreamableManager::StreamInternal(const FSoftObjectPath& InTargetN
 			{
 				Existing->bLoadFailed = true;
 				UE_LOG(LogStreamableManager, Log, TEXT("Failed attempt to load %s"), *TargetName.ToString());
+
+				Handle->SetError(FStreamableError{ EStreamableErrorType::PackageLoadError, TEXT("Static load failed") });
 			}
 			Existing->bAsyncLoadRequestOutstanding = false;
 		}
@@ -1368,6 +1432,8 @@ FStreamable* FStreamableManager::StreamInternal(const FSoftObjectPath& InTargetN
 				UE_LOG(LogStreamableManager, Error, TEXT("Failed attempt to load %s; it is not a valid LongPackageName"), *Package);
 				Existing->bLoadFailed = true;
 				Existing->bAsyncLoadRequestOutstanding = false;
+
+				Handle->SetError(FStreamableError{ EStreamableErrorType::PackageLoadError, TEXT("Unkown package")});
 			}
 			else
 			{
@@ -1403,12 +1469,12 @@ bool FStreamableManager::ShouldStripDebugName()
 	return bShouldStripDebugName;
 }
 
-TSharedPtr<FStreamableHandle> FStreamableManager::RequestAsyncLoadInternal(TArray<FSoftObjectPath>&& TargetsToStream, FStreamableDelegate&& DelegateToCall, TAsyncLoadPriority Priority, bool bManageActiveHandle, bool bStartStalled, FString&& DebugName)
+TSharedPtr<FStreamableHandle> FStreamableManager::RequestAsyncLoadInternal(TArray<FSoftObjectPath>&& TargetsToStream, FStreamableDelegateWithHandle&& DelegateToCall, TAsyncLoadPriority Priority, bool bManageActiveHandle, bool bStartStalled, FString&& DebugName)
 {
 	LLM_SCOPE(ELLMTag::StreamingManager);
 
 	// Schedule a new callback, this will get called when all related async loads are completed
-	TSharedRef<FStreamableHandle> NewRequest = MakeShareable(new FStreamableHandle());
+	TSharedRef<FStreamableHandle> NewRequest = MakeShared<FStreamableHandle>(FStreamableHandle::FPrivateToken());
 	NewRequest->CompleteDelegate = MoveTemp(DelegateToCall);
 	NewRequest->OwningManager = this;
 	NewRequest->RequestedAssets = MoveTemp(TargetsToStream);
@@ -1525,7 +1591,7 @@ TSharedPtr<FStreamableHandle> FStreamableManager::RequestSyncLoadInternal(TArray
 	bForceSynchronousLoads = IsInAsyncLoadingThread() || IsEventDrivenLoaderEnabled() || !IsAsyncLoading();
 
 	// Do an async load and wait to complete. In some cases this will do a sync load due to safety issues
-	TSharedPtr<FStreamableHandle> Request = RequestAsyncLoadInternal(MoveTemp(TargetsToStream), FStreamableDelegate(), AsyncLoadHighPriority, bManageActiveHandle, false, MoveTemp(DebugName));
+	TSharedPtr<FStreamableHandle> Request = RequestAsyncLoadInternal(MoveTemp(TargetsToStream), FStreamableDelegateWithHandle(), AsyncLoadHighPriority, bManageActiveHandle, false, MoveTemp(DebugName));
 
 	bForceSynchronousLoads = false;
 
@@ -1533,8 +1599,18 @@ TSharedPtr<FStreamableHandle> FStreamableManager::RequestSyncLoadInternal(TArray
 	{
 		EAsyncPackageState::Type Result = Request->WaitUntilComplete();
 
-		ensureMsgf(Result == EAsyncPackageState::Complete, TEXT("RequestSyncLoad of %s resulted in bad async load result %d"), *Request->DebugName, Result);
-		ensureMsgf(Request->HasLoadCompleted(), TEXT("RequestSyncLoad of %s completed early, not actually completed!"), *Request->DebugName);
+		const bool bBadPackageState = !ensureMsgf(
+			Result == EAsyncPackageState::Complete,
+			TEXT("RequestSyncLoad of %s resulted in bad async load result %d"), *Request->DebugName, Result);
+
+		const bool bLoadNotComplete = !ensureMsgf(
+			Request->HasLoadCompleted(),
+			TEXT("RequestSyncLoad of %s completed early, not actually completed!"), *Request->DebugName);
+
+		if (bBadPackageState || bLoadNotComplete)
+		{
+			Request->SetError(FStreamableError{ EStreamableErrorType::PackageLoadError, TEXT("Sync load incomplete") });
+		}
 	}
 
 	return Request;
@@ -1677,7 +1753,7 @@ void FStreamableManager::FindInMemory(FSoftObjectPath& InOutTargetName, struct F
 	}
 }
 
-void FStreamableManager::AsyncLoadCallback(FSoftObjectPath TargetName, UPackage* Package)
+void FStreamableManager::AsyncLoadCallback(FSoftObjectPath TargetName, UPackage* Package, EAsyncLoadingResult::Type Result)
 {
 	check(IsInGameThread());
 
@@ -1700,6 +1776,7 @@ void FStreamableManager::AsyncLoadCallback(FSoftObjectPath TargetName, UPackage*
 		{
 			UE_LOG(LogStreamableManager, Verbose, TEXT("AsyncLoadCallback called for %s when not waiting on a load request, was loaded early by sync load"), *TargetName.ToString());
 		}
+
 		if (Existing->Target)
 		{
 			UE_LOG(LogStreamableManager, Verbose, TEXT("    Found target %s"), *Existing->Target->GetFullName());
@@ -1708,6 +1785,7 @@ void FStreamableManager::AsyncLoadCallback(FSoftObjectPath TargetName, UPackage*
 		{
 			// Async load failed to find the object
 			Existing->bLoadFailed = true;
+			ensure(Result != EAsyncLoadingResult::Succeeded);
 			UE_LOG(LogStreamableManager, Verbose, TEXT("    Failed async load for %s"), *TargetName.ToString());
 		}
 	}
@@ -1846,14 +1924,14 @@ void FStreamableManager::Unload(const FSoftObjectPath& Target)
 	}
 }
 
-TSharedPtr<FStreamableHandle> FStreamableManager::CreateCombinedHandle(const TConstArrayView<TSharedPtr<FStreamableHandle> >& ChildHandles, const FString& DebugName, EStreamableManagerCombinedHandleOptions Options)
+TSharedPtr<FStreamableHandle> FStreamableManager::CreateCombinedHandle(const TConstArrayView<TSharedPtr<FStreamableHandle>>& ChildHandles, const FString& DebugName, EStreamableManagerCombinedHandleOptions Options)
 {
 	if (!ensure(ChildHandles.Num() > 0))
 	{
 		return nullptr;
 	}
 
-	TSharedRef<FStreamableHandle> NewRequest = MakeShareable(new FStreamableHandle());
+	TSharedRef<FStreamableHandle> NewRequest = MakeShared<FStreamableHandle>(FStreamableHandle::FPrivateToken());
 	NewRequest->OwningManager = this;
 	NewRequest->bIsCombinedHandle = true;
 	NewRequest->DebugName = DebugName;
@@ -1949,6 +2027,12 @@ TSharedPtr<FStreamableHandle> FStreamableManager::CreateCombinedHandle(const TCo
 		{
 			++NewRequest->CanceledChildCount;
 		}
+
+		// Propogate any errors to new parent(s)
+		if (ChildHandle->Error)
+		{
+			ChildHandle->SetErrorOnParents(*ChildHandle->Error);
+		}
 	}
 
 	if (EnumHasAllFlags(Options, EStreamableManagerCombinedHandleOptions::MergeDebugNames))
