Texture Collection fixes
- Fixing Resource Collection shader parameter validation.
- Adding UTextureCollection support to Material streaming logic.

diff --git a/Engine/Source/Runtime/Engine/Private/Materials/Material.cpp b/Engine/Source/Runtime/Engine/Private/Materials/Material.cpp
--- a/Engine/Source/Runtime/Engine/Private/Materials/Material.cpp
+++ b/Engine/Source/Runtime/Engine/Private/Materials/Material.cpp
@@ -53,6 +53,7 @@
 #include "Materials/MaterialUniformExpressions.h"
 #include "Engine/SubsurfaceProfile.h"
 #include "Engine/SpecularProfile.h"
+#include "Engine/TextureCollection.h"
 #include "EditorSupportDelegates.h"
 #include "ComponentRecreateRenderStateContext.h"
 #include "ShaderCompiler.h"
@@ -1223,6 +1224,19 @@ void UMaterial::GetUsedTextures(TArray<UTexture*>& OutTextures, EMaterialQuality
 				}
 			}
 
+			for (const FMaterialTextureCollectionParameterInfo& Parameter : CurrentResource->GetUniformTextureCollectionExpressions())
+			{
+				UTextureCollection* TextureCollection = nullptr;
+				Parameter.GetGameThreadTextureCollectionValue(this, *CurrentResource, TextureCollection);
+				if (TextureCollection)
+				{
+					for (UTexture* Texture : TextureCollection->Textures)
+					{
+						OutTextures.AddUnique(Texture);
+					}
+				}
+			}
+
 #if WITH_EDITOR
 /*
 			// Also look for any scalar parameters that are acting as lookups for an atlas texture, and store the atlas texture
diff --git a/Engine/Source/Runtime/Engine/Private/Materials/MaterialInstance.cpp b/Engine/Source/Runtime/Engine/Private/Materials/MaterialInstance.cpp
--- a/Engine/Source/Runtime/Engine/Private/Materials/MaterialInstance.cpp
+++ b/Engine/Source/Runtime/Engine/Private/Materials/MaterialInstance.cpp
@@ -1248,7 +1248,7 @@ int32 UMaterialInstance::GetBlendablePriority(const UMaterial* Base) const
 	}
 }
 
-void UMaterialInstance::GetTextureExpressionValues(const FMaterialResource* MaterialResource, TArray<UTexture*>& OutTextures, TArray< TArray<int32> >* OutIndices) const
+void UMaterialInstance::GetTextureExpressionValues(const FMaterialResource* MaterialResource, TArray<UTexture*>& OutTextures, TArray<TArray<int32>>* OutIndices, bool bIncludeTextureCollections) const
 {
 	check(MaterialResource);
 	const FUniformExpressionSet& UniformExpressions = MaterialResource->GetUniformExpressions();
@@ -1287,6 +1287,26 @@ void UMaterialInstance::GetTextureExpressionValues(const FMaterialResource* Mate
 			}
 		}
 	}
+
+	if (bIncludeTextureCollections)
+	{
+		checkf(OutIndices == nullptr, TEXT("Texture Collections don't work with Texture Indices."));
+
+		for (int32 TextureCollectionIndex = 0; TextureCollectionIndex < UniformExpressions.GetNumTextureCollections(); TextureCollectionIndex++)
+		{
+			UTextureCollection* TextureCollection = nullptr;
+			UniformExpressions.GetGameThreadTextureCollectionValue(TextureCollectionIndex, this, *MaterialResource, TextureCollection);
+
+			if (TextureCollection)
+			{
+				for (UTexture* Texture : TextureCollection->Textures)
+				{
+					OutTextures.AddUnique(Texture);
+				}
+			}
+		}
+	}
+
 #if WITH_EDITOR
 	for (int32 i = 0; i < TransientTextureParameterOverrides.Num(); ++i)
 	{
@@ -1296,6 +1316,23 @@ void UMaterialInstance::GetTextureExpressionValues(const FMaterialResource* Mate
 #endif
 }
 
+void UMaterialInstance::GetTextureCollectionExpressionValues(const FMaterialResource* MaterialResource, TArray<UTextureCollection*>& OutTextureCollections) const
+{
+	check(MaterialResource);
+	const FUniformExpressionSet& UniformExpressions = MaterialResource->GetUniformExpressions();
+
+	for (int32 TextureCollectionIndex = 0; TextureCollectionIndex < UniformExpressions.GetNumTextureCollections(); TextureCollectionIndex++)
+	{
+		UTextureCollection* TextureCollection = nullptr;
+		UniformExpressions.GetGameThreadTextureCollectionValue(TextureCollectionIndex, this, *MaterialResource, TextureCollection);
+
+		if (TextureCollection)
+		{
+			OutTextureCollections.AddUnique(TextureCollection);
+		}
+	}
+}
+
 void UMaterialInstance::GetUsedTextures(TArray<UTexture*>& OutTextures, EMaterialQualityLevel::Type QualityLevel, bool bAllQualityLevels, ERHIFeatureLevel::Type FeatureLevel, bool bAllFeatureLevels) const
 {
 	OutTextures.Empty();
@@ -1353,7 +1390,7 @@ void UMaterialInstance::GetUsedTextures(TArray<UTexture*>& OutTextures, EMateria
 
 			for (const FMaterialResource* MaterialResource : MatchedResources)
 			{
-				GetTextureExpressionValues(MaterialResource, OutTextures);
+				GetTextureExpressionValues(MaterialResource, OutTextures, nullptr, true);
 			}
 		}
 		else
diff --git a/Engine/Source/Runtime/Engine/Private/Materials/MaterialShared.cpp b/Engine/Source/Runtime/Engine/Private/Materials/MaterialShared.cpp
--- a/Engine/Source/Runtime/Engine/Private/Materials/MaterialShared.cpp
+++ b/Engine/Source/Runtime/Engine/Private/Materials/MaterialShared.cpp
@@ -961,6 +961,11 @@ TArrayView<const FMaterialTextureParameterInfo> FMaterial::GetUniformTextureExpr
 	return GetUniformExpressions().UniformTextureParameters[(uint32)Type];
 }
 
+TConstArrayView<FMaterialTextureCollectionParameterInfo> FMaterial::GetUniformTextureCollectionExpressions() const
+{
+	return GetUniformExpressions().UniformTextureCollectionParameters;
+}
+
 TArrayView<const FMaterialNumericParameterInfo> FMaterial::GetUniformNumericParameterExpressions() const
 { 
 	return GetUniformExpressions().UniformNumericParameters;
diff --git a/Engine/Source/Runtime/Engine/Private/Materials/MaterialUniformExpressions.cpp b/Engine/Source/Runtime/Engine/Private/Materials/MaterialUniformExpressions.cpp
--- a/Engine/Source/Runtime/Engine/Private/Materials/MaterialUniformExpressions.cpp
+++ b/Engine/Source/Runtime/Engine/Private/Materials/MaterialUniformExpressions.cpp
@@ -804,6 +804,14 @@ void FUniformExpressionSet::GetTextureValue(int32 Index, const FMaterialRenderCo
 	}
 }
 
+void FUniformExpressionSet::GetGameThreadTextureCollectionValue(int32 Index, const UMaterialInterface* MaterialInterface, const FMaterial& Material, UTextureCollection*& OutValue) const
+{
+	check(IsInGameThread() || IsInParallelGameThread());
+	OutValue = nullptr;
+	const FMaterialTextureCollectionParameterInfo& Parameter = GetTextureCollectionParameter(Index);
+	Parameter.GetGameThreadTextureCollectionValue(MaterialInterface, Material, OutValue);
+}
+
 void FUniformExpressionSet::FillUniformBuffer(const FMaterialRenderContext& MaterialRenderContext, const FUniformExpressionCache& UniformExpressionCache, const FRHIUniformBufferLayout* UniformBufferLayout, uint8* TempBuffer, int TempBufferSize) const
 {
 	FillUniformBuffer(MaterialRenderContext, UniformExpressionCache.AllocatedVTs, UniformBufferLayout, TempBuffer, TempBufferSize);
@@ -1838,16 +1846,24 @@ void FMaterialTextureParameterInfo::GetGameThreadTextureValue(const UMaterialInt
 	}
 }
 
+void FMaterialTextureParameterInfo::GetGameThreadTextureValue(const UMaterialInterface* MaterialInterface, const FMaterial& Material, USparseVolumeTexture*& OutValue) const
+{
+	if (ParameterInfo.Name.IsNone() || !MaterialInterface->GetSparseVolumeTextureParameterValue(ParameterInfo, OutValue))
+	{
+		OutValue = GetIndexedTexture<USparseVolumeTexture>(Material, TextureIndex);
+	}
+}
+
 void FMaterialTextureCollectionParameterInfo::GetTextureCollection(const FMaterialRenderContext& Context, const UTextureCollection*& OutTextureCollection) const
 {
 	Context.GetTextureCollectionParameterValue(ParameterInfo, TextureCollectionIndex, OutTextureCollection);
 }
 
-void FMaterialTextureParameterInfo::GetGameThreadTextureValue(const UMaterialInterface* MaterialInterface, const FMaterial& Material, USparseVolumeTexture*& OutValue) const
+void FMaterialTextureCollectionParameterInfo::GetGameThreadTextureCollectionValue(const UMaterialInterface* MaterialInterface, const FMaterial& Material, UTextureCollection*& OutValue) const
 {
-	if (ParameterInfo.Name.IsNone() || !MaterialInterface->GetSparseVolumeTextureParameterValue(ParameterInfo, OutValue))
+	if (ParameterInfo.Name.IsNone() || !MaterialInterface->GetTextureCollectionParameterValue(ParameterInfo, OutValue))
 	{
-		OutValue = GetIndexedTexture<USparseVolumeTexture>(Material, TextureIndex);
+		OutValue = GetIndexedTextureCollection(Material, TextureCollectionIndex);
 	}
 }
 
diff --git a/Engine/Source/Runtime/Engine/Public/MaterialShared.h b/Engine/Source/Runtime/Engine/Public/MaterialShared.h
--- a/Engine/Source/Runtime/Engine/Public/MaterialShared.h
+++ b/Engine/Source/Runtime/Engine/Public/MaterialShared.h
@@ -554,6 +554,7 @@ public:
 		return !operator==(Lhs, Rhs);
 	}
 
+	ENGINE_API void GetGameThreadTextureCollectionValue(const UMaterialInterface* MaterialInterface, const FMaterial& Material, UTextureCollection*& OutValue) const;
 	void GetTextureCollection(const FMaterialRenderContext& Context, const UTextureCollection*& OutTextureCollection) const;
 
 	LAYOUT_FIELD(FHashedMaterialParameterInfo, ParameterInfo);
@@ -679,6 +680,7 @@ public:
 
 	inline const FMaterialNumericParameterInfo& GetNumericParameter(uint32 Index) const { return UniformNumericParameters[Index]; }
 	inline const FMaterialTextureParameterInfo& GetTextureParameter(EMaterialTextureParameterType Type, int32 Index) const { return UniformTextureParameters[(uint32)Type][Index]; }
+	const FMaterialTextureCollectionParameterInfo& GetTextureCollectionParameter(int32 Index) const { return UniformTextureCollectionParameters[Index]; }
 
 	inline int32 GetNumTextures(EMaterialTextureParameterType Type) const { return UniformTextureParameters[(uint32)Type].Num(); }
 	ENGINE_API void GetGameThreadTextureValue(EMaterialTextureParameterType Type, int32 Index, const UMaterialInterface* MaterialInterface, const FMaterial& Material, UTexture*& OutValue, bool bAllowOverride = true) const;
@@ -686,6 +688,9 @@ public:
 	ENGINE_API void GetTextureValue(int32 Index, const FMaterialRenderContext& Context, const FMaterial& Material, const URuntimeVirtualTexture*& OutValue) const;
 	ENGINE_API void GetTextureValue(int32 Index, const FMaterialRenderContext& Context, const FMaterial& Material, const USparseVolumeTexture*& OutValue) const;
 
+	int32 GetNumTextureCollections() const { return UniformTextureCollectionParameters.Num(); }
+	ENGINE_API void GetGameThreadTextureCollectionValue(int32 Index, const UMaterialInterface* MaterialInterface, const FMaterial& Material, UTextureCollection*& OutValue) const;
+
 	int32 FindOrAddTextureParameter(EMaterialTextureParameterType Type, const FMaterialTextureParameterInfo& Info);
 	int32 FindOrAddTextureCollectionParameter(const FMaterialTextureCollectionParameterInfo& Info);
 	int32 FindOrAddExternalTextureParameter(const FMaterialExternalTextureParameterInfo& Info);
@@ -2365,6 +2370,7 @@ public:
 	// Accessors.
 	ENGINE_API const FUniformExpressionSet& GetUniformExpressions() const;
 	ENGINE_API TArrayView<const FMaterialTextureParameterInfo> GetUniformTextureExpressions(EMaterialTextureParameterType Type) const;
+	ENGINE_API TConstArrayView<FMaterialTextureCollectionParameterInfo> GetUniformTextureCollectionExpressions() const;
 	ENGINE_API TArrayView<const FMaterialNumericParameterInfo> GetUniformNumericParameterExpressions() const;
 
 	inline TArrayView<const FMaterialTextureParameterInfo> GetUniform2DTextureExpressions() const { return GetUniformTextureExpressions(EMaterialTextureParameterType::Standard2D); }
diff --git a/Engine/Source/Runtime/Engine/Public/Materials/MaterialInstance.h b/Engine/Source/Runtime/Engine/Public/Materials/MaterialInstance.h
--- a/Engine/Source/Runtime/Engine/Public/Materials/MaterialInstance.h
+++ b/Engine/Source/Runtime/Engine/Public/Materials/MaterialInstance.h
@@ -1117,7 +1117,8 @@ protected:
 
 	ENGINE_API bool SetParentInternal(class UMaterialInterface* NewParent, bool RecacheShaders);
 
-	void GetTextureExpressionValues(const FMaterialResource* MaterialResource, TArray<UTexture*>& OutTextures, TArray< TArray<int32> >* OutIndices = nullptr) const;
+	void GetTextureExpressionValues(const FMaterialResource* MaterialResource, TArray<UTexture*>& OutTextures, TArray<TArray<int32>>* OutIndices = nullptr, bool bIncludeTextureCollections = false) const;
+	void GetTextureCollectionExpressionValues(const FMaterialResource* MaterialResource, TArray<UTextureCollection*>& OutTextureCollections) const;
 
 	UE_DEPRECATED(4.26, "Calling UpdatePermutationAllocations is no longer necessary")
 	inline void UpdatePermutationAllocations(FMaterialResourceDeferredDeletionArray* ResourcesToFree = nullptr) {}
diff --git a/Engine/Source/Runtime/RenderCore/Private/ShaderParameterMetadata.cpp b/Engine/Source/Runtime/RenderCore/Private/ShaderParameterMetadata.cpp
--- a/Engine/Source/Runtime/RenderCore/Private/ShaderParameterMetadata.cpp
+++ b/Engine/Source/Runtime/RenderCore/Private/ShaderParameterMetadata.cpp
@@ -379,6 +379,10 @@ EShaderCodeResourceBindingType ParseShaderResourceBindingType(const TCHAR* Shade
 	{
 		BindingType = EShaderCodeResourceBindingType::RaytracingAccelerationStructure;
 	}
+	else if (!bIsRWResource && ShaderTypeLength == 19 && FCString::Strncmp(ComparedShaderType, TEXT("FResourceCollection"), ShaderTypeLength) == 0)
+	{
+		BindingType = EShaderCodeResourceBindingType::ResourceCollection;
+	}
 	else if (bIsRasterizerOrderedResource)
 	{
 		BindingType = EShaderCodeResourceBindingType::RasterizerOrderedTexture2D;
@@ -742,8 +746,7 @@ void FShaderParametersMetadata::InitializeLayout(FRHIUniformBufferLayoutInitiali
 				{
 					bIsValidBindingType = (
 						BaseType == UBMT_SRV ||
-						BaseType == UBMT_RDG_BUFFER_SRV ||
-						BaseType == UBMT_RESOURCE_COLLECTION);
+						BaseType == UBMT_RDG_BUFFER_SRV);
 				}
 				else if (BindingType == EShaderCodeResourceBindingType::RaytracingAccelerationStructure)
 				{
@@ -772,6 +775,10 @@ void FShaderParametersMetadata::InitializeLayout(FRHIUniformBufferLayoutInitiali
 						BaseType == UBMT_UAV ||
 						BaseType == UBMT_RDG_BUFFER_UAV);
 				}
+				else if (BindingType == EShaderCodeResourceBindingType::ResourceCollection)
+				{
+					bIsValidBindingType = (BaseType == UBMT_RESOURCE_COLLECTION);
+				}
 				else
 				{
 					//unimplemented();
