Fix crashes when skeletal meshes are saved with more bones than their skeleton

Addresses a number of locations where OOB access could occur if bone arrays were mismatched. This could happen if a skeleton was not resaved to include additional bones from a mesh, or via a merge conflict resolve.

diff --git a/Engine/Plugins/Animation/BlendStack/Source/Runtime/Private/AnimNode_BlendStack.cpp b/Engine/Plugins/Animation/BlendStack/Source/Runtime/Private/AnimNode_BlendStack.cpp
--- a/Engine/Plugins/Animation/BlendStack/Source/Runtime/Private/AnimNode_BlendStack.cpp
+++ b/Engine/Plugins/Animation/BlendStack/Source/Runtime/Private/AnimNode_BlendStack.cpp
@@ -519,8 +519,13 @@ void GetBlendInWeightsCompactPose(const FBlendStackAnimPlayer& AnimPlayer, const
 	const float WeightForZeroBlendInTime = AnimPlayer.GetTimeToActivation() > 0.f ? 0.f : 1.f;
 	for (int32 CompactPoseBoneIdx = 0; CompactPoseBoneIdx < CompactPoseNumBones; ++CompactPoseBoneIdx)
 	{
-		const int32 SkeletonBoneIdx = BoneContainer.GetSkeletonIndex(FCompactPoseBoneIndex(CompactPoseBoneIdx));
-		check(SkeletonBoneIdx != INDEX_NONE);
+		const FSkeletonPoseBoneIndex SkeletonBoneIdx = BoneContainer.GetSkeletonPoseIndexFromCompactPoseIndex(FCompactPoseBoneIndex(CompactPoseBoneIdx));
+		if(!SkeletonBoneIdx.IsValid())
+		{
+			Weights[CompactPoseBoneIdx] = 0.0f;
+			continue;
+		}
+
 		const float TotalBlendInTimeBoneIdx = TotalBlendInTimePerBone[SkeletonBoneIdx];
 		if (TotalBlendInTimeBoneIdx < UE_SMALL_NUMBER)
 		{
diff --git a/Engine/Plugins/Experimental/AnimNext/Source/AnimNext/Private/GenerationTools.cpp b/Engine/Plugins/Experimental/AnimNext/Source/AnimNext/Private/GenerationTools.cpp
--- a/Engine/Plugins/Experimental/AnimNext/Source/AnimNext/Private/GenerationTools.cpp
+++ b/Engine/Plugins/Experimental/AnimNext/Source/AnimNext/Private/GenerationTools.cpp
@@ -128,13 +128,14 @@ struct FCompareBoneIndexType
 				const FSkeletonToMeshLinkup& LinkupTable = const_cast<USkeleton*>(Skeleton)->FindOrAddMeshLinkupData(SkeletalMesh);
 
 				// Generate a Skeleton to LOD look up table
-				const int32 NumMeshBones = LinkupTable.SkeletonToMeshTable.Num();
+				const int32 NumSkelBones = Skeleton->GetReferenceSkeleton().GetNum();
+				const int32 NumMeshBones = SkeletalMesh->GetRefSkeleton().GetNum();
 				const int32 NumOrderedBones = LODBoneIndexToMeshBoneIndexMap.Num();
 
 				TArray<FBoneIndexType> SkeletonBoneIndexToLODBoneIndexMap;
 				TArray<FBoneIndexType> MeshBoneIndexToLODBoneIndexMap;
 				TArray<FBoneIndexType> LODBoneIndexToSkeletonBoneIndexMap;
-				SkeletonBoneIndexToLODBoneIndexMap.Init(INDEX_NONE, NumMeshBones);
+				SkeletonBoneIndexToLODBoneIndexMap.Init(INDEX_NONE, NumSkelBones);
 				MeshBoneIndexToLODBoneIndexMap.Init(INDEX_NONE, NumMeshBones);
 				LODBoneIndexToSkeletonBoneIndexMap.Init(INDEX_NONE, NumOrderedBones);
 
@@ -144,12 +145,26 @@ struct FCompareBoneIndexType
 					const FMeshPoseBoneIndex MeshBoneIndex(LODBoneIndexToMeshBoneIndexMap[LODBoneIndex]);
 
 					// Remap our skeletal mesh bone index into the skeleton bone index we output for
-					const FSkeletonPoseBoneIndex SkeletonBoneIndex(LinkupTable.MeshToSkeletonTable[MeshBoneIndex.GetInt()]);
-					ensure(SkeletonBoneIndex.IsValid());	// We expect the skeletal mesh bone to map to a valid skeleton bone
+					FSkeletonPoseBoneIndex SkeletonBoneIndex(INDEX_NONE);
+					if(LinkupTable.MeshToSkeletonTable.IsValidIndex(MeshBoneIndex.GetInt()))
+					{
+						SkeletonBoneIndex = FSkeletonPoseBoneIndex(LinkupTable.MeshToSkeletonTable[MeshBoneIndex.GetInt()]);
+					}
 
-					SkeletonBoneIndexToLODBoneIndexMap[static_cast<FBoneIndexType>(SkeletonBoneIndex.GetInt())] = LODBoneIndex;
-					MeshBoneIndexToLODBoneIndexMap[static_cast<FBoneIndexType>(MeshBoneIndex.GetInt())] = LODBoneIndex;
-					LODBoneIndexToSkeletonBoneIndexMap[LODBoneIndex] = static_cast<FBoneIndexType>(SkeletonBoneIndex.GetInt());
+					if(SkeletonBoneIndexToLODBoneIndexMap.IsValidIndex(SkeletonBoneIndex.GetInt()))
+					{
+						SkeletonBoneIndexToLODBoneIndexMap[static_cast<FBoneIndexType>(SkeletonBoneIndex.GetInt())] = LODBoneIndex;
+						LODBoneIndexToSkeletonBoneIndexMap[LODBoneIndex] = static_cast<FBoneIndexType>(SkeletonBoneIndex.GetInt());
+					}
+					else
+					{
+						LODBoneIndexToSkeletonBoneIndexMap[LODBoneIndex] = INDEX_NONE;
+					}
+
+					if(MeshBoneIndexToLODBoneIndexMap.IsValidIndex(MeshBoneIndex.GetInt()))
+					{
+						MeshBoneIndexToLODBoneIndexMap[static_cast<FBoneIndexType>(MeshBoneIndex.GetInt())] = LODBoneIndex;
+					}
 				}
 
 				TArray<FBoneIndexType> MeshBoneIndexToLODBoneIndex;
diff --git a/Engine/Source/Runtime/Engine/Private/Animation/AnimationDecompression.cpp b/Engine/Source/Runtime/Engine/Private/Animation/AnimationDecompression.cpp
--- a/Engine/Source/Runtime/Engine/Private/Animation/AnimationDecompression.cpp
+++ b/Engine/Source/Runtime/Engine/Private/Animation/AnimationDecompression.cpp
@@ -210,27 +210,31 @@ void DecompressPose(FCompactPose& OutPose,
 		{
 			for (FCompactPoseBoneIndex BoneIndex(bFirstTrackIsRootBone ? 1 : 0); BoneIndex < NumCompactBones; ++BoneIndex)
 			{
-				const int32 TargetSkeletonBoneIndex = RequiredBones.GetSkeletonIndex(BoneIndex);
+				const FSkeletonPoseBoneIndex TargetSkeletonBoneIndex = RequiredBones.GetSkeletonPoseIndexFromCompactPoseIndex(BoneIndex);
+				if(!TargetSkeletonBoneIndex.IsValid())
+				{
+					continue;
+				}
 
 				// Mesh space additives do not require fix-up
 				if (DecompressionContext.GetAdditiveType() == AAT_LocalSpaceBase)
 				{
-					OutPose[BoneIndex].SetRotation(SkeletonRemapping.RetargetAdditiveRotationToTargetSkeleton(TargetSkeletonBoneIndex, OutPose[BoneIndex].GetRotation()));
+					OutPose[BoneIndex].SetRotation(SkeletonRemapping.RetargetAdditiveRotationToTargetSkeleton(TargetSkeletonBoneIndex.GetInt(), OutPose[BoneIndex].GetRotation()));
 				}
 			
 				// Check what retarget mode to use for the translational retargeting for this specific bone.
-				const int32 SourceSkeletonBoneIndex = SkeletonRemapping.GetSourceSkeletonBoneIndex(TargetSkeletonBoneIndex);
+				const int32 SourceSkeletonBoneIndex = SkeletonRemapping.GetSourceSkeletonBoneIndex(TargetSkeletonBoneIndex.GetInt());
 				const EBoneTranslationRetargetingMode::Type RetargetMode = FAnimationRuntime::GetBoneTranslationRetargetingMode(
 					bUseSourceRetargetModes,
 					SourceSkeletonBoneIndex,
-					TargetSkeletonBoneIndex,
+					TargetSkeletonBoneIndex.GetInt(),
 					SourceSkeleton,
 					TargetSkeleton,
 					bDisableRetargeting);
 
 				if (RetargetMode != EBoneTranslationRetargetingMode::Skeleton)
 				{
-					OutPose[BoneIndex].SetTranslation(SkeletonRemapping.RetargetAdditiveTranslationToTargetSkeleton(TargetSkeletonBoneIndex, OutPose[BoneIndex].GetTranslation()));
+					OutPose[BoneIndex].SetTranslation(SkeletonRemapping.RetargetAdditiveTranslationToTargetSkeleton(TargetSkeletonBoneIndex.GetInt(), OutPose[BoneIndex].GetTranslation()));
 				}
 			}
 		}
@@ -238,22 +242,27 @@ void DecompressPose(FCompactPose& OutPose,
 		{
 			for (FCompactPoseBoneIndex BoneIndex(bFirstTrackIsRootBone ? 1 : 0); BoneIndex < NumCompactBones; ++BoneIndex)
 			{
-				const int32 TargetSkeletonBoneIndex = RequiredBones.GetSkeletonIndex(BoneIndex);
-				OutPose[BoneIndex].SetRotation(SkeletonRemapping.RetargetBoneRotationToTargetSkeleton(TargetSkeletonBoneIndex, OutPose[BoneIndex].GetRotation()));
+				const FSkeletonPoseBoneIndex TargetSkeletonBoneIndex = RequiredBones.GetSkeletonPoseIndexFromCompactPoseIndex(BoneIndex);
+				if(!TargetSkeletonBoneIndex.IsValid())
+				{
+					continue;
+				}
+
+				OutPose[BoneIndex].SetRotation(SkeletonRemapping.RetargetBoneRotationToTargetSkeleton(TargetSkeletonBoneIndex.GetInt(), OutPose[BoneIndex].GetRotation()));
 
 				// Check what retarget mode to use for the translational retargeting for this specific bone.
-				const int32 SourceSkeletonBoneIndex = SkeletonRemapping.GetSourceSkeletonBoneIndex(TargetSkeletonBoneIndex);
+				const int32 SourceSkeletonBoneIndex = SkeletonRemapping.GetSourceSkeletonBoneIndex(TargetSkeletonBoneIndex.GetInt());
 				const EBoneTranslationRetargetingMode::Type RetargetMode = FAnimationRuntime::GetBoneTranslationRetargetingMode(
 					bUseSourceRetargetModes,
 					SourceSkeletonBoneIndex,
-					TargetSkeletonBoneIndex,
+					TargetSkeletonBoneIndex.GetInt(),
 					SourceSkeleton,
 					TargetSkeleton,
 					bDisableRetargeting);
 
 				if (RetargetMode != EBoneTranslationRetargetingMode::Skeleton)
 				{
-					OutPose[BoneIndex].SetTranslation(SkeletonRemapping.RetargetBoneTranslationToTargetSkeleton(TargetSkeletonBoneIndex, OutPose[BoneIndex].GetTranslation()));
+					OutPose[BoneIndex].SetTranslation(SkeletonRemapping.RetargetBoneTranslationToTargetSkeleton(TargetSkeletonBoneIndex.GetInt(), OutPose[BoneIndex].GetTranslation()));
 				}
 			}
 		}
