Implement Distance, Desaturation and ChannelMask

- All three implementations required no new Operators in the IR and could be built w/ existing operators.
- Clear the compiler error arrays on the FMaterial since we call translate multiple times in the Material Editor resulting in duplicate errors being reported.  This logic is at parity with the existing HLSL Translator.
- Fixed assert for ternary operators.  Ternary operators go through the binary operator code path (to lower the second argument).  We assert if we aren't a binary or ternary operator at this point.
- Improve the error reporting in `ValidateOperatorAndGetResultType` to output the operator name and potentially the argument types when the types don't match.
- When lowering types we were emitting `float` instead of `MaterialFloat`, which can be `half` on certain platforms.
- When doing casts we were emitting `float` instead of `MaterialFloat`.
- Simplify UMaterialExpressionLength::Build

#rnx

diff --git a/Engine/Source/Runtime/Engine/Private/Materials/MaterialExpressionsIR.cpp b/Engine/Source/Runtime/Engine/Private/Materials/MaterialExpressionsIR.cpp
--- a/Engine/Source/Runtime/Engine/Private/Materials/MaterialExpressionsIR.cpp
+++ b/Engine/Source/Runtime/Engine/Private/Materials/MaterialExpressionsIR.cpp
@@ -290,7 +290,7 @@ void UMaterialExpressionStaticBool::Build(MIR::FEmitter& Emitter)
 	Emitter.Output(0, Emitter.EmitConstantBoolScalar(Value));
 }
 
-static MIR::FValue* BuildMaterialExpressionParameter(UMaterialExpressionParameter* ParameterExpr, UE::MIR::FEmitter& Emitter)
+static MIR::FValue* BuildMaterialExpressionParameter(UMaterialExpressionParameter* ParameterExpr, MIR::FEmitter& Emitter)
 {
 	FMaterialParameterMetadata Metadata;
 	if (!ParameterExpr->GetParameterValue(Metadata))
@@ -302,19 +302,41 @@ static MIR::FValue* BuildMaterialExpressionParameter(UMaterialExpressionParamete
 	return Emitter.EmitParameter(ParameterExpr->GetParameterName(), Metadata);
 }
 
-void UMaterialExpressionParameter::Build(UE::MIR::FEmitter& Emitter)
+void UMaterialExpressionParameter::Build(MIR::FEmitter& Emitter)
 {
 	Emitter.Output(0, BuildMaterialExpressionParameter(this, Emitter));
 }
 
-void UMaterialExpressionVectorParameter::Build(UE::MIR::FEmitter& Emitter)
+void UMaterialExpressionVectorParameter::Build(MIR::FEmitter& Emitter)
 {
 	MIR::FValue* Value = BuildMaterialExpressionParameter(this, Emitter);
 	Emitter.Output(0, Value);
 	Emitter.Output(1, Emitter.EmitSubscript(Value, 0));
-	Emitter.Output(2, Emitter.EmitSubscript(Value, 0));
-	Emitter.Output(3, Emitter.EmitSubscript(Value, 0));
-	Emitter.Output(4, Emitter.EmitSubscript(Value, 0));
+	Emitter.Output(2, Emitter.EmitSubscript(Value, 1));
+	Emitter.Output(3, Emitter.EmitSubscript(Value, 2));
+	Emitter.Output(4, Emitter.EmitSubscript(Value, 3));
+}
+
+void UMaterialExpressionChannelMaskParameter::Build(MIR::FEmitter& Emitter)
+{
+	MIR::FValue* InputValue = Emitter.Input(&Input);
+
+	UE_MIR_CHECKPOINT();
+
+	InputValue = Emitter.EmitConstruct(MIR::FPrimitiveType::GetVector(MIR::ScalarKind_Float, 4), InputValue);
+	MIR::FValue* ParameterValue = BuildMaterialExpressionParameter(this, Emitter);
+
+	UE_MIR_CHECKPOINT();
+
+	MIR::FValue* DotResult = Emitter.EmitOperator(MIR::BO_Dot, InputValue, ParameterValue);
+
+	UE_MIR_CHECKPOINT();
+
+	Emitter.Output(0, DotResult);
+	Emitter.Output(1, Emitter.EmitSubscript(DotResult, 0));
+	Emitter.Output(2, Emitter.EmitSubscript(DotResult, 1));
+	Emitter.Output(3, Emitter.EmitSubscript(DotResult, 2));
+	Emitter.Output(4, Emitter.EmitSubscript(DotResult, 3));
 }
 
 void UMaterialExpressionStaticBoolParameter::Build(MIR::FEmitter& Emitter)
@@ -327,7 +349,7 @@ void UMaterialExpressionStaticBoolParameter::Build(MIR::FEmitter& Emitter)
 	Emitter.Output(0, Value);
 }
 
-void UMaterialExpressionStaticSwitch::Build(UE::MIR::FEmitter& Emitter)
+void UMaterialExpressionStaticSwitch::Build(MIR::FEmitter& Emitter)
 {
 	Emitter.DefaultInputToBoolScalar(&Value, (bool)DefaultValue);
 	bool bCondition = Emitter.InputConstantBoolScalar(&Value);
@@ -337,7 +359,7 @@ void UMaterialExpressionStaticSwitch::Build(UE::MIR::FEmitter& Emitter)
 	Emitter.Output(0, Emitter.Input(bCondition ? &A : &B));
 }
 
-void UMaterialExpressionStaticSwitchParameter::Build(UE::MIR::FEmitter& Emitter)
+void UMaterialExpressionStaticSwitchParameter::Build(MIR::FEmitter& Emitter)
 {
 	MIR::FValue* Value = BuildMaterialExpressionParameter(this, Emitter);
 
@@ -492,14 +514,8 @@ void UMaterialExpressionLength::Build(UE::MIR::FEmitter& Emitter)
 	MIR::FValue* Value = Emitter.Input(&Input);
 	UE_MIR_CHECKPOINT();
 
-	if (Value->Type->AsPrimitive()->IsScalar())
-	{
-		BuildUnaryArithmeticOperatorFromValue(Emitter, MIR::UO_Abs, Value);
-	}
-	else
-	{
-		BuildUnaryArithmeticOperatorFromValue(Emitter, MIR::UO_Length, Value);
-	}
+	const MIR::EOperator Op = Value->Type->AsPrimitive()->IsScalar() ? MIR::UO_Abs : MIR::UO_Length;
+	BuildUnaryArithmeticOperatorFromValue(Emitter, Op, Value);
 }
 
 void UMaterialExpressionRound::Build(UE::MIR::FEmitter& Emitter)
@@ -628,6 +644,51 @@ void UMaterialExpressionAdd::Build(MIR::FEmitter& Emitter)
 	BuildBinaryOperatorFromInputWithDefault(Emitter, MIR::BO_Add, &A, ConstA, &B, ConstB);
 }
 
+void UMaterialExpressionDesaturation::Build(UE::MIR::FEmitter& Emitter)
+{
+	MIR::FValue* ColorValue = Emitter.Input(&Input);
+	UE_MIR_CHECKPOINT();
+
+	ColorValue = Emitter.EmitConstruct(MIR::FPrimitiveType::GetVector(MIR::ScalarKind_Float, 3), ColorValue);
+	UE_MIR_CHECKPOINT();
+
+	MIR::FValue* GreyOrLerpValue = Emitter.EmitOperator(MIR::BO_Dot, ColorValue, Emitter.EmitConstantFloat3(FVector3f(LuminanceFactors)));
+	UE_MIR_CHECKPOINT();
+
+	MIR::FValue* FractionValue = Emitter.TryInput(&Fraction);
+	if (FractionValue)
+	{
+		// Cast the Grey and Fraction value to the color type value so the lerp op works below.
+		GreyOrLerpValue = Emitter.EmitConstruct(ColorValue->Type, GreyOrLerpValue);
+		FractionValue = Emitter.EmitConstruct(ColorValue->Type, FractionValue);
+
+		UE_MIR_CHECKPOINT();
+
+		GreyOrLerpValue = Emitter.EmitOperator(MIR::TO_Lerp, ColorValue, GreyOrLerpValue, FractionValue);
+	}
+
+	UE_MIR_CHECKPOINT();
+
+	Emitter.Output(0, GreyOrLerpValue);
+}
+
+void UMaterialExpressionDistance::Build(MIR::FEmitter& Emitter)
+{
+	MIR::FValue* LhsValue = Emitter.Input(&A);
+	MIR::FValue* RhsValue = Emitter.Input(&B);
+	UE_MIR_CHECKPOINT();
+
+	MIR::FValue* DeltaValue = Emitter.EmitOperator(MIR::BO_Subtract, LhsValue, RhsValue);
+	UE_MIR_CHECKPOINT();
+
+	const bool bScalarInputs = LhsValue->Type->AsPrimitive()->IsScalar() && RhsValue->Type->AsPrimitive()->IsScalar();
+	const MIR::EOperator Op = bScalarInputs ? MIR::UO_Abs : MIR::UO_Length;
+	MIR::FValue* ResultValue = Emitter.EmitOperator(Op, DeltaValue);
+	UE_MIR_CHECKPOINT();
+
+	Emitter.Output(0, ResultValue);
+}
+
 void UMaterialExpressionFmod::Build(MIR::FEmitter& Emitter)
 { 
 	BuildBinaryOperatorFromInput(Emitter, MIR::BO_Fmod, &A, &B);
diff --git a/Engine/Source/Runtime/Engine/Private/Materials/MaterialIREmitter.cpp b/Engine/Source/Runtime/Engine/Private/Materials/MaterialIREmitter.cpp
--- a/Engine/Source/Runtime/Engine/Private/Materials/MaterialIREmitter.cpp
+++ b/Engine/Source/Runtime/Engine/Private/Materials/MaterialIREmitter.cpp
@@ -1148,7 +1148,7 @@ static FPrimitiveTypePtr ValidateOperatorAndGetResultType(FEmitter* Emitter, EOp
 		// All operators have primitive parameter types.
 		if (!ArgumentPrimitiveType)
 		{
-			Emitter->Errorf(TEXT("Operator %s argument does not have primitive type."), ArgumentsStr[i]);
+			Emitter->Errorf(TEXT("Operator %s's %s argument does not have primitive type."), OperatorToString(Op), ArgumentsStr[i]);
 			return nullptr;
 		}
 
@@ -1159,7 +1159,7 @@ static FPrimitiveTypePtr ValidateOperatorAndGetResultType(FEmitter* Emitter, EOp
 		{
 			if (ArgumentPrimitiveType->ScalarKind != ScalarKind_Bool)
 			{
-				Emitter->Errorf(TEXT("Operator %s argument component types are not boolean."), ArgumentsStr[i]);
+				Emitter->Errorf(TEXT("Operator %s's %s argument component types are not boolean."), OperatorToString(Op), ArgumentsStr[i]);
 				bValid = false;
 			}
 		}
@@ -1168,7 +1168,7 @@ static FPrimitiveTypePtr ValidateOperatorAndGetResultType(FEmitter* Emitter, EOp
 		{
 			if (!ScalarKindIsArithmetic(ArgumentPrimitiveType->ScalarKind))
 			{
-				Emitter->Errorf(TEXT("Operator %s argument component types expected to be arithmetic (int, float, double, etc)."), ArgumentsStr[i]);
+				Emitter->Errorf(TEXT("Operator %s's %s argument component types expected to be arithmetic (int, float, double, etc)."), OperatorToString(Op), ArgumentsStr[i]);
 				bValid = false;
 			}
 		}
@@ -1177,7 +1177,7 @@ static FPrimitiveTypePtr ValidateOperatorAndGetResultType(FEmitter* Emitter, EOp
 		{
 			if (!ScalarKindIsAnyFloat(ArgumentPrimitiveType->ScalarKind))
 			{
-				Emitter->Errorf(TEXT("Operator %s argument component types expected to be floating point (float, double, etc)."), ArgumentsStr[i]);
+				Emitter->Errorf(TEXT("Operator %s's %s argument component types expected to be floating point (float, double, etc)."), OperatorToString(Op), ArgumentsStr[i]);
 				bValid = false;
 			}
 		}
@@ -1186,7 +1186,7 @@ static FPrimitiveTypePtr ValidateOperatorAndGetResultType(FEmitter* Emitter, EOp
 		{
 			if (ArgumentPrimitiveType != FirstArgumentPrimitiveType)
 			{
-				Emitter->Errorf(TEXT("Operator %s argument type expected to be the same as first argument."), ArgumentsStr[i]);
+				Emitter->Errorf(TEXT("Operator %s's %s argument type (%s) expected to be the same as first argument (%s)."), OperatorToString(Op), ArgumentsStr[i], ArgumentPrimitiveType->Spelling.GetData(), FirstArgumentPrimitiveType->Spelling.GetData());
 				bValid = false;
 			}
 		}
@@ -1195,7 +1195,7 @@ static FPrimitiveTypePtr ValidateOperatorAndGetResultType(FEmitter* Emitter, EOp
 		{
 			if (ArgumentPrimitiveType->IsMatrix())
 			{
-				Emitter->Errorf(TEXT("Operator %s argument type expected to be a scalar or vector."), ArgumentsStr[i]);
+				Emitter->Errorf(TEXT("Operator %s's %s argument type expected to be a scalar or vector."), OperatorToString(Op), ArgumentsStr[i]);
 				bValid = false;
 			}
 		}
@@ -1204,7 +1204,7 @@ static FPrimitiveTypePtr ValidateOperatorAndGetResultType(FEmitter* Emitter, EOp
 		{
 			if (!ArgumentPrimitiveType->IsVector() || ArgumentPrimitiveType->GetNumComponents() != 3)
 			{
-				Emitter->Errorf(TEXT("Operator %s argument type expected to be a 3D vector."), ArgumentsStr[i]);
+				Emitter->Errorf(TEXT("Operator %s's %s argument type expected to be a 3D vector."), OperatorToString(Op), ArgumentsStr[i]);
 				bValid = false;
 			}
 		}
diff --git a/Engine/Source/Runtime/Engine/Private/Materials/MaterialIRToHLSLTranslator.cpp b/Engine/Source/Runtime/Engine/Private/Materials/MaterialIRToHLSLTranslator.cpp
--- a/Engine/Source/Runtime/Engine/Private/Materials/MaterialIRToHLSLTranslator.cpp
+++ b/Engine/Source/Runtime/Engine/Private/Materials/MaterialIRToHLSLTranslator.cpp
@@ -516,7 +516,7 @@ struct FTranslator : FMaterialIRToHLSLTranslation
 			// Binary
 			if (Operator->BArg)
 			{
-				check(MIR::IsBinaryOperator(Operator->Op));
+				check(MIR::IsBinaryOperator(Operator->Op) || MIR::IsTernaryOperator(Operator->Op));
 				Printer << ", " << LowerValue(Operator->BArg);
 			}
 
@@ -765,7 +765,7 @@ struct FTranslator : FMaterialIRToHLSLTranslation
 			{
 				case MIR::ScalarKind_Bool:	Printer << TEXT("bool"); break;
 				case MIR::ScalarKind_Int: 	Printer << TEXT("int"); break;
-				case MIR::ScalarKind_Float:	Printer << TEXT("float"); break;
+				case MIR::ScalarKind_Float:	Printer << TEXT("MaterialFloat"); break;
 			}
 
 			if (ArithmeticType->NumRows > 1)
diff --git a/Engine/Source/Runtime/Engine/Private/Materials/MaterialIRTypes.cpp b/Engine/Source/Runtime/Engine/Private/Materials/MaterialIRTypes.cpp
--- a/Engine/Source/Runtime/Engine/Private/Materials/MaterialIRTypes.cpp
+++ b/Engine/Source/Runtime/Engine/Private/Materials/MaterialIRTypes.cpp
@@ -254,7 +254,7 @@ const TCHAR* ScalarKindToString(EScalarKind Kind)
 	{
 		case ScalarKind_Bool: return TEXT("bool"); break;
 		case ScalarKind_Int: return TEXT("int"); break;
-		case ScalarKind_Float: return TEXT("float"); break;
+		case ScalarKind_Float: return TEXT("MaterialFloat"); break;
 		default: UE_MIR_UNREACHABLE();
 	}
 }
diff --git a/Engine/Source/Runtime/Engine/Private/Materials/MaterialShared.cpp b/Engine/Source/Runtime/Engine/Private/Materials/MaterialShared.cpp
--- a/Engine/Source/Runtime/Engine/Private/Materials/MaterialShared.cpp
+++ b/Engine/Source/Runtime/Engine/Private/Materials/MaterialShared.cpp
@@ -3298,6 +3298,10 @@ bool FMaterial::Translate_New(const FMaterialShaderMapId& InShaderMapId,
 	FMaterialCompilationOutput& OutCompilationOutput,
 	TRefCountPtr<FSharedShaderCompilerEnvironment>& OutMaterialEnvironment)
 {
+	// Clear existing Material Errors.
+	CompileErrors.Empty();
+	ErrorExpressions.Empty();
+
 	FMaterialIRModule Module;
 
 	GetMaterialInterface()->MaterialInsight.Reset(new FMaterialInsights);
diff --git a/Engine/Source/Runtime/Engine/Public/Materials/MaterialExpressionChannelMaskParameter.h b/Engine/Source/Runtime/Engine/Public/Materials/MaterialExpressionChannelMaskParameter.h
--- a/Engine/Source/Runtime/Engine/Public/Materials/MaterialExpressionChannelMaskParameter.h
+++ b/Engine/Source/Runtime/Engine/Public/Materials/MaterialExpressionChannelMaskParameter.h
@@ -34,6 +34,7 @@ class UMaterialExpressionChannelMaskParameter : public UMaterialExpressionVector
 	virtual bool SetParameterValue(FName InParameterName, FLinearColor InValue, EMaterialExpressionSetParameterValueFlags Flags) override;
 	virtual void PostEditChangeProperty(FPropertyChangedEvent& PropertyChangedEvent) override;
 
+	virtual void Build(UE::MIR::FEmitter& Emitter) override;
 	virtual int32 Compile(class FMaterialCompiler* Compiler, int32 OutputIndex) override;
 	virtual void GetCaption(TArray<FString>& OutCaptions) const override;
 	
diff --git a/Engine/Source/Runtime/Engine/Public/Materials/MaterialExpressionDesaturation.h b/Engine/Source/Runtime/Engine/Public/Materials/MaterialExpressionDesaturation.h
--- a/Engine/Source/Runtime/Engine/Public/Materials/MaterialExpressionDesaturation.h
+++ b/Engine/Source/Runtime/Engine/Public/Materials/MaterialExpressionDesaturation.h
@@ -35,6 +35,7 @@ class UMaterialExpressionDesaturation : public UMaterialExpression
 
 	//~ Begin UMaterialExpression Interface
 #if WITH_EDITOR
+	virtual void Build(UE::MIR::FEmitter& Emitter) override;
 	virtual int32 Compile(class FMaterialCompiler* Compiler, int32 OutputIndex) override;
 	virtual void GetCaption(TArray<FString>& OutCaptions) const override
 	{
diff --git a/Engine/Source/Runtime/Engine/Public/Materials/MaterialExpressionDistance.h b/Engine/Source/Runtime/Engine/Public/Materials/MaterialExpressionDistance.h
--- a/Engine/Source/Runtime/Engine/Public/Materials/MaterialExpressionDistance.h
+++ b/Engine/Source/Runtime/Engine/Public/Materials/MaterialExpressionDistance.h
@@ -25,6 +25,7 @@ public:
 
 	//~ Begin UMaterialExpression Interface
 #if WITH_EDITOR
+	virtual void Build(UE::MIR::FEmitter& Emitter) override;
 	virtual int32 Compile(class FMaterialCompiler* Compiler, int32 OutputIndex) override;
 	virtual void GetCaption(TArray<FString>& OutCaptions) const override;
 #endif
