Fix crash AnimNext trying to spawn UI on worker thread.

#tests PIE

diff --git a/Engine/Source/Runtime/Engine/Private/Animation/Skeleton.cpp b/Engine/Source/Runtime/Engine/Private/Animation/Skeleton.cpp
--- a/Engine/Source/Runtime/Engine/Private/Animation/Skeleton.cpp
+++ b/Engine/Source/Runtime/Engine/Private/Animation/Skeleton.cpp
@@ -898,31 +898,42 @@ void USkeleton::BuildLinkupData(const USkinnedAsset* InSkinnedAsset, FSkeletonTo
 
 		if (SkeletonBoneIndex == INDEX_NONE)
 		{
+			bool bSpawnedErrorMessage = false;
+
 #if WITH_EDITOR
 			// If we're in editor, and skeleton is missing a bone, fix it.
 			// not currently supported in-game.
 
 			static FName NAME_LoadErrors("LoadErrors");
-			FMessageLog LoadErrors(NAME_LoadErrors);
 
-			TSharedRef<FTokenizedMessage> Message = LoadErrors.Info();
-			Message->AddToken(FTextToken::Create(LOCTEXT("SkeletonBuildLinkupMissingBones1", "The Skeleton ")));
-			Message->AddToken(FAssetNameToken::Create(GetPathName(), FText::FromString( GetNameSafe(this) ) ));
-			Message->AddToken(FTextToken::Create(LOCTEXT("SkeletonBuildLinkupMissingBones2", " is missing bones that SkeletalMesh ")));
-			Message->AddToken(FAssetNameToken::Create(InSkinnedAsset->GetPathName(), FText::FromString( GetNameSafe(InSkinnedAsset) )));
-			Message->AddToken(FTextToken::Create(LOCTEXT("SkeletonBuildLinkupMissingBones3", "  needs. They will be added now. Please save the Skeleton!")));
-			LoadErrors.Open();
+			if (IsInGameThread() || IsInSlateThread())
+			{
+				FMessageLog LoadErrors(NAME_LoadErrors);
+
+				TSharedRef<FTokenizedMessage> Message = LoadErrors.Info();
+				Message->AddToken(FTextToken::Create(LOCTEXT("SkeletonBuildLinkupMissingBones1", "The Skeleton ")));
+				Message->AddToken(FAssetNameToken::Create(GetPathName(), FText::FromString(GetNameSafe(this))));
+				Message->AddToken(FTextToken::Create(LOCTEXT("SkeletonBuildLinkupMissingBones2", " is missing bones that SkeletalMesh ")));
+				Message->AddToken(FAssetNameToken::Create(InSkinnedAsset->GetPathName(), FText::FromString(GetNameSafe(InSkinnedAsset))));
+				Message->AddToken(FTextToken::Create(LOCTEXT("SkeletonBuildLinkupMissingBones3", "  needs. They will be added now. Please save the Skeleton!")));
+				LoadErrors.Open();
+
+				bSpawnedErrorMessage = true;
+			}
 
 			// Re-add all SkelMesh bones to the Skeleton.
 			MergeAllBonesToBoneTree(InSkinnedAsset);
 
 			// Fix missing bone.
 			SkeletonBoneIndex = SkeletonRefSkel.FindBoneIndex(MeshBoneName);
-#else
-			// If we're not in editor, we still want to know which skeleton is missing a bone.
-			UE_LOG(LogAnimation, Error, TEXT("USkeleton::BuildLinkup: The Skeleton %s, is missing bones that SkeletalMesh %s needs. MeshBoneName %s"),
-				*GetNameSafe(this), *GetNameSafe(InSkinnedAsset), *MeshBoneName.ToString());
-#endif
+#endif // WITH_EDITOR
+
+			if (!bSpawnedErrorMessage)
+			{
+				// If we're not in editor, we still want to know which skeleton is missing a bone.
+				UE_LOG(LogAnimation, Error, TEXT("USkeleton::BuildLinkup: The Skeleton %s, is missing bones that SkeletalMesh %s needs. MeshBoneName %s"),
+					*GetNameSafe(this), *GetNameSafe(InSkinnedAsset), *MeshBoneName.ToString());
+			}
 		}
 
 		NewMeshLinkup.MeshToSkeletonTable[MeshBoneIndex] = SkeletonBoneIndex;
