World: Add support for disabling Time Dilation in GetEffectiveTimeDilation() without modifying any of the time dilation values; When disabled, time dilation can still be set via existing functions, but calling GetEffectiveTimeDilation() will return 1.0, indicating no time dilation. Once set back to true, GetEffectiveTimeDilation() will return the time dilation as last set (even if it was set while time dilation was disabled).

diff --git a/Engine/Source/Runtime/Engine/Classes/GameFramework/WorldSettings.h b/Engine/Source/Runtime/Engine/Classes/GameFramework/WorldSettings.h
--- a/Engine/Source/Runtime/Engine/Classes/GameFramework/WorldSettings.h
+++ b/Engine/Source/Runtime/Engine/Classes/GameFramework/WorldSettings.h
@@ -455,6 +455,18 @@ public:
 	UPROPERTY(EditAnywhere, config, Category = HLODSystem, AdvancedDisplay, meta=(EditConditionHides, EditCondition="WorldPartition == nullptr", DisplayAfter="HierarchicalLODSetup"))
 	uint8 bGenerateSingleClusterForLevel : 1;
 
+	/**
+	 * Allows game logic to temporarily disable all forms of time dilation (when using GetEffectiveTimeDilation().
+	 * When false, logic is still allowed to set time dilation values, but their values will be ignored when calling
+	 * GetEffectiveTimeDilation() until this is set back to true.
+	 *
+	 * Transient because it's used for temporary overrides and shouldn't be saved/loaded from level packages. Not
+	 * replicated because in multiplayer a server should modify the real time dilation, instead of ignoring it on
+	 * a client (which will lead to position desyncs due to delta time mismatches).
+	*/
+	UPROPERTY(Transient)
+	uint8 bAllowTimeDilation : 1;
+
 #if WITH_EDITORONLY_DATA
 	/** if set to true, this hide the streaming disabled warning available in the viewport */
 	UPROPERTY(EditAnywhere, Category = WorldPartitionSetup, AdvancedDisplay, meta = (EditCondition = "WorldPartition != nullptr"))
@@ -806,9 +818,18 @@ public:
 	 */
 	ENGINE_API virtual float GetGravityZ() const;
 
+	/**
+	 * Multiplies all time dilation sources together and returns the effective overall time dilation. When bAllowTimeDilation is
+	 * set to false, this will return 1.0 regardless of the time dilation sources. See SetAllowTimeDilation(...)/bAllowTimeDilation.
+	 */
 	virtual float GetEffectiveTimeDilation() const
 	{
-		return TimeDilation * CinematicTimeDilation * DemoPlayTimeDilation;
+		if (bAllowTimeDilation)
+		{
+			return TimeDilation * CinematicTimeDilation * DemoPlayTimeDilation;
+		}
+
+		return 1.0f;
 	}
 
 	/**
@@ -816,10 +837,13 @@ public:
 	 */
 	ENGINE_API virtual float FixupDeltaSeconds(float DeltaSeconds, float RealDeltaSeconds);
 
-	/** Sets the global time dilation value (subject to clamping). Returns the final value that was set. */
+	/** Sets the global time dilation value (subject to clamping). Returns the final value that was set (which may be clamped) even if bAllowTimeDilation is set to false. */
 	ENGINE_API virtual float SetTimeDilation(float NewTimeDilation);
 
-	/** Sets the cinematic time dilation value (subject to clamping). Returns the final value that was set. */
+	/** Sets whether or not Time Dilation is allowed. When true, GetEffectiveTimeDilation() will return the time dilation as determined by the various dilation factors, when false it will return 1.0f. */
+	ENGINE_API virtual void SetAllowTimeDilation(bool NewAllowTimeDilation);
+
+	/** Sets the cinematic time dilation value (subject to clamping). Returns the final value that was set (which may be clamped) even if bAllowTimeDilation is set to false. */
 	ENGINE_API virtual float SetCinematicTimeDilation(float NewCinematicTimeDilation);
 
 	/** @return configuration for NavigationSystem's creation. Null means
diff --git a/Engine/Source/Runtime/Engine/Private/WorldSettings.cpp b/Engine/Source/Runtime/Engine/Private/WorldSettings.cpp
--- a/Engine/Source/Runtime/Engine/Private/WorldSettings.cpp
+++ b/Engine/Source/Runtime/Engine/Private/WorldSettings.cpp
@@ -122,6 +122,7 @@ PRAGMA_ENABLE_DEPRECATION_WARNINGS
 	TimeDilation = 1.0f;
 	CinematicTimeDilation = 1.0f;
 	DemoPlayTimeDilation = 1.0f;
+	bAllowTimeDilation = true;
 	PackedLightAndShadowMapTextureSize = 1024;
 	SetHidden(false);
 
@@ -341,6 +342,12 @@ float AWorldSettings::SetCinematicTimeDilation(float NewCinematicTimeDilation)
 	return CinematicTimeDilation;
 }
 
+
+void AWorldSettings::SetAllowTimeDilation(bool NewAllowTimeDilation) 
+{ 
+	bAllowTimeDilation = NewAllowTimeDilation; 
+}
+
 void AWorldSettings::NotifyBeginPlay()
 {
 	UWorld* World = GetWorld();
