[HWRT] Fix incorrect MaxVertices in FRayTracingGeometrySegment for skeletal meshes.

diff --git a/Engine/Source/Runtime/Engine/Private/SkeletalMeshSceneProxy.cpp b/Engine/Source/Runtime/Engine/Private/SkeletalMeshSceneProxy.cpp
--- a/Engine/Source/Runtime/Engine/Private/SkeletalMeshSceneProxy.cpp
+++ b/Engine/Source/Runtime/Engine/Private/SkeletalMeshSceneProxy.cpp
@@ -965,8 +965,6 @@ void FSkeletalMeshSceneProxy::GetDynamicRayTracingInstances(FRayTracingInstanceC
 			
 			const bool bIgnorePreviewFilter = true;
 
-			uint32 TotalNumVertices = 0;
-
 			for (FSkeletalMeshSectionIter Iter(LODIndex, *MeshObject, LODData, LODSection, bIgnorePreviewFilter); Iter; ++Iter)
 			{
 				const FSkelMeshRenderSection& Section = Iter.GetSection();
@@ -977,11 +975,10 @@ void FSkeletalMeshSceneProxy::GetDynamicRayTracingInstances(FRayTracingInstanceC
 				CreateBaseMeshBatch(Collector.GetReferenceView(), LODData, LODIndex, SectionIndex, SectionElementInfo, MeshBatch, ESkinVertexFactoryMode::RayTracing);
 
 				RayTracingInstance.Materials.Add(MeshBatch);
-
-				TotalNumVertices += Section.GetNumVertices();
 			}
 
 			RayTracingInstance.InstanceTransforms.Add(GetLocalToWorld());
+			const uint32 VertexBufferNumVertices = LODData.StaticVertexBuffers.PositionVertexBuffer.GetNumVertices();
 			const uint32 VertexBufferStride = LODData.StaticVertexBuffers.PositionVertexBuffer.GetStride();
 
 			const FVertexFactory* VertexFactory = MeshObject->GetSkinVertexFactory(Collector.GetReferenceView(), LODIndex, 0, ESkinVertexFactoryMode::RayTracing);
@@ -1001,7 +998,7 @@ void FSkeletalMeshSceneProxy::GetDynamicRayTracingInstances(FRayTracingInstanceC
 
 					FRayTracingGeometrySegment Segment;
 					Segment.VertexBufferStride = VertexBufferStride;
-					Segment.MaxVertices = TotalNumVertices;
+					Segment.MaxVertices = VertexBufferNumVertices;
 					Segment.FirstPrimitive = Section.BaseIndex / 3;
 					Segment.NumPrimitives = Section.NumTriangles;
 					Segment.bEnabled = !MeshObject->IsMaterialHidden(LODIndex, SectionElementInfo.UseMaterialIndex) && !Section.bDisabled && Section.bVisibleInRayTracing;
diff --git a/Engine/Source/Runtime/Engine/Private/SkeletalRenderCPUSkin.cpp b/Engine/Source/Runtime/Engine/Private/SkeletalRenderCPUSkin.cpp
--- a/Engine/Source/Runtime/Engine/Private/SkeletalRenderCPUSkin.cpp
+++ b/Engine/Source/Runtime/Engine/Private/SkeletalRenderCPUSkin.cpp
@@ -370,6 +370,7 @@ void FSkeletalMeshObjectCPUSkin::FSkeletalMeshObjectLOD::InitResources(const FSk
 		FSkeletalMeshLODRenderData& LODModel = SkelMeshRenderData->LODRenderData[LODIndex];
 		FBufferRHIRef VertexBufferRHI = LODModel.StaticVertexBuffers.PositionVertexBuffer.VertexBufferRHI;
 		FBufferRHIRef IndexBufferRHI = LODModel.MultiSizeIndexContainer.GetIndexBuffer()->IndexBufferRHI;
+		uint32 VertexBufferNumVertices = LODModel.StaticVertexBuffers.PositionVertexBuffer.GetNumVertices();
 		uint32 VertexBufferStride = LODModel.StaticVertexBuffers.PositionVertexBuffer.GetStride();
 
 		uint32 TrianglesCount = 0;
@@ -381,7 +382,7 @@ void FSkeletalMeshObjectCPUSkin::FSkeletalMeshObjectLOD::InitResources(const FSk
 
 		TArray<FSkelMeshRenderSection>* RenderSections = &LODModel.RenderSections;
 		ENQUEUE_RENDER_COMMAND(InitSkeletalRenderCPUSkinRayTracingGeometry)(UE::RenderCommandPipe::SkeletalMesh,
-			[this, VertexBufferRHI, IndexBufferRHI, VertexBufferStride, TrianglesCount, RenderSections, &SourceGeometry = LODModel.SourceRayTracingGeometry](FRHICommandListBase& RHICmdList)
+			[this, VertexBufferRHI, IndexBufferRHI, VertexBufferNumVertices, VertexBufferStride, TrianglesCount, RenderSections, &SourceGeometry = LODModel.SourceRayTracingGeometry](FRHICommandListBase& RHICmdList)
 			{
 				FRayTracingGeometryInitializer Initializer;
 				static const FName DebugName("FSkeletalMeshObjectCPUSkin");
@@ -395,12 +396,6 @@ void FSkeletalMeshObjectCPUSkin::FSkeletalMeshObjectLOD::InitResources(const FSk
 				TArray<FRayTracingGeometrySegment> GeometrySections;
 				GeometrySections.Reserve(RenderSections->Num());
 
-				uint32 TotalNumVertices = 0;
-				for (const FSkelMeshRenderSection& Section : *RenderSections)
-				{
-					TotalNumVertices += Section.GetNumVertices();
-				}
-
 				for (const FSkelMeshRenderSection& Section : *RenderSections)
 				{
 					FRayTracingGeometrySegment Segment;
@@ -408,7 +403,7 @@ void FSkeletalMeshObjectCPUSkin::FSkeletalMeshObjectLOD::InitResources(const FSk
 					Segment.VertexBufferStride = VertexBufferStride;
 					Segment.VertexBufferOffset = 0;
 					Segment.VertexBufferElementType = VET_Float3;
-					Segment.MaxVertices = TotalNumVertices;
+					Segment.MaxVertices = VertexBufferNumVertices;
 					Segment.FirstPrimitive = Section.BaseIndex / 3;
 					Segment.NumPrimitives = Section.NumTriangles;
 					Segment.bEnabled = !Section.bDisabled && Section.bVisibleInRayTracing;
diff --git a/Engine/Source/Runtime/Engine/Private/SkeletalRenderGPUSkin.cpp b/Engine/Source/Runtime/Engine/Private/SkeletalRenderGPUSkin.cpp
--- a/Engine/Source/Runtime/Engine/Private/SkeletalRenderGPUSkin.cpp
+++ b/Engine/Source/Runtime/Engine/Private/SkeletalRenderGPUSkin.cpp
@@ -1113,16 +1113,15 @@ void FSkeletalMeshObjectGPUSkin::UpdateRayTracingGeometry_Internal(
 		{
 
 			FBufferRHIRef IndexBufferRHI = LODModel.MultiSizeIndexContainer.GetIndexBuffer()->IndexBufferRHI;
+			uint32 VertexBufferNumVertices = LODModel.StaticVertexBuffers.PositionVertexBuffer.GetNumVertices();
 			uint32 VertexBufferStride = LODModel.StaticVertexBuffers.PositionVertexBuffer.GetStride();
 			
 			//#dxr_todo: do we need support for separate sections in FRayTracingGeometryData?
 			uint32 TotalNumTriangles = 0;
-			uint32 TotalNumVertices = 0;
 			for (int32 SectionIndex = 0; SectionIndex < LODModel.RenderSections.Num(); SectionIndex++)
 			{
 				const FSkelMeshRenderSection& Section = LODModel.RenderSections[SectionIndex];
 				TotalNumTriangles += Section.NumTriangles;
-				TotalNumVertices += Section.GetNumVertices();
 			}
 
 			FRayTracingGeometryInitializer Initializer;
@@ -1158,7 +1157,7 @@ void FSkeletalMeshObjectGPUSkin::UpdateRayTracingGeometry_Internal(
 				Segment.VertexBufferElementType = VET_Float3;
 				Segment.VertexBufferStride = VertexBufferStride;
 				Segment.VertexBufferOffset = 0;
-				Segment.MaxVertices = TotalNumVertices;
+				Segment.MaxVertices = VertexBufferNumVertices;
 				Segment.FirstPrimitive = Section.BaseIndex / 3;
 				Segment.NumPrimitives = Section.NumTriangles;
 
diff --git a/Engine/Source/Runtime/Engine/Private/SkeletalRenderStatic.cpp b/Engine/Source/Runtime/Engine/Private/SkeletalRenderStatic.cpp
--- a/Engine/Source/Runtime/Engine/Private/SkeletalRenderStatic.cpp
+++ b/Engine/Source/Runtime/Engine/Private/SkeletalRenderStatic.cpp
@@ -86,12 +86,6 @@ void FSkeletalMeshObjectStatic::InitResources(const FSkinnedMeshSceneProxyDesc&
 							TArray<FRayTracingGeometrySegment> GeometrySections;
 							GeometrySections.Reserve(RenderSections->Num());
 
-							uint32 TotalNumVertices = 0;
-							for (const FSkelMeshRenderSection& Section : *RenderSections)
-							{
-								TotalNumVertices += Section.GetNumVertices();
-							}
-
 							for (const FSkelMeshRenderSection& Section : *RenderSections)
 							{
 								FRayTracingGeometrySegment Segment;
@@ -99,7 +93,7 @@ void FSkeletalMeshObjectStatic::InitResources(const FSkinnedMeshSceneProxyDesc&
 								Segment.VertexBufferElementType = VET_Float3;
 								Segment.VertexBufferOffset = 0;
 								Segment.VertexBufferStride = PositionVertexBufferPtr->GetStride();
-								Segment.MaxVertices = TotalNumVertices;
+								Segment.MaxVertices = PositionVertexBufferPtr->GetNumVertices();
 								Segment.FirstPrimitive = Section.BaseIndex / 3;
 								Segment.NumPrimitives = Section.NumTriangles;
 								Segment.bEnabled = !Section.bDisabled && Section.bVisibleInRayTracing;
