Sequencer: User reports crash when saving sequence with spawnable. Appears to be stale data in the Skeletal Mesh Component. See slack linked in Jira for details.

#okforgitub public,internal,animation

diff --git a/Engine/Source/Runtime/Engine/Private/Components/SkeletalMeshComponent.cpp b/Engine/Source/Runtime/Engine/Private/Components/SkeletalMeshComponent.cpp
--- a/Engine/Source/Runtime/Engine/Private/Components/SkeletalMeshComponent.cpp
+++ b/Engine/Source/Runtime/Engine/Private/Components/SkeletalMeshComponent.cpp
@@ -3495,7 +3495,10 @@ void USkeletalMeshComponent::ResetAnimInstanceDynamics(ETeleportType InTeleportT
 
 	for(UAnimInstance* LinkedInstance : LinkedInstances)
 	{
-		LinkedInstance->ResetDynamics(InTeleportType);
+		if (LinkedInstance)
+		{
+			LinkedInstance->ResetDynamics(InTeleportType);
+		}
 	}
 
 	if(PostProcessAnimInstance)
