Iris - Update autonomous role replication condition as part of updating the owning net connection, regardless whether it's to be enabled or disabled. Fixes issue with for example COND_SimulatedOnly properties not replicating to the former autonomous proxy connection when unpossessing a pawn.

diff --git a/Engine/Source/Runtime/Engine/Private/Actor.cpp b/Engine/Source/Runtime/Engine/Private/Actor.cpp
--- a/Engine/Source/Runtime/Engine/Private/Actor.cpp
+++ b/Engine/Source/Runtime/Engine/Private/Actor.cpp
@@ -4383,9 +4383,9 @@ void AActor::SetAutonomousProxy(const bool bInAutonomousProxy, const bool bAllow
 
 #if UE_WITH_IRIS
 			// Update autonomous role condition
-			if (UReplicationSystem* ReplicationSystem = UE::Net::FReplicationSystemUtil::GetReplicationSystem(GetNetOwner()))
+			if (UReplicationSystem* ReplicationSystem = UE::Net::FReplicationSystemUtil::GetReplicationSystem(this))
 			{
-				if (UObjectReplicationBridge* Bridge = ReplicationSystem->GetReplicationBridgeAs<UObjectReplicationBridge>())
+				if (const UObjectReplicationBridge* Bridge = ReplicationSystem->GetReplicationBridgeAs<UObjectReplicationBridge>())
 				{
 					const UE::Net::FNetRefHandle RefHandle = Bridge->GetReplicatedRefHandle(this);
 					if (RefHandle.IsValid())
diff --git a/Engine/Source/Runtime/Engine/Private/ActorReplication.cpp b/Engine/Source/Runtime/Engine/Private/ActorReplication.cpp
--- a/Engine/Source/Runtime/Engine/Private/ActorReplication.cpp
+++ b/Engine/Source/Runtime/Engine/Private/ActorReplication.cpp
@@ -1141,16 +1141,14 @@ void AActor::UpdateOwningNetConnection()
 				{
 					ReplicationSystem->SetOwningNetConnection(RefHandle, NewOwningNetConnectionId);
 					// Update autonomous proxy condition
-					if (Actor->GetRemoteRole() == ROLE_AutonomousProxy)
-					{
-						const bool bEnableAutonomousCondition = true;
-						ReplicationSystem->SetReplicationConditionConnectionFilter(RefHandle, EReplicationCondition::RoleAutonomous, NewOwningNetConnectionId, bEnableAutonomousCondition);
-					}
+					const bool bEnableAutonomousCondition = Actor->GetRemoteRole() == ROLE_AutonomousProxy;
+					ReplicationSystem->SetReplicationConditionConnectionFilter(RefHandle, EReplicationCondition::RoleAutonomous, NewOwningNetConnectionId, bEnableAutonomousCondition);
 				}
 			}
 		}
 	} while (EveryChildren.Num());
 }
+
 #endif // UE_WITH_IRIS
 
 void AActor::UpdateReplicatePhysicsCondition()
