Added support for recording CSV stats in AActor::ForceNetUpdate and AActor::FlushNetDormancy.

Recording is only enabled in Non-Shipping builds and when compiling for Dedicated servers only.
The new CSV Categories (Actor_FlushNetDormancy and Actor_ForceNetUpdate) are disabled by default.

By default we record the NativeParentClass name of the Actor.

The CVar: net.Debug.ActorClassNameTypeCSV, controls which type of class name to record:
- 0: Record the Parent native class name of the given Actor
- 1: Record the TopMost non-native class name of the given Actor
- 2: Record the Actor class name

#tests Test in a local test and shipping builds.

diff --git a/Engine/Source/Runtime/CoreUObject/Private/UObject/UObjectBaseUtility.cpp b/Engine/Source/Runtime/CoreUObject/Private/UObject/UObjectBaseUtility.cpp
--- a/Engine/Source/Runtime/CoreUObject/Private/UObject/UObjectBaseUtility.cpp
+++ b/Engine/Source/Runtime/CoreUObject/Private/UObject/UObjectBaseUtility.cpp
@@ -691,6 +691,26 @@ UClass* GetParentNativeClass(UClass* Class)
 	return Class;
 }
 
+UClass* GetTopMostNonNativeClass(UClass* Class, bool bAllowReturnNativeClass)
+{
+	while (Class && !Class->IsNative())
+	{
+		UClass* SuperClass = Class->GetSuperClass();
+		// is the parent class native?
+		if (SuperClass && SuperClass->IsNative())
+		{
+			// return the current class instead (Top most non-native class)
+			return Class;
+		}
+
+		// otherwise keep going up the hierarchy
+		Class = SuperClass;
+	}
+
+	// If the input given class is Native, return it only when bFallBackToNative is true, otherwise return nullptr
+	return bAllowReturnNativeClass ? Class : nullptr;
+}
+
 #if !STATS && !ENABLE_STATNAMEDEVENTS && USE_LIGHTWEIGHT_STATS_FOR_HITCH_DETECTION && USE_HITCH_DETECTION && USE_LIGHTWEIGHT_UOBJECT_STATS_FOR_HITCH_DETECTION
 #include "HAL/ThreadHeartBeat.h"
 #include "HAL/ThreadManager.h"
diff --git a/Engine/Source/Runtime/CoreUObject/Public/UObject/UObjectBaseUtility.h b/Engine/Source/Runtime/CoreUObject/Public/UObject/UObjectBaseUtility.h
--- a/Engine/Source/Runtime/CoreUObject/Public/UObject/UObjectBaseUtility.h
+++ b/Engine/Source/Runtime/CoreUObject/Public/UObject/UObjectBaseUtility.h
@@ -927,6 +927,19 @@ FORCEINLINE FString GetFullNameSafe(const UObjectBaseUtility* Object)
  */
 COREUOBJECT_API UClass* GetParentNativeClass(UClass* Class);
 
+/** 
+* Returns the TopMost non-native parent class of the supplied class.
+* 
+* @param Class - The class used to find its TopMost non-native parent class.
+* @param bAllowReturnNativeClass - If supplied class is native, it will be returned.
+*
+* @return
+*	- TopMost non-native parent class found in the hierarchy.
+*	- Class if the supplied Class is already native and bAllowReturnNativeClass is true.
+*	- nullptr if the supplied class is already native and bAllowReturnNativeClass is false.
+*/
+COREUOBJECT_API UClass* GetTopMostNonNativeClass(UClass* Class, bool bAllowReturnNativeClass);
+
 #if !defined(USE_LIGHTWEIGHT_UOBJECT_STATS_FOR_HITCH_DETECTION)
 #define USE_LIGHTWEIGHT_UOBJECT_STATS_FOR_HITCH_DETECTION (1)
 #endif
diff --git a/Engine/Source/Runtime/Engine/Private/Actor.cpp b/Engine/Source/Runtime/Engine/Private/Actor.cpp
--- a/Engine/Source/Runtime/Engine/Private/Actor.cpp
+++ b/Engine/Source/Runtime/Engine/Private/Actor.cpp
@@ -118,6 +118,16 @@ FCriticalSection CSVActorClassNameToCountMapLock;
 
 #endif // (CSV_PROFILER_STATS && !UE_BUILD_SHIPPING)
 
+/** Detailed CSV Recording of Actor Net Stats */
+#ifndef CSV_RECORD_DETAILED_ACTOR_NET_STATS
+#define CSV_RECORD_DETAILED_ACTOR_NET_STATS (CSV_PROFILER_STATS && !UE_BUILD_SHIPPING && UE_SERVER)
+#endif // CSV_RECORD_DETAILED_ACTOR_NET_STATS
+
+#if CSV_RECORD_DETAILED_ACTOR_NET_STATS
+CSV_DEFINE_CATEGORY(Actor_FlushNetDormancy, false);
+CSV_DEFINE_CATEGORY(Actor_ForceNetUpdate, false);
+#endif // CSV_RECORD_DETAILED_ACTOR_NET_STATS
+
 #if WITH_EDITOR
 namespace ActorUtils
 {
@@ -146,6 +156,73 @@ uint32 AActor::BeginPlayCallDepth = 0;
 
 namespace ActorUtils
 {
+	namespace CVars
+	{
+#if CSV_RECORD_DETAILED_ACTOR_NET_STATS
+		// Which type of class name to record when the CSV categories: Actor_FlushNetDormancy or Actor_ForceNetUpdate are enabled
+		// From least verbose to most verbose:
+		// 	- 0: Record the Parent native class name of the given Actor
+		// 	- 1: Record the TopMost non-native class name of the given Actor
+		// 	- 2: Record the Actor class name
+		int32 CsvActorClassVerbosity = 0;
+		FAutoConsoleVariableRef CVarCsvActorClassVerbosity(
+			TEXT("net.ForceFlush.CsvActorClassVerbosity")
+			, CsvActorClassVerbosity
+			, TEXT("Which Actor type of class name to record when the CSV categories: Actor_FlushNetDormancy or Actor_ForceNetUpdate are enabled"
+					" 0: Record the Parent native class name of the given Actor."
+					" 1: Record the TopMost non-native class name of the given Actor."
+					" 2: Record the Actor class name.")
+			, ECVF_Cheat
+		);
+#endif // CSV_RECORD_DETAILED_ACTOR_NET_STATS
+	}
+
+#if CSV_RECORD_DETAILED_ACTOR_NET_STATS
+	static void CsvRecordClassName(UClass* Class, int32 CsvCategoryIndex)
+	{
+		FCsvProfiler* CsvProfiler = FCsvProfiler::Get();
+		if (!CsvProfiler)
+		{
+			return;
+		}
+
+		if (!CsvProfiler->IsCapturing())
+		{
+			return;
+		}
+
+		if (!CsvProfiler->IsCategoryEnabled(CsvCategoryIndex))
+		{
+			return;
+		}
+
+		switch (ActorUtils::CVars::CsvActorClassVerbosity)
+		{
+			case 0: // Record the Parent native class name of the given Actor
+			{
+				Class = GetParentNativeClass(Class);
+				break;
+			}
+			case 1: // Record the TopMost non-native class name of the given Actor
+			{
+				constexpr bool bAllowReturnNativeClass = true;
+				Class = GetTopMostNonNativeClass(Class, bAllowReturnNativeClass);
+				break;
+			}
+			case 2: // Record the Actor class name (reuse Class variable)
+			{
+				break;
+			}
+		}
+				
+		if (ensure(Class))
+		{
+			CsvProfiler->RecordCustomStat(*GetNameSafe(Class), CsvCategoryIndex, 1, ECsvCustomStatOp::Accumulate);
+		}
+	}
+#endif // CSV_RECORD_DETAILED_ACTOR_NET_STATS
+
+
 	template<class Function>
 	static void ForEachNetDriver(UEngine* Engine, const UWorld* const World, const Function InFunction)
 	{
@@ -2780,6 +2857,10 @@ bool AActor::IsRelevancyOwnerFor(const AActor* ReplicatedActor, const AActor* Ac
 
 void AActor::ForceNetUpdate()
 {
+#if CSV_RECORD_DETAILED_ACTOR_NET_STATS
+	ActorUtils::CsvRecordClassName(GetClass(), CSV_CATEGORY_INDEX(Actor_ForceNetUpdate));
+#endif // CSV_RECORD_DETAILED_ACTOR_NET_STATS
+
 	UNetDriver* NetDriver = GetNetDriver();
 
 	if (GetLocalRole() == ROLE_Authority)
@@ -2881,6 +2962,10 @@ void AActor::FlushNetDormancy()
 
 	QUICK_SCOPE_CYCLE_COUNTER(NET_AActor_FlushNetDormancy);
 
+#if CSV_RECORD_DETAILED_ACTOR_NET_STATS
+	ActorUtils::CsvRecordClassName(GetClass(), CSV_CATEGORY_INDEX(Actor_FlushNetDormancy));
+#endif // CSV_RECORD_DETAILED_ACTOR_NET_STATS
+
 	bool bWasDormInitial = false;
 	if (NetDormancy == DORM_Initial)
 	{
