Add LevelInstance support to Entities

-Changed the level_entity to have a transform_component, this provides a clean and logical to transfer the transfrom of the ALevelInstance actor to the scenegraph
-Added functionnality in ALevelEntity to transfer the transform from the Actor to the entities within it.
-Changed ALevelEntity::GetOrCreateLevelEntityActor to only take ULevels since it's too easy to pass the wrong info when passing a UWorld
-Added ALevelEntity::GetLevelEntityActor/GetLevelEntityActorChecked that takes an entity as input and returns the correct LevelEntityActor.
-Added level instance editing state tracking in UStaticMeshEntityComponent_Implementation so that we could create the render proxies correctly
-Modified AEntityProxyActor to forward level instance editing state changes to the entity
-Blocked AEntityProxyActor::TransformUpdated from affecting local transforms on entities when we're applying level transforms since the hierarchical nature of the transforms requires these change to be done in a specific order (and this is not happening during ApplyLevelTransform)

diff --git a/Engine/Source/Runtime/Engine/Public/PrimitiveSceneProxy.h b/Engine/Source/Runtime/Engine/Public/PrimitiveSceneProxy.h
--- a/Engine/Source/Runtime/Engine/Public/PrimitiveSceneProxy.h
+++ b/Engine/Source/Runtime/Engine/Public/PrimitiveSceneProxy.h
@@ -269,7 +269,7 @@ public:
 	 * This is called in the game thread when the object enters/leaves LevelInstance levels.
 	 * @param bInEditingState - true if the parent actor belongs to an editing LevelInstance sublevel
 	 */
-	void SetLevelInstanceEditingState_GameThread(const bool bInEditingState);
+	ENGINE_API void SetLevelInstanceEditingState_GameThread(const bool bInEditingState);
 
 	/**
 	 * Updates hover state for the primitive proxy. This simply sends a message to the rendering thread to call SetHovered_RenderThread.
