Added engine restart requirement to local fog volume on trnaslucent option.

#tests PC editor dx12
#rnx

diff --git a/Engine/Source/Runtime/Engine/Classes/Engine/RendererSettings.h b/Engine/Source/Runtime/Engine/Classes/Engine/RendererSettings.h
--- a/Engine/Source/Runtime/Engine/Classes/Engine/RendererSettings.h
+++ b/Engine/Source/Runtime/Engine/Classes/Engine/RendererSettings.h
@@ -731,9 +731,10 @@ class URendererSettings : public UDeveloperSettings
 		ToolTip="The axis that sorting will occur along when Translucent Sort Policy is set to SortAlongAxis."))
 	FVector TranslucentSortAxis;
 
-	UPROPERTY(config, EditAnywhere, Category=Translucency, meta=(
+	UPROPERTY(config, EditAnywhere, Category=Translucency, meta=( 
 		ConsoleVariable="r.LocalFogVolume.ApplyOnTranslucent",
-		ToolTip="Allow local fog volumes to be combined and rendered over translucent meshes. Only per vertex evaluation is supported as of today. It requires r.SupportLocalFogVolumes to be true"))
+		ToolTip="Allow local fog volumes to be combined and rendered over translucent meshes. Only per vertex evaluation is supported as of today. It requires r.SupportLocalFogVolumes to be true",
+		ConfigRestartRequired = true))
 	uint32 bLocalFogVolumeApplyOnTranslucent:1;
 
 	UPROPERTY(config, EditAnywhere, Category = VR, meta = (
