[DemoRecording] Additional crash fix. Don't access pointers that might not be valid.

diff --git a/Engine/Source/Runtime/Engine/Private/ReplayHelper.cpp b/Engine/Source/Runtime/Engine/Private/ReplayHelper.cpp
--- a/Engine/Source/Runtime/Engine/Private/ReplayHelper.cpp
+++ b/Engine/Source/Runtime/Engine/Private/ReplayHelper.cpp
@@ -630,10 +630,18 @@ void FReplayHelper::SaveCheckpoint(UNetConnection* Connection)
 					if (!bDeltaCheckpoint || NetworkObjectInfo->bDirtyForReplay)
 					{
 						const TWeakObjectPtr<AActor>& WeakActor = NetworkObjectInfo->WeakActor;
+
+						// If the actor isn't valid we can't safely access its properties nor serialize it.
+						if (!WeakActor.IsValid())
+						{
+							NetworkObjectInfo->bDirtyForReplay = false;
+							continue;
+						}
+
 						AActor* Actor = NetworkObjectInfo->Actor;
 
 						// check to see if it should replicate at all
-						bool bCheckpointActor = WeakActor.IsValid() && ((Actor->GetRemoteRole() != ROLE_None || Actor->GetTearOff()) && (Actor == Connection->PlayerController || Cast<APlayerController>(Actor) == nullptr));
+						bool bCheckpointActor = ((Actor->GetRemoteRole() != ROLE_None || Actor->GetTearOff()) && (Actor == Connection->PlayerController || Cast<APlayerController>(Actor) == nullptr));
 						
 						// now look for an open channel
 						bCheckpointActor = bCheckpointActor && ActorChannelMap.Contains(WeakActor);
@@ -646,8 +654,12 @@ void FReplayHelper::SaveCheckpoint(UNetConnection* Connection)
 						
 						if (bCheckpointActor)
 						{
-							ActorArray.Add({ WeakActor, Actor->GetOuter() });
+#if DO_CHECK && !(UE_BUILD_SHIPPING || UE_BUILD_TEST)
+							const ULevel* ActorLevel = Actor->GetLevel();
+							checkf(Actor->GetOuter() == ActorLevel, TEXT("Outer %p %s != Level %p %s. Actor %s"), Actor->GetOuter(), ToCStr(GetNameSafe(Actor->GetOuter())), ActorLevel, ToCStr(GetNameSafe(ActorLevel)), ToCStr(GetNameSafe(Actor)));
+#endif
 
+							ActorArray.Add({ WeakActor, Actor->GetOuter() });
 							NetworkObjectInfo->bDirtyForReplay = false;
 						}
 					}
@@ -667,10 +679,6 @@ void FReplayHelper::SaveCheckpoint(UNetConnection* Connection)
 			const UObject* CurrentLevelToIndex = ActorArray[CurrentIt].Level;
 			const FLevelStatus& LevelStatus = FindOrAddLevelStatus(*Cast<const ULevel>(CurrentLevelToIndex));
 
-#if !(UE_BUILD_SHIPPING || UE_BUILD_TEST)
-			// Validate that we get the correct level
-			check(Cast<const ULevel>(CurrentLevelToIndex) == ActorArray[CurrentIt].Actor->GetLevel());
-#endif
 			while (CurrentIt < EndIt && (CurrentLevelToIndex == ActorArray[CurrentIt].Level))
 			{
 				int32 PendingIndex = CheckpointSaveContext.PendingCheckpointActors.Add({ ActorArray[CurrentIt].Actor, LevelStatus.LevelIndex });
