[WorldPartitionLevelHelper] Fix an issue that could cause Level Instance to trigger an assert indicating that it hasn't been deserialized prior to PostLoad.

diff --git a/Engine/Source/Runtime/Engine/Private/WorldPartition/WorldPartitionLevelHelper.cpp b/Engine/Source/Runtime/Engine/Private/WorldPartition/WorldPartitionLevelHelper.cpp
--- a/Engine/Source/Runtime/Engine/Private/WorldPartition/WorldPartitionLevelHelper.cpp
+++ b/Engine/Source/Runtime/Engine/Private/WorldPartition/WorldPartitionLevelHelper.cpp
@@ -514,6 +514,16 @@ bool FWorldPartitionLevelHelper::LoadActorsWithPropertyOverridesInternal(FLoadAc
 
 			InstancingContext.AddPackageMapping(PackageName, AssetPackageInstanceName);
 			InstancingContext.FixupSoftObjectPath(RemappedPath);
+
+			// If packages are already loaded, add a reference, to make sure they're not trashed before completion callback is called
+			if (UPackage* WorldPackage = FindPackage(nullptr, *RemappedContainerPackage.ToString()))
+			{
+				LoadProgress->Params.PackageReferencer->AddReference(WorldPackage);
+			}
+			if (UPackage* ActorPackage = FindPackage(nullptr, *AssetPackageInstanceName.ToString()))
+			{
+				LoadProgress->Params.PackageReferencer->AddReference(ActorPackage);
+			}
 		}
 
 		FName RemappedPackageName = InstancingContext.RemapPackage(PackageName);
@@ -689,6 +699,16 @@ bool FWorldPartitionLevelHelper::LoadActorsInternal(FLoadActorsParams&& InParams
 		const FLinkerInstancingContext ContainerInstancingContext = FLinkerInstancingContext::DuplicateContext(LinkerInstancingContexts.FindChecked(PackageObjectMapping->ContainerID));
 		const FName PackageName = ContainerInstancingContext.RemapPackage(PackageToLoad);
 
+		if (!PackageObjectMapping->ContainerID.IsMainContainer())
+		{
+			const FName ContainerPackageName = ContainerInstancingContext.RemapPackage(PackageObjectMapping->ContainerPackage);
+			if (UPackage* ContainerPackage = FindPackage(nullptr, *ContainerPackageName.ToString()))
+			{
+				// If container package is already loaded, add a reference, to make sure it's not trashed before completion callback is called
+				InParams.PackageReferencer->AddReference(ContainerPackage);
+			}
+		}
+
 		FLoadPackageAsyncDelegate CompletionCallback = FLoadPackageAsyncDelegate::CreateLambda([LoadProgress, PackageObjectMapping, LoadedOverrides = InLoadedPropertyOverrides, PackageReferencer = InParams.PackageReferencer, OuterWorld = InParams.OuterWorld, DestLevel = InParams.DestLevel, CompletionCallback = InParams.CompletionCallback](const FName& LoadedPackageName, UPackage* LoadedPackage, EAsyncLoadingResult::Type Result)
 		{
 			const FName ActorName = *FPaths::GetExtension(PackageObjectMapping->Path.ToString());
