Material Editor: Add "Preview Node" button to MaterialNode title section

Misc: Remove constants from basic math node captions (can see this info inline in the nodes)
Misc: SGraphNode Fixed Node Title Background failing to fill the title area in some cases

diff --git a/Engine/Source/Editor/GraphEditor/Private/MaterialNodes/SGraphNodeMaterialBase.cpp b/Engine/Source/Editor/GraphEditor/Private/MaterialNodes/SGraphNodeMaterialBase.cpp
--- a/Engine/Source/Editor/GraphEditor/Private/MaterialNodes/SGraphNodeMaterialBase.cpp
+++ b/Engine/Source/Editor/GraphEditor/Private/MaterialNodes/SGraphNodeMaterialBase.cpp
@@ -54,6 +54,7 @@
 #include "UObject/ObjectPtr.h"
 #include "UnrealClient.h"
 #include "Widgets/Images/SImage.h"
+#include "Widgets/Input/SButton.h"
 #include "Widgets/Input/SCheckBox.h"
 #include "Widgets/Layout/SBorder.h"
 #include "Widgets/Layout/SBox.h"
@@ -301,6 +302,7 @@ void FPreviewElement::Draw_RenderThread(FRDGBuilder& GraphBuilder, const FDrawPa
 
 void SGraphNodeMaterialBase::Construct(const FArguments& InArgs, UMaterialGraphNode* InNode)
 {
+	this->TitleBorderMargin = FMargin(10.f, 5.f, 50.f, 3.f);
 	this->GraphNode = InNode;
 	this->MaterialNode = InNode;
 
@@ -552,22 +554,39 @@ void SGraphNodeMaterialBase::SetDefaultTitleAreaWidget(TSharedRef<SOverlay> Defa
 	{
 		DefaultTitleAreaWidget->AddSlot()
 		.HAlign(HAlign_Right)
-		.VAlign(VAlign_Center)
+		.VAlign(VAlign_Top)
 		.Padding(FMargin(5))
 		[
-			SNew(SCheckBox)
-			.OnCheckStateChanged( this, &SGraphNodeMaterialBase::OnExpressionPreviewChanged )
-			.IsChecked( IsExpressionPreviewChecked() )
-			.Cursor(EMouseCursor::Default)
-			.Style(FAppStyle::Get(), "Graph.Node.AdvancedView")
+			SNew(SHorizontalBox)
+			+ SHorizontalBox::Slot()
+			.AutoWidth()
 			[
-				SNew(SHorizontalBox)
-				+SHorizontalBox::Slot()
-				.VAlign(VAlign_Center)
-				.HAlign(HAlign_Center)
+				SNew(SButton)
+				.ButtonStyle(FAppStyle::Get(), "HoverHintOnly")
+				.OnClicked(this, &SGraphNodeMaterialBase::OnTogglePreviewClicked)
+				.Cursor(EMouseCursor::Default)
 				[
 					SNew(SImage)
-					. Image(GetExpressionPreviewArrow())
+					.Image(FAppStyle::GetBrush(TEXT("Icons.Preview")))
+				]
+			]
+			+ SHorizontalBox::Slot()
+			.AutoWidth()
+			[
+				SNew(SCheckBox)
+				.OnCheckStateChanged( this, &SGraphNodeMaterialBase::OnExpressionPreviewChanged )
+				.IsChecked( IsExpressionPreviewChecked() )
+				.Cursor(EMouseCursor::Default)
+				.Style(FAppStyle::Get(), "Graph.Node.AdvancedView")
+				[
+					SNew(SHorizontalBox)
+					+SHorizontalBox::Slot()
+					.VAlign(VAlign_Center)
+					.HAlign(HAlign_Center)
+					[
+						SNew(SImage)
+						.Image(GetExpressionPreviewArrow())
+					]
 				]
 			]
 		];
@@ -674,6 +693,12 @@ EVisibility SGraphNodeMaterialBase::ExpressionPreviewVisibility() const
 	return bShowPreview ? EVisibility::Visible : EVisibility::Collapsed;
 }
 
+FReply SGraphNodeMaterialBase::OnTogglePreviewClicked()
+{
+	MaterialNode->ToggleNodePreview();
+	return FReply::Handled();
+}
+
 void SGraphNodeMaterialBase::OnExpressionPreviewChanged( const ECheckBoxState NewCheckedState )
 {
 	UMaterialExpression* MaterialExpression = MaterialNode->MaterialExpression;
diff --git a/Engine/Source/Editor/GraphEditor/Private/SGraphNode.cpp b/Engine/Source/Editor/GraphEditor/Private/SGraphNode.cpp
--- a/Engine/Source/Editor/GraphEditor/Private/SGraphNode.cpp
+++ b/Engine/Source/Editor/GraphEditor/Private/SGraphNode.cpp
@@ -863,6 +863,8 @@ void SGraphNode::UpdateGraphNode()
 	TSharedRef<SOverlay> DefaultTitleAreaWidget =
 		SNew(SOverlay)
 		+SOverlay::Slot()
+		.HAlign(HAlign_Fill)
+		.VAlign(VAlign_Fill)
 		[
 			SNew(SImage)
 			.Image( FAppStyle::GetBrush("Graph.Node.TitleGloss") )
@@ -870,7 +872,7 @@ void SGraphNode::UpdateGraphNode()
 		]
 		+SOverlay::Slot()
 		.HAlign(HAlign_Fill)
-		.VAlign(VAlign_Center)
+		.VAlign(VAlign_Fill)
 		[
 			SNew(SHorizontalBox)
 			+ SHorizontalBox::Slot()
diff --git a/Engine/Source/Editor/GraphEditor/Public/MaterialNodes/SGraphNodeMaterialBase.h b/Engine/Source/Editor/GraphEditor/Public/MaterialNodes/SGraphNodeMaterialBase.h
--- a/Engine/Source/Editor/GraphEditor/Public/MaterialNodes/SGraphNodeMaterialBase.h
+++ b/Engine/Source/Editor/GraphEditor/Public/MaterialNodes/SGraphNodeMaterialBase.h
@@ -127,8 +127,11 @@ protected:
 	/** Returns text to over lay over the expression preview viewport */
 	FText ExpressionPreviewOverlayText() const;
 
+	/** Called when we click to preview node button */
+	FReply OnTogglePreviewClicked();
+
 	/** Show/hide Expression Preview */
-	void OnExpressionPreviewChanged( const ECheckBoxState NewCheckedState );
+	void OnExpressionPreviewChanged(const ECheckBoxState NewCheckedState);
 
 	/** hidden == unchecked, shown == checked */
 	ECheckBoxState IsExpressionPreviewChecked() const;
diff --git a/Engine/Source/Editor/MaterialEditor/Private/MaterialEditor.cpp b/Engine/Source/Editor/MaterialEditor/Private/MaterialEditor.cpp
--- a/Engine/Source/Editor/MaterialEditor/Private/MaterialEditor.cpp
+++ b/Engine/Source/Editor/MaterialEditor/Private/MaterialEditor.cpp
@@ -4986,20 +4986,23 @@ void FMaterialEditor::OnPreviewNode()
 	{
 		for (FGraphPanelSelectionSet::TConstIterator NodeIt(SelectedNodes); NodeIt; ++NodeIt)
 		{
-			UMaterialGraphNode* GraphNode = Cast<UMaterialGraphNode>(*NodeIt);
-			if (GraphNode)
+			if (UMaterialGraphNode* GraphNode = Cast<UMaterialGraphNode>(*NodeIt))
 			{
-				if (TSharedPtr<SGraphEditor> FocusedGraphEd = FocusedGraphEdPtr.Pin())
-				{
-					FocusedGraphEd->NotifyGraphChanged();
-				}
-				
-				SetPreviewExpression(GraphNode->MaterialExpression);
+				PreviewNode(*GraphNode);
 			}
 		}
 	}
 }
 
+void FMaterialEditor::PreviewNode(UMaterialGraphNode& InGraphNode)
+{
+	if (TSharedPtr<SGraphEditor> FocusedGraphEd = FocusedGraphEdPtr.Pin())
+	{
+		FocusedGraphEd->NotifyGraphChanged();
+	}
+	SetPreviewExpression(InGraphNode.MaterialExpression);
+}
+
 void FMaterialEditor::OnToggleRealtimePreview()
 {
 	const FGraphPanelSelectionSet SelectedNodes = GetSelectedNodes();
@@ -5943,6 +5946,7 @@ void FMaterialEditor::SetPreviewExpression(UMaterialExpression* NewPreviewExpres
 {
 	UMaterialExpressionFunctionOutput* FunctionOutput = Cast<UMaterialExpressionFunctionOutput>(NewPreviewExpression);
 
+	// If the NewPreviewExpression is null or the same as existing, stop previewing
 	if (!NewPreviewExpression || PreviewExpression == NewPreviewExpression)
 	{
 		if (FunctionOutput)
diff --git a/Engine/Source/Editor/MaterialEditor/Private/MaterialEditor.h b/Engine/Source/Editor/MaterialEditor/Private/MaterialEditor.h
--- a/Engine/Source/Editor/MaterialEditor/Private/MaterialEditor.h
+++ b/Engine/Source/Editor/MaterialEditor/Private/MaterialEditor.h
@@ -812,6 +812,13 @@ private:
 	void OnConvertNamedRerouteToReroute();
 	/** Command for previewing a selected node */
 	void OnPreviewNode();
+
+	/** 
+	 *	Call to Preview a InGraphNode.
+	 *	If InGraphNode is the currently-previewed node, will stop previewing
+	 */
+	virtual void PreviewNode(UMaterialGraphNode& InGraphNode) override;
+
 	/** Command for toggling real time preview of selected node */
 	void OnToggleRealtimePreview();
 	/** Command to select nodes downstream of selected node */
diff --git a/Engine/Source/Editor/MaterialEditor/Public/IMaterialEditor.h b/Engine/Source/Editor/MaterialEditor/Public/IMaterialEditor.h
--- a/Engine/Source/Editor/MaterialEditor/Public/IMaterialEditor.h
+++ b/Engine/Source/Editor/MaterialEditor/Public/IMaterialEditor.h
@@ -13,6 +13,7 @@ class IMaterialEditorModule;
 class UMaterialExpression;
 class UMaterialExpressionComment;
 class UMaterialExpressionComposite;
+class UMaterialGraphNode;
 class UMaterialInterface;
 
 /**
@@ -162,6 +163,9 @@ public:
 
 	/** Updates materials used in material stats */
 	virtual void RefreshStatsMaterials() {};
+
+	/** Start/Stop Previewing the specified material graph node */
+	virtual void PreviewNode(UMaterialGraphNode& InGraphNode) {};
 private:
 	FMaterialEditorClosedEvent MaterialEditorClosedEvent;
 	FRegisterTabSpawnersEvent RegisterTabSpawnersEvent;
diff --git a/Engine/Source/Editor/UnrealEd/Classes/MaterialGraph/MaterialGraphNode.h b/Engine/Source/Editor/UnrealEd/Classes/MaterialGraph/MaterialGraphNode.h
--- a/Engine/Source/Editor/UnrealEd/Classes/MaterialGraph/MaterialGraphNode.h
+++ b/Engine/Source/Editor/UnrealEd/Classes/MaterialGraph/MaterialGraphNode.h
@@ -8,6 +8,7 @@
 #include "MaterialGraph/MaterialGraphNode_Base.h"
 #include "MaterialGraphNode.generated.h"
 
+class IMaterialEditor;
 class UEdGraphPin;
 
 UCLASS(MinimalAPI)
@@ -89,6 +90,8 @@ public:
 
 	virtual void PropagatePropertyChange() override;
 
+	UNREALED_API void ToggleNodePreview();
+
 private:
 	/** Make sure the MaterialExpression is owned by the Material */
 	void ResetMaterialExpressionOwner();
diff --git a/Engine/Source/Editor/UnrealEd/Private/MaterialGraphNode.cpp b/Engine/Source/Editor/UnrealEd/Private/MaterialGraphNode.cpp
--- a/Engine/Source/Editor/UnrealEd/Private/MaterialGraphNode.cpp
+++ b/Engine/Source/Editor/UnrealEd/Private/MaterialGraphNode.cpp
@@ -8,6 +8,7 @@
 #include "Widgets/Notifications/SNotificationList.h"
 #include "Framework/Notifications/NotificationManager.h"
 #include "Framework/MultiBox/MultiBoxBuilder.h"
+#include "IMaterialEditor.h"
 #include "ToolMenus.h"
 #include "MaterialGraph/MaterialGraphSchema.h"
 
@@ -52,6 +53,7 @@
 #include "GraphEditorSettings.h"
 #include "Framework/Commands/GenericCommands.h"
 #include "ScopedTransaction.h"
+#include "Subsystems/AssetEditorSubsystem.h"
 #include "MaterialNodes/SGraphNodeMaterialBase.h"
 #include "SourceCodeNavigation.h"
 
@@ -833,6 +835,24 @@ void UMaterialGraphNode::PropagatePropertyChange()
 	}
 }
 
+void UMaterialGraphNode::ToggleNodePreview()
+{
+	if (!MaterialExpression->Material)
+	{
+		return;
+	}
+
+	UAssetEditorSubsystem* AssetEditorSubsystem = GEditor->GetEditorSubsystem<UAssetEditorSubsystem>();
+	IAssetEditorInstance* FoundAssetEditor = AssetEditorSubsystem->FindEditorForAsset(MaterialExpression->Material, true);
+	if (!FoundAssetEditor)
+	{
+		return;
+	}
+
+	IMaterialEditor* MaterialEditor = static_cast<IMaterialEditor*>(FoundAssetEditor);
+	MaterialEditor->PreviewNode(*this);
+}
+
 void UMaterialGraphNode::ResetMaterialExpressionOwner()
 {
 	if (MaterialExpression)
diff --git a/Engine/Source/Runtime/Engine/Private/Materials/MaterialExpressions.cpp b/Engine/Source/Runtime/Engine/Private/Materials/MaterialExpressions.cpp
--- a/Engine/Source/Runtime/Engine/Private/Materials/MaterialExpressions.cpp
+++ b/Engine/Source/Runtime/Engine/Private/Materials/MaterialExpressions.cpp
@@ -4576,18 +4576,7 @@ int32 UMaterialExpressionAdd::Compile(class FMaterialCompiler* Compiler, int32 O
 
 void UMaterialExpressionAdd::GetCaption(TArray<FString>& OutCaptions) const
 {
-	FString ret = TEXT("Add");
-
-	FExpressionInput ATraced = A.GetTracedInput();
-	FExpressionInput BTraced = B.GetTracedInput();
-	if(!ATraced.Expression || !BTraced.Expression)
-	{
-		ret += TEXT("(");
-		ret += ATraced.Expression ? TEXT(",") : FString::Printf( TEXT("%.4g,"), ConstA);
-		ret += BTraced.Expression ? TEXT(")") : FString::Printf( TEXT("%.4g)"), ConstB);
-	}
-
-	OutCaptions.Add(ret);
+	OutCaptions.Add(TEXT("Add"));
 }
 #endif // WITH_EDITOR
 
@@ -4607,19 +4596,7 @@ int32 UMaterialExpressionMultiply::Compile(class FMaterialCompiler* Compiler, in
 
 void UMaterialExpressionMultiply::GetCaption(TArray<FString>& OutCaptions) const
 {
-	FString ret = TEXT("Multiply");
-
-	FExpressionInput ATraced = A.GetTracedInput();
-	FExpressionInput BTraced = B.GetTracedInput();
-
-	if(!ATraced.Expression || !BTraced.Expression)
-	{
-		ret += TEXT("(");
-		ret += ATraced.Expression ? TEXT(",") : FString::Printf( TEXT("%.4g,"), ConstA);
-		ret += BTraced.Expression ? TEXT(")") : FString::Printf( TEXT("%.4g)"), ConstB);
-	}
-
-	OutCaptions.Add(ret);
+	OutCaptions.Add(TEXT("Multiply"));
 }
 #endif // WITH_EDITOR
 
@@ -4636,19 +4613,7 @@ int32 UMaterialExpressionDivide::Compile(class FMaterialCompiler* Compiler, int3
 
 void UMaterialExpressionDivide::GetCaption(TArray<FString>& OutCaptions) const
 {
-	FString ret = TEXT("Divide");
-
-	FExpressionInput ATraced = A.GetTracedInput();
-	FExpressionInput BTraced = B.GetTracedInput();
-
-	if(!ATraced.Expression || !BTraced.Expression)
-	{
-		ret += TEXT("(");
-		ret += ATraced.Expression ? TEXT(",") : FString::Printf( TEXT("%.4g,"), ConstA);
-		ret += BTraced.Expression ? TEXT(")") : FString::Printf( TEXT("%.4g)"), ConstB);
-	}
-
-	OutCaptions.Add(ret);
+	OutCaptions.Add(TEXT("Divide"));
 }
 
 int32 UMaterialExpressionSubtract::Compile(class FMaterialCompiler* Compiler, int32 OutputIndex)
@@ -4663,18 +4628,7 @@ int32 UMaterialExpressionSubtract::Compile(class FMaterialCompiler* Compiler, in
 
 void UMaterialExpressionSubtract::GetCaption(TArray<FString>& OutCaptions) const
 {
-	FString ret = TEXT("Subtract");
-
-	FExpressionInput ATraced = A.GetTracedInput();
-	FExpressionInput BTraced = B.GetTracedInput();
-	if(!ATraced.Expression || !BTraced.Expression)
-	{
-		ret += TEXT("(");
-		ret += ATraced.Expression ? TEXT(",") : FString::Printf( TEXT("%.4g,"), ConstA);
-		ret += BTraced.Expression ? TEXT(")") : FString::Printf( TEXT("%.4g)"), ConstB);
-	}
-
-	OutCaptions.Add(ret);
+	OutCaptions.Add(TEXT("Subtract"));
 }
 
 int32 UMaterialExpressionSmoothStep::Compile(class FMaterialCompiler* Compiler, int32 OutputIndex)
@@ -4691,21 +4645,7 @@ int32 UMaterialExpressionSmoothStep::Compile(class FMaterialCompiler* Compiler,
 
 void UMaterialExpressionSmoothStep::GetCaption(TArray<FString>& OutCaptions) const
 {
-	FString ret = TEXT("SmoothStep");
-
-	FExpressionInput MinTraced = Min.GetTracedInput();
-	FExpressionInput MaxTraced = Max.GetTracedInput();
-	FExpressionInput ValueTraced = Value.GetTracedInput();
-
-	if (!MinTraced.Expression || !MaxTraced.Expression || !ValueTraced.Expression)
-	{
-		ret += TEXT("(");
-		ret += MinTraced.Expression ? TEXT(",") : FString::Printf(TEXT("%.4g,"), ConstMin);
-		ret += MaxTraced.Expression ? TEXT(",") : FString::Printf(TEXT("%.4g,"), ConstMax);
-		ret += ValueTraced.Expression ? TEXT(")") : FString::Printf(TEXT("%.4g)"), ConstValue);
-	}
-
-	OutCaptions.Add(ret);
+	OutCaptions.Add(TEXT("SmoothStep"));
 }
 #endif // WITH_EDITOR
 
@@ -4725,19 +4665,7 @@ int32 UMaterialExpressionStep::Compile(class FMaterialCompiler* Compiler, int32
 
 void UMaterialExpressionStep::GetCaption(TArray<FString>& OutCaptions) const
 {
-	FString ret = TEXT("Step");
-
-	FExpressionInput YTraced = Y.GetTracedInput();
-	FExpressionInput XTraced = X.GetTracedInput();
-
-	if (!YTraced.Expression || !XTraced.Expression)
-	{
-		ret += TEXT("(");
-		ret += YTraced.Expression ? TEXT(",") : FString::Printf(TEXT("%.4g,"), ConstY);
-		ret += XTraced.Expression ? TEXT(")") : FString::Printf(TEXT("%.4g)"), ConstX);
-	}
-
-	OutCaptions.Add(ret);
+	OutCaptions.Add(TEXT("Step"));
 }
 #endif // WITH_EDITOR
 
@@ -4760,26 +4688,7 @@ int32 UMaterialExpressionInverseLinearInterpolate::Compile(class FMaterialCompil
 
 void UMaterialExpressionInverseLinearInterpolate::GetCaption(TArray<FString>& OutCaptions) const
 {
-	FString ret = TEXT("InvLerp");
-
-	FExpressionInput MinTraced = A.GetTracedInput();
-	FExpressionInput MaxTraced = B.GetTracedInput();
-	FExpressionInput ValueTraced = Value.GetTracedInput();
-
-	if (!MinTraced.Expression || !MaxTraced.Expression || !ValueTraced.Expression)
-	{
-		ret += TEXT("(");
-		ret += MinTraced.Expression ? TEXT(",") : FString::Printf(TEXT("%.4g,"), ConstA);
-		ret += MaxTraced.Expression ? TEXT(",") : FString::Printf(TEXT("%.4g,"), ConstB);
-		ret += ValueTraced.Expression ? TEXT(")") : FString::Printf(TEXT("%.4g)"), ConstValue);
-	}
-
-	if (bClampResult)
-	{
-		ret += TEXT(" Clamped");
-	}
-
-	OutCaptions.Add(ret);
+	OutCaptions.Add(TEXT("InvLerp"));
 }
 #endif // WITH_EDITOR
 
@@ -4801,21 +4710,7 @@ int32 UMaterialExpressionLinearInterpolate::Compile(class FMaterialCompiler* Com
 
 void UMaterialExpressionLinearInterpolate::GetCaption(TArray<FString>& OutCaptions) const
 {
-	FString ret = TEXT("Lerp");
-
-	FExpressionInput ATraced = A.GetTracedInput();
-	FExpressionInput BTraced = B.GetTracedInput();
-	FExpressionInput AlphaTraced = Alpha.GetTracedInput();
-
-	if(!ATraced.Expression || !BTraced.Expression || !AlphaTraced.Expression)
-	{
-		ret += TEXT("(");
-		ret += ATraced.Expression ? TEXT(",") : FString::Printf( TEXT("%.4g,"), ConstA);
-		ret += BTraced.Expression ? TEXT(",") : FString::Printf( TEXT("%.4g,"), ConstB);
-		ret += AlphaTraced.Expression ? TEXT(")") : FString::Printf( TEXT("%.4g)"), ConstAlpha);
-	}
-
-	OutCaptions.Add(ret);
+	OutCaptions.Add(TEXT("Lerp"));
 }
 #endif // WITH_EDITOR
 
@@ -5035,8 +4930,7 @@ int32 UMaterialExpressionClamp::Compile(class FMaterialCompiler* Compiler, int32
 
 void UMaterialExpressionClamp::GetCaption(TArray<FString>& OutCaptions) const
 {
-	FString	NewCaption = TEXT( "Clamp" );
-	OutCaptions.Add(NewCaption);
+	OutCaptions.Add(TEXT("Clamp"));
 }
 #endif // WITH_EDITOR
 
@@ -5081,19 +4975,7 @@ int32 UMaterialExpressionMin::Compile(class FMaterialCompiler* Compiler, int32 O
 
 void UMaterialExpressionMin::GetCaption(TArray<FString>& OutCaptions) const
 {
-	FString ret = TEXT("Min");
-
-	FExpressionInput ATraced = A.GetTracedInput();
-	FExpressionInput BTraced = B.GetTracedInput();
-
-	if(!ATraced.Expression || !BTraced.Expression)
-	{
-		ret += TEXT("(");
-		ret += ATraced.Expression ? TEXT(",") : FString::Printf( TEXT("%.4g,"), ConstA);
-		ret += BTraced.Expression ? TEXT(")") : FString::Printf( TEXT("%.4g)"), ConstB);
-	}
-
-	OutCaptions.Add(ret);
+	OutCaptions.Add(TEXT("Min"));
 }
 #endif // WITH_EDITOR
 
@@ -5113,19 +4995,7 @@ int32 UMaterialExpressionMax::Compile(class FMaterialCompiler* Compiler, int32 O
 
 void UMaterialExpressionMax::GetCaption(TArray<FString>& OutCaptions) const
 {
-	FString ret = TEXT("Max");
-
-	FExpressionInput ATraced = A.GetTracedInput();
-	FExpressionInput BTraced = B.GetTracedInput();
-
-	if(!ATraced.Expression || !BTraced.Expression)
-	{
-		ret += TEXT("(");
-		ret += ATraced.Expression ? TEXT(",") : FString::Printf( TEXT("%.4g,"), ConstA);
-		ret += BTraced.Expression ? TEXT(")") : FString::Printf( TEXT("%.4g)"), ConstB);
-	}
-
-	OutCaptions.Add(ret);
+	OutCaptions.Add(TEXT("Max"));
 }
 #endif // WITH_EDITOR
 
diff --git a/Engine/Source/Runtime/SlateCore/Private/Styling/StarshipCoreStyle.cpp b/Engine/Source/Runtime/SlateCore/Private/Styling/StarshipCoreStyle.cpp
--- a/Engine/Source/Runtime/SlateCore/Private/Styling/StarshipCoreStyle.cpp
+++ b/Engine/Source/Runtime/SlateCore/Private/Styling/StarshipCoreStyle.cpp
@@ -429,6 +429,8 @@ TSharedRef<ISlateStyle> FStarshipCoreStyle::Create()
 		Style->Set("Icons.Merge", new IMAGE_BRUSH_SVG("Starship/Common/Merge", Icon16x16));
 
 		Style->Set("Icons.Calendar", new IMAGE_BRUSH_SVG("Starship/Common/Calendar", Icon16x16));
+
+		Style->Set("Icons.Preview", new IMAGE_BRUSH_SVG("Starship/Common/preview-default", Icon20x20));
 	}
 
 	// Tool panels
