This change removes skeletal mesh asset properties that were incorrectly added to the global UCS modified property list in 5.5.0 and 5.5.1.  Only assets serialized in those versions are affected.  This change also only affects blueprint added skeletal mesh components.  Components added natively were unaffected by the original issue.

diff --git a/Engine/Source/Editor/UnrealEd/Private/Kismet2/BlueprintEditorUtils.cpp b/Engine/Source/Editor/UnrealEd/Private/Kismet2/BlueprintEditorUtils.cpp
--- a/Engine/Source/Editor/UnrealEd/Private/Kismet2/BlueprintEditorUtils.cpp
+++ b/Engine/Source/Editor/UnrealEd/Private/Kismet2/BlueprintEditorUtils.cpp
@@ -137,6 +137,8 @@
 #include "UObject/FastReferenceCollector.h"
 #include "INotifyFieldValueChanged.h"
 
+#include "Components/SkeletalMeshComponent.h"
+
 #define LOCTEXT_NAMESPACE "Blueprint"
 
 DEFINE_LOG_CATEGORY(LogBlueprintDebug);
@@ -600,6 +602,70 @@ FUCSComponentId::FUCSComponentId(const UK2Node_AddComponent* UCSNode)
 
 FBlueprintEditorUtils::FOnRefreshAllNodes FBlueprintEditorUtils::OnRefreshAllNodesEvent;
 
+namespace UE::Private
+{
+bool UNREALED_API GraphContainsSetSkeletalMeshAssetNode(UEdGraph* Graph)
+{
+	struct FNodeGraphTraverser
+	{
+		TSet<UEdGraphNode*> VisitedNodes;
+		
+		bool TraverseGraph(UEdGraph* Graph)
+		{
+			bool bFoundSetSkeletalMeshCall = false;
+			for (int32 Index = 0; Index < Graph->Nodes.Num(); Index++)
+			{
+				UK2Node_CallFunction* FunctionNode = Cast<UK2Node_CallFunction>(Graph->Nodes[Index]);
+				if (FunctionNode && !VisitedNodes.Contains(FunctionNode))
+				{
+					VisitedNodes.Add(FunctionNode);
+
+					const UEdGraphNode* ResultEventNode = nullptr;
+					UEdGraph* FunctionGraph = FunctionNode->GetFunctionGraph(ResultEventNode);
+					if (FunctionGraph)
+					{
+						bFoundSetSkeletalMeshCall = TraverseGraph(FunctionGraph);
+						if (bFoundSetSkeletalMeshCall)
+						{
+							break;
+						}
+					}
+					else if (FunctionNode->FunctionReference.GetMemberName() == GET_FUNCTION_NAME_CHECKED(USkeletalMeshComponent, SetSkeletalMeshAsset))
+					{
+						bFoundSetSkeletalMeshCall = true;
+						break;
+					}
+				}
+				else
+				{
+					UK2Node_MacroInstance* MacroInstanceNode = Cast<UK2Node_MacroInstance>(Graph->Nodes[Index]);
+					if (MacroInstanceNode && !VisitedNodes.Contains(MacroInstanceNode))
+					{
+						VisitedNodes.Add(MacroInstanceNode);
+
+						UEdGraph* MacroGraph = MacroInstanceNode->GetMacroGraph();
+						if (MacroGraph)
+						{
+							bFoundSetSkeletalMeshCall = TraverseGraph(MacroGraph);
+							if (bFoundSetSkeletalMeshCall)
+							{
+								break;
+							}
+						}
+					}
+				}
+				
+			}
+
+			return bFoundSetSkeletalMeshCall;
+		}
+	};
+
+	FNodeGraphTraverser GraphTraverser;
+	return GraphTraverser.TraverseGraph(Graph);	
+}
+}
+
 void FBlueprintEditorUtils::RefreshAllNodes(UBlueprint* Blueprint)
 {
 	if (!Blueprint || (Blueprint->HasAllFlags(RF_WasLoaded) && !Blueprint->HasAllFlags(RF_LoadCompleted)))
diff --git a/Engine/Source/Runtime/Engine/Classes/Components/ActorComponent.h b/Engine/Source/Runtime/Engine/Classes/Components/ActorComponent.h
--- a/Engine/Source/Runtime/Engine/Classes/Components/ActorComponent.h
+++ b/Engine/Source/Runtime/Engine/Classes/Components/ActorComponent.h
@@ -1377,6 +1377,9 @@ private:
 	friend class FComponentReregisterContextBase;
 	friend class FComponentRecreateRenderStateContext;
 	friend struct FActorComponentTickFunction;
+#if WITH_EDITOR
+	friend struct SkeletalMeshUCSModifiedDataCleanupHelper;
+#endif
 
 	//~ Begin Methods for Replicated Members.
 protected:
diff --git a/Engine/Source/Runtime/Engine/Private/Components/SkeletalMeshComponent.cpp b/Engine/Source/Runtime/Engine/Private/Components/SkeletalMeshComponent.cpp
--- a/Engine/Source/Runtime/Engine/Private/Components/SkeletalMeshComponent.cpp
+++ b/Engine/Source/Runtime/Engine/Private/Components/SkeletalMeshComponent.cpp
@@ -52,6 +52,8 @@
 #include "SkeletalMeshComponentInstanceData.h"
 #include "Editor.h"
 #include "Subsystems/AssetEditorSubsystem.h"
+#include "Misc/TransactionallySafeRWScopeLock.h"
+#include "Kismet2/BlueprintEditorUtils.h"
 #endif
 
 #include UE_INLINE_GENERATED_CPP_BY_NAME(SkeletalMeshComponent)
@@ -93,6 +95,14 @@ static TAutoConsoleVariable<int32> CVarCacheLocalSpaceBounds(
 	1,
 	TEXT("If 1 (default) local-space bounds are calculated and cached, otherwise worldspace bounds are built and cached (and inverse transformed to produce local bounds)."));
 
+#if WITH_EDITOR
+static bool GExtendSkeletalMeshDataCleanupRange = false;
+static FAutoConsoleVariableRef CVarExtendSkeletalMeshDataCleanupRange(
+	TEXT("a.ExtendSkeletalMeshDataCleanupRange"),
+	GExtendSkeletalMeshDataCleanupRange,
+	TEXT("Extend the version check for skel mesh components with possible broken skeletal mesh asset properties to include 5.6.0"));
+#endif
+
 DECLARE_CYCLE_STAT_EXTERN(TEXT("Anim Instance Spawn Time"), STAT_AnimSpawnTime, STATGROUP_Anim, );
 DEFINE_STAT(STAT_AnimSpawnTime);
 DEFINE_STAT(STAT_PostAnimEvaluation);
@@ -119,6 +129,85 @@ FAutoConsoleTaskPriority CPrio_ParallelAnimationEvaluationTask(
 /** Static Multicaster fired when SkeletalMeshComponent finalizes the regeneration of the required bones list for the current LOD*/
 /*static*/ FOnLODRequiredBonesUpdateMulticast USkeletalMeshComponent::OnLODRequiredBonesUpdate_Static;
 
+#if WITH_EDITOR
+namespace UE::Private
+{
+	extern UNREALED_API bool GraphContainsSetSkeletalMeshAssetNode(UEdGraph* Graph);
+}
+
+struct SkeletalMeshUCSModifiedDataCleanupHelper
+{
+	static void CleanupData(USkeletalMeshComponent* Component)
+	{
+		static const FName SkeletalMeshPropertyName(TEXT("SkeletalMesh"));
+
+		AActor* OwningActor = Component->GetOwner();
+		if (!OwningActor || !OwningActor->GetLevel())
+		{
+			return;
+		}
+
+		UBlueprint* Blueprint = Cast<UBlueprint>(OwningActor->GetClass()->ClassGeneratedBy);
+		if (!Blueprint)
+		{
+			return;
+		}
+
+		UEdGraph* CSGraph = FBlueprintEditorUtils::FindUserConstructionScript(Blueprint);
+		if (!CSGraph || UE::Private::GraphContainsSetSkeletalMeshAssetNode(CSGraph))
+		{
+			return;
+		}
+		
+		// Check if the property is actually overridden, if not, we can early out
+		UClass* ComponentClass = Component->GetClass();
+		FProperty* MeshAssetProperty = ComponentClass->FindPropertyByName(SkeletalMeshPropertyName);
+		if (MeshAssetProperty)
+		{	
+			UObject* ComponentArchetype = Component->GetArchetype();
+			uint8* DataPtr = MeshAssetProperty->ContainerPtrToValuePtr<uint8>((uint8*)Component);
+			uint8* DefaultValue = MeshAssetProperty->ContainerPtrToValuePtrForDefaults<uint8>(ComponentClass, (uint8*)ComponentArchetype);
+
+			if (MeshAssetProperty->Identical(DataPtr, DefaultValue, PPF_DeepCompareInstances))
+			{
+				return;
+			}
+		}
+
+		bool bFoundUCSEntry = false;
+		TArray<FSimpleMemberReference> UCSModifiedProperties;
+
+		{
+			FTransactionallySafeRWScopeLock Lock(Component->AllUCSModifiedPropertiesLock, SLT_Write);
+			TArray<FSimpleMemberReference>* UCSModifiedPropertiesPtr = Component->AllUCSModifiedProperties.Find(Component);
+			if (UCSModifiedPropertiesPtr)
+			{
+				for (int32 Index = 0; Index < UCSModifiedPropertiesPtr->Num(); Index++)
+				{
+					FSimpleMemberReference& Reference = (*UCSModifiedPropertiesPtr)[Index];
+					if (Reference.MemberName == SkeletalMeshPropertyName)
+					{
+						bFoundUCSEntry = true;
+						UCSModifiedProperties = *UCSModifiedPropertiesPtr;
+						UCSModifiedProperties.RemoveAt(Index);
+						break;
+					}
+				}
+			}
+
+			if (bFoundUCSEntry)
+			{
+				Component->AllUCSModifiedProperties.Remove(Component);
+				if (UCSModifiedProperties.Num() > 0)
+				{
+					Component->AllUCSModifiedProperties.Add(Component, MoveTemp(UCSModifiedProperties));
+				}
+			}
+		}
+	}
+};
+#endif
+
 #if !UE_BUILD_SHIPPING
 CSV_DEFINE_CATEGORY(AnimationParallelEvaluation, true);
 struct FParallelAnimationEvaluationStats
@@ -452,6 +541,29 @@ void USkeletalMeshComponent::Serialize(FArchive& Ar)
 
 	Super::Serialize(Ar);
 
+#if WITH_EDITOR	
+	// Fix for UE-241328. 5.5.0 and 5.5.1 introduced a bug where skeletal mesh asset properties on blueprint added skeletal mesh components would
+	// be incorrectly added to the global UCS modified property list. This affected serialization of overrides of those properties. The fix attempts
+	// to clean up the UCS modified property list, removing the invalid properties and leaving those that are valid (ie. are set via UCS).
+	// It's possible that some packages saved in 5.6.0 could be affected by the bug. This would happen when building from source if an integration
+	// from Main had been made after the original change was introduced in Release-5.5. In these cases, the GExtendSkeletalMeshDataCleanupRange 
+	// cvar can be enabled so the fix up is inclusive of 5.6.0 packages.
+	if (Ar.IsLoading())
+	{
+		if (CreationMethod == EComponentCreationMethod::SimpleConstructionScript || CreationMethod == EComponentCreationMethod::UserConstructionScript)
+		{
+			FEngineVersionBase ArVersion = Ar.EngineVer();
+			uint16 Major = ArVersion.GetMajor();
+			uint16 Minor = ArVersion.GetMinor();
+			uint16 Patch = ArVersion.GetPatch();
+			if (Major == 5 && ((GExtendSkeletalMeshDataCleanupRange && Minor == 6 && Patch == 0) || (Minor == 5 && (Patch == 0 || Patch == 1))))
+			{
+				SkeletalMeshUCSModifiedDataCleanupHelper::CleanupData(this);
+			}
+		}
+	}	
+#endif
+
 	// to count memory : TODO: REMOVE?
 	if (Ar.IsCountingMemory())
 	{
