Fix to account for change in scope lock API since release-5.5

diff --git a/Engine/Source/Runtime/Engine/Private/Components/SkeletalMeshComponent.cpp b/Engine/Source/Runtime/Engine/Private/Components/SkeletalMeshComponent.cpp
--- a/Engine/Source/Runtime/Engine/Private/Components/SkeletalMeshComponent.cpp
+++ b/Engine/Source/Runtime/Engine/Private/Components/SkeletalMeshComponent.cpp
@@ -52,7 +52,7 @@
 #include "SkeletalMeshComponentInstanceData.h"
 #include "Editor.h"
 #include "Subsystems/AssetEditorSubsystem.h"
-#include "Misc/TransactionallySafeRWScopeLock.h"
+#include "Misc/ScopeRWLock.h"
 #include "Kismet2/BlueprintEditorUtils.h"
 #endif
 
@@ -178,7 +178,7 @@ struct SkeletalMeshUCSModifiedDataCleanupHelper
 		TArray<FSimpleMemberReference> UCSModifiedProperties;
 
 		{
-			FTransactionallySafeRWScopeLock Lock(Component->AllUCSModifiedPropertiesLock, SLT_Write);
+			UE::TRWScopeLock Lock(Component->AllUCSModifiedPropertiesLock, SLT_Write);
 			TArray<FSimpleMemberReference>* UCSModifiedPropertiesPtr = Component->AllUCSModifiedProperties.Find(Component);
 			if (UCSModifiedPropertiesPtr)
 			{
