Automate 'Verify that compiling an Animation Blueprint utilizing Inertialization blend logic without an Intertialization node in the Anim Graph displays a warning' test

#rnx

diff --git a/Engine/Source/Runtime/Engine/Classes/Animation/AnimNode_StateMachine.h b/Engine/Source/Runtime/Engine/Classes/Animation/AnimNode_StateMachine.h
--- a/Engine/Source/Runtime/Engine/Classes/Animation/AnimNode_StateMachine.h
+++ b/Engine/Source/Runtime/Engine/Classes/Animation/AnimNode_StateMachine.h
@@ -325,7 +325,7 @@ protected:
 		return GetRelevantAssetPlayerInterfaceFromState(Context.AnimInstanceProxy, StateInfo);
 	}
 
-	ENGINE_API void LogInertializationRequestError(const FAnimationUpdateContext& Context, int32 PreviousState, int32 NextState);
+	ENGINE_API virtual void LogInertializationRequestError(const FAnimationUpdateContext& Context, int32 PreviousState, int32 NextState);
 
 	/** Queues a new transition request, returns true if the transition request was successfully queued */
 	ENGINE_API bool RequestTransitionEvent(const FTransitionEvent& InTransitionEvent);
