HeterogeneousVolume: Update the VolumeResolution property value on tick, if the SVT frame has a different value for it.

diff --git a/Engine/Source/Runtime/Engine/Private/Components/HeterogeneousVolumeComponent.cpp b/Engine/Source/Runtime/Engine/Private/Components/HeterogeneousVolumeComponent.cpp
--- a/Engine/Source/Runtime/Engine/Private/Components/HeterogeneousVolumeComponent.cpp
+++ b/Engine/Source/Runtime/Engine/Private/Components/HeterogeneousVolumeComponent.cpp
@@ -688,6 +688,7 @@ void UHeterogeneousVolumeComponent::TickComponent(float DeltaTime, ELevelTick Ti
 				FIntVector PerFrameVolumeResolution = SparseVolumeTextureFrame->GetVolumeResolution();
 				if (VolumeResolution != PerFrameVolumeResolution)
 				{
+					VolumeResolution = PerFrameVolumeResolution;
 					MarkRenderTransformDirty();
 				}
 
