Refactor material function inputs to track cascading function call changes internally instead; this allows us to enable support for Substrate functions that do not have access to the compiler to map the full material graph structure correctly, whilst also retaining legacy behaviour. Removed bCompilingFunctionPreview which duplicated tracking logic.

Additionally added Substrate support for the Set and Get material functions, as well as a default Substrate slab to support function inputs which are not connected to a preview.

diff --git a/Engine/Source/Runtime/Core/Private/UObject/FortniteMainBranchObjectVersion.cpp b/Engine/Source/Runtime/Core/Private/UObject/FortniteMainBranchObjectVersion.cpp
--- a/Engine/Source/Runtime/Core/Private/UObject/FortniteMainBranchObjectVersion.cpp
+++ b/Engine/Source/Runtime/Core/Private/UObject/FortniteMainBranchObjectVersion.cpp
@@ -14,6 +14,6 @@ TMap<FGuid, FGuid> FFortniteMainBranchObjectVersion::GetSystemGuids()
 	SystemGuids.Add(DevGuids.Niagara_LatestScriptCompileVersion, FGuid("670A2EBDEE1A4794BB77BBFF331EA7EF"));
 	SystemGuids.Add(DevGuids.SkeletalMeshDerivedDataVersion, FGuid("2F90180769024483BB7D90AACEB2798D"));
 	SystemGuids.Add(DevGuids.STATICMESH_DERIVEDDATA_VER, FGuid("656E338C78C34641A5D3BFAE303997AF"));
-	SystemGuids.Add(DevGuids.MaterialTranslationDDCVersion, FGuid("9119259dbf0449a3a53ad39df9fd13a4"));
+	SystemGuids.Add(DevGuids.MaterialTranslationDDCVersion, FGuid("38a1c24b1eb84bccbc79ff8fd7b73a72"));
 	return SystemGuids;
 }
diff --git a/Engine/Source/Runtime/Engine/Private/Materials/MaterialExpressionSubstrate.cpp b/Engine/Source/Runtime/Engine/Private/Materials/MaterialExpressionSubstrate.cpp
--- a/Engine/Source/Runtime/Engine/Private/Materials/MaterialExpressionSubstrate.cpp
+++ b/Engine/Source/Runtime/Engine/Private/Materials/MaterialExpressionSubstrate.cpp
@@ -885,8 +885,7 @@ int32 UMaterialExpressionSubstrateSlabBSDF::Compile(class FMaterialCompiler* Com
 	return OutputCodeChunk;
 }
 
-#if ENABLE_MATERIAL_LAYER_PROTOTYPE
-int32 UMaterialExpressionSubstrateSlabBSDF::CompileDefaultSlab(FMaterialCompiler* Compiler, FVector3f DiffuseAlbedoOverride)
+int32 UMaterialExpressionSubstrateSlabBSDF::CompileDefaultSlab(FMaterialCompiler* Compiler, FVector3f EmissiveOverride)
 {
 	/**
 	 * This creates a default slab, primarily for material function previews where the input may be optional/missing.
@@ -908,7 +907,7 @@ int32 UMaterialExpressionSubstrateSlabBSDF::CompileDefaultSlab(FMaterialCompiler
 	const FSubstrateRegisteredSharedLocalBasis NewRegisteredSharedLocalBasis = SubstrateCompilationInfoCreateSharedLocalBasis(Compiler, NormalCodeChunk, TangentCodeChunk);
 	SubstrateOperator.BSDFRegisteredSharedLocalBasis = NewRegisteredSharedLocalBasis;
 
-	int32 DiffuseAlbedoCodeChunk = Compiler->Constant3(DiffuseAlbedoOverride.X, DiffuseAlbedoOverride.Y, DiffuseAlbedoOverride.Z);
+	int32 DiffuseAlbedoCodeChunk = Compiler->Constant3(0.0f, 0.0f, 0.0f);
 	const float DefaultF0 = 0.04f;
 	int32 F0CodeChunk = Compiler->Constant3(DefaultF0, DefaultF0, DefaultF0);
 	int32 RoughnessCodeChunk = Compiler->Constant(0.5f);
@@ -936,7 +935,7 @@ int32 UMaterialExpressionSubstrateSlabBSDF::CompileDefaultSlab(FMaterialCompiler
 		SSSMFPScaleCodeChunk,
 		SSSPhaseAnisotropyCodeChunk,
 		Compiler->Constant(0.0f),
-		Compiler->Constant3(0.0f, 0.0f, 0.0f),
+		Compiler->Constant3(EmissiveOverride.X, EmissiveOverride.Y, EmissiveOverride.Z),
 		SecondRoughnessCodeChunk,
 		SecondRoughnessWeightCodeChunk,
 		Compiler->Constant(0.0f),										// SecondRoughnessAsSimpleClearCoat
@@ -953,7 +952,6 @@ int32 UMaterialExpressionSubstrateSlabBSDF::CompileDefaultSlab(FMaterialCompiler
 		Compiler->GetSubstrateSharedLocalBasisIndexMacro(NewRegisteredSharedLocalBasis),
 		(SubstrateOperator.Index != INDEX_NONE && SubstrateOperator.BSDFIndex != INDEX_NONE) && (!SubstrateOperator.bUseParameterBlending || (SubstrateOperator.bUseParameterBlending && SubstrateOperator.bRootOfParameterBlendingSubTree)) ? &SubstrateOperator : nullptr);
 }
-#endif //ENABLE_MATERIAL_LAYER_PROTOTYPE
 
 FSubstrateMaterialComplexity UMaterialExpressionSubstrateSlabBSDF::GetHighestComplexity() const
 {
diff --git a/Engine/Source/Runtime/Engine/Private/Materials/MaterialExpressions.cpp b/Engine/Source/Runtime/Engine/Private/Materials/MaterialExpressions.cpp
--- a/Engine/Source/Runtime/Engine/Private/Materials/MaterialExpressions.cpp
+++ b/Engine/Source/Runtime/Engine/Private/Materials/MaterialExpressions.cpp
@@ -6148,13 +6148,6 @@ uint32 UMaterialExpressionGetMaterialAttributes::GetOutputType(int32 OutputIndex
 
 bool UMaterialExpressionGetMaterialAttributes::IsResultSubstrateMaterial(int32 OutputIndex)
 {	
-#if ENABLE_MATERIAL_LAYER_PROTOTYPE
-	if (--OutputIndex >= 0 && OutputIndex < AttributeGetTypes.Num())
-	{
-		return FMaterialAttributeDefinitionMap::GetValueType(AttributeGetTypes[OutputIndex]) == MCT_Substrate;
-	}
-#endif
-
 	if (MaterialAttributes.Expression)
 	{
 		return MaterialAttributes.Expression->IsResultSubstrateMaterial(MaterialAttributes.OutputIndex);
@@ -6164,11 +6157,7 @@ bool UMaterialExpressionGetMaterialAttributes::IsResultSubstrateMaterial(int32 O
 
 void UMaterialExpressionGetMaterialAttributes::GatherSubstrateMaterialInfo(FSubstrateMaterialInfo& SubstrateMaterialInfo, int32 OutputIndex)
 {
-	if (MaterialAttributes.Expression
-#if ENABLE_MATERIAL_LAYER_PROTOTYPE
-	&& IsResultSubstrateMaterial(OutputIndex)
-#endif
-	)
+	if (IsResultSubstrateMaterial(OutputIndex))
 	{
 		MaterialAttributes.Expression->GatherSubstrateMaterialInfo(SubstrateMaterialInfo, MaterialAttributes.OutputIndex);
 	}
@@ -6176,11 +6165,7 @@ void UMaterialExpressionGetMaterialAttributes::GatherSubstrateMaterialInfo(FSubs
 
 FSubstrateOperator* UMaterialExpressionGetMaterialAttributes::SubstrateGenerateMaterialTopologyTree(class FMaterialCompiler* Compiler, class UMaterialExpression* Parent, int32 OutputIndex)
 {
-	if (MaterialAttributes.Expression
-#if ENABLE_MATERIAL_LAYER_PROTOTYPE
-		&& IsResultSubstrateMaterial(OutputIndex)
-#endif
-	)
+	if (IsResultSubstrateMaterial(OutputIndex))
 	{
 		return MaterialAttributes.Expression->SubstrateGenerateMaterialTopologyTree(Compiler, Parent, MaterialAttributes.OutputIndex);
 	}
@@ -6508,55 +6493,50 @@ bool UMaterialExpressionSetMaterialAttributes::ConnectInputAttribute(EMaterialPr
 	}
 	return false;
 }
+#endif // ENABLE_MATERIAL_LAYER_PROTOTYPE
 
 bool UMaterialExpressionSetMaterialAttributes::GetSubstrateMaterialInputIndex(int32 OutputIndex, int32& InputIndex)
 {
-	for (InputIndex = Inputs.Num() - 1; InputIndex >= GET_SET_MA_MATERIALATTRIBUTESINDEX; InputIndex--)
+	for (InputIndex = Inputs.Num() - 1; InputIndex > GET_SET_MA_MATERIALATTRIBUTESINDEX; InputIndex--)
 	{
-		if (GetInputType(InputIndex) == MCT_Substrate)
-		{			
-			return Inputs[InputIndex].IsConnected();
-		}
-		else if (GetInputType(InputIndex) == MCT_MaterialAttributes && Inputs[InputIndex].IsConnected())
+		if (GetInputType(InputIndex) == MCT_Substrate && Inputs[InputIndex].IsConnected())
 		{
-			return Inputs[InputIndex].Expression->IsResultSubstrateMaterial(Inputs[InputIndex].OutputIndex);
+			return true;
 		}
 	}
 
+	//MA input is always in position 0
+	InputIndex = GET_SET_MA_MATERIALATTRIBUTESINDEX;
+	if(GetInputType(InputIndex) == MCT_MaterialAttributes && Inputs[InputIndex].IsConnected())
+	{
+		return Inputs[InputIndex].Expression->IsResultSubstrateMaterial(Inputs[InputIndex].OutputIndex);
+	}	
+
 	return false;
 }
-#endif // ENABLE_MATERIAL_LAYER_PROTOTYPE
 
 bool UMaterialExpressionSetMaterialAttributes::IsResultSubstrateMaterial(int32 OutputIndex)
 {
-#if ENABLE_MATERIAL_LAYER_PROTOTYPE
 	int32 InputIndex = INDEX_NONE;
 	return GetSubstrateMaterialInputIndex(OutputIndex, InputIndex);
-#else
-	return false;
-#endif
 }
 
 void UMaterialExpressionSetMaterialAttributes::GatherSubstrateMaterialInfo(FSubstrateMaterialInfo& SubstrateMaterialInfo, int32 OutputIndex)
 {
-#if ENABLE_MATERIAL_LAYER_PROTOTYPE
 	int32 InputIndex = INDEX_NONE;
 	if(GetSubstrateMaterialInputIndex(OutputIndex, InputIndex))
 	{
 		Inputs[InputIndex].Expression->GatherSubstrateMaterialInfo(SubstrateMaterialInfo, Inputs[InputIndex].OutputIndex);
 	}
-#endif
 }
 
 FSubstrateOperator* UMaterialExpressionSetMaterialAttributes::SubstrateGenerateMaterialTopologyTree(FMaterialCompiler* Compiler, UMaterialExpression* Parent, int32 OutputIndex)
 {
-#if ENABLE_MATERIAL_LAYER_PROTOTYPE
 	int32 InputIndex = INDEX_NONE;
 	if (GetSubstrateMaterialInputIndex(OutputIndex, InputIndex))
 	{
 		return Inputs[InputIndex].Expression->SubstrateGenerateMaterialTopologyTree(Compiler, Parent, Inputs[InputIndex].OutputIndex);
 	}
-#endif
 	return nullptr;
 }
 
@@ -14679,31 +14659,27 @@ void UMaterialFunction::LinkIntoCaller(const TArray<FFunctionExpressionInput>& C
 	// Go through all the function's input expressions and hook their inputs up to the corresponding expression in the material being compiled.
 	for (UMaterialExpression* CurrentExpression : GetExpressions())
 	{
-		UMaterialExpressionFunctionInput* InputExpression = Cast<UMaterialExpressionFunctionInput>(CurrentExpression);
-
-		if (InputExpression)
+		if(UMaterialExpressionFunctionInput* InputExpression = Cast<UMaterialExpressionFunctionInput>(CurrentExpression))
 		{
-			// Mark that we are compiling the function as used in a material
-			InputExpression->bCompilingFunctionPreview = false;
 			// Initialize for this function call
-			InputExpression->EffectivePreviewDuringCompile = InputExpression->Preview;
+			FExpressionInput* CurrentPreview = InputExpression->AddNewEffectivePreviewDuringCompile(InputExpression->Preview);
 
 			// Get the FExpressionInput which stores information about who this input node should be linked to in order to compile
 			const FExpressionInput* MatchingInput = FindInputByExpression(InputExpression, CallerInputs);
 
-			if (MatchingInput 
+			if (CurrentPreview && MatchingInput 
 				// Only change the connection if the input has a valid connection,
-				// Otherwise we will need what's connected to the Preview input if bCompilingFunctionPreview is true
+				// Otherwise we will need what's connected to the Preview input
 				&& (MatchingInput->Expression || !InputExpression->bUsePreviewValueAsDefault))
 			{
 				// Connect this input to the expression in the material that it should be connected to
-				InputExpression->EffectivePreviewDuringCompile.Expression = MatchingInput->Expression;
-				InputExpression->EffectivePreviewDuringCompile.OutputIndex = MatchingInput->OutputIndex;
-				InputExpression->EffectivePreviewDuringCompile.Mask = MatchingInput->Mask;
-				InputExpression->EffectivePreviewDuringCompile.MaskR = MatchingInput->MaskR;
-				InputExpression->EffectivePreviewDuringCompile.MaskG = MatchingInput->MaskG;
-				InputExpression->EffectivePreviewDuringCompile.MaskB = MatchingInput->MaskB;
-				InputExpression->EffectivePreviewDuringCompile.MaskA = MatchingInput->MaskA;		
+				CurrentPreview->Expression = MatchingInput->Expression;
+				CurrentPreview->OutputIndex = MatchingInput->OutputIndex;
+				CurrentPreview->Mask = MatchingInput->Mask;
+				CurrentPreview->MaskR = MatchingInput->MaskR;
+				CurrentPreview->MaskG = MatchingInput->MaskG;
+				CurrentPreview->MaskB = MatchingInput->MaskB;
+				CurrentPreview->MaskA = MatchingInput->MaskA;
 			}
 		}
 	}
@@ -14713,14 +14689,9 @@ void UMaterialFunction::UnlinkFromCaller()
 {
 	for (UMaterialExpression* CurrentExpression : GetExpressions())
 	{
-		UMaterialExpressionFunctionInput* InputExpression = Cast<UMaterialExpressionFunctionInput>(CurrentExpression);
-
-		if (InputExpression)
+		if(UMaterialExpressionFunctionInput* InputExpression = Cast<UMaterialExpressionFunctionInput>(CurrentExpression))
 		{
-			// Restore the default value
-			InputExpression->bCompilingFunctionPreview = true;
-			// Clear the reference to make stale accesses obvious
-			InputExpression->EffectivePreviewDuringCompile.Expression = nullptr;
+			InputExpression->RemoveLastEffectivePreviewDuringCompile();
 		}
 	}
 }
@@ -17101,6 +17072,17 @@ void UMaterialExpressionMaterialFunctionCall::PostEditChangeProperty(FPropertyCh
 void UMaterialExpressionMaterialFunctionCall::LinkFunctionIntoCaller(FMaterialCompiler* Compiler)
 {
 	MaterialFunction->LinkIntoCaller(FunctionInputs);
+	PushParameterOwner(Compiler);
+}
+
+void UMaterialExpressionMaterialFunctionCall::UnlinkFunctionFromCaller(FMaterialCompiler* Compiler)
+{
+	PopParameterOwner(Compiler);
+	MaterialFunction->UnlinkFromCaller();
+}
+
+void UMaterialExpressionMaterialFunctionCall::PushParameterOwner(FMaterialCompiler* Compiler)
+{
 	// Update parameter owner when stepping into layer functions.
 	// This is an optional step when we only want to march the material graph (e.g. to gather Substrate material topology)
 	if (Compiler && MaterialFunction->GetMaterialFunctionUsage() != EMaterialFunctionUsage::Default)
@@ -17109,14 +17091,13 @@ void UMaterialExpressionMaterialFunctionCall::LinkFunctionIntoCaller(FMaterialCo
 	}
 }
 
-void UMaterialExpressionMaterialFunctionCall::UnlinkFunctionFromCaller(FMaterialCompiler* Compiler)
+void UMaterialExpressionMaterialFunctionCall::PopParameterOwner(FMaterialCompiler* Compiler)
 {
 	if (Compiler && MaterialFunction->GetMaterialFunctionUsage() != EMaterialFunctionUsage::Default)
 	{
 		const FMaterialParameterInfo PoppedParameterInfo = Compiler->PopParameterOwner();
 		check(PoppedParameterInfo == FunctionParameterInfo);
 	}
-	MaterialFunction->UnlinkFromCaller();
 }
 
 int32 UMaterialExpressionMaterialFunctionCall::Compile(class FMaterialCompiler* Compiler, int32 OutputIndex)
@@ -17687,7 +17668,6 @@ uint32 UMaterialExpressionMaterialFunctionCall::GetInputType(int32 InputIndex)
 UMaterialExpressionFunctionInput::UMaterialExpressionFunctionInput(const FObjectInitializer& ObjectInitializer)
 	: Super(ObjectInitializer)
 {
-	bCompilingFunctionPreview = true;
 	InputType = FunctionInput_Vector3;
 	InputName = TEXT("In");
 
@@ -17709,6 +17689,30 @@ void UMaterialExpressionFunctionInput::PostDuplicate(bool bDuplicateForPIE)
 }
 
 #if WITH_EDITOR
+FExpressionInput* UMaterialExpressionFunctionInput::AddNewEffectivePreviewDuringCompile(FExpressionInput& InEffectivePreview)
+{
+	PushEffectivePreviewDuringCompile(new FExpressionInput(InEffectivePreview));
+	return EffectivePreviewDuringCompile.Last();
+}
+
+void UMaterialExpressionFunctionInput::PushEffectivePreviewDuringCompile(FExpressionInput* InEffectivePreview)
+{
+	EffectivePreviewDuringCompile.Push(InEffectivePreview);
+}
+
+FExpressionInput* UMaterialExpressionFunctionInput::PopEffectivePreviewDuringCompile()
+{
+	return EffectivePreviewDuringCompile.IsEmpty() ? nullptr : EffectivePreviewDuringCompile.Pop();
+}
+
+void UMaterialExpressionFunctionInput::RemoveLastEffectivePreviewDuringCompile()
+{
+	FExpressionInput* PoppedInput = PopEffectivePreviewDuringCompile();
+	if (PoppedInput)
+	{
+		delete PoppedInput;
+	}
+}
 
 void UMaterialExpressionFunctionInput::PostEditImport()
 {
@@ -17808,24 +17812,20 @@ int32 UMaterialExpressionFunctionInput::CompilePreviewValue(FMaterialCompiler* C
 			return Compiler->Constant4(PreviewValue.X, PreviewValue.Y, PreviewValue.Z, PreviewValue.W);
 		case FunctionInput_MaterialAttributes:
 		{
-#if ENABLE_MATERIAL_LAYER_PROTOTYPE
 			if (AttributeID == FMaterialAttributeDefinitionMap::GetID(MP_EmissiveColor))
 			{
 				return Compiler->Constant3(PreviewValue.X, PreviewValue.Y, PreviewValue.Z);
 			}
 
 			if (!Substrate::IsSubstrateEnabled() || AttributeID != FMaterialAttributeDefinitionMap::GetID(MP_FrontMaterial))
-#endif //ENABLE_MATERIAL_LAYER_PROTOTYPE
 			{
 				return FMaterialAttributeDefinitionMap::CompileDefaultExpression(Compiler, AttributeID);
 			}
 		}
 		case FunctionInput_Substrate:
-#if ENABLE_MATERIAL_LAYER_PROTOTYPE
 		{
 			return UMaterialExpressionSubstrateSlabBSDF::CompileDefaultSlab(Compiler, FVector3f(PreviewValue.X, PreviewValue.Y, PreviewValue.Z));
 		}
-#endif //ENABLE_MATERIAL_LAYER_PROTOTYPE
 		case FunctionInput_Texture2D:
 		case FunctionInput_TextureCube:
 		case FunctionInput_Texture2DArray:
@@ -17860,33 +17860,31 @@ int32 UMaterialExpressionFunctionInput::Compile(class FMaterialCompiler* Compile
 	check(InputType < FunctionInput_MAX);
 
 	// If we are being compiled as part of a material which calls this function
-	FExpressionInput EffectivePreviewDuringCompileTracedInput = EffectivePreviewDuringCompile.GetTracedInput();
+	FExpressionInput* LocalPreviewDuringCompile = PopEffectivePreviewDuringCompile();
 	int32 ExpressionResult = INDEX_NONE;
-	if (EffectivePreviewDuringCompileTracedInput.Expression && !bCompilingFunctionPreview)
+	if (LocalPreviewDuringCompile && LocalPreviewDuringCompile->GetTracedInput().Expression)
 	{
+		FExpressionInput EffectivePreviewDuringCompileTracedInput = LocalPreviewDuringCompile->GetTracedInput();
 		// Stay in this function if we are compiling an expression that is in the current function
 		// This can happen if bUsePreviewValueAsDefault is true and the calling material didn't override the input
 		if (bUsePreviewValueAsDefault && EffectivePreviewDuringCompileTracedInput.Expression->GetOuter() == GetOuter())
 		{
 			// Compile the function input
-			ExpressionResult = EffectivePreviewDuringCompile.Compile(Compiler);
+			ExpressionResult = EffectivePreviewDuringCompileTracedInput.Compile(Compiler);
 		}
 		else
 		{
 			// Tell the compiler that we are leaving the function
 			FMaterialFunctionCompileState* FunctionState = Compiler->PopFunction();
 
-			// Backup EffectivePreviewDuringCompile which will be modified by UnlinkFromCaller and LinkIntoCaller of any potential chained function calls to the same function
-			FExpressionInput LocalPreviewDuringCompile = EffectivePreviewDuringCompile;
-
 			// Restore the function since we are leaving it
-			FunctionState->FunctionCall->UnlinkFunctionFromCaller(Compiler);
+			FunctionState->FunctionCall->PopParameterOwner(Compiler);
 
 			// Compile the function input
-			ExpressionResult = LocalPreviewDuringCompile.Compile(Compiler);
+			ExpressionResult = EffectivePreviewDuringCompileTracedInput.Compile(Compiler);
 
 			// Link the function's inputs into the caller graph before entering
-			FunctionState->FunctionCall->LinkFunctionIntoCaller(Compiler);
+			FunctionState->FunctionCall->PushParameterOwner(Compiler);
 
 			// Tell the compiler that we are re-entering the function
 			Compiler->PushFunction(FunctionState);
@@ -17894,7 +17892,7 @@ int32 UMaterialExpressionFunctionInput::Compile(class FMaterialCompiler* Compile
 	}
 	else
 	{
-		if (bCompilingFunctionPreview || bUsePreviewValueAsDefault)
+		if (!LocalPreviewDuringCompile || bUsePreviewValueAsDefault)
 		{
 			// If we are compiling the function in a preview material, such as when editing the function,
 			// Compile the preview value or texture and output a texture object.
@@ -17906,6 +17904,11 @@ int32 UMaterialExpressionFunctionInput::Compile(class FMaterialCompiler* Compile
 		}
 	}
 
+	if (LocalPreviewDuringCompile)
+	{
+		PushEffectivePreviewDuringCompile(LocalPreviewDuringCompile);
+	}
+
 	if (ExpressionResult != INDEX_NONE)
 	{
 		// Cast to the type that the function author specified
@@ -17997,63 +18000,100 @@ bool UMaterialExpressionFunctionInput::IsResultMaterialAttributes(int32 OutputIn
 
 bool UMaterialExpressionFunctionInput::IsResultSubstrateMaterial(int32 OutputIndex)
 {
+	bool bResult = false;
 	if (FunctionInput_Substrate == InputType)
 	{
-		return true;
+		bResult = true;
 	}
-	else
+	else if(IsResultMaterialAttributes(OutputIndex))
 	{
-#if ENABLE_MATERIAL_LAYER_PROTOTYPE
-		//This ensures the default Substrate operator is always generated for layering, reflecting the behaviour in SubstrateGenerateMaterialTopologyTree
-		return Substrate::IsSubstrateEnabled() && IsResultMaterialAttributes(OutputIndex);
-#else
-		return false;
-#endif
+		FExpressionInput* LocalPreviewDuringCompile = PopEffectivePreviewDuringCompile();
+		if (LocalPreviewDuringCompile && LocalPreviewDuringCompile->GetTracedInput().Expression)
+		{
+			bResult = LocalPreviewDuringCompile->GetTracedInput().Expression->IsResultSubstrateMaterial(LocalPreviewDuringCompile->OutputIndex);
+		}
+		else if(bUsePreviewValueAsDefault)
+		{
+			if (Preview.Expression)
+			{
+				bResult = Preview.Expression->IsResultSubstrateMaterial(Preview.OutputIndex);
+			}
+			else
+			{
+				//Ensure default values force slab generation for MA type inputs.
+				bResult = true;
+			}
+		}
+
+		if (LocalPreviewDuringCompile)
+		{
+			PushEffectivePreviewDuringCompile(LocalPreviewDuringCompile);
+		}
 	}
+	return bResult;
 }
 
 void UMaterialExpressionFunctionInput::GatherSubstrateMaterialInfo(FSubstrateMaterialInfo& SubstrateMaterialInfo, int32 OutputIndex)
 {
-#if ENABLE_MATERIAL_LAYER_PROTOTYPE
 	if (IsResultSubstrateMaterial(OutputIndex))
-#else
-	if (FunctionInput_Substrate == InputType)
-#endif
 	{
-		FExpressionInput EffectivePreviewDuringCompileTracedInput = EffectivePreviewDuringCompile.GetTracedInput();
-		if (EffectivePreviewDuringCompileTracedInput.GetTracedInput().Expression)
+		FExpressionInput* LocalPreviewDuringCompile = PopEffectivePreviewDuringCompile();
+		if (LocalPreviewDuringCompile && LocalPreviewDuringCompile->GetTracedInput().Expression)
+		{
+			LocalPreviewDuringCompile->GetTracedInput().Expression->GatherSubstrateMaterialInfo(SubstrateMaterialInfo, LocalPreviewDuringCompile->OutputIndex);
+		}
+		else if(Preview.Expression && bUsePreviewValueAsDefault)
+		{
+			Preview.Expression->GatherSubstrateMaterialInfo(SubstrateMaterialInfo, Preview.OutputIndex);
+		}
+
+		if (LocalPreviewDuringCompile)
+		{
+			PushEffectivePreviewDuringCompile(LocalPreviewDuringCompile);
+		}
+
+		if (SubstrateMaterialInfo.IsValid())
 		{
-			EffectivePreviewDuringCompileTracedInput.Expression->GatherSubstrateMaterialInfo(SubstrateMaterialInfo, EffectivePreviewDuringCompileTracedInput.OutputIndex);
 			return;
 		}
 	}
-#if ENABLE_MATERIAL_LAYER_PROTOTYPE
 	SubstrateMaterialInfo.AddShadingModel(SSM_DefaultLit);
 	SubstrateMaterialInfo.AddGuid(MaterialExpressionGuid);
-#endif
 }
 
 FSubstrateOperator* UMaterialExpressionFunctionInput::SubstrateGenerateMaterialTopologyTree(class FMaterialCompiler* Compiler, class UMaterialExpression* Parent, int32 OutputIndex)
 {
-#if ENABLE_MATERIAL_LAYER_PROTOTYPE
+	FSubstrateOperator* SlabOperator = nullptr;
 	if (IsResultSubstrateMaterial(OutputIndex))
-#else
-	if (FunctionInput_Substrate == InputType)
-#endif
 	{
-		FExpressionInput EffectivePreviewDuringCompileTracedInput = EffectivePreviewDuringCompile.GetTracedInput();
-		if (EffectivePreviewDuringCompileTracedInput.GetTracedInput().Expression)
+		FExpressionInput* LocalPreviewDuringCompile = PopEffectivePreviewDuringCompile();
+		if (LocalPreviewDuringCompile && LocalPreviewDuringCompile->GetTracedInput().Expression)
+		{
+			// Backup EffectivePreviewDuringCompile which will be modified by UnlinkFromCaller and LinkIntoCaller of any potential chained function calls to the same function
+			FExpressionInput LocalPreviewTracedInput = LocalPreviewDuringCompile->GetTracedInput();
+			SlabOperator = LocalPreviewTracedInput.Expression->SubstrateGenerateMaterialTopologyTree(Compiler, Parent, LocalPreviewTracedInput.OutputIndex);
+		}
+		else if(Preview.Expression && bUsePreviewValueAsDefault)
+		{
+			SlabOperator = Preview.Expression->SubstrateGenerateMaterialTopologyTree(Compiler, Parent, Preview.OutputIndex);
+		}
+
+		if (LocalPreviewDuringCompile)
 		{
-			return EffectivePreviewDuringCompileTracedInput.Expression->SubstrateGenerateMaterialTopologyTree(Compiler, Parent, EffectivePreviewDuringCompileTracedInput.OutputIndex);
+			PushEffectivePreviewDuringCompile(LocalPreviewDuringCompile);
 		}
 	}
+
 	// If we are parsing for a material function input we always needs to return a default valid BSDF operator at least 
 	// If the material requires an input thenthe UI will forcethe user to proviced one.
 	// If this is not the case however, we need to compile a default material to compile the shader, or to preview the material function.
-	FSubstrateOperator& DefaultSlabOperator = Compiler->SubstrateCompilationRegisterOperator(SUBSTRATE_OPERATOR_BSDF, Compiler->SubstrateTreeStackGetPathUniqueId(), this, Parent, Compiler->SubstrateTreeStackGetParentPathUniqueId());
-	DefaultSlabOperator.BSDFType = SUBSTRATE_BSDF_TYPE_SLAB;
-	DefaultSlabOperator.ThicknessIndex = Compiler->SubstrateThicknessStackGetThicknessIndex();
-	return &DefaultSlabOperator;
+	if(!SlabOperator)
+	{
+		SlabOperator = &Compiler->SubstrateCompilationRegisterOperator(SUBSTRATE_OPERATOR_BSDF, Compiler->SubstrateTreeStackGetPathUniqueId(), this, Parent, Compiler->SubstrateTreeStackGetParentPathUniqueId());
+		SlabOperator->BSDFType = SUBSTRATE_BSDF_TYPE_SLAB;
+		SlabOperator->ThicknessIndex = Compiler->SubstrateThicknessStackGetThicknessIndex();
+	}
+	return SlabOperator;
 }
 
 uint32 UMaterialExpressionFunctionInput::GetInputType(int32 InputIndex)
diff --git a/Engine/Source/Runtime/Engine/Public/Materials/MaterialExpressionFunctionInput.h b/Engine/Source/Runtime/Engine/Public/Materials/MaterialExpressionFunctionInput.h
--- a/Engine/Source/Runtime/Engine/Public/Materials/MaterialExpressionFunctionInput.h
+++ b/Engine/Source/Runtime/Engine/Public/Materials/MaterialExpressionFunctionInput.h
@@ -73,23 +73,17 @@ class UMaterialExpressionFunctionInput : public UMaterialExpression
 	UPROPERTY(EditAnywhere, BlueprintReadWrite, Category=MaterialExpression)
 	int32 SortPriority;
 
-	/** 
-	 * true when this expression is being compiled in a function preview, 
-	 * false when this expression is being compiled into a material that uses the function.
-	 * Only valid in Compile()
-	 */
-	UPROPERTY(transient)
-	uint32 bCompilingFunctionPreview:1;
-
-	/** The Preview input to use during compilation from another material, when bCompilingFunctionPreview is false. */
-	FExpressionInput EffectivePreviewDuringCompile;
+	UE_DEPRECATED(5.6, "bCompilingFunctionPreview has been deprecated.")
+	UPROPERTY(transient, meta = (DeprecatedProperty, DeprecationMessage = "bCompilingFunctionPreview has been removed from function inputs and has no functional usage, it has been replaced by automatic state tracking between function calls."))
+	uint32 bCompilingFunctionPreview_DEPRECATED : 1;
 
 #if WITH_EDITOR
-	/** Returns the appropriate preview expression when compiling a function or material preview. */
-	UMaterialExpression* GetEffectivePreviewExpression()
-	{
-		return bCompilingFunctionPreview ? Preview.Expression : EffectivePreviewDuringCompile.Expression;
-	}
+	/** 
+	* These create and delete new FExpression objects on the heap, but are public to allow creation and removal from external systems. 
+	* Use the corresponding pop and push to manage internal behaviour of EffectivePreviews within this class.
+	*/
+	FExpressionInput* AddNewEffectivePreviewDuringCompile(FExpressionInput& InEffectivePreview);
+	void RemoveLastEffectivePreviewDuringCompile();
 #endif
 
 	//~ Begin UObject Interface.
@@ -129,13 +123,19 @@ class UMaterialExpressionFunctionInput : public UMaterialExpression
 
 
 private:
-
 #if WITH_EDITOR
 	/** Helper function which compiles this expression for previewing. */
 	int32 CompilePreviewValue(FMaterialCompiler* Compiler);
 
 	/** Stashed data between a Pre/PostEditChange event */
 	FName InputNameBackup;
+
+	/** Cascading preview inputs to use when compiling from another material and/or function */
+	TArray<FExpressionInput*> EffectivePreviewDuringCompile;
+
+	/** Push and pop management of the effective previews should be restricted to internal functions only to avoid memory leaks. */
+	void PushEffectivePreviewDuringCompile(FExpressionInput* InEffectivePreview);
+	FExpressionInput* PopEffectivePreviewDuringCompile();
 #endif // WITH_EDITOR
 };
 
diff --git a/Engine/Source/Runtime/Engine/Public/Materials/MaterialExpressionMaterialFunctionCall.h b/Engine/Source/Runtime/Engine/Public/Materials/MaterialExpressionMaterialFunctionCall.h
--- a/Engine/Source/Runtime/Engine/Public/Materials/MaterialExpressionMaterialFunctionCall.h
+++ b/Engine/Source/Runtime/Engine/Public/Materials/MaterialExpressionMaterialFunctionCall.h
@@ -113,6 +113,9 @@ class UMaterialExpressionMaterialFunctionCall : public UMaterialExpression
 #if WITH_EDITOR
 	ENGINE_API void UnlinkFunctionFromCaller(FMaterialCompiler* Compiler);
 	ENGINE_API void LinkFunctionIntoCaller(FMaterialCompiler* Compiler);
+
+	ENGINE_API void PushParameterOwner(FMaterialCompiler* Compiler);
+	ENGINE_API void PopParameterOwner(FMaterialCompiler* Compiler);
 #endif
 
 	//~ Begin UMaterialExpression Interface
diff --git a/Engine/Source/Runtime/Engine/Public/Materials/MaterialExpressionSubstrate.h b/Engine/Source/Runtime/Engine/Public/Materials/MaterialExpressionSubstrate.h
--- a/Engine/Source/Runtime/Engine/Public/Materials/MaterialExpressionSubstrate.h
+++ b/Engine/Source/Runtime/Engine/Public/Materials/MaterialExpressionSubstrate.h
@@ -347,9 +347,7 @@ class UMaterialExpressionSubstrateSlabBSDF : public UMaterialExpressionSubstrate
 	//~ Begin UMaterialExpression Interface
 #if WITH_EDITOR
 	virtual int32 Compile(class FMaterialCompiler* Compiler, int32 OutputIndex) override;
-#if ENABLE_MATERIAL_LAYER_PROTOTYPE
-	static int32 CompileDefaultSlab(class FMaterialCompiler* Compiler, FVector3f DiffuseAlbedoOverride = {0.18f, 0.18f, 0.18f});
-#endif //ENABLE_MATERIAL_LAYER_PROTOTYPE
+	static int32 CompileDefaultSlab(class FMaterialCompiler* Compiler, FVector3f EmissiveOverride = FVector3f::Zero());
 	virtual void GetCaption(TArray<FString>& OutCaptions) const override;
 	virtual uint32 GetOutputType(int32 OutputIndex) override;
 	virtual uint32 GetInputType(int32 InputIndex) override;
