Fix bug of repcomponent not being auto added when the entity was outer to the transient package.

#rnx

diff --git a/Engine/Source/Runtime/Engine/Private/Tests/NetTestHelpers.cpp b/Engine/Source/Runtime/Engine/Private/Tests/NetTestHelpers.cpp
--- a/Engine/Source/Runtime/Engine/Private/Tests/NetTestHelpers.cpp
+++ b/Engine/Source/Runtime/Engine/Private/Tests/NetTestHelpers.cpp
@@ -485,7 +485,7 @@ UObject* FTestWorlds::FindReplicatedObjectOnClient(UObject* ServerObject, uint32
 		return nullptr;
 	}
 
-	if (ServerObject->GetWorld() != Server.GetWorld())
+	if (ServerObject->GetWorld() != Server.GetWorld() && ServerObject->GetOuter() != GetTransientPackage())
 	{
 		ensureMsgf(false, TEXT("FTestWorlds::FindReplicatedObjectOnClient received object %s not part of the Server world"), *GetFullNameSafe(ServerObject));
 		return nullptr;
