Fixed ReplicationSystemTest compilation error with FNetIDVariant
Added FNetIDVariant::ToString

diff --git a/Engine/Plugins/Runtime/ReplicationSystemTestPlugin/Source/Private/Tests/Engine/TestNetIDVariant.cpp b/Engine/Plugins/Runtime/ReplicationSystemTestPlugin/Source/Private/Tests/Engine/TestNetIDVariant.cpp
--- a/Engine/Plugins/Runtime/ReplicationSystemTestPlugin/Source/Private/Tests/Engine/TestNetIDVariant.cpp
+++ b/Engine/Plugins/Runtime/ReplicationSystemTestPlugin/Source/Private/Tests/Engine/TestNetIDVariant.cpp
@@ -6,6 +6,21 @@
 #include "Net/NetIDVariant.h"
 #include "UObject/CoreNet.h"
 
+namespace UE::Net
+{
+
+FTestMessage& operator<<(FTestMessage& Message, const FNetIDVariant& NetID)
+{
+	return Message << *NetID.ToString();
+}
+
+FTestMessage& operator<<(FTestMessage& Message, const FNetworkGUID& NetGUID)
+{
+	return Message << *NetGUID.ToString();
+}
+
+}
+
 namespace UE::Net::Private
 {
 	UE_NET_TEST(NetIDVariant, Empty)
diff --git a/Engine/Source/Runtime/Engine/Private/Net/NetIDVariant.cpp b/Engine/Source/Runtime/Engine/Private/Net/NetIDVariant.cpp
--- a/Engine/Source/Runtime/Engine/Private/Net/NetIDVariant.cpp
+++ b/Engine/Source/Runtime/Engine/Private/Net/NetIDVariant.cpp
@@ -18,7 +18,7 @@ FNetIDVariant::FNetIDVariant(FNetRefHandle NetRefHandle)
 }
 #endif
 
-bool FNetIDVariant::operator==(FNetIDVariant RHS) const
+bool FNetIDVariant::operator==(const FNetIDVariant& RHS) const
 {
 	if (Variant.GetIndex() != RHS.Variant.GetIndex())
 	{
@@ -117,4 +117,19 @@ FArchive& operator<<(FArchive& Ar, FNetIDVariant& NetID)
 	return Ar;
 }
 
+FString FNetIDVariant::ToString() const
+{
+	if (Variant.IsType<FNetworkGUID>())
+	{
+		return Variant.Get<FNetworkGUID>().ToString();
+	}
+#if UE_WITH_IRIS
+	else if (Variant.IsType<FNetRefHandle>())
+	{
+		return Variant.Get<FNetRefHandle>().ToString();
+	}
+#endif
+	return FString(TEXT("Invalid"));
+}
+
 }
\ No newline at end of file
diff --git a/Engine/Source/Runtime/Engine/Public/Net/NetIDVariant.h b/Engine/Source/Runtime/Engine/Public/Net/NetIDVariant.h
--- a/Engine/Source/Runtime/Engine/Public/Net/NetIDVariant.h
+++ b/Engine/Source/Runtime/Engine/Public/Net/NetIDVariant.h
@@ -51,7 +51,10 @@ public:
 	}
 
 	/** Equality comparison */
-	bool ENGINE_API operator==(FNetIDVariant RHS) const;
+	bool ENGINE_API operator==(const FNetIDVariant& RHS) const;
+
+	/** Output to string */
+	ENGINE_API FString ToString() const;
 
 private:
 	FVariantType Variant;
