Prevent a rare crash if SetupInstancedMeshBatch is called before CreateRenderThreadResources creates LODLooseUniformBuffers.

diff --git a/Engine/Source/Runtime/Engine/Classes/Engine/InstancedStaticMesh.h b/Engine/Source/Runtime/Engine/Classes/Engine/InstancedStaticMesh.h
--- a/Engine/Source/Runtime/Engine/Classes/Engine/InstancedStaticMesh.h
+++ b/Engine/Source/Runtime/Engine/Classes/Engine/InstancedStaticMesh.h
@@ -561,6 +561,7 @@ protected:
 private:
 
 	void SetupProxy(const FInstancedStaticMeshSceneProxyDesc& InProxyDesc);
+	bool CanSetupInstancedMeshBatch(int32 LODIndex) const;
 
 	/** Stores a loose uniform buffer per LOD, used for static view relevance. */
 	TMap<uint32, FInstancedStaticMeshVFLooseUniformShaderParametersRef> LODLooseUniformBuffers;
diff --git a/Engine/Source/Runtime/Engine/Private/InstancedStaticMesh.cpp b/Engine/Source/Runtime/Engine/Private/InstancedStaticMesh.cpp
--- a/Engine/Source/Runtime/Engine/Private/InstancedStaticMesh.cpp
+++ b/Engine/Source/Runtime/Engine/Private/InstancedStaticMesh.cpp
@@ -1516,6 +1516,12 @@ void FInstancedStaticMeshSceneProxy::SetupProxy(const FInstancedStaticMeshSceneP
 	}
 }
 
+bool FInstancedStaticMeshSceneProxy::CanSetupInstancedMeshBatch(int32 LODIndex) const
+{
+	// TODO: If this fires somehow our LODLooseUniformBuffers created in CreateRenderThreadResources is empty, a full dmp should be investigated
+	check(LODLooseUniformBuffers.Num() == InstancedRenderData.VertexFactories.Num());
+	return LODIndex < LODLooseUniformBuffers.Num() && LODIndex < InstancedRenderData.VertexFactories.Num();
+}
 
 void FInstancedStaticMeshSceneProxy::CreateRenderThreadResources(FRHICommandListBase& RHICmdList)
 {
@@ -1599,7 +1605,7 @@ void FInstancedStaticMeshSceneProxy::GetLightRelevance(const FLightSceneProxy* L
 
 bool FInstancedStaticMeshSceneProxy::GetShadowMeshElement(int32 LODIndex, int32 BatchIndex, uint8 InDepthPriorityGroup, FMeshBatch& OutMeshBatch, bool bDitheredLODTransition) const
 {
-	if (LODIndex < InstancedRenderData.VertexFactories.Num() && FStaticMeshSceneProxy::GetShadowMeshElement(LODIndex, BatchIndex, InDepthPriorityGroup, OutMeshBatch, bDitheredLODTransition))
+	if (CanSetupInstancedMeshBatch(LODIndex) && FStaticMeshSceneProxy::GetShadowMeshElement(LODIndex, BatchIndex, InDepthPriorityGroup, OutMeshBatch, bDitheredLODTransition))
 	{
 		SetupInstancedMeshBatch(LODIndex, BatchIndex, OutMeshBatch);
 		return true;
@@ -1610,7 +1616,7 @@ bool FInstancedStaticMeshSceneProxy::GetShadowMeshElement(int32 LODIndex, int32
 /** Sets up a FMeshBatch for a specific LOD and element. */
 bool FInstancedStaticMeshSceneProxy::GetMeshElement(int32 LODIndex, int32 BatchIndex, int32 ElementIndex, uint8 InDepthPriorityGroup, bool bUseSelectionOutline, bool bAllowPreCulledIndices, FMeshBatch& OutMeshBatch) const
 {
-	if (LODIndex < InstancedRenderData.VertexFactories.Num() && FStaticMeshSceneProxy::GetMeshElement(LODIndex, BatchIndex, ElementIndex, InDepthPriorityGroup, bUseSelectionOutline, bAllowPreCulledIndices, OutMeshBatch))
+	if (CanSetupInstancedMeshBatch(LODIndex) && FStaticMeshSceneProxy::GetMeshElement(LODIndex, BatchIndex, ElementIndex, InDepthPriorityGroup, bUseSelectionOutline, bAllowPreCulledIndices, OutMeshBatch))
 	{
 		SetupInstancedMeshBatch(LODIndex, BatchIndex, OutMeshBatch);
 		return true;
@@ -1621,7 +1627,7 @@ bool FInstancedStaticMeshSceneProxy::GetMeshElement(int32 LODIndex, int32 BatchI
 /** Sets up a wireframe FMeshBatch for a specific LOD. */
 bool FInstancedStaticMeshSceneProxy::GetWireframeMeshElement(int32 LODIndex, int32 BatchIndex, const FMaterialRenderProxy* WireframeRenderProxy, uint8 InDepthPriorityGroup, bool bAllowPreCulledIndices, FMeshBatch& OutMeshBatch) const
 {
-	if (LODIndex < InstancedRenderData.VertexFactories.Num() && FStaticMeshSceneProxy::GetWireframeMeshElement(LODIndex, BatchIndex, WireframeRenderProxy, InDepthPriorityGroup, bAllowPreCulledIndices, OutMeshBatch))
+	if (CanSetupInstancedMeshBatch(LODIndex) && FStaticMeshSceneProxy::GetWireframeMeshElement(LODIndex, BatchIndex, WireframeRenderProxy, InDepthPriorityGroup, bAllowPreCulledIndices, OutMeshBatch))
 	{
 		SetupInstancedMeshBatch(LODIndex, BatchIndex, OutMeshBatch);
 		return true;
