[NewMaterialTranslator] - Implement Square Root and Saturate

- The old translator emitted constants with 8 decimals points of precision (see int32 FHLSLMaterialTranslator::Constant(float X)), update the new translator to do the same.
- Remove extra new lines in between each material attribute when writing out HLSL to more match the generated code from the old translator.

#rnx

diff --git a/Engine/Source/Runtime/Engine/Private/Materials/MaterialExpressionsIR.cpp b/Engine/Source/Runtime/Engine/Private/Materials/MaterialExpressionsIR.cpp
--- a/Engine/Source/Runtime/Engine/Private/Materials/MaterialExpressionsIR.cpp
+++ b/Engine/Source/Runtime/Engine/Private/Materials/MaterialExpressionsIR.cpp
@@ -578,6 +578,16 @@ void UMaterialExpressionSine::Build(UE::MIR::FEmitter& Emitter)
 	BuildUnaryArithmeticOperatorFromInput(Emitter, MIR::UO_Sin, &Input);
 }
 
+void UMaterialExpressionSaturate::Build(UE::MIR::FEmitter& Emitter)
+{
+	BuildUnaryArithmeticOperatorFromInput(Emitter, MIR::UO_Saturate, &Input);
+}
+
+void UMaterialExpressionSquareRoot::Build(UE::MIR::FEmitter& Emitter)
+{
+	BuildUnaryArithmeticOperatorFromInput(Emitter, MIR::UO_Sqrt, &Input);
+}
+
 void UMaterialExpressionTangent::Build(UE::MIR::FEmitter& Emitter)
 {
 	BuildUnaryArithmeticOperatorFromInput(Emitter, MIR::UO_Tan, &Input);
diff --git a/Engine/Source/Runtime/Engine/Private/Materials/MaterialIR.cpp b/Engine/Source/Runtime/Engine/Private/Materials/MaterialIR.cpp
--- a/Engine/Source/Runtime/Engine/Private/Materials/MaterialIR.cpp
+++ b/Engine/Source/Runtime/Engine/Private/Materials/MaterialIR.cpp
@@ -423,7 +423,9 @@ const TCHAR* OperatorToString(EOperator Op)
 		case UO_Logarithm2: return TEXT("Logarithm2");
 		case UO_Logarithm10: return TEXT("Logarithm10");
 		case UO_Round: return TEXT("Round");
+		case UO_Saturate: return TEXT("Saturate");
 		case UO_Sin: return TEXT("Sine");
+		case UO_Sqrt: return TEXT("Sqrt");
 		case UO_Tan: return TEXT("Tangent");
 		case UO_Truncate: return TEXT("Truncate");
 
diff --git a/Engine/Source/Runtime/Engine/Private/Materials/MaterialIREmitter.cpp b/Engine/Source/Runtime/Engine/Private/Materials/MaterialIREmitter.cpp
--- a/Engine/Source/Runtime/Engine/Private/Materials/MaterialIREmitter.cpp
+++ b/Engine/Source/Runtime/Engine/Private/Materials/MaterialIREmitter.cpp
@@ -821,7 +821,9 @@ static TOptional<T> FoldScalarArithmeticOperator(FEmitter* Emitter, EOperator Op
 				{
 					return FMath::RoundToDouble(A);
 				}
+			case UO_Saturate: return FMath::Min<T>(1.0f, FMath::Max<T>(0.0f, A));
 			case UO_Sin: return FMath::Sin(A);
+			case UO_Sqrt: return FMath::Sqrt(A);
 			case UO_Tan: return FMath::Tan(A);
 			case UO_Truncate:
 				if constexpr (std::is_same<T, float>::value)
@@ -947,7 +949,9 @@ static FValue* TryFoldOperatorScalar(FEmitter* Emitter, EOperator Op, FValue* A,
 		case UO_Negate:
 		case UO_Not:
 		case UO_Round:
+		case UO_Saturate:
 		case UO_Sin:
+		case UO_Sqrt:
 		case UO_Tan:
 		case UO_Truncate:
 			break;
@@ -1089,7 +1093,9 @@ static const FOperatorSignature* GetOperatorSignature(EOperator Op)
 		S[UO_Logarithm2] 			= UnaryArithmetic;
 		S[UO_Logarithm10] 			= UnaryArithmetic;
 		S[UO_Round] 				= UnaryArithmetic;
+		S[UO_Saturate]				= UnaryArithmetic;
 		S[UO_Sin] 					= UnaryArithmetic;
+		S[UO_Sqrt]					= UnaryArithmetic;
 		S[UO_Tan] 					= UnaryArithmetic;
 		S[UO_Truncate] 				= UnaryArithmetic;
 		S[BO_GreaterThan] 			= BinaryComparison;
diff --git a/Engine/Source/Runtime/Engine/Private/Materials/MaterialIRToHLSLTranslator.cpp b/Engine/Source/Runtime/Engine/Private/Materials/MaterialIRToHLSLTranslator.cpp
--- a/Engine/Source/Runtime/Engine/Private/Materials/MaterialIRToHLSLTranslator.cpp
+++ b/Engine/Source/Runtime/Engine/Private/Materials/MaterialIRToHLSLTranslator.cpp
@@ -296,11 +296,6 @@ struct FTranslator : FMaterialIRToHLSLTranslation
 			{
 				Printer << EndOfStatement;
 			}
-
-			if (Instr->Kind == MIR::VK_SetMaterialOutput)
-			{
-				Printer << NewLine;
-			}
         }
 
         NumLocals = OldNumLocals;
@@ -363,7 +358,7 @@ struct FTranslator : FMaterialIRToHLSLTranslation
 		{
 			case MIR::ScalarKind_Bool:  Printer.Buffer.Append(Constant->Boolean ? TEXT("true") : TEXT("false")); break;
 			case MIR::ScalarKind_Int:   Printer.Buffer.Appendf(TEXT("%") PRId64, Constant->Integer); break;
-			case MIR::ScalarKind_Float: Printer.Buffer.Appendf(TEXT("%.5ff"), Constant->Float); break;
+			case MIR::ScalarKind_Float: Printer.Buffer.Appendf(TEXT("%.8ff"), Constant->Float); break;
 		}
 	}
 
@@ -491,7 +486,9 @@ struct FTranslator : FMaterialIRToHLSLTranslation
 				case MIR::UO_Floor: OpString = TEXT("floor"); break;
 				case MIR::UO_Frac: OpString = TEXT("frac"); break;
 				case MIR::UO_Length: OpString = TEXT("length"); break;
+				case MIR::UO_Saturate: OpString = TEXT("saturate"); break;
 				case MIR::UO_Sin: OpString = TEXT("sin"); break;
+				case MIR::UO_Sqrt: OpString = TEXT("sqrt"); break;
 				case MIR::UO_Tan: OpString = TEXT("tan"); break;
 				case MIR::UO_Truncate: OpString = TEXT("trunc"); break;
 				case MIR::UO_Exponential: OpString = TEXT("exp"); break;
diff --git a/Engine/Source/Runtime/Engine/Public/Materials/MaterialExpressionSaturate.h b/Engine/Source/Runtime/Engine/Public/Materials/MaterialExpressionSaturate.h
--- a/Engine/Source/Runtime/Engine/Public/Materials/MaterialExpressionSaturate.h
+++ b/Engine/Source/Runtime/Engine/Public/Materials/MaterialExpressionSaturate.h
@@ -20,6 +20,7 @@ public:
 
 	//~ Begin UMaterialExpression Interface
 #if WITH_EDITOR
+	virtual void Build(UE::MIR::FEmitter& Emitter) override;
 	virtual int32 Compile(class FMaterialCompiler* Compiler, int32 OutputIndex) override;
 	virtual void GetCaption(TArray<FString>& OutCaptions) const override;
 	virtual void GetExpressionToolTip(TArray<FString>& OutToolTip) override;
diff --git a/Engine/Source/Runtime/Engine/Public/Materials/MaterialExpressionSquareRoot.h b/Engine/Source/Runtime/Engine/Public/Materials/MaterialExpressionSquareRoot.h
--- a/Engine/Source/Runtime/Engine/Public/Materials/MaterialExpressionSquareRoot.h
+++ b/Engine/Source/Runtime/Engine/Public/Materials/MaterialExpressionSquareRoot.h
@@ -22,6 +22,7 @@ public:
 
 	//~ Begin UMaterialExpression Interface
 #if WITH_EDITOR
+	virtual void Build(UE::MIR::FEmitter& Emitter) override;
 	virtual int32 Compile(class FMaterialCompiler* Compiler, int32 OutputIndex) override;
 	virtual void GetCaption(TArray<FString>& OutCaptions) const override;
 	virtual FText GetKeywords() const override {return FText::FromString(TEXT("sqrt"));}
diff --git a/Engine/Source/Runtime/Engine/Public/Materials/MaterialIR.h b/Engine/Source/Runtime/Engine/Public/Materials/MaterialIR.h
--- a/Engine/Source/Runtime/Engine/Public/Materials/MaterialIR.h
+++ b/Engine/Source/Runtime/Engine/Public/Materials/MaterialIR.h
@@ -328,7 +328,9 @@ enum EOperator
 	UO_Logarithm2,
 	UO_Logarithm10,
 	UO_Round,
+	UO_Saturate,
 	UO_Sin,
+	UO_Sqrt,
 	UO_Tan,
 	UO_Truncate,
 
