Remove hack from vbr seektable

diff --git a/Engine/Source/Runtime/Engine/Private/StreamedAudioChunkSeekTable.cpp b/Engine/Source/Runtime/Engine/Private/StreamedAudioChunkSeekTable.cpp
--- a/Engine/Source/Runtime/Engine/Private/StreamedAudioChunkSeekTable.cpp
+++ b/Engine/Source/Runtime/Engine/Private/StreamedAudioChunkSeekTable.cpp
@@ -244,16 +244,9 @@ struct FVariableRateSeekTable : public FStreamedAudioChunkSeekTable::ISeekTableI
 	uint32 FindOffset(uint32 InTimeInAudioFrames) const override
 	{
 		check(Offsets.Num() == Times.Num());
-		const int32 Index = Algo::LowerBound(Times, InTimeInAudioFrames);
-		if (Offsets.IsValidIndex(Index))
+		if (const int32 Index = Algo::UpperBound(Times, InTimeInAudioFrames); Offsets.IsValidIndex(Index))
 		{
-			// Minor hack here. RADA is the only customer of this currently and needs
-			// the one block behind the one we've found. 
-			if (Index > 0)
-			{
-				return Offsets[Index-1];
-			}
-			return Offsets[0];
+			return Offsets[Index];
 		}
 		return INDEX_NONE;
 	}
