Add additional pixel formats to FImportanceTexture to support high dynamic range

Data is linearized on load for RGBA8 formats, and floating point data is loaded directly.

diff --git a/Engine/Source/Runtime/Engine/Classes/Kismet/ImportanceSamplingLibrary.h b/Engine/Source/Runtime/Engine/Classes/Kismet/ImportanceSamplingLibrary.h
--- a/Engine/Source/Runtime/Engine/Classes/Kismet/ImportanceSamplingLibrary.h
+++ b/Engine/Source/Runtime/Engine/Classes/Kismet/ImportanceSamplingLibrary.h
@@ -37,7 +37,7 @@ namespace EImportanceWeight
 USTRUCT(BlueprintType, meta = (HasNativeMake = "/Script/Engine.ImportanceSamplingLibrary.MakeImportanceTexture", HasNativeBreak = "/Script/Engine.ImportanceSamplingLibrary.BreakImportanceTexture"))
 struct FImportanceTexture
 {
-	GENERATED_BODY();
+	GENERATED_BODY()
 
 	// active texture dimensions, capped to 1024 x 1024
 	UPROPERTY()
@@ -57,10 +57,14 @@ struct FImportanceTexture
 	UPROPERTY()
 	TArray<float> ConditionalCDF;
 
+	UE_DEPRECATED(5.6, "FColor TextureData has been deprecated, use LinearTextureData")
+	UPROPERTY(Transient, meta = (DeprecatedProperty, DeprecationMessage = "Please use LinearTextureData instead."))
+	TArray<FColor> TextureData_DEPRECATED;
+	
 	// packed copy of MIP level data for filtered sampling (capped to 1024x1024)
 	// local copy seems better than allocating and copying the same data temporarily for each sample
-	UPROPERTY()
-	TArray<FColor> TextureData;
+	UPROPERTY(Transient)
+	TArray<FLinearColor> LinearTextureData;
 
 	// Original texture object for Break function
 	UPROPERTY()
@@ -72,13 +76,16 @@ struct FImportanceTexture
 
 public:
 	/* Default constructor, must Initialize before use */
-	FImportanceTexture() : Size(ForceInitToZero), NumMips(0), Texture(0), Weighting(EImportanceWeight::Luminance) {}
+	FImportanceTexture();
 
 	/* Constructor with initialization */
-	FImportanceTexture(UTexture2D *SourceTexture, TEnumAsByte<EImportanceWeight::Type> WeightingFunc)
-	{
-		Initialize(SourceTexture, WeightingFunc);
-	}
+	FImportanceTexture(UTexture2D *SourceTexture, TEnumAsByte<EImportanceWeight::Type> WeightingFunc);
+
+	FImportanceTexture(const FImportanceTexture& Other);
+	FImportanceTexture(FImportanceTexture&& Other);
+	FImportanceTexture& operator=(const FImportanceTexture& Other);
+	FImportanceTexture& operator=(FImportanceTexture&& Other);
+
 
 	/* Allocate and compute PDF arrays for a texture */
 	void Initialize(UTexture2D *SourceTexture, TEnumAsByte<EImportanceWeight::Type> WeightingFunc);
diff --git a/Engine/Source/Runtime/Engine/Private/Kismet/ImportanceSamplingLibrary.cpp b/Engine/Source/Runtime/Engine/Private/Kismet/ImportanceSamplingLibrary.cpp
--- a/Engine/Source/Runtime/Engine/Private/Kismet/ImportanceSamplingLibrary.cpp
+++ b/Engine/Source/Runtime/Engine/Private/Kismet/ImportanceSamplingLibrary.cpp
@@ -21,6 +21,23 @@
 // for 1024x1024 with 2^20 texels and 11 mips, that's 24-20 = 4 bits of probability precision
 #define MAX_MIP_LEVELS 11
 
+PRAGMA_DISABLE_DEPRECATION_WARNINGS
+/* Default constructor, must Initialize before use */
+FImportanceTexture::FImportanceTexture() : Size(ForceInitToZero), NumMips(0), Texture(0), Weighting(EImportanceWeight::Luminance) {}
+
+/* Constructor with initialization */
+FImportanceTexture::FImportanceTexture(UTexture2D *SourceTexture, TEnumAsByte<EImportanceWeight::Type> WeightingFunc)
+{
+	Initialize(SourceTexture, WeightingFunc);
+}
+
+FImportanceTexture::FImportanceTexture(const FImportanceTexture& Other) = default;
+FImportanceTexture::FImportanceTexture(FImportanceTexture&& Other) = default;
+FImportanceTexture& FImportanceTexture::operator=(const FImportanceTexture& Other) = default;
+FImportanceTexture& FImportanceTexture::operator=(FImportanceTexture&& Other) = default;
+
+PRAGMA_ENABLE_DEPRECATION_WARNINGS
+
 FLinearColor FImportanceTexture::GetColorTrilinear(FVector2D Position, float Mip) const
 {
 	float IntMip = FMath::FloorToFloat(Mip);
@@ -44,22 +61,20 @@ FLinearColor FImportanceTexture::GetColorBilinear(FVector2D Position, int32 Mip)
 
 	// At bottom MIP, return single texel
 	size_t TextureOffset = LevelStart + IntPos.Y * MipSize.X + IntPos.X;
-	FLinearColor Color00(TextureData[TextureOffset]);
+	FLinearColor Color00(LinearTextureData[TextureOffset]);
 	if (MipSize.X == 1 || MipSize.Y == 1) return Color00;
 	
 	// MIP texel blending should be in linear space, so this includes conversions to/from sRGB
-	FLinearColor Color01(TextureData[TextureOffset + MipSize.X]);
-	FLinearColor Color10(TextureData[TextureOffset + 1]);
-	FLinearColor Color11(TextureData[TextureOffset + MipSize.X + 1]);
+	FLinearColor Color01(LinearTextureData[TextureOffset + MipSize.X]);
+	FLinearColor Color10(LinearTextureData[TextureOffset + 1]);
+	FLinearColor Color11(LinearTextureData[TextureOffset + MipSize.X + 1]);
 	return FMath::Lerp(
 		FMath::Lerp(Color00, Color10, TexelBlend.X), 
 		FMath::Lerp(Color01, Color11, TexelBlend.X), TexelBlend.Y);
 }
 
-float FImportanceTexture::ImportanceWeight(FColor Texel, TEnumAsByte<EImportanceWeight::Type> WeightingFunc) const
+static float ImportanceWeightLinear(FLinearColor LinearTexel, TEnumAsByte<EImportanceWeight::Type> WeightingFunc)
 {
-	FLinearColor LinearTexel = (Texture.IsValid() && Texture->SRGB) ? FLinearColor(Texel) : Texel.ReinterpretAsLinear();
-
 	switch (WeightingFunc) {
 	case EImportanceWeight::Luminance:	return LinearTexel.GetLuminance();
 	case EImportanceWeight::Red:		return LinearTexel.R;
@@ -70,37 +85,26 @@ float FImportanceTexture::ImportanceWeight(FColor Texel, TEnumAsByte<EImportance
 	return 1.f;
 }
 
-void FImportanceTexture::Initialize(UTexture2D *SourceTexture, TEnumAsByte<EImportanceWeight::Type> WeightingFunc)
+float FImportanceTexture::ImportanceWeight(FColor Texel, TEnumAsByte<EImportanceWeight::Type> WeightingFunc) const
 {
-#if WITH_EDITOR
-	// For now, this doesn't support being refreshed once the source texture compilation is done.
-	// So force it to finish its compilation if it is not yet done to avoid getting the default's texture data.
-	if (SourceTexture && SourceTexture->IsCompiling())
-	{
-		FTextureCompilingManager::Get().FinishCompilation({ SourceTexture });
-	}
-#endif
+	FLinearColor LinearTexel = (Texture.IsValid() && Texture->SRGB) ? FLinearColor(Texel) : Texel.ReinterpretAsLinear();
+	return ImportanceWeightLinear(LinearTexel, WeightingFunc);
+}
 
-	if (!SourceTexture || SourceTexture->GetPixelFormat() != EPixelFormat::PF_B8G8R8A8)
-	{
-		Texture = 0;
-		FFrame::KismetExecutionMessage(TEXT("Importance Texture only supports RGBA8 textures"), ELogVerbosity::Error);
-		return;
-	}
 
-	// after this, safe to re-initialize
-	Texture = SourceTexture;
-	Weighting = WeightingFunc;
 
+template <typename ColorType>
+static void FetchTextureData(TArray<FLinearColor>& LinearTextureData, FIntPoint& Size, int& NumMips, UTexture2D* SourceTexture)
+{
 	// Save a copy of all MIP data for later color lookups
 	// Doing GetMipData for each sample would allocate and copy the entire MIP chain for each access
-	TArray<FColor*> MipData;
+	TArray<ColorType*> MipData;
 	int32 SourceMips = SourceTexture->GetNumMips();
 	int32 FirstMip = FMath::Max(0, SourceMips - MAX_MIP_LEVELS);
 	NumMips = SourceMips - FirstMip;
 	MipData.AddZeroed(NumMips);
 	// this gets the PlatformData mip data, not the Source:
-	SourceTexture->GetMipData(FirstMip, (void**)MipData.GetData());
+	SourceTexture->GetMipData(FirstMip, reinterpret_cast<void**>(MipData.GetData()));
 
 	// Copy just the needed MIP data and adjust size
 	// todo: why is it being copied? just use MipData
@@ -109,37 +113,91 @@ void FImportanceTexture::Initialize(UTexture2D *SourceTexture, TEnumAsByte<EImpo
 	Size = FIntPoint(((SrcSize.X - 1) >> FirstMip) + 1, ((SrcSize.Y - 1) >> FirstMip) + 1);
 	FIntPoint LastMipSize(((Size.X - 1) >> (NumMips - 1)) + 1, ((Size.Y - 1) >> (NumMips - 1)) + 1);
 	size_t MipDataSize = (4 * Size.X * Size.Y - LastMipSize.X * LastMipSize.Y) / 3;
-	TextureData.SetNumUninitialized(MipDataSize);
+	LinearTextureData.SetNumUninitialized(MipDataSize);
 	for (int32 mip = 0; mip < NumMips; ++mip) {
 		FIntPoint LevelSize(((Size.X - 1) >> mip) + 1, ((Size.Y - 1) >> mip) + 1);
 		size_t LevelStart = (Size.X * Size.Y - LevelSize.X * LevelSize.Y) * 4 / 3;
-		FMemory::Memcpy((void*)&TextureData[LevelStart], (void*)MipData[mip], LevelSize.X*LevelSize.Y*sizeof(FColor));
+		for (int Index = 0, NumPixels = LevelSize.X * LevelSize.Y; Index < NumPixels; Index++)
+		{
+			ColorType Texel = MipData[mip][Index];
+			if constexpr(std::is_same_v<ColorType, FColor>)
+			{
+				// in RGBA8 case, we should include the possibility of non-SRGB textures
+				LinearTextureData[LevelStart + Index] = SourceTexture->SRGB ? FLinearColor(Texel) : Texel.ReinterpretAsLinear();
+			}
+			else
+			{
+				LinearTextureData[LevelStart + Index] = FLinearColor(Texel);
+			}
+			
+		}
+	}
+
+	// get rid of temporary copy of MIP data
+	for (auto MipLevel : MipData)
+	{
+		FMemory::Free(MipLevel);
+	}
+}
+
+void FImportanceTexture::Initialize(UTexture2D *SourceTexture, TEnumAsByte<EImportanceWeight::Type> WeightingFunc)
+{
+#if WITH_EDITOR
+	// For now, this doesn't support being refreshed once the source texture compilation is done.
+	// So force it to finish its compilation if it is not yet done to avoid getting the default's texture data.
+	if (SourceTexture && SourceTexture->IsCompiling())
+	{
+		FTextureCompilingManager::Get().FinishCompilation({ SourceTexture });
+	}
+#endif
+
+	// after this, safe to re-initialize
+	Texture = SourceTexture;
+	Weighting = WeightingFunc;
+
+	if (!SourceTexture)
+	{
+		FFrame::KismetExecutionMessage(TEXT("Importance Texture is missing a source texture"), ELogVerbosity::Error);
+		return;
+	}
+
+	switch (SourceTexture->GetPixelFormat())
+	{
+		case EPixelFormat::PF_B8G8R8A8:
+			FetchTextureData<FColor>(LinearTextureData, Size, NumMips, SourceTexture);
+			break;
+		case EPixelFormat::PF_FloatRGBA:
+			FetchTextureData<FFloat16Color>(LinearTextureData, Size, NumMips, SourceTexture);
+			break;
+		case EPixelFormat::PF_A32B32G32R32F:
+			FetchTextureData<FLinearColor>(LinearTextureData, Size, NumMips, SourceTexture);
+			break;
+		default:
+			FFrame::KismetExecutionMessage(*FString::Printf(TEXT("Importance Texture does not support texture format %d"), SourceTexture->GetPixelFormat()), ELogVerbosity::Error);
+			Texture = nullptr; // Reset pointer to we don't try to sample from this texture
+			return;
 	}
 
 	// accumulate un-normalized marginal CDF for image, and conditional CDF for each row
 	MarginalCDF.SetNumUninitialized(Size.Y + 1);
 	ConditionalCDF.SetNumUninitialized((Size.X + 1) * Size.Y);
 	MarginalCDF[0] = 0.f;
+
 	for (int y = 1; y <= Size.Y; ++y)
 	{
 		// Accumulate along row
-		FColor *ColorRow = &MipData[0][(y - 1) * Size.X];
+
+		FLinearColor *ColorRow = &LinearTextureData[(y - 1) * Size.X];
 		float *CDFRow = &ConditionalCDF[(y - 1) * (Size.X + 1)];
 		CDFRow[0] = 0.f;
 		for (int x = 1; x <= Size.X; ++x)
 		{
-			CDFRow[x] = CDFRow[x - 1] + ImportanceWeight(ColorRow[x - 1], WeightingFunc);
+			CDFRow[x] = CDFRow[x - 1] + ImportanceWeightLinear(ColorRow[x - 1], WeightingFunc);
 		}
 
 		// add row total to global total
 		MarginalCDF[y] = MarginalCDF[y - 1] + CDFRow[Size.X];
 	}
-
-	// get rid of temporary copy of MIP data
-	for (auto MipLevel : MipData)
-	{
-		FMemory::Free(MipLevel);
-	}
 }
 
 void FImportanceTexture::ImportanceSample(const FVector2D & Rand, int Samples, float Intensity, FVector2D & SamplePosition, FLinearColor & SampleColor, float & SampleIntensity, float & SampleSize) const
