SkyAtmosphere - added sampling world direction override on SkyAtmosphereViewLuminance node.

#tests PC editor dx12
#rnx

diff --git a/Engine/Shaders/Private/MaterialTemplate.ush b/Engine/Shaders/Private/MaterialTemplate.ush
--- a/Engine/Shaders/Private/MaterialTemplate.ush
+++ b/Engine/Shaders/Private/MaterialTemplate.ush
@@ -2304,7 +2304,7 @@ float3 MaterialExpressionSkyAtmosphereLightDiskLuminance(FMaterialPixelParameter
 	return LightDiskLuminance;
 }
 
-float3 MaterialExpressionSkyAtmosphereViewLuminance(FMaterialPixelParameters Parameters)
+float3 MaterialExpressionSkyAtmosphereViewLuminance(FMaterialPixelParameters Parameters, float3 ViewDir)
 {
 #if MATERIAL_SKY_ATMOSPHERE && PROJECT_SUPPORT_SKY_ATMOSPHERE
 	if (ResolvedView.RenderingReflectionCaptureMask == 0.0f && !IsSkyAtmosphereRenderedInMain(ResolvedView.EnvironmentComponentsFlags))
@@ -2314,7 +2314,6 @@ float3 MaterialExpressionSkyAtmosphereViewLuminance(FMaterialPixelParameters Par
 
 	// The referencial used to build the Sky View lut
 	float3x3 LocalReferencial = GetSkyViewLutReferential(ResolvedView.SkyViewLutReferential);
-	const float3 ViewDir = -Parameters.CameraVector;
 
 	float ViewHeight = ResolvedView.SkyPlanetTranslatedWorldCenterAndViewHeight.w;
 	if(IsOrthoProjection())
diff --git a/Engine/Source/Runtime/Engine/Private/Materials/HLSLMaterialTranslator.cpp b/Engine/Source/Runtime/Engine/Private/Materials/HLSLMaterialTranslator.cpp
--- a/Engine/Source/Runtime/Engine/Private/Materials/HLSLMaterialTranslator.cpp
+++ b/Engine/Source/Runtime/Engine/Private/Materials/HLSLMaterialTranslator.cpp
@@ -12230,10 +12230,15 @@ int32 FHLSLMaterialTranslator::SkyAtmosphereLightDiskLuminance(int32 LightIndex,
 		OverrideAtmosphereLightDiscCosHalfApexAngle == INDEX_NONE ? *GetParameterCode(Constant(-1.0f)) : *GetParameterCode(OverrideAtmosphereLightDiscCosHalfApexAngle));
 }
 
-int32 FHLSLMaterialTranslator::SkyAtmosphereViewLuminance()
+int32 FHLSLMaterialTranslator::SkyAtmosphereViewLuminance(int32 WorldDirectionOverrideCodeChunk)
 {
+	if (WorldDirectionOverrideCodeChunk == INDEX_NONE)
+	{
+		WorldDirectionOverrideCodeChunk = Mul(Constant3(-1.0f, -1.0f, -1.0f), CameraVector());
+	}
+
 	bUsesSkyAtmosphere = true;
-	return AddCodeChunk(MCT_Float3, TEXT("MaterialExpressionSkyAtmosphereViewLuminance(Parameters)"));
+	return AddCodeChunk(MCT_Float3, TEXT("MaterialExpressionSkyAtmosphereViewLuminance(Parameters, %s)"), *GetParameterCode(WorldDirectionOverrideCodeChunk));
 }
 
 int32 FHLSLMaterialTranslator::SkyAtmosphereAerialPerspective(int32 WorldPosition, EPositionOrigin PositionOrigin)
diff --git a/Engine/Source/Runtime/Engine/Private/Materials/HLSLMaterialTranslator.h b/Engine/Source/Runtime/Engine/Private/Materials/HLSLMaterialTranslator.h
--- a/Engine/Source/Runtime/Engine/Private/Materials/HLSLMaterialTranslator.h
+++ b/Engine/Source/Runtime/Engine/Private/Materials/HLSLMaterialTranslator.h
@@ -1260,7 +1260,7 @@ protected:
 	virtual int32 SkyAtmosphereLightIlluminanceOnGround(int32 LightIndex) override;
 	virtual int32 SkyAtmosphereLightDirection(int32 LightIndex) override;
 	virtual int32 SkyAtmosphereLightDiskLuminance(int32 LightIndex, int32 OverrideAtmosphereLightDiscCosHalfApexAngle) override;
-	virtual int32 SkyAtmosphereViewLuminance() override;
+	virtual int32 SkyAtmosphereViewLuminance(int32 WorldDirectionOverrideCodeChunk) override;
 	virtual int32 SkyAtmosphereAerialPerspective(int32 WorldPosition, EPositionOrigin PositionOrigin) override;
 	virtual int32 SkyAtmosphereDistantLightScatteredLuminance() override;
 
diff --git a/Engine/Source/Runtime/Engine/Private/Materials/MaterialExpressions.cpp b/Engine/Source/Runtime/Engine/Private/Materials/MaterialExpressions.cpp
--- a/Engine/Source/Runtime/Engine/Private/Materials/MaterialExpressions.cpp
+++ b/Engine/Source/Runtime/Engine/Private/Materials/MaterialExpressions.cpp
@@ -21457,7 +21457,12 @@ void UMaterialExpressionSkyAtmosphereLightDiskLuminance::GetCaption(TArray<FStri
 #if WITH_EDITOR
 int32 UMaterialExpressionSkyAtmosphereViewLuminance::Compile(class FMaterialCompiler* Compiler, int32 OutputIndex)
 {
-	return Compiler->SkyAtmosphereViewLuminance();
+	int32 WorldDirectionCodeChunk = INDEX_NONE;
+	if (WorldDirection.GetTracedInput().Expression)
+	{
+		WorldDirectionCodeChunk = WorldDirection.Compile(Compiler);
+	}
+	return Compiler->SkyAtmosphereViewLuminance(WorldDirectionCodeChunk);
 }
 
 void UMaterialExpressionSkyAtmosphereViewLuminance::GetCaption(TArray<FString>& OutCaptions) const
diff --git a/Engine/Source/Runtime/Engine/Public/MaterialCompiler.h b/Engine/Source/Runtime/Engine/Public/MaterialCompiler.h
--- a/Engine/Source/Runtime/Engine/Public/MaterialCompiler.h
+++ b/Engine/Source/Runtime/Engine/Public/MaterialCompiler.h
@@ -658,7 +658,7 @@ public:
 	virtual int32 SkyAtmosphereLightIlluminanceOnGround(int32 LightIndex) = 0;
 	virtual int32 SkyAtmosphereLightDirection(int32 LightIndex) = 0;
 	virtual int32 SkyAtmosphereLightDiskLuminance(int32 LightIndex, int32 OverrideAtmosphereLightDiscCosHalfApexAngle) = 0;
-	virtual int32 SkyAtmosphereViewLuminance() = 0;
+	virtual int32 SkyAtmosphereViewLuminance(int32 WorldDirectionOverrideCodeChunk) = 0;
 	virtual int32 SkyAtmosphereAerialPerspective(int32 WorldPosition, EPositionOrigin PositionOrigin) = 0;
 	UE_DEPRECATED(5.4, "Use SkyAtmosphereAerialPerspective(int32, EPositionOrigin) instead")
 	int32 SkyAtmosphereAerialPerspective(int32 WorldPosition)
@@ -1248,9 +1248,9 @@ public:
 		return Compiler->SkyAtmosphereLightDiskLuminance(LightIndex, OverrideAtmosphereLightDiscCosHalfApexAngle);
 	}
 
-	virtual int32 SkyAtmosphereViewLuminance() override
+	virtual int32 SkyAtmosphereViewLuminance(int32 WorldDirectionOverrideCodeChunk) override
 	{
-		return Compiler->SkyAtmosphereViewLuminance();
+		return Compiler->SkyAtmosphereViewLuminance(WorldDirectionOverrideCodeChunk);
 	}
 
 	virtual int32 SkyAtmosphereAerialPerspective(int32 WorldPosition, EPositionOrigin PositionOrigin) override
diff --git a/Engine/Source/Runtime/Engine/Public/Materials/MaterialExpressionSkyAtmosphereViewLuminance.h b/Engine/Source/Runtime/Engine/Public/Materials/MaterialExpressionSkyAtmosphereViewLuminance.h
--- a/Engine/Source/Runtime/Engine/Public/Materials/MaterialExpressionSkyAtmosphereViewLuminance.h
+++ b/Engine/Source/Runtime/Engine/Public/Materials/MaterialExpressionSkyAtmosphereViewLuminance.h
@@ -13,13 +13,16 @@ class UMaterialExpressionSkyAtmosphereViewLuminance : public UMaterialExpression
 {
 	GENERATED_BODY()
 
+	/** Sample the sky atmosphere using WorldDirection instead of the default camera-to-pixel V direction. */
+	UPROPERTY()
+	FExpressionInput WorldDirection;
+
 public:
 
 	//~ Begin UMaterialExpression Interface
 #if WITH_EDITOR
 	virtual int32 Compile(class FMaterialCompiler* Compiler, int32 OutputIndex) override;
 	virtual void GetCaption(TArray<FString>& OutCaptions) const override;
-
 #endif
 	//~ End UMaterialExpression Interface
 };
