Added warning when sampling the previous sky light capture into a IsSky material.

#tests PC editor dx12
#rnx

diff --git a/Engine/Source/Runtime/Engine/Private/Materials/HLSLMaterialTranslator.cpp b/Engine/Source/Runtime/Engine/Private/Materials/HLSLMaterialTranslator.cpp
--- a/Engine/Source/Runtime/Engine/Private/Materials/HLSLMaterialTranslator.cpp
+++ b/Engine/Source/Runtime/Engine/Private/Materials/HLSLMaterialTranslator.cpp
@@ -12265,6 +12265,10 @@ int32 FHLSLMaterialTranslator::SkyLightEnvMapSample(int32 DirectionCodeChunk, in
 	{
 		return Errorf(TEXT("The SkyLightEnvMapSample node can only be used when material Domain is set to Surface."));
 	}
+	if (Material->IsSky())
+	{
+		UE_LOG(LogMaterial, Warning, TEXT("Using SkyLightEnvMapSample from a IsSky material can result in visual artifact. For instance, if the previous frame capture was super bright, it might leak onto a new frame, e.g. transtion from menu to game."));
+	}
 
 	return AddCodeChunk(MCT_Float3, TEXT("MaterialExpressionSkyLightEnvMapSample(%s, %s)"), *GetParameterCode(DirectionCodeChunk), *GetParameterCode(RoughnessCodeChunk));
 }
