Sideloaded animation fixes
- Made sideload testing available on a CVar rather than a define
- Fetch the correct target compressed data (for the running platform) to side-load the data to in editor
- Renamed some API (+ redirector) to clear up usage/readability a bit
- Added export for ValidateCompressionSettings (used in sideloading)

diff --git a/Engine/Source/Runtime/Engine/Classes/Animation/AnimSequence.h b/Engine/Source/Runtime/Engine/Classes/Animation/AnimSequence.h
--- a/Engine/Source/Runtime/Engine/Classes/Animation/AnimSequence.h
+++ b/Engine/Source/Runtime/Engine/Classes/Animation/AnimSequence.h
@@ -20,6 +20,7 @@
 #include "Animation/CustomAttributes.h"
 #include "Animation/AnimData/AnimDataNotifications.h"
 #include "Animation/AttributeCurve.h"
+#include "Async/SharedRecursiveMutex.h"
 #include "UObject/PerPlatformProperties.h"
 #include "IO/IoHash.h"
 
@@ -808,6 +809,7 @@ protected:
 	TMap<FIoHash, TUniquePtr<FCompressedAnimSequence>> DataByPlatformKeyHash;
 	TMap<FIoHash, TPimplPtr<UE::Anim::FAnimationSequenceAsyncCacheTask>> CacheTasksByKeyHash;
 	mutable FRWLock CompressedDataLock;
+	mutable UE::FSharedRecursiveMutex SharedCompressedDataMutex;
 
 	// Cached FIoHash keys stored according to ITargetPlatform type-hash
 	mutable TMap<uint32, FIoHash> PlatformHashToKeyHash;
@@ -827,7 +829,7 @@ protected:
 	ENGINE_API bool IsCompiling() const;
 	ENGINE_API bool RequiresResidency(const FIoHash& InKeyHash) const;
 
-	void ValidateCompressionSettings();
+	ENGINE_API void ValidateCompressionSettings();
 
 	const FCompressedAnimSequence& GetPlatformCompressedData(const ITargetPlatform* InTargetPlatform) const;
 #endif // WITH_EDITOR
@@ -935,16 +937,25 @@ protected:
 		FCompressedAnimationDataReadScope(const UAnimSequence* InAnimSequence) : FReadScopeLock(InAnimSequence->CompressedDataLock)
 		{
 			check(InAnimSequence);
-		}		
+			SharedLock = MakeUnique<UE::TSharedLock<UE::FSharedRecursiveMutex>>(InAnimSequence->SharedCompressedDataMutex);
+		}
+		TUniquePtr<UE::TSharedLock<UE::FSharedRecursiveMutex>> SharedLock;
 	};
 	
 	/** Helper struct to enter/leave compressed data lock in write-mode */
 	struct FCompressedAnimationDataWriteScope : public FWriteScopeLock
 	{
-		FCompressedAnimationDataWriteScope(const UAnimSequence* InAnimSequence) : FWriteScopeLock(InAnimSequence->CompressedDataLock)
+		FCompressedAnimationDataWriteScope(const UAnimSequence* InAnimSequence) : FWriteScopeLock(InAnimSequence->CompressedDataLock), AnimSequence(InAnimSequence)
 		{
 			check(InAnimSequence);
+			InAnimSequence->SharedCompressedDataMutex.Lock();
+		}
+		~FCompressedAnimationDataWriteScope()
+		{
+			AnimSequence->SharedCompressedDataMutex.Unlock();
 		}
+
+		const UAnimSequence* AnimSequence = nullptr;
 	};
 #endif // WITH_EDITOR
 
diff --git a/Engine/Source/Runtime/Engine/Private/Animation/AnimSequence.cpp b/Engine/Source/Runtime/Engine/Private/Animation/AnimSequence.cpp
--- a/Engine/Source/Runtime/Engine/Private/Animation/AnimSequence.cpp
+++ b/Engine/Source/Runtime/Engine/Private/Animation/AnimSequence.cpp
@@ -5144,7 +5144,7 @@ void UAnimSequence::CacheDerivedDataForPlatform(const ITargetPlatform* TargetPla
 {
 	const FIoHash KeyHash = BeginCacheDerivedData(TargetPlatform);
 	EndCacheDerivedData(KeyHash);
-	check(HasCompressedDataForHash(KeyHash));
+	check(KeyHash.IsZero() || HasCompressedDataForHash(KeyHash));
 }
 
 FFrameRate UAnimSequence::GetTargetSamplingFrameRate(const ITargetPlatform* InPlatform) const
