Shrink FCollisionResponse::ResponseArray post update as its slack memory might not be used and might be wasted over time. That does save 5MB on low end platforms

#rnx

diff --git a/Engine/Source/Runtime/Engine/Private/PhysicsEngine/BodyInstance.cpp b/Engine/Source/Runtime/Engine/Private/PhysicsEngine/BodyInstance.cpp
--- a/Engine/Source/Runtime/Engine/Private/PhysicsEngine/BodyInstance.cpp
+++ b/Engine/Source/Runtime/Engine/Private/PhysicsEngine/BodyInstance.cpp
@@ -283,6 +283,9 @@ void FCollisionResponse::UpdateArrayFromResponseContainer()
 			}
 		}
 	}
+
+	// Shrink array as its slack memory might not be used and might be wasted over time.
+	ResponseArray.Shrink();
 }
 
 #endif // WITH_EDITOR
