New material translator push.
- Added a "SimplifyOperator" step to "EmitOperator" so that the operator gets a chance to be simplified to a simpler operator if we can prove that some identity is true. For instance this would change Length to Abs when applied on a scalar.
- Added lerp() identities (fold to 'a' or 'b' when factor is 0 or is 1) to folding logic.
- Added support for Clamp, but the translator does not perform arguments auto conversion yet.

#rnx

diff --git a/Engine/Source/Runtime/Engine/Private/Materials/MaterialExpressionsIR.cpp b/Engine/Source/Runtime/Engine/Private/Materials/MaterialExpressionsIR.cpp
--- a/Engine/Source/Runtime/Engine/Private/Materials/MaterialExpressionsIR.cpp
+++ b/Engine/Source/Runtime/Engine/Private/Materials/MaterialExpressionsIR.cpp
@@ -514,8 +514,7 @@ void UMaterialExpressionLength::Build(UE::MIR::FEmitter& Emitter)
 	MIR::FValue* Value = Emitter.Input(&Input);
 	UE_MIR_CHECKPOINT();
 
-	const MIR::EOperator Op = Value->Type->AsPrimitive()->IsScalar() ? MIR::UO_Abs : MIR::UO_Length;
-	BuildUnaryArithmeticOperatorFromValue(Emitter, Op, Value);
+	BuildUnaryArithmeticOperatorFromValue(Emitter, MIR::UO_Length, Value);
 }
 
 void UMaterialExpressionRound::Build(UE::MIR::FEmitter& Emitter)
@@ -691,9 +690,7 @@ void UMaterialExpressionDistance::Build(MIR::FEmitter& Emitter)
 	MIR::FValue* DeltaValue = Emitter.EmitOperator(MIR::BO_Subtract, LhsValue, RhsValue);
 	UE_MIR_CHECKPOINT();
 
-	const bool bScalarInputs = LhsValue->Type->AsPrimitive()->IsScalar() && RhsValue->Type->AsPrimitive()->IsScalar();
-	const MIR::EOperator Op = bScalarInputs ? MIR::UO_Abs : MIR::UO_Length;
-	MIR::FValue* ResultValue = Emitter.EmitOperator(Op, DeltaValue);
+	MIR::FValue* ResultValue = Emitter.EmitOperator(MIR::UO_Length, DeltaValue);
 	UE_MIR_CHECKPOINT();
 
 	Emitter.Output(0, ResultValue);
@@ -972,6 +969,36 @@ void BuildTernaryArithmeticOperator(
 	Emitter.Output(0, Emitter.EmitOperator(Op, ValueA, ValueB, ValueC));
 }
 
+void UMaterialExpressionClamp::Build(UE::MIR::FEmitter& Emitter)
+{
+	MIR::FValue* InputValue = Emitter.Input(&Input);
+
+	UE_MIR_CHECKPOINT();
+
+	MIR::FValue* MinValue = Emitter.DefaultInputToFloatScalar(&Min, MinDefault).Input(&Min);
+	MIR::FValue* MaxValue = Emitter.DefaultInputToFloatScalar(&Max, MaxDefault).Input(&Max);
+
+	UE_MIR_CHECKPOINT();
+
+	MIR::FValue* OutputValue = nullptr;
+	if (ClampMode == CMODE_Clamp)
+	{
+		OutputValue = Emitter.EmitOperator(MIR::TO_Clamp, InputValue, MinValue, MaxValue);
+	}
+	else if (ClampMode == CMODE_ClampMin)
+	{
+		OutputValue = Emitter.EmitOperator(MIR::BO_Max, InputValue, MinValue);
+	}
+	else if (ClampMode == CMODE_ClampMax)
+	{
+		OutputValue = Emitter.EmitOperator(MIR::BO_Min, InputValue, MaxValue);
+	}
+
+	UE_MIR_CHECKPOINT();
+
+	Emitter.Output(0, OutputValue);
+}
+
 void UMaterialExpressionLinearInterpolate::Build(UE::MIR::FEmitter& Emitter)
 {
 	BuildTernaryArithmeticOperator(Emitter, MIR::TO_Lerp, &A, ConstA, &B, ConstB, &Alpha, ConstAlpha);
diff --git a/Engine/Source/Runtime/Engine/Private/Materials/MaterialIR.cpp b/Engine/Source/Runtime/Engine/Private/Materials/MaterialIR.cpp
--- a/Engine/Source/Runtime/Engine/Private/Materials/MaterialIR.cpp
+++ b/Engine/Source/Runtime/Engine/Private/Materials/MaterialIR.cpp
@@ -168,9 +168,19 @@ bool FValue::IsFalse() const
 	return Constant && Constant->IsBool() && Constant->Boolean == false;
 }
 
-bool FValue::IsExactlyZero() const
+bool FValue::AreAllExactlyZero() const
 {
-	if (const MIR::FConstant* Constant = As<MIR::FConstant>())
+	if (const MIR::FDimensional* Dimensional = As<MIR::FDimensional>())
+	{
+		for (const MIR::FValue* Component : Dimensional->GetComponents())
+		{
+			if (!Component->AreAllExactlyZero())
+			{
+				return false;
+			}
+		}
+	}
+	else if (const MIR::FConstant* Constant = As<MIR::FConstant>())
 	{
 		return (Constant->IsInteger() && Constant->Integer == 0)
 			|| (Constant->IsFloat() && Constant->Float == 0.0f);
@@ -178,9 +188,19 @@ bool FValue::IsExactlyZero() const
 	return false;
 }
 
-bool FValue::IsNearlyZero() const
+bool FValue::AreAllNearlyZero() const
 {
-	if (const MIR::FConstant* Constant = As<MIR::FConstant>())
+	if (const MIR::FDimensional* Dimensional = As<MIR::FDimensional>())
+	{
+		for (const MIR::FValue* Component : Dimensional->GetComponents())
+		{
+			if (!Component->AreAllNearlyZero())
+			{
+				return false;
+			}
+		}
+	}
+	else if (const MIR::FConstant* Constant = As<MIR::FConstant>())
 	{
 		return (Constant->IsInteger() && Constant->Integer == 0)
 			|| (Constant->IsFloat() && FMath::IsNearlyZero(Constant->Float));
@@ -188,9 +208,19 @@ bool FValue::IsNearlyZero() const
 	return false;
 }
 
-bool FValue::IsExactlyOne() const
+bool FValue::AreAllExactlyOne() const
 {
-	if (const MIR::FConstant* Constant = As<MIR::FConstant>())
+	if (const MIR::FDimensional* Dimensional = As<MIR::FDimensional>())
+	{
+		for (const MIR::FValue* Component : Dimensional->GetComponents())
+		{
+			if (!Component->AreAllExactlyOne())
+			{
+				return false;
+			}
+		}
+	}
+	else if (const MIR::FConstant* Constant = As<MIR::FConstant>())
 	{
 		return (Constant->IsInteger() && Constant->Integer == 1)
 			|| (Constant->IsFloat() && Constant->Float == 1.0f);
@@ -198,9 +228,19 @@ bool FValue::IsExactlyOne() const
 	return false;
 }
 
-bool FValue::IsNearlyOne() const
+bool FValue::AreAllNearlyOne() const
 {
-	if (const MIR::FConstant* Constant = As<MIR::FConstant>())
+	if (const MIR::FDimensional* Dimensional = As<MIR::FDimensional>())
+	{
+		for (const MIR::FValue* Component : Dimensional->GetComponents())
+		{
+			if (!Component->AreAllNearlyOne())
+			{
+				return false;
+			}
+		}
+	}
+	else if (const MIR::FConstant* Constant = As<MIR::FConstant>())
 	{
 		return (Constant->IsInteger() && Constant->Integer == 1)
 			|| (Constant->IsFloat() && FMath::IsNearlyEqual(Constant->Float, 1.0f));
@@ -391,6 +431,11 @@ bool IsBitOperator(EOperator Op)
 	return Op == UO_Not || (Op >= BO_Or && Op <= BO_Xor);
 }
 
+bool IsUnaryOperator(EOperator Op)
+{
+	return Op >= UO_FirstUnaryOperator && Op < BO_FirstBinaryOperator;
+}
+
 bool IsBinaryOperator(EOperator Op)
 {
 	return Op >= BO_FirstBinaryOperator && Op < TO_FirstTernaryOperator;
@@ -405,57 +450,56 @@ const TCHAR* OperatorToString(EOperator Op)
 {
 	switch (Op)
 	{
-		case O_Invalid: return TEXT("Invalid");
-
-		/* Unary Arithmetic */
+		// Note: sorted alphabetically
+		case BO_Add: return TEXT("Add");
+		case BO_And: return TEXT("And");
+		case BO_Cross: return TEXT("Cross");
+		case BO_Divide: return TEXT("Divide");
+		case BO_Dot: return TEXT("Dot");
+		case BO_Equals: return TEXT("Equals");
+		case BO_Fmod: return TEXT("Fmod");
+		case BO_GreaterThan: return TEXT("GreaterThan");
+		case BO_GreaterThanOrEquals: return TEXT("GreaterThanOrEquals");
+		case BO_LessThan: return TEXT("LessThan");
+		case BO_LessThanOrEquals: return TEXT("LessThanOrEquals");
+		case BO_Max: return TEXT("Max");
+		case BO_Min: return TEXT("Min");
+		case BO_Multiply: return TEXT("Multiply");
+		case BO_NotEquals: return TEXT("NotEquals");
+		case BO_Or: return TEXT("Or");
+		case BO_Subtract: return TEXT("Subtract");
+		case BO_Xor: return TEXT("Xor");
+		case TO_Clamp: return TEXT("Clamp");
+		case TO_Lerp: return TEXT("Lerp");
+		case TO_Smoothstep: return TEXT("Smoothstep");
 		case UO_Abs: return TEXT("Abs");
-		case UO_ACos: return TEXT("Arccosine");
-		case UO_ASin: return TEXT("Arcsine");
-		case UO_ATan: return TEXT("Arctangent");
+		case UO_ACos: return TEXT("ACos");
+		case UO_ASin: return TEXT("ASin");
+		case UO_ATan: return TEXT("ATan");
 		case UO_Ceil: return TEXT("Ceil");
-		case UO_Cos: return TEXT("Cosine");
+		case UO_Cos: return TEXT("Cos");
 		case UO_Exponential: return TEXT("Exponential");
 		case UO_Exponential2: return TEXT("Exponential2");
 		case UO_Floor: return TEXT("Floor");
 		case UO_Frac: return TEXT("Frac");
 		case UO_Length: return TEXT("Length");
 		case UO_Logarithm: return TEXT("Logarithm");
-		case UO_Logarithm2: return TEXT("Logarithm2");
 		case UO_Logarithm10: return TEXT("Logarithm10");
+		case UO_Logarithm2: return TEXT("Logarithm2");
+		case UO_Negate: return TEXT("Negate");
+		case UO_Not: return TEXT("Not");
 		case UO_Round: return TEXT("Round");
 		case UO_Saturate: return TEXT("Saturate");
-		case UO_Sin: return TEXT("Sine");
+		case UO_Sin: return TEXT("Sin");
 		case UO_Sqrt: return TEXT("Sqrt");
-		case UO_Tan: return TEXT("Tangent");
+		case UO_Tan: return TEXT("Tan");
 		case UO_Truncate: return TEXT("Truncate");
 
-		/* Binary Comparisons */
-		case BO_GreaterThan: return TEXT("GreaterThan");
-		case BO_GreaterThanOrEquals: return TEXT("GreaterThanOrEquals");
-		case BO_LessThan: return TEXT("LowerThan");
-		case BO_LessThanOrEquals: return TEXT("LowerThanOrEquals");
-		case BO_Equals: return TEXT("Equals");
-		case BO_NotEquals: return TEXT("NotEquals");
-	
-		/* Binary Arithmetic */
-		case BO_Add: return TEXT("Add");
-		case BO_Fmod: return TEXT("Fmod");
-		case BO_Subtract: return TEXT("Subtract");
-		case BO_Multiply: return TEXT("Multiply");
-		case BO_Divide: return TEXT("Divide");
-		case BO_Min: return TEXT("Min");
-		case BO_Max: return TEXT("Max");
-
-		/* Binary Intrinsic Operators */
-		case BO_Dot: return TEXT("DotProduct");
-		case BO_Cross: return TEXT("CrossProduct");
-		
-		/* Ternary Arithmetic */
-		case TO_Lerp: return TEXT("Lerp");
-		case TO_Smoothstep: return TEXT("Smoothstep");
-
-		default: UE_MIR_UNREACHABLE();
+		case O_Invalid: return TEXT("Invalid");
+		case OperatorCount: UE_MIR_UNREACHABLE();
 	}
+
+	return TEXT("???");
 }
 
 } // namespace UE::MIR
diff --git a/Engine/Source/Runtime/Engine/Private/Materials/MaterialIREmitter.cpp b/Engine/Source/Runtime/Engine/Private/Materials/MaterialIREmitter.cpp
--- a/Engine/Source/Runtime/Engine/Private/Materials/MaterialIREmitter.cpp
+++ b/Engine/Source/Runtime/Engine/Private/Materials/MaterialIREmitter.cpp
@@ -821,7 +821,6 @@ static TOptional<T> FoldScalarArithmeticOperator(FEmitter* Emitter, EOperator Op
 				{
 					return FMath::RoundToDouble(A);
 				}
-			case UO_Saturate: return FMath::Min<T>(1.0f, FMath::Max<T>(0.0f, A));
 			case UO_Sin: return FMath::Sin(A);
 			case UO_Sqrt: return FMath::Sqrt(A);
 			case UO_Tan: return FMath::Tan(A);
@@ -837,31 +836,44 @@ static TOptional<T> FoldScalarArithmeticOperator(FEmitter* Emitter, EOperator Op
 			case BO_Fmod: return FGenericPlatformMath::Fmod(A, B);
 			case TO_Lerp: return FMath::Lerp<T>(A, B, C);
 			case TO_Smoothstep: return FMath::SmoothStep<T>(A, B, C);
+			default: break;
 		}
 	}
 
 	switch (Operator)
 	{
+		case UO_Abs: return FGenericPlatformMath::Abs<T>(A);
 		case UO_Negate: return -A;
+		case UO_Saturate: return FMath::Clamp(A, T(0), T(1));
 		case BO_Add: return A + B;
 		case BO_Subtract: return A - B;
 		case BO_Multiply: return A * B;
 		case BO_Divide: return A / B;
 		case BO_Min: return FMath::Min<T>(A, B);
 		case BO_Max: return FMath::Max<T>(A, B);
-		case UO_Abs: return FGenericPlatformMath::Abs<T>(A);
+		case TO_Clamp: return FMath::Clamp(A, B, C);
 		default: UE_MIR_UNREACHABLE();
 	}
 }
 
-static FValue* TryFoldOperatorScalar(FEmitter* Emitter, EOperator Op, FValue* A, FValue* B, FValue* C)
+// It tries to apply a known identity of specified operator, e.g. "x + 0 = x ∀ x ∈ R".
+// If it returns a value, the operation has been "folded" and the returned value is the
+// result (in the example above, it would return "x").
+// If it returns null, the end result could not be inferred, but the operator could have
+// still been changed to some other (with lower complexity). For example "clamp(x, 0, 1)"
+// will change to "saturate(x)".
+static FValue* TrySimplifyOperator(FEmitter* Emitter, EOperator& Op, FValue*& A, FValue*& B, FValue*& C)
 {
-	FPrimitiveTypePtr PrimitiveType = A->Type->AsPrimitive();
-
-	// Check whether we can fold the operation based on identies (e.g. "x + 0 = x ∀ x ∈ R").
-	// These simplifications don't requre that both lhs and rhs are constant.
 	switch (Op)
 	{
+		/* Unary Operators */
+		case UO_Length:
+			if (A->Type->AsScalar())
+			{
+				Op = UO_Abs;
+			}
+			break;
+
 		/* Binary Comparisons */
 		case BO_GreaterThan:
 		case BO_LessThan:
@@ -881,24 +893,24 @@ static FValue* TryFoldOperatorScalar(FEmitter* Emitter, EOperator Op, FValue* A,
 			}
 			break;
 
-		/* Arithmetic */
+		/* Binary Arithmetic */
 		case BO_Add:
-			if (A->IsNearlyZero())
+			if (A->AreAllNearlyZero())
 			{
 				return B;
 			}
-			else if (B->IsNearlyZero())
+			else if (B->AreAllNearlyZero())
 			{
 				return A;
 			}
 			break;
 
 		case BO_Subtract:
-			if (B->IsNearlyZero())
+			if (B->AreAllNearlyZero())
 			{
 				return A;
 			}
-			else if (A->IsNearlyZero())
+			else if (A->AreAllNearlyZero())
 			{
 				// return Emitter->EmitUnaryOperator(UO_Minus, Rhs);
 				UE_MIR_TODO();
@@ -906,58 +918,60 @@ static FValue* TryFoldOperatorScalar(FEmitter* Emitter, EOperator Op, FValue* A,
 			break;
 
 		case BO_Multiply:
-			if (A->IsNearlyZero() || B->IsNearlyOne())
+			if (A->AreAllNearlyZero() || B->AreAllNearlyOne())
 			{
 				return A;
 			}
-			else if (A->IsNearlyOne() || B->IsNearlyZero())
+			else if (A->AreAllNearlyOne() || B->AreAllNearlyZero())
 			{
 				return B;
 			}
 			break;
 
 		case BO_Divide:
-			if (A->IsNearlyZero() || B->IsNearlyOne())
+			if (A->AreAllNearlyZero() || B->AreAllNearlyOne())
 			{
 				return A;
 			}
 			break;
+		
+		case TO_Clamp:
+			if (B->AreAllNearlyZero() && C->AreAllNearlyOne())
+			{
+				Op = UO_Saturate;
+				B = nullptr;
+				C = nullptr;
+			}
+			break;
 
 		case TO_Lerp:
-			break; // todo
-
-		case TO_Smoothstep:
-			break; // todo
-
-		/* No simplifications for these operations */
-		case BO_Fmod:
-		case BO_Max:
-		case BO_Min:
-		case UO_Abs:
-		case UO_ACos:
-		case UO_ASin:
-		case UO_ATan:
-		case UO_Ceil:
-		case UO_Cos:
-		case UO_Exponential:
-		case UO_Exponential2:
-		case UO_Floor:
-		case UO_Frac:
-		case UO_Logarithm:
-		case UO_Logarithm10:
-		case UO_Logarithm2:
-		case UO_Negate:
-		case UO_Not:
-		case UO_Round:
-		case UO_Saturate:
-		case UO_Sin:
-		case UO_Sqrt:
-		case UO_Tan:
-		case UO_Truncate:
+			if (C->AreAllNearlyZero())
+			{
+				return A;
+			}
+			else if (C->AreAllNearlyOne())
+			{
+				return B;
+			}
 			break;
 
 		default:
-			UE_MIR_UNREACHABLE();
+			break;
+	}
+
+	return nullptr;
+}
+
+// Tries to fold (statically evaluate) the operator, assuming that the arguments are all scalar.
+// It returns either the result of the operator or null if it could not be folded.
+static FValue* TryFoldOperatorScalar(FEmitter* Emitter, EOperator Op, FValue* A, FValue* B, FValue* C)
+{
+	FPrimitiveTypePtr PrimitiveType = A->Type->AsPrimitive();
+
+	// Try to simplify the operator. This could potentially change Op, A, B and C.
+	if (FValue* Simplified = TrySimplifyOperator(Emitter, Op, A, B, C))
+	{
+		return Simplified;
 	}
 
 	// Verify that both lhs and rhs are constants, otherwise we cannot fold the operation.
@@ -1116,13 +1130,14 @@ static const FOperatorSignature* GetOperatorSignature(EOperator Op)
 		S[BO_Xor] 					= BinaryLogical;
 		S[BO_Dot] 					= { { OT_Arithmetic | OT_NotMatrix, OT_SameAsFirstArgument }, ORT_FirstArgumentTypeToScalar };
 		S[BO_Cross] 				= { { OT_Arithmetic | OT_Vector3, OT_SameAsFirstArgument }, ORT_FirstArgumentType };
+		S[TO_Clamp]					= TernaryArithmetic;
 		S[TO_Lerp] 					= TernaryArithmetic;
 		S[TO_Smoothstep]			= TernaryArithmetic;
 		return S;
 	} ();
 	return &Signatures[Op];
 }
-
+ 
 /// Validates that the types of the arguments are valid for specified operator.
 /// If valid, it returns the type of the result. Otherwise if it is not valid, it returns nullptr.
 static FPrimitiveTypePtr ValidateOperatorAndGetResultType(FEmitter* Emitter, EOperator Op, FValue* A, FValue* B, FValue* C)
@@ -1244,7 +1259,7 @@ static FPrimitiveTypePtr ValidateOperatorAndGetResultType(FEmitter* Emitter, EOp
 /// 	op(v, w) == [op(v_0, w_0), ..., op(v_n, w_n)]
 static bool IsComponentwiseOperator(EOperator Op)
 {
-	return Op != BO_Dot && Op != BO_Cross && Op != UO_Length;
+	return Op != BO_Dot && Op != BO_Cross;
 }
 
 /// Tries to fold the operator by applying the operator componentwise on arguments components.
@@ -1463,7 +1478,13 @@ FValue* FEmitter::EmitOperator(EOperator Op, FValue* A, FValue* B, FValue* C)
 		return FPoison::Get();
 	}
 
-	// Try folding the operator folding first.
+	// Try to apply some operator identity to simplify the operator.
+	if (FValue* Simplified = TrySimplifyOperator(this, Op, A, B, C))
+	{
+		return Simplified;
+	}
+
+	// Try folding the operator first.
 	if (FValue* FoldedValue = TryFoldOperator(this, Op, A, B, C, ResultType))
 	{
 		return FoldedValue;
diff --git a/Engine/Source/Runtime/Engine/Private/Materials/MaterialIRToHLSLTranslator.cpp b/Engine/Source/Runtime/Engine/Private/Materials/MaterialIRToHLSLTranslator.cpp
--- a/Engine/Source/Runtime/Engine/Private/Materials/MaterialIRToHLSLTranslator.cpp
+++ b/Engine/Source/Runtime/Engine/Private/Materials/MaterialIRToHLSLTranslator.cpp
@@ -140,30 +140,30 @@ static const TCHAR* GetHLSLTypeString(EMaterialValueType Type)
 {
 	switch (Type)
 	{
-	case MCT_Float1: return TEXT("MaterialFloat");
-	case MCT_Float2: return TEXT("MaterialFloat2");
-	case MCT_Float3: return TEXT("MaterialFloat3");
-	case MCT_Float4: return TEXT("MaterialFloat4");
-	case MCT_Float: return TEXT("MaterialFloat");
-	case MCT_Texture2D: return TEXT("texture2D");
-	case MCT_TextureCube: return TEXT("textureCube");
-	case MCT_Texture2DArray: return TEXT("texture2DArray");
-	case MCT_VolumeTexture: return TEXT("volumeTexture");
-	case MCT_StaticBool: return TEXT("static bool");
-	case MCT_Bool:  return TEXT("bool");
-	case MCT_MaterialAttributes: return TEXT("FMaterialAttributes");
-	case MCT_TextureExternal: return TEXT("TextureExternal");
-	case MCT_TextureVirtual: return TEXT("TextureVirtual");
-	case MCT_VTPageTableResult: return TEXT("VTPageTableResult");
-	case MCT_ShadingModel: return TEXT("uint");
-	case MCT_UInt: return TEXT("uint");
-	case MCT_UInt1: return TEXT("uint");
-	case MCT_UInt2: return TEXT("uint2");
-	case MCT_UInt3: return TEXT("uint3");
-	case MCT_UInt4: return TEXT("uint4");
-	case MCT_Substrate: return TEXT("FSubstrateData");
-	case MCT_TextureCollection: return TEXT("FResourceCollection");
-	default: return TEXT("unknown");
+		case MCT_Float1: return TEXT("MaterialFloat");
+		case MCT_Float2: return TEXT("MaterialFloat2");
+		case MCT_Float3: return TEXT("MaterialFloat3");
+		case MCT_Float4: return TEXT("MaterialFloat4");
+		case MCT_Float: return TEXT("MaterialFloat");
+		case MCT_Texture2D: return TEXT("texture2D");
+		case MCT_TextureCube: return TEXT("textureCube");
+		case MCT_Texture2DArray: return TEXT("texture2DArray");
+		case MCT_VolumeTexture: return TEXT("volumeTexture");
+		case MCT_StaticBool: return TEXT("static bool");
+		case MCT_Bool:  return TEXT("bool");
+		case MCT_MaterialAttributes: return TEXT("FMaterialAttributes");
+		case MCT_TextureExternal: return TEXT("TextureExternal");
+		case MCT_TextureVirtual: return TEXT("TextureVirtual");
+		case MCT_VTPageTableResult: return TEXT("VTPageTableResult");
+		case MCT_ShadingModel: return TEXT("uint");
+		case MCT_UInt: return TEXT("uint");
+		case MCT_UInt1: return TEXT("uint");
+		case MCT_UInt2: return TEXT("uint2");
+		case MCT_UInt3: return TEXT("uint3");
+		case MCT_UInt4: return TEXT("uint4");
+		case MCT_Substrate: return TEXT("FSubstrateData");
+		case MCT_TextureCollection: return TEXT("FResourceCollection");
+		default: return TEXT("unknown");
 	};
 }
 
@@ -274,7 +274,7 @@ struct FTranslator : FMaterialIRToHLSLTranslation
             
 			if (Instr->NumUsers >= 1)
 			{
-                FString LocalStr = FString::Printf(TEXT("l%d"), NumLocals);
+                FString LocalStr = FString::Printf(TEXT("_%d"), NumLocals);
                 ++NumLocals;
 
                 Printer << LowerType(Instr->Type) << TEXT(" ") << LocalStr;
@@ -454,18 +454,16 @@ struct FTranslator : FMaterialIRToHLSLTranslation
 			const TCHAR* OpString;
 			switch (Operator->Op)
 			{
+				case MIR::BO_Add: OpString = TEXT(" + "); break;
+				case MIR::BO_Divide: OpString = TEXT(" / "); break;
+				case MIR::BO_Equals: OpString = TEXT("=="); break;
 				case MIR::BO_GreaterThan: OpString = TEXT(">"); break;
 				case MIR::BO_GreaterThanOrEquals: OpString = TEXT(">="); break;
 				case MIR::BO_LessThan: OpString = TEXT("<"); break;
 				case MIR::BO_LessThanOrEquals: OpString = TEXT("<="); break;
-				case MIR::BO_Equals: OpString = TEXT("=="); break;
+				case MIR::BO_Multiply: OpString = TEXT(" * "); break;
 				case MIR::BO_NotEquals: OpString = TEXT("!="); break;
-				
-				case MIR::BO_Add: OpString = TEXT(" + "); break;
 				case MIR::BO_Subtract: OpString = TEXT(" - "); break;
-				case MIR::BO_Multiply: OpString = TEXT(" * "); break;
-				case MIR::BO_Divide: OpString = TEXT(" / "); break;
-
 				default: UE_MIR_UNREACHABLE();
 			}
 			Printer << OpString;
@@ -477,33 +475,34 @@ struct FTranslator : FMaterialIRToHLSLTranslation
 			const TCHAR* OpString;
 			switch (Operator->Op)
 			{
+				case MIR::TO_Clamp: OpString = TEXT("clamp"); break;
+				case MIR::BO_Cross: OpString = TEXT("cross"); break;
+				case MIR::BO_Dot: OpString = TEXT("dot"); break;
+				case MIR::BO_Fmod: OpString = TEXT("fmod"); break;
+				case MIR::BO_Max: OpString = TEXT("max"); break;
+				case MIR::BO_Min: OpString = TEXT("min"); break;
+				case MIR::TO_Lerp: OpString = TEXT("lerp"); break;
+				case MIR::TO_Smoothstep: OpString = TEXT("smoothstep"); break;
 				case MIR::UO_Abs: OpString = TEXT("abs"); break;
 				case MIR::UO_ACos: OpString = TEXT("acos"); break;
 				case MIR::UO_ASin: OpString = TEXT("asin"); break;
 				case MIR::UO_ATan: OpString = TEXT("atan"); break;
 				case MIR::UO_Ceil: OpString = TEXT("ceil"); break;
 				case MIR::UO_Cos: OpString = TEXT("cos"); break;
+				case MIR::UO_Exponential: OpString = TEXT("exp"); break;
+				case MIR::UO_Exponential2: OpString = TEXT("exp2"); break;
 				case MIR::UO_Floor: OpString = TEXT("floor"); break;
 				case MIR::UO_Frac: OpString = TEXT("frac"); break;
 				case MIR::UO_Length: OpString = TEXT("length"); break;
+				case MIR::UO_Logarithm: OpString = TEXT("log"); break;
+				case MIR::UO_Logarithm10: OpString = TEXT("log10"); break;
+				case MIR::UO_Logarithm2: OpString = TEXT("log2"); break;
+				case MIR::UO_Round: OpString = TEXT("round"); break;
 				case MIR::UO_Saturate: OpString = TEXT("saturate"); break;
 				case MIR::UO_Sin: OpString = TEXT("sin"); break;
 				case MIR::UO_Sqrt: OpString = TEXT("sqrt"); break;
 				case MIR::UO_Tan: OpString = TEXT("tan"); break;
 				case MIR::UO_Truncate: OpString = TEXT("trunc"); break;
-				case MIR::UO_Exponential: OpString = TEXT("exp"); break;
-				case MIR::UO_Exponential2: OpString = TEXT("exp2"); break;
-				case MIR::UO_Logarithm: OpString = TEXT("log"); break;
-				case MIR::UO_Logarithm2: OpString = TEXT("log2"); break;
-				case MIR::UO_Logarithm10: OpString = TEXT("log10"); break;
-				case MIR::UO_Round: OpString = TEXT("round"); break;
-				case MIR::BO_Fmod: OpString = TEXT("fmod"); break;
-				case MIR::BO_Min: OpString = TEXT("min"); break;
-				case MIR::BO_Max: OpString = TEXT("max"); break;
-				case MIR::BO_Dot: OpString = TEXT("dot"); break;
-				case MIR::BO_Cross: OpString = TEXT("cross"); break;
-				case MIR::TO_Lerp: OpString = TEXT("lerp"); break;
-				case MIR::TO_Smoothstep: OpString = TEXT("smoothstep"); break;
 				default: UE_MIR_UNREACHABLE();
 			}
 
diff --git a/Engine/Source/Runtime/Engine/Public/Materials/MaterialExpressionClamp.h b/Engine/Source/Runtime/Engine/Public/Materials/MaterialExpressionClamp.h
--- a/Engine/Source/Runtime/Engine/Public/Materials/MaterialExpressionClamp.h
+++ b/Engine/Source/Runtime/Engine/Public/Materials/MaterialExpressionClamp.h
@@ -49,6 +49,7 @@ public:
 
 	//~ Begin UMaterialExpression Interface
 #if WITH_EDITOR
+	virtual void Build(UE::MIR::FEmitter& Emitter) override;
 	virtual int32 Compile(class FMaterialCompiler* Compiler, int32 OutputIndex) override;
 	virtual void GetCaption(TArray<FString>& OutCaptions) const override;
 #endif
diff --git a/Engine/Source/Runtime/Engine/Public/Materials/MaterialIR.h b/Engine/Source/Runtime/Engine/Public/Materials/MaterialIR.h
--- a/Engine/Source/Runtime/Engine/Public/Materials/MaterialIR.h
+++ b/Engine/Source/Runtime/Engine/Public/Materials/MaterialIR.h
@@ -90,16 +90,16 @@ struct FValue
 	bool IsFalse() const;
 
 	// Returns whether this value is arithmetic and exactly zero.
-	bool IsExactlyZero() const;
+	bool AreAllExactlyZero() const;
 
 	// Returns whether this value is arithmetic and approximately zero.
-	bool IsNearlyZero() const;
+	bool AreAllNearlyZero() const;
 
 	// Returns whether this value is arithmetic and exactly one.
-	bool IsExactlyOne() const;
+	bool AreAllExactlyOne() const;
 
 	// Returns whether this value is arithmetic and approximately one.
-	bool IsNearlyOne() const;
+	bool AreAllNearlyOne() const;
 	
 	// Returns this value's textures if it has one (nullptr otherwise).
 	TOptional<TPair<UTexture*, EMaterialSamplerType>> GetTexture();
@@ -365,7 +365,8 @@ enum EOperator
 	/* Ternary Arithmetic Operators */
 	TO_FirstTernaryOperator,
 
-	TO_Lerp = TO_FirstTernaryOperator,
+	TO_Clamp = TO_FirstTernaryOperator,
+	TO_Lerp,
 	TO_Smoothstep,
 
 	OperatorCount,
@@ -378,7 +379,7 @@ bool IsComparisonOperator(EOperator Op);
 bool IsBitOperator(EOperator Op);
 
 //
-bool IsBinaryOperator(EOperator Op);
+bool IsUnaryOperator(EOperator Op);
 
 //
 bool IsBinaryOperator(EOperator Op);
