[DebugDraw] added an associated world to FDebugDrawDelegateHelper to only draw text labels on the associated canvas

diff --git a/Engine/Source/Runtime/AIModule/Private/EnvironmentQuery/EQSRenderingComponent.cpp b/Engine/Source/Runtime/AIModule/Private/EnvironmentQuery/EQSRenderingComponent.cpp
--- a/Engine/Source/Runtime/AIModule/Private/EnvironmentQuery/EQSRenderingComponent.cpp
+++ b/Engine/Source/Runtime/AIModule/Private/EnvironmentQuery/EQSRenderingComponent.cpp
@@ -233,7 +233,7 @@ void FEQSSceneProxy::CollectEQSData(const FEnvQueryResult* ResultItems, const FE
 		}
 	}
 }
- 
+
 void FEQSRenderingDebugDrawDelegateHelper::DrawDebugLabels(UCanvas* Canvas, APlayerController* PC)
 {
 	if ( !ActorOwner
@@ -243,12 +243,6 @@ void FEQSRenderingDebugDrawDelegateHelper::DrawDebugLabels(UCanvas* Canvas, APla
 		return;
 	}
 
-	// little hacky test but it's the only way to remove text rendering from bad worlds, when using UDebugDrawService for it
-	if (Canvas && Canvas->SceneView && Canvas->SceneView->Family && Canvas->SceneView->Family->Scene && Canvas->SceneView->Family->Scene->GetWorld() != ActorOwner->GetWorld())
-	{
-		return;
-	}
-
 	FDebugDrawDelegateHelper::DrawDebugLabels(Canvas, PC);
 }
 #endif //USE_EQS_DEBUGGER
diff --git a/Engine/Source/Runtime/Engine/Private/DebugRenderSceneProxy.cpp b/Engine/Source/Runtime/Engine/Private/DebugRenderSceneProxy.cpp
--- a/Engine/Source/Runtime/Engine/Private/DebugRenderSceneProxy.cpp
+++ b/Engine/Source/Runtime/Engine/Private/DebugRenderSceneProxy.cpp
@@ -15,6 +15,7 @@
 #include "Materials/Material.h"
 #include "Materials/MaterialRenderProxy.h"
 #include "MeshElementCollector.h"
+#include "SceneInterface.h"
 
 static TAutoConsoleVariable<bool> CVarDebugRenderAllowFrustumCulling(
 	TEXT("r.DebugRender.AllowFrustumCulling"), true,
@@ -48,12 +49,25 @@ void FDebugDrawDelegateHelper::RegisterDebugDrawDelegateInternal()
 	// called in multiple scenarios, most notably blueprint recompilation or changing level visibility in the editor.
 	if (State == InitializedState)
 	{
-		DebugTextDrawingDelegate = FDebugDrawDelegate::CreateRaw(this, &FDebugDrawDelegateHelper::DrawDebugLabels);
+		DebugTextDrawingDelegate = FDebugDrawDelegate::CreateRaw(this, &FDebugDrawDelegateHelper::HandleDrawDebugLabels);
 		DebugTextDrawingDelegateHandle = UDebugDrawService::Register(*ViewFlagName, DebugTextDrawingDelegate);
 		State = RegisteredState;
 	}
 }
 
+void FDebugDrawDelegateHelper::InitDelegateHelper(const FDebugRenderSceneProxy* InSceneProxy)
+{
+	check(IsInParallelGameThread() || IsInGameThread());
+
+	Texts.Reset();
+	Texts.Append(InSceneProxy->Texts);
+	ViewFlagName = InSceneProxy->ViewFlagName;
+	TextWithoutShadowDistance = InSceneProxy->TextWithoutShadowDistance;
+	State = (State == UndefinedState) ? InitializedState : State;
+	FarClippingDistance = InSceneProxy->FarClippingDistance;
+	AssociatedWorld = InSceneProxy->GetScene().GetWorld();
+}
+
 void FDebugDrawDelegateHelper::RequestRegisterDebugDrawDelegate(FRegisterComponentContext* Context)
 {
 	bDeferredRegister = Context != nullptr;
@@ -111,10 +125,24 @@ uint32 FDebugRenderSceneProxy::GetAllocatedSize(void) const
 }
 
 
-void FDebugDrawDelegateHelper::DrawDebugLabels(UCanvas* Canvas, APlayerController*)
+void FDebugDrawDelegateHelper::HandleDrawDebugLabels(UCanvas* Canvas, APlayerController* PlayerController)
 {
-	TRACE_CPUPROFILER_EVENT_SCOPE(FDebugDrawDelegateHelper::DrawDebugLabels);
+	TRACE_CPUPROFILER_EVENT_SCOPE(FDebugDrawDelegateHelper::OnDrawDebugLabels);
+
+	check (Canvas);
+
+	// Ignore if the helper's associated scene doesn't match the one associated to the canvas
+	if (Canvas->Canvas != nullptr
+		&& AssociatedWorld != Canvas->Canvas->GetScene()->GetWorld())
+	{
+		return;
+	}
+
+	DrawDebugLabels(Canvas, PlayerController);
+}
 
+void FDebugDrawDelegateHelper::DrawDebugLabels(UCanvas* Canvas, APlayerController*)
+{
 	const FColor OldDrawColor = Canvas->DrawColor;
 	const FFontRenderInfo FontRenderInfo = Canvas->CreateFontRenderInfo(true, false);
 	const FFontRenderInfo FontRenderInfoWithShadow = Canvas->CreateFontRenderInfo(true, true);
diff --git a/Engine/Source/Runtime/Engine/Public/DebugRenderSceneProxy.h b/Engine/Source/Runtime/Engine/Public/DebugRenderSceneProxy.h
--- a/Engine/Source/Runtime/Engine/Public/DebugRenderSceneProxy.h
+++ b/Engine/Source/Runtime/Engine/Public/DebugRenderSceneProxy.h
@@ -427,26 +427,7 @@ protected:
 	typedef TArray<FDebugRenderSceneProxy::FText3d> TextArray;
 
 public:
-	void InitDelegateHelper(const FDebugRenderSceneProxy* InSceneProxy)
-	{
-		check(IsInParallelGameThread() || IsInGameThread());
-
-		Texts.Reset();
-		Texts.Append(InSceneProxy->Texts);
-		ViewFlagName = InSceneProxy->ViewFlagName;
-		TextWithoutShadowDistance = InSceneProxy->TextWithoutShadowDistance;
-		State = (State == UndefinedState) ? InitializedState : State;
-		FarClippingDistance = InSceneProxy->FarClippingDistance;
-	}
-
-	UE_DEPRECATED(5.0, "This method is deprecated. Call RequestRegisterDebugDrawDelegate or override RegisterDebugDrawDelegateInternal instead.")
-	virtual void RegisterDebugDrawDelgate() final { RegisterDebugDrawDelegateInternal(); }
-	UE_DEPRECATED(5.0, "This method is deprecated. Use UnregisterDebugDrawDelegate instead.")
-	virtual void UnregisterDebugDrawDelgate() { UnregisterDebugDrawDelegate(); }
-	UE_DEPRECATED(5.0, "This method is deprecated. Use ReregisterDebugDrawDelegate instead.")
-	void ReregisterDebugDrawDelgate() { ReregisterDebugDrawDelegate(); }
-	UE_DEPRECATED(5.0, "This method is deprecated. Call RequestRegisterDebugDrawDelegate or override RegisterDebugDrawDelegateInternal instead.")
-	virtual void RegisterDebugDrawDelegate() final { RegisterDebugDrawDelegateInternal(); }
+	ENGINE_API void InitDelegateHelper(const FDebugRenderSceneProxy* InSceneProxy);
 
 	/**
 	 * Method that should be called at render state creation (i.e. CreateRenderState_Concurrent).
@@ -472,7 +453,8 @@ protected:
 	/** called to set up debug drawing delegate in UDebugDrawService if you want to draw labels */
 	ENGINE_API virtual void RegisterDebugDrawDelegateInternal();
 
-	ENGINE_API virtual void DrawDebugLabels(UCanvas* Canvas, APlayerController*);
+	ENGINE_API void HandleDrawDebugLabels(UCanvas* Canvas, APlayerController* PlayerController);
+	ENGINE_API virtual void DrawDebugLabels(UCanvas* Canvas, APlayerController* PlayerController);
 	void ResetTexts() { Texts.Reset(); }
 	const TextArray& GetTexts() const { return Texts; }
 	float GetTextWithoutShadowDistance() const {return TextWithoutShadowDistance; }
@@ -491,6 +473,12 @@ protected:
 	bool bDeferredRegister = false;
 
 private:
+	/**
+	 * Weak pointer to the associated world used when drawing labels to skip canvas from a different world.
+	 * This is required since the current implementation relies on a global delegate so multiple canvas might
+	 * be active (e.g., LevelEditor + any asset editor with viewport)
+	 */
+	TWeakObjectPtr<UWorld> AssociatedWorld;
 	TextArray Texts;
 	FString ViewFlagName;
 	float TextWithoutShadowDistance;
