Remove the unused "Input time" engine stat.

The "Input Time" stat used to measure the actual time it takes to poll input back in UE3. UE input has since been refactored long ago, and this stat is no longer in use anywhere.

The "input time" engine stat was originally added in 2004 for UE3 with the change 69619. Since then it has been refactored to a scoped cycle stat with 104116. The use of this scoped cycle stat was eventually removed all together with 1686667, never to be seen again...

#rnx

diff --git a/Engine/Source/Runtime/Engine/Private/UnrealEngine.cpp b/Engine/Source/Runtime/Engine/Private/UnrealEngine.cpp
--- a/Engine/Source/Runtime/Engine/Private/UnrealEngine.cpp
+++ b/Engine/Source/Runtime/Engine/Private/UnrealEngine.cpp
@@ -13417,7 +13417,6 @@ DEFINE_STAT(STAT_IntentionalHitch);
 DEFINE_STAT(STAT_DeferredTickTime);
 
 /** Input stat */
-DEFINE_STAT(STAT_InputTime);
 DEFINE_STAT(STAT_InputLatencyTime);
 
 /** HUD stat */
diff --git a/Engine/Source/Runtime/Engine/Public/EngineStats.h b/Engine/Source/Runtime/Engine/Public/EngineStats.h
--- a/Engine/Source/Runtime/Engine/Public/EngineStats.h
+++ b/Engine/Source/Runtime/Engine/Public/EngineStats.h
@@ -9,7 +9,6 @@
  * Engine stats
  */
 /** Input stat */
-DECLARE_CYCLE_STAT_EXTERN(TEXT("Input Time"),STAT_InputTime,STATGROUP_Engine, );
 DECLARE_FLOAT_COUNTER_STAT_EXTERN(TEXT("Input Latency"),STAT_InputLatencyTime,STATGROUP_Engine, );
 /** HUD stat */
 DECLARE_CYCLE_STAT_EXTERN(TEXT("HUD Time"),STAT_HudTime,STATGROUP_Engine, );
