Fixed LightComponent generating non-deterministic GUID in cook.

diff --git a/Engine/Source/Runtime/Engine/Private/Components/LightComponent.cpp b/Engine/Source/Runtime/Engine/Private/Components/LightComponent.cpp
--- a/Engine/Source/Runtime/Engine/Private/Components/LightComponent.cpp
+++ b/Engine/Source/Runtime/Engine/Private/Components/LightComponent.cpp
@@ -263,7 +263,23 @@ void ULightComponentBase::ValidateLightGUIDs()
 
 void ULightComponentBase::UpdateLightGUIDs()
 {
-	OriginalLightGuid = (HasStaticShadowing() ? FGuid::NewGuid() : FGuid());
+	if (HasStaticShadowing())
+	{
+#if WITH_EDITOR
+		if (IsRunningCookCommandlet())
+		{
+			OriginalLightGuid = FGuid::NewDeterministicGuid(GetFullName());
+		}
+		else
+#endif
+		{
+			OriginalLightGuid = FGuid::NewGuid();
+		}
+	}
+	else
+	{
+		OriginalLightGuid = FGuid();
+	}
 	LightGuid = (OriginalLightGuid.IsValid() && GetOwner() )? FGuid::Combine( OriginalLightGuid , FActorInstanceGuid::GetActorInstanceGuid(*GetOwner())) : FGuid();
 }
 
