World Partition Streaming Generation: resolve the default runtime grid in the actor descriptor view when an actor has its grid set to none to avoid false positives when validating runtime grids.

diff --git a/Engine/Source/Runtime/Engine/Private/WorldPartition/RuntimeHashSet/WorldPartitionRuntimeHashSet.cpp b/Engine/Source/Runtime/Engine/Private/WorldPartition/RuntimeHashSet/WorldPartitionRuntimeHashSet.cpp
--- a/Engine/Source/Runtime/Engine/Private/WorldPartition/RuntimeHashSet/WorldPartitionRuntimeHashSet.cpp
+++ b/Engine/Source/Runtime/Engine/Private/WorldPartition/RuntimeHashSet/WorldPartitionRuntimeHashSet.cpp
@@ -296,6 +296,11 @@ void UWorldPartitionRuntimeHashSet::FlushStreamingContent()
 	UpdateRuntimeDataGridMap();
 }
 
+FName UWorldPartitionRuntimeHashSet::GetDefaultGrid() const
+{
+	return RuntimePartitions.Num() ? RuntimePartitions[0].Name : NAME_None;
+}
+
 bool UWorldPartitionRuntimeHashSet::IsValidGrid(FName GridName, const UClass* ActorClass) const
 {
 	return GridName.IsNone() || !!ResolveRuntimePartition(GridName);
diff --git a/Engine/Source/Runtime/Engine/Private/WorldPartition/WorldPartitionRuntimeSpatialHash.cpp b/Engine/Source/Runtime/Engine/Private/WorldPartition/WorldPartitionRuntimeSpatialHash.cpp
--- a/Engine/Source/Runtime/Engine/Private/WorldPartition/WorldPartitionRuntimeSpatialHash.cpp
+++ b/Engine/Source/Runtime/Engine/Private/WorldPartition/WorldPartitionRuntimeSpatialHash.cpp
@@ -1689,6 +1689,11 @@ void UWorldPartitionRuntimeSpatialHash::FlushStreamingContent()
 	StreamingGrids.Empty();
 }
 
+FName UWorldPartitionRuntimeSpatialHash::GetDefaultGrid() const
+{
+	return Grids.Num() ? Grids[0].GridName : NAME_None;
+}
+
 bool UWorldPartitionRuntimeSpatialHash::IsValidGrid(FName GridName, const UClass* ActorClass) const
 {
 	// We always consider grids from auto-generated HLODs to be valid
diff --git a/Engine/Source/Runtime/Engine/Private/WorldPartition/WorldPartitionStreamingGeneration.cpp b/Engine/Source/Runtime/Engine/Private/WorldPartition/WorldPartitionStreamingGeneration.cpp
--- a/Engine/Source/Runtime/Engine/Private/WorldPartition/WorldPartitionStreamingGeneration.cpp
+++ b/Engine/Source/Runtime/Engine/Private/WorldPartition/WorldPartitionStreamingGeneration.cpp
@@ -186,6 +186,11 @@ FName FStreamingGenerationActorDescView::GetRuntimeGrid() const
 		return ParentView->GetRuntimeGrid();
 	}
 
+	if (ResolvedRuntimeGrid.IsSet())
+	{
+		return ResolvedRuntimeGrid.GetValue();
+	}
+
 	return Super::GetRuntimeGrid();
 }
 
@@ -269,7 +274,8 @@ FString FStreamingGenerationActorDescView::ToString(FWorldPartitionActorDesc::ET
 	FString Result = Super::ToString(Mode);
 	if (Mode >= FWorldPartitionActorDesc::EToStringMode::Compact)
 	{
-		Result += FString::Printf(TEXT(" View=(RuntimeHLODLayer=%s bIsForcedNonSpatiallyLoaded=%s bIsForcedNoRuntimeGrid=%s bIsForcedNoDataLayers=%s bIsForceNoHLODLayer=%s bIsForcedHLODIrrelevant=%s)"), 
+		Result += FString::Printf(TEXT(" View=(RuntimeGrid=%s RuntimeHLODLayer=%s bIsForcedNonSpatiallyLoaded=%s bIsForcedNoRuntimeGrid=%s bIsForcedNoDataLayers=%s bIsForceNoHLODLayer=%s bIsForcedHLODIrrelevant=%s)"), 
+			ResolvedRuntimeGrid.IsSet() ? *ResolvedRuntimeGrid->ToString() : TEXT("none"),
 			RuntimeHLODLayer.IsSet() ? *RuntimeHLODLayer->ToString() : TEXT("none"),
 			*LexToString(bIsForcedNonSpatiallyLoaded),
 			*LexToString(bIsForcedNoRuntimeGrid),
@@ -379,6 +385,11 @@ void FStreamingGenerationActorDescView::SetRuntimeDataLayerInstanceNames(const F
 	}
 }
 
+void FStreamingGenerationActorDescView::SetRuntimeGrid(FName InRuntimeGrid)
+{
+	ResolvedRuntimeGrid = InRuntimeGrid;
+}
+
 void FStreamingGenerationActorDescView::SetRuntimeReferences(const TArray<FGuid>& InRuntimeReferences)
 {
 	RuntimeReferences = InRuntimeReferences;
@@ -728,6 +739,10 @@ class FWorldPartitionStreamingGenerator
 		{
 			ActorDescView.SetForcedNoRuntimeGrid();
 		}
+		else if (ActorDescView.GetRuntimeGrid().IsNone())
+		{			
+			ActorDescView.SetRuntimeGrid(DefaultGrid);
+		}
 	}
 
 	void ResolveRuntimeDataLayers(FStreamingGenerationActorDescView& ActorDescView, const TArray<const FWorldDataLayersActorDesc*>& InDataLayerResolvers)
@@ -1746,6 +1761,7 @@ public:
 		bool bIsPIE;
 		TMap<FGuid, const UActorDescContainerInstance*> ActorGuidsToContainerInstanceMap;
 		TArray<TSubclassOf<AActor>> FilteredClasses;
+		FName DefaultGrid;
 		TFunction<bool(FName, const UClass*)> IsValidGrid;
 		TFunction<bool(FName GridName, const FSoftObjectPath&)> IsValidHLODLayer;
 
@@ -1757,6 +1773,7 @@ public:
 		FWorldPartitionStreamingGeneratorParams& SetIsPIE(bool bInIsPIE) { bIsPIE = bInIsPIE; return *this; }
 		FWorldPartitionStreamingGeneratorParams& SetActorGuidsToContainerInstanceMap(const TMap<FGuid, const UActorDescContainerInstance*>& InActorGuidsToContainerInstanceMap) { ActorGuidsToContainerInstanceMap = InActorGuidsToContainerInstanceMap; return *this; }
 		FWorldPartitionStreamingGeneratorParams& SetFilteredClasses(const TArray<TSubclassOf<AActor>>& InFilteredClasses) { FilteredClasses = InFilteredClasses; return *this; }
+		FWorldPartitionStreamingGeneratorParams& SetDefaultGrid(FName InDefaultGrid) { DefaultGrid = InDefaultGrid; return *this; }
 		FWorldPartitionStreamingGeneratorParams& SetIsValidGrid(TFunction<bool(FName, const UClass*)> InIsValidGrid) { IsValidGrid = InIsValidGrid; return *this; }
 		FWorldPartitionStreamingGeneratorParams& SetIsValidHLODLayer(TFunction<bool(FName GridName, const FSoftObjectPath&)> InIsValidHLODLayer) { IsValidHLODLayer = InIsValidHLODLayer; return *this; }
 				
@@ -1770,6 +1787,7 @@ public:
 		, bHandleUnsavedActors(Params.bHandleUnsavedActors)
 		, bIsPIE(Params.bIsPIE)
 		, FilteredClasses(Params.FilteredClasses)
+		, DefaultGrid(Params.DefaultGrid)
 		, IsValidGrid(Params.IsValidGrid)
 		, IsValidHLODLayer(Params.IsValidHLODLayer)
 		, ErrorHandler(Params.ErrorHandler)
@@ -1780,7 +1798,6 @@ public:
 		DataLayerManager = OwningWorld ? WorldPartitionContext->GetDataLayerManager() : nullptr;
 	}
 
-	
 	void PreparationPhase(const FStreamingGenerationContainerInstanceCollection& ContainerInstanceCollection)
 	{
 		CreateActorContainers(ContainerInstanceCollection);
@@ -1941,6 +1958,7 @@ private:
 	bool bHandleUnsavedActors;
 	bool bIsPIE;
 	TArray<TSubclassOf<AActor>> FilteredClasses;
+	FName DefaultGrid;
 	TFunction<bool(FName, const UClass*)> IsValidGrid;
 	TFunction<bool(FName, const FSoftObjectPath&)> IsValidHLODLayer;
 	IStreamingGenerationErrorHandler* ErrorHandler;
@@ -2014,6 +2032,7 @@ bool UWorldPartition::GenerateContainerStreaming(const FGenerateStreamingParams&
 		.SetWorldPartitionContext(this)
 		.SetHandleUnsavedActors(bIsPIE && !GetTypedOuter<UWorld>()->IsGameWorld())
 		.SetIsPIE(bIsPIE)
+		.SetDefaultGrid(RuntimeHash ? RuntimeHash->GetDefaultGrid() : NAME_None)
 		.SetIsValidGrid([this](FName GridName, const UClass* ActorClass) { return RuntimeHash->IsValidGrid(GridName, ActorClass); })
 		.SetIsValidHLODLayer([this](FName GridName, const FSoftObjectPath& HLODLayerPath) { return RuntimeHash->IsValidHLODLayer(GridName, HLODLayerPath); })
 		.SetErrorHandler(ErrorHandlerSelector.Get())
@@ -2114,6 +2133,7 @@ TUniquePtr<IStreamingGenerationContext> UWorldPartition::GenerateStreamingGenera
 
 	FWorldPartitionStreamingGenerator::FWorldPartitionStreamingGeneratorParams StreamingGeneratorParams = FWorldPartitionStreamingGenerator::FWorldPartitionStreamingGeneratorParams()
 		.SetWorldPartitionContext(this)
+		.SetDefaultGrid(RuntimeHash ? RuntimeHash->GetDefaultGrid() : NAME_None)
 		.SetIsValidGrid([this](FName GridName, const UClass* ActorClass) { return RuntimeHash->IsValidGrid(GridName, ActorClass); })
 		.SetIsValidHLODLayer([this](FName GridName, const FSoftObjectPath& HLODLayerPath) { return RuntimeHash->IsValidHLODLayer(GridName, HLODLayerPath); })
 		.SetErrorHandler(ErrorHandlerSelector.Get())
@@ -2184,6 +2204,7 @@ void UWorldPartition::SetupHLODActors(const FSetupHLODActorsParams& Params)
 			.SetErrorHandler(ErrorHandlerSelector.Get())
 			.SetEnableStreaming(IsStreamingEnabled())
 			.SetFilteredClasses({ AWorldPartitionHLOD::StaticClass() })
+			.SetDefaultGrid(RuntimeHash ? RuntimeHash->GetDefaultGrid() : NAME_None)
 			.SetIsValidGrid([this](FName GridName, const UClass* ActorClass) { return RuntimeHash->IsValidGrid(GridName, ActorClass); })
 			.SetIsValidHLODLayer([this](FName GridName, const FSoftObjectPath& HLODLayerPath) { return RuntimeHash->IsValidHLODLayer(GridName, HLODLayerPath); });
 
@@ -2418,6 +2439,7 @@ void UWorldPartition::CheckForErrors(const FCheckForErrorsParams& InParams)
 			.SetHandleUnsavedActors(!!WorldPartition)
 			.SetIsPIE(WorldPartition ? WorldPartition->bIsPIE : false)
 			.SetErrorHandler(ErrorHandlerSelector.Get())
+			.SetDefaultGrid(WorldPartitionRuntimeHash ? WorldPartitionRuntimeHash->GetDefaultGrid() : NAME_None)
 			.SetIsValidGrid([WorldPartitionRuntimeHash](FName GridName, const UClass* ActorClass) { return WorldPartitionRuntimeHash ? WorldPartitionRuntimeHash->IsValidGrid(GridName, ActorClass) : true; })
 			.SetIsValidHLODLayer([WorldPartitionRuntimeHash](FName GridName, const FSoftObjectPath& HLODLayerPath) { return WorldPartitionRuntimeHash ? WorldPartitionRuntimeHash->IsValidHLODLayer(GridName, HLODLayerPath) : true; })
 			.SetEnableStreaming(InParams.bEnableStreaming)
diff --git a/Engine/Source/Runtime/Engine/Public/WorldPartition/RuntimeHashSet/WorldPartitionRuntimeHashSet.h b/Engine/Source/Runtime/Engine/Public/WorldPartition/RuntimeHashSet/WorldPartitionRuntimeHashSet.h
--- a/Engine/Source/Runtime/Engine/Public/WorldPartition/RuntimeHashSet/WorldPartitionRuntimeHashSet.h
+++ b/Engine/Source/Runtime/Engine/Public/WorldPartition/RuntimeHashSet/WorldPartitionRuntimeHashSet.h
@@ -166,6 +166,7 @@ public:
 	ENGINE_API virtual bool SupportsHLODs() const override;
 	ENGINE_API virtual bool SetupHLODActors(const IStreamingGenerationContext* StreamingGenerationContext, const UWorldPartition::FSetupHLODActorsParams& Params) const override;
 	ENGINE_API virtual bool GenerateStreaming(class UWorldPartitionStreamingPolicy* StreamingPolicy, const IStreamingGenerationContext* StreamingGenerationContext, TArray<FString>* OutPackagesToGenerate) override;
+	ENGINE_API virtual FName GetDefaultGrid() const override;
 	ENGINE_API virtual bool IsValidGrid(FName GridName, const UClass* ActorClass) const override;
 	ENGINE_API virtual bool IsValidHLODLayer(FName GridName, const FSoftObjectPath& HLODLayerPath) const override;
 	ENGINE_API virtual void DumpStateLog(FHierarchicalLogArchive& Ar) const override;
diff --git a/Engine/Source/Runtime/Engine/Public/WorldPartition/WorldPartitionRuntimeHash.h b/Engine/Source/Runtime/Engine/Public/WorldPartition/WorldPartitionRuntimeHash.h
--- a/Engine/Source/Runtime/Engine/Public/WorldPartition/WorldPartitionRuntimeHash.h
+++ b/Engine/Source/Runtime/Engine/Public/WorldPartition/WorldPartitionRuntimeHash.h
@@ -189,6 +189,7 @@ class UWorldPartitionRuntimeHash : public UObject
 	ENGINE_API virtual TArray<UWorldPartitionRuntimeCell*> GetAlwaysLoadedCells() const;
 	ENGINE_API virtual bool GenerateStreaming(class UWorldPartitionStreamingPolicy* StreamingPolicy, const IStreamingGenerationContext* StreamingGenerationContext, TArray<FString>* OutPackagesToGenerate);
 	virtual bool SetupHLODActors(const IStreamingGenerationContext* StreamingGenerationContext, const UWorldPartition::FSetupHLODActorsParams& Params) const { return false; }
+	virtual FName GetDefaultGrid() const { return NAME_None; }
 	virtual bool IsValidGrid(FName GridName, const UClass* ActorClass) const { return false; }
 	virtual bool IsValidHLODLayer(FName GridName, const FSoftObjectPath& HLODLayerPath) const { return false; }
 	virtual void DrawPreview() const {}
diff --git a/Engine/Source/Runtime/Engine/Public/WorldPartition/WorldPartitionRuntimeSpatialHash.h b/Engine/Source/Runtime/Engine/Public/WorldPartition/WorldPartitionRuntimeSpatialHash.h
--- a/Engine/Source/Runtime/Engine/Public/WorldPartition/WorldPartitionRuntimeSpatialHash.h
+++ b/Engine/Source/Runtime/Engine/Public/WorldPartition/WorldPartitionRuntimeSpatialHash.h
@@ -326,6 +326,7 @@ public:
 	ENGINE_API virtual void PreSetupHLODActors(const UWorldPartition* InWorldPartition, const UWorldPartition::FSetupHLODActorsParams& InParams) const override;
 	ENGINE_API virtual void PostSetupHLODActors(const UWorldPartition* InWorldPartition, const UWorldPartition::FSetupHLODActorsParams& InParams) const override;
 	ENGINE_API virtual bool SetupHLODActors(const IStreamingGenerationContext* StreamingGenerationContext, const UWorldPartition::FSetupHLODActorsParams& Params) const override;
+	ENGINE_API virtual FName GetDefaultGrid() const override;
 	ENGINE_API virtual bool IsValidGrid(FName GridName, const UClass* ActorClass) const override;
 	ENGINE_API virtual bool IsValidHLODLayer(FName GridName, const FSoftObjectPath& HLODLayerPath) const override { return true; }
 	ENGINE_API virtual void DrawPreview() const override;
diff --git a/Engine/Source/Runtime/Engine/Public/WorldPartition/WorldPartitionStreamingGeneration.h b/Engine/Source/Runtime/Engine/Public/WorldPartition/WorldPartitionStreamingGeneration.h
--- a/Engine/Source/Runtime/Engine/Public/WorldPartition/WorldPartitionStreamingGeneration.h
+++ b/Engine/Source/Runtime/Engine/Public/WorldPartition/WorldPartitionStreamingGeneration.h
@@ -75,6 +75,7 @@ public:
 	ENGINE_API void SetForcedNonSpatiallyLoaded();
 	ENGINE_API void SetForcedNoRuntimeGrid();
 	ENGINE_API void SetForcedNoDataLayers();
+	ENGINE_API void SetRuntimeGrid(FName InRuntimeGrid);
 	ENGINE_API void SetRuntimeDataLayerInstanceNames(const FDataLayerInstanceNames& InRuntimeDataLayerInstanceNames);
 	ENGINE_API void SetRuntimeReferences(const TArray<FGuid>& InRuntimeReferences);
 	ENGINE_API void SetEditorReferences(const TArray<FGuid>& InEditorReferences);
@@ -114,6 +115,7 @@ private:
 	// Used by GetDataLayerInstanceNames() to avoid returning a copy
 	mutable FDataLayerInstanceNames LastReturnedDataLayerInstanceNames; 
 
+	TOptional<FName> ResolvedRuntimeGrid;
 	TOptional<FDataLayerInstanceNames> ResolvedDataLayerInstanceNames;
 	TOptional<FDataLayerInstanceNames> RuntimeDataLayerInstanceNames;
 	TOptional<TArray<FGuid>> RuntimeReferences;
