Fix for RAC failing on some sample rates

diff --git a/Engine/Source/Runtime/Engine/Private/AudioFormatSettings.cpp b/Engine/Source/Runtime/Engine/Private/AudioFormatSettings.cpp
--- a/Engine/Source/Runtime/Engine/Private/AudioFormatSettings.cpp
+++ b/Engine/Source/Runtime/Engine/Private/AudioFormatSettings.cpp
@@ -71,10 +71,10 @@ namespace Audio
 			, SampleRate(InWave->GetImportedSampleRate())
 			, NumChannels(InWave->NumChannels)
 		{
-			// Override sample-rate?
+			// Override sample-rate? (only if resample for device is true).
 			if (const FPlatformAudioCookOverrides* CookOverrides = FPlatformCompressionUtilities::GetCookOverrides(*InFormatSettings->IniPlatformName.ToString()) )
 			{
-				if (const float SampleRateOverride = InWave->GetSampleRateForCompressionOverrides(CookOverrides); SampleRateOverride != -1.f)
+				if (const float SampleRateOverride = InWave->GetSampleRateForCompressionOverrides(CookOverrides); CookOverrides->bResampleForDevice && SampleRateOverride != -1.f)
 				{
 					SampleRate = SampleRateOverride;
 				}
@@ -119,7 +119,7 @@ namespace Audio
 	}
 
 	bool FAudioFormatSettings::IsFormatAllowed(const FPlatformWaveState& InWave) const 
-	{
+	{		
 		if (const IAudioFormat* Format = FindFormat(InWave.FormatName))
 		{
 			// Platform supported?
