Add distance-based teleport settings to ClothComponent.

This matches behavior of SKM-based cloth.

diff --git a/Engine/Plugins/ChaosClothAsset/Source/ChaosClothAssetEngine/Private/ChaosClothAsset/ClothComponent.cpp b/Engine/Plugins/ChaosClothAsset/Source/ChaosClothAssetEngine/Private/ChaosClothAsset/ClothComponent.cpp
--- a/Engine/Plugins/ChaosClothAsset/Source/ChaosClothAssetEngine/Private/ChaosClothAsset/ClothComponent.cpp
+++ b/Engine/Plugins/ChaosClothAsset/Source/ChaosClothAssetEngine/Private/ChaosClothAsset/ClothComponent.cpp
@@ -8,6 +8,7 @@
 #include "ChaosClothAsset/ClothSimulationModel.h"
 #include "ChaosClothAsset/ClothSimulationProxy.h"
 #include "Chaos/CollectionPropertyFacade.h"
+#include "ClothingSimulationTeleportHelpers.h"
 #include "Components/SkeletalMeshComponent.h"
 #include "GeometryCollection/ManagedArrayCollection.h"
 #include "HAL/IConsoleManager.h"
@@ -46,9 +47,14 @@ UChaosClothComponent::UChaosClothComponent(const FObjectInitializer& ObjectIniti
 #if WITH_EDITORONLY_DATA
 	, bSimulateInEditor(0)
 #endif
+	, TeleportDistanceThreshold(300.f)
+	, TeleportRotationThreshold(0.f)
+	, ClothTeleportDistThresholdSquared(UE::ClothingSimulation::TeleportHelpers::ComputeTeleportDistanceThresholdSquared(TeleportDistanceThreshold))
+	, ClothTeleportCosineThresholdInRad(UE::ClothingSimulation::TeleportHelpers::ComputeTeleportCosineRotationThreshold(TeleportRotationThreshold))
 	, CollisionSources(MakeUnique<UE::Chaos::ClothAsset::FCollisionSources>(this, (bool)bCollideWithEnvironment))
 {
 	PrimaryComponentTick.EndTickGroup = TG_PostPhysics;
+	PrevRootBoneMatrix = GetBoneMatrix(0);
 }
 
 UChaosClothComponent::UChaosClothComponent(FVTableHelper& Helper)
@@ -202,6 +208,14 @@ PRAGMA_ENABLE_DEPRECATION_WARNINGS
 		{
 			CollisionSources->SetCollideWithEnvironment(bCollideWithEnvironment);
 		}
+		if (Property->GetFName() == GET_MEMBER_NAME_CHECKED(UChaosClothComponent, TeleportDistanceThreshold))
+		{
+			ClothTeleportDistThresholdSquared = UE::ClothingSimulation::TeleportHelpers::ComputeTeleportDistanceThresholdSquared(TeleportDistanceThreshold);
+		}
+		if (Property->GetFName() == GET_MEMBER_NAME_CHECKED(UChaosClothComponent, TeleportRotationThreshold))
+		{
+			ClothTeleportCosineThresholdInRad = UE::ClothingSimulation::TeleportHelpers::ComputeTeleportCosineRotationThreshold(TeleportRotationThreshold);
+		}
 	}
 
 	Super::PostEditChangeProperty(PropertyChangedEvent);
@@ -281,6 +295,11 @@ void UChaosClothComponent::TickComponent(float DeltaTime, enum ELevelTick TickTy
 	
 	// Make sure that the previous frame simulation has completed
 	HandleExistingParallelSimulation();
+
+	if (ClothSimulationProxy.IsValid())
+	{
+		UpdateClothTeleport();
+	}
 	
 	if (!SimulationAsset.DataflowAsset)
 	{
@@ -624,6 +643,29 @@ void UChaosClothComponent::StartNewParallelSimulation(float DeltaTime)
 	}
 }
 
+void UChaosClothComponent::UpdateClothTeleport()
+{
+	const FMatrix CurRootBoneMat = GetBoneMatrix(0);
+	EClothingTeleportMode TeleportMode = bTeleport ? (bReset ? EClothingTeleportMode::TeleportAndReset : EClothingTeleportMode::Teleport) : EClothingTeleportMode::None;
+	TeleportMode = UE::ClothingSimulation::TeleportHelpers::CalculateClothingTeleport(TeleportMode, CurRootBoneMat, PrevRootBoneMatrix, bReset, ClothTeleportDistThresholdSquared, ClothTeleportCosineThresholdInRad);
+
+	switch (TeleportMode)
+	{
+	default:
+	case EClothingTeleportMode::None:
+		bTeleport = false;
+		break;
+	case EClothingTeleportMode::TeleportAndReset:
+		bReset = true;
+		// fallthrough
+	case EClothingTeleportMode::Teleport:
+		bTeleport = true;
+		break;
+	}
+
+	PrevRootBoneMatrix = CurRootBoneMat;
+}
+
 void UChaosClothComponent::HandleExistingParallelSimulation()
 {
 	if (bBindToLeaderComponent)
@@ -762,3 +804,24 @@ const FDataflowSimulationProxy* UChaosClothComponent::GetSimulationProxy() const
 	return ClothSimulationProxy.Get();
 }
 
+float UChaosClothComponent::GetTeleportRotationThreshold() const
+{
+	return TeleportRotationThreshold;
+}
+
+void UChaosClothComponent::SetTeleportRotationThreshold(float Threshold)
+{
+	TeleportRotationThreshold = Threshold;
+	ClothTeleportCosineThresholdInRad = UE::ClothingSimulation::TeleportHelpers::ComputeTeleportCosineRotationThreshold(TeleportRotationThreshold);
+}
+
+float UChaosClothComponent::GetTeleportDistanceThreshold() const
+{
+	return TeleportDistanceThreshold;
+}
+
+void UChaosClothComponent::SetTeleportDistanceThreshold(float Threshold)
+{
+	TeleportDistanceThreshold = Threshold;
+	ClothTeleportDistThresholdSquared = UE::ClothingSimulation::TeleportHelpers::ComputeTeleportDistanceThresholdSquared(TeleportDistanceThreshold);
+}
diff --git a/Engine/Plugins/ChaosClothAsset/Source/ChaosClothAssetEngine/Public/ChaosClothAsset/ClothComponent.h b/Engine/Plugins/ChaosClothAsset/Source/ChaosClothAssetEngine/Public/ChaosClothAsset/ClothComponent.h
--- a/Engine/Plugins/ChaosClothAsset/Source/ChaosClothAssetEngine/Public/ChaosClothAsset/ClothComponent.h
+++ b/Engine/Plugins/ChaosClothAsset/Source/ChaosClothAssetEngine/Public/ChaosClothAsset/ClothComponent.h
@@ -207,6 +207,38 @@ public:
 	float GetClothGeometryScale() const { return ClothGeometryScale; }
 	void SetClothGeometryScale(float Scale) { ClothGeometryScale = Scale; }
 
+	/**
+	* Gets the teleportation distance threshold.
+	*
+	* @return Threshold value.
+	*/
+	UFUNCTION(BlueprintGetter, Category = ClothComponent)
+	float GetTeleportDistanceThreshold() const;
+
+	/**
+	* Sets the teleportation distance threshold.
+	*
+	* @param threshold Threshold value.
+	*/
+	UFUNCTION(BlueprintSetter, Category = ClothComponent)
+	void SetTeleportDistanceThreshold(float Threshold);
+
+	/**
+	* Gets the teleportation rotation threshold.
+	*
+	* @return Threshold in degrees.
+	*/
+	UFUNCTION(BlueprintGetter, Category = ClothComponent)
+	float GetTeleportRotationThreshold() const;
+
+	/**
+	* Sets the teleportation rotation threshold.
+	*
+	* @param threshold Threshold in degrees.
+	*/
+	UFUNCTION(BlueprintSetter, Category = ClothComponent)
+	void SetTeleportRotationThreshold(float Threshold);
+
 #if WITH_EDITOR
 	/** Update config properties from the asset. Will only update existing values.*/
 	void UpdateConfigProperties();
@@ -277,6 +309,7 @@ private:
 	bool ShouldWaitForParallelSimulationInTickComponent() const;
 	void UpdateComponentSpaceTransforms();
 	void UpdateVisibility();
+	void UpdateClothTeleport();
 
 	friend UE::Chaos::ClothAsset::FClothComponentCacheAdapter;
 
@@ -337,6 +370,31 @@ private:
 	TObjectPtr<UChaosClothAssetBase> Asset;
 #endif
 
+	/**
+	* Conduct teleportation if the character's movement is greater than this threshold in 1 frame.
+	* Zero or negative values will skip the check.
+	* You can also do force teleport manually using ForceNextUpdateTeleport() / ForceNextUpdateTeleportAndReset().
+	*/
+	UPROPERTY(EditAnywhere, BlueprintGetter = GetTeleportDistanceThreshold, BlueprintSetter = SetTeleportDistanceThreshold, Category = ClothComponent)
+	float TeleportDistanceThreshold;
+
+	/**
+	* Rotation threshold in degrees, ranging from 0 to 180.
+	* Conduct teleportation if the character's rotation is greater than this threshold in 1 frame.
+	* Zero or negative values will skip the check.
+	*/
+	UPROPERTY(EditAnywhere, BlueprintGetter = GetTeleportRotationThreshold, BlueprintSetter = SetTeleportRotationThreshold, Category = ClothComponent)
+	float TeleportRotationThreshold;
+
+	/** Used for pre-computation using tTeleportDistanceThreshold property */
+	float ClothTeleportDistThresholdSquared;
+
+	/** Used for pre-computation using TeleportRotationThreshold property */
+	float ClothTeleportCosineThresholdInRad;
+
+	/** previous root bone matrix to compare the difference and decide to do clothing teleport  */
+	FMatrix	PrevRootBoneMatrix;
+
 	/**
 	 * Simulation properties per cloth model. 
 	 * Transient because it doesn't need to be serialized but contains pointers to UObject that needs not be garbage collected.
diff --git a/Engine/Source/Runtime/ClothingSystemRuntimeInterface/Private/ClothingSimulationTeleportHelpers.cpp b/Engine/Source/Runtime/ClothingSystemRuntimeInterface/Private/ClothingSimulationTeleportHelpers.cpp
new file mode 100644
--- /dev/null
+++ b/Engine/Source/Runtime/ClothingSystemRuntimeInterface/Private/ClothingSimulationTeleportHelpers.cpp
@@ -0,0 +1,67 @@
+// Copyright Epic Games, Inc. All Rights Reserved.
+
+#include "ClothingSimulationTeleportHelpers.h"
+#include "ClothingSystemRuntimeTypes.h"
+#include "HAL/IConsoleManager.h"
+
+static TAutoConsoleVariable<bool> CVarClothTeleportOverride(TEXT("p.Cloth.TeleportOverride"), false, TEXT("Force console variable teleport override values over skeletal mesh properties.\n Default: false."));
+static TAutoConsoleVariable<bool> CVarClothResetAfterTeleport(TEXT("p.Cloth.ResetAfterTeleport"), true, TEXT("Require p.Cloth.TeleportOverride. Reset the clothing after moving the clothing position (called teleport).\n Default: true."));
+static TAutoConsoleVariable<float> CVarClothTeleportDistanceThreshold(TEXT("p.Cloth.TeleportDistanceThreshold"), 300.f, TEXT("Require p.Cloth.TeleportOverride. Conduct teleportation if the character's movement is greater than this threshold in 1 frame.\n Zero or negative values will skip the check.\n Default: 300."));
+static TAutoConsoleVariable<float> CVarClothTeleportRotationThreshold(TEXT("p.Cloth.TeleportRotationThreshold"), 0.f, TEXT("Require p.Cloth.TeleportOverride. Rotation threshold in degrees, ranging from 0 to 180.\n Conduct teleportation if the character's rotation is greater than this threshold in 1 frame.\n Zero or negative values will skip the check.\n Default 0."));
+
+namespace UE::ClothingSimulation::TeleportHelpers
+{
+	EClothingTeleportMode CalculateClothingTeleport(EClothingTeleportMode CurrentTeleportMode, const FMatrix& CurRootBoneMat, const FMatrix& PrevRootBoneMat, bool bResetAfterTeleport, float ClothTeleportDistThresholdSquared, float ClothTeleportCosineThresholdInRad)
+	{
+		EClothingTeleportMode ClothTeleportMode = CurrentTeleportMode;
+
+		// CVar overrides
+		bool bResetAfterTeleportOverride;
+		float ClothTeleportDistThresholdSquaredOverride;
+		float ClothTeleportCosineThresholdInRadOverride;
+
+		if (CVarClothTeleportOverride.GetValueOnGameThread())
+		{
+			bResetAfterTeleportOverride = CVarClothResetAfterTeleport.GetValueOnGameThread();
+			const float TeleportDistanceThresholdOverride = CVarClothTeleportDistanceThreshold.GetValueOnGameThread();
+			ClothTeleportDistThresholdSquaredOverride = TeleportDistanceThresholdOverride > 0.f ? FMath::Square(TeleportDistanceThresholdOverride) : 0.f;
+			const float TeleportRotationThresholdOverride = CVarClothTeleportRotationThreshold.GetValueOnGameThread();
+			ClothTeleportCosineThresholdInRadOverride = FMath::Cos(FMath::DegreesToRadians(TeleportRotationThresholdOverride));
+		}
+		else
+		{
+			bResetAfterTeleportOverride = bResetAfterTeleport;
+			ClothTeleportDistThresholdSquaredOverride = ClothTeleportDistThresholdSquared;
+			ClothTeleportCosineThresholdInRadOverride = ClothTeleportCosineThresholdInRad;
+		}
+
+		// distance check 
+		// TeleportDistanceThreshold is greater than Zero and not teleported yet
+		if (ClothTeleportDistThresholdSquaredOverride > 0 && ClothTeleportMode == EClothingTeleportMode::None)
+		{
+			float DistSquared = FVector::DistSquared(PrevRootBoneMat.GetOrigin(), CurRootBoneMat.GetOrigin());
+			if (DistSquared > ClothTeleportDistThresholdSquaredOverride) // if it has traveled too far
+			{
+				ClothTeleportMode = bResetAfterTeleportOverride ? EClothingTeleportMode::TeleportAndReset : EClothingTeleportMode::Teleport;
+			}
+		}
+
+		// rotation check
+		// if TeleportRotationThreshold is greater than Zero and the user didn't do force teleport
+		if (ClothTeleportCosineThresholdInRadOverride < 1 && ClothTeleportMode == EClothingTeleportMode::None)
+		{
+			// Detect whether teleportation is needed or not
+			// Rotation matrix's transpose means an inverse but can't use a transpose because this matrix includes scales
+			FMatrix AInvB = CurRootBoneMat * PrevRootBoneMat.InverseFast();
+			float Trace = AInvB.M[0][0] + AInvB.M[1][1] + AInvB.M[2][2];
+			float CosineTheta = (Trace - 1.0f) / 2.0f; // trace = 1+2cos(theta) for a 3x3 matrix
+
+			if (CosineTheta < ClothTeleportCosineThresholdInRadOverride) // has the root bone rotated too much
+			{
+				ClothTeleportMode = bResetAfterTeleportOverride ? EClothingTeleportMode::TeleportAndReset : EClothingTeleportMode::Teleport;
+			}
+		}
+
+		return ClothTeleportMode;
+	}
+}
diff --git a/Engine/Source/Runtime/ClothingSystemRuntimeInterface/Public/ClothingSimulationTeleportHelpers.h b/Engine/Source/Runtime/ClothingSystemRuntimeInterface/Public/ClothingSimulationTeleportHelpers.h
new file mode 100644
--- /dev/null
+++ b/Engine/Source/Runtime/ClothingSystemRuntimeInterface/Public/ClothingSimulationTeleportHelpers.h
@@ -0,0 +1,25 @@
+// Copyright Epic Games, Inc. All Rights Reserved.
+#pragma once
+
+#include "Math/UnrealMathUtility.h"
+
+enum class EClothingTeleportMode : uint8;
+
+namespace UE::ClothingSimulation::TeleportHelpers
+{
+	inline float ComputeTeleportCosineRotationThreshold(float ThresholdInDegrees)
+	{
+		// Threshold <= 0 disables check.
+		// Cos(0) = 1, so set to 1 for all negative values.
+		// CalculateClothingTeleport will use ClothTeleportCosineThreshold >= 1 to disable.
+		return ThresholdInDegrees > 0.f ? FMath::Cos(FMath::DegreesToRadians(ThresholdInDegrees)) : 1.f;
+	}
+
+	inline float ComputeTeleportDistanceThresholdSquared(float ThresholdDistance)
+	{
+		// Threshold <= 0 disables check.
+		return ThresholdDistance > 0.f ? FMath::Square(ThresholdDistance) : 0.f;
+	}
+
+	CLOTHINGSYSTEMRUNTIMEINTERFACE_API EClothingTeleportMode CalculateClothingTeleport(EClothingTeleportMode CurrentTeleportMode, const FMatrix& CurRootBoneMat, const FMatrix& PrevRootBoneMat, bool bResetAfterTeleport, float ClothTeleportDistThresholdSquared, float ClothTeleportCosineThresholdInRad);
+};
diff --git a/Engine/Source/Runtime/Engine/Classes/Components/SkeletalMeshComponent.h b/Engine/Source/Runtime/Engine/Classes/Components/SkeletalMeshComponent.h
--- a/Engine/Source/Runtime/Engine/Classes/Components/SkeletalMeshComponent.h
+++ b/Engine/Source/Runtime/Engine/Classes/Components/SkeletalMeshComponent.h
@@ -1515,9 +1515,6 @@ private:
 	/** Used for pre-computation using tTeleportDistanceThreshold property */
 	float ClothTeleportDistThresholdSquared;
 
-	ENGINE_API void ComputeTeleportRotationThresholdInRadians();
-	ENGINE_API void ComputeTeleportDistanceThresholdInRadians();
-
 public:
 
 	/** Checked whether we have already ticked the pose this frame */
diff --git a/Engine/Source/Runtime/Engine/Private/Components/SkeletalMeshComponent.cpp b/Engine/Source/Runtime/Engine/Private/Components/SkeletalMeshComponent.cpp
--- a/Engine/Source/Runtime/Engine/Private/Components/SkeletalMeshComponent.cpp
+++ b/Engine/Source/Runtime/Engine/Private/Components/SkeletalMeshComponent.cpp
@@ -37,6 +37,7 @@
 #include "ClothCollisionSource.h"
 #include "ClothingSimulationInterface.h"
 #include "ClothingSimulationInteractor.h"
+#include "ClothingSimulationTeleportHelpers.h"
 #include "Features/IModularFeatures.h"
 #include "Misc/RuntimeErrors.h"
 #include "SkeletalMeshSceneProxy.h"
@@ -455,8 +456,8 @@ USkeletalMeshComponent::USkeletalMeshComponent(const FObjectInitializer& ObjectI
 	PrevRootRigidBodyScale = FVector3f(GetRelativeScale3D());
 
 	// pre-compute cloth teleport thresholds for performance
-	ComputeTeleportRotationThresholdInRadians();
-	ComputeTeleportDistanceThresholdInRadians();
+	ClothTeleportDistThresholdSquared = UE::ClothingSimulation::TeleportHelpers::ComputeTeleportDistanceThresholdSquared(TeleportDistanceThreshold);
+	ClothTeleportCosineThresholdInRad = UE::ClothingSimulation::TeleportHelpers::ComputeTeleportCosineRotationThreshold(TeleportRotationThreshold);
 
 	bBindClothToLeaderComponent = false;
 	bClothingSimulationSuspended = false;
@@ -1499,12 +1500,12 @@ void USkeletalMeshComponent::PostEditChangeProperty(FPropertyChangedEvent& Prope
 
 		if ( PropertyThatChanged->GetFName() == GET_MEMBER_NAME_CHECKED( USkeletalMeshComponent, TeleportDistanceThreshold ) )
 		{
-			ComputeTeleportDistanceThresholdInRadians();
+			ClothTeleportDistThresholdSquared = UE::ClothingSimulation::TeleportHelpers::ComputeTeleportDistanceThresholdSquared(TeleportDistanceThreshold);
 		}
 
 		if ( PropertyThatChanged->GetFName() == GET_MEMBER_NAME_CHECKED( USkeletalMeshComponent, TeleportRotationThreshold ) )
 		{
-			ComputeTeleportRotationThresholdInRadians();
+			ClothTeleportCosineThresholdInRad = UE::ClothingSimulation::TeleportHelpers::ComputeTeleportCosineRotationThreshold(TeleportRotationThreshold);
 		}
 	}
 }
@@ -5113,7 +5114,7 @@ float USkeletalMeshComponent::GetTeleportRotationThreshold() const
 void USkeletalMeshComponent::SetTeleportRotationThreshold(float Threshold)
 {
 	TeleportRotationThreshold = Threshold;
-	ComputeTeleportRotationThresholdInRadians();
+	ClothTeleportCosineThresholdInRad = UE::ClothingSimulation::TeleportHelpers::ComputeTeleportCosineRotationThreshold(TeleportRotationThreshold);
 }
 
 float USkeletalMeshComponent::GetTeleportDistanceThreshold() const
@@ -5124,17 +5125,7 @@ float USkeletalMeshComponent::GetTeleportDistanceThreshold() const
 void USkeletalMeshComponent::SetTeleportDistanceThreshold(float Threshold)
 {
 	TeleportDistanceThreshold = Threshold;
-	ComputeTeleportDistanceThresholdInRadians();
-}
-
-void USkeletalMeshComponent::ComputeTeleportRotationThresholdInRadians()
-{
-	ClothTeleportCosineThresholdInRad = FMath::Cos(FMath::DegreesToRadians(TeleportRotationThreshold));
-}
-
-void USkeletalMeshComponent::ComputeTeleportDistanceThresholdInRadians()
-{
-	ClothTeleportDistThresholdSquared = TeleportDistanceThreshold * TeleportDistanceThreshold;
+	ClothTeleportDistThresholdSquared = UE::ClothingSimulation::TeleportHelpers::ComputeTeleportDistanceThresholdSquared(TeleportDistanceThreshold);
 }
 
 void USkeletalMeshComponent::SetDisableAnimCurves(bool bInDisableAnimCurves)
diff --git a/Engine/Source/Runtime/Engine/Private/SkeletalMeshComponentPhysics.cpp b/Engine/Source/Runtime/Engine/Private/SkeletalMeshComponentPhysics.cpp
--- a/Engine/Source/Runtime/Engine/Private/SkeletalMeshComponentPhysics.cpp
+++ b/Engine/Source/Runtime/Engine/Private/SkeletalMeshComponentPhysics.cpp
@@ -17,6 +17,7 @@
 #include "AnimationRuntime.h"
 #include "ClothCollisionData.h"
 #include "ClothingSimulationInteractor.h"
+#include "ClothingSimulationTeleportHelpers.h"
 #include "Rendering/RenderCommandPipes.h"
 
 #include "Logging/MessageLog.h"
@@ -57,11 +58,6 @@ TAutoConsoleVariable<int32> CVarEnableClothPhysics(TEXT("p.ClothPhysics"), 1, TE
 static TAutoConsoleVariable<int32> CVarEnableClothPhysicsUseTaskThread(TEXT("p.ClothPhysics.UseTaskThread"), 1, TEXT("If 1, run cloth on the task thread. If 0, run on game thread."));
 static TAutoConsoleVariable<int32> CVarClothPhysicsTickWaitForParallelClothTask(TEXT("p.ClothPhysics.WaitForParallelClothTask"), 0, TEXT("If 1, always wait for cloth task completion in the Cloth Tick function. If 0, wait at end-of-frame updates instead if allowed by component settings"));
 
-static TAutoConsoleVariable<bool> CVarClothTeleportOverride(TEXT("p.Cloth.TeleportOverride"), false, TEXT("Force console variable teleport override values over skeletal mesh properties.\n Default: false."));
-static TAutoConsoleVariable<bool> CVarClothResetAfterTeleport(TEXT("p.Cloth.ResetAfterTeleport"), true, TEXT("Require p.Cloth.TeleportOverride. Reset the clothing after moving the clothing position (called teleport).\n Default: true."));
-static TAutoConsoleVariable<float> CVarClothTeleportDistanceThreshold(TEXT("p.Cloth.TeleportDistanceThreshold"), 300.f, TEXT("Require p.Cloth.TeleportOverride. Conduct teleportation if the character's movement is greater than this threshold in 1 frame.\n Zero or negative values will skip the check.\n Default: 300."));
-static TAutoConsoleVariable<float> CVarClothTeleportRotationThreshold(TEXT("p.Cloth.TeleportRotationThreshold"), 0.f, TEXT("Require p.Cloth.TeleportOverride. Rotation threshold in degrees, ranging from 0 to 180.\n Conduct teleportation if the character's rotation is greater than this threshold in 1 frame.\n Zero or negative values will skip the check.\n Default 0."));
-
 static TAutoConsoleVariable<int32> CVarEnableKinematicDeferralPrePhysicsCondition(TEXT("p.EnableKinematicDeferralPrePhysicsCondition"), 1, TEXT("If is 1, and deferral would've been disallowed due to EUpdateTransformFlags, allow if in PrePhysics tick. If 0, condition is unchanged."));
 
 static TAutoConsoleVariable<bool> CVarUpdateRigidBodyScaling(TEXT("p.UpdateRigidBodyScaling"), true, TEXT("Updates the rigid body scaling for simulation when the actor scaling is changed \n Default: true."));
@@ -758,8 +754,8 @@ void USkeletalMeshComponent::InitArticulated(FPhysScene* PhysScene)
 	PrevRootRigidBodyScale = FVector3f(GetRelativeScale3D()); //Update the root rigid body scale. 
 
 	// pre-compute cloth teleport thresholds for performance
-	ComputeTeleportDistanceThresholdInRadians();
-	ComputeTeleportRotationThresholdInRadians();
+	ClothTeleportDistThresholdSquared = UE::ClothingSimulation::TeleportHelpers::ComputeTeleportDistanceThresholdSquared(TeleportDistanceThreshold);
+	ClothTeleportCosineThresholdInRad = UE::ClothingSimulation::TeleportHelpers::ComputeTeleportCosineRotationThreshold(TeleportRotationThreshold);
 }
 
 TAutoConsoleVariable<int32> CVarEnableRagdollPhysics(TEXT("p.RagdollPhysics"), 1, TEXT("If 1, ragdoll physics will be used. Otherwise just root body is simulated"));
@@ -3328,55 +3324,7 @@ void USkeletalMeshComponent::CheckClothTeleport()
 	// Get the root bone transform
 	FMatrix CurRootBoneMat = GetBoneMatrix(0);
 
-	// CVar overrides
-	bool bResetAfterTeleportOverride;
-	float TeleportDistanceThresholdOverride;
-	float ClothTeleportDistThresholdSquaredOverride;
-	float TeleportRotationThresholdOverride;
-	float ClothTeleportCosineThresholdInRadOverride;
-	if (CVarClothTeleportOverride.GetValueOnGameThread())
-	{
-		bResetAfterTeleportOverride = CVarClothResetAfterTeleport.GetValueOnGameThread();
-		TeleportDistanceThresholdOverride = CVarClothTeleportDistanceThreshold.GetValueOnGameThread();
-		ClothTeleportDistThresholdSquaredOverride = FMath::Square(TeleportDistanceThresholdOverride);
-		TeleportRotationThresholdOverride = CVarClothTeleportRotationThreshold.GetValueOnGameThread();
-		ClothTeleportCosineThresholdInRadOverride = FMath::Cos(FMath::DegreesToRadians(TeleportRotationThresholdOverride));
-	}
-	else
-	{
-		bResetAfterTeleportOverride = bResetAfterTeleport;
-		TeleportDistanceThresholdOverride = TeleportDistanceThreshold;
-		ClothTeleportDistThresholdSquaredOverride = ClothTeleportDistThresholdSquared;
-		TeleportRotationThresholdOverride = TeleportRotationThreshold;
-		ClothTeleportCosineThresholdInRadOverride = ClothTeleportCosineThresholdInRad;
-	}
-
-	// distance check 
-	// TeleportDistanceThreshold is greater than Zero and not teleported yet
-	if(TeleportDistanceThresholdOverride > 0 && ClothTeleportMode == EClothingTeleportMode::None)
-	{
-		float DistSquared = FVector::DistSquared(PrevRootBoneMatrix.GetOrigin(), CurRootBoneMat.GetOrigin());
-		if ( DistSquared > ClothTeleportDistThresholdSquaredOverride ) // if it has traveled too far
-		{
-			ClothTeleportMode = bResetAfterTeleportOverride ? EClothingTeleportMode::TeleportAndReset : EClothingTeleportMode::Teleport;
-		}
-	}
-
-	// rotation check
-	// if TeleportRotationThreshold is greater than Zero and the user didn't do force teleport
-	if(TeleportRotationThresholdOverride > 0 && ClothTeleportMode == EClothingTeleportMode::None)
-	{
-		// Detect whether teleportation is needed or not
-		// Rotation matrix's transpose means an inverse but can't use a transpose because this matrix includes scales
-		FMatrix AInvB = CurRootBoneMat * PrevRootBoneMatrix.InverseFast();
-		float Trace = AInvB.M[0][0] + AInvB.M[1][1] + AInvB.M[2][2];
-		float CosineTheta = (Trace - 1.0f) / 2.0f; // trace = 1+2cos(theta) for a 3x3 matrix
-
-		if ( CosineTheta < ClothTeleportCosineThresholdInRadOverride ) // has the root bone rotated too much
-		{
-			ClothTeleportMode = bResetAfterTeleportOverride ? EClothingTeleportMode::TeleportAndReset : EClothingTeleportMode::Teleport;
-		}
-	}
+	ClothTeleportMode = UE::ClothingSimulation::TeleportHelpers::CalculateClothingTeleport(ClothTeleportMode, CurRootBoneMat, PrevRootBoneMatrix, bResetAfterTeleport, ClothTeleportDistThresholdSquared, ClothTeleportCosineThresholdInRad);
 
 	PrevRootBoneMatrix = CurRootBoneMat;
 }
