Missing include causing compilation error.

diff --git a/Engine/Source/Runtime/Engine/Private/OnlineEngineInterface.cpp b/Engine/Source/Runtime/Engine/Private/OnlineEngineInterface.cpp
--- a/Engine/Source/Runtime/Engine/Private/OnlineEngineInterface.cpp
+++ b/Engine/Source/Runtime/Engine/Private/OnlineEngineInterface.cpp
@@ -2,6 +2,7 @@
 
 #include "Net/OnlineEngineInterface.h"
 #include "UObject/Package.h"
+#include "Misc/CommandLine.h"
 #include "Misc/ConfigCacheIni.h"
 
 #include UE_INLINE_GENERATED_CPP_BY_NAME(OnlineEngineInterface)
@@ -22,13 +23,18 @@ UOnlineEngineInterface* UOnlineEngineInterface::Get()
 		FString OnlineEngineInterfaceClassName;
 		GConfig->GetString(TEXT("/Script/Engine.OnlineEngineInterface"), TEXT("ClassName"), OnlineEngineInterfaceClassName, GEngineIni);
 
-		// To not break licensees using this, prefer this if it is present, and warn. Remove in 5.7
-		bool bUseOnlineServicesV2 = false;
-		if (GConfig->GetBool(TEXT("/Script/Engine.OnlineEngineInterface"), TEXT("bUseOnlineServicesV2"), bUseOnlineServicesV2, GEngineIni))
+		bool bUseOnlineServices = FParse::Param(FCommandLine::Get(), TEXT("bUseOnlineServices"));
+		if (bUseOnlineServices)
+		{
+			OnlineEngineInterfaceClassName = TEXT("/Script/OnlineSubsystemUtils.OnlineServicesEngineInterfaceImpl");
+		}
+
+		// To not break licensees using the config override, prefer it if it is present, and warn. Remove in 5.7
+		if (GConfig->GetBool(TEXT("/Script/Engine.OnlineEngineInterface"), TEXT("bUseOnlineServicesV2"), bUseOnlineServices, GEngineIni))
 		{
 			const TCHAR* V1ClassName = TEXT("/Script/OnlineSubsystemUtils.OnlineEngineInterfaceImpl");
 			const TCHAR* V2ClassName = TEXT("/Script/OnlineSubsystemUtils.OnlineServicesEngineInterfaceImpl");
-			OnlineEngineInterfaceClassName = bUseOnlineServicesV2 ? V2ClassName : V1ClassName;
+			OnlineEngineInterfaceClassName = bUseOnlineServices ? V2ClassName : V1ClassName;
 			UE_LOG(LogOnlineEngine, Warning, TEXT("bUseOnlineServicesV2 is deprecated, please instead configure [/Script/Engine.OnlineEngineInterface]:ClassName=%s"), *OnlineEngineInterfaceClassName);
 		}
 		
