[NetPhysics] CVar to override enabled / disable of Physics Replication LOD.

diff --git a/Engine/Source/Runtime/Engine/Private/PhysicsEngine/PhysicsReplicationLOD.cpp b/Engine/Source/Runtime/Engine/Private/PhysicsEngine/PhysicsReplicationLOD.cpp
--- a/Engine/Source/Runtime/Engine/Private/PhysicsEngine/PhysicsReplicationLOD.cpp
+++ b/Engine/Source/Runtime/Engine/Private/PhysicsEngine/PhysicsReplicationLOD.cpp
@@ -12,8 +12,8 @@
 
 namespace PhysicsReplicationLODCVars
 {
-	bool bEnabled = false;
-	static FAutoConsoleVariableRef CVarEnabled(TEXT("p.ReplicationLOD.Enabled"), bEnabled, TEXT("Override a disabled PhysicsReplicationLOD to become enabled, can't disable an LOD that is set enabled via other means like project settings or API call."));
+	int32 OverrideEnabled = 0;
+	static FAutoConsoleVariableRef CVarOverrideEnabled(TEXT("p.ReplicationLOD.OverrideEnabled"), OverrideEnabled, TEXT("0 = Use Project Settings or API calls. 1 = Override to enabled. -1 = Override to disabled."));
 
 	float MinimumBaseDistance = -1.0f;
 	static FAutoConsoleVariableRef CVarMinimumBaseDistance(TEXT("p.ReplicationLOD.MinimumBaseDistance"), MinimumBaseDistance, TEXT("Override the minimum distance in the physics replication LODs Base Distance. Negative value disables the override and project settings will apply."));
@@ -84,7 +84,11 @@ void FPhysicsReplicationLOD::SetEnabled(const bool bInEnabled)
 
 bool FPhysicsReplicationLOD::IsEnabled() const
 {
-	return bEnabled || PhysicsReplicationLODCVars::bEnabled;
+	if (PhysicsReplicationLODCVars::OverrideEnabled == -1)
+	{
+		return false;
+	}
+	return bEnabled || (PhysicsReplicationLODCVars::OverrideEnabled == 1);
 }
 
 void FPhysicsReplicationLOD::RegisterFocalPoint_External(const UPrimitiveComponent* Component, FName BoneName)
@@ -138,7 +142,11 @@ IPhysicsReplicationLODAsync* FPhysicsReplicationLOD::GetPhysicsReplicationLOD_In
 
 bool FPhysicsReplicationLODAsync::IsEnabled() const
 {
-	return (bEnabled || PhysicsReplicationLODCVars::bEnabled) && FocalParticles.Num() > 0;
+	if (PhysicsReplicationLODCVars::OverrideEnabled == -1)
+	{
+		return false;
+	}
+	return (bEnabled || (PhysicsReplicationLODCVars::OverrideEnabled == 1)) && FocalParticles.Num() > 0;
 }
 
 void FPhysicsReplicationLODAsync::OnPostInitialize_Internal()
