Correctly remove LocalVertexFactory shaders from MD_UI materials

#tests local cook of Lyra to verify removal and Editor in debugger to ensure they compile

diff --git a/Engine/Source/Runtime/Engine/Private/LocalVertexFactory.cpp b/Engine/Source/Runtime/Engine/Private/LocalVertexFactory.cpp
--- a/Engine/Source/Runtime/Engine/Private/LocalVertexFactory.cpp
+++ b/Engine/Source/Runtime/Engine/Private/LocalVertexFactory.cpp
@@ -255,7 +255,7 @@ bool FLocalVertexFactory::ShouldCompilePermutation(const FVertexFactoryShaderPer
 	// Only compile this permutation inside the editor - it's not applicable in games, but occasionally the editor needs it.
 	if (Parameters.MaterialParameters.MaterialDomain == MD_UI)
 	{
-		return !!WITH_EDITOR;
+		return EnumHasAllFlags(Parameters.Flags, EShaderPermutationFlags::HasEditorOnlyData);
 	}
 
 	return true; 
