Make sure Texture objects referenced by a TextureCollection are loaded before initializing the TextureCollection.

diff --git a/Engine/Source/Runtime/Engine/Private/TextureCollection.cpp b/Engine/Source/Runtime/Engine/Private/TextureCollection.cpp
--- a/Engine/Source/Runtime/Engine/Private/TextureCollection.cpp
+++ b/Engine/Source/Runtime/Engine/Private/TextureCollection.cpp
@@ -89,6 +89,15 @@ void UTextureCollection::PostLoad()
 
 	if (!IsTemplate())
 	{
+		// Make sure all referenced textures are PostLoad-ed before we start accessing them in UpdateResource()
+		for (UTexture* Texture : this->Textures)
+		{
+			if (Texture)
+			{
+				Texture->ConditionalPostLoad();
+			}
+		}
+
 		UpdateResource();
 	}
 }
