Minor fix for incorrect Vector parameter name checking in MaterialParameterCollection code.

diff --git a/Engine/Source/Runtime/Engine/Private/Materials/ParameterCollection.cpp b/Engine/Source/Runtime/Engine/Private/Materials/ParameterCollection.cpp
--- a/Engine/Source/Runtime/Engine/Private/Materials/ParameterCollection.cpp
+++ b/Engine/Source/Runtime/Engine/Private/Materials/ParameterCollection.cpp
@@ -200,7 +200,7 @@ bool UMaterialParameterCollection::SetScalarParameterDefaultValue(FName Paramete
 
 bool UMaterialParameterCollection::SetVectorParameterDefaultValueByInfo(FCollectionVectorParameter VectorParameter)
 {
-	if(GetScalarParameterByName(VectorParameter.ParameterName))
+	if(GetVectorParameterByName(VectorParameter.ParameterName))
     {
     	// if the input parameter exists, pass the name and value down to SetVectorParameterDefaultValue
     	// since we want to preserve the Guid of the parameter that's already on the asset
