Adding TextureCollection support to MaterialInstanceDynamic.

diff --git a/Engine/Source/Runtime/Engine/Private/Materials/MaterialInstanceDynamic.cpp b/Engine/Source/Runtime/Engine/Private/Materials/MaterialInstanceDynamic.cpp
--- a/Engine/Source/Runtime/Engine/Private/Materials/MaterialInstanceDynamic.cpp
+++ b/Engine/Source/Runtime/Engine/Private/Materials/MaterialInstanceDynamic.cpp
@@ -10,6 +10,7 @@
 #include "UObject/Package.h"
 #include "Materials/MaterialInstanceSupport.h"
 #include "Engine/Texture.h"
+#include "Engine/TextureCollection.h"
 #include "Misc/RuntimeErrors.h"
 #include "ObjectCacheEventSink.h"
 
@@ -275,6 +276,30 @@ UTexture* UMaterialInstanceDynamic::K2_GetTextureParameterValueByInfo(const FMat
 	return Result;
 }
 
+void UMaterialInstanceDynamic::SetTextureCollectionParameterValueByInfo(const FMaterialParameterInfo& ParameterInfo, UTextureCollection* Value)
+{
+	SetTextureCollectionParameterValueInternal(ParameterInfo, Value);
+}
+
+void UMaterialInstanceDynamic::SetTextureCollectionParameterValue(FName ParameterName, UTextureCollection* Value)
+{
+	FMaterialParameterInfo ParameterInfo(ParameterName);
+	SetTextureCollectionParameterValueByInfo(ParameterInfo, Value);
+}
+
+UTextureCollection* UMaterialInstanceDynamic::K2_GetTextureCollectionParameterValueByInfo(const FMaterialParameterInfo& ParameterInfo)
+{
+	UTextureCollection* Result = NULL;
+	Super::GetTextureCollectionParameterValue(ParameterInfo, Result);
+	return Result;
+}
+
+UTextureCollection* UMaterialInstanceDynamic::K2_GetTextureCollectionParameterValue(FName ParameterName)
+{
+	FMaterialParameterInfo ParameterInfo(ParameterName);
+	return K2_GetTextureCollectionParameterValueByInfo(ParameterInfo);
+}
+
 void UMaterialInstanceDynamic::SetFontParameterValue(const FMaterialParameterInfo& ParameterInfo,class UFont* FontValue,int32 FontPage)
 {
 	SetFontParameterValueInternal(ParameterInfo,FontValue,FontPage);
diff --git a/Engine/Source/Runtime/Engine/Public/Materials/MaterialInstanceDynamic.h b/Engine/Source/Runtime/Engine/Public/Materials/MaterialInstanceDynamic.h
--- a/Engine/Source/Runtime/Engine/Public/Materials/MaterialInstanceDynamic.h
+++ b/Engine/Source/Runtime/Engine/Public/Materials/MaterialInstanceDynamic.h
@@ -1,6 +1,5 @@
 // Copyright Epic Games, Inc. All Rights Reserved.
 
-
 #pragma once
 
 #include "CoreMinimal.h"
@@ -8,6 +7,9 @@
 #include "Materials/MaterialInstance.h"
 #include "MaterialInstanceDynamic.generated.h"
 
+class UTexture;
+class UTextureCollection;
+
 UCLASS(hidecategories=Object, collapsecategories, BlueprintType, MinimalAPI)
 class UMaterialInstanceDynamic : public UMaterialInstance
 {
@@ -86,6 +88,22 @@ class UMaterialInstanceDynamic : public UMaterialInstance
 	UFUNCTION(BlueprintCallable, meta = (DisplayName = "Get Texture Parameter Value By Info", ScriptName = "GetTextureParameterValueByInfo"), Category = "Rendering|Material")
 	ENGINE_API class UTexture* K2_GetTextureParameterValueByInfo(const FMaterialParameterInfo& ParameterInfo);
 
+	/** Set an MID texture collection parameter value */
+	UFUNCTION(BlueprintCallable, Category="Rendering|Material")
+	ENGINE_API void SetTextureCollectionParameterValue(FName ParameterName, UTextureCollection* Value);
+
+	/** Set an MID texture collection parameter value using MPI (to allow access to layer parameters) */
+	UFUNCTION(BlueprintCallable, Category = "Rendering|Material")
+	ENGINE_API void SetTextureCollectionParameterValueByInfo(const FMaterialParameterInfo& ParameterInfo, UTextureCollection* Value);
+
+	/** Get the current MID texture collection parameter value, using MPI (to allow access to layer parameters) */
+	UFUNCTION(BlueprintCallable, meta = (DisplayName = "Get Texture Collection Parameter Value By Info", ScriptName = "GetTextureCollectionParameterValueByInfo"), Category = "Rendering|Material")
+	ENGINE_API UTextureCollection* K2_GetTextureCollectionParameterValueByInfo(const FMaterialParameterInfo& ParameterInfo);
+
+	/** Get the current MID texture collection parameter value */
+	UFUNCTION(BlueprintCallable, meta=(DisplayName = "Get Texture Collection Parameter Value", ScriptName = "GetTextureCollectionParameterValue"), Category="Rendering|Material")
+	ENGINE_API UTextureCollection* K2_GetTextureCollectionParameterValue(FName ParameterName);
+
 	/** Set an MID vector parameter value */
 	UFUNCTION(BlueprintCallable, meta=(Keywords = "SetColorParameterValue"), Category="Rendering|Material")
 	ENGINE_API void SetVectorParameterValue(FName ParameterName, FLinearColor Value);
