New translator changes:
- As() and AsInstruction() are now free functions in the MIR namespace and take a value as argument. This argument can be null (in which case, the functions simply return null too). This removes risks with the undefined behavior with previous instances of "this &&".
- Made all MIR Emitter APIs safe to poison arguments. This simplifies API usage as now the emitter will simply return poison if any of its API arguments is poison.
- Renamed EOperatorParameterTypeFilter enum entries from "OT_" to "OPTF_" for consistency.
- Renamed "OT_SameAsFirstArgument" to "OPTF_AsFirstArgument".
- EmitOperator() will now try to cast arguments after the first to the first argument type on OPTF_AsFirstArgument. This greatly simplifies emitting operations as appropriate arguments casting is now done automatically.
- Removed many UE_MIR_CHECKPOINT() instances as now redundant.
- Thanks to the improvements above, simplified several expression Build() implementations.

#rnx

diff --git a/Engine/Source/Runtime/Engine/Private/Materials/MaterialExpressionsIR.cpp b/Engine/Source/Runtime/Engine/Private/Materials/MaterialExpressionsIR.cpp
--- a/Engine/Source/Runtime/Engine/Private/Materials/MaterialExpressionsIR.cpp
+++ b/Engine/Source/Runtime/Engine/Private/Materials/MaterialExpressionsIR.cpp
@@ -600,15 +600,6 @@ static void BuildBinaryOperatorFromValue(
 	MIR::FValue* LhsValue,
 	MIR::FValue* RhsValue)
 {
-	// Determine operation input/output type by looking at the first connected input and picking float1 otherwise.
-	MIR::FTypePtr ResultType = Emitter.GetCommonPrimitiveType(LhsValue->Type->AsPrimitive(), RhsValue->Type->AsPrimitive());
-
-	// Convert operand values to determined result type. 
-	LhsValue = Emitter.EmitConstruct(ResultType, LhsValue);
-	RhsValue = Emitter.EmitConstruct(ResultType, RhsValue);
-
-	UE_MIR_CHECKPOINT();
-
 	// Finally emit the binary operator.
 	MIR::FValue* Value = Emitter.EmitOperator(Op, LhsValue, RhsValue);
 
@@ -785,6 +776,14 @@ void UMaterialExpressionIfThenElse::Build(MIR::FEmitter& Emitter)
 	Emitter.Output(0, OutputValue);
 }
 
+static MIR::FValue* EmitAlmostEquals(MIR::FEmitter& Emitter, MIR::FValue* A, MIR::FValue* B, float Threshold)
+{
+	return Emitter.EmitOperator(
+		MIR::BO_LessThanOrEquals, 
+		Emitter.EmitOperator(MIR::UO_Abs, Emitter.EmitOperator(MIR::BO_Subtract, A, B)),
+		Emitter.EmitConstantFloatScalar(Threshold));
+}
+
 void UMaterialExpressionIf::Build(MIR::FEmitter& Emitter)
 {
 	// Create default values flowing into disconnected inputs
@@ -794,8 +793,6 @@ void UMaterialExpressionIf::Build(MIR::FEmitter& Emitter)
 	Emitter.DefaultInputToFloatScalarZero(&AEqualsB);
 	Emitter.DefaultInputToFloatScalarZero(&ALessThanB);
 
-	UE_MIR_CHECKPOINT();
-
 	// Get input values and check their types are what we expect.
 	MIR::FValue* AValue = Emitter.InputScalar(&A);
 	MIR::FValue* BValue = Emitter.InputScalar(&B);
@@ -803,33 +800,9 @@ void UMaterialExpressionIf::Build(MIR::FEmitter& Emitter)
 	MIR::FValue* AEqualsBValue = Emitter.InputPrimitive(&AEqualsB);
 	MIR::FValue* ALessThanBValue = Emitter.InputPrimitive(&ALessThanB);
 
-	UE_MIR_CHECKPOINT();
-
-	// Get the primitive common type between the conditional arguments (e.g. if inputs are int and float, it will return float).
-	MIR::FPrimitiveTypePtr ConditionArgsType = Emitter.GetCommonPrimitiveType(AValue->Type->AsPrimitive(), BValue->Type->AsPrimitive());
-
-	UE_MIR_CHECKPOINT();
-
-	// Convert both conditional argument values to the common type. This is expected to work as the common type could be found.
-	AValue = Emitter.EmitConstruct(ConditionArgsType, AValue);
-	BValue = Emitter.EmitConstruct(ConditionArgsType, BValue);
-
-	// Now determine the output type by taking the common primitive type between result values.
-	MIR::FPrimitiveTypePtr OutputType = Emitter.GetCommonPrimitiveType(AGreaterThanBValue->Type->AsPrimitive(), AEqualsBValue->Type->AsPrimitive());
-	OutputType = Emitter.GetCommonPrimitiveType(OutputType, ALessThanBValue->Type->AsPrimitive());
-
-	UE_MIR_CHECKPOINT();
-
-	// Convert result values to the common result type.
-	AGreaterThanBValue = Emitter.EmitConstruct(OutputType, AGreaterThanBValue);
-	AEqualsBValue = Emitter.EmitConstruct(OutputType, AEqualsBValue);
-	ALessThanBValue = Emitter.EmitConstruct(OutputType, ALessThanBValue);
-
-	UE_MIR_CHECKPOINT();
-
 	// Emit the comparison expressions.
 	MIR::FValue* ALessThanBConditionValue = Emitter.EmitOperator(MIR::BO_LessThan, AValue, BValue);
-	MIR::FValue* AEqualsBConditionValue = Emitter.EmitOperator(MIR::BO_Equals, AValue, BValue);
+	MIR::FValue* AEqualsBConditionValue = EmitAlmostEquals(Emitter, AValue, BValue, EqualsThreshold);
 
 	// And finally emit the full conditional expression.
 	MIR::FValue* OutputValue = Emitter.EmitBranch(AEqualsBConditionValue, AEqualsBValue, AGreaterThanBValue);
@@ -862,10 +835,12 @@ static void BuildTextureSampleExpression(MIR::FValue* Texture, UMaterialExpressi
 
 void UMaterialExpressionTextureSample::Build(UE::MIR::FEmitter& Emitter)
 {
-	MIR::FValue* TextureValue = Emitter.TryInput(&TextureObject) ? Emitter.InputTexture(&TextureObject) : nullptr;
-	UE_MIR_CHECKPOINT();
-
-	if (!TextureValue)
+	MIR::FValue* TextureValue = Emitter.TryInput(&TextureObject);
+	if (TextureValue)
+	{
+		TextureValue = Emitter.InputTexture(&TextureObject);
+	}
+	else
 	{
 		if (!Texture)
 		{
@@ -876,6 +851,8 @@ void UMaterialExpressionTextureSample::Build(UE::MIR::FEmitter& Emitter)
 		TextureValue = Emitter.EmitTextureObject(Texture.Get(), SamplerType);
 	}
 
+	UE_MIR_CHECKPOINT();
+
 	if (TextureValue && TextureValue->GetTexture()->Get<EMaterialSamplerType>() != SAMPLERTYPE_Color)
 	{
 		Emitter.Error(TEXT("Input texture sampler type must be color."));
@@ -953,33 +930,19 @@ void BuildTernaryArithmeticOperator(
 	Emitter.DefaultInputToFloatScalar(B, ConstB);
 	Emitter.DefaultInputToFloatScalar(C, ConstC);
 
-	UE_MIR_CHECKPOINT();
-
 	MIR::FValue* ValueA = Emitter.InputPrimitive(A);
 	MIR::FValue* ValueB = Emitter.Input(B);
 	MIR::FValue* ValueC = Emitter.Input(C);
 
-	UE_MIR_CHECKPOINT();
-
-	ValueB = Emitter.EmitConstruct(ValueA->Type, ValueB);
-	ValueC = Emitter.EmitConstruct(ValueA->Type, ValueC);
-
-	UE_MIR_CHECKPOINT();
-
 	Emitter.Output(0, Emitter.EmitOperator(Op, ValueA, ValueB, ValueC));
 }
 
 void UMaterialExpressionClamp::Build(UE::MIR::FEmitter& Emitter)
 {
 	MIR::FValue* InputValue = Emitter.Input(&Input);
-
-	UE_MIR_CHECKPOINT();
-
 	MIR::FValue* MinValue = Emitter.DefaultInputToFloatScalar(&Min, MinDefault).Input(&Min);
 	MIR::FValue* MaxValue = Emitter.DefaultInputToFloatScalar(&Max, MaxDefault).Input(&Max);
 
-	UE_MIR_CHECKPOINT();
-
 	MIR::FValue* OutputValue = nullptr;
 	if (ClampMode == CMODE_Clamp)
 	{
@@ -994,8 +957,6 @@ void UMaterialExpressionClamp::Build(UE::MIR::FEmitter& Emitter)
 		OutputValue = Emitter.EmitOperator(MIR::BO_Min, InputValue, MaxValue);
 	}
 
-	UE_MIR_CHECKPOINT();
-
 	Emitter.Output(0, OutputValue);
 }
 
diff --git a/Engine/Source/Runtime/Engine/Private/Materials/MaterialIR.cpp b/Engine/Source/Runtime/Engine/Private/Materials/MaterialIR.cpp
--- a/Engine/Source/Runtime/Engine/Private/Materials/MaterialIR.cpp
+++ b/Engine/Source/Runtime/Engine/Private/Materials/MaterialIR.cpp
@@ -57,16 +57,6 @@ bool FValue::IsA(EValueKind InKind) const
 	return Kind == InKind;
 }
 
-FInstruction* FValue::AsInstruction()
-{
-	return this && (Kind > VK_InstructionBegin && Kind < VK_InstructionEnd) ? static_cast<FInstruction*>(this) : nullptr;
-}
-
-const FInstruction* FValue::AsInstruction() const
-{
-	return const_cast<FValue*>(this)->AsInstruction();
-}
-
 bool FValue::IsPoison() const
 {
 	return Kind == VK_Poison;
@@ -263,6 +253,16 @@ TOptional<TPair<UTexture*, EMaterialSamplerType>> FValue::GetTexture()
 	return {};
 }
 
+FInstruction* AsInstruction(FValue* Value)
+{
+	return Value && (Value->Kind > VK_InstructionBegin && Value->Kind < VK_InstructionEnd) ? static_cast<FInstruction*>(Value) : nullptr;
+}
+
+const FInstruction* AsInstruction(const FValue* Value)
+{
+	return AsInstruction(const_cast<FValue*>(Value));
+}
+
 FPoison* FPoison::Get()
 {
 	static FPoison Poison{ VK_Poison, VF_None, FType::GetPoison() };
diff --git a/Engine/Source/Runtime/Engine/Private/Materials/MaterialIRDebug.cpp b/Engine/Source/Runtime/Engine/Private/Materials/MaterialIRDebug.cpp
--- a/Engine/Source/Runtime/Engine/Private/Materials/MaterialIRDebug.cpp
+++ b/Engine/Source/Runtime/Engine/Private/Materials/MaterialIRDebug.cpp
@@ -125,7 +125,7 @@ void DebugDumpIRUseGraph(const FMaterialIRModule& Module)
 		// End the node declaration
 		Content.Append(TEXT(">]\n"));
 
-		const FInstruction* Instr = Value->AsInstruction();
+		const FInstruction* Instr = AsInstruction(Value);
 		if (DumpInstructionSequence && Instr && Instr->Next)
 		{
 			Content.Appendf(TEXT("\"%p\" -> \"%p\" [color=\"red\"]\n"), Instr, Instr->Next);
@@ -155,7 +155,7 @@ void DebugDumpIRUseGraph(const FMaterialIRModule& Module)
 
 			if (DumpInstructionSequence && Instr)
 			{
-				const FInstruction* UseInstr = Use->AsInstruction();
+				const FInstruction* UseInstr = AsInstruction(Use);
 				if (UseInstr && UseInstr->Block != Instr->Block)
 				{
 					Content.Appendf(TEXT("\"%p\" -> \"%p\" [color=\"red\", style=\"dashed\"]\n"), UseInstr, Instr);
diff --git a/Engine/Source/Runtime/Engine/Private/Materials/MaterialIREmitter.cpp b/Engine/Source/Runtime/Engine/Private/Materials/MaterialIREmitter.cpp
--- a/Engine/Source/Runtime/Engine/Private/Materials/MaterialIREmitter.cpp
+++ b/Engine/Source/Runtime/Engine/Private/Materials/MaterialIREmitter.cpp
@@ -204,22 +204,23 @@ FValue* FEmitter::TryInput(const FExpressionInput* InInput)
 
 FValue* FEmitter::Input(const FExpressionInput* InInput)
 {
-	FValue* Value = TryInput(InInput);
-	if (!Value)
+	if (FValue* Value = TryInput(InInput))
 	{
-		Errorf(TEXT("Missing '%s' input value."), *SlowFindInputName(Expression, InInput).ToString());
-		return MIR::FPoison::Get();
+		return Value;
 	}
-	return Value;
+
+	Errorf(TEXT("Missing '%s' input value."), *SlowFindInputName(Expression, InInput).ToString());
+	return MIR::FPoison::Get();
 }
 
 FValue* FEmitter::InputOfType(const FExpressionInput* InInput, ETypeKind Kind)
 {
 	FValue* Value = Input(InInput);
-	if (Value->IsPoison())
+	if (IsAnyNonNullAPoison(Value))
 	{
 		return Value;
 	}
+
 	CheckInputTypeIs(InInput, Value, Kind);
 	return Value;
 }
@@ -227,12 +228,12 @@ FValue* FEmitter::InputOfType(const FExpressionInput* InInput, ETypeKind Kind)
 bool FEmitter::InputConstantBoolScalar(const FExpressionInput* InInput)
 {
 	FValue* Value = InputBoolScalar(InInput);
-	if (Value->IsPoison())
+	if (IsAnyNonNullAPoison(Value))
 	{
 		return false;
 	}
 
-	auto Constant = Value->As<FConstant>();
+	const FConstant* Constant = MIR::As<FConstant>(Value);
 	if (!Constant)
 	{
 		Errorf(TEXT("Expected constant bool value through '%s' input."), *SlowFindInputName(Expression, InInput).ToString());
@@ -259,7 +260,7 @@ FValue* FEmitter::InputFloatScalar(const FExpressionInput* InInput)
 FValue* FEmitter::InputScalar(const FExpressionInput* InInput)
 {
 	FValue* Value = InputPrimitive(InInput);
-	if (Value->IsPoison())
+	if (IsAnyNonNullAPoison(Value))
 	{
 		return Value;
 	}
@@ -271,10 +272,11 @@ FValue* FEmitter::InputScalar(const FExpressionInput* InInput)
 FValue* FEmitter::InputPrimitive(const FExpressionInput* InInput)
 {
 	FValue* Value = Input(InInput);
-	if (Value->IsPoison())
+	if (IsAnyNonNullAPoison(Value))
 	{
 		return Value;
 	}
+
 	CheckInputTypeIs(InInput, Value, TypeKind_Primitive);
 	return Value;
 }
@@ -282,10 +284,11 @@ FValue* FEmitter::InputPrimitive(const FExpressionInput* InInput)
 FValue* FEmitter::InputTexture(const FExpressionInput* InInput)
 {
 	FValue* Value = Input(InInput);
-	if (Value->IsPoison())
+	if (IsAnyNonNullAPoison(Value))
 	{
 		return Value;
 	}
+
 	CheckValueIsTexture(Value);
 	return Value;
 }
@@ -306,12 +309,12 @@ bool FEmitter::ToConstantBoolScalar(MIR::FValue* Value)
 {
 	CheckValueIsScalar(Value, ScalarKind_Bool);
 
-	if (!bExpressionHasError)
+	if (IsInvalid())
 	{
 		return false;
 	}
 
-	if (MIR::FConstant* Constant = Value->As<MIR::FConstant>())
+	if (MIR::FConstant* Constant = As<MIR::FConstant>(Value))
 	{
 		return Constant->Boolean;
 	}
@@ -322,15 +325,25 @@ bool FEmitter::ToConstantBoolScalar(MIR::FValue* Value)
 
 void FEmitter::CheckInputIsScalar(const FExpressionInput* InInput, FValue* InputValue)
 {
+	if (IsAnyNonNullAPoison(InputValue))
+	{
+		return;
+	}
+
 	FPrimitiveTypePtr PrimitiveType = InputValue->Type->AsPrimitive();
 	if (!PrimitiveType || !PrimitiveType->IsScalar())
 	{
-		Errorf(TEXT("Input '%s' expected to be a scalar. It is %s instead."), *SlowFindInputName(Expression, InInput).ToString(), InputValue->Type->GetSpelling().GetData());
+		Errorf(TEXT("Input '%s' expected to be a scalar. It is a %s instead."), *SlowFindInputName(Expression, InInput).ToString(), InputValue->Type->GetSpelling().GetData());
 	}
 }
 
 void FEmitter::CheckInputIsScalar(const FExpressionInput* InInput, FValue* InputValue, EScalarKind Kind)
 {
+	if (IsAnyNonNullAPoison(InputValue))
+	{
+		return;
+	}
+
 	FPrimitiveTypePtr PrimitiveType = InputValue->Type->AsPrimitive();
 	if (!PrimitiveType || !PrimitiveType->IsScalar() || PrimitiveType->ScalarKind != Kind)
 	{
@@ -340,6 +353,11 @@ void FEmitter::CheckInputIsScalar(const FExpressionInput* InInput, FValue* Input
 
 void FEmitter::CheckInputTypeIs(const FExpressionInput* InInput, FValue* InputValue, ETypeKind Kind)
 {
+	if (IsAnyNonNullAPoison(InputValue))
+	{
+		return;
+	}
+
 	if (InputValue && InputValue->Type->Kind != Kind)
 	{
 		Errorf(TEXT("Input '%s' expected to be have type %s. It is %s instead."), *SlowFindInputName(Expression, InInput).ToString(), TypeKindToString(Kind), InputValue->Type->GetSpelling().GetData());
@@ -348,7 +366,7 @@ void FEmitter::CheckInputTypeIs(const FExpressionInput* InInput, FValue* InputVa
 
 void FEmitter::CheckValueIsScalar(const FValue* Value, EScalarKind Kind)
 {
-	if (Value->IsPoison())
+	if (IsAnyNonNullAPoison(Value))
 	{
 		return;
 	}
@@ -363,7 +381,7 @@ void FEmitter::CheckValueIsScalar(const FValue* Value, EScalarKind Kind)
 
 void FEmitter::CheckValueIsTexture(const FValue* Value)
 {
-	if (Value->IsPoison())
+	if (IsAnyNonNullAPoison(Value))
 	{
 		return;
 	}
@@ -485,6 +503,11 @@ FValue* FEmitter::EmitConstantInt4(const FIntVector4& InValue)
 
 FValue* FEmitter::EmitVector2(FValue* InX, FValue* InY)
 {
+	if (IsAnyNonNullAPoison(InX, InY))
+	{
+		return FPoison::Get();
+	}
+
 	check(InX->Type->AsScalar());
 	check(InX->Type == InY->Type);
 
@@ -498,6 +521,11 @@ FValue* FEmitter::EmitVector2(FValue* InX, FValue* InY)
 
 FValue* FEmitter::EmitVector3(FValue* InX, FValue* InY, FValue* InZ)
 {
+	if (IsAnyNonNullAPoison(InX, InY, InZ))
+	{
+		return FPoison::Get();
+	}
+
 	check(InX->Type->AsScalar());
 	check(InX->Type == InY->Type);
 	check(InY->Type == InZ->Type);
@@ -513,6 +541,11 @@ FValue* FEmitter::EmitVector3(FValue* InX, FValue* InY, FValue* InZ)
 
 FValue* FEmitter::EmitVector4(FValue* InX, FValue* InY, FValue* InZ, FValue* InW)
 {
+	if (IsAnyNonNullAPoison(InX, InY, InZ, InW))
+	{
+		return FPoison::Get();
+	}
+
 	check(InX->Type->AsScalar());
 	check(InX->Type == InY->Type);
 	check(InY->Type == InZ->Type);
@@ -537,6 +570,11 @@ FValue* FEmitter::EmitExternalInput(EExternalInput Id)
 
 FValue* FEmitter::EmitSubscript(FValue* Value, int32 Index)
 {
+	if (IsAnyNonNullAPoison(Value))
+	{
+		return Value;
+	}
+
 	FPrimitiveTypePtr PrimitiveType = Value->Type->AsPrimitive();
 	if (!PrimitiveType)
 	{
@@ -551,14 +589,14 @@ FValue* FEmitter::EmitSubscript(FValue* Value, int32 Index)
 	}
 
 	// Getting first component and Value is already a scalar, just return itself.
-	if (FDimensional* DimensionalValue = Value->As<FDimensional>())
+	if (FDimensional* DimensionalValue = As<FDimensional>(Value))
 	{
 		check(Index < DimensionalValue->GetComponents().Num());
 		return DimensionalValue->GetComponents()[Index];
 	}
 	
 	// Avoid subscripting a subscript (e.g. no value.xy.x)
-	if (FSubscript* Subscript = Value->As<FSubscript>())
+	if (FSubscript* Subscript = As<FSubscript>(Value))
 	{
 		Value = Subscript->Arg;
 	}
@@ -573,6 +611,11 @@ FValue* FEmitter::EmitSubscript(FValue* Value, int32 Index)
 
 FValue* FEmitter::EmitSwizzle(FValue* Value, FSwizzleMask Mask)
 {
+	if (IsAnyNonNullAPoison(Value))
+	{
+		return Value;
+	}
+
 	// At least one component must have been specified.
 	check(Mask.NumComponents > 0);
 
@@ -722,12 +765,12 @@ FValue* FEmitter::EmitCustomPrimitiveData(uint32 PrimitiveDataIndex)
 	return nullptr;
 }
 
-FSetMaterialOutput* FEmitter::EmitSetMaterialOutput(EMaterialProperty InProperty, FValue* InArgValue)
+FSetMaterialOutput* FEmitter::EmitSetMaterialOutput(EMaterialProperty InProperty, FValue* Arg)
 {
 	FSetMaterialOutput Proto = MakePrototype<FSetMaterialOutput>(nullptr);
 	Proto.Block 	= Module->RootBlock;
 	Proto.Property  = InProperty;
-	Proto.Arg 		= InArgValue;
+	Proto.Arg 		= Arg;
 
 	FSetMaterialOutput* Instr = static_cast<FSetMaterialOutput*>(EmitPrototype(this, Proto));
 	Module->Outputs.Add(Instr);
@@ -975,9 +1018,9 @@ static FValue* TryFoldOperatorScalar(FEmitter* Emitter, EOperator Op, FValue* A,
 	}
 
 	// Verify that both lhs and rhs are constants, otherwise we cannot fold the operation.
-	const FConstant* AConstant = A->As<FConstant>();
-	const FConstant* BConstant = B->As<FConstant>();
-	const FConstant* CConstant = C->As<FConstant>();
+	const FConstant* AConstant = As<FConstant>(A);
+	const FConstant* BConstant = As<FConstant>(B);
+	const FConstant* CConstant = As<FConstant>(C);
 	if (!AConstant ||  (IsBinaryOperator(Op) && !BConstant) || (IsTernaryOperator(Op) && !CConstant))
 	{
 		return nullptr;
@@ -1047,14 +1090,14 @@ static FValue* TryFoldOperatorScalar(FEmitter* Emitter, EOperator Op, FValue* A,
 /// Used to filter what parameter *primitive* types operators can take.
 enum EOperatorParameterTypeFilter
 {
-	OT_Unknown = 0xff,					///< Unspecified
-	OT_Any = 0,							///< Any primitive type (no filters)
-	OT_Boolean = 1 << 0,				///< Boolean primitive type of any dimension
-	OT_Arithmetic = 1 << 1,				///< Arithmetic primitive type of any dimension (i.e. that supports arithmetic operations)
-	OT_AnyFloat = 1 << 2,				///< Floating point primitive type of any dimension
-	OT_SameAsFirstArgument = 1 << 3,	///< Must be the same type as the first argument
-	OT_NotMatrix = 1 << 4,				///< Any primitive type except matrices
-	OT_Vector3 = 1 << 5,				///< 3D vector of any scalar type
+	OPTF_Unknown = 0xff,				///< Unspecified
+	OPTF_Any = 0,						///< Any primitive type (no filters)
+	OPTF_Boolean = 1 << 0,				///< Boolean primitive type of any dimension
+	OPTF_Arithmetic = 1 << 1,			///< Arithmetic primitive type of any dimension (i.e. that supports arithmetic operations)
+	OPTF_AnyFloat = 1 << 2,				///< Floating point primitive type of any dimension
+	OPTF_AsFirstArgument = 1 << 3,		///< Must be the same type as the first argument
+	OPTF_NotMatrix = 1 << 4,			///< Any primitive type except matrices
+	OPTF_Vector3 = 1 << 5,				///< 3D vector of any scalar type
 };
 
 inline EOperatorParameterTypeFilter operator|(EOperatorParameterTypeFilter A, EOperatorParameterTypeFilter B)
@@ -1074,7 +1117,7 @@ enum EOperatorReturnType
 /// The signature of an operator consisting of its parameter and return type information.
 struct FOperatorSignature
 {
-	EOperatorParameterTypeFilter ParameterTypes[3] = { OT_Unknown, OT_Unknown, OT_Unknown };
+	EOperatorParameterTypeFilter ParameterTypes[3] = { OPTF_Unknown, OPTF_Unknown, OPTF_Unknown };
 	EOperatorReturnType ReturnType = ORT_Unknown;
 };
 
@@ -1083,15 +1126,15 @@ static const FOperatorSignature* GetOperatorSignature(EOperator Op)
 {
 	static const FOperatorSignature* Signatures = [] ()
 	{
-		const FOperatorSignature UnaryArithmetic = { { OT_Arithmetic }, ORT_FirstArgumentType };
-		const FOperatorSignature BinaryArithmetic = { { OT_Arithmetic, OT_SameAsFirstArgument }, ORT_FirstArgumentType };
-		const FOperatorSignature TernaryArithmetic = { { OT_Arithmetic, OT_SameAsFirstArgument, OT_SameAsFirstArgument }, ORT_FirstArgumentType };
-		const FOperatorSignature BinaryComparison = { { OT_Any, OT_SameAsFirstArgument }, ORT_BooleanWithFirstArgumentDimensions };
-		const FOperatorSignature BinaryLogical = { { OT_Boolean, OT_SameAsFirstArgument }, ORT_FirstArgumentType };
+		const FOperatorSignature UnaryArithmetic = { { OPTF_Arithmetic }, ORT_FirstArgumentType };
+		const FOperatorSignature BinaryArithmetic = { { OPTF_Arithmetic, OPTF_AsFirstArgument }, ORT_FirstArgumentType };
+		const FOperatorSignature TernaryArithmetic = { { OPTF_Arithmetic, OPTF_AsFirstArgument, OPTF_AsFirstArgument }, ORT_FirstArgumentType };
+		const FOperatorSignature BinaryComparison = { { OPTF_Any, OPTF_AsFirstArgument }, ORT_BooleanWithFirstArgumentDimensions };
+		const FOperatorSignature BinaryLogical = { { OPTF_Boolean, OPTF_AsFirstArgument }, ORT_FirstArgumentType };
 
 		static FOperatorSignature S[MIR::OperatorCount];
 		S[UO_Negate] 				= UnaryArithmetic;
-		S[UO_Not] 					= { { OT_Boolean }, ORT_FirstArgumentType };
+		S[UO_Not] 					= { { OPTF_Boolean }, ORT_FirstArgumentType };
 		S[UO_Abs] 					= UnaryArithmetic;
 		S[UO_ACos] 					= UnaryArithmetic;
 		S[UO_ASin] 					= UnaryArithmetic;
@@ -1102,7 +1145,7 @@ static const FOperatorSignature* GetOperatorSignature(EOperator Op)
 		S[UO_Exponential2] 			= UnaryArithmetic;
 		S[UO_Floor] 				= UnaryArithmetic;
 		S[UO_Frac]					= UnaryArithmetic;
-		S[UO_Length]				= { { OT_Arithmetic | OT_NotMatrix }, ORT_FirstArgumentTypeToScalar };
+		S[UO_Length]				= { { OPTF_Arithmetic | OPTF_NotMatrix }, ORT_FirstArgumentTypeToScalar };
 		S[UO_Logarithm] 			= UnaryArithmetic;
 		S[UO_Logarithm2] 			= UnaryArithmetic;
 		S[UO_Logarithm10] 			= UnaryArithmetic;
@@ -1128,8 +1171,8 @@ static const FOperatorSignature* GetOperatorSignature(EOperator Op)
 		S[BO_Or] 					= BinaryLogical;
 		S[BO_And] 					= BinaryLogical;
 		S[BO_Xor] 					= BinaryLogical;
-		S[BO_Dot] 					= { { OT_Arithmetic | OT_NotMatrix, OT_SameAsFirstArgument }, ORT_FirstArgumentTypeToScalar };
-		S[BO_Cross] 				= { { OT_Arithmetic | OT_Vector3, OT_SameAsFirstArgument }, ORT_FirstArgumentType };
+		S[BO_Dot] 					= { { OPTF_Arithmetic | OPTF_NotMatrix, OPTF_AsFirstArgument }, ORT_FirstArgumentTypeToScalar };
+		S[BO_Cross] 				= { { OPTF_Arithmetic | OPTF_Vector3, OPTF_AsFirstArgument }, ORT_FirstArgumentType };
 		S[TO_Clamp]					= TernaryArithmetic;
 		S[TO_Lerp] 					= TernaryArithmetic;
 		S[TO_Smoothstep]			= TernaryArithmetic;
@@ -1140,7 +1183,7 @@ static const FOperatorSignature* GetOperatorSignature(EOperator Op)
  
 /// Validates that the types of the arguments are valid for specified operator.
 /// If valid, it returns the type of the result. Otherwise if it is not valid, it returns nullptr.
-static FPrimitiveTypePtr ValidateOperatorAndGetResultType(FEmitter* Emitter, EOperator Op, FValue* A, FValue* B, FValue* C)
+static FPrimitiveTypePtr ValidateOperatorAndGetResultType(FEmitter* Emitter, EOperator Op, FValue* A, FValue*& B, FValue*& C)
 {
 	// Argument A must have always been provided.
 	check(A);
@@ -1174,9 +1217,18 @@ static FPrimitiveTypePtr ValidateOperatorAndGetResultType(FEmitter* Emitter, EOp
 		}
 
 		EOperatorParameterTypeFilter Filter = Signature->ParameterTypes[i];
-		check(Filter != OT_Unknown); // No signature specified for this operator.
+		check(Filter != OPTF_Unknown); // No signature specified for this operator.
+
+		if (Filter & OPTF_AsFirstArgument)
+		{
+			check(i > 0); // This check can't apply to the first argument.
+
+			// Cast this argument to the first argument primitive type
+			Arguments[i] = Emitter->EmitConstruct(FirstArgumentPrimitiveType, Arguments[i]);
+			bValid &= !Arguments[i]->IsPoison();
+		}
 
-		if (Filter & OT_Boolean)
+		if (Filter & OPTF_Boolean)
 		{
 			if (ArgumentPrimitiveType->ScalarKind != ScalarKind_Bool)
 			{
@@ -1185,7 +1237,7 @@ static FPrimitiveTypePtr ValidateOperatorAndGetResultType(FEmitter* Emitter, EOp
 			}
 		}
 		
-		if (Filter & OT_Arithmetic)
+		if (Filter & OPTF_Arithmetic)
 		{
 			if (!ScalarKindIsArithmetic(ArgumentPrimitiveType->ScalarKind))
 			{
@@ -1194,7 +1246,7 @@ static FPrimitiveTypePtr ValidateOperatorAndGetResultType(FEmitter* Emitter, EOp
 			}
 		}
 		
-		if (Filter & OT_AnyFloat)
+		if (Filter & OPTF_AnyFloat)
 		{
 			if (!ScalarKindIsAnyFloat(ArgumentPrimitiveType->ScalarKind))
 			{
@@ -1202,17 +1254,8 @@ static FPrimitiveTypePtr ValidateOperatorAndGetResultType(FEmitter* Emitter, EOp
 				bValid = false;
 			}
 		}
-
-		if (Filter & OT_SameAsFirstArgument)
-		{
-			if (ArgumentPrimitiveType != FirstArgumentPrimitiveType)
-			{
-				Emitter->Errorf(TEXT("Operator %s's %s argument type (%s) expected to be the same as first argument (%s)."), OperatorToString(Op), ArgumentsStr[i], ArgumentPrimitiveType->Spelling.GetData(), FirstArgumentPrimitiveType->Spelling.GetData());
-				bValid = false;
-			}
-		}
 		
-		if (Filter & OT_NotMatrix)
+		if (Filter & OPTF_NotMatrix)
 		{
 			if (ArgumentPrimitiveType->IsMatrix())
 			{
@@ -1221,7 +1264,7 @@ static FPrimitiveTypePtr ValidateOperatorAndGetResultType(FEmitter* Emitter, EOp
 			}
 		}
 
-		if (Filter & OT_Vector3)
+		if (Filter & OPTF_Vector3)
 		{
 			if (!ArgumentPrimitiveType->IsVector() || ArgumentPrimitiveType->GetNumComponents() != 3)
 			{
@@ -1231,6 +1274,11 @@ static FPrimitiveTypePtr ValidateOperatorAndGetResultType(FEmitter* Emitter, EOp
 		}
 	}
 
+	// Arguments might have changed, update the references.
+	A = Arguments[0];
+	B = Arguments[1];
+	C = Arguments[2];
+
 	if (!bValid)
 	{
 		return nullptr;
@@ -1368,7 +1416,7 @@ static FValue* TryFoldComponentwiseOperator(FEmitter* Emitter, EOperator Op, FVa
 static bool TryUnpackConstantScalarOrVector(FValue* V, TArrayView<FConstant*> OutComponents, int32& OutNumComponents)
 {
 	// V not specified? Or not a scalar/vector?
-	FDimensional* Dimensional = V->As<FDimensional>();
+	FDimensional* Dimensional = As<FDimensional>(V);
 	if (!Dimensional || V->Type->AsPrimitive()->IsMatrix())
 	{
 		return false;
@@ -1377,7 +1425,7 @@ static bool TryUnpackConstantScalarOrVector(FValue* V, TArrayView<FConstant*> Ou
 	TArrayView<FValue*> Components = Dimensional->GetComponents();
 	for (int32 i = 0; i < Components.Num(); ++i)
 	{
-		OutComponents[i] = Components[i]->As<FConstant>();
+		OutComponents[i] = As<FConstant>(Components[i]);
 		if (!OutComponents[i])
 		{
 			return false;
@@ -1471,6 +1519,11 @@ static FValue* TryFoldOperator(FEmitter* Emitter, EOperator Op, FValue* A, FValu
 
 FValue* FEmitter::EmitOperator(EOperator Op, FValue* A, FValue* B, FValue* C)
 {
+	if (IsAnyNonNullAPoison(A, B, C))
+	{
+		return FPoison::Get();
+	}
+
 	// Validate the operation and retrieve the result type.
 	FPrimitiveTypePtr ResultType = ValidateOperatorAndGetResultType(this, Op, A, B, C);
 	if (!ResultType)
@@ -1503,10 +1556,14 @@ FValue* FEmitter::EmitOperator(EOperator Op, FValue* A, FValue* B, FValue* C)
 FValue* FEmitter::EmitBranch(FValue* Condition, FValue* True, FValue* False)
 {
 	// Condition must be of type bool
-	check(Condition->Type->IsBoolScalar());
+	Condition = EmitConstruct(FPrimitiveType::GetBool1(), Condition);
+	if (IsAnyNonNullAPoison(Condition))
+	{
+		return FPoison::Get();
+	}
 
 	// If the condition is a scalar constant, then simply evaluate the result now.
-	if (const FConstant* ConstCondition = Condition->As<FConstant>())
+	if (const FConstant* ConstCondition = As<FConstant>(Condition))
 	{
 		return ConstCondition->Boolean ? True : False;
 	}
@@ -1514,9 +1571,20 @@ FValue* FEmitter::EmitBranch(FValue* Condition, FValue* True, FValue* False)
 	// Condition must be of type bool
 	check(Condition->Type->IsBoolScalar());
 
-	// If the condition is not static, the types of the true and false arguments must match.
-	// The resulting type will be that of the branch instruction created.
-	check(True->Type == False->Type);
+	// If the condition is not static, make both true and false arguments have the same type,
+	// by casting false argument into the true's type.
+	FTypePtr CommonType = GetCommonType(True->Type, False->Type);
+	if (!CommonType)
+	{
+		return FPoison::Get();
+	}
+
+	True = EmitConstruct(True->Type, True);
+	False = EmitConstruct(True->Type, False);
+	if (IsAnyNonNullAPoison(True, False))
+	{
+		return FPoison::Get();
+	}
 
 	// Create the branch instruction.
 	FBranch Proto = MakePrototype<FBranch>(True->Type);
@@ -1541,7 +1609,7 @@ static FValue* CastConstant(FEmitter* Emitter, FConstant* Constant, EScalarKind
 		{
 			switch (TargetKind)
 			{
-				case ScalarKind_Bool: UE_MIR_TODO();
+				case ScalarKind_Bool: return nullptr;
 				case ScalarKind_Int: return Emitter->EmitConstantIntScalar(Constant->Integer);
 				case ScalarKind_Float: return Emitter->EmitConstantFloatScalar((TFloat)Constant->Integer);
 				default: UE_MIR_UNREACHABLE();
@@ -1552,7 +1620,7 @@ static FValue* CastConstant(FEmitter* Emitter, FConstant* Constant, EScalarKind
 		{
 			switch (TargetKind)
 			{
-				case ScalarKind_Bool: UE_MIR_TODO();
+				case ScalarKind_Bool: return nullptr;
 				case ScalarKind_Int: return Emitter->EmitConstantIntScalar((int32)Constant->Float);
 				default: UE_MIR_UNREACHABLE();
 			}
@@ -1587,7 +1655,7 @@ static FValue* ConstructPrimitiveValue(FEmitter* Emitter, FPrimitiveTypePtr Targ
 		}
 
 		// Construct the scalar from a constant.
-		if (FConstant* ConstantInitializer = Initializer->As<FConstant>())
+		if (FConstant* ConstantInitializer = As<FConstant>(Initializer))
 		{
 			return CastConstant(Emitter, ConstantInitializer, InitializerPrimitiveType->ScalarKind, TargetPrimitiveType->ScalarKind);
 		}
@@ -1657,7 +1725,7 @@ static FValue* ConstructPrimitiveValue(FEmitter* Emitter, FPrimitiveTypePtr Targ
 		check(TargetPrimitiveType->IsMatrix());
 
 		// 
-		if (FDimensional* DimensionalInitializer = Initializer->As<FDimensional>())
+		if (FDimensional* DimensionalInitializer = As<FDimensional>(Initializer))
 		{
 			// Create the result dimensional value.
 			FDimensional* Result = NewDimensionalValue(Emitter, TargetPrimitiveType);
@@ -1705,6 +1773,7 @@ FValue* FEmitter::EmitConstruct(FTypePtr Type, FValue* Initializer)
 	if (!Result)
 	{
 		Errorf(TEXT("Cannot construct a '%s' from a '%s'."), TEXT(""), TEXT(""));
+		return FPoison::Get();
 	}
 	
 	return Result;
@@ -1724,28 +1793,45 @@ FValue* FEmitter::EmitTextureSample(FValue* Texture, FValue* TexCoord, ESamplerS
 	return EmitPrototype(this, Prototype);
 }
 
-FPrimitiveTypePtr FEmitter::GetCommonPrimitiveType(FPrimitiveTypePtr A, FPrimitiveTypePtr B)
+FTypePtr FEmitter::TryGetCommonType(FTypePtr A, FTypePtr B)
 {
 	// Trivial case: types are equal
 	if (A == B)
 	{
 		return A;
 	}
+
+	FPrimitiveTypePtr PrimitiveA = A->AsPrimitive();
+	FPrimitiveTypePtr PrimitiveB = B->AsPrimitive();
+	if (!PrimitiveA || !PrimitiveB)
+	{
+		return nullptr;
+	}
 	
 	// If both A and B are matrices, their dimensions must match (equality check above didn't trigger).
-	if (A->IsMatrix() || B->IsMatrix())
+	if (PrimitiveA->IsMatrix() || PrimitiveB->IsMatrix())
 	{
-		Errorf(TEXT("No common primitive type between `%s` and `%s`."), A->Spelling.GetData(), B->Spelling.GetData());
 		return nullptr;
 	}
 
 	// Neither A nor B are matrices, but single scalar or vector. Return the largest.
-	check(A->NumColumns == 1 && B->NumColumns == 1);
-	EScalarKind ScalarKind = FMath::Max(A->ScalarKind, B->ScalarKind);
-	int32 NumRows = FMath::Max(A->NumRows, B->NumRows);
+	check(PrimitiveA->NumColumns == 1 && PrimitiveB->NumColumns == 1);
+	EScalarKind ScalarKind = FMath::Max(PrimitiveA->ScalarKind, PrimitiveB->ScalarKind);
+	int32 NumRows = FMath::Max(PrimitiveA->NumRows, PrimitiveB->NumRows);
 	return FPrimitiveType::Get(ScalarKind, NumRows, 1);
 }
 
+FTypePtr FEmitter::GetCommonType(FTypePtr A, FTypePtr B)
+{
+	if (FTypePtr CommonType = TryGetCommonType(A, B))
+	{
+		return CommonType;
+	}
+
+	Errorf(TEXT("No common type between `%s` and `%s`."), A->GetSpelling().GetData(), B->GetSpelling().GetData());
+	return nullptr;
+}
+
 void FEmitter::Error(FStringView Message)
 {
 	FMaterialIRModule::FError Error;
diff --git a/Engine/Source/Runtime/Engine/Private/Materials/MaterialIRModuleBuilder.cpp b/Engine/Source/Runtime/Engine/Private/Materials/MaterialIRModuleBuilder.cpp
--- a/Engine/Source/Runtime/Engine/Private/Materials/MaterialIRModuleBuilder.cpp
+++ b/Engine/Source/Runtime/Engine/Private/Materials/MaterialIRModuleBuilder.cpp
@@ -345,7 +345,7 @@ struct FMaterialIRModuleBuilderImpl
 				}
 
 				// This value is used. If it hasn't been analyzed yet, analyze it now.
-				MIR::FInstruction* Use = UseValue->AsInstruction();
+				MIR::FInstruction* Use = MIR::AsInstruction(UseValue);
 				if (!Use)
 				{
 					if (!(UseValue->Flags & MIR::VF_Analyzed))
@@ -374,7 +374,7 @@ struct FMaterialIRModuleBuilderImpl
 
 			for (MIR::FValue* UseValue : Instr->GetUses())
 			{
-				if (MIR::FInstruction* UseAsInstruction = UseValue->AsInstruction())
+				if (MIR::FInstruction* UseAsInstruction = MIR::AsInstruction(UseValue))
 				{
 					UseAsInstruction->NumUsers += 1;
 				}
@@ -416,7 +416,7 @@ struct FMaterialIRModuleBuilderImpl
 			for (int32 UseIndex = 0; UseIndex < Uses.Num(); ++UseIndex)
 			{
 				MIR::FValue* Use = Uses[UseIndex];
-				MIR::FInstruction* UseInstr = Use->AsInstruction();
+				MIR::FInstruction* UseInstr = MIR::AsInstruction(Use);
 				if (!UseInstr)
 				{
 					continue;
diff --git a/Engine/Source/Runtime/Engine/Private/Materials/MaterialIRToHLSLTranslator.cpp b/Engine/Source/Runtime/Engine/Private/Materials/MaterialIRToHLSLTranslator.cpp
--- a/Engine/Source/Runtime/Engine/Private/Materials/MaterialIRToHLSLTranslator.cpp
+++ b/Engine/Source/Runtime/Engine/Private/Materials/MaterialIRToHLSLTranslator.cpp
@@ -305,7 +305,7 @@ struct FTranslator : FMaterialIRToHLSLTranslation
 	
 	ENoOp LowerValue(MIR::FValue* InValue)
 	{
-		if (MIR::FInstruction* Instr = InValue->AsInstruction())
+		if (MIR::FInstruction* Instr = MIR::AsInstruction(InValue))
 		{
 			if (Instr->NumUsers <= 1 && IsFoldable(Instr))
 			{
diff --git a/Engine/Source/Runtime/Engine/Public/Materials/MaterialIR.h b/Engine/Source/Runtime/Engine/Public/Materials/MaterialIR.h
--- a/Engine/Source/Runtime/Engine/Public/Materials/MaterialIR.h
+++ b/Engine/Source/Runtime/Engine/Public/Materials/MaterialIR.h
@@ -59,12 +59,6 @@ struct FValue
 	// Returns whether this value is of specified kind.
 	bool IsA(EValueKind InKind) const;
 
-	// Tries to cast this value to an instruction and returns it (nullptr otherwise).
-	FInstruction* AsInstruction();
-
-	// Tries to cast this value to an instruction and returns it (nullptr otherwise).
-	const FInstruction* AsInstruction() const;
-
 	// Returns the array of this value's uses. An use is another value referenced by this one (e.g. the operands of a binary expression).
 	TArrayView<const FValue*> GetUses() const;
 
@@ -106,19 +100,47 @@ struct FValue
 
 	// Tries to cast this value to specified type T and returns the casted pointer, if possible (nullptr otherwise).
 	template <typename T>
-	T* As() { return this && IsA(T::TypeKind) ? static_cast<T*>(this) : nullptr; }
+	T* As() { return IsA(T::TypeKind) ? static_cast<T*>(this) : nullptr; }
 
 	// Tries to cast this value to specified type T and returns the casted pointer, if possible (nullptr otherwise).
 	template <typename T>
-	const T* As() const { return this && IsA(T::TypeKind) ? static_cast<const T*>(this) : nullptr; }
+	const T* As() const { return IsA(T::TypeKind) ? static_cast<const T*>(this) : nullptr; }
 };
 
+// If specified value is not null, it tries to cast this value T and returns it.
+// Otherwise, it returns null.
+template <typename T>
+T* As(FValue* Value)
+{
+	return Value && Value->IsA(T::TypeKind) ? static_cast<T*>(Value) : nullptr;
+}
+
+// If specified value is not null, it tries to cast this value T and returns it.
+// Otherwise, it returns null.
+template <typename T>
+const T* As(const FValue* Value)
+{
+	return Value && Value->IsA(T::TypeKind) ? static_cast<const T*>(Value) : nullptr;
+}
+
+// If specified value is not null, it tries to cast this value to an instruction and returns it.
+// Otherwise, it returns null.
+FInstruction* AsInstruction(FValue* Value);
+
+// If specified value is not null, it tries to cast this value to an instruction and returns it.
+// Otherwise, it returns null.
+const FInstruction* AsInstruction(const FValue* Value);
+
 template <EValueKind TTypeKind>
 struct TValue : FValue
 {
 	static constexpr EValueKind TypeKind = TTypeKind;
 };
 
+// A poison value represents an invalid value. It is produced by the emitter when an
+// invalid operation is performed. Poison values can be passed as arguments to other operations,
+// but they are "contagious": any instruction emitted with a poison value as an argument
+// will itself produce a poison value.
 struct FPoison : TValue<VK_Poison>
 {
 	static FPoison* Get();
@@ -435,5 +457,26 @@ struct FTextureSample : TInstruction<VK_TextureSample>
 	// TODO: TexCoordDerivatives;
 };
 
+// --- Helper functions ---
+
+static inline bool IsAnyNonNullAPoison()
+{
+	return false;
+}
+
+/// Returns whether there is a poison in the values specified, which could be null.
+template <typename... TTail>
+bool IsAnyNonNullAPoison(const FValue* Head, TTail... Tail)
+{
+	if (Head && Head->IsPoison())
+	{
+		return true;
+	}
+	else
+	{
+		return IsAnyNonNullAPoison(Tail...);
+	}
+}
+
 } // namespace UE::MIR
 #endif
diff --git a/Engine/Source/Runtime/Engine/Public/Materials/MaterialIREmitter.h b/Engine/Source/Runtime/Engine/Public/Materials/MaterialIREmitter.h
--- a/Engine/Source/Runtime/Engine/Public/Materials/MaterialIREmitter.h
+++ b/Engine/Source/Runtime/Engine/Public/Materials/MaterialIREmitter.h
@@ -156,7 +156,7 @@ public:
 
 	/* Instructions */
 
-	FSetMaterialOutput* EmitSetMaterialOutput(EMaterialProperty InProperty, FValue* InArgValue);
+	FSetMaterialOutput* EmitSetMaterialOutput(EMaterialProperty InProperty, FValue* Arg);
 	FValue* EmitOperator(EOperator Operator, FValue* A, FValue* B = nullptr, FValue* C = nullptr);
 	FValue* EmitBranch(FValue* Condition, FValue* True, FValue* False);
 	FValue* EmitConstruct(FTypePtr Type, FValue* Initializer);
@@ -164,7 +164,8 @@ public:
 
 	/* Types */
 
-	FPrimitiveTypePtr GetCommonPrimitiveType(FPrimitiveTypePtr A, FPrimitiveTypePtr B);
+	FTypePtr TryGetCommonType(FTypePtr A, FTypePtr B);
+	FTypePtr GetCommonType(FTypePtr A, FTypePtr B);
 
 	/* Error reporting */
 
