Initial pass at integrating IAD support into steamable manager.
- Progress not implemented yet.
- Fix some issues with error handling.
- Fix requests for packages (explicitly packages, not the asset in the package) not finding the package after load completes.

diff --git a/Engine/Source/Runtime/Core/Private/IO/IoStatus.cpp b/Engine/Source/Runtime/Core/Private/IO/IoStatus.cpp
--- a/Engine/Source/Runtime/Core/Private/IO/IoStatus.cpp
+++ b/Engine/Source/Runtime/Core/Private/IO/IoStatus.cpp
@@ -64,6 +64,11 @@ bool FIoStatus::operator==(const FIoStatus& Other) const
 
 FString FIoStatus::ToString() const
 {
+	if (ErrorMessage[0] == TEXT('\0'))
+	{
+		return FString::Format(TEXT("({0})"), { GetIoErrorText(ErrorCode) });
+	}
+
 	return FString::Format(TEXT("{0} ({1})"), { ErrorMessage, GetIoErrorText(ErrorCode) });
 }
 
diff --git a/Engine/Source/Runtime/Engine/Classes/Engine/StreamableManager.h b/Engine/Source/Runtime/Engine/Classes/Engine/StreamableManager.h
--- a/Engine/Source/Runtime/Engine/Classes/Engine/StreamableManager.h
+++ b/Engine/Source/Runtime/Engine/Classes/Engine/StreamableManager.h
@@ -4,6 +4,7 @@
 
 #include "CoreMinimal.h"
 #include "Containers/ArrayView.h"
+#include "IO/IoStoreOnDemand.h"
 #include "Math/NumericLimits.h"
 #include "Misc/PackageAccessTracking.h"
 #include "Misc/SourceLocation.h"
@@ -26,17 +27,102 @@
 	DownloadError,
 };
 
+ENGINE_API const TCHAR* LexToString(EStreamableErrorType ErrorType);
+
 struct FStreamableError
 {
+	FStreamableError(EStreamableErrorType InErrorType, FString InErrorMessage)
+		: ErrorMessage(MoveTemp(InErrorMessage))
+		, ErrorType(InErrorType)
+	{}
+
+	FString ErrorMessage; // TODO: Can this just be a TCHAR* or string view ?
 	EStreamableErrorType ErrorType = EStreamableErrorType::Unknown;
-	FString ErrorMessage;
 };
 
+ENGINE_API FString LexToString(const FStreamableError Error);
+
 /** Defines FStreamableDelegate delegate interface */
 using FStreamableDelegate = TDelegate<void()>;
 using FStreamableDelegateWithHandle = TDelegate<void(TSharedPtr<struct FStreamableHandle>)>;
 using FStreamableUpdateDelegate = TDelegate<void(TSharedRef<struct FStreamableHandle>)>;
 
+namespace UE::StreamableManager::Private
+{
+	// Default priority for all async loads
+	static constexpr TAsyncLoadPriority DefaultAsyncLoadPriority = 0;
+	// Priority to try and load immediately
+	static constexpr TAsyncLoadPriority AsyncLoadHighPriority = 100;
+
+	inline FStreamableDelegateWithHandle WrapDelegate(const FStreamableDelegate& InDelegate)
+	{
+		if (InDelegate.IsBound())
+		{
+			return FStreamableDelegateWithHandle::CreateLambda(
+				[InDelegate](TSharedPtr<FStreamableHandle>)
+				{
+					// Delegates may get tick delayed so may not be safe to call if bound to a GC'd object
+					InDelegate.ExecuteIfBound();
+				});
+		}
+
+		return FStreamableDelegateWithHandle();
+	}
+
+	inline FStreamableDelegateWithHandle WrapDelegate(FStreamableDelegate&& InDelegate)
+	{
+		if (InDelegate.IsBound())
+		{
+			return FStreamableDelegateWithHandle::CreateLambda(
+				[InDelegate = MoveTemp(InDelegate)](TSharedPtr<FStreamableHandle>)
+				{
+					// Delegates may get tick delayed so may not be safe to call if bound to a GC'd object
+					InDelegate.ExecuteIfBound();
+				});
+		}
+
+		return FStreamableDelegateWithHandle();
+	}
+
+	template<class> inline constexpr bool always_false_v = false;
+
+	struct FDownloadContext;
+}
+
+/** Handle that pins downloaded packages in the local cache, may be shared between requests */
+struct FStreamableDownloadCachePin
+{
+	ENGINE_API FStreamableDownloadCachePin();
+	ENGINE_API ~FStreamableDownloadCachePin();
+
+	template<typename DebugNameType = UE::FSourceLocation>
+	static FStreamableDownloadCachePin Create(DebugNameType&& DebugNameOrLocation)
+	{
+		FStreamableDownloadCachePin Result;
+		if constexpr (std::is_same_v<std::decay_t<DebugNameType>, UE::FSourceLocation>)
+		{
+			Result.ContentHandle = UE::IoStore::FOnDemandContentHandle::Create(
+				UE::SourceLocation::ToFileAndLineString(DebugNameOrLocation));
+		}
+		else
+		{
+			Result.ContentHandle = UE::IoStore::FOnDemandContentHandle::Create(
+				Forward<DebugNameType>(DebugNameOrLocation));
+		}
+
+		return Result;
+	}
+
+	ENGINE_API FString GetDebugName() const;
+
+	ENGINE_API void Release();
+	ENGINE_API bool IsValid() const;
+	operator bool() const { return IsValid(); }
+private:
+	UE::IoStore::FOnDemandContentHandle ContentHandle;
+	friend struct UE::StreamableManager::Private::FDownloadContext;
+};
+
 #if WITH_EDITOR
 enum class ECookLoadType : uint8;
 #endif
@@ -147,16 +233,24 @@
 		return bIsCombinedHandle;
 	}
 
+	/** True if an error occured */
 	bool HasError() const
 	{
 		return Error.IsSet();
 	}
 
+	/** Get error details */
 	const TOptional<FStreamableError>& GetError() const
 	{
 		return Error;
 	}
 
+	/** Get download cache pin */
+	FStreamableDownloadCachePin GetDownloadCachePin() const
+	{
+		return DownloadCachePin;
+	}
+
 	/** Returns true if we've done all the loading we can now, ie all handles are either completed or stalled */
 	ENGINE_API bool HasLoadCompletedOrStalled() const;
 
@@ -226,7 +320,7 @@
 	ENGINE_API bool BindCancelDelegate(FStreamableDelegateWithHandle NewDelegate);
 	ENGINE_API bool BindCancelDelegate(FStreamableDelegate NewDelegate);
 
-	/** Bind delegate that is called periodically as delegate updates, only works if loading is in progress. This will overwrite any already bound delegate! */
+	/** Bind delegate that is called periodically as request updates, only works if loading is in progress. This will overwrite any already bound delegate! */
 	ENGINE_API bool BindUpdateDelegate(FStreamableUpdateDelegate NewDelegate);
 
 	/**
@@ -430,6 +524,7 @@
 private:
 	friend struct FStreamableManager;
 	friend struct FStreamable;
+	friend struct UE::StreamableManager::Private::FDownloadContext;
 
 	/** Called from self as part of SetDebugNameIfEmptyOrDefault */
 	bool IsHandleNameEmptyOrDefault() const;
@@ -459,26 +554,6 @@
 	void SetError(FStreamableError&& Error);
 	void SetErrorOnParents(const FStreamableError& Error);
 
-	TOptional<FStreamableError> Error;
-
-	/** True if this request has finished loading. It may still be active, or it may have been released */
-	bool bLoadCompleted;
-
-	/** True if this request was released, which will stop it from keeping hard GC references */
-	bool bReleased;
-
-	/** True if this request was explicitly canceled, which stops it from calling the completion delegate and immediately releases it */
-	bool bCanceled;
-
-	/** True if this handle has been created but not yet actually requested. This handle is probably waiting for a resource like a chunk to be available */
-	bool bStalled;
-
-	/** If true, this handle will be released when it finishes loading */
-	bool bReleaseWhenLoaded;
-
-	/** If true, this handle will be released when it finishes loading */
-	bool bIsCombinedHandle;
-
 	/** Delegate to call when streaming is completed */
 	FStreamableDelegateWithHandle CompleteDelegate;
 
@@ -491,6 +566,15 @@
 	/** Name of this handle, passed in by caller to help in debugging */
 	FString DebugName;
 
+	/** Error details in case of an error */
+	TOptional<FStreamableError> Error;
+
+	/** Handle that holds packages in the download cache */
+	FStreamableDownloadCachePin DownloadCachePin;
+
+	/** Download book-keeping, released when download finishes */
+	TUniquePtr<UE::StreamableManager::Private::FDownloadContext> DownloadContext;
+
 	/** The async priority for this request */
 	TAsyncLoadPriority Priority;
 
@@ -526,6 +610,24 @@
 #if WITH_EDITOR
 	ECookLoadType CookLoadType;
 #endif
+
+	/** True if this request has finished loading. It may still be active, or it may have been released */
+	bool bLoadCompleted;
+
+	/** True if this request was released, which will stop it from keeping hard GC references */
+	bool bReleased;
+
+	/** True if this request was explicitly canceled, which stops it from calling the completion delegate and immediately releases it */
+	bool bCanceled;
+
+	/** True if this handle has been created but not yet actually requested. This handle is probably waiting for a resource like a chunk to be available */
+	bool bStalled;
+
+	/** If true, this handle will be released when it finishes loading */
+	bool bReleaseWhenLoaded;
+
+	/** If true, this handle will be released when it finishes loading */
+	bool bIsCombinedHandle;
 };
 
 enum class EStreamableManagerCombinedHandleOptions : uint8 
@@ -541,13 +643,61 @@
 
 ENUM_CLASS_FLAGS(EStreamableManagerCombinedHandleOptions);
 
-/** A native class for managing streaming assets in and keeping them in memory. AssetManager is the global singleton version of this with blueprint access */
-struct FStreamableManager : public FGCObject
+namespace UE::StreamableManager
 {
 	// Default priority for all async loads
 	static constexpr TAsyncLoadPriority DefaultAsyncLoadPriority = 0;
 	// Priority to try and load immediately
 	static constexpr TAsyncLoadPriority AsyncLoadHighPriority = 100;
+}
 
+enum class EStreamableDownloadPriority : uint8
+{
+	/** Below normal priority. */
+	Low,
+	/** Normal priority. */
+	Normal,
+	/** Above normal priority. */
+	High
+};
+
+/** Parameters for an async load */
+struct FStreamableAsyncLoadParams
+{
+	/** Assets to load off disk */
+	TArray<FSoftObjectPath> TargetsToStream;
+	/** [optional] Delegate to call when load finishes. Will be called on the next tick if asset is already loaded, or many seconds later */
+	FStreamableDelegateWithHandle OnComplete;
+	/** [optional] Delegate to call when the load is canceled. Will be called on the next tick or many seconds later */
+	FStreamableDelegateWithHandle OnCancel;
+	/** [optional] delegate that is called periodically as request updates */
+	FStreamableUpdateDelegate OnUpdate;
+	/** [optional] Priority to pass to the streaming system, higher priority will be loaded first */
+	TAsyncLoadPriority Priority = UE::StreamableManager::Private::DefaultAsyncLoadPriority;
+	/** [optional] If true, the manager will keep the streamable handle active until explicitly released */
+	bool bManageActiveHandle = false;
+	/** [optional] If true, the handle will start in a stalled state and will not attempt to actually async load until StartStalledHandle is called on it */
+	bool bStartStalled = false;
+};
+
+/** Additional parameters for downloading packages */
+struct FStreamableDownloadParams
+{
+	/** Handle to keep packages in the local cache. If no cache pin is provided, a new one will be created. */
+	FStreamableDownloadCachePin CachePin;
+	/** Download priority */
+	EStreamableDownloadPriority Priority = EStreamableDownloadPriority::Normal;
+	/** If set, any soft references will also be installed */
+	bool bInstallSoftReferences = false;
+};
+
+/** A native class for managing streaming assets in and keeping them in memory. AssetManager is the global singleton version of this with blueprint access */
+struct FStreamableManager : public FGCObject
+{
+	// Default priority for all async loads
+	static constexpr TAsyncLoadPriority DefaultAsyncLoadPriority = UE::StreamableManager::Private::DefaultAsyncLoadPriority;
+	// Priority to try and load immediately
+	static constexpr TAsyncLoadPriority AsyncLoadHighPriority = UE::StreamableManager::Private::AsyncLoadHighPriority;
+
 	/** 
 	 * This is the primary streamable operation. Requests streaming of one or more target objects. When complete, a delegate function is called. Returns a Streamable Handle.
@@ -569,5 +719,30 @@
 		DebugNameType&& DebugNameOrLocation = UE::FSourceLocation::Current());
 
 	/** 
+	 * This is the primary streamable operation. Requests streaming of one or more target objects. When complete, a delegate function is called. Returns a Streamable Handle.
+	 *
+	 * @param Params				Load Parameters
+	 * @param DebugName				[optional] Name of this handle, either FString or anything that can construct FString, will be reported in debug tools, will report Source Location if DebugName is not specified explicitly
+	 */ 
+	template< typename DebugNameType = UE::FSourceLocation >
+	TSharedPtr<FStreamableHandle> RequestAsyncLoad(
+		FStreamableAsyncLoadParams&& Params,
+		DebugNameType&& DebugNameOrLocation = UE::FSourceLocation::Current());
+
+	/** 
+	 * This is the primary streamable operation. Requests streaming of one or more target objects. When complete, a delegate function is called. Returns a Streamable Handle.
+	 * This overload will download packages if needed.
+	 *
+	 * @param Params				Load Parameters
+	 * @param DownloadParams		Additional Parameters to control downloading
+	 * @param DebugName				[optional] Name of this handle, either FString or anything that can construct FString, will be reported in debug tools, will report Source Location if DebugName is not specified explicitly
+	 */ 
+	template< typename DebugNameType = UE::FSourceLocation >
+	TSharedPtr<FStreamableHandle> RequestAsyncLoad(
+		FStreamableAsyncLoadParams&& LoadParams,
+		FStreamableDownloadParams&& DownloadParams,
+		DebugNameType&& DebugNameOrLocation = UE::FSourceLocation::Current());
+
+	/** 
 	 * Synchronously load a set of assets, and return a handle. This can be very slow and may stall the game thread for several seconds.
 	 * 
@@ -663,6 +838,7 @@
 	FStreamableManager& operator=(const FStreamableManager&) = delete;
 private:
 	friend FStreamableHandle;
+	friend UE::StreamableManager::Private::FDownloadContext;
 
 	void RemoveReferencedAsset(const FSoftObjectPath& Target, TSharedRef<FStreamableHandle> Handle);
 	void StartHandleRequests(TSharedRef<FStreamableHandle> Handle);
@@ -678,7 +854,7 @@
 	void AsyncLoadCallback(FSoftObjectPath Request, UPackage* Package, EAsyncLoadingResult::Type Result);
 
 	ENGINE_API static bool ShouldStripDebugName();
-	ENGINE_API TSharedPtr<FStreamableHandle> RequestAsyncLoadInternal(TArray<FSoftObjectPath>&& TargetsToStream, FStreamableDelegateWithHandle&& DelegateToCall, TAsyncLoadPriority Priority, bool bManageActiveHandle, bool bStartStalled, FString&& DebugName);
+	ENGINE_API TSharedPtr<FStreamableHandle> RequestAsyncLoadInternal(FStreamableAsyncLoadParams&& Params, FStreamableDownloadParams&& DownloadParams, FString&& DebugName);
 	ENGINE_API TSharedPtr<FStreamableHandle> RequestSyncLoadInternal(TArray<FSoftObjectPath>&& TargetsToStream, bool bManageActiveHandle, FString&& DebugName);
 
 	/** Map of paths to streamable objects, this will be the post-redirector name */
@@ -750,41 +926,6 @@
 	}
 }
 
-namespace UE::StreamableManager::Private
-{
-	inline FStreamableDelegateWithHandle WrapDelegate(const FStreamableDelegate& InDelegate)
-	{
-		if (InDelegate.IsBound())
-		{
-			return FStreamableDelegateWithHandle::CreateLambda(
-				[InDelegate](TSharedPtr<FStreamableHandle>)
-				{
-					// Delegates may get tick delayed so may not be safe to call if bound to a GC'd object
-					InDelegate.ExecuteIfBound();
-				});
-		}
-
-		return FStreamableDelegateWithHandle();
-	}
-
-	inline FStreamableDelegateWithHandle WrapDelegate(FStreamableDelegate&& InDelegate)
-	{
-		if (InDelegate.IsBound())
-		{
-			return FStreamableDelegateWithHandle::CreateLambda(
-				[InDelegate = MoveTemp(InDelegate)](TSharedPtr<FStreamableHandle>)
-				{
-					// Delegates may get tick delayed so may not be safe to call if bound to a GC'd object
-					InDelegate.ExecuteIfBound();
-				});
-		}
-
-		return FStreamableDelegateWithHandle();
-	}
-
-	template<class> inline constexpr bool always_false_v = false;
-}
-
 template< typename PathContainerType, typename FuncType, typename DebugNameType >
 TSharedPtr<FStreamableHandle> FStreamableManager::RequestAsyncLoad(
 	PathContainerType&& TargetsToStream,
@@ -794,23 +935,25 @@
 	bool bStartStalled,
 	DebugNameType&& DebugNameOrLocation)
 {
-	FStreamableDelegateWithHandle DelegateToCall;
+	FStreamableAsyncLoadParams Params;
+	Params.TargetsToStream = TArray<FSoftObjectPath>{ Forward<PathContainerType>(TargetsToStream) };
+
 	if constexpr (std::is_constructible_v<FStreamableDelegateWithHandle, FuncType>)
 	{
-		DelegateToCall = FStreamableDelegateWithHandle(Forward<FuncType>(Callback));
+		Params.OnComplete = FStreamableDelegateWithHandle(Forward<FuncType>(Callback));
 	}
 	else if constexpr (std::is_constructible_v<FStreamableDelegate, FuncType>)
 	{
-		DelegateToCall = UE::StreamableManager::Private::WrapDelegate(
+		Params.OnComplete = UE::StreamableManager::Private::WrapDelegate(
 			FStreamableDelegate(Forward<FuncType>(Callback)));
 	}
 	else if constexpr (std::is_invocable_v<FuncType, TSharedPtr<FStreamableHandle>>)
 	{
-		DelegateToCall = FStreamableDelegateWithHandle::CreateLambda(Forward<FuncType>(Callback));
+		Params.OnComplete = FStreamableDelegateWithHandle::CreateLambda(Forward<FuncType>(Callback));
 	}
 	else if constexpr (std::is_invocable_v<FuncType>)
 	{
-		DelegateToCall = UE::StreamableManager::Private::WrapDelegate(
+		Params.OnComplete = UE::StreamableManager::Private::WrapDelegate(
 			FStreamableDelegate::CreateLambda(Forward<FuncType>(Callback)));
 	}
 	else
@@ -819,24 +962,70 @@
 		static_assert(UE::StreamableManager::Private::always_false_v<FuncType>, "Cannot bind callback");
 	}
 
+	Params.Priority = Priority;
+	Params.bManageActiveHandle = bManageActiveHandle;
+	Params.bStartStalled = bStartStalled;
+
 	if constexpr (std::is_same_v<std::decay_t<DebugNameType>, UE::FSourceLocation>)
 	{
 		return RequestAsyncLoadInternal(
-			TArray<FSoftObjectPath>{ Forward<PathContainerType>(TargetsToStream) },
-			MoveTemp(DelegateToCall),
-			Priority,
-			bManageActiveHandle,
-			bStartStalled,
+			MoveTemp(Params),
+			FStreamableDownloadParams(),
 			ShouldStripDebugName() ? FString() : UE::SourceLocation::ToFileAndLineString(DebugNameOrLocation));
 	}
 	else
 	{
 		return RequestAsyncLoadInternal(
-			TArray<FSoftObjectPath>{ Forward<PathContainerType>(TargetsToStream) },
-			MoveTemp(DelegateToCall),
-			Priority,
-			bManageActiveHandle,
-			bStartStalled,
+			MoveTemp(Params),
+			FStreamableDownloadParams(),
 			ShouldStripDebugName() ? FString() : FString{ Forward<DebugNameType>(DebugNameOrLocation) });
 	}
 }
+
+template< typename DebugNameType>
+TSharedPtr<FStreamableHandle> FStreamableManager::RequestAsyncLoad(
+	FStreamableAsyncLoadParams&& Params,
+	DebugNameType&& DebugNameOrLocation)
+{
+	if constexpr (std::is_same_v<std::decay_t<DebugNameType>, UE::FSourceLocation>)
+	{
+		return RequestAsyncLoadInternal(
+			MoveTemp(Params),
+			FStreamableDownloadParams(),
+			ShouldStripDebugName() ? FString() : UE::SourceLocation::ToFileAndLineString(DebugNameOrLocation));
+	}
+	else
+	{
+		return RequestAsyncLoadInternal(
+			MoveTemp(Params),
+			FStreamableDownloadParams(),
+			ShouldStripDebugName() ? FString() : FString{ Forward<DebugNameType>(DebugNameOrLocation) });
+	}
+}
+
+template< typename DebugNameType >
+TSharedPtr<FStreamableHandle> FStreamableManager::RequestAsyncLoad(
+	FStreamableAsyncLoadParams&& LoadParams,
+	FStreamableDownloadParams&& DownloadParams,
+	DebugNameType&& DebugNameOrLocation)
+{
+	if (!DownloadParams.CachePin)
+	{
+		DownloadParams.CachePin = FStreamableDownloadCachePin::Create(Forward<DebugNameType>(DebugNameOrLocation));
+	}
+
+	if constexpr (std::is_same_v<std::decay_t<DebugNameType>, UE::FSourceLocation>)
+	{
+		return RequestAsyncLoadInternal(
+			MoveTemp(LoadParams),
+			MoveTemp(DownloadParams),
+			ShouldStripDebugName() ? FString() : UE::SourceLocation::ToFileAndLineString(DebugNameOrLocation));
+	}
+	else
+	{
+		return RequestAsyncLoadInternal(
+			MoveTemp(LoadParams),
+			MoveTemp(DownloadParams),
+			ShouldStripDebugName() ? FString() : FString{ Forward<DebugNameType>(DebugNameOrLocation) });
+	}
+}
diff --git a/Engine/Source/Runtime/Engine/Engine.Build.cs b/Engine/Source/Runtime/Engine/Engine.Build.cs
--- a/Engine/Source/Runtime/Engine/Engine.Build.cs
+++ b/Engine/Source/Runtime/Engine/Engine.Build.cs
@@ -22,6 +22,7 @@ public class Engine : ModuleRules
 				"AudioMixer", 
 				"AudioMixerCore",
 				"InputCore",
+				"IoStoreOnDemand",
 				"MovieSceneCapture", 
 				"PacketHandler", 
 				"Renderer",
@@ -167,6 +168,7 @@ public class Engine : ModuleRules
 				"Icmp",
 				"UniversalObjectLocator",
 				"XmlParser",
+				"IoStoreOnDemand",
 			}
 		);
 
diff --git a/Engine/Source/Runtime/Engine/Private/StreamableManager.cpp b/Engine/Source/Runtime/Engine/Private/StreamableManager.cpp
--- a/Engine/Source/Runtime/Engine/Private/StreamableManager.cpp
+++ b/Engine/Source/Runtime/Engine/Private/StreamableManager.cpp
@@ -4,6 +4,7 @@
 #include "UObject/ICookInfo.h"
 #include "UObject/ObjectRedirector.h"
 #include "UObject/Package.h"
+#include "Algo/Transform.h"
 #include "Misc/PackageAccessTrackingOps.h"
 #include "Misc/PackageName.h"
 #include "Misc/ScopeLock.h"
@@ -45,6 +46,33 @@ static FAutoConsoleVariableRef CVarStreamableStripDebugNameInShipping(
 	ECVF_Default
 );
 
+const TCHAR* LexToString(EStreamableErrorType ErrorType)
+{
+	switch (ErrorType)
+	{
+	default:
+	case EStreamableErrorType::Unknown:
+		return TEXT("Unknown");
+	case EStreamableErrorType::PackageLoadError:
+		return TEXT("PackageLoadError");
+	case EStreamableErrorType::PackageLoadCanceled:
+		return TEXT("PackageLoadCanceled");
+	case EStreamableErrorType::DownloadError:
+		return TEXT("DownloadError");
+	}
+}
+
+FString LexToString(const FStreamableError Error)
+{
+	return FString::Printf(TEXT("%s - %s"), LexToString(Error.ErrorType), *Error.ErrorMessage);
+}
+
+FStreamableDownloadCachePin::FStreamableDownloadCachePin() = default;
+FStreamableDownloadCachePin::~FStreamableDownloadCachePin() = default;
+FString FStreamableDownloadCachePin::GetDebugName() const { return LexToString(ContentHandle); }
+void FStreamableDownloadCachePin::Release() { ContentHandle.Reset(); }
+bool FStreamableDownloadCachePin::IsValid() const { return ContentHandle.IsValid(); }
+
 const FString FStreamableHandle::HandleDebugName_Preloading = FString(TEXT("Preloading"));
 const FString FStreamableHandle::HandleDebugName_AssetList = FString(TEXT("LoadAssetList"));
 const FString FStreamableHandle::HandleDebugName_CombinedHandle = FString(TEXT("CreateCombinedHandle"));
@@ -345,20 +373,179 @@ TStreamableHandleContextDataTypeID FStreamableHandleContextDataBase::AllocateCla
 	return Result;
 }
 
+namespace UE::StreamableManager::Private
+{
+	// self-deleting context for managing package downloads for a load request
+	struct FDownloadContext
+	{
+		UE::IoStore::FOnDemandInstallRequest InstallRequest;
+		EStreamableDownloadPriority Priority = EStreamableDownloadPriority::Normal;
+		bool bInstallSoftReferences = false;
+
+		bool IsStarted() const
+		{
+			return !InstallRequest.IsNone();
+		}
+
+		~FDownloadContext()
+		{
+			InstallRequest.Cancel();
+		}
+
+		void Cancel(const TSharedRef<FStreamableHandle>& OwningHandle)
+		{
+			if (IsStarted())
+			{
+				// Will be cleaned up when request calls back
+				InstallRequest.Cancel();
+			}
+			else
+			{
+				// No callback coming, destroy self now
+				OwningHandle->DownloadContext = nullptr;
+			}
+		}
+
+		void CompleteLoad(const TSharedRef<FStreamableHandle>& OwningHandle)
+		{
+			OwningHandle->CompleteLoad();
+			if (OwningHandle->bReleaseWhenLoaded)
+			{
+				OwningHandle->ReleaseHandle();
+			}
+		}
+
+		void CompleteLoad(const TSharedRef<FStreamableHandle>& OwningHandle, FStreamableError&& Error)
+		{
+			OwningHandle->SetError(MoveTemp(Error));
+			CompleteLoad(OwningHandle);
+		}
+
+		void Start(const TSharedRef<FStreamableHandle>& OwningHandle)
+		{
+			if (!InstallRequest.IsNone())
+			{
+				checkf(false, TEXT("IAD install already started"));
+				return;
+			}
+
+			UE::IoStore::IOnDemandIoStore* IoStore = UE::IoStore::TryGetOnDemandIoStore();
+			if (!IoStore)
+			{
+				CompleteLoad(OwningHandle, FStreamableError{ EStreamableErrorType::DownloadError, TEXT("IoStore not found") });
+				OwningHandle->DownloadContext = nullptr; // Destroy self
+				return;
+			}
+
+			UE::IoStore::FOnDemandInstallArgs InstallArgs;
+			InstallArgs.PackageIds.Reserve(OwningHandle->RequestedAssets.Num());
+			Algo::Transform(OwningHandle->RequestedAssets, InstallArgs.PackageIds, [](const FSoftObjectPath& AssetPath)
+				{
+					return FPackageId::FromName(AssetPath.GetLongPackageFName());
+				});
+
+			check(OwningHandle->DownloadCachePin);
+			InstallArgs.ContentHandle = OwningHandle->DownloadCachePin.ContentHandle;
+
+			InstallArgs.Options |= UE::IoStore::EOnDemandInstallOptions::CallbackOnGameThread;
+			if (bInstallSoftReferences)
+			{
+				InstallArgs.Options |= UE::IoStore::EOnDemandInstallOptions::InstallSoftReferences;
+			}
+
+			InstallArgs.Priority = UE::IoStore::EOnDemandInstallPriority::Normal;
+			switch (Priority)
+			{
+			case EStreamableDownloadPriority::Low:
+				InstallArgs.Priority = UE::IoStore::EOnDemandInstallPriority::Low;
+				break;
+			case EStreamableDownloadPriority::High:
+				InstallArgs.Priority = UE::IoStore::EOnDemandInstallPriority::High;
+				break;
+			}
+
+			InstallRequest = IoStore->Install(MoveTemp(InstallArgs),
+				[WeakOwningHandle = OwningHandle.ToWeakPtr()](const UE::IoStore::FOnDemandInstallResult& OnDemandInstallResult)
+				{
+					TSharedPtr<FStreamableHandle> OwningHandle = WeakOwningHandle.Pin();
+					if (OwningHandle && OwningHandle->DownloadContext)
+					{
+						OwningHandle->DownloadContext->OnInstallComplete(OwningHandle.ToSharedRef(), OnDemandInstallResult);
+					}
+				},
+				[WeakOwningHandle = OwningHandle.ToWeakPtr()](const UE::IoStore::FOnDemandInstallProgress& InProgress)
+				{
+					TSharedPtr<FStreamableHandle> OwningHandle = WeakOwningHandle.Pin();
+					if (OwningHandle && OwningHandle->DownloadContext)
+					{
+						OwningHandle->DownloadContext->OnInstallProgress(OwningHandle.ToSharedRef(), InProgress);
+					}
+				}
+			);
+		}
+
+		void OnInstallComplete(const TSharedRef<FStreamableHandle>& OwningHandle, const UE::IoStore::FOnDemandInstallResult& OnDemandInstallResult)
+		{
+			// Destroy self at end of scope
+			TUniquePtr<FDownloadContext> Self = MoveTemp(OwningHandle->DownloadContext);
+			OwningHandle->DownloadContext = nullptr;
+
+			if (!OnDemandInstallResult.Status.IsOk())
+			{
+				if (OwningHandle->WasCanceled() && OnDemandInstallResult.Status.GetErrorCode() == EIoErrorCode::Cancelled)
+				{
+					CompleteLoad(OwningHandle);
+				}
+				else
+				{
+					CompleteLoad(OwningHandle, FStreamableError{ EStreamableErrorType::DownloadError, OnDemandInstallResult.Status.ToString() });
+				}
+
+				return;
+			}
+
+			if (OwningHandle->OwningManager)
+			{
+				// StartHandleRequests will progress to asset loading now that the download context is gone
+				OwningHandle->OwningManager->StartHandleRequests(OwningHandle->AsShared());
+			}
+		}
+
+		void OnInstallProgress(const TSharedRef<FStreamableHandle>& OwningHandle, const UE::IoStore::FOnDemandInstallProgress& InProgress)
+		{
+			// TODO: progress
+		}
+	};
+
+	static FStreamableError GetStreamableError(EAsyncLoadingResult::Type Result)
+	{
+		switch (Result)
+		{
+		case EAsyncLoadingResult::Failed:
+			// We could supply an error string if async loading bubbled one up...
+			// Possibly GetExplanationForUnavailablePackage?
+			return FStreamableError{ EStreamableErrorType::PackageLoadError, TEXT("Async load failed") };
+		case EAsyncLoadingResult::Canceled:
+			return FStreamableError{ EStreamableErrorType::PackageLoadCanceled, TEXT("Async load canceled") };
+		}
+
+		return FStreamableError{ EStreamableErrorType::Unknown, FString() };
+	}
+}
 
 FStreamableHandle::FStreamableHandle(FPrivateToken)
-	: bLoadCompleted(false)
-	, bReleased(false)
-	, bCanceled(false)
-	, bStalled(false)
-	, bReleaseWhenLoaded(false)
-	, bIsCombinedHandle(false)
-	, Priority(0)
+	: Priority(0)
 	, StreamablesLoading(0)
 	, OwningManager(nullptr)
 #if WITH_EDITOR
 	, CookLoadType(ECookLoadType::Unspecified)
 #endif
+	, bLoadCompleted(false)
+	, bReleased(false)
+	, bCanceled(false)
+	, bStalled(false)
+	, bReleaseWhenLoaded(false)
+	, bIsCombinedHandle(false)
 {
 }
 
@@ -653,7 +840,7 @@ void FStreamableHandle::CancelHandle()
 
 	// Remove from referenced list. If it is stalled then it won't have been registered with
 	// the manager yet
-	if (!bStalled)
+	if (!bStalled && !DownloadContext)
 	{
 		for (const FSoftObjectPath& AssetRef : RequestedAssets)
 		{
@@ -661,6 +848,12 @@ void FStreamableHandle::CancelHandle()
 		}
 	}
 
+	if (DownloadContext)
+	{
+		// This may destroy DownloadContext if its not pending
+		DownloadContext->Cancel(AsShared());
+	}
+
 	// Remove from explicit list
 	OwningManager->ManagedActiveHandles.Remove(SharedThis);
 
@@ -960,16 +1153,9 @@ void FStreamableHandle::AsyncLoadCallbackWrapper(const FName& PackageName, UPack
 	// Needed so we can bind with a shared pointer for safety
 	if (OwningManager)
 	{
-		switch (Result)
+		if (Result != EAsyncLoadingResult::Succeeded)
 		{
-		case EAsyncLoadingResult::Failed:
-			// We could supply an error string if async loading bubbled one up...
-			// Possibly GetExplanationForUnavailablePackage?
-			SetError(FStreamableError{ EStreamableErrorType::PackageLoadError, TEXT("Async load failed")});
-			break;
-		case EAsyncLoadingResult::Canceled:
-			SetError(FStreamableError{ EStreamableErrorType::PackageLoadCanceled, TEXT("Async load canceled")});
-			break;
+			SetError(UE::StreamableManager::Private::GetStreamableError(Result));
 		}
 
 		OwningManager->AsyncLoadCallback(MoveTemp(TargetName), Package, Result);
@@ -1469,18 +1655,20 @@ bool FStreamableManager::ShouldStripDebugName()
 	return bShouldStripDebugName;
 }
 
-TSharedPtr<FStreamableHandle> FStreamableManager::RequestAsyncLoadInternal(TArray<FSoftObjectPath>&& TargetsToStream, FStreamableDelegateWithHandle&& DelegateToCall, TAsyncLoadPriority Priority, bool bManageActiveHandle, bool bStartStalled, FString&& DebugName)
+TSharedPtr<FStreamableHandle> FStreamableManager::RequestAsyncLoadInternal(FStreamableAsyncLoadParams&& Params, FStreamableDownloadParams&& DownloadParams, FString&& DebugName)
 {
 	LLM_SCOPE(ELLMTag::StreamingManager);
 
 	// Schedule a new callback, this will get called when all related async loads are completed
 	TSharedRef<FStreamableHandle> NewRequest = MakeShared<FStreamableHandle>(FStreamableHandle::FPrivateToken());
-	NewRequest->CompleteDelegate = MoveTemp(DelegateToCall);
+	NewRequest->CompleteDelegate = MoveTemp(Params.OnComplete);
+	NewRequest->CancelDelegate = MoveTemp(Params.OnCancel);
+	NewRequest->UpdateDelegate = MoveTemp(Params.OnUpdate);
 	NewRequest->OwningManager = this;
-	NewRequest->RequestedAssets = MoveTemp(TargetsToStream);
+	NewRequest->RequestedAssets = MoveTemp(Params.TargetsToStream);
 	NewRequest->DebugName = MoveTemp(DebugName);
 
-	NewRequest->Priority = Priority;
+	NewRequest->Priority = Params.Priority;
 #if UE_WITH_PACKAGE_ACCESS_TRACKING
 	PackageAccessTracking_Private::FTrackedData* AccumulatedScopeData = PackageAccessTracking_Private::FPackageAccessRefScope::GetCurrentThreadAccumulatedData();
 	if (AccumulatedScopeData)
@@ -1507,7 +1695,8 @@ TSharedPtr<FStreamableHandle> FStreamableManager::RequestAsyncLoadInternal(TArra
 		}
 		else if (FPackageName::IsShortPackageName(TargetName.GetLongPackageFName()))
 		{
-			UE_LOG(LogStreamableManager, Error, TEXT("RequestAsyncLoad(%s) called with invalid package name %s"), *DebugName, *TargetName.ToString());
+			UE_LOG(LogStreamableManager, Error, TEXT("RequestAsyncLoad(%s) called with invalid package name %s"), *NewRequest->DebugName, *TargetName.ToString());
+			NewRequest->CancelDelegate.Unbind(); // cancel delegate is not exepected since a null handle is returned
 			NewRequest->CancelHandle();
 			return nullptr;
 		}
@@ -1517,7 +1706,8 @@ TSharedPtr<FStreamableHandle> FStreamableManager::RequestAsyncLoadInternal(TArra
 	if (NumValidRequests == 0)
 	{
 		// Original array was empty or all null
-		UE_LOG(LogStreamableManager, Display, TEXT("RequestAsyncLoad(%s) called with empty or only null assets!"), *DebugName);
+		UE_LOG(LogStreamableManager, Display, TEXT("RequestAsyncLoad(%s) called with empty or only null assets!"), *NewRequest->DebugName);
+		NewRequest->CancelDelegate.Unbind(); // cancel delegate is not exepected since a null handle is returned
 		NewRequest->CancelHandle();
 		return nullptr;
 	} 
@@ -1542,7 +1732,7 @@ TSharedPtr<FStreamableHandle> FStreamableManager::RequestAsyncLoadInternal(TArra
 		}
 
 		// Some valid, some null
-		UE_LOG(LogStreamableManager, Display, TEXT("RequestAsyncLoad(%s) called with both valid and null assets, null assets removed from %s!"), *DebugName, *RequestedSet);
+		UE_LOG(LogStreamableManager, Display, TEXT("RequestAsyncLoad(%s) called with both valid and null assets, null assets removed from %s!"), *NewRequest->DebugName, *RequestedSet);
 	}
 
 	if (TargetSet.Num() != NewRequest->RequestedAssets.Num())
@@ -1559,22 +1749,30 @@ TSharedPtr<FStreamableHandle> FStreamableManager::RequestAsyncLoadInternal(TArra
 			RequestedSet += Asset.ToString();
 		}
 
-		UE_LOG(LogStreamableManager, Verbose, TEXT("RequestAsyncLoad(%s) called with duplicate assets, duplicates removed from %s!"), *DebugName, *RequestedSet);
+		UE_LOG(LogStreamableManager, Verbose, TEXT("RequestAsyncLoad(%s) called with duplicate assets, duplicates removed from %s!"), *NewRequest->DebugName, *RequestedSet);
 #endif
 
 		NewRequest->RequestedAssets = TargetSet.Array();
 	}
 
-	if (bManageActiveHandle)
+	if (Params.bManageActiveHandle)
 	{
 		// This keeps a reference around until explicitly released
 		ManagedActiveHandles.Add(NewRequest);
 	}
 
-	if (bStartStalled)
+	if (DownloadParams.CachePin)
 	{
-		NewRequest->bStalled = true;
+		NewRequest->DownloadCachePin = DownloadParams.CachePin;
+		NewRequest->DownloadContext = MakeUnique<UE::StreamableManager::Private::FDownloadContext>();
+		NewRequest->DownloadContext->Priority = DownloadParams.Priority;
+		NewRequest->DownloadContext->bInstallSoftReferences = DownloadParams.bInstallSoftReferences;
 	}
+
+	if (Params.bStartStalled)
+	{
+		NewRequest->bStalled = true;
+	} 
 	else
 	{
 		StartHandleRequests(NewRequest);
@@ -1590,8 +1788,13 @@ TSharedPtr<FStreamableHandle> FStreamableManager::RequestSyncLoadInternal(TArray
 	// Otherwise, only do a sync load if there are no background sync loads, this is faster but will cause a sync flush
 	bForceSynchronousLoads = IsInAsyncLoadingThread() || IsEventDrivenLoaderEnabled() || !IsAsyncLoading();
 
+	FStreamableAsyncLoadParams AsyncParams;
+	AsyncParams.TargetsToStream = MoveTemp(TargetsToStream);
+	AsyncParams.Priority = AsyncLoadHighPriority;
+	AsyncParams.bManageActiveHandle = bManageActiveHandle;
+
 	// Do an async load and wait to complete. In some cases this will do a sync load due to safety issues
-	TSharedPtr<FStreamableHandle> Request = RequestAsyncLoadInternal(MoveTemp(TargetsToStream), FStreamableDelegateWithHandle(), AsyncLoadHighPriority, bManageActiveHandle, false, MoveTemp(DebugName));
+	TSharedPtr<FStreamableHandle> Request = RequestAsyncLoadInternal(MoveTemp(AsyncParams), FStreamableDownloadParams(), MoveTemp(DebugName));
 
 	bForceSynchronousLoads = false;
 
@@ -1618,6 +1821,13 @@ TSharedPtr<FStreamableHandle> FStreamableManager::RequestSyncLoadInternal(TArray
 
 void FStreamableManager::StartHandleRequests(TSharedRef<FStreamableHandle> Handle)
 {
+	if (Handle->DownloadContext)
+	{
+		Handle->DownloadContext->Start(Handle);
+		// StartHandleRequests will be called again when the download completes
+		return;
+	}
+
 	TRACE_LOADTIME_REQUEST_GROUP_SCOPE(TEXT("StreamableManager - %s"), *Handle->GetDebugName());
 
 	TArray<FStreamable *> ExistingStreamables;
@@ -1676,7 +1886,7 @@ TArray<int32> FStreamableManager::GetAsyncLoadRequestIds(TSharedRef<FStreamableH
 }
 
 UE_TRACE_EVENT_BEGIN(Cpu, StreamableManager_LoadSynchronous, NoSync)
-	UE_TRACE_EVENT_FIELD(UE::Trace::WideString, AssetPath)
+UE_TRACE_EVENT_FIELD(UE::Trace::WideString, AssetPath)
 UE_TRACE_EVENT_END()
 
 UObject* FStreamableManager::LoadSynchronous(const FSoftObjectPath& Target, bool bManageActiveHandle, TSharedPtr<FStreamableHandle>* RequestHandlePointer, UE::FSourceLocation Location)
@@ -1717,9 +1927,27 @@ void FStreamableManager::FindInMemory(FSoftObjectPath& InOutTargetName, struct F
 	UE_LOG(LogStreamableManager, Verbose, TEXT("     Searching in memory for %s"), *InOutTargetName.ToString());
 	Existing->Target = nullptr;
 
-	UObject* Asset = Package ? 
-		StaticFindObjectFast(UObject::StaticClass(), Package, InOutTargetName.GetAssetFName()) :
-		StaticFindObject(UObject::StaticClass(), InOutTargetName.GetAssetPath(), /*ExactClass*/false);	
+	UObject* Asset = nullptr;
+	if (InOutTargetName.GetAssetFName().IsNone())
+	{
+		if (Package && Package->GetFName() == InOutTargetName.GetLongPackageFName())
+		{
+			Asset = Package;
+		}
+		else
+		{
+			Asset = StaticFindObjectFast(UObject::StaticClass(), nullptr, InOutTargetName.GetLongPackageFName());
+		}
+	}
+	else if (Package)
+	{
+		Asset = StaticFindObjectFast(UObject::StaticClass(), Package, InOutTargetName.GetAssetFName());
+	}
+	else
+	{
+		Asset = StaticFindObject(UObject::StaticClass(), InOutTargetName.GetAssetPath(), /*ExactClass*/false);
+	}
+
 	if (Asset)
 	{
 		if (InOutTargetName.IsAsset())
@@ -1762,6 +1990,8 @@ void FStreamableManager::AsyncLoadCallback(FSoftObjectPath TargetName, UPackage*
 	UE_LOG(LogStreamableManager, Verbose, TEXT("Stream Complete callback %s"), *TargetName.ToString());
 	if (Existing)
 	{
+		const bool bCompleteRequests = Existing->bAsyncLoadRequestOutstanding;
+
 		if (Existing->bAsyncLoadRequestOutstanding)
 		{
 			Existing->bAsyncLoadRequestOutstanding = false;
@@ -1769,8 +1999,6 @@ void FStreamableManager::AsyncLoadCallback(FSoftObjectPath TargetName, UPackage*
 			{
 				FindInMemory(TargetName, Existing, Package);
 			}
-
-			CheckCompletedRequests(TargetName, Existing);
 		}
 		else
 		{
@@ -1785,9 +2013,28 @@ void FStreamableManager::AsyncLoadCallback(FSoftObjectPath TargetName, UPackage*
 		{
 			// Async load failed to find the object
 			Existing->bLoadFailed = true;
-			ensure(Result != EAsyncLoadingResult::Succeeded);
 			UE_LOG(LogStreamableManager, Verbose, TEXT("    Failed async load for %s"), *TargetName.ToString());
 		}
+
+		if (Result != EAsyncLoadingResult::Succeeded)
+		{
+			for (const TSharedRef<FStreamableHandle>& Handle : Existing->LoadingHandles)
+			{
+				Handle->SetError(UE::StreamableManager::Private::GetStreamableError(Result));
+			}
+		}
+		else if (!Existing->Target)
+		{
+			for (const TSharedRef<FStreamableHandle>& Handle : Existing->LoadingHandles)
+			{
+				Handle->SetError(FStreamableError{ EStreamableErrorType::PackageLoadError, TEXT("Failed to find target asset in memory") });
+			}
+		}
+
+		if (bCompleteRequests)
+		{
+			CheckCompletedRequests(TargetName, Existing);
+		}
 	}
 	else
 	{
