Added wp.Runtime.DrawRuntimeHash2DScaleFactor to allow to scale down the 2D debug display of the WP streaming grids.

diff --git a/Engine/Source/Runtime/Engine/Private/WorldPartition/WorldPartitionSubsystem.cpp b/Engine/Source/Runtime/Engine/Private/WorldPartition/WorldPartitionSubsystem.cpp
--- a/Engine/Source/Runtime/Engine/Private/WorldPartition/WorldPartitionSubsystem.cpp
+++ b/Engine/Source/Runtime/Engine/Private/WorldPartition/WorldPartitionSubsystem.cpp
@@ -85,6 +85,14 @@ static FAutoConsoleCommand CVarDrawWorldPartitionIndex(
 		}
 	}));
 
+static float GDrawRuntimeHash2DScaleFactor = 1;
+static FAutoConsoleVariableRef CVarDrawRuntimeHash2DScaleFactor(
+	TEXT("wp.Runtime.DrawRuntimeHash2DScaleFactor"),
+	GDrawRuntimeHash2DScaleFactor,
+	TEXT("Allows to reduce the size on screen of the 2D debug display of world partition runtime hash (value must be between 0.1 and 1)"),
+	ECVF_Default
+);
+
 static int32 GDrawRuntimeHash3D = 0;
 static FAutoConsoleCommand CVarDrawRuntimeHash3D(
 	TEXT("wp.Runtime.ToggleDrawRuntimeHash3D"),
@@ -1427,7 +1435,8 @@ void UWorldPartitionSubsystem::Draw(UCanvas* Canvas, class APlayerController* PC
 		const FVector2D CanvasBottomRightPadding(10.f, 10.f);
 		const FVector2D CanvasMinimumSize(100.f, 100.f);
 		const FVector2D CanvasMaxScreenSize = FVector2D::Max(MaxScreenRatio*FVector2D(Canvas->ClipX, Canvas->ClipY) - CanvasBottomRightPadding - CurrentOffset, CanvasMinimumSize);
-		const FVector2D PartitionCanvasSize = FVector2D(CanvasMaxScreenSize.X, CanvasMaxScreenSize.Y);
+		const float DrawRuntimeHash2DScaleFactor = FMath::Clamp<float>(GDrawRuntimeHash2DScaleFactor, 0.1, 1);
+		const FVector2D PartitionCanvasSize = FVector2D(CanvasMaxScreenSize.X, CanvasMaxScreenSize.Y) * DrawRuntimeHash2DScaleFactor;
 
 		FBox2D WorldRegion(ForceInit);
 		WorldPartitionsDraw2DContext.SetNum(RegisteredWorldPartitions.Num(), EAllowShrinking::No);
