Implemented support for setting Texture Collection parameters on UMaterials and UMaterialInstanceConstant

diff --git a/Engine/Source/Runtime/Engine/Private/Materials/Material.cpp b/Engine/Source/Runtime/Engine/Private/Materials/Material.cpp
--- a/Engine/Source/Runtime/Engine/Private/Materials/Material.cpp
+++ b/Engine/Source/Runtime/Engine/Private/Materials/Material.cpp
@@ -1593,6 +1593,13 @@ bool UMaterial::SetTextureParameterValueEditorOnly(FName ParameterName, class UT
 	return SetParameterValueEditorOnly(ParameterName, Meta);
 }
 
+bool UMaterial::SetTextureCollectionParameterValueEditorOnly(FName ParameterName, class UTextureCollection* InValue)
+{
+	FMaterialParameterMetadata Meta;
+	Meta.Value = InValue;
+	return SetParameterValueEditorOnly(ParameterName, Meta);
+}
+
 bool UMaterial::SetRuntimeVirtualTextureParameterValueEditorOnly(FName ParameterName, class URuntimeVirtualTexture* InValue)
 {
 	FMaterialParameterMetadata Meta;
diff --git a/Engine/Source/Runtime/Engine/Private/Materials/MaterialInstanceConstant.cpp b/Engine/Source/Runtime/Engine/Private/Materials/MaterialInstanceConstant.cpp
--- a/Engine/Source/Runtime/Engine/Private/Materials/MaterialInstanceConstant.cpp
+++ b/Engine/Source/Runtime/Engine/Private/Materials/MaterialInstanceConstant.cpp
@@ -53,6 +53,13 @@ UTexture* UMaterialInstanceConstant::K2_GetTextureParameterValue(FName Parameter
 	return Result;
 }
 
+UTextureCollection* UMaterialInstanceConstant::K2_GetTextureCollectionParameterValue(FName ParameterName)
+{
+	UTextureCollection* Result = NULL;
+	Super::GetTextureCollectionParameterValue(ParameterName, Result);
+	return Result;
+}
+
 UPhysicalMaterialMask* UMaterialInstanceConstant::GetPhysicalMaterialMask() const
 {
 	return PhysMaterialMask;
@@ -126,6 +133,12 @@ void UMaterialInstanceConstant::SetTextureParameterValueEditorOnly(const FMateri
 	SetTextureParameterValueInternal(ParameterInfo,Value);
 }
 
+void UMaterialInstanceConstant::SetTextureCollectionParameterValueEditorOnly(const FMaterialParameterInfo& ParameterInfo, UTextureCollection* Value)
+{
+	check(GIsEditor || IsRunningCommandlet());
+	SetTextureCollectionParameterValueInternal(ParameterInfo,Value);
+}
+
 void UMaterialInstanceConstant::SetRuntimeVirtualTextureParameterValueEditorOnly(const FMaterialParameterInfo& ParameterInfo, URuntimeVirtualTexture* Value)
 {
 	check(GIsEditor || IsRunningCommandlet());
diff --git a/Engine/Source/Runtime/Engine/Public/Materials/Material.h b/Engine/Source/Runtime/Engine/Public/Materials/Material.h
--- a/Engine/Source/Runtime/Engine/Public/Materials/Material.h
+++ b/Engine/Source/Runtime/Engine/Public/Materials/Material.h
@@ -1372,6 +1372,7 @@ public:
 	ENGINE_API bool SetVectorParameterValueEditorOnly(FName ParameterName, FLinearColor InValue);
 	ENGINE_API bool SetScalarParameterValueEditorOnly(FName ParameterName, float InValue);
 	ENGINE_API bool SetTextureParameterValueEditorOnly(FName ParameterName, class UTexture* InValue);
+	ENGINE_API bool SetTextureCollectionParameterValueEditorOnly(FName ParameterName, class UTextureCollection* InValue);
 	ENGINE_API bool SetRuntimeVirtualTextureParameterValueEditorOnly(FName ParameterName, class URuntimeVirtualTexture* InValue);
 	ENGINE_API bool SetSparseVolumeTextureParameterValueEditorOnly(FName ParameterName, class USparseVolumeTexture* InValue);
 	ENGINE_API bool SetFontParameterValueEditorOnly(FName ParameterName, class UFont* InFontValue, int32 InFontPage);
diff --git a/Engine/Source/Runtime/Engine/Public/Materials/MaterialInstanceConstant.h b/Engine/Source/Runtime/Engine/Public/Materials/MaterialInstanceConstant.h
--- a/Engine/Source/Runtime/Engine/Public/Materials/MaterialInstanceConstant.h
+++ b/Engine/Source/Runtime/Engine/Public/Materials/MaterialInstanceConstant.h
@@ -57,6 +57,10 @@ class UMaterialInstanceConstant : public UMaterialInstance
 	UFUNCTION(BlueprintCallable, meta=(DisplayName = "Get Texture Parameter Value", ScriptName = "GetTextureParameterValue"), Category="Rendering|Material")
 	class UTexture* K2_GetTextureParameterValue(FName ParameterName);
 
+	/** Get the MIC texture collection parameter value */
+	UFUNCTION(BlueprintCallable, meta=(DisplayName = "Get Texture Collection Parameter Value", ScriptName = "GetTextureCollectionParameterValue"), Category="Rendering|Material")
+	class UTextureCollection* K2_GetTextureCollectionParameterValue(FName ParameterName);
+
 	/** Get the MIC vector parameter value */
 	UFUNCTION(BlueprintCallable, meta=(DisplayName = "Get Vector Parameter Value", ScriptName = "GetVectorParameterValue", Keywords = "GetColorParameterValue"), Category="Rendering|Material")
 	FLinearColor K2_GetVectorParameterValue(FName ParameterName);
@@ -92,6 +96,7 @@ class UMaterialInstanceConstant : public UMaterialInstance
 	ENGINE_API void SetScalarParameterValueEditorOnly(const FMaterialParameterInfo& ParameterInfo, float Value);
 	ENGINE_API void SetScalarParameterAtlasEditorOnly(const FMaterialParameterInfo& ParameterInfo, FScalarParameterAtlasInstanceData AtlasData);
 	ENGINE_API void SetTextureParameterValueEditorOnly(const FMaterialParameterInfo& ParameterInfo, class UTexture* Value);
+	ENGINE_API void SetTextureCollectionParameterValueEditorOnly(const FMaterialParameterInfo& ParameterInfo, class UTextureCollection* Value);
 	ENGINE_API void SetRuntimeVirtualTextureParameterValueEditorOnly(const FMaterialParameterInfo& ParameterInfo, class URuntimeVirtualTexture* Value);
 	ENGINE_API void SetSparseVolumeTextureParameterValueEditorOnly(const FMaterialParameterInfo& ParameterInfo, class USparseVolumeTexture* Value);
 	ENGINE_API void SetFontParameterValueEditorOnly(const FMaterialParameterInfo& ParameterInfo, class UFont* FontValue, int32 FontPage);
