[MaterialEditor] Adding a null-check for PreviewOutput in UMaterialFunction::GetPreviewMaterial() to prevent a crash when creating a new Material Function asset.

[RN] minor creator/ui

diff --git a/Engine/Source/Runtime/Engine/Private/Materials/MaterialExpressions.cpp b/Engine/Source/Runtime/Engine/Private/Materials/MaterialExpressions.cpp
--- a/Engine/Source/Runtime/Engine/Private/Materials/MaterialExpressions.cpp
+++ b/Engine/Source/Runtime/Engine/Private/Materials/MaterialExpressions.cpp
@@ -13546,7 +13546,10 @@ UMaterialInterface* UMaterialFunction::GetPreviewMaterial()
 		}
 
 		// Set the chosen output as preview.
-		PreviewOutput->ConnectToPreviewMaterial(PreviewMaterial, 0);
+		if (PreviewOutput != nullptr)
+		{
+			PreviewOutput->ConnectToPreviewMaterial(PreviewMaterial, 0);
+		}
 
 		PreviewMaterial->MaterialDomain = PreviewMaterialDomain;
 
