[AutoRTFM] Fix mixed open / closed writes to FUObjectItem::SerialNumber

The memory validator correctly identified this as a mixed open / closed write when reconstructing UObjects.

diff --git a/Engine/Source/Runtime/CoreUObject/Private/UObject/UObjectArray.cpp b/Engine/Source/Runtime/CoreUObject/Private/UObject/UObjectArray.cpp
--- a/Engine/Source/Runtime/CoreUObject/Private/UObject/UObjectArray.cpp
+++ b/Engine/Source/Runtime/CoreUObject/Private/UObject/UObjectArray.cpp
@@ -502,10 +502,14 @@ int32 FUObjectArray::AllocateSerialNumber(int32 Index)
 	if (!SerialNumber)
 	{
 		// Open around PrimarySerialNumber. If we fail/abort a transaction, we don't need to undo this; we simply allow it to grow for the next use.
-		UE_AUTORTFM_OPEN
+		// Disable validation for this open as we're performing an explicitly
+		// recorded write to SerialNumber which the memory validator will treat
+		// as a false-positive mixed open / closed write.
+		UE_AUTORTFM_OPEN_NO_VALIDATION
 		{
 			SerialNumber = PrimarySerialNumber.Increment();
 			UE_CLOG(SerialNumber <= START_SERIAL_NUMBER, LogUObjectArray, Fatal, TEXT("UObject serial numbers overflowed (trying to allocate serial number %d)."), SerialNumber);
+			AutoRTFM::RecordOpenWrite(const_cast<int32*>(SerialNumberPtr)); // const_cast to remove volatile
 			int32 ValueWas = FPlatformAtomics::InterlockedCompareExchange((int32*)SerialNumberPtr, SerialNumber, 0);
 			if (ValueWas != 0)
 			{
diff --git a/Engine/Source/Runtime/Engine/Private/Tests/AutoRTFM/AutoRTFMActorComponentTests.cpp b/Engine/Source/Runtime/Engine/Private/Tests/AutoRTFM/AutoRTFMActorComponentTests.cpp
--- a/Engine/Source/Runtime/Engine/Private/Tests/AutoRTFM/AutoRTFMActorComponentTests.cpp
+++ b/Engine/Source/Runtime/Engine/Private/Tests/AutoRTFM/AutoRTFMActorComponentTests.cpp
@@ -24,6 +24,7 @@
 #include "UObject/Package.h"
 #include "UObject/UObjectArray.h"
 #include "UObject/UObjectGlobals.h"
+#include "UObject/WeakObjectPtrTemplates.h"
 
 #if WITH_DEV_AUTOMATION_TESTS
 
@@ -513,19 +514,56 @@ AUTORTFM_ACTOR_COMPONENT_TEST(ReconstructActor)
 	TestTrueExpr(0 == Old->ActorCategory);
 	Old->ActorCategory = 123;
 
-	AAutoRTFMTestActor* New = nullptr;
-
+	TWeakObjectPtr<AAutoRTFMTestActor> OldWeak = Old;
+	TestTrueExpr(OldWeak == Old);
+	
+	const uint32 OldID = Old->GetUniqueID();
+	
 	AutoRTFM::Testing::Abort([&]
 		{
 			// Reconstruct the object by using the same name, then immediately abort.
-			New = NewObject<AAutoRTFMTestActor>(Outer, Name);
+			AAutoRTFMTestActor* New = NewObject<AAutoRTFMTestActor>(Outer, Name);
+			TestTrueExpr(Old == New);
+			TestTrueExpr(OldWeak == New);
+			TestTrueExpr(OldID == New->GetUniqueID());
+			TestTrueExpr(0 == New->ActorCategory);
 			AutoRTFM::AbortTransaction();
 		});
 
-	TestTrueExpr(nullptr == New);
+	TestTrueExpr(OldWeak == Old);
+	TestTrueExpr(OldID == Old->GetUniqueID());
 	TestTrueExpr(123 == Old->ActorCategory);
 }
 
+// Reproduction for mixed open / closed writes on FUObjectItem::SerialNumber.
+AUTORTFM_ACTOR_COMPONENT_TEST(ReconstructActorThenObtainWeakPtr) // SOL-7678
+{
+	AutoRTFM::Testing::Abort([&]
+	{
+		// Construct a UObject.
+		// The FUObjectItem::SerialNumber begins with 0, assigned in the open
+		// with validation disabled. This is open-write is slightly iffy, but
+		// can be considered part of the UObject construction / reconstruction
+		// logic, which has a lot of "special leniency".
+		UAutoRTFMTestPrimitiveComponent* Original = NewObject<UAutoRTFMTestPrimitiveComponent>(Actor, "ReconstructedComponent");
+		// Reconstruct the object.
+		// This calls FUObjectArray::ResetSerialNumber() which assigns 0 to
+		// FUObjectItem::SerialNumber in the closed.
+		UAutoRTFMTestPrimitiveComponent* Reconstructed = NewObject<UAutoRTFMTestPrimitiveComponent>(Actor, "ReconstructedComponent");
+		TestTrueExpr(Original == Reconstructed);
+		// Obtain a weak pointer to the object.
+		// This calls into FUObjectArray::AllocateSerialNumber() which assigns a
+		// new number to FUObjectItem::SerialNumber in the open.
+		// The validator would catch the closed write followed by the open
+		// write. The applied fix for this is to disable validation on this
+		// open, and to explicitly record the write. To the memory validator
+		// this behaves like a regular closed write, despite the logic using
+		// atomic CAS which is not permitted in the closed.
+		FWeakObjectPtr WeakReconstructed = Reconstructed;
+		AutoRTFM::AbortTransaction();
+	});
+}
+
 AUTORTFM_ACTOR_COMPONENT_TEST(ReconstructCameraShake) // SOL-7529
 {
 	FName Name = TEXT("ShakyMcShakeface");
