[DebugDraw] small header cleanup

diff --git a/Engine/Plugins/Animation/GameplayInsights/Source/RewindDebuggerVLog/Private/RewindDebuggerVLog.h b/Engine/Plugins/Animation/GameplayInsights/Source/RewindDebuggerVLog/Private/RewindDebuggerVLog.h
--- a/Engine/Plugins/Animation/GameplayInsights/Source/RewindDebuggerVLog/Private/RewindDebuggerVLog.h
+++ b/Engine/Plugins/Animation/GameplayInsights/Source/RewindDebuggerVLog/Private/RewindDebuggerVLog.h
@@ -1,13 +1,13 @@
 // Copyright Epic Games, Inc. All Rights Reserved.
 #pragma once
 
-#include "CoreMinimal.h"
 #include "IRewindDebugger.h"
 #include "IRewindDebuggerExtension.h"
-#include "UObject/WeakObjectPtr.h"
 #include "VLogRenderingActor.h"
 #include "ToolMenu.h"
 
+class UFont;
+
 // Rewind debugger extension for Visual Logger support
 
 class FRewindDebuggerVLog : public IRewindDebuggerExtension
diff --git a/Engine/Source/Developer/LogVisualizer/Private/VisualLoggerRenderingActorBase.cpp b/Engine/Source/Developer/LogVisualizer/Private/VisualLoggerRenderingActorBase.cpp
--- a/Engine/Source/Developer/LogVisualizer/Private/VisualLoggerRenderingActorBase.cpp
+++ b/Engine/Source/Developer/LogVisualizer/Private/VisualLoggerRenderingActorBase.cpp
@@ -1,7 +1,9 @@
 // Copyright Epic Games, Inc. All Rights Reserved.
 
 #include "VisualLoggerRenderingActorBase.h"
+#include "DynamicMeshBuilder.h"
 #include "LogVisualizerSettings.h"
+#include "Materials/Material.h"
 #if WITH_EDITOR
 #include "GeomTools.h"
 #endif // WITH_EDITOR
diff --git a/Engine/Source/Runtime/Engine/Private/DebugRenderSceneProxy.cpp b/Engine/Source/Runtime/Engine/Private/DebugRenderSceneProxy.cpp
--- a/Engine/Source/Runtime/Engine/Private/DebugRenderSceneProxy.cpp
+++ b/Engine/Source/Runtime/Engine/Private/DebugRenderSceneProxy.cpp
@@ -27,8 +27,19 @@ static TAutoConsoleVariable<float> CVarDebugRenderOverrideFarClippingPlane(
 	TEXT("Allows to override the far clipping plane for debug shapes and text (in cm), only effective if > 0.0 and if r.DebugRender.AllowFrustumCulling is enabled"),
 	ECVF_RenderThreadSafe);
 
-// FPrimitiveSceneProxy interface.
+// FDebugRenderSceneProxy::FMesh
+FDebugRenderSceneProxy::FMesh::FMesh() = default;
+FDebugRenderSceneProxy::FMesh::~FMesh() = default;
+
+FDebugRenderSceneProxy::FMesh::FMesh(const FMesh& Other)
+	: Vertices(Other.Vertices)
+	, Indices(Other.Indices)
+	, Box(Other.Box)
+	, Color(Other.Color)
+{
+}
 
+// FPrimitiveSceneProxy interface
 FDebugRenderSceneProxy::FDebugRenderSceneProxy(const UPrimitiveComponent* InComponent)
 	: FPrimitiveSceneProxy(InComponent)
 	, ViewFlagIndex(uint32(FEngineShowFlags::FindIndexByName(TEXT("Game"))))
diff --git a/Engine/Source/Runtime/Engine/Public/DebugRenderSceneProxy.h b/Engine/Source/Runtime/Engine/Public/DebugRenderSceneProxy.h
--- a/Engine/Source/Runtime/Engine/Public/DebugRenderSceneProxy.h
+++ b/Engine/Source/Runtime/Engine/Public/DebugRenderSceneProxy.h
@@ -6,12 +6,17 @@
 
 #pragma once
 
+#if UE_ENABLE_INCLUDE_ORDER_DEPRECATED_IN_5_6
 #include "CoreMinimal.h"
-#include "PrimitiveSceneProxy.h"
 #include "Materials/Material.h"
 #include "DynamicMeshBuilder.h"
+#endif // UE_ENABLE_INCLUDE_ORDER_DEPRECATED_IN_5_6
+
+#include "PrimitiveSceneProxy.h"
 #include "SceneView.h"
 
+struct FDynamicMeshVertex;
+struct FConvexVolume;
 class APlayerController;
 class FMaterialRenderProxy;
 class FMeshElementCollector;
@@ -328,10 +333,14 @@ public:
 
 	struct FMesh
 	{
+		ENGINE_API FMesh();
+		ENGINE_API ~FMesh();
+		ENGINE_API FMesh(const FMesh& Other);
+
 		TArray<FDynamicMeshVertex> Vertices;
 		TArray <uint32> Indices;
 		FBox Box = FBox(ForceInit);
-		FColor Color;
+		FColor Color = FColor::White;
 	};
 
 	struct FCapsule
diff --git a/Engine/Source/Runtime/Experimental/Dataflow/Engine/Private/Dataflow/DataflowDebugDraw.cpp b/Engine/Source/Runtime/Experimental/Dataflow/Engine/Private/Dataflow/DataflowDebugDraw.cpp
--- a/Engine/Source/Runtime/Experimental/Dataflow/Engine/Private/Dataflow/DataflowDebugDraw.cpp
+++ b/Engine/Source/Runtime/Experimental/Dataflow/Engine/Private/Dataflow/DataflowDebugDraw.cpp
@@ -3,6 +3,7 @@
 #include "Dataflow/DataflowDebugDraw.h"
 #include "Dataflow/DataflowDebugDrawComponent.h"
 #include "DynamicMesh/DynamicMesh3.h"
+#include "DynamicMeshBuilder.h"
 
 FDataflowDebugDraw::FDataflowDebugDraw(FDataflowDebugRenderSceneProxy* DebugRenderSceneProxy) :
 	DebugRenderSceneProxy(DebugRenderSceneProxy)
diff --git a/Engine/Source/Runtime/NavigationSystem/Private/NavMesh/NavTestRenderingComponent.cpp b/Engine/Source/Runtime/NavigationSystem/Private/NavMesh/NavTestRenderingComponent.cpp
--- a/Engine/Source/Runtime/NavigationSystem/Private/NavMesh/NavTestRenderingComponent.cpp
+++ b/Engine/Source/Runtime/NavigationSystem/Private/NavMesh/NavTestRenderingComponent.cpp
@@ -5,12 +5,14 @@
 #include "NavMesh/RecastNavMesh.h"
 
 #if UE_ENABLE_DEBUG_DRAWING
+#include "Debug/DebugDrawService.h"
+#include "DynamicMeshBuilder.h"
 #include "Engine/Canvas.h"
 #include "Engine/Engine.h"
-#include "PrimitiveDrawingUtils.h"
+#include "Materials/Material.h"
 #include "Materials/MaterialRenderProxy.h"
 #include "MeshElementCollector.h"
-#include "Debug/DebugDrawService.h"
+#include "PrimitiveDrawingUtils.h"
 
 #include UE_INLINE_GENERATED_CPP_BY_NAME(NavTestRenderingComponent)
 
diff --git a/Engine/Source/Runtime/NavigationSystem/Private/NavMesh/RecastNavMeshGenerator.cpp b/Engine/Source/Runtime/NavigationSystem/Private/NavMesh/RecastNavMeshGenerator.cpp
--- a/Engine/Source/Runtime/NavigationSystem/Private/NavMesh/RecastNavMeshGenerator.cpp
+++ b/Engine/Source/Runtime/NavigationSystem/Private/NavMesh/RecastNavMeshGenerator.cpp
@@ -1,6 +1,8 @@
 // Copyright Epic Games, Inc. All Rights Reserved.
 
 #include "NavMesh/RecastNavMeshGenerator.h"
+
+#include "DynamicMeshBuilder.h"
 #include "AI/Navigation/NavRelevantInterface.h"
 #include "Compression/OodleDataCompression.h"
 #include "Engine/Level.h"
