Refactor of experimental substrate material layering macro over to a per project cvar.

diff --git a/Engine/Source/Editor/DetailCustomizations/Private/MaterialAttributePropertyDetails.cpp b/Engine/Source/Editor/DetailCustomizations/Private/MaterialAttributePropertyDetails.cpp
--- a/Engine/Source/Editor/DetailCustomizations/Private/MaterialAttributePropertyDetails.cpp
+++ b/Engine/Source/Editor/DetailCustomizations/Private/MaterialAttributePropertyDetails.cpp
@@ -43,11 +43,7 @@ void FMaterialAttributePropertyDetails::CustomizeDetails(IDetailLayoutBuilder& D
 		// We do not want to blend Substrate BSDF which could increase the complexity of the material (BSDF / Slab count) for each blend operation.
 		// Only parameters are allowed to finally be transform into Substrate at the end of the chain.
 		// This filtering is ran only when constructing the UI element.
-#if ENABLE_MATERIAL_LAYER_PROTOTYPE
-		if(Substrate::IsSubstrateEnabled() || NameGUIDPair.Key != ExcludedFrontMaterialName)
-#else
-		if (NameGUIDPair.Key != ExcludedFrontMaterialName)
-#endif
+		if(Substrate::IsMaterialLayeringSupportEnabled() || NameGUIDPair.Key != ExcludedFrontMaterialName)
 		{
 			AttributeDisplayNameList.Add(MakeShareable(new FString(NameGUIDPair.Key)));
 		}
diff --git a/Engine/Source/Editor/MaterialEditor/Private/MaterialEditor.cpp b/Engine/Source/Editor/MaterialEditor/Private/MaterialEditor.cpp
--- a/Engine/Source/Editor/MaterialEditor/Private/MaterialEditor.cpp
+++ b/Engine/Source/Editor/MaterialEditor/Private/MaterialEditor.cpp
@@ -752,8 +752,7 @@ void FMaterialEditor::InitMaterialEditor( const EToolkitMode::Type Mode, const T
 				// We can check the usage here and add the appropriate inputs too (e.g. Layer==1MA, Blend==2MA)
 				if (MaterialFunction->GetMaterialFunctionUsage() == EMaterialFunctionUsage::MaterialLayer)
 				{
-#if ENABLE_MATERIAL_LAYER_PROTOTYPE
-					if (Substrate::IsSubstrateEnabled())
+					if (Substrate::IsMaterialLayeringSupportEnabled())
 					{
 						UMaterialExpression* Input = CreateNewMaterialExpression(UMaterialExpressionFunctionInput::StaticClass(), FVector2D(-800, 300), false, true);
 						if (Input)
@@ -780,7 +779,6 @@ void FMaterialEditor::InitMaterialEditor( const EToolkitMode::Type Mode, const T
 
 					}
 					else
-#endif //ENABLE_MATERIAL_LAYER_PROTOTYPE
 					{
 						UMaterialExpression* Input = CreateNewMaterialExpression(UMaterialExpressionFunctionInput::StaticClass(), FVector2D(-350, 300), false, true);
 						if (Input)
@@ -805,8 +803,7 @@ void FMaterialEditor::InitMaterialEditor( const EToolkitMode::Type Mode, const T
 				}
 				else if (MaterialFunction->GetMaterialFunctionUsage() == EMaterialFunctionUsage::MaterialLayerBlend)
 				{
-#if ENABLE_MATERIAL_LAYER_PROTOTYPE
-					if (Substrate::IsSubstrateEnabled())
+					if (Substrate::IsMaterialLayeringSupportEnabled())
 					{
 						{
 							UMaterialExpression* InputBottom = CreateNewMaterialExpression(UMaterialExpressionFunctionInput::StaticClass(), FVector2D(-1100, 200), false, true);
@@ -831,10 +828,10 @@ void FMaterialEditor::InitMaterialEditor( const EToolkitMode::Type Mode, const T
 
 							if(InputTop && InputBottom && GetDefault<UEditorExperimentalSettings>()->bExampleLayersAndBlends)
 							{
-								static TObjectPtr<UMaterialFunction> DefaultBlendFunction = FindObject<UMaterialFunction>(nullptr, DEFAULT_MATERIALLAYERBLEND_PATH);
+								static TObjectPtr<UMaterialFunction> DefaultBlendFunction = FindObject<UMaterialFunction>(nullptr, DEFAULT_SUBSTRATE_MATERIALLAYERBLEND_PATH);
 								if (!DefaultBlendFunction)
 								{
-									DefaultBlendFunction = LoadObject<UMaterialFunction>(nullptr, DEFAULT_MATERIALLAYERBLEND_PATH);
+									DefaultBlendFunction = LoadObject<UMaterialFunction>(nullptr, DEFAULT_SUBSTRATE_MATERIALLAYERBLEND_PATH);
 								}
 							
 								if (DefaultBlendFunction)
@@ -911,7 +908,6 @@ void FMaterialEditor::InitMaterialEditor( const EToolkitMode::Type Mode, const T
 						}
 					}
 					else
-#endif //ENABLE_MATERIAL_LAYER_PROTOTYPE
 					{
 						// "Top layer" should be below "bottom layer" on the graph, to align with B on blend nodes
 						UMaterialExpression* InputTop = CreateNewMaterialExpression(UMaterialExpressionFunctionInput::StaticClass(), FVector2D(-300, 400), false, true);
@@ -1029,12 +1025,11 @@ void FMaterialEditor::UpdateGenerator()
 		TArray<UObject*> Objects;
 		Objects.Add(MaterialEditorInstance);
 		Generator->SetObjects(Objects);
-#if ENABLE_MATERIAL_LAYER_PROTOTYPE
+
 		if (MaterialLayersFunctionsInstance.IsValid())
 		{
 		 	MaterialLayersFunctionsInstance->SetEditorInstance(MaterialEditorInstance);
 		}
-#endif
 	}
 }
 
@@ -1797,8 +1792,7 @@ void FMaterialEditor::CreateInternalWidgets()
 		.InGenerator(Generator)
 		.InMaterialDetailsView(MaterialDetailsView);
 
-	
-	if (GBEnableMaterialLayerPrototype)
+	if(Substrate::IsMaterialLayeringSupportEnabled())
 	{
 		MaterialLayersFunctionsInstance = SNew(SMaterialLayersFunctionsInstanceWrapper)
 			.InMaterialEditorInstance(MaterialEditorInstance)
@@ -1923,12 +1917,10 @@ void FMaterialEditor::OnChangeBreadCrumbGraph(UEdGraph* InGraph)
 void FMaterialEditor::GeneratorRowsRefreshed()
 {
 	MaterialParametersOverviewWidget->Refresh();
-#if !ENABLE_MATERIAL_LAYER_PROTOTYPE
-	if (MaterialLayersFunctionsInstance)
+	if (!Substrate::IsMaterialLayeringSupportEnabled() && MaterialLayersFunctionsInstance)
 	{
 		MaterialLayersFunctionsInstance->Refresh();
 	}
-#endif
 }
 
 FName FMaterialEditor::GetToolkitFName() const
@@ -3157,11 +3149,7 @@ void FMaterialEditor::UpdateMaterialinfoList_Old()
 			int32 NumOutputs = 0;
 			// For Material Layers
 
-			if (MaterialFunction->GetMaterialFunctionUsage() == EMaterialFunctionUsage::MaterialLayer
-#if ENABLE_MATERIAL_LAYER_PROTOTYPE
-			&& !Substrate::IsSubstrateEnabled()
-#endif
-			)
+			if (!Substrate::IsMaterialLayeringSupportEnabled() && MaterialFunction->GetMaterialFunctionUsage() == EMaterialFunctionUsage::MaterialLayer)
 			{
 				// Material layers must have a single MA input and output only
 				for (UMaterialExpression* Expression : MaterialFunction->GetExpressions())
@@ -3197,11 +3185,7 @@ void FMaterialEditor::UpdateMaterialinfoList_Old()
 					FailingExpression.Add(nullptr);
 				}
 			}
-			else if (MaterialFunction->GetMaterialFunctionUsage() == EMaterialFunctionUsage::MaterialLayerBlend
-#if ENABLE_MATERIAL_LAYER_PROTOTYPE
-			&& !Substrate::IsSubstrateEnabled()
-#endif
-			)
+			else if (!Substrate::IsMaterialLayeringSupportEnabled() && MaterialFunction->GetMaterialFunctionUsage() == EMaterialFunctionUsage::MaterialLayerBlend)
 			{
 				// Material layer blends can have two MA inputs and single MA output only
 				for (UMaterialExpression* Expression : MaterialFunction->GetExpressions())
diff --git a/Engine/Source/Editor/MaterialEditor/Private/MaterialEditorDetailCustomization.cpp b/Engine/Source/Editor/MaterialEditor/Private/MaterialEditorDetailCustomization.cpp
--- a/Engine/Source/Editor/MaterialEditor/Private/MaterialEditorDetailCustomization.cpp
+++ b/Engine/Source/Editor/MaterialEditor/Private/MaterialEditorDetailCustomization.cpp
@@ -998,9 +998,11 @@ void FMaterialExpressionLayersParameterDetails::CustomizeDetails(IDetailLayoutBu
 	// for expression parameters all their properties are in one category based on their class name.
 
 	FName LayerCategory = FName("Layers");
-#if ENABLE_MATERIAL_LAYER_PROTOTYPE
-	DetailLayout.HideCategory(LayerCategory);
-#endif
+	if(Substrate::IsMaterialLayeringSupportEnabled())
+	{
+		DetailLayout.HideCategory(LayerCategory);
+	}
+
 	IDetailCategoryBuilder& Category = DetailLayout.EditCategory(LayerCategory);
 	// Get a handle to the property we are about to edit
 	GroupPropertyHandle = DetailLayout.GetProperty("DefaultLayers");
diff --git a/Engine/Source/Editor/MaterialEditor/Private/MaterialEditorUtilities.cpp b/Engine/Source/Editor/MaterialEditor/Private/MaterialEditorUtilities.cpp
--- a/Engine/Source/Editor/MaterialEditor/Private/MaterialEditorUtilities.cpp
+++ b/Engine/Source/Editor/MaterialEditor/Private/MaterialEditorUtilities.cpp
@@ -569,11 +569,13 @@ void FMaterialEditorUtilities::GetVisibleMaterialParametersFromExpression(
 			{
 				TArray<TObjectPtr<UMaterialExpressionMaterialFunctionCall>>& LayerCallers = LayersExpression->LayerCallers;
 				TArray<TObjectPtr<UMaterialExpressionMaterialFunctionCall>>& BlendCallers = LayersExpression->BlendCallers;
-#if ENABLE_MATERIAL_LAYER_PROTOTYPE
-				TSharedPtr< FMaterialLayersFunctionsRuntimeGraphCache > LayerRuntimeTree = LayersValue.RuntimeGraphCache;
-				LayerCallers = LayerRuntimeTree->LayerCallers;
-				BlendCallers = LayerRuntimeTree->BlendCallers;
-#endif
+				
+				if(Substrate::IsMaterialLayeringSupportEnabled())
+				{
+					TSharedPtr< FMaterialLayersFunctionsRuntimeGraphCache > LayerRuntimeTree = LayersValue.RuntimeGraphCache;
+					LayerCallers = LayerRuntimeTree->LayerCallers;
+					BlendCallers = LayerRuntimeTree->BlendCallers;
+				}
 
 				for (auto& Layer : LayerCallers)
 				{
diff --git a/Engine/Source/Editor/MaterialEditor/Private/MaterialInstanceEditor.cpp b/Engine/Source/Editor/MaterialEditor/Private/MaterialInstanceEditor.cpp
--- a/Engine/Source/Editor/MaterialEditor/Private/MaterialInstanceEditor.cpp
+++ b/Engine/Source/Editor/MaterialEditor/Private/MaterialInstanceEditor.cpp
@@ -250,13 +250,9 @@ void FMaterialInstanceEditor::RegisterTabSpawners(const TSharedRef<class FTabMan
 
 	if (!bIsFunctionPreviewMaterial)
 	{
+		FText DisplayName = Substrate::IsMaterialLayeringSupportEnabled() ? LOCTEXT("MaterialLayersTab", "Material Layers") : LOCTEXT("LayerPropertiesTab", "Layer Parameters");
 		InTabManager->RegisterTabSpawner(FMaterialInstanceEditorTabs::LayerPropertiesTabId, FOnSpawnTab::CreateSP(this, &FMaterialInstanceEditor::SpawnTab_LayerProperties))
-#if ENABLE_MATERIAL_LAYER_PROTOTYPE		
-			.SetDisplayName(LOCTEXT("MaterialLayersTab", "Material Layers"))
-#else
-			.SetDisplayName(LOCTEXT("LayerPropertiesTab", "Layer Parameters"))
-#endif
-		
+			.SetDisplayName(DisplayName)
 			.SetGroup(WorkspaceMenuCategoryRef)
 			.SetIcon(FSlateIcon(FAppStyle::GetAppStyleSetName(), "LevelEditor.Tabs.Layers"));
 	}
@@ -266,17 +262,19 @@ void FMaterialInstanceEditor::RegisterTabSpawners(const TSharedRef<class FTabMan
 		.SetGroup(WorkspaceMenuCategoryRef)
 		.SetIcon(FSlateIcon(FAppStyle::GetAppStyleSetName(), "LevelEditor.Tabs.Details"));
 
-#if ENABLE_MATERIAL_LAYER_PROTOTYPE
-	InTabManager->RegisterTabSpawner(FMaterialInstanceEditorTabs::AssetBrowserTabId, FOnSpawnTab::CreateSP(this, &FMaterialInstanceEditor::SpawnTab_AssetBrowser))
-		.SetDisplayName(LOCTEXT("AssetBrowserTab", "Asset Browser"))
-		.SetGroup(WorkspaceMenuCategoryRef)
-		.SetIcon(FSlateIcon(FAppStyle::GetAppStyleSetName(), "LevelEditor.Tabs.Details"));
+	if(Substrate::IsMaterialLayeringSupportEnabled())
+	{
+		InTabManager->RegisterTabSpawner(FMaterialInstanceEditorTabs::AssetBrowserTabId, FOnSpawnTab::CreateSP(this, &FMaterialInstanceEditor::SpawnTab_AssetBrowser))
+			.SetDisplayName(LOCTEXT("AssetBrowserTab", "Asset Browser"))
+			.SetGroup(WorkspaceMenuCategoryRef)
+			.SetIcon(FSlateIcon(FAppStyle::GetAppStyleSetName(), "LevelEditor.Tabs.Details"));
+
+		InTabManager->RegisterTabSpawner(FMaterialInstanceEditorTabs::SubstrateTabId, FOnSpawnTab::CreateSP(this, &FMaterialInstanceEditor::SpawnTab_Substrate))
+			.SetDisplayName(LOCTEXT("SubstrateTab", "Substrate"))
+			.SetGroup(WorkspaceMenuCategoryRef)
+			.SetIcon(FSlateIcon(FAppStyle::GetAppStyleSetName(), "Kismet.Tabs.Palette"));	// SUBSTRATE_TODO a Substrate icon
+	}
 
-	InTabManager->RegisterTabSpawner(FMaterialInstanceEditorTabs::SubstrateTabId, FOnSpawnTab::CreateSP(this, &FMaterialInstanceEditor::SpawnTab_Substrate))
-		.SetDisplayName(LOCTEXT("SubstrateTab", "Substrate"))
-		.SetGroup(WorkspaceMenuCategoryRef)
-		.SetIcon(FSlateIcon(FAppStyle::GetAppStyleSetName(), "Kismet.Tabs.Palette"));	// SUBSTRATE_TODO a Substrate icon
-#endif
 	MaterialStatsManager->RegisterTabs();
 
 	OnRegisterTabSpawners().Broadcast(InTabManager);
@@ -293,11 +291,12 @@ void FMaterialInstanceEditor::UnregisterTabSpawners(const TSharedRef<class FTabM
 		InTabManager->UnregisterTabSpawner(FMaterialInstanceEditorTabs::LayerPropertiesTabId);
 	}
 	InTabManager->UnregisterTabSpawner( FMaterialInstanceEditorTabs::PreviewSettingsTabId );
-#if ENABLE_MATERIAL_LAYER_PROTOTYPE
-	InTabManager->UnregisterTabSpawner( FMaterialInstanceEditorTabs::AssetBrowserTabId );
+	if(Substrate::IsMaterialLayeringSupportEnabled())
+	{
+		InTabManager->UnregisterTabSpawner( FMaterialInstanceEditorTabs::AssetBrowserTabId );
+		InTabManager->UnregisterTabSpawner(FMaterialInstanceEditorTabs::SubstrateTabId);
+	}
 
-	InTabManager->UnregisterTabSpawner(FMaterialInstanceEditorTabs::SubstrateTabId);
-#endif
 	MaterialStatsManager->UnregisterTabs();
 
 	OnUnregisterTabSpawners().Broadcast(InTabManager);
@@ -475,70 +474,73 @@ void FMaterialInstanceEditor::InitMaterialInstanceEditor( const EToolkitMode::Ty
 	if (!bIsFunctionPreviewMaterial)
 	{
 		
-#if ENABLE_MATERIAL_LAYER_PROTOTYPE
-		StandaloneDefaultLayout = FTabManager::NewLayout("Standalone_MaterialInstanceEditor_Layout_v9")
-			->AddArea
-			(
-				FTabManager::NewPrimaryArea()->SetOrientation(Orient_Vertical)
-				->Split
+		if(Substrate::IsMaterialLayeringSupportEnabled())
+		{
+			StandaloneDefaultLayout = FTabManager::NewLayout("Standalone_MaterialInstanceEditor_Layout_v9")
+				->AddArea
 				(
-					FTabManager::NewSplitter()->SetOrientation(Orient_Horizontal)->SetSizeCoefficient(0.9f)
+					FTabManager::NewPrimaryArea()->SetOrientation(Orient_Vertical)
 					->Split
 					(
-						FTabManager::NewStack()->SetSizeCoefficient(0.70f)->SetHideTabWell(true)
-						->AddTab(FMaterialInstanceEditorTabs::PreviewTabId, ETabState::OpenedTab)
-						->AddTab(FMaterialInstanceEditorTabs::PreviewSettingsTabId, ETabState::ClosedTab)
-					)
-					->Split
-					(
-						FTabManager::NewSplitter()->SetOrientation(Orient_Vertical)->SetSizeCoefficient(0.7f)
+						FTabManager::NewSplitter()->SetOrientation(Orient_Horizontal)->SetSizeCoefficient(0.9f)
 						->Split
 						(
-					FTabManager::NewSplitter()->SetOrientation(Orient_Horizontal)->SetSizeCoefficient(0.3f)
+							FTabManager::NewStack()->SetSizeCoefficient(0.70f)->SetHideTabWell(true)
+							->AddTab(FMaterialInstanceEditorTabs::PreviewTabId, ETabState::OpenedTab)
+							->AddTab(FMaterialInstanceEditorTabs::PreviewSettingsTabId, ETabState::ClosedTab)
+						)
+						->Split
+						(
+							FTabManager::NewSplitter()->SetOrientation(Orient_Vertical)->SetSizeCoefficient(0.7f)
 							->Split
 							(
-								FTabManager::NewStack()/*->SetSizeCoefficient(0.30f)*/
-								->AddTab(FMaterialInstanceEditorTabs::LayerPropertiesTabId, ETabState::OpenedTab)
+						FTabManager::NewSplitter()->SetOrientation(Orient_Horizontal)->SetSizeCoefficient(0.3f)
+								->Split
+								(
+									FTabManager::NewStack()/*->SetSizeCoefficient(0.30f)*/
+									->AddTab(FMaterialInstanceEditorTabs::LayerPropertiesTabId, ETabState::OpenedTab)
+								)
+								->Split
+								(
+									FTabManager::NewStack()/*->SetSizeCoefficient(0.30f)*/
+									->AddTab(FMaterialInstanceEditorTabs::AssetBrowserTabId, ETabState::OpenedTab)
+								)
 							)
 							->Split
 							(
-								FTabManager::NewStack()/*->SetSizeCoefficient(0.30f)*/
-								->AddTab(FMaterialInstanceEditorTabs::AssetBrowserTabId, ETabState::OpenedTab)
+							FTabManager::NewStack()->SetSizeCoefficient(0.30f)
+									->AddTab(FMaterialInstanceEditorTabs::PropertiesTabId, ETabState::OpenedTab)
 							)
 						)
-						->Split
-						(
-						FTabManager::NewStack()->SetSizeCoefficient(0.30f)
-								->AddTab(FMaterialInstanceEditorTabs::PropertiesTabId, ETabState::OpenedTab)
-						)
 					)
-				)
-			);
-#else
+				);
+		}
+		else
+		{
 			StandaloneDefaultLayout = FTabManager::NewLayout("Standalone_MaterialInstanceEditor_Layout_v8")
-			->AddArea
-			(
-				FTabManager::NewPrimaryArea()->SetOrientation(Orient_Vertical)
-				->Split
+				->AddArea
 				(
-					FTabManager::NewSplitter()->SetOrientation(Orient_Horizontal)->SetSizeCoefficient(0.9f)
-					->Split
-					(
-						FTabManager::NewStack()->SetSizeCoefficient(0.70f)->SetHideTabWell(true)
-						->AddTab(FMaterialInstanceEditorTabs::PreviewTabId, ETabState::OpenedTab)
-						->AddTab(FMaterialInstanceEditorTabs::PreviewSettingsTabId, ETabState::ClosedTab)
-					)
+					FTabManager::NewPrimaryArea()->SetOrientation(Orient_Vertical)
 					->Split
 					(
-						FTabManager::NewStack()->SetSizeCoefficient(0.30f)
-						->AddTab(FMaterialInstanceEditorTabs::PropertiesTabId, ETabState::OpenedTab)
-						->AddTab(FMaterialInstanceEditorTabs::LayerPropertiesTabId, ETabState::OpenedTab)
-						->SetForegroundTab(FMaterialInstanceEditorTabs::PropertiesTabId)
+						FTabManager::NewSplitter()->SetOrientation(Orient_Horizontal)->SetSizeCoefficient(0.9f)
+						->Split
+						(
+							FTabManager::NewStack()->SetSizeCoefficient(0.70f)->SetHideTabWell(true)
+							->AddTab(FMaterialInstanceEditorTabs::PreviewTabId, ETabState::OpenedTab)
+							->AddTab(FMaterialInstanceEditorTabs::PreviewSettingsTabId, ETabState::ClosedTab)
+						)
+						->Split
+						(
+							FTabManager::NewStack()->SetSizeCoefficient(0.30f)
+							->AddTab(FMaterialInstanceEditorTabs::PropertiesTabId, ETabState::OpenedTab)
+							->AddTab(FMaterialInstanceEditorTabs::LayerPropertiesTabId, ETabState::OpenedTab)
+							->SetForegroundTab(FMaterialInstanceEditorTabs::PropertiesTabId)
+						)
 					)
-				)
-			);
-#endif
+				);
 		}
+	}
 	
 	const bool bCreateDefaultStandaloneMenu = true;
 	const bool bCreateDefaultToolbar = true;
@@ -1110,22 +1112,32 @@ TSharedRef<SDockTab> FMaterialInstanceEditor::SpawnTab_LayerProperties(const FSp
 {
 	check(Args.GetTabId().TabType == FMaterialInstanceEditorTabs::LayerPropertiesTabId);
 
-	TSharedRef<SDockTab> SpawnedTab = SNew(SDockTab)
+	if(Substrate::IsMaterialLayeringSupportEnabled())
+	{
+		TSharedRef<SDockTab> SpawnedTab = SNew(SDockTab)
 		.Label(LOCTEXT("MaterialLayerPropertiesTitle", "Layer Parameters"))
 		[
-#if ENABLE_MATERIAL_LAYER_PROTOTYPE
 			MaterialLayersFunctionsInstance.ToSharedRef()
-#else
-			SNew(SBorder)
-			.Padding(4)
+		];
+
+		AddToSpawnedToolPanels(Args.GetTabId().TabType, SpawnedTab);
+		return SpawnedTab;
+	}
+	else
+	{
+		TSharedRef<SDockTab> SpawnedTab = SNew(SDockTab)
+		.Label(LOCTEXT("MaterialLayerPropertiesTitle", "Layer Parameters"))
 		[
-			MaterialLayersFunctionsInstance.ToSharedRef()
-		]
-#endif
+			SNew(SBorder)
+				.Padding(4)
+				[
+					MaterialLayersFunctionsInstance.ToSharedRef()
+				]
 		];
 
-	AddToSpawnedToolPanels(Args.GetTabId().TabType, SpawnedTab);
-	return SpawnedTab;
+		AddToSpawnedToolPanels(Args.GetTabId().TabType, SpawnedTab);
+		return SpawnedTab;
+	}
 }
 
 TSharedRef<SDockTab> FMaterialInstanceEditor::SpawnTab_PreviewSettings(const FSpawnTabArgs& Args)
@@ -1807,10 +1819,10 @@ void FMaterialInstanceEditor::Refresh()
 
 	RefreshOnScreenMessages();
 
-#if ENABLE_MATERIAL_LAYER_PROTOTYPE
-	if (SubstrateWidget)
+	if (Substrate::IsMaterialLayeringSupportEnabled() && SubstrateWidget)
+	{
 		SubstrateWidget->UpdateFromMaterial();
-#endif
+	}
 }
 
 void FMaterialInstanceEditor::PostUndo( bool bSuccess )
diff --git a/Engine/Source/Editor/MaterialEditor/Private/MaterialLayersFunctionsCustomization.cpp b/Engine/Source/Editor/MaterialEditor/Private/MaterialLayersFunctionsCustomization.cpp
--- a/Engine/Source/Editor/MaterialEditor/Private/MaterialLayersFunctionsCustomization.cpp
+++ b/Engine/Source/Editor/MaterialEditor/Private/MaterialLayersFunctionsCustomization.cpp
@@ -27,6 +27,7 @@
 #include "Framework/SlateDelegates.h"
 #include "Misc/Attribute.h"
 #include "Styling/SlateTypes.h"
+#include "RenderUtils.h"
 #include "MaterialPropertyHelpers.h"
 
 
@@ -108,11 +109,11 @@ void FMaterialLayersFunctionsCustomization::RebuildChildren()
 
 void FMaterialLayersFunctionsCustomization::GenerateHeaderRowContent(FDetailWidgetRow& NodeRow)
 {
-
-#if ENABLE_MATERIAL_LAYER_PROTOTYPE
-	// early out if the new prototype is enabled as we only want to edit FMaterialLayersFunctions through the new Substrate tree UI
-	return;
-#endif
+	if(Substrate::IsMaterialLayeringSupportEnabled())
+	{
+		// early out if Substrate layering is enabled as we only want to edit FMaterialLayersFunctions through the new Substrate tree UI
+		return;
+	}
 	
 	TAttribute<bool>::FGetter IsResetButtonEnabledDelegate = TAttribute<bool>::FGetter::CreateSP(this, &FMaterialLayersFunctionsCustomization::IsResetToDefaultVisible);
 	TAttribute<bool> IsEnabledAttribute = TAttribute<bool>::Create(IsResetButtonEnabledDelegate);
@@ -150,10 +151,11 @@ void FMaterialLayersFunctionsCustomization::GenerateHeaderRowContent(FDetailWidg
 
 void FMaterialLayersFunctionsCustomization::GenerateChildContent(IDetailChildrenBuilder& ChildrenBuilder)
 {
-#if ENABLE_MATERIAL_LAYER_PROTOTYPE
-	// early out if the new prototype is enabled as we only want to edit FMaterialLayersFunctions through the new Substrate tree UI
-	return;
-#endif
+	if(Substrate::IsMaterialLayeringSupportEnabled())
+	{
+		// early out if the new prototype is enabled as we only want to edit FMaterialLayersFunctions through the new Substrate tree UI
+		return;
+	}
 	
 	uint32 LayerChildren;
 	LayerHandle->GetNumChildren(LayerChildren);
diff --git a/Engine/Source/Editor/MaterialEditor/Private/MaterialPropertyHelpers.cpp b/Engine/Source/Editor/MaterialEditor/Private/MaterialPropertyHelpers.cpp
--- a/Engine/Source/Editor/MaterialEditor/Private/MaterialPropertyHelpers.cpp
+++ b/Engine/Source/Editor/MaterialEditor/Private/MaterialPropertyHelpers.cpp
@@ -136,21 +136,18 @@ void FMaterialPropertyHelpers::OnMaterialLayerAssetChanged(const struct FAssetDa
 			{
 				InMaterialFunction->Layers[Index] = LayerFunction;
 				InMaterialFunction->UnlinkLayerFromParent(Index);
-#if ENABLE_MATERIAL_LAYER_PROTOTYPE
-				// update name of the layer
-				InMaterialFunction->EditorOnly.LayerNames[Index] = FText::FromName(InAssetData.AssetName);
-#endif
+				if(Substrate::IsMaterialLayeringSupportEnabled())
+				{
+					// update name of the layer
+					InMaterialFunction->EditorOnly.LayerNames[Index] = FText::FromName(InAssetData.AssetName);
+				}
 			}
 			break;
 		case EMaterialParameterAssociation::BlendParameter:
 			if (InMaterialFunction->Blends[Index] != LayerFunction)
 			{
 				InMaterialFunction->Blends[Index] = LayerFunction;
-#if ENABLE_MATERIAL_LAYER_PROTOTYPE
-				InMaterialFunction->UnlinkLayerFromParent(Index);
-#else
-				InMaterialFunction->UnlinkLayerFromParent(Index + 1); // Blend indices are offset by 1, no blend for base layer
-#endif
+				InMaterialFunction->UnlinkLayerFromParent(Substrate::IsMaterialLayeringSupportEnabled() ? Index : Index + 1); // Blend indices are offset by 1 in legacy, no blend for base layer
 			}
 			break;
 		default:
diff --git a/Engine/Source/Editor/MaterialEditor/Private/SMaterialEditorTopologyWidget.cpp b/Engine/Source/Editor/MaterialEditor/Private/SMaterialEditorTopologyWidget.cpp
--- a/Engine/Source/Editor/MaterialEditor/Private/SMaterialEditorTopologyWidget.cpp
+++ b/Engine/Source/Editor/MaterialEditor/Private/SMaterialEditorTopologyWidget.cpp
@@ -252,17 +252,20 @@ static const TSharedRef<SWidget> InternalProcessOperatorAsThumbnails(
 				int32 ExpressionIndex = INDEX_NONE;
 				const float ThumbnailSize = 40.0f;
 				UObject* ThumbnailObject = nullptr;
-#if ENABLE_MATERIAL_LAYER_PROTOTYPE
-				const TSharedPtr<FMaterialLayersFunctionsRuntimeGraphCache> LayerTreeCache = LayersFunctions->RuntimeGraphCache;
-				if (LayerTreeCache)
-				{
-					ExpressionIndex = LayerTreeCache->FindExpressionIndex(Op.MaterialExpressionGuid, LayersFunctions);
-				}
-				if (ExpressionIndex >= 0)
+
+				if(Substrate::IsMaterialLayeringSupportEnabled())
 				{
-					ThumbnailObject = LayerTreeCache->NodePreviewMaterials[ExpressionIndex];
+					const TSharedPtr<FMaterialLayersFunctionsRuntimeGraphCache> LayerTreeCache = LayersFunctions->RuntimeGraphCache;
+					if (LayerTreeCache)
+					{
+						ExpressionIndex = LayerTreeCache->FindExpressionIndex(Op.MaterialExpressionGuid, LayersFunctions);
+					}
+					if (ExpressionIndex >= 0)
+					{
+						ThumbnailObject = LayerTreeCache->NodePreviewMaterials[ExpressionIndex];
+					}
 				}
-#endif
+
 				const TSharedPtr<FAssetThumbnail> AssetThumbnail = MakeShareable(new FAssetThumbnail(ThumbnailObject, ThumbnailSize, ThumbnailSize, UThumbnailManager::Get().GetSharedThumbnailPool()));
 				TSharedRef<SWidget> ThumbnailWidget = SNew(SBorder)
 					.Padding(2.0f)
diff --git a/Engine/Source/Editor/MaterialEditor/Private/SMaterialEditorViewport.cpp b/Engine/Source/Editor/MaterialEditor/Private/SMaterialEditorViewport.cpp
--- a/Engine/Source/Editor/MaterialEditor/Private/SMaterialEditorViewport.cpp
+++ b/Engine/Source/Editor/MaterialEditor/Private/SMaterialEditorViewport.cpp
@@ -1092,9 +1092,8 @@ void SMaterialEditor3DPreviewViewport::PopulateViewportOverlays(TSharedRef<class
 		];
 		// clang-format on
 
-#if ENABLE_MATERIAL_LAYER_PROTOTYPE
 	// only show topology widget for material instances
-	if (Cast<UMaterialInstance>(MaterialEditorPtr.Pin()->GetMaterialInterface()))
+	if (Substrate::IsMaterialLayeringSupportEnabled() && Cast<UMaterialInstance>(MaterialEditorPtr.Pin()->GetMaterialInterface()))
 	{
 		Overlay->AddSlot()
 			// clang-format off
@@ -1106,7 +1105,7 @@ void SMaterialEditor3DPreviewViewport::PopulateViewportOverlays(TSharedRef<class
 			];
 		// clang-format on
 	}
-#endif
+
 	// clang-format on
 	// add the feature level display widget
 	Overlay->AddSlot()
diff --git a/Engine/Source/Editor/MaterialEditor/Private/SMaterialLayersFunctionsTree.cpp b/Engine/Source/Editor/MaterialEditor/Private/SMaterialLayersFunctionsTree.cpp
--- a/Engine/Source/Editor/MaterialEditor/Private/SMaterialLayersFunctionsTree.cpp
+++ b/Engine/Source/Editor/MaterialEditor/Private/SMaterialLayersFunctionsTree.cpp
@@ -48,6 +48,7 @@
 #include "ThumbnailRendering/ThumbnailManager.h"
 #include "Styling/AppStyle.h"
 #include "Styling/StyleColors.h"
+#include "RenderUtils.h"
 
 #define LOCTEXT_NAMESPACE "MaterialLayerCustomization"
 
@@ -1792,7 +1793,7 @@ void SMaterialLayersFunctionsInstanceWrapper::Refresh()
 	TSharedRef<SWidget> ThumbnailWidget = SNullWidget::NullWidget;
 	TSharedPtr<SBox> ThumbnailBox = nullptr;
 	FText MaterialName = FText::GetEmpty();
-	if (GBEnableMaterialLayerPrototype)
+	if (Substrate::IsMaterialLayeringSupportEnabled())
 	{
 		UObject* ThumbnailObject = MaterialEditorInstance->GetMaterialInterface().Get();
 		const TSharedPtr<FAssetThumbnail> AssetThumbnail = MakeShareable(new FAssetThumbnail(ThumbnailObject, ThumbnailSize, ThumbnailSize, NestedTree->GetTreeThumbnailPool()));
@@ -1808,7 +1809,7 @@ void SMaterialLayersFunctionsInstanceWrapper::Refresh()
 				return NestedTree->RelinkLayersToParent();
 			});
 		
-		if (GBEnableMaterialLayerPrototype)
+		if (Substrate::IsMaterialLayeringSupportEnabled())
 		{
 			HeaderBox = SNew(SHorizontalBox)
 			+ SHorizontalBox::Slot()
@@ -1883,7 +1884,7 @@ void SMaterialLayersFunctionsInstanceWrapper::Refresh()
 		];
 		if (FMaterialPropertyHelpers::IsOverriddenExpression(NestedTree->GetFunctionParameter()))
 		{
-			if (GBEnableMaterialLayerPrototype)
+			if (Substrate::IsMaterialLayeringSupportEnabled())
 			{
 				FMenuBuilder DetailViewOptions( true, nullptr);
 				DetailViewOptions.AddMenuEntry(
@@ -1954,7 +1955,7 @@ void SMaterialLayersFunctionsInstanceWrapper::Refresh()
 			}
 		}
 
-		if (!GBEnableMaterialLayerPrototype)
+		if (!Substrate::IsMaterialLayeringSupportEnabled())
 		{
 			HeaderBox->AddSlot()
 				.FillWidth(1.0f)
@@ -2015,7 +2016,7 @@ void SMaterialLayersFunctionsInstanceWrapper::Refresh()
 
 void SMaterialLayersFunctionsInstanceWrapper::Construct(const FArguments& InArgs)
 {
-	if (GBEnableMaterialLayerPrototype)
+	if (Substrate::IsMaterialLayeringSupportEnabled())
 	{
 		NestedTree = SNew(SMaterialSubstrateTree)
 			.InMaterialEditorInstance(InArgs._InMaterialEditorInstance)
diff --git a/Engine/Source/Editor/MaterialEditor/Private/SMaterialLayersFunctionsTree.h b/Engine/Source/Editor/MaterialEditor/Private/SMaterialLayersFunctionsTree.h
--- a/Engine/Source/Editor/MaterialEditor/Private/SMaterialLayersFunctionsTree.h
+++ b/Engine/Source/Editor/MaterialEditor/Private/SMaterialLayersFunctionsTree.h
@@ -19,12 +19,6 @@
 #include "Widgets/Views/STreeView.h"
 #include "Materials/MaterialExpressionMaterialSample.h"
 
-#if ENABLE_MATERIAL_LAYER_PROTOTYPE
-static constexpr bool GBEnableMaterialLayerPrototype = true;
-#else
-static constexpr bool GBEnableMaterialLayerPrototype = false;
-#endif
-
 class IPropertyHandle;
 class SMaterialLayersFunctionsInstanceTree;
 class UDEditorParameterValue;
diff --git a/Engine/Source/Editor/MaterialEditor/Private/SMaterialSubstrateTree.cpp b/Engine/Source/Editor/MaterialEditor/Private/SMaterialSubstrateTree.cpp
--- a/Engine/Source/Editor/MaterialEditor/Private/SMaterialSubstrateTree.cpp
+++ b/Engine/Source/Editor/MaterialEditor/Private/SMaterialSubstrateTree.cpp
@@ -61,6 +61,7 @@
 #include "Materials/MaterialFunctionMaterialLayer.h"
 #include "Materials/MaterialFunctionMaterialLayerBlend.h"
 #include "MaterialShared.h"
+#include "RenderUtils.h"
 
 #define LOCTEXT_NAMESPACE "MaterialSubstrateTree"
 
@@ -1323,16 +1324,17 @@ void SMaterialSubstrateTree::RecursiveCreateWidgets(FRecursiveCreateWidgetsConte
 				UMaterialInterface* MaterialInterface = MaterialEditorInstance->GetMaterialInterface();
 				const FMaterialResource* MaterialResource = MaterialInterface->GetMaterialResource(ERHIFeatureLevel::SM6);
 				const FMaterialLayersFunctions* LayersFunctions = MaterialResource->GetMaterialLayers();
-#if ENABLE_MATERIAL_LAYER_PROTOTYPE
-				TSharedPtr<FMaterialLayersFunctionsRuntimeGraphCache> LayerTreeCache = LayersFunctions->RuntimeGraphCache;
-				if (LayerTreeCache)
+				
+				if(Substrate::IsMaterialLayeringSupportEnabled())
 				{
-					LayerTreeCache->NodePreviewMaterials[InNodeId] = (UMaterial*)EditedMaterial;
+					TSharedPtr<FMaterialLayersFunctionsRuntimeGraphCache> LayerTreeCache = LayersFunctions->RuntimeGraphCache;
+					if (LayerTreeCache)
+					{
+						LayerTreeCache->NodePreviewMaterials[InNodeId] = (UMaterial*)EditedMaterial;
+					}
 				}
-#endif
 			}
-#if ENABLE_MATERIAL_LAYER_PROTOTYPE
-			else 
+			else if(Substrate::IsMaterialLayeringSupportEnabled())
 			{
 				UMaterialInterface* MaterialInterface = MaterialEditorInstance->GetMaterialInterface();
 				const FMaterialResource* MaterialResource = MaterialInterface->GetMaterialResource(ERHIFeatureLevel::SM6);
@@ -1357,7 +1359,6 @@ void SMaterialSubstrateTree::RecursiveCreateWidgets(FRecursiveCreateWidgetsConte
 					}
 				}
 			}
-#endif
 		}
 
 		StackProperty->Children.Add(ChildProperty);
@@ -1511,8 +1512,7 @@ void SMaterialSubstrateTree::CreateGroupsWidget()
 		DeferredSearches.Empty();
 
 		// Create the hierarchy of Sorted items recursivelly following the LayerFunctions Tree
-#if ENABLE_MATERIAL_LAYER_PROTOTYPE
-		if (FunctionParameterHandle)
+		if (Substrate::IsMaterialLayeringSupportEnabled() && FunctionParameterHandle)
 		{
 			TSharedPtr<IPropertyHandle>	LayerHandle = FunctionParameterHandle->GetChildHandle("Layers").ToSharedRef();
 			TSharedPtr<IPropertyHandle> BlendHandle = FunctionParameterHandle->GetChildHandle("Blends").ToSharedRef();
@@ -1549,7 +1549,6 @@ void SMaterialSubstrateTree::CreateGroupsWidget()
 				RecursiveCreateWidgets(&Context, RootChildren[Index], LayerProperties, true);
 			}
 		}
-#endif // ENABLE_MATERIAL_LAYER_PROTOTYPE
 	}
 	
 	SetParentsExpansionState();
diff --git a/Engine/Source/Editor/MaterialEditor/Public/MaterialPropertyHelpers.h b/Engine/Source/Editor/MaterialEditor/Public/MaterialPropertyHelpers.h
--- a/Engine/Source/Editor/MaterialEditor/Public/MaterialPropertyHelpers.h
+++ b/Engine/Source/Editor/MaterialEditor/Public/MaterialPropertyHelpers.h
@@ -10,6 +10,7 @@
 #include "Widgets/Layout/SSplitter.h"
 #include "MaterialEditor/MaterialEditorInstanceConstant.h"
 #include "DragAndDrop/DecoratedDragDropOp.h"
+#include "RenderUtils.h"
 #include "MaterialPropertyHelpers.generated.h"
 
 
@@ -114,9 +115,11 @@ public:
 
 	FLayerDragDropOp(TSharedPtr<IDraggableItem> InOwningStack)
 	{
-#if ENABLE_MATERIAL_LAYER_PROTOTYPE
-		MouseCursor = EMouseCursor::GrabHandClosed;
-#endif
+		if(Substrate::IsMaterialLayeringSupportEnabled())
+		{
+			MouseCursor = EMouseCursor::GrabHandClosed;
+		}
+
 		OwningStack = InOwningStack;
 		DecoratorWidget = SNew(SBorder)
 			.BorderImage(FAppStyle::GetBrush("Graph.ConnectorFeedback.Border"))
diff --git a/Engine/Source/Runtime/Engine/Classes/Engine/RendererSettings.h b/Engine/Source/Runtime/Engine/Classes/Engine/RendererSettings.h
--- a/Engine/Source/Runtime/Engine/Classes/Engine/RendererSettings.h
+++ b/Engine/Source/Runtime/Engine/Classes/Engine/RendererSettings.h
@@ -1228,6 +1228,15 @@ class URendererSettings : public UDeveloperSettings
 		ConfigRestartRequired = true))
 		uint32 SubstrateDebugAdvancedVisualizationShaders : 1;
 
+	/**
+	"Enable substrate layer support."
+	*/
+	UPROPERTY(config, EditAnywhere, Category = Substrate, meta = (
+		ConsoleVariable = "r.Substrate.EnableLayerSupport", DisplayName = "Substrate enable material layer support (Experimental)",
+		ToolTip = "Enable Substrate material layering and UI(Experimental).Note: This support is one way; the legacy layer materials auto-upgrade and can only be reversed manually after assets are re - saved.",
+		ConfigRestartRequired = true))
+		uint32 SubstrateEnableLayerSupport : 1;
+
 	/**
 	"Enable rough diffuse material."
 	*/
diff --git a/Engine/Source/Runtime/Engine/Internal/Materials/MaterialExpressionMaterialSample.cpp b/Engine/Source/Runtime/Engine/Internal/Materials/MaterialExpressionMaterialSample.cpp
--- a/Engine/Source/Runtime/Engine/Internal/Materials/MaterialExpressionMaterialSample.cpp
+++ b/Engine/Source/Runtime/Engine/Internal/Materials/MaterialExpressionMaterialSample.cpp
@@ -10,7 +10,7 @@
 UMaterialExpressionMaterialSample::UMaterialExpressionMaterialSample(const FObjectInitializer& ObjectInitializer)
 	: Super(ObjectInitializer)
 {
-#if ENABLE_MATERIAL_LAYER_PROTOTYPE
+#if ENABLE_MATERIAL_SAMPLE_PROTOTYPE
 	bCollapsed = false;
 	bShowOutputNameOnPin = true;
 	Outputs.Empty();
@@ -22,7 +22,7 @@ UMaterialExpressionMaterialSample::UMaterialExpressionMaterialSample(const FObje
 }
 
 #if WITH_EDITOR 
-#if ENABLE_MATERIAL_LAYER_PROTOTYPE
+#if ENABLE_MATERIAL_SAMPLE_PROTOTYPE
 int32 UMaterialExpressionMaterialSample::DynamicCompile(FMaterialCompiler* Compiler, int32 OutputIndex, bool bCompilePreview)
 {
 	int32 Result = INDEX_NONE;
@@ -98,7 +98,7 @@ FSubstrateOperator* UMaterialExpressionMaterialSample::SubstrateGenerateMaterial
 	}
 	return nullptr;
 }
-#endif //ENABLE_MATERIAL_LAYER_PROTOTYPE
+#endif //ENABLE_MATERIAL_SAMPLE_PROTOTYPE
 
 void UMaterialExpressionMaterialSample::GetCaption(TArray<FString>& OutCaptions) const
 {
diff --git a/Engine/Source/Runtime/Engine/Internal/Materials/MaterialExpressionMaterialSample.h b/Engine/Source/Runtime/Engine/Internal/Materials/MaterialExpressionMaterialSample.h
--- a/Engine/Source/Runtime/Engine/Internal/Materials/MaterialExpressionMaterialSample.h
+++ b/Engine/Source/Runtime/Engine/Internal/Materials/MaterialExpressionMaterialSample.h
@@ -11,6 +11,8 @@
 
 #include "MaterialExpressionMaterialSample.generated.h"
 
+#define ENABLE_MATERIAL_SAMPLE_PROTOTYPE 0
+
 class FMaterialCompiler;
 
 UCLASS(MinimalAPI)
@@ -23,7 +25,7 @@ class UMaterialExpressionMaterialSample: public UMaterialExpression
 
 	//~ Begin UMaterialExpression Interface
 #if WITH_EDITOR 
-#if ENABLE_MATERIAL_LAYER_PROTOTYPE
+#if ENABLE_MATERIAL_SAMPLE_PROTOTYPE
 	ENGINE_API virtual int32 Compile(class FMaterialCompiler* Compiler, int32 OutputIndex) override;
 	ENGINE_API virtual int32 CompilePreview(FMaterialCompiler* Compiler, int32 OutputIndex) override;
 	ENGINE_API virtual uint32 GetOutputType(int32 OutputIndex) override;
diff --git a/Engine/Source/Runtime/Engine/Private/Materials/Material.cpp b/Engine/Source/Runtime/Engine/Private/Materials/Material.cpp
--- a/Engine/Source/Runtime/Engine/Private/Materials/Material.cpp
+++ b/Engine/Source/Runtime/Engine/Private/Materials/Material.cpp
@@ -3615,10 +3615,9 @@ bool UMaterial::ConvertMaterialToSubstrateMaterial(bool bAllowEmptyMaterialUpdat
 	bool bRelinkCustomOutputNodes = false;
 	UMaterialExpressionSubstrateShadingModels* ConvertNode = nullptr;
 	// Connect all the legacy pin into the conversion node
-#if ENABLE_MATERIAL_LAYER_PROTOTYPE
 	//First check if this material contains a layer node connection, as that requires different upgrade behavior to convert to the Substrate material layer system
-	bool bLayersMaterial = false;
-	if (bUseMaterialAttributes && EditorOnly->MaterialAttributes.Expression && !EditorOnly->FrontMaterial.IsConnected())
+	bool bIsSubstrateLayersMaterial = false;
+	if (Substrate::IsMaterialLayeringSupportEnabled() && bUseMaterialAttributes && EditorOnly->MaterialAttributes.Expression && !EditorOnly->FrontMaterial.IsConnected())
 	{
 		UMaterialExpression* FinalNode = nullptr;
 		UMaterialExpressionMaterialAttributeLayers* LayersNode = nullptr;
@@ -3628,7 +3627,7 @@ bool UMaterial::ConvertMaterialToSubstrateMaterial(bool bAllowEmptyMaterialUpdat
 			{
 				LayersNode = Cast<UMaterialExpressionMaterialAttributeLayers>(Expression);
 				FinalNode = EditorOnly->MaterialAttributes.Expression;
-				bLayersMaterial = true;
+				bIsSubstrateLayersMaterial = true;
 				break;
 			}
 		}
@@ -3670,8 +3669,7 @@ bool UMaterial::ConvertMaterialToSubstrateMaterial(bool bAllowEmptyMaterialUpdat
 		}
 	}
 	
-	if(!bLayersMaterial)
-#endif //ENABLE_MATERIAL_LAYER_PROTOTYPE
+	if(!bIsSubstrateLayersMaterial)
 	{
 		const bool bHasAnySubstrateNodes = HasAnyExpressionsInMaterialAndFunctionsOfType<UMaterialExpressionSubstrateBSDF>();
 		if (bUseMaterialAttributes && EditorOnly->MaterialAttributes.Expression && !EditorOnly->FrontMaterial.IsConnected() && !EditorOnly->MaterialAttributes.Expression->IsResultSubstrateMaterial(EditorOnly->MaterialAttributes.OutputIndex) && !bHasAnySubstrateNodes) // M_Rifle cause issues there
diff --git a/Engine/Source/Runtime/Engine/Private/Materials/MaterialExpressions.cpp b/Engine/Source/Runtime/Engine/Private/Materials/MaterialExpressions.cpp
--- a/Engine/Source/Runtime/Engine/Private/Materials/MaterialExpressions.cpp
+++ b/Engine/Source/Runtime/Engine/Private/Materials/MaterialExpressions.cpp
@@ -6047,7 +6047,6 @@ UMaterialExpressionGetMaterialAttributes::UMaterialExpressionGetMaterialAttribut
 }
 
 #if WITH_EDITOR
-#if ENABLE_MATERIAL_LAYER_PROTOTYPE
 int32 UMaterialExpressionGetMaterialAttributes::CreateOrGetOutputAttribute(EMaterialProperty Attribute)
 {
 	int32 OutputIndex = INDEX_NONE;
@@ -6079,7 +6078,6 @@ int32 UMaterialExpressionGetMaterialAttributes::CreateOrGetOutputAttribute(EMate
 	}
 	return OutputIndex;
 }
-#endif //ENABLE_MATERIAL_LAYER_PROTOTYPE
 
 int32 UMaterialExpressionGetMaterialAttributes::Compile(class FMaterialCompiler* Compiler, int32 OutputIndex)
 {
@@ -6452,7 +6450,6 @@ uint32 UMaterialExpressionSetMaterialAttributes::GetInputType(int32 InputIndex)
 	return InputType;
 }
 
-#if ENABLE_MATERIAL_LAYER_PROTOTYPE
 int32 UMaterialExpressionSetMaterialAttributes::CreateOrGetInputAttribute(EMaterialProperty Attribute)
 {
 	int32 InputsIndex = INDEX_NONE;
@@ -6498,7 +6495,6 @@ bool UMaterialExpressionSetMaterialAttributes::ConnectInputAttribute(EMaterialPr
 	}
 	return false;
 }
-#endif // ENABLE_MATERIAL_LAYER_PROTOTYPE
 
 bool UMaterialExpressionSetMaterialAttributes::GetSubstrateMaterialInputIndex(int32 OutputIndex, int32& InputIndex)
 {
@@ -6763,7 +6759,7 @@ void UMaterialExpressionBlendMaterialAttributes::GatherSubstrateMaterialInfo(FSu
 
 FSubstrateOperator* UMaterialExpressionBlendMaterialAttributes::SubstrateGenerateMaterialTopologyTree(class FMaterialCompiler* Compiler, class UMaterialExpression* Parent, int32 OutputIndex)
 {
-	// SUBSTRATE_TODO: this likely no longer work. We would need to stop parsing and always do parameter blending at this stage.
+	// SUBSTRATE_TODO: this likely no longer work (or only works with the new layering system). We would need to stop parsing and always do parameter blending at this stage.
 	const bool bUseParameterBlending = false;
 	FSubstrateOperator& SubstrateOperator = Compiler->SubstrateCompilationRegisterOperator(SUBSTRATE_OPERATOR_HORIZONTAL, Compiler->SubstrateTreeStackGetPathUniqueId(), this, Parent, Compiler->SubstrateTreeStackGetParentPathUniqueId(), bUseParameterBlending);
 
@@ -6791,11 +6787,30 @@ FSubstrateOperator* UMaterialExpressionBlendMaterialAttributes::SubstrateGenerat
 }
 #endif // WITH_EDITOR
 
-#if ENABLE_MATERIAL_LAYER_PROTOTYPE
+
+UMaterialExpressionLegacyBlendMaterialAttributes::UMaterialExpressionLegacyBlendMaterialAttributes(const FObjectInitializer& ObjectInitializer)
+	: Super(ObjectInitializer)
+{
+	if(!Substrate::IsMaterialLayeringSupportEnabled())
+	{
+		Outputs.Reset();
+		bShowInputs = false;
+		bShowOutputs = false;
+		bCollapsed = true;
+	}
+}
+
 #if WITH_EDITOR
 void UMaterialExpressionLegacyBlendMaterialAttributes::GetCaption(TArray<FString>& OutCaptions) const
 {
-	OutCaptions.Add(TEXT("LegacyBlendMaterialAttributes"));
+	if(Substrate::IsMaterialLayeringSupportEnabled())
+	{
+		OutCaptions.Add(TEXT("LegacyBlendMaterialAttributes"));
+	}
+	else
+	{
+		OutCaptions.Add(TEXT("InvalidLegacyBlend_ForSubstrateLayeringOnly"));
+	}
 }
 bool UMaterialExpressionLegacyBlendMaterialAttributes::CanEditChange(const FProperty* InProperty) const
 {
@@ -6879,6 +6894,11 @@ bool GetStaticBoolExpressionOutput(FMaterialCompiler* Compiler, FExpressionInput
 
 int32 UMaterialExpressionLegacyBlendMaterialAttributes::Compile(FMaterialCompiler* Compiler, int32 OutputIndex)
 {
+	if (!Substrate::IsMaterialLayeringSupportEnabled())
+	{
+		return Compiler->Errorf(TEXT("LegacyBlendMaterialAttributes is only usable with Substrate material layering enabled."));
+	}
+
 	VertexAttributeBlendType = EMaterialAttributeBlend::Blend;
 	PixelAttributeBlendType = EMaterialAttributeBlend::Blend;	
 	
@@ -6956,7 +6976,6 @@ int32 UMaterialExpressionLegacyBlendMaterialAttributes::Compile(FMaterialCompile
 	return Compiler->Lerp(ResultA, ResultB, ResultAlpha);
 }
 #endif // WITH_EDITOR
-#endif //ENABLE_MATERIAL_LAYER_PROTOTYPE
 
 //
 //	UMaterialExpressionMaterialAttributeLayers
@@ -7008,305 +7027,308 @@ void UMaterialExpressionMaterialAttributeLayers::PostEditChangeProperty(FPropert
 #if WITH_EDITOR
 void UMaterialExpressionMaterialAttributeLayers::RebuildLayerGraph(bool bReportErrors)
 {
-#if ENABLE_MATERIAL_LAYER_PROTOTYPE
-	const FMaterialLayersFunctions* LayersTree = (ParamLayers ? ParamLayers : &DefaultLayers);
-	const TArray<UMaterialFunctionInterface*>& Layers = GetLayers();
-	const TArray<UMaterialFunctionInterface*>& Blends = GetBlends();
-	const TArray<bool>& LayerStates = GetLayerStates();
-
-	if (!LayersTree->RuntimeGraphCache)
+	if(Substrate::IsMaterialLayeringSupportEnabled())
 	{
-		LayersTree->RuntimeGraphCache = MakeShareable(new FMaterialLayersFunctionsRuntimeGraphCache());
-	}
-	TSharedPtr< FMaterialLayersFunctionsRuntimeGraphCache> LayersTreeCache = LayersTree->RuntimeGraphCache;
+		const FMaterialLayersFunctions* LayersTree = (ParamLayers ? ParamLayers : &DefaultLayers);
+		const TArray<UMaterialFunctionInterface*>& Layers = GetLayers();
+		const TArray<UMaterialFunctionInterface*>& Blends = GetBlends();
+		const TArray<bool>& LayerStates = GetLayerStates();
 
-	// Pre-populate callers, we maintain these transient objects to avoid
-	// heavy UObject recreation as the graphs are frequently rebuilt
-	while (LayersTreeCache->LayerCallers.Num() < Layers.Num())
-	{
-		LayersTreeCache->LayerCallers.Add(NewObject<UMaterialExpressionMaterialFunctionCall>(GetTransientPackage()));
-	}
-	while (LayersTreeCache->BlendCallers.Num() < Blends.Num())
-	{
-		LayersTreeCache->BlendCallers.Add(NewObject<UMaterialExpressionMaterialFunctionCall>(GetTransientPackage()));
-	}
+		if (!LayersTree->RuntimeGraphCache)
+		{
+			LayersTree->RuntimeGraphCache = MakeShareable(new FMaterialLayersFunctionsRuntimeGraphCache());
+		}
+		TSharedPtr< FMaterialLayersFunctionsRuntimeGraphCache> LayersTreeCache = LayersTree->RuntimeGraphCache;
 
-	// Reset graph connectivity
-	bIsLayerGraphBuilt = false;
-	NumActiveLayerCallers = 0;
-	NumActiveBlendCallers = 0;
+		// Pre-populate callers, we maintain these transient objects to avoid
+		// heavy UObject recreation as the graphs are frequently rebuilt
+		while (LayersTreeCache->LayerCallers.Num() < Layers.Num())
+		{
+			LayersTreeCache->LayerCallers.Add(NewObject<UMaterialExpressionMaterialFunctionCall>(GetTransientPackage()));
+		}
+		while (LayersTreeCache->BlendCallers.Num() < Blends.Num())
+		{
+			LayersTreeCache->BlendCallers.Add(NewObject<UMaterialExpressionMaterialFunctionCall>(GetTransientPackage()));
+		}
 
-	if (ValidateLayerConfiguration(nullptr, bReportErrors))
-	{
-		// Initialize layer function callers
-		for (int32 LayerIndex = 0; LayerIndex < Layers.Num(); ++LayerIndex)
+		// Reset graph connectivity
+		bIsLayerGraphBuilt = false;
+		NumActiveLayerCallers = 0;
+		NumActiveBlendCallers = 0;
+
+		if (ValidateLayerConfiguration(nullptr, bReportErrors))
 		{
-			if (Layers[LayerIndex] && LayerStates[LayerIndex])
+			// Initialize layer function callers
+			for (int32 LayerIndex = 0; LayerIndex < Layers.Num(); ++LayerIndex)
 			{
-				LayersTreeCache->LayerCallers[LayerIndex]->MaterialFunction = Layers[LayerIndex];
-				LayersTreeCache->LayerCallers[LayerIndex]->FunctionParameterInfo.Association = EMaterialParameterAssociation::LayerParameter;
-				LayersTreeCache->LayerCallers[LayerIndex]->FunctionParameterInfo.Index = LayerIndex;
+				if (Layers[LayerIndex] && LayerStates[LayerIndex])
+				{
+					LayersTreeCache->LayerCallers[LayerIndex]->MaterialFunction = Layers[LayerIndex];
+					LayersTreeCache->LayerCallers[LayerIndex]->FunctionParameterInfo.Association = EMaterialParameterAssociation::LayerParameter;
+					LayersTreeCache->LayerCallers[LayerIndex]->FunctionParameterInfo.Index = LayerIndex;
 
-				Layers[LayerIndex]->GetInputsAndOutputs(LayersTreeCache->LayerCallers[LayerIndex]->FunctionInputs, LayersTreeCache->LayerCallers[LayerIndex]->FunctionOutputs);
+					Layers[LayerIndex]->GetInputsAndOutputs(LayersTreeCache->LayerCallers[LayerIndex]->FunctionInputs, LayersTreeCache->LayerCallers[LayerIndex]->FunctionOutputs);
 				
-				for (FFunctionExpressionOutput& FunctionOutput : LayersTreeCache->LayerCallers[LayerIndex]->FunctionOutputs)
-				{
-					LayersTreeCache->LayerCallers[LayerIndex]->Outputs.Add(FunctionOutput.Output);
-				}
+					for (FFunctionExpressionOutput& FunctionOutput : LayersTreeCache->LayerCallers[LayerIndex]->FunctionOutputs)
+					{
+						LayersTreeCache->LayerCallers[LayerIndex]->Outputs.Add(FunctionOutput.Output);
+					}
 
-				// Recursively run through internal functions to allow connection of inputs/outputs
-				LayersTreeCache->LayerCallers[LayerIndex]->UpdateFromFunctionResource();
+					// Recursively run through internal functions to allow connection of inputs/outputs
+					LayersTreeCache->LayerCallers[LayerIndex]->UpdateFromFunctionResource();
 
-				++NumActiveLayerCallers;
-			}
-			else
-			{
-				// Empty entries for opaque layers
-				LayersTreeCache->LayerCallers[LayerIndex]->MaterialFunction = nullptr;
+					++NumActiveLayerCallers;
+				}
+				else
+				{
+					// Empty entries for opaque layers
+					LayersTreeCache->LayerCallers[LayerIndex]->MaterialFunction = nullptr;
+				}
 			}
-		}
 
-		for (int32 BlendIndex = 0; BlendIndex < Blends.Num(); ++BlendIndex)
-		{
-			if (Blends[BlendIndex])
+			for (int32 BlendIndex = 0; BlendIndex < Blends.Num(); ++BlendIndex)
 			{
-				LayersTreeCache->BlendCallers[BlendIndex]->MaterialFunction = Blends[BlendIndex];
-				LayersTreeCache->BlendCallers[BlendIndex]->FunctionParameterInfo.Association = EMaterialParameterAssociation::BlendParameter;
-				LayersTreeCache->BlendCallers[BlendIndex]->FunctionParameterInfo.Index = BlendIndex;
-
-				Blends[BlendIndex]->GetInputsAndOutputs(LayersTreeCache->BlendCallers[BlendIndex]->FunctionInputs, LayersTreeCache->BlendCallers[BlendIndex]->FunctionOutputs);
-				for (FFunctionExpressionOutput& FunctionOutput : LayersTreeCache->BlendCallers[BlendIndex]->FunctionOutputs)
+				if (Blends[BlendIndex])
 				{
-					LayersTreeCache->BlendCallers[BlendIndex]->Outputs.Add(FunctionOutput.Output);
-				}
+					LayersTreeCache->BlendCallers[BlendIndex]->MaterialFunction = Blends[BlendIndex];
+					LayersTreeCache->BlendCallers[BlendIndex]->FunctionParameterInfo.Association = EMaterialParameterAssociation::BlendParameter;
+					LayersTreeCache->BlendCallers[BlendIndex]->FunctionParameterInfo.Index = BlendIndex;
 
-				// Recursively run through internal functions to allow connection of inputs/ouputs
-				LayersTreeCache->BlendCallers[BlendIndex]->UpdateFromFunctionResource();
+					Blends[BlendIndex]->GetInputsAndOutputs(LayersTreeCache->BlendCallers[BlendIndex]->FunctionInputs, LayersTreeCache->BlendCallers[BlendIndex]->FunctionOutputs);
+					for (FFunctionExpressionOutput& FunctionOutput : LayersTreeCache->BlendCallers[BlendIndex]->FunctionOutputs)
+					{
+						LayersTreeCache->BlendCallers[BlendIndex]->Outputs.Add(FunctionOutput.Output);
+					}
+
+					// Recursively run through internal functions to allow connection of inputs/ouputs
+					LayersTreeCache->BlendCallers[BlendIndex]->UpdateFromFunctionResource();
 
-				++NumActiveBlendCallers;
+					++NumActiveBlendCallers;
+				}
+				else
+				{
+					// Empty entries for opaque layers
+					LayersTreeCache->BlendCallers[BlendIndex]->MaterialFunction = nullptr;
+				}
 			}
-			else
+
+			// Empty out unused callers
+			for (int32 CallerIndex = Layers.Num(); CallerIndex < LayersTreeCache->LayerCallers.Num(); ++CallerIndex)
 			{
-				// Empty entries for opaque layers
-				LayersTreeCache->BlendCallers[BlendIndex]->MaterialFunction = nullptr;
+				LayersTreeCache->LayerCallers[CallerIndex]->MaterialFunction = nullptr;
 			}
-		}
-
-		// Empty out unused callers
-		for (int32 CallerIndex = Layers.Num(); CallerIndex < LayersTreeCache->LayerCallers.Num(); ++CallerIndex)
-		{
-			LayersTreeCache->LayerCallers[CallerIndex]->MaterialFunction = nullptr;
-		}
 
-		for (int32 CallerIndex = Blends.Num(); CallerIndex < LayersTreeCache->BlendCallers.Num(); ++CallerIndex)
-		{
-			LayersTreeCache->BlendCallers[CallerIndex]->MaterialFunction = nullptr;
-		}
+			for (int32 CallerIndex = Blends.Num(); CallerIndex < LayersTreeCache->BlendCallers.Num(); ++CallerIndex)
+			{
+				LayersTreeCache->BlendCallers[CallerIndex]->MaterialFunction = nullptr;
+			}
 
 
-		LayersTreeCache->NodeMaterialGraphExpressions.Empty(LayersTree->Tree.Nodes.Num());
-		LayersTreeCache->NodeMaterialGraphExpressions.InsertZeroed(0, LayersTree->Tree.Nodes.Num());
-		LayersTreeCache->NodeMaterialGraphExpressionsBlends.Empty(LayersTree->Tree.Nodes.Num());
-		LayersTreeCache->NodeMaterialGraphExpressionsBlends.InsertZeroed(0, LayersTree->Tree.Nodes.Num());
-		LayersTreeCache->NodePreviewMaterials.Empty(LayersTree->Tree.Nodes.Num());
-		LayersTreeCache->NodePreviewMaterials.InsertZeroed(0, LayersTree->Tree.Nodes.Num());
-		struct TreeGraphBuilder
-		{
-			UMaterialExpressionMaterialAttributeLayers* This = nullptr;
-			const FMaterialLayersFunctions* Tree = nullptr;
-			TSharedPtr< FMaterialLayersFunctionsRuntimeGraphCache> TreeCache;
-
-			UMaterialExpressionMaterialFunctionCall* AssembleBranchGraph(int32 InSourceId)
+			LayersTreeCache->NodeMaterialGraphExpressions.Empty(LayersTree->Tree.Nodes.Num());
+			LayersTreeCache->NodeMaterialGraphExpressions.InsertZeroed(0, LayersTree->Tree.Nodes.Num());
+			LayersTreeCache->NodeMaterialGraphExpressionsBlends.Empty(LayersTree->Tree.Nodes.Num());
+			LayersTreeCache->NodeMaterialGraphExpressionsBlends.InsertZeroed(0, LayersTree->Tree.Nodes.Num());
+			LayersTreeCache->NodePreviewMaterials.Empty(LayersTree->Tree.Nodes.Num());
+			LayersTreeCache->NodePreviewMaterials.InsertZeroed(0, LayersTree->Tree.Nodes.Num());
+			struct TreeGraphBuilder
 			{
-				// Fetch the Layer function of this node as the base background layer
-				// Always null for the root node.	
-				UMaterialExpressionMaterialFunctionCall* PreviousLayerInput = nullptr;
-				if (InSourceId != -1)
-				{
-					int32 BackgroundLayerFuncIdx = Tree->GetLayerFuncIndex(InSourceId);
-					PreviousLayerInput = (TreeCache->LayerCallers.IsValidIndex(BackgroundLayerFuncIdx) ? TreeCache->LayerCallers[BackgroundLayerFuncIdx] : nullptr);
-					PreviousLayerInput = (PreviousLayerInput && PreviousLayerInput->MaterialFunction ? PreviousLayerInput : nullptr);
-					if (PreviousLayerInput)
-						TreeCache->ExpressionCollection.AddExpression(PreviousLayerInput);
-				}
+				UMaterialExpressionMaterialAttributeLayers* This = nullptr;
+				const FMaterialLayersFunctions* Tree = nullptr;
+				TSharedPtr< FMaterialLayersFunctionsRuntimeGraphCache> TreeCache;
 
-				// For child of this node, try to assemble the blend/layer tree 
-				auto NodeChildrenIds = Tree->GetNodeChildren(InSourceId);
-				for (auto NodeId : NodeChildrenIds)
+				UMaterialExpressionMaterialFunctionCall* AssembleBranchGraph(int32 InSourceId)
 				{
-					// Assemble the branch of layer / blend under this sub node by calling recursively  into that sub node
-					UMaterialExpressionMaterialFunctionCall* CurrentLayerInput = AssembleBranchGraph(NodeId);
+					// Fetch the Layer function of this node as the base background layer
+					// Always null for the root node.	
+					UMaterialExpressionMaterialFunctionCall* PreviousLayerInput = nullptr;
+					if (InSourceId != -1)
+					{
+						int32 BackgroundLayerFuncIdx = Tree->GetLayerFuncIndex(InSourceId);
+						PreviousLayerInput = (TreeCache->LayerCallers.IsValidIndex(BackgroundLayerFuncIdx) ? TreeCache->LayerCallers[BackgroundLayerFuncIdx] : nullptr);
+						PreviousLayerInput = (PreviousLayerInput && PreviousLayerInput->MaterialFunction ? PreviousLayerInput : nullptr);
+						if (PreviousLayerInput)
+							TreeCache->ExpressionCollection.AddExpression(PreviousLayerInput);
+					}
 
+					// For child of this node, try to assemble the blend/layer tree 
+					auto NodeChildrenIds = Tree->GetNodeChildren(InSourceId);
+					for (auto NodeId : NodeChildrenIds)
+					{
+						// Assemble the branch of layer / blend under this sub node by calling recursively  into that sub node
+						UMaterialExpressionMaterialFunctionCall* CurrentLayerInput = AssembleBranchGraph(NodeId);
 
-					// Fetch the blend function of the node
-					int32 BlendFuncIdx = Tree->GetBlendFuncIndex(NodeId);
-					UMaterialExpressionMaterialFunctionCall* CurrentBlendCaller = (TreeCache->BlendCallers.IsValidIndex(BlendFuncIdx) ? TreeCache->BlendCallers[BlendFuncIdx] : nullptr);
-					CurrentBlendCaller = (CurrentBlendCaller && CurrentBlendCaller->MaterialFunction ? CurrentBlendCaller : nullptr);
 
-					TreeCache->NodeMaterialGraphExpressionsBlends[NodeId] = CurrentBlendCaller;
+						// Fetch the blend function of the node
+						int32 BlendFuncIdx = Tree->GetBlendFuncIndex(NodeId);
+						UMaterialExpressionMaterialFunctionCall* CurrentBlendCaller = (TreeCache->BlendCallers.IsValidIndex(BlendFuncIdx) ? TreeCache->BlendCallers[BlendFuncIdx] : nullptr);
+						CurrentBlendCaller = (CurrentBlendCaller && CurrentBlendCaller->MaterialFunction ? CurrentBlendCaller : nullptr);
+
+						TreeCache->NodeMaterialGraphExpressionsBlends[NodeId] = CurrentBlendCaller;
 					
-					if (CurrentBlendCaller) // Blend is the new current layer if valid and connect the previous layer and current layer is valid
-					{
-						if (PreviousLayerInput)
-							TreeCache->ExpressionCollection.AddExpression(CurrentBlendCaller);
+						if (CurrentBlendCaller) // Blend is the new current layer if valid and connect the previous layer and current layer is valid
+						{
+							if (PreviousLayerInput)
+								TreeCache->ExpressionCollection.AddExpression(CurrentBlendCaller);
 
-						if (PreviousLayerInput)
-							CurrentBlendCaller->FunctionInputs[0].Input.Connect(0, PreviousLayerInput);
-						if (CurrentLayerInput)
-							CurrentBlendCaller->FunctionInputs[1].Input.Connect(0, CurrentLayerInput);
+							if (PreviousLayerInput)
+								CurrentBlendCaller->FunctionInputs[0].Input.Connect(0, PreviousLayerInput);
+							if (CurrentLayerInput)
+								CurrentBlendCaller->FunctionInputs[1].Input.Connect(0, CurrentLayerInput);
 
-						PreviousLayerInput = CurrentBlendCaller;
+							PreviousLayerInput = CurrentBlendCaller;
+						}
+						else if (CurrentLayerInput) // No Blend then the new current layer if valid
+						{
+							PreviousLayerInput = ToRawPtr(CurrentLayerInput);
+						}
+						//else this node's layer and and blend functions aren't valid, ignore and move on to the next sub node
 					}
-					else if (CurrentLayerInput) // No Blend then the new current layer if valid
+
+					if (InSourceId != -1)
 					{
-						PreviousLayerInput = ToRawPtr(CurrentLayerInput);
+						TreeCache->NodeMaterialGraphExpressions[InSourceId] = (PreviousLayerInput);
 					}
-					//else this node's layer and and blend functions aren't valid, ignore and move on to the next sub node
-				}
 
-				if (InSourceId != -1)
-				{
-					TreeCache->NodeMaterialGraphExpressions[InSourceId] = (PreviousLayerInput);
+					return PreviousLayerInput;
 				}
+			} BuildTreeGraph = { this, LayersTree, LayersTreeCache };
 
-				return PreviousLayerInput;
-			}
-		} BuildTreeGraph = { this, LayersTree, LayersTreeCache };
+			OutputCaller = BuildTreeGraph.AssembleBranchGraph(-1);
 
-		OutputCaller = BuildTreeGraph.AssembleBranchGraph(-1);
+			if (OutputCaller && (OutputCaller->MaterialFunction == nullptr))
+			{	
+				OutputCaller = nullptr;
+			}
 
-		if (OutputCaller && (OutputCaller->MaterialFunction == nullptr))
-		{	
-			OutputCaller = nullptr;
+			bIsLayerGraphBuilt = true;
 		}
-
-#else
-
-	const TArray<UMaterialFunctionInterface*>& Layers = GetLayers();
-	const TArray<UMaterialFunctionInterface*>& Blends = GetBlends();
-	const TArray<bool>& LayerStates = GetLayerStates();
-	
-	// Pre-populate callers, we maintain these transient objects to avoid
-	// heavy UObject recreation as the graphs are frequently rebuilt
-	while (LayerCallers.Num() < Layers.Num())
-	{
-		LayerCallers.Add(NewObject<UMaterialExpressionMaterialFunctionCall>(GetTransientPackage()));
 	}
-	while (BlendCallers.Num() < Blends.Num())
+	else
 	{
-		BlendCallers.Add(NewObject<UMaterialExpressionMaterialFunctionCall>(GetTransientPackage()));
-	}
+		const TArray<UMaterialFunctionInterface*>& Layers = GetLayers();
+		const TArray<UMaterialFunctionInterface*>& Blends = GetBlends();
+		const TArray<bool>& LayerStates = GetLayerStates();
+	
+		// Pre-populate callers, we maintain these transient objects to avoid
+		// heavy UObject recreation as the graphs are frequently rebuilt
+		while (LayerCallers.Num() < Layers.Num())
+		{
+			LayerCallers.Add(NewObject<UMaterialExpressionMaterialFunctionCall>(GetTransientPackage()));
+		}
+		while (BlendCallers.Num() < Blends.Num())
+		{
+			BlendCallers.Add(NewObject<UMaterialExpressionMaterialFunctionCall>(GetTransientPackage()));
+		}
 
-	// Reset graph connectivity
-	bIsLayerGraphBuilt = false;
-	NumActiveLayerCallers = 0;
-	NumActiveBlendCallers = 0;
+		// Reset graph connectivity
+		bIsLayerGraphBuilt = false;
+		NumActiveLayerCallers = 0;
+		NumActiveBlendCallers = 0;
 
-	if (ValidateLayerConfiguration(nullptr, bReportErrors))
-	{
+		if (ValidateLayerConfiguration(nullptr, bReportErrors))
+		{
 
 
-		// Initialize layer function callers
-		for (int32 LayerIndex = 0; LayerIndex < Layers.Num(); ++LayerIndex)
-		{
-			if (Layers[LayerIndex] && LayerStates[LayerIndex])
+			// Initialize layer function callers
+			for (int32 LayerIndex = 0; LayerIndex < Layers.Num(); ++LayerIndex)
 			{
-				int32 CallerIndex = NumActiveLayerCallers;
-				LayerCallers[CallerIndex]->MaterialFunction = Layers[LayerIndex];
-				LayerCallers[CallerIndex]->FunctionParameterInfo.Association = EMaterialParameterAssociation::LayerParameter;
-				LayerCallers[CallerIndex]->FunctionParameterInfo.Index = LayerIndex;
-				++NumActiveLayerCallers;
-
-				Layers[LayerIndex]->GetInputsAndOutputs(LayerCallers[CallerIndex]->FunctionInputs, LayerCallers[CallerIndex]->FunctionOutputs);
-				for (FFunctionExpressionOutput& FunctionOutput : LayerCallers[CallerIndex]->FunctionOutputs)
+				if (Layers[LayerIndex] && LayerStates[LayerIndex])
 				{
-					LayerCallers[CallerIndex]->Outputs.Add(FunctionOutput.Output);
-				}
+					int32 CallerIndex = NumActiveLayerCallers;
+					LayerCallers[CallerIndex]->MaterialFunction = Layers[LayerIndex];
+					LayerCallers[CallerIndex]->FunctionParameterInfo.Association = EMaterialParameterAssociation::LayerParameter;
+					LayerCallers[CallerIndex]->FunctionParameterInfo.Index = LayerIndex;
+					++NumActiveLayerCallers;
+
+					Layers[LayerIndex]->GetInputsAndOutputs(LayerCallers[CallerIndex]->FunctionInputs, LayerCallers[CallerIndex]->FunctionOutputs);
+					for (FFunctionExpressionOutput& FunctionOutput : LayerCallers[CallerIndex]->FunctionOutputs)
+					{
+						LayerCallers[CallerIndex]->Outputs.Add(FunctionOutput.Output);
+					}
 
-				// Optional: Single material attributes input, the base input to the stack
-				if (LayerCallers[CallerIndex]->FunctionInputs.Num() > 0)
-				{
-					if (Input.GetTracedInput().Expression)
+					// Optional: Single material attributes input, the base input to the stack
+					if (LayerCallers[CallerIndex]->FunctionInputs.Num() > 0)
 					{
-						LayerCallers[CallerIndex]->FunctionInputs[0].Input = Input;
+						if (Input.GetTracedInput().Expression)
+						{
+							LayerCallers[CallerIndex]->FunctionInputs[0].Input = Input;
+						}
 					}
-				}
 
-				// Recursively run through internal functions to allow connection of inputs/outputs
-				LayerCallers[CallerIndex]->UpdateFromFunctionResource();
+					// Recursively run through internal functions to allow connection of inputs/outputs
+					LayerCallers[CallerIndex]->UpdateFromFunctionResource();
+				}
 			}
-		}
 
-		for (int32 BlendIndex = 0; BlendIndex < Blends.Num(); ++BlendIndex)
-		{
-			const int32 LayerIndex = BlendIndex + 1;
-			if (Layers[LayerIndex] && LayerStates[LayerIndex])
+			for (int32 BlendIndex = 0; BlendIndex < Blends.Num(); ++BlendIndex)
 			{
-				int32 CallerIndex = NumActiveBlendCallers;
-				++NumActiveBlendCallers;
-
-				if (Blends[BlendIndex])
+				const int32 LayerIndex = BlendIndex + 1;
+				if (Layers[LayerIndex] && LayerStates[LayerIndex])
 				{
-					BlendCallers[CallerIndex]->MaterialFunction = Blends[BlendIndex];
-					BlendCallers[CallerIndex]->FunctionParameterInfo.Association = EMaterialParameterAssociation::BlendParameter;
-					BlendCallers[CallerIndex]->FunctionParameterInfo.Index = BlendIndex;
+					int32 CallerIndex = NumActiveBlendCallers;
+					++NumActiveBlendCallers;
 
-					Blends[BlendIndex]->GetInputsAndOutputs(BlendCallers[CallerIndex]->FunctionInputs, BlendCallers[CallerIndex]->FunctionOutputs);
-					for (FFunctionExpressionOutput& FunctionOutput : BlendCallers[CallerIndex]->FunctionOutputs)
+					if (Blends[BlendIndex])
 					{
-						BlendCallers[CallerIndex]->Outputs.Add(FunctionOutput.Output);
-					}
+						BlendCallers[CallerIndex]->MaterialFunction = Blends[BlendIndex];
+						BlendCallers[CallerIndex]->FunctionParameterInfo.Association = EMaterialParameterAssociation::BlendParameter;
+						BlendCallers[CallerIndex]->FunctionParameterInfo.Index = BlendIndex;
 
-					// Recursively run through internal functions to allow connection of inputs/ouputs
-					BlendCallers[CallerIndex]->UpdateFromFunctionResource();
-				}
-				else
-				{
-					// Empty entries for opaque layers
-					BlendCallers[CallerIndex]->MaterialFunction = nullptr;
+						Blends[BlendIndex]->GetInputsAndOutputs(BlendCallers[CallerIndex]->FunctionInputs, BlendCallers[CallerIndex]->FunctionOutputs);
+						for (FFunctionExpressionOutput& FunctionOutput : BlendCallers[CallerIndex]->FunctionOutputs)
+						{
+							BlendCallers[CallerIndex]->Outputs.Add(FunctionOutput.Output);
+						}
+
+						// Recursively run through internal functions to allow connection of inputs/ouputs
+						BlendCallers[CallerIndex]->UpdateFromFunctionResource();
+					}
+					else
+					{
+						// Empty entries for opaque layers
+						BlendCallers[CallerIndex]->MaterialFunction = nullptr;
+					}
 				}
 			}
-		}
 
-		// Empty out unused callers
-		for (int32 CallerIndex = NumActiveLayerCallers; CallerIndex < LayerCallers.Num(); ++CallerIndex)
-		{
-			LayerCallers[CallerIndex]->MaterialFunction = nullptr;
-		}
-
-		for (int32 CallerIndex = NumActiveBlendCallers; CallerIndex < BlendCallers.Num(); ++CallerIndex)
-		{
-			BlendCallers[CallerIndex]->MaterialFunction = nullptr;
-		}
+			// Empty out unused callers
+			for (int32 CallerIndex = NumActiveLayerCallers; CallerIndex < LayerCallers.Num(); ++CallerIndex)
+			{
+				LayerCallers[CallerIndex]->MaterialFunction = nullptr;
+			}
 
-		// Assemble function chain so each layer blends with the previous
-		if (NumActiveLayerCallers >= 2 && NumActiveBlendCallers >= 1)
-		{
-			if (BlendCallers[0]->MaterialFunction)
+			for (int32 CallerIndex = NumActiveBlendCallers; CallerIndex < BlendCallers.Num(); ++CallerIndex)
 			{
-				BlendCallers[0]->FunctionInputs[0].Input.Connect(0, LayerCallers[0]);
-				BlendCallers[0]->FunctionInputs[1].Input.Connect(0, LayerCallers[1]);
+				BlendCallers[CallerIndex]->MaterialFunction = nullptr;
 			}
 
-			for (int32 LayerIndex = 2; LayerIndex < NumActiveLayerCallers; ++LayerIndex)
+			// Assemble function chain so each layer blends with the previous
+			if (NumActiveLayerCallers >= 2 && NumActiveBlendCallers >= 1)
 			{
-				if (BlendCallers[LayerIndex - 1]->MaterialFunction)
+				if (BlendCallers[0]->MaterialFunction)
 				{
-					// Active blend input is previous blend or direct layer if previous is opaque
-					UMaterialExpressionMaterialFunctionCall* BlendInput = BlendCallers[LayerIndex - 2];
-					BlendInput = BlendInput->MaterialFunction ? BlendInput : ToRawPtr(LayerCallers[LayerIndex - 1]);
+					BlendCallers[0]->FunctionInputs[0].Input.Connect(0, LayerCallers[0]);
+					BlendCallers[0]->FunctionInputs[1].Input.Connect(0, LayerCallers[1]);
+				}
 
-					BlendCallers[LayerIndex - 1]->FunctionInputs[0].Input.Connect(0, BlendInput);
-					BlendCallers[LayerIndex - 1]->FunctionInputs[1].Input.Connect(0, LayerCallers[LayerIndex]);
+				for (int32 LayerIndex = 2; LayerIndex < NumActiveLayerCallers; ++LayerIndex)
+				{
+					if (BlendCallers[LayerIndex - 1]->MaterialFunction)
+					{
+						// Active blend input is previous blend or direct layer if previous is opaque
+						UMaterialExpressionMaterialFunctionCall* BlendInput = BlendCallers[LayerIndex - 2];
+						BlendInput = BlendInput->MaterialFunction ? BlendInput : ToRawPtr(LayerCallers[LayerIndex - 1]);
+
+						BlendCallers[LayerIndex - 1]->FunctionInputs[0].Input.Connect(0, BlendInput);
+						BlendCallers[LayerIndex - 1]->FunctionInputs[1].Input.Connect(0, LayerCallers[LayerIndex]);
+					}
 				}
 			}
+			bIsLayerGraphBuilt = true;
 		}
-
-#endif
-
-		bIsLayerGraphBuilt = true;
-	}
-	else if (bReportErrors)
+	}	
+	
+	if (!bIsLayerGraphBuilt && bReportErrors)
 	{
 		UE_LOG(LogMaterial, Warning, TEXT("Failed to build layer graph for %s."), Material ? *(Material->GetName()) : TEXT("Unknown"));
 	}
@@ -7360,8 +7382,7 @@ bool UMaterialExpressionMaterialAttributeLayers::ValidateLayerConfiguration(FMat
 					bIsValid = false;
 				}
 			}
-#if !ENABLE_MATERIAL_LAYER_PROTOTYPE
-			else
+			else if (!Substrate::IsMaterialLayeringSupportEnabled())
 			{
 				TArray<UMaterialExpressionFunctionInput*> InputExpressions;
 				Layer->GetAllExpressionsOfType<UMaterialExpressionFunctionInput>(InputExpressions, false);
@@ -7371,7 +7392,6 @@ bool UMaterialExpressionMaterialAttributeLayers::ValidateLayerConfiguration(FMat
 					bIsValid = false;
 				}
 			}
-#endif
 
 			if (LayerStates[LayerIndex])
 			{
@@ -7403,11 +7423,8 @@ bool UMaterialExpressionMaterialAttributeLayers::ValidateLayerConfiguration(FMat
 			{
 				TArray<UMaterialExpressionFunctionInput*> InputExpressions;
 				Blend->GetAllExpressionsOfType<UMaterialExpressionFunctionInput>(InputExpressions, false);
-#if ENABLE_MATERIAL_LAYER_PROTOTYPE
-				if (InputExpressions.Num() < 2)
-#else
-				if (InputExpressions.Num() != 2)
-#endif
+
+				if (Substrate::IsMaterialLayeringSupportEnabled() ? InputExpressions.Num() < 2 : InputExpressions.Num() != 2)
 				{
 					COMPILER_OR_LOG_ERROR(TEXT("Blend %i, %s, must have two MaterialAttributes inputs."), BlendIndex, *Blend->GetName());
 					bIsValid = false;
@@ -7416,24 +7433,24 @@ bool UMaterialExpressionMaterialAttributeLayers::ValidateLayerConfiguration(FMat
 		}
 		
 		// Null blends signify an opaque layer so count as valid for the sake of graph validation
-#if ENABLE_MATERIAL_LAYER_PROTOTYPE
-		if (Layers[BlendIndex] && LayerStates[BlendIndex])
-#else
-		if (Layers[BlendIndex+1] && LayerStates[BlendIndex+1])
-#endif
+		if (Substrate::IsMaterialLayeringSupportEnabled() ? Layers[BlendIndex] && LayerStates[BlendIndex] : Layers[BlendIndex+1] && LayerStates[BlendIndex+1])
 		{
 			++NumActiveBlends;
 		}
 	}
 
-#if ENABLE_MATERIAL_LAYER_PROTOTYPE
-	bool bValidGraphLayout = (NumActiveLayers == 0 && NumActiveBlends == 0)		// Pass-through
-		|| (NumActiveLayers == NumActiveBlends);			    // Single layer
-#else
-	bool bValidGraphLayout = (NumActiveLayers == 0 && NumActiveBlends == 0)		// Pass-through
-		|| (NumActiveLayers == 1 && NumActiveBlends == 0)						// Single layer
-		|| (NumActiveLayers >= 2 && NumActiveBlends == NumActiveLayers - 1);	// Blend graph
-#endif
+	bool bValidGraphLayout = false;
+	if(Substrate::IsMaterialLayeringSupportEnabled())
+	{
+		bValidGraphLayout = (NumActiveLayers == 0 && NumActiveBlends == 0)		// Pass-through
+			|| (NumActiveLayers == NumActiveBlends);			    // Single layer
+	}
+	else
+	{
+		bValidGraphLayout = (NumActiveLayers == 0 && NumActiveBlends == 0)		// Pass-through
+			|| (NumActiveLayers == 1 && NumActiveBlends == 0)						// Single layer
+			|| (NumActiveLayers >= 2 && NumActiveBlends == NumActiveLayers - 1);	// Blend graph
+	}
 
 	if (!bValidGraphLayout)
 	{
@@ -7538,15 +7555,11 @@ int32 UMaterialExpressionMaterialAttributeLayers::Compile(FMaterialCompiler* Com
 
 	if (ValidateLayerConfiguration(Compiler, true) && bIsLayerGraphBuilt)
 	{
-#if ENABLE_MATERIAL_LAYER_PROTOTYPE
-		if (OutputCaller)
+		if (Substrate::IsMaterialLayeringSupportEnabled() && OutputCaller)
 		{
 			Result = OutputCaller->Compile(Compiler, 0);
 		}
-		else 
-#endif //ENABLE_MATERIAL_LAYER_PROTOTYPE
-
-		if (NumActiveBlendCallers > 0 && BlendCallers[NumActiveBlendCallers-1]->MaterialFunction)
+		else if (NumActiveBlendCallers > 0 && BlendCallers[NumActiveBlendCallers-1]->MaterialFunction)
 		{
 			// Multiple blended layers
 			Result = BlendCallers[NumActiveBlendCallers-1]->Compile(Compiler, 0);
@@ -7586,43 +7599,44 @@ int32 UMaterialExpressionMaterialAttributeLayers::Compile(FMaterialCompiler* Com
 	return Result;
 }
 
-#if ENABLE_MATERIAL_LAYER_PROTOTYPE
 bool UMaterialExpressionMaterialAttributeLayers::IsResultSubstrateMaterial(int32 OutputIndex)
 {
-	return Substrate::IsSubstrateEnabled();
+	return Substrate::IsMaterialLayeringSupportEnabled();
 }
 
 FSubstrateOperator* UMaterialExpressionMaterialAttributeLayers::SubstrateGenerateMaterialTopologyTree(FMaterialCompiler* Compiler, UMaterialExpression* Parent, int32 OutputIndex)
 {
 	FSubstrateOperator* OutOperator = nullptr;
 
-	const FMaterialLayersFunctions* OverrideLayers = Compiler->GetMaterialLayers();
-	OverrideLayerGraph(OverrideLayers);
-
-	//Behaviour mirrors the behaviour of the existing compile function
-	if (ValidateLayerConfiguration(Compiler, true) && bIsLayerGraphBuilt)
+	if(Substrate::IsMaterialLayeringSupportEnabled())
 	{
-		if (OutputCaller)
+		const FMaterialLayersFunctions* OverrideLayers = Compiler->GetMaterialLayers();
+		OverrideLayerGraph(OverrideLayers);
+
+		//Behaviour mirrors the behaviour of the existing compile function
+		if (ValidateLayerConfiguration(Compiler, true) && bIsLayerGraphBuilt)
 		{
-			OutOperator = OutputCaller->SubstrateGenerateMaterialTopologyTree(Compiler, Parent, 0);
-			Compiler->FeedbackMaterialLayersInstancedGraphFromCompilation(OverrideLayers);
-		}
+			if (OutputCaller)
+			{
+				OutOperator = OutputCaller->SubstrateGenerateMaterialTopologyTree(Compiler, Parent, 0);
+				Compiler->FeedbackMaterialLayersInstancedGraphFromCompilation(OverrideLayers);
+			}
 		
-		if (!OutOperator)
-		{
-			//If this is reached, compile a default operator to avoid a crash
-			FSubstrateOperator& DefaultOperator = Compiler->SubstrateCompilationRegisterOperator(SUBSTRATE_OPERATOR_BSDF, Compiler->SubstrateTreeStackGetPathUniqueId(), this, Parent, Compiler->SubstrateTreeStackGetParentPathUniqueId());
-			DefaultOperator.BSDFType = SUBSTRATE_BSDF_TYPE_SLAB;
-			DefaultOperator.ThicknessIndex = Compiler->SubstrateThicknessStackGetThicknessIndex();
-			OutOperator = &DefaultOperator;
+			if (!OutOperator)
+			{
+				//If this is reached, compile a default operator to avoid a crash
+				FSubstrateOperator& DefaultOperator = Compiler->SubstrateCompilationRegisterOperator(SUBSTRATE_OPERATOR_BSDF, Compiler->SubstrateTreeStackGetPathUniqueId(), this, Parent, Compiler->SubstrateTreeStackGetParentPathUniqueId());
+				DefaultOperator.BSDFType = SUBSTRATE_BSDF_TYPE_SLAB;
+				DefaultOperator.ThicknessIndex = Compiler->SubstrateThicknessStackGetThicknessIndex();
+				OutOperator = &DefaultOperator;
+			}
 		}
-	}
 
-	OverrideLayerGraph(nullptr);
+		OverrideLayerGraph(nullptr);
+	}
 
 	return OutOperator;
 }
-#endif //ENABLE_MATERIAL_LAYER_PROTOTYPE
 
 void UMaterialExpressionMaterialAttributeLayers::GetCaption(TArray<FString>& OutCaptions) const
 {
@@ -9080,9 +9094,7 @@ bool UMaterialExpressionFeatureLevelSwitch::IsInputConnectionRequired(int32 Inpu
 
 bool UMaterialExpressionFeatureLevelSwitch::IsResultMaterialAttributes(int32 OutputIndex)
 {
-#if !ENABLE_MATERIAL_LAYER_PROTOTYPE
-	check(OutputIndex == 0);
-#endif
+	check(Substrate::IsMaterialLayeringSupportEnabled() || OutputIndex == 0);
 	for (FExpressionInputIterator It{ this }; It; ++It)
 	{
 		// If there is a loop anywhere in this expression's inputs then we can't risk checking them
@@ -13681,7 +13693,6 @@ void UMaterialFunction::Serialize(FArchive& Ar)
 }
 
 #if WITH_EDITORONLY_DATA
-#if ENABLE_MATERIAL_LAYER_PROTOTYPE
 void UMaterialFunction::GetAllCustomOutputExpressions(TArray<class UMaterialExpressionCustomOutput*>& OutCustomOutputs) const
 {
 	for (UMaterialExpression* Expression : GetExpressions())
@@ -13696,7 +13707,7 @@ void UMaterialFunction::GetAllCustomOutputExpressions(TArray<class UMaterialExpr
 
 void UMaterialFunction::ConvertExpressionsBetweenLegacyAndSubstrate()
 {
-	if (!Substrate::IsSubstrateEnabled() || GetMaterialFunctionUsage() == EMaterialFunctionUsage::Default)
+	if (!Substrate::IsMaterialLayeringSupportEnabled() || GetMaterialFunctionUsage() == EMaterialFunctionUsage::Default)
 	{
 		return;
 	}
@@ -13933,10 +13944,10 @@ void UMaterialFunction::ConvertExpressionsBetweenLegacyAndSubstrate()
 			if (UMaterialExpressionBlendMaterialAttributes* BlendNode = Cast<UMaterialExpressionBlendMaterialAttributes>(InputExpression))
 			{
 				//Load the existing substrate default blend function which will replace the legacy blend node.
-				static TObjectPtr<UMaterialFunction> DefaultBlendFunction = FindObject<UMaterialFunction>(GetTransientPackage(), DEFAULT_MATERIALLAYERBLEND_PATH);
+				static TObjectPtr<UMaterialFunction> DefaultBlendFunction = FindObject<UMaterialFunction>(GetTransientPackage(), DEFAULT_SUBSTRATE_MATERIALLAYERBLEND_PATH);
 				if (!DefaultBlendFunction)
 				{
-					DefaultBlendFunction = LoadObject<UMaterialFunction>(GetTransientPackage(), DEFAULT_MATERIALLAYERBLEND_PATH);
+					DefaultBlendFunction = LoadObject<UMaterialFunction>(GetTransientPackage(), DEFAULT_SUBSTRATE_MATERIALLAYERBLEND_PATH);
 				}
 
 				if(DefaultBlendFunction)
@@ -14128,7 +14139,6 @@ void UMaterialFunction::ConvertExpressionsBetweenLegacyAndSubstrate()
 	MoveNodeInHorizonalAxis(OutputNode, 50);
 	OutputNode->bCollapsed = true;
 }
-#endif //ENABLE_MATERIAL_LAYER_PROTOTYPE
 #endif //WITH_EDITORONLY_DATA
 
 void UMaterialFunction::PostLoad()
@@ -14152,9 +14162,7 @@ void UMaterialFunction::PostLoad()
 		EditorOnly->ExpressionCollection.EditorComments = MoveTemp(FunctionEditorComments_DEPRECATED);
 	}
 	
-#if ENABLE_MATERIAL_LAYER_PROTOTYPE
 	ConvertExpressionsBetweenLegacyAndSubstrate();
-#endif //ENABLE_MATERIAL_LAYER_PROTOTYPE
 
 	if (EditorOnly)
 	{
@@ -14468,8 +14476,7 @@ bool UMaterialFunction::ValidateFunctionUsage(FMaterialCompiler* Compiler, const
 #if WITH_EDITOR
 	if (GetMaterialFunctionUsage() == EMaterialFunctionUsage::MaterialLayer)
 	{
-#if ENABLE_MATERIAL_LAYER_PROTOTYPE
-		if (Substrate::IsSubstrateEnabled())
+		if (Substrate::IsMaterialLayeringSupportEnabled())
 		{
 			// Material layers must have a single MA input and output only
 			for (UMaterialExpression* Expression : GetExpressions())
@@ -14514,7 +14521,6 @@ bool UMaterialFunction::ValidateFunctionUsage(FMaterialCompiler* Compiler, const
 			}
 		}
 		else
-#endif //ENABLE_MATERIAL_LAYER_PROTOTYPE
 		{
 			// Material layers must have a single MA input and output only
 			for (UMaterialExpression* Expression : GetExpressions())
@@ -14553,8 +14559,7 @@ bool UMaterialFunction::ValidateFunctionUsage(FMaterialCompiler* Compiler, const
 	}
 	else if (GetMaterialFunctionUsage() == EMaterialFunctionUsage::MaterialLayerBlend)
 	{
-#if ENABLE_MATERIAL_LAYER_PROTOTYPE
-		if(Substrate::IsSubstrateEnabled())
+		if(Substrate::IsMaterialLayeringSupportEnabled())
 		{
 			// Material layer blends can have up to two MA inputs and single MA output only
 			for (UMaterialExpression* Expression : GetExpressions())
@@ -14598,7 +14603,6 @@ bool UMaterialFunction::ValidateFunctionUsage(FMaterialCompiler* Compiler, const
 			}
 		}
 		else
-#endif //ENABLE_MATERIAL_LAYER_PROTOTYPE
 		{
 			// Material layer blends can have up to two MA inputs and single MA output only
 			for (UMaterialExpression* Expression : GetExpressions())
@@ -16027,47 +16031,41 @@ bool FMaterialLayersFunctionsRuntimeData::SerializeFromMismatchedTag(const FProp
 
 void FMaterialLayersFunctionsRuntimeData::PostSerialize(const FArchive& Ar)
 {
-
-#if ENABLE_MATERIAL_LAYER_PROTOTYPE
-	if (Ar.IsLoading())
-	{	
-		if (Substrate::IsSubstrateEnabled())
+	if (Substrate::IsMaterialLayeringSupportEnabled() && Ar.IsLoading())
+	{
+		// When loading, if a legacy PRE substrate tree layer version:
+		// The tree need to be recreated from the list
+		// of layers and blends assuming this is a flat hierarchy of layers:
+		// For PRE substrate tree layer version:
+		// TODO: We could use:
+		// 	if (Ar.CustomVer(FRenderingObjectVersion::GUID) < FRenderingObjectVersion::AddedMaterialLayersSubstrateSupport)
+		if ((Tree.Nodes.Num() != Layers.Num()) || (Layers.Num() > Blends.Num()))
 		{
-			// When loading, if a legacy PRE substrate tree layer version:
-			// The tree need to be recreated from the list
-			// of layers and blends assuming this is a flat hierarchy of layers:
-			// For PRE substrate tree layer version:
-			// TODO: We could use:
-			// 	if (Ar.CustomVer(FRenderingObjectVersion::GUID) < FRenderingObjectVersion::AddedMaterialLayersSubstrateSupport)
-			if ((Tree.Nodes.Num() != Layers.Num()) || (Layers.Num() > Blends.Num()))
-			{
-				int NumLegacyLayers = Layers.Num();
-				int NumLegacyBlends = Blends.Num();
+			int NumLegacyLayers = Layers.Num();
+			int NumLegacyBlends = Blends.Num();
 
-				Tree.Empty();
+			Tree.Empty();
 
-				// First add an extra empty Blend to match the number of layers
-				while (Blends.Num() < Layers.Num())
-				{
-					Blends.AddDefaulted(); // normally we should go through this only once
-				}
+			// First add an extra empty Blend to match the number of layers
+			while (Blends.Num() < Layers.Num())
+			{
+				Blends.AddDefaulted(); // normally we should go through this only once
+			}
 
-				Layers.AddDefaulted();
-				Blends.AddDefaulted();
+			Layers.AddDefaulted();
+			Blends.AddDefaulted();
 
-				// Add a layer node at root 
-				FLayerNodeId LayerNodeId = Tree.AddNode({ NumLegacyLayers,  NumLegacyLayers }, -1);
+			// Add a layer node at root 
+			FLayerNodeId LayerNodeId = Tree.AddNode({ NumLegacyLayers,  NumLegacyLayers }, -1);
 
-				// Third rebuiild the tree of nodes referencing the layers and blends
-				for (int32 l = 0; l < NumLegacyLayers; ++l)
-				{
-					// And then ONE MORE node to be the attributes first child
-					Tree.AddNode({ l,  l == 0 ? NumLegacyLayers - 1 : l-1}, LayerNodeId);
-				}
+			// Third rebuiild the tree of nodes referencing the layers and blends
+			for (int32 l = 0; l < NumLegacyLayers; ++l)
+			{
+				// And then ONE MORE node to be the attributes first child
+				Tree.AddNode({ l,  l == 0 ? NumLegacyLayers - 1 : l-1}, LayerNodeId);
 			}
 		}
 	}
-#endif
 }
 
 
@@ -16137,10 +16135,10 @@ void FMaterialLayersFunctions::SerializeLegacy(FArchive& Ar)
 
 void FMaterialLayersFunctions::PostSerialize(const FArchive& Ar)
 {
-#if ENABLE_MATERIAL_LAYER_PROTOTYPE
-	FMaterialLayersFunctionsRuntimeData::PostSerialize(Ar);
-#endif
-
+	if(Substrate::IsMaterialLayeringSupportEnabled())
+	{
+		FMaterialLayersFunctionsRuntimeData::PostSerialize(Ar);
+	}
 
 #if WITH_EDITORONLY_DATA
 	if (Ar.IsLoading())
@@ -16174,15 +16172,10 @@ void FMaterialLayersFunctions::PostSerialize(const FArchive& Ar)
 			EditorOnly.DeletedParentLayerGuids = MoveTemp(DeletedParentLayerGuids_DEPRECATED);
 		}
 
-
-#if ENABLE_MATERIAL_LAYER_PROTOTYPE
-		const int32 NumLayers = EditorOnly.LayerStates.Num();
-#else
-		const int32 NumLayers = Layers.Num();
+		const int32 NumLayers = Substrate::IsMaterialLayeringSupportEnabled() ? EditorOnly.LayerStates.Num() : Layers.Num();
 		// In legacy version we use Layers.Num, but in the newer version above we compare against EdtorOnly.LayerStates.Num because
 		// Layers.Num is modified in the RuntimeData legacy conversion and incorrect for the upcoming test.
 		// then further in the call CheckAndRepairPostSerializeEditorOnlyDataForRuntimeData() the editor only fileds are converted to the new tree of layers datastructure.
-#endif // ENABLE_MATERIAL_LAYER_PROTOTYPE
 		if (EditorOnly.LayerGuids.Num() != NumLayers ||
 			EditorOnly.LayerLinkStates.Num() != NumLayers)
 		{
@@ -16207,10 +16200,7 @@ void FMaterialLayersFunctions::PostSerialize(const FArchive& Ar)
 			}
 		}
 
-#if ENABLE_MATERIAL_LAYER_PROTOTYPE
 		CheckAndRepairPostSerializeEditorOnlyDataForRuntimeData(GetRuntime(), EditorOnly);
-#endif // ENABLE_MATERIAL_LAYER_PROTOTYPE
-
 	} // Ar.IsLoading()
 
 #endif // WITH_EDITORONLY_DATA	
@@ -16221,8 +16211,7 @@ void FMaterialLayersFunctions::PostSerialize(const FArchive& Ar)
 
 void FMaterialLayersFunctions::CheckAndRepairPostSerializeEditorOnlyDataForRuntimeData(FMaterialLayersFunctionsRuntimeData& Runtime, FMaterialLayersFunctionsEditorOnlyData& EditorOnly)
 {
-#if ENABLE_MATERIAL_LAYER_PROTOTYPE
-	if (Substrate::IsSubstrateEnabled())
+	if (Substrate::IsMaterialLayeringSupportEnabled())
 	{
 		// This method is called right after unserialization
 		// Use this fact to fix editor only data in case of a legacy version of the LayersFunctions
@@ -16244,35 +16233,34 @@ void FMaterialLayersFunctions::CheckAndRepairPostSerializeEditorOnlyDataForRunti
 			}
 		}
 	}
-#endif
 }
 
 void FMaterialLayersFunctions::AddDefaultBackgroundLayer()
 {
-#if ENABLE_MATERIAL_LAYER_PROTOTYPE
-	// This call is only valid for the very first layer created from the constructor.
-	check(Layers.Num() == 0 && Blends.Num() == 0 && Tree.IsEmpty());
-	AppendLayerNode(-1, -1);
-	FText LayerName = FText(LOCTEXT("Background", "Background"));
-	EditorOnly.LayerNames[0] = LayerName;
-	// Use a consistent Guid for the background layer
-	// Default constructor assigning different guids will break FStructUtils::AttemptToFindUninitializedScriptStructMembers
-	EditorOnly.LayerGuids[0] = (BackgroundGuid);
-
-#else
-
-	// Default to a non-blended "background" layer
-	Layers.AddDefaulted();	
-	EditorOnly.LayerStates.Add(true);
-	FText LayerName = FText(LOCTEXT("Background", "Background"));
-	EditorOnly.LayerNames.Add(LayerName);
-	EditorOnly.RestrictToLayerRelatives.Add(false);
-	// Use a consistent Guid for the background layer
-	// Default constructor assigning different guids will break FStructUtils::AttemptToFindUninitializedScriptStructMembers
-	EditorOnly.LayerGuids.Add(BackgroundGuid);
-	EditorOnly.LayerLinkStates.Add(EMaterialLayerLinkState::NotFromParent);
-
-#endif // ENABLE_MATERIAL_LAYER_PROTOTYPE
+	if(Substrate::IsMaterialLayeringSupportEnabled())
+	{
+		// This call is only valid for the very first layer created from the constructor.
+		check(Layers.Num() == 0 && Blends.Num() == 0 && Tree.IsEmpty());
+		AppendLayerNode(-1, -1);
+		FText LayerName = FText(LOCTEXT("Background", "Background"));
+		EditorOnly.LayerNames[0] = LayerName;
+		// Use a consistent Guid for the background layer
+		// Default constructor assigning different guids will break FStructUtils::AttemptToFindUninitializedScriptStructMembers
+		EditorOnly.LayerGuids[0] = (BackgroundGuid);
+	}
+	else
+	{
+		// Default to a non-blended "background" layer
+		Layers.AddDefaulted();	
+		EditorOnly.LayerStates.Add(true);
+		FText LayerName = FText(LOCTEXT("Background", "Background"));
+		EditorOnly.LayerNames.Add(LayerName);
+		EditorOnly.RestrictToLayerRelatives.Add(false);
+		// Use a consistent Guid for the background layer
+		// Default constructor assigning different guids will break FStructUtils::AttemptToFindUninitializedScriptStructMembers
+		EditorOnly.LayerGuids.Add(BackgroundGuid);
+		EditorOnly.LayerLinkStates.Add(EMaterialLayerLinkState::NotFromParent);
+	}
 }
 
 int32 FMaterialLayersFunctions::AppendBlendedLayer()
@@ -16302,59 +16290,59 @@ int32 FMaterialLayersFunctions::AddLayerCopy(const FMaterialLayersFunctionsRunti
 
 	Layers.Add(Source.Layers[SourceLayerIndex]);
 
-#if ENABLE_MATERIAL_LAYER_PROTOTYPE
-	// Legacy parents pre substrate support
-	if (Source.Blends.Num() < Source.Layers.Num())
+	if(Substrate::IsMaterialLayeringSupportEnabled())
 	{
-		if (LayerIndex > 0)
+		// Legacy parents pre substrate support
+		if (Source.Blends.Num() < Source.Layers.Num())
 		{
-			Blends.Add(Source.Blends[SourceLayerIndex - 1]);
+			if (LayerIndex > 0)
+			{
+				Blends.Add(Source.Blends[SourceLayerIndex - 1]);
+			}
+			else
+			{
+				Blends.AddDefaulted();
+			}
 		}
+		// Source is substrate tree, same number of layersand blends
 		else
 		{
-			Blends.AddDefaulted();
+			Blends.Add(Source.Blends[SourceLayerIndex]);
 		}
 	}
-	// Source is substrate tree, same number of layersand blends
-	else
-	{
-		Blends.Add(Source.Blends[SourceLayerIndex]);
-	}
-#else
-	if (LayerIndex > 0)
+	else if (LayerIndex > 0)
 	{
 		Blends.Add(Source.Blends[SourceLayerIndex - 1]);
 	}
-#endif
 
 	EditorOnly.LayerStates.Add(bVisible);
 	EditorOnly.LayerNames.Add(SourceEditorOnly.LayerNames[SourceLayerIndex]);
 	EditorOnly.RestrictToLayerRelatives.Add(SourceEditorOnly.RestrictToLayerRelatives[SourceLayerIndex]);
 
-#if ENABLE_MATERIAL_LAYER_PROTOTYPE
-	// Legacy parents pre substrate support
-	if (SourceEditorOnly.RestrictToBlendRelatives.Num() < SourceEditorOnly.RestrictToLayerRelatives.Num())
+	if(Substrate::IsMaterialLayeringSupportEnabled())
 	{
-		if (LayerIndex > 0)
+		// Legacy parents pre substrate support
+		if (SourceEditorOnly.RestrictToBlendRelatives.Num() < SourceEditorOnly.RestrictToLayerRelatives.Num())
 		{
-			EditorOnly.RestrictToBlendRelatives.Add(SourceEditorOnly.RestrictToBlendRelatives[SourceLayerIndex - 1]);
-}
+			if (LayerIndex > 0)
+			{
+				EditorOnly.RestrictToBlendRelatives.Add(SourceEditorOnly.RestrictToBlendRelatives[SourceLayerIndex - 1]);
+	}
+			else
+			{
+				EditorOnly.RestrictToBlendRelatives.AddDefaulted();
+			}
+		}
+		// Source is substrate tree, same number of layersand blends
 		else
 		{
-			EditorOnly.RestrictToBlendRelatives.AddDefaulted();
+			EditorOnly.RestrictToBlendRelatives.Add(SourceEditorOnly.RestrictToBlendRelatives[SourceLayerIndex]);
 		}
 	}
-	// Source is substrate tree, same number of layersand blends
-	else
-	{
-		EditorOnly.RestrictToBlendRelatives.Add(SourceEditorOnly.RestrictToBlendRelatives[SourceLayerIndex]);
-	}
-#else
-	if (LayerIndex > 0)
+	else if (LayerIndex > 0)
 	{
 		EditorOnly.RestrictToBlendRelatives.Add(SourceEditorOnly.RestrictToBlendRelatives[SourceLayerIndex - 1]);
 	}
-#endif
 
 	EditorOnly.LayerGuids.Add(SourceEditorOnly.LayerGuids[SourceLayerIndex]);
 	EditorOnly.LayerLinkStates.Add(LinkState);
@@ -16371,20 +16359,19 @@ void FMaterialLayersFunctions::InsertLayerCopy(const FMaterialLayersFunctionsRun
 	check(LinkState != EMaterialLayerLinkState::Uninitialized);
 	check(LayerIndex > 0);
 	Layers.Insert(Source.Layers[SourceLayerIndex], LayerIndex);
-#if ENABLE_MATERIAL_LAYER_PROTOTYPE
-	Blends.Insert(Source.Blends[SourceLayerIndex], LayerIndex);
-#else
-	Blends.Insert(Source.Blends[SourceLayerIndex - 1], LayerIndex - 1);
-#endif
-
+	if(Substrate::IsMaterialLayeringSupportEnabled())
+	{
+		Blends.Insert(Source.Blends[SourceLayerIndex], LayerIndex);
+		EditorOnly.RestrictToBlendRelatives.Insert(SourceEditorOnly.RestrictToBlendRelatives[SourceLayerIndex], LayerIndex);
+	}
+	else
+	{
+		Blends.Insert(Source.Blends[SourceLayerIndex - 1], LayerIndex - 1);
+		EditorOnly.RestrictToBlendRelatives.Insert(SourceEditorOnly.RestrictToBlendRelatives[SourceLayerIndex - 1], LayerIndex - 1);
+	}
 	EditorOnly.LayerStates.Insert(SourceEditorOnly.LayerStates[SourceLayerIndex], LayerIndex);
 	EditorOnly.LayerNames.Insert(SourceEditorOnly.LayerNames[SourceLayerIndex], LayerIndex);
 	EditorOnly.RestrictToLayerRelatives.Insert(SourceEditorOnly.RestrictToLayerRelatives[SourceLayerIndex], LayerIndex);
-#if ENABLE_MATERIAL_LAYER_PROTOTYPE
-	EditorOnly.RestrictToBlendRelatives.Insert(SourceEditorOnly.RestrictToBlendRelatives[SourceLayerIndex], LayerIndex);
-#else
-	EditorOnly.RestrictToBlendRelatives.Insert(SourceEditorOnly.RestrictToBlendRelatives[SourceLayerIndex - 1], LayerIndex - 1);
-#endif
 	EditorOnly.LayerGuids.Insert(SourceEditorOnly.LayerGuids[SourceLayerIndex], LayerIndex);
 	EditorOnly.LayerLinkStates.Insert(LinkState, LayerIndex);
 }
@@ -16393,25 +16380,32 @@ void FMaterialLayersFunctions::RemoveBlendedLayerAt(int32 Index)
 {
 	if (Layers.IsValidIndex(Index))
 	{
-#if ENABLE_MATERIAL_LAYER_PROTOTYPE
-		check(Layers.IsValidIndex(Index) && Blends.IsValidIndex(Index));
-		Layers.RemoveAt(Index);
-		Blends.RemoveAt(Index);
+		if(Substrate::IsMaterialLayeringSupportEnabled())
+		{
+			check(Layers.IsValidIndex(Index) && Blends.IsValidIndex(Index));
+			Layers.RemoveAt(Index);
+			Blends.RemoveAt(Index);
 
-		check(EditorOnly.LayerStates.IsValidIndex(Index) &&
-			EditorOnly.LayerNames.IsValidIndex(Index) &&
-			EditorOnly.RestrictToLayerRelatives.IsValidIndex(Index) &&
-			EditorOnly.RestrictToBlendRelatives.IsValidIndex(Index));
-#else
-		check(Layers.IsValidIndex(Index) && Blends.IsValidIndex(Index - 1));
-		Layers.RemoveAt(Index);
-		Blends.RemoveAt(Index - 1);
+			check(EditorOnly.LayerStates.IsValidIndex(Index) &&
+				EditorOnly.LayerNames.IsValidIndex(Index) &&
+				EditorOnly.RestrictToLayerRelatives.IsValidIndex(Index) &&
+				EditorOnly.RestrictToBlendRelatives.IsValidIndex(Index));
+
+			EditorOnly.RestrictToBlendRelatives.RemoveAt(Index);
+		}
+		else
+		{
+			check(Layers.IsValidIndex(Index) && Blends.IsValidIndex(Index - 1));
+			Layers.RemoveAt(Index);
+			Blends.RemoveAt(Index - 1);
 		
-		check(EditorOnly.LayerStates.IsValidIndex(Index) &&
-			EditorOnly.LayerNames.IsValidIndex(Index) &&
-			EditorOnly.RestrictToLayerRelatives.IsValidIndex(Index) &&
-			EditorOnly.RestrictToBlendRelatives.IsValidIndex(Index - 1));
-#endif
+			check(EditorOnly.LayerStates.IsValidIndex(Index) &&
+				EditorOnly.LayerNames.IsValidIndex(Index) &&
+				EditorOnly.RestrictToLayerRelatives.IsValidIndex(Index) &&
+				EditorOnly.RestrictToBlendRelatives.IsValidIndex(Index - 1));
+
+			EditorOnly.RestrictToBlendRelatives.RemoveAt(Index - 1);
+		}
 
 		if (EditorOnly.LayerLinkStates[Index] != EMaterialLayerLinkState::NotFromParent)
 		{
@@ -16424,11 +16418,6 @@ void FMaterialLayersFunctions::RemoveBlendedLayerAt(int32 Index)
 		EditorOnly.LayerStates.RemoveAt(Index);
 		EditorOnly.LayerNames.RemoveAt(Index);
 		EditorOnly.RestrictToLayerRelatives.RemoveAt(Index);
-#if ENABLE_MATERIAL_LAYER_PROTOTYPE
-		EditorOnly.RestrictToBlendRelatives.RemoveAt(Index);
-#else
-		EditorOnly.RestrictToBlendRelatives.RemoveAt(Index - 1);
-#endif
 		EditorOnly.LayerGuids.RemoveAt(Index);
 		EditorOnly.LayerLinkStates.RemoveAt(Index);
 	}
@@ -16441,19 +16430,19 @@ void FMaterialLayersFunctions::MoveBlendedLayer(int32 SrcLayerIndex, int32 DstLa
 	if (SrcLayerIndex != DstLayerIndex)
 	{
 		Layers.Swap(SrcLayerIndex, DstLayerIndex);
-#if ENABLE_MATERIAL_LAYER_PROTOTYPE
-		Blends.Swap(SrcLayerIndex, DstLayerIndex);
-#else
-		Blends.Swap(SrcLayerIndex - 1, DstLayerIndex - 1);
-#endif
+		if (Substrate::IsMaterialLayeringSupportEnabled())
+		{
+			Blends.Swap(SrcLayerIndex, DstLayerIndex);
+			EditorOnly.RestrictToBlendRelatives.Swap(SrcLayerIndex, DstLayerIndex);
+		}
+		else
+		{
+			Blends.Swap(SrcLayerIndex - 1, DstLayerIndex - 1);
+			EditorOnly.RestrictToBlendRelatives.Swap(SrcLayerIndex - 1, DstLayerIndex - 1);
+		}
 		EditorOnly.LayerStates.Swap(SrcLayerIndex, DstLayerIndex);
 		EditorOnly.LayerNames.Swap(SrcLayerIndex, DstLayerIndex);
 		EditorOnly.RestrictToLayerRelatives.Swap(SrcLayerIndex, DstLayerIndex);
-#if ENABLE_MATERIAL_LAYER_PROTOTYPE
-		EditorOnly.RestrictToBlendRelatives.Swap(SrcLayerIndex, DstLayerIndex);
-#else
-		EditorOnly.RestrictToBlendRelatives.Swap(SrcLayerIndex - 1, DstLayerIndex - 1);
-#endif
 		EditorOnly.LayerGuids.Swap(SrcLayerIndex, DstLayerIndex);
 		EditorOnly.LayerLinkStates.Swap(SrcLayerIndex, DstLayerIndex);
 	}
@@ -16549,11 +16538,8 @@ bool FMaterialLayersFunctions::MatchesParent(const FMaterialLayersFunctionsRunti
 		{
 			return false;
 		}
-#if ENABLE_MATERIAL_LAYER_PROTOTYPE
-		if (LayerIndex > 0 && Runtime.Blends[LayerIndex] != ParentRuntime.Blends[ParentLayerIndex])
-#else
-		if (LayerIndex > 0 && Runtime.Blends[LayerIndex - 1] != ParentRuntime.Blends[ParentLayerIndex - 1])
-#endif
+
+		if (LayerIndex > 0 && (Substrate::IsMaterialLayeringSupportEnabled() ? Runtime.Blends[LayerIndex] == ParentRuntime.Blends[ParentLayerIndex] : Runtime.Blends[LayerIndex - 1] == ParentRuntime.Blends[ParentLayerIndex - 1]))
 		{
 			return false;
 		}
@@ -16624,12 +16610,8 @@ bool FMaterialLayersFunctions::ResolveParent(const FMaterialLayersFunctionsRunti
 			else
 			{
 				// See if we match layer in parent
-				if (Runtime.Layers[LayerIndex] == ParentRuntime.Layers[ParentLayerIndex] &&
-#if ENABLE_MATERIAL_LAYER_PROTOTYPE
-				(LayerIndex == 0 || Runtime.Blends[LayerIndex] == ParentRuntime.Blends[ParentLayerIndex]))
-#else
-				(LayerIndex == 0 || Runtime.Blends[LayerIndex - 1] == ParentRuntime.Blends[ParentLayerIndex - 1]))
-#endif
+				if (Runtime.Layers[LayerIndex] == ParentRuntime.Layers[ParentLayerIndex]
+					&& (LayerIndex == 0 || (Substrate::IsMaterialLayeringSupportEnabled() ? Runtime.Blends[LayerIndex] == ParentRuntime.Blends[ParentLayerIndex] : Runtime.Blends[LayerIndex - 1] == ParentRuntime.Blends[ParentLayerIndex - 1])))
 				{
 					// Parent layer matches, so link to parent
 					ResolvedLayerIndex = ResolvedLayers.AddLayerCopy(ParentRuntime, ParentEditorOnly, ParentLayerIndex, bLayerVisible, EMaterialLayerLinkState::LinkedToParent);
@@ -16641,6 +16623,7 @@ bool FMaterialLayersFunctions::ResolveParent(const FMaterialLayersFunctionsRunti
 					ResolvedLayers.EditorOnly.LayerGuids[ResolvedLayerIndex] = ParentEditorOnly.LayerGuids[ParentLayerIndex]; // Still need to match guid to parent
 				}
 
+
 				check(!ParentLayerIndices.Contains(ParentLayerIndex));
 				ParentLayerIndices.Add(ParentLayerIndex);
 			}
@@ -16748,35 +16731,36 @@ bool FMaterialLayersFunctions::ResolveParent(const FMaterialLayersFunctionsRunti
 		}
 	}
 
-#if ENABLE_MATERIAL_LAYER_PROTOTYPE
-	struct Traverser
+	if(Substrate::IsMaterialLayeringSupportEnabled())
 	{
-		void Recursive(const FMaterialLayersFunctionsTree& Source, FLayerNodeId InSourceId, FMaterialLayersFunctionsTree& Dest, FLayerNodeId InDestParentId)
+		struct Traverser
 		{
-
-			FLayerNode* SourceNode = Source.GetNode(InSourceId);
-			FLayerNodeId NewNodeId = FMaterialLayersFunctionsTree::InvalidId;
-			FLayerNodePayload Payload;
-			if (SourceNode)
+			void Recursive(const FMaterialLayersFunctionsTree& Source, FLayerNodeId InSourceId, FMaterialLayersFunctionsTree& Dest, FLayerNodeId InDestParentId)
 			{
-				Payload = Source.Payloads[InSourceId];
-				NewNodeId = Dest.AddNode(Payload, InDestParentId);
-			}
 
-			auto ChildrenIds = Source.GetChildrenIds(InSourceId);
-			for (auto SourceSubId : ChildrenIds)
-			{
-				Recursive(Source, SourceSubId, Dest, NewNodeId);
-			}
+				FLayerNode* SourceNode = Source.GetNode(InSourceId);
+				FLayerNodeId NewNodeId = FMaterialLayersFunctionsTree::InvalidId;
+				FLayerNodePayload Payload;
+				if (SourceNode)
+				{
+					Payload = Source.Payloads[InSourceId];
+					NewNodeId = Dest.AddNode(Payload, InDestParentId);
+				}
 
-		}
-	} ReplicateTree;
+				auto ChildrenIds = Source.GetChildrenIds(InSourceId);
+				for (auto SourceSubId : ChildrenIds)
+				{
+					Recursive(Source, SourceSubId, Dest, NewNodeId);
+				}
+
+			}
+		} ReplicateTree;
 	
-	ReplicateTree.Recursive(Runtime.Tree, -1, ResolvedLayers.Tree, -1);
+		ReplicateTree.Recursive(Runtime.Tree, -1, ResolvedLayers.Tree, -1);
 
-	// NOTE: For debug purpose at the moment, we want to monitor the end result tree data state.
-	MLFT::Log(&Runtime, __FUNCTION__);
-#endif
+		// NOTE: For debug purpose at the moment, we want to monitor the end result tree data state.
+		MLFT::Log(&Runtime, __FUNCTION__);
+	}
 
 	Runtime = MoveTemp(static_cast<FMaterialLayersFunctionsRuntimeData&>(ResolvedLayers));
 	EditorOnly = MoveTemp(ResolvedLayers.EditorOnly);
@@ -16789,9 +16773,7 @@ void FMaterialLayersFunctions::Validate(const FMaterialLayersFunctionsRuntimeDat
 {
 	if (Runtime.Layers.Num() > 0)
 	{
-#if !ENABLE_MATERIAL_LAYER_PROTOTYPE
-		check(Runtime.Blends.Num() == Runtime.Layers.Num() - 1);
-#endif // ENABLE_MATERIAL_LAYER_PROTOTYPE
+		check(Substrate::IsMaterialLayeringSupportEnabled() || Runtime.Blends.Num() == Runtime.Layers.Num() - 1);
 		check(Runtime.Layers.Num() == EditorOnly.LayerStates.Num());
 		check(Runtime.Layers.Num() == EditorOnly.LayerNames.Num());
 		check(Runtime.Layers.Num() == EditorOnly.LayerGuids.Num());
@@ -16853,21 +16835,21 @@ bool FMaterialLayersFunctions::CanAppendLayerNode(FLayerNodeId InParent) const
 
 FLayerNodeId FMaterialLayersFunctions::AppendLayerNode(FLayerNodeId InParent, int32 InSiblingIndex)
 {
-#if ENABLE_MATERIAL_LAYER_PROTOTYPE
-	if (!CanAppendLayerNode(InParent))
-		return FMaterialLayersFunctionsTree::InvalidId;
+	if(Substrate::IsMaterialLayeringSupportEnabled())
+	{
+		if (!CanAppendLayerNode(InParent))
+			return FMaterialLayersFunctionsTree::InvalidId;
 
-	const int32 PayloadIndex = AppendBlendedLayer();
+		const int32 PayloadIndex = AppendBlendedLayer();
 
-	// InPayload
-	FLayerNodeId NewNodeId = Tree.AddNode({ PayloadIndex, PayloadIndex }, InParent, InSiblingIndex);
+		// InPayload
+		FLayerNodeId NewNodeId = Tree.AddNode({ PayloadIndex, PayloadIndex }, InParent, InSiblingIndex);
 
-	MLFT::Log(this, __FUNCTION__);
+		MLFT::Log(this, __FUNCTION__);
 
-	return NewNodeId;
-#else
+		return NewNodeId;
+	}
 	return FMaterialLayersFunctionsTree::InvalidId;
-#endif
 }
 
 bool FMaterialLayersFunctions::CanRemoveLayerNode(FLayerNodeId InNodeId) const
@@ -16883,8 +16865,7 @@ bool FMaterialLayersFunctions::CanRemoveLayerNode(FLayerNodeId InNodeId) const
 
 void FMaterialLayersFunctions::RemoveLayerNodeAt(FLayerNodeId InNodeId)
 {
-#if ENABLE_MATERIAL_LAYER_PROTOTYPE
-	if (!CanRemoveLayerNode(InNodeId))
+	if (!Substrate::IsMaterialLayeringSupportEnabled() || !CanRemoveLayerNode(InNodeId))
 		return;
 
 	if (Tree.IsValidId(InNodeId))
@@ -16931,7 +16912,6 @@ void FMaterialLayersFunctions::RemoveLayerNodeAt(FLayerNodeId InNodeId)
 	}
 
 	MLFT::Log(this, __FUNCTION__);
-#endif
 }
 	
 
@@ -16956,8 +16936,7 @@ void FMaterialLayersFunctions::MoveLayerNode(FLayerNodeId InNodeId, FLayerNodeId
 
 UMaterial* FMaterialLayersFunctions::GetRuntimeNodePreviewMaterial(FLayerNodeId InNodeId) const
 {
-#if ENABLE_MATERIAL_LAYER_PROTOTYPE
-	if (RuntimeGraphCache)
+	if (Substrate::IsMaterialLayeringSupportEnabled() && RuntimeGraphCache)
 	{
 		if (Tree.IsValidId(InNodeId))
 		{
@@ -16987,7 +16966,6 @@ UMaterial* FMaterialLayersFunctions::GetRuntimeNodePreviewMaterial(FLayerNodeId
 			return RuntimeGraphCache->NodePreviewMaterials[InNodeId];	
 		}
 	}
-#endif
 	return nullptr;
 }
 
diff --git a/Engine/Source/Runtime/Engine/Private/Materials/MaterialInstance.cpp b/Engine/Source/Runtime/Engine/Private/Materials/MaterialInstance.cpp
--- a/Engine/Source/Runtime/Engine/Private/Materials/MaterialInstance.cpp
+++ b/Engine/Source/Runtime/Engine/Private/Materials/MaterialInstance.cpp
@@ -294,18 +294,16 @@ bool FMaterialInstanceResource::GetParameterValue(EMaterialParameterType Type, c
 	// Check for hard-coded parameters
 	if (Type == EMaterialParameterType::Scalar && ParameterInfo.Name == SubsurfaceProfile::GetSubsurfaceProfileParameterName())
 	{
-#if !ENABLE_MATERIAL_LAYER_PROTOTYPE
-		check(ParameterInfo.Association == EMaterialParameterAssociation::GlobalParameter);
-#endif
+		check(Substrate::IsMaterialLayeringSupportEnabled() || ParameterInfo.Association == EMaterialParameterAssociation::GlobalParameter);
+
 		const USubsurfaceProfile* MySubsurfaceProfileRT = GetSubsurfaceProfileRT();
 		OutValue = SubsurfaceProfile::GetSubsurfaceProfileId(MySubsurfaceProfileRT);
 		bResult = true;
 	}
 	else if (Type == EMaterialParameterType::Scalar && NumSubsurfaceProfileRT() > 0)
 	{
-#if !ENABLE_MATERIAL_LAYER_PROTOTYPE
-		check(ParameterInfo.Association == EMaterialParameterAssociation::GlobalParameter);
-#endif
+		check(Substrate::IsMaterialLayeringSupportEnabled() || ParameterInfo.Association == EMaterialParameterAssociation::GlobalParameter);
+
 		const USubsurfaceProfile* SSProfileOverrideRT = GetSubsurfaceProfileRT();
 		// Substrate general SubsurfaceProfile
 		for (uint32 It = 0, Count = NumSubsurfaceProfileRT(); It < Count; ++It)
@@ -327,9 +325,7 @@ bool FMaterialInstanceResource::GetParameterValue(EMaterialParameterType Type, c
 		{
 			if (ParameterInfo.Name == SpecularProfile::GetSpecularProfileParameterName(GetSpecularProfileRT(It)))
 			{
-#if !ENABLE_MATERIAL_LAYER_PROTOTYPE
-				check(ParameterInfo.Association == EMaterialParameterAssociation::GlobalParameter);
-#endif
+				check(Substrate::IsMaterialLayeringSupportEnabled() || ParameterInfo.Association == EMaterialParameterAssociation::GlobalParameter);
 				OutValue = SpecularProfile::GetSpecularProfileId(SPOverrideRT ? SPOverrideRT : GetSpecularProfileRT(It));
 				bResult = true;
 				break;
diff --git a/Engine/Source/Runtime/Engine/Public/Materials/MaterialExpression.h b/Engine/Source/Runtime/Engine/Public/Materials/MaterialExpression.h
--- a/Engine/Source/Runtime/Engine/Public/Materials/MaterialExpression.h
+++ b/Engine/Source/Runtime/Engine/Public/Materials/MaterialExpression.h
@@ -691,10 +691,4 @@ struct FExpressionOutputIterator
     }
 };
 
-#endif // WITH_EDITOR
-
-#ifndef ENABLE_MATERIAL_LAYER_PROTOTYPE
-//WIP Internal Only - Please do not enable + use this in any projects other than for iterating the new layering prototype system until further notice
-//Temporarily placed here for the best possible catch all position when in use - will be removed for 5.5
-#define ENABLE_MATERIAL_LAYER_PROTOTYPE 0
-#endif
\ No newline at end of file
+#endif // WITH_EDITOR
\ No newline at end of file
diff --git a/Engine/Source/Runtime/Engine/Public/Materials/MaterialExpressionBlendMaterialAttributes.h b/Engine/Source/Runtime/Engine/Public/Materials/MaterialExpressionBlendMaterialAttributes.h
--- a/Engine/Source/Runtime/Engine/Public/Materials/MaterialExpressionBlendMaterialAttributes.h
+++ b/Engine/Source/Runtime/Engine/Public/Materials/MaterialExpressionBlendMaterialAttributes.h
@@ -72,7 +72,7 @@ class UMaterialExpressionBlendMaterialAttributes : public UMaterialExpression
 UCLASS(collapsecategories, hidecategories = Object, MinimalAPI)
 class UMaterialExpressionLegacyBlendMaterialAttributes : public UMaterialExpressionBlendMaterialAttributes
 {
-	GENERATED_BODY()
+	GENERATED_UCLASS_BODY()
 
 public:
 
@@ -91,7 +91,6 @@ public:
 	UPROPERTY(EditAnywhere, Category = MaterialAttributes)
 	TEnumAsByte<EMaterialAttributeBlendFunction::Type> BlendFunctionType;
 
-#if ENABLE_MATERIAL_LAYER_PROTOTYPE
 #if WITH_EDITOR
 	virtual void GetCaption(TArray<FString>& OutCaptions) const override;
 	virtual bool CanEditChange(const FProperty* InProperty) const override;
@@ -101,5 +100,4 @@ public:
 	virtual uint32 GetInputType(int32 InputIndex) override;
 	virtual int32 Compile(class FMaterialCompiler* Compiler, int32 OutputIndex) override;
 #endif // WITH_EDITOR
-#endif //ENABLE_MATERIAL_LAYER_PROTOTYPE
 };
\ No newline at end of file
diff --git a/Engine/Source/Runtime/Engine/Public/Materials/MaterialExpressionGetMaterialAttributes.h b/Engine/Source/Runtime/Engine/Public/Materials/MaterialExpressionGetMaterialAttributes.h
--- a/Engine/Source/Runtime/Engine/Public/Materials/MaterialExpressionGetMaterialAttributes.h
+++ b/Engine/Source/Runtime/Engine/Public/Materials/MaterialExpressionGetMaterialAttributes.h
@@ -36,10 +36,7 @@ class UMaterialExpressionGetMaterialAttributes : public UMaterialExpression
  
 	//~ Begin UMaterialExpression Interface
 #if WITH_EDITOR
-#if ENABLE_MATERIAL_LAYER_PROTOTYPE
 	ENGINE_API int32 CreateOrGetOutputAttribute(EMaterialProperty Attribute);
-#endif
-
 	virtual int32 Compile(class FMaterialCompiler* Compiler, int32 OutputIndex) override;
 	virtual void GetCaption(TArray<FString>& OutCaptions) const override;
 	virtual FName GetInputName(int32 InputIndex) const override;
diff --git a/Engine/Source/Runtime/Engine/Public/Materials/MaterialExpressionMaterialAttributeLayers.h b/Engine/Source/Runtime/Engine/Public/Materials/MaterialExpressionMaterialAttributeLayers.h
--- a/Engine/Source/Runtime/Engine/Public/Materials/MaterialExpressionMaterialAttributeLayers.h
+++ b/Engine/Source/Runtime/Engine/Public/Materials/MaterialExpressionMaterialAttributeLayers.h
@@ -112,10 +112,8 @@ class UMaterialExpressionMaterialAttributeLayers : public UMaterialExpression
 	virtual bool IsInputConnectionRequired(int32 InputIndex) const override {return false;}
 	virtual uint32 GetInputType(int32 InputIndex) override;
 	virtual bool IsResultMaterialAttributes(int32 OutputIndex) override {return true;}
-#if ENABLE_MATERIAL_LAYER_PROTOTYPE
 	virtual bool IsResultSubstrateMaterial(int32 OutputIndex) override;
 	virtual FSubstrateOperator* SubstrateGenerateMaterialTopologyTree(class FMaterialCompiler* Compiler, class UMaterialExpression* Parent, int32 OutputIndex) override;
-#endif //ENABLE_MATERIAL_LAYER_PROTOTYPE
 #endif
 	//~ End UMaterialExpression Interface
 
diff --git a/Engine/Source/Runtime/Engine/Public/Materials/MaterialExpressionSetMaterialAttributes.h b/Engine/Source/Runtime/Engine/Public/Materials/MaterialExpressionSetMaterialAttributes.h
--- a/Engine/Source/Runtime/Engine/Public/Materials/MaterialExpressionSetMaterialAttributes.h
+++ b/Engine/Source/Runtime/Engine/Public/Materials/MaterialExpressionSetMaterialAttributes.h
@@ -43,17 +43,14 @@ class UMaterialExpressionSetMaterialAttributes : public UMaterialExpression
 	virtual bool IsInputConnectionRequired(int32 InputIndex) const override {return true;}
 	virtual bool IsResultMaterialAttributes(int32 OutputIndex) override {return true;}
 
-#if ENABLE_MATERIAL_LAYER_PROTOTYPE
 	ENGINE_API int32 CreateOrGetInputAttribute(EMaterialProperty Attribute);
 	ENGINE_API bool ConnectInputAttribute(EMaterialProperty Attribute, UMaterialExpression* Expression, int32 OutputIndex = 0);
-#endif
 	bool GetSubstrateMaterialInputIndex(int32 OutputIndex, int32& InputIndex);
 	virtual bool IsResultSubstrateMaterial(int32 OutputIndex) override;
 	virtual void GatherSubstrateMaterialInfo(FSubstrateMaterialInfo& SubstrateMaterialInfo, int32 OutputIndex) override;
 	virtual FSubstrateOperator* SubstrateGenerateMaterialTopologyTree(FMaterialCompiler* Compiler, UMaterialExpression* Parent, int32 OutputIndex) override;
 	virtual void GetExpressionToolTip(TArray<FString>& OutToolTip) override;
 
-
 	uint64 GetConnectedInputs() const;
 #endif
 	//~ End UMaterialExpression Interface
diff --git a/Engine/Source/Runtime/Engine/Public/Materials/MaterialFunction.h b/Engine/Source/Runtime/Engine/Public/Materials/MaterialFunction.h
--- a/Engine/Source/Runtime/Engine/Public/Materials/MaterialFunction.h
+++ b/Engine/Source/Runtime/Engine/Public/Materials/MaterialFunction.h
@@ -177,13 +177,13 @@ public:
 
 	/** Returns If returns an empty string, use the default class name for the material function. Otherwise, the string will be the name shown when the function is exposed to users in the material graph as a node, or from the contextual menu when searching for nodes. */
 	virtual FString GetUserExposedCaption() const override { return UserExposedCaption; }
-#if ENABLE_MATERIAL_LAYER_PROTOTYPE
+
 	/** Converts legacy Material Layer Functions to utilise Substrate */
 	void ConvertExpressionsBetweenLegacyAndSubstrate();
 
 	/** Retrieves an array of all custom output expressions in a Material Function */
 	void GetAllCustomOutputExpressions(TArray<class UMaterialExpressionCustomOutput*>& OutCustomOutputs) const;
-#endif //ENABLE_MATERIAL_LAYER_PROTOTYPE
+
 #endif // WITH_EDITORONLY_DATA
 
 #if WITH_EDITOR
diff --git a/Engine/Source/Runtime/Engine/Public/Materials/MaterialFunctionInterface.h b/Engine/Source/Runtime/Engine/Public/Materials/MaterialFunctionInterface.h
--- a/Engine/Source/Runtime/Engine/Public/Materials/MaterialFunctionInterface.h
+++ b/Engine/Source/Runtime/Engine/Public/Materials/MaterialFunctionInterface.h
@@ -34,9 +34,8 @@ enum class EMaterialFunctionUsage : uint8
 	MaterialLayer,
 	MaterialLayerBlend
 };
-#if ENABLE_MATERIAL_LAYER_PROTOTYPE
-#define DEFAULT_MATERIALLAYERBLEND_PATH TEXT("/Engine/EngineMaterials/Blends/MLB_Horizontal")
-#endif
+
+#define DEFAULT_SUBSTRATE_MATERIALLAYERBLEND_PATH TEXT("/Engine/EngineMaterials/Blends/MLB_Horizontal")
 
 using FMFRecursionGuard = TMaterialRecursionGuard<class UMaterialFunctionInterface>;
 
diff --git a/Engine/Source/Runtime/Engine/Public/Materials/MaterialLayersFunctions.h b/Engine/Source/Runtime/Engine/Public/Materials/MaterialLayersFunctions.h
--- a/Engine/Source/Runtime/Engine/Public/Materials/MaterialLayersFunctions.h
+++ b/Engine/Source/Runtime/Engine/Public/Materials/MaterialLayersFunctions.h
@@ -411,22 +411,18 @@ struct FMaterialLayersFunctionsRuntimeData
 	UPROPERTY(EditAnywhere, Category = MaterialLayers)
 	TArray<TObjectPtr<class UMaterialFunctionInterface>> Blends;
 
-	// ENABLE_MATERIAL_LAYER_PROTOTYPE : UProperty describing the tree data structure
+	// UProperty describing the tree data structure
 	UPROPERTY()
 	FMaterialLayersFunctionsTree Tree;
 
-#if ENABLE_MATERIAL_LAYER_PROTOTYPE
 	mutable TSharedPtr<FMaterialLayersFunctionsRuntimeGraphCache> RuntimeGraphCache;
-#endif
 
 	FMaterialLayersFunctionsRuntimeData() = default;
 	FMaterialLayersFunctionsRuntimeData(const FMaterialLayersFunctionsRuntimeData& Rhs)
 		: Layers(Rhs.Layers)
 		, Blends(Rhs.Blends)
-#if ENABLE_MATERIAL_LAYER_PROTOTYPE
 		, Tree(Rhs.Tree)
 		, RuntimeGraphCache(Rhs.RuntimeGraphCache)
-#endif
 	{}
 
 	FMaterialLayersFunctionsRuntimeData(const FMaterialLayersFunctions& Rhs) = delete;
@@ -435,10 +431,8 @@ struct FMaterialLayersFunctionsRuntimeData
 	{
 		Layers = Rhs.Layers;
 		Blends = Rhs.Blends;
-#if ENABLE_MATERIAL_LAYER_PROTOTYPE
 		Tree = Rhs.Tree;
 		RuntimeGraphCache = Rhs.RuntimeGraphCache;
-#endif
 		return *this;
 	}
 
@@ -450,18 +444,14 @@ struct FMaterialLayersFunctionsRuntimeData
 	{
 		Layers.Empty();
 		Blends.Empty();
-#if ENABLE_MATERIAL_LAYER_PROTOTYPE
 		Tree.Empty();
 		RuntimeGraphCache.Reset();
-#endif
 	}
 
 	FORCEINLINE bool operator==(const FMaterialLayersFunctionsRuntimeData& Other) const
 	{
 		if (Layers != Other.Layers || Blends != Other.Blends
-#if ENABLE_MATERIAL_LAYER_PROTOTYPE
 			|| Tree != Other.Tree
-#endif
 		)
 		{
 			return false;
diff --git a/Engine/Source/Runtime/RenderCore/Private/RenderUtils.cpp b/Engine/Source/Runtime/RenderCore/Private/RenderUtils.cpp
--- a/Engine/Source/Runtime/RenderCore/Private/RenderUtils.cpp
+++ b/Engine/Source/Runtime/RenderCore/Private/RenderUtils.cpp
@@ -1847,6 +1847,12 @@ static TAutoConsoleVariable<int32> CVarSubstrateDBufferPass(
 	TEXT("Apply DBuffer after the base-pass as a separate pass. Read only because when this is changed, it will require the recompilation of all shaders."),
 	ECVF_ReadOnly | ECVF_RenderThreadSafe);
 
+static TAutoConsoleVariable<int32> CVarSubstrateEnableLayerSupport(
+	TEXT("r.Substrate.EnableLayerSupport"),
+	0,
+	TEXT("Enable Substrate material layering and UI (Experimental). Note: This support is one way; the legacy layer materials auto-upgrade and can only be reversed manually after assets are re-saved."),
+	ECVF_ReadOnly | ECVF_RenderThreadSafe);
+
 // Transition render settings that will disappear when Substrate gets enabled
 static TAutoConsoleVariable<int32> CVarMaterialRoughDiffuse(
 	TEXT("r.Material.RoughDiffuse"),
@@ -2117,6 +2123,11 @@ namespace Substrate
 	{
 		return CVarSubstrateDebugAdvancedVisualizationShaders.GetValueOnAnyThread() > 0;
 	}
+
+	bool IsMaterialLayeringSupportEnabled()
+	{
+		return IsSubstrateEnabled() && CVarSubstrateEnableLayerSupport.GetValueOnAnyThread() > 0;
+	}
 }
 
 bool HasFirstPersonGBufferBit(EShaderPlatform ShaderPlatform)
diff --git a/Engine/Source/Runtime/RenderCore/Public/RenderUtils.h b/Engine/Source/Runtime/RenderCore/Public/RenderUtils.h
--- a/Engine/Source/Runtime/RenderCore/Public/RenderUtils.h
+++ b/Engine/Source/Runtime/RenderCore/Public/RenderUtils.h
@@ -596,6 +596,7 @@ namespace Substrate
 	RENDERCORE_API bool IsDBufferPassEnabled(EShaderPlatform InPlatform);
 	RENDERCORE_API bool IsOpaqueRoughRefractionEnabled();
 	RENDERCORE_API bool IsAdvancedVisualizationEnabled();
+	RENDERCORE_API bool IsMaterialLayeringSupportEnabled();
 
 	RENDERCORE_API uint32 GetRayTracingMaterialPayloadSizeInBytes(bool bFullySimplifiedMaterial);
 
