Allow wp.Runtime.UpdateStreaming.EnableAsyncUpdate to work for listen servers.

diff --git a/Engine/Source/Runtime/Engine/Private/WorldPartition/WorldPartitionStreamingPolicy.cpp b/Engine/Source/Runtime/Engine/Private/WorldPartition/WorldPartitionStreamingPolicy.cpp
--- a/Engine/Source/Runtime/Engine/Private/WorldPartition/WorldPartitionStreamingPolicy.cpp
+++ b/Engine/Source/Runtime/Engine/Private/WorldPartition/WorldPartitionStreamingPolicy.cpp
@@ -345,7 +345,8 @@ void UWorldPartitionStreamingPolicy::UpdateStreamingState()
 	UpdateStreamingSources(bCanOptimizeUpdate);
 
 	// Determine if update will be async or not
-	const bool bCanUpdateAsync = UWorldPartitionStreamingPolicy::IsAsyncUpdateStreamingStateEnabled && !bIsServer && bCanStream && WorldPartition->IsInitialized() && !IsInBlockTillLevelStreamingCompleted();
+	const bool bIsDedicatedServer = (World->GetNetMode() == NM_DedicatedServer);
+	const bool bCanUpdateAsync = UWorldPartitionStreamingPolicy::IsAsyncUpdateStreamingStateEnabled && !bIsDedicatedServer && bCanStream && WorldPartition->IsInitialized() && !IsInBlockTillLevelStreamingCompleted();
 
 	// Detect if nothing relevant changed and early out
 	const uint32 NewUpdateStreamingHash = ComputeUpdateStreamingHash(bCanOptimizeUpdate);
