First Person: Fixed setting primitive to FirstPersonWorldSpaceRepresentation not properly configuring the primitive to behave as if it was set to bOwnerNoSee. Also disabled most shadow casting options on FirstPerson primitives as these don't properly work with such primitives. Shadowing is instead solved with First Person Self Shadows and an invisible FirstPersonWorldSpaceRepresentation primitive for casting shadow onto the scene.

diff --git a/Engine/Source/Runtime/Engine/Classes/Components/PrimitiveComponent.h b/Engine/Source/Runtime/Engine/Classes/Components/PrimitiveComponent.h
--- a/Engine/Source/Runtime/Engine/Classes/Components/PrimitiveComponent.h
+++ b/Engine/Source/Runtime/Engine/Classes/Components/PrimitiveComponent.h
@@ -445,8 +445,8 @@ public:
 	UPROPERTY(EditAnywhere, AdvancedDisplay, BlueprintReadOnly, Category = Rendering, Interp)
 	uint8 bHoldout : 1;
 
-	/** If this is True, this component won't be visible when the view actor is the component's owner, directly or indirectly. Will be internally overridden to true when FirstPersonPrimitiveType is set to WorldSpaceRepresentation. */
-	UPROPERTY(EditAnywhere, AdvancedDisplay, BlueprintReadOnly, Category = Rendering)
+	/** If this is True, this component won't be visible when the view actor is the component's owner, directly or indirectly. Will be internally set to true when FirstPersonPrimitiveType is set to WorldSpaceRepresentation. */
+	UPROPERTY(EditAnywhere, AdvancedDisplay, BlueprintReadOnly, Category = Rendering, meta=(EditCondition="FirstPersonPrimitiveType != EFirstPersonPrimitiveType::WorldSpaceRepresentation"))
 	uint8 bOwnerNoSee:1;
 
 	/** If this is True, this component will only be visible when the view actor is the component's owner, directly or indirectly. */
@@ -507,24 +507,24 @@ public:
 	UPROPERTY(EditAnywhere, AdvancedDisplay, BlueprintReadOnly, Category=Lighting, meta=(EditCondition="bAffectDynamicIndirectLighting", DisplayName = "Affect Indirect Lighting While Hidden"), Interp)
 	uint8 bAffectIndirectLightingWhileHidden:1;
 
-	/** Controls whether the primitive should affect dynamic distance field lighting methods.  This flag is only used if CastShadow is true. **/
+	/** Controls whether the primitive should affect dynamic distance field lighting methods.  This flag is only used if CastShadow is true. */
 	UPROPERTY(EditAnywhere, BlueprintReadOnly, Category=Lighting, AdvancedDisplay)
 	uint8 bAffectDistanceFieldLighting:1;
 
-	/** Controls whether the primitive should cast shadows in the case of non precomputed shadowing.  This flag is only used if CastShadow is true. **/
-	UPROPERTY(EditAnywhere, BlueprintReadOnly, Category=Lighting, AdvancedDisplay, meta=(EditCondition="CastShadow", DisplayName = "Dynamic Shadow"))
+	/** Controls whether the primitive should cast shadows in the case of non precomputed shadowing. This flag is only used if CastShadow is true and if FirstPersonPrimitiveType is not set to FirstPerson. */
+	UPROPERTY(EditAnywhere, BlueprintReadOnly, Category=Lighting, AdvancedDisplay, meta=(EditCondition="CastShadow && (FirstPersonPrimitiveType != EFirstPersonPrimitiveType::FirstPerson)", DisplayName = "Dynamic Shadow"))
 	uint8 bCastDynamicShadow:1;
 
-	/** Whether the object should cast a static shadow from shadow casting lights.  This flag is only used if CastShadow is true. */
-	UPROPERTY(EditAnywhere, BlueprintReadOnly, Category=Lighting, AdvancedDisplay, meta=(EditCondition="CastShadow", DisplayName = "Static Shadow"))
+	/** Whether the object should cast a static shadow from shadow casting lights. This flag is only used if CastShadow is true and if FirstPersonPrimitiveType is not set to FirstPerson. */
+	UPROPERTY(EditAnywhere, BlueprintReadOnly, Category=Lighting, AdvancedDisplay, meta=(EditCondition="CastShadow && (FirstPersonPrimitiveType != EFirstPersonPrimitiveType::FirstPerson)", DisplayName = "Static Shadow"))
 	uint8 bCastStaticShadow:1;
 
 	/** 
 	 * Whether the object should cast a volumetric translucent shadow.
 	 * Volumetric translucent shadows are useful for primitives with smoothly changing opacity like particles representing a volume, 
-	 * But have artifacts when used on highly opaque surfaces.
+	 * but have artifacts when used on highly opaque surfaces. This flag is only used if CastShadow is true and if FirstPersonPrimitiveType is not set to FirstPerson.
 	 */
-	UPROPERTY(EditAnywhere, AdvancedDisplay, BlueprintReadOnly, Category=Lighting, meta=(EditCondition="CastShadow", DisplayName = "Volumetric Translucent Shadow"))
+	UPROPERTY(EditAnywhere, AdvancedDisplay, BlueprintReadOnly, Category=Lighting, meta=(EditCondition="CastShadow && (FirstPersonPrimitiveType != EFirstPersonPrimitiveType::FirstPerson)", DisplayName = "Volumetric Translucent Shadow"))
 	uint8 bCastVolumetricTranslucentShadow:1;
 
 	/**
@@ -542,9 +542,9 @@ public:
 	uint8 bSelfShadowOnly:1;
 
 	/** 
-	 * When enabled, the component will be rendering into the far shadow cascades (only for directional lights).
+	 * When enabled, the component will be rendering into the far shadow cascades (only for directional lights). This flag is only used if CastShadow is true and if FirstPersonPrimitiveType is not set to FirstPerson.
 	 */
-	UPROPERTY(EditAnywhere, AdvancedDisplay, BlueprintReadOnly, Category=Lighting, meta=(EditCondition="CastShadow", DisplayName = "Far Shadow"))
+	UPROPERTY(EditAnywhere, AdvancedDisplay, BlueprintReadOnly, Category=Lighting, meta=(EditCondition="CastShadow && (FirstPersonPrimitiveType != EFirstPersonPrimitiveType::FirstPerson)", DisplayName = "Far Shadow"))
 	uint8 bCastFarShadow:1;
 
 	/** 
@@ -564,10 +564,9 @@ public:
 	/** 
 	 *	If true, the primitive will cast shadows even if bHidden is true.
 	 *	Controls whether the primitive should cast shadows when hidden.
-	 *	This flag is only used if CastShadow is true and will be internally 
-	 *  overridden to false if FirstPersonPrimitiveType is set to WorldSpaceRepresentation.
+	 *	This flag is only used if CastShadow is true and if FirstPersonPrimitiveType is not set to WorldSpaceRepresentation.
 	 */
-	UPROPERTY(EditAnywhere, AdvancedDisplay, BlueprintReadOnly, Category=Lighting, meta=(EditCondition="CastShadow", DisplayName = "Hidden Shadow"), Interp)
+	UPROPERTY(EditAnywhere, AdvancedDisplay, BlueprintReadOnly, Category=Lighting, meta=(EditCondition="CastShadow && (FirstPersonPrimitiveType != EFirstPersonPrimitiveType::WorldSpaceRepresentation)", DisplayName = "Hidden Shadow"), Interp)
 	uint8 bCastHiddenShadow:1;
 
 	/** Whether this primitive should cast dynamic shadows as if it were a two sided material. */
diff --git a/Engine/Source/Runtime/Engine/Classes/Components/SkinnedMeshComponent.h b/Engine/Source/Runtime/Engine/Classes/Components/SkinnedMeshComponent.h
--- a/Engine/Source/Runtime/Engine/Classes/Components/SkinnedMeshComponent.h
+++ b/Engine/Source/Runtime/Engine/Classes/Components/SkinnedMeshComponent.h
@@ -821,15 +821,16 @@ public:
 	/** 
 	 * Whether to use the capsule representation (when present) from a skeletal mesh's ShadowPhysicsAsset for direct shadowing from lights.
 	 * This type of shadowing is approximate but handles extremely wide area shadowing well.  The softness of the shadow depends on the light's LightSourceAngle / SourceRadius.
-	 * This flag will force bCastInsetShadow to be enabled.
+	 * This flag will force bCastInsetShadow to be enabled. This flag is only used if CastShadow is true and if FirstPersonPrimitiveType is not set to FirstPerson.
 	 */
-	UPROPERTY(EditAnywhere, AdvancedDisplay, BlueprintReadOnly, Category=Lighting, meta=(EditCondition="CastShadow", DisplayName = "Capsule Direct Shadow"))
+	UPROPERTY(EditAnywhere, AdvancedDisplay, BlueprintReadOnly, Category=Lighting, meta=(EditCondition="CastShadow && (FirstPersonPrimitiveType != EFirstPersonPrimitiveType::FirstPerson)", DisplayName = "Capsule Direct Shadow"))
 	uint8 bCastCapsuleDirectShadow:1;
 
 	/** 
 	 * Whether to use the capsule representation (when present) from a skeletal mesh's ShadowPhysicsAsset for shadowing indirect lighting (from lightmaps or skylight).
+	 * This flag is only used if CastShadow is true and if FirstPersonPrimitiveType is not set to FirstPerson.
 	 */
-	UPROPERTY(EditAnywhere, AdvancedDisplay, BlueprintReadOnly, Category=Lighting, meta=(EditCondition="CastShadow", DisplayName = "Capsule Indirect Shadow"))
+	UPROPERTY(EditAnywhere, AdvancedDisplay, BlueprintReadOnly, Category=Lighting, meta=(EditCondition="CastShadow && (FirstPersonPrimitiveType != EFirstPersonPrimitiveType::FirstPerson)", DisplayName = "Capsule Indirect Shadow"))
 	uint8 bCastCapsuleIndirectShadow:1;
 
 	/** Whether or not to CPU skin this component, requires render data refresh after changing */
diff --git a/Engine/Source/Runtime/Engine/Private/Components/PrimitiveComponent.cpp b/Engine/Source/Runtime/Engine/Private/Components/PrimitiveComponent.cpp
--- a/Engine/Source/Runtime/Engine/Private/Components/PrimitiveComponent.cpp
+++ b/Engine/Source/Runtime/Engine/Private/Components/PrimitiveComponent.cpp
@@ -1460,14 +1460,10 @@ bool UPrimitiveComponent::CanEditChange(const FProperty* InProperty) const
 		const FName PropertyName = InProperty->GetFName();
 
 		static FName LightmassSettingsName = TEXT("LightmassSettings");
-		static FName LightingChannelsName = GET_MEMBER_NAME_CHECKED(UPrimitiveComponent, LightingChannels);
 		static FName SingleSampleShadowFromStationaryLightsName = GET_MEMBER_NAME_CHECKED(UPrimitiveComponent, bSingleSampleShadowFromStationaryLights);
 		static FName IndirectLightingCacheQualityName = GET_MEMBER_NAME_CHECKED(UPrimitiveComponent, IndirectLightingCacheQuality);
 		static FName CastCinematicShadowName = GET_MEMBER_NAME_CHECKED(UPrimitiveComponent, bCastCinematicShadow);
 		static FName CastInsetShadowName = GET_MEMBER_NAME_CHECKED(UPrimitiveComponent, bCastInsetShadow);
-		static FName CastShadowName = GET_MEMBER_NAME_CHECKED(UPrimitiveComponent, CastShadow);
-		static FName CastHiddenShadowName = GET_MEMBER_NAME_CHECKED(UPrimitiveComponent, bCastHiddenShadow);
-		static FName OwnerNoSeeName = GET_MEMBER_NAME_CHECKED(UPrimitiveComponent, bOwnerNoSee);
 
 		if (PropertyName == LightmassSettingsName)
 		{
@@ -1504,13 +1500,6 @@ bool UPrimitiveComponent::CanEditChange(const FProperty* InProperty) const
 		{
 			return !bSelfShadowOnly;
 		}
-
-		// FirstPersonWorldSpaceRepresentation primitives must have bCastHiddenShadows=false and bOwnerNoSee=true to work correctly with VSM.
-		// Other setups are not expected to be needed, so to avoid confusion, we'll gray out these properties and simply override them on the scene proxy.
-		if (PropertyName == CastHiddenShadowName || PropertyName == OwnerNoSeeName)
-		{
-			return FirstPersonPrimitiveType != EFirstPersonPrimitiveType::WorldSpaceRepresentation;
-		}
 	}
 
 	return bIsEditable;
diff --git a/Engine/Source/Runtime/Engine/Private/Components/StaticMeshComponent.cpp b/Engine/Source/Runtime/Engine/Private/Components/StaticMeshComponent.cpp
--- a/Engine/Source/Runtime/Engine/Private/Components/StaticMeshComponent.cpp
+++ b/Engine/Source/Runtime/Engine/Private/Components/StaticMeshComponent.cpp
@@ -2107,12 +2107,12 @@ bool UStaticMeshComponent::CanEditChange(const FProperty* InProperty) const
 
 		if (PropertyName == GET_MEMBER_NAME_STRING_CHECKED(UStaticMeshComponent, bCastDistanceFieldIndirectShadow))
 		{
-			return Mobility != EComponentMobility::Static && CastShadow && bCastDynamicShadow;
+			return Mobility != EComponentMobility::Static && CastShadow && bCastDynamicShadow && FirstPersonPrimitiveType != EFirstPersonPrimitiveType::FirstPerson;
 		}
 
 		if (PropertyName == GET_MEMBER_NAME_STRING_CHECKED(UStaticMeshComponent, DistanceFieldIndirectShadowMinVisibility))
 		{
-			return Mobility != EComponentMobility::Static && bCastDistanceFieldIndirectShadow && CastShadow && bCastDynamicShadow;
+			return Mobility != EComponentMobility::Static && bCastDistanceFieldIndirectShadow && CastShadow && bCastDynamicShadow && FirstPersonPrimitiveType != EFirstPersonPrimitiveType::FirstPerson;
 		}
 
 		if (PropertyName == GET_MEMBER_NAME_STRING_CHECKED(UStaticMeshComponent, bOverrideDistanceFieldSelfShadowBias))
diff --git a/Engine/Source/Runtime/Engine/Private/InstancedStaticMesh.cpp b/Engine/Source/Runtime/Engine/Private/InstancedStaticMesh.cpp
--- a/Engine/Source/Runtime/Engine/Private/InstancedStaticMesh.cpp
+++ b/Engine/Source/Runtime/Engine/Private/InstancedStaticMesh.cpp
@@ -4696,7 +4696,7 @@ static bool ComponentRequestsCPUAccess(UInstancedStaticMeshComponent* InComponen
 			auto ShaderPlatform = GetFeatureLevelShaderPlatform(FeatureLevel);
 			bNeedsCPUAccess |= PrimitiveNeedsDistanceFieldSceneData(
 				ShouldAllPrimitivesHaveDistanceField(ShaderPlatform),
-				/* bCastsDynamicIndirectShadow */ InComponent->bCastDynamicShadow && InComponent->CastShadow && InComponent->bCastDistanceFieldIndirectShadow && InComponent->Mobility != EComponentMobility::Static,
+				/* bCastsDynamicIndirectShadow */ InComponent->bCastDynamicShadow && InComponent->CastShadow && InComponent->bCastDistanceFieldIndirectShadow && InComponent->Mobility != EComponentMobility::Static && InComponent->FirstPersonPrimitiveType != EFirstPersonPrimitiveType::FirstPerson,
 				InComponent->bAffectDistanceFieldLighting,
 				true, /* conservatively overestimate DrawInGame - it has complex logic in the Proxy. */
 				InComponent->bCastHiddenShadow,
diff --git a/Engine/Source/Runtime/Engine/Private/PrimitiveSceneProxy.cpp b/Engine/Source/Runtime/Engine/Private/PrimitiveSceneProxy.cpp
--- a/Engine/Source/Runtime/Engine/Private/PrimitiveSceneProxy.cpp
+++ b/Engine/Source/Runtime/Engine/Private/PrimitiveSceneProxy.cpp
@@ -342,7 +342,7 @@ void FPrimitiveSceneProxyDesc::InitializeFromPrimitiveComponent(const UPrimitive
 		bIsOwnedByFoliage = FFoliageHelper::IsOwnedByFoliage(ActorOwner);
 #endif
 
-		if(bOnlyOwnerSee || bOwnerNoSee || bUseViewOwnerDepthPriorityGroup)
+		if(bOnlyOwnerSee || bOwnerNoSee || bUseViewOwnerDepthPriorityGroup || bIsFirstPersonWorldSpaceRepresentation)
 		{
 			// Make a list of the actors which directly or indirectly own the InComponent.
 			for(const AActor* CurrentOwner = ActorOwner;CurrentOwner;CurrentOwner = CurrentOwner->GetOwner())
@@ -614,7 +614,7 @@ FPrimitiveSceneProxy::FPrimitiveSceneProxy(const FPrimitiveSceneProxyDesc& InPro
 			DrawInEditor &= !InProxyDesc.IsHiddenEd();
 		#endif
 
-		if(bOnlyOwnerSee || bOwnerNoSee || bUseViewOwnerDepthPriorityGroup)
+		if(bOnlyOwnerSee || bOwnerNoSee || bUseViewOwnerDepthPriorityGroup || bIsFirstPersonWorldSpaceRepresentation)
 		{
 			Owners = MoveTemp(InProxyDesc.ActorOwners);
 		}
@@ -719,13 +719,24 @@ FPrimitiveSceneProxy::FPrimitiveSceneProxy(const FPrimitiveSceneProxyDesc& InPro
 	}
 #endif
 
-	// FirstPersonWorldSpaceRepresentation primitives must have bCastHiddenShadow=false and bOwnerNoSee=true.
+	// FirstPerson primitives won't properly work with most shadow solutions, so just disable them here to give a smoother user experience out of the box.
+	// Such primitives are intended to be used in conjunction with FirstPersonWorldSpaceRepresentation primitives, which are used to cast shadow onto the scene.
+	// Self shadowing for FirstPerson primitives is solved with First Person Self Shadows, which are currently implemented with screen space ray marching.
+	// Note that bCastContactShadow and a few of the legacy shadow cast flags are not overridden here. This is intentional as there might be valid use cases for them in this context.
+	if (bIsFirstPerson)
+	{
+		bCastDynamicShadow = false;
+		bCastStaticShadow = false;
+		bCastVolumetricTranslucentShadow = false;
+		bCastFarShadow = false;
+	}
+	// FirstPersonWorldSpaceRepresentation primitives must have bOwnerNoSee=true and bCastHiddenShadow=false.
 	// This way, they are skipped by regular shadow rendering and can be picked up by the special first person
 	// shadow path in VSM. Raytraced shadows ignore these flags anyways, so this works for all cases.
-	if (bIsFirstPersonWorldSpaceRepresentation)
+	else if (bIsFirstPersonWorldSpaceRepresentation)
 	{
-		bCastHiddenShadow = false;
 		bOwnerNoSee = true;
+		bCastHiddenShadow = false;
 	}
 }
 
diff --git a/Engine/Source/Runtime/Engine/Private/SkeletalMeshSceneProxy.cpp b/Engine/Source/Runtime/Engine/Private/SkeletalMeshSceneProxy.cpp
--- a/Engine/Source/Runtime/Engine/Private/SkeletalMeshSceneProxy.cpp
+++ b/Engine/Source/Runtime/Engine/Private/SkeletalMeshSceneProxy.cpp
@@ -119,8 +119,8 @@ FSkeletalMeshSceneProxy::FSkeletalMeshSceneProxy(const FSkinnedMeshSceneProxyDes
 
 	bIsCPUSkinned = MeshObject->IsCPUSkinned();
 
-	bCastCapsuleDirectShadow = InMeshDesc.bCastDynamicShadow && InMeshDesc.CastShadow && InMeshDesc.bCastCapsuleDirectShadow;
-	bCastsDynamicIndirectShadow = InMeshDesc.bCastDynamicShadow && InMeshDesc.CastShadow && InMeshDesc.bCastCapsuleIndirectShadow;
+	bCastCapsuleDirectShadow = InMeshDesc.bCastDynamicShadow && InMeshDesc.CastShadow && InMeshDesc.bCastCapsuleDirectShadow && !InMeshDesc.bIsFirstPerson;
+	bCastsDynamicIndirectShadow = InMeshDesc.bCastDynamicShadow && InMeshDesc.CastShadow && InMeshDesc.bCastCapsuleIndirectShadow && !InMeshDesc.bIsFirstPerson;
 
 	DynamicIndirectShadowMinVisibility = FMath::Clamp(InMeshDesc.CapsuleIndirectShadowMinVisibility, 0.0f, 1.0f);
 
diff --git a/Engine/Source/Runtime/Engine/Private/StaticMeshSceneProxy.cpp b/Engine/Source/Runtime/Engine/Private/StaticMeshSceneProxy.cpp
--- a/Engine/Source/Runtime/Engine/Private/StaticMeshSceneProxy.cpp
+++ b/Engine/Source/Runtime/Engine/Private/StaticMeshSceneProxy.cpp
@@ -322,7 +322,7 @@ FStaticMeshSceneProxy::FStaticMeshSceneProxy(const FStaticMeshSceneProxyDesc& In
 	CardRepresentationData = RenderData->LODResources[0].CardRepresentationData;
 
 	bSupportsDistanceFieldRepresentation = MaterialRelevance.bOpaque && !MaterialRelevance.bUsesSingleLayerWaterMaterial && DistanceFieldData && DistanceFieldData->IsValid();
-	bCastsDynamicIndirectShadow = InProxyDesc.bCastDynamicShadow && InProxyDesc.CastShadow && InProxyDesc.bCastDistanceFieldIndirectShadow && InProxyDesc.Mobility != EComponentMobility::Static;
+	bCastsDynamicIndirectShadow = InProxyDesc.bCastDynamicShadow && InProxyDesc.CastShadow && InProxyDesc.bCastDistanceFieldIndirectShadow && InProxyDesc.Mobility != EComponentMobility::Static && !InProxyDesc.bIsFirstPerson;
 	DynamicIndirectShadowMinVisibility = FMath::Clamp(InProxyDesc.DistanceFieldIndirectShadowMinVisibility, 0.0f, 1.0f);
 	DistanceFieldSelfShadowBias = FMath::Max(InProxyDesc.bOverrideDistanceFieldSelfShadowBias ? InProxyDesc.DistanceFieldSelfShadowBias : InProxyDesc.GetStaticMesh()->DistanceFieldSelfShadowBias, 0.0f);
 
