Don't bind the virtual texture editor resource RHI texture to the UTexture::TextureReference.
The TextureReference can be left pointing at DefaultTexture.
The editor resource texture is only used for thumbnail rendering and for that it is directly retrieved through the resource.
Fixes OOM in bindless heap when batch recreating virtual textures.

diff --git a/Engine/Source/Runtime/Engine/Private/Texture2D.cpp b/Engine/Source/Runtime/Engine/Private/Texture2D.cpp
--- a/Engine/Source/Runtime/Engine/Private/Texture2D.cpp
+++ b/Engine/Source/Runtime/Engine/Private/Texture2D.cpp
@@ -1779,7 +1779,6 @@ void FVirtualTexture2DResource::InitializeEditorResources(IVirtualTexture* InVir
 		TextureRHI = Texture2DRHI;
 		TextureRHI->SetName(TextureName);
 		RHIBindDebugLabelName(TextureRHI, *Name);
-		RHIUpdateTextureReference(TextureReferenceRHI, TextureRHI);
 	}
 }
 #endif // WITH_EDITOR
