Fixed duplicated ActorFolders fixup

diff --git a/Engine/Source/Runtime/Engine/Private/Level.cpp b/Engine/Source/Runtime/Engine/Private/Level.cpp
--- a/Engine/Source/Runtime/Engine/Private/Level.cpp
+++ b/Engine/Source/Runtime/Engine/Private/Level.cpp
@@ -3307,7 +3307,7 @@ void ULevel::FixupActorFolders()
 
 			SortedDuplicatePaths.Sort([](const FName& FolderPathA, const FName& FolderPathB)
 			{
-				return FolderPathA.LexicalLess(FolderPathB);
+				return !FolderPathA.LexicalLess(FolderPathB);
 			});
 
 			TSet<UActorFolder*> FoldersToDelete;
@@ -3326,33 +3326,59 @@ void ULevel::FixupActorFolders()
 			};
 
 			TSet<UActorFolder*> FoldersToKeep;
-			TMap<UActorFolder*, UActorFolder*> DuplicateFolders;
-
-			for (FName& DuplicatePath : SortedDuplicatePaths)
+			auto FindFolderToKeep = [&FoldersToKeep, HasParentToDelete](TArray<UActorFolder*>& Folders)
 			{
-				// Choose to keep the folder that has no parent to delete
-				TArray<UActorFolder*>& Folders = PathToFolders.FindChecked(DuplicatePath);
 				UActorFolder* FolderToKeep = nullptr;
+ 				// Check if folder is already marked to keep by one of its children
 				for (UActorFolder* Folder : Folders)
 				{
-					if (!HasParentToDelete(Folder))
+					if (FoldersToKeep.Contains(Folder))
 					{
 						FolderToKeep = Folder;
 						break;
 					}
 				}
+				if (!FolderToKeep)
+				{
+					for (UActorFolder* Folder : Folders)
+					{
+						if (!HasParentToDelete(Folder))
+						{
+							FolderToKeep = Folder;
+							break;
+						}
+					}
+				}
+				return FolderToKeep;
+			};
 
-				// Validation
+			TMap<UActorFolder*, UActorFolder*> DuplicateFolders;
+			for (FName& DuplicatePath : SortedDuplicatePaths)
+			{
+				// Choose to keep the folder that has no parent to delete
+				TArray<UActorFolder*>& Folders = PathToFolders.FindChecked(DuplicatePath);
+				UActorFolder* FolderToKeep = FindFolderToKeep(Folders);
 				check(FolderToKeep);
-				bool bAlreadyExists = false;
-				FoldersToKeep.Add(FolderToKeep, &bAlreadyExists);
-				check(!bAlreadyExists);
+
+				// Mark folder to keep
+				check(!FoldersToDelete.Contains(FolderToKeep));
+				FoldersToKeep.Add(FolderToKeep);
+
+				// Also mark to keep the whole parent chain
+				UActorFolder* Parent = FolderToKeep->GetParent();
+				while (Parent)
+				{
+					check(!FoldersToDelete.Contains(Parent));
+					FoldersToKeep.Add(Parent);
+					Parent = Parent->GetParent();
+				}
 
 				// All other folders are considered duplicated and will be marked as deleted
 				for (UActorFolder* Folder : Folders)
 				{
 					if (Folder != FolderToKeep)
 					{
+						check(!FoldersToKeep.Contains(Folder));
 						FoldersToDelete.Add(Folder);
 
 						// Cleanup FolderLabelToActorFolders from duplicate folder as we don't want GetActorFolder(Path) to resolve to a duplicate ActorFolder
