Fix material function conversion with new EMaterialSubSurfaceType.

diff --git a/Engine/Source/Runtime/Engine/Private/Materials/Material.cpp b/Engine/Source/Runtime/Engine/Private/Materials/Material.cpp
--- a/Engine/Source/Runtime/Engine/Private/Materials/Material.cpp
+++ b/Engine/Source/Runtime/Engine/Private/Materials/Material.cpp
@@ -4204,6 +4204,42 @@ bool UMaterial::ConvertMaterialToSubstrateMaterial(bool bAllowEmptyMaterialUpdat
 
 	return bInvalidateShader;
 }
+
+static void ConvertSlabExpressionMaterialSubSurfaceType(TConstArrayView<TObjectPtr<UMaterialExpression>> InExpressions, EBlendMode InBlendMode, bool bInIsThinSurface)
+{
+	for (UMaterialExpression* Expression : InExpressions)
+	{
+		if (UMaterialExpressionSubstrateSlabBSDF* SlabExpression = Cast<UMaterialExpressionSubstrateSlabBSDF>(Expression))
+		{
+			PRAGMA_DISABLE_DEPRECATION_WARNINGS
+			if (!IsOpaqueBlendMode(InBlendMode))
+			{
+				SlabExpression->SubSurfaceType = EMaterialSubSurfaceType::MSS_SimpleVolume;
+			}
+			else if (SlabExpression->bUseSSSDiffusion)
+			{
+				SlabExpression->SubSurfaceType = EMaterialSubSurfaceType::MSS_Diffusion;
+			}
+			else if (bInIsThinSurface)
+			{
+				SlabExpression->SubSurfaceType = EMaterialSubSurfaceType::MSS_TwoSidedWrap;
+			}
+			else
+			{
+				SlabExpression->SubSurfaceType = EMaterialSubSurfaceType::MSS_Wrap;
+			}
+			PRAGMA_ENABLE_DEPRECATION_WARNINGS
+		}
+		else if (UMaterialExpressionMaterialFunctionCall* FunctionCall = Cast<UMaterialExpressionMaterialFunctionCall>(Expression))
+		{
+			if (FunctionCall->MaterialFunction)
+			{
+				ConvertSlabExpressionMaterialSubSurfaceType(FunctionCall->MaterialFunction->GetExpressions(), InBlendMode, bInIsThinSurface);
+			}
+		}
+	}
+}
+
 #endif // WITH_EDITOR
 
 TMap<FGuid, UMaterialInterface*> LightingGuidFixupMap;
@@ -4506,26 +4542,7 @@ void UMaterial::PostLoad()
 	{
 		if (FortMainVer < FFortniteMainBranchObjectVersion::AddMaterialSubstrateSubsurfaceType)
 		{
-			for (UMaterialExpression* Expression : GetExpressions())
-			{
-				if (UMaterialExpressionSubstrateSlabBSDF* SlabExpression = Cast<UMaterialExpressionSubstrateSlabBSDF>(Expression))
-				{
-					PRAGMA_DISABLE_DEPRECATION_WARNINGS
-					if (SlabExpression->bUseSSSDiffusion)
-					{
-						SlabExpression->SubSurfaceType = EMaterialSubSurfaceType::MSS_Diffusion;
-					}
-					else if (bIsThinSurface)
-					{
-						SlabExpression->SubSurfaceType = EMaterialSubSurfaceType::MSS_TwoSidedWrap;
-					}
-					else
-					{
-						SlabExpression->SubSurfaceType = EMaterialSubSurfaceType::MSS_Wrap;
-					}
-					PRAGMA_ENABLE_DEPRECATION_WARNINGS
-				}
-			}
+			ConvertSlabExpressionMaterialSubSurfaceType(GetExpressions(), BlendMode, bIsThinSurface);
 		}
 
 		// Substrate materials conversion needs to be done after expressions are cached, otherwise material function won't have 
