[WorldPartitionLevelHelper] Don't attempt to unload packages while AsyncLoading.

WorldPartitionlevelHelper releases the reference and adds packages to PreGCPackagesToUnload in FinalizeRuntimeLevel, right after the loading is done.
However, AsyncLoading2 also tracks package references using GlobalImportStore. Package references are removed from the GlobalImportStore when FAsyncPackage2 is deleted, which can happen after PreGarbageCollect, due to usage of DeferredDeletePackages queue.
If we get another request, which involves a package that has already been trashed (via FWorldPartitionPackageHelper::UnloadPackage) but not yet removed from the GlobalImportStore, AsyncLoading2 will attempt to reuse that package. Since it has already been trashed at that point, it'll lead to undesired behavior.
To prevent this from happening, don't attempt to unload packages while AsyncLoading - IsAsyncLoading returns true until all AsyncPackages have not been deleted - see PackagesWithRemainingWorkCounter and DeleteAsyncPackage

diff --git a/Engine/Source/Runtime/Engine/Private/WorldPartition/WorldPartitionLevelHelper.cpp b/Engine/Source/Runtime/Engine/Private/WorldPartition/WorldPartitionLevelHelper.cpp
--- a/Engine/Source/Runtime/Engine/Private/WorldPartition/WorldPartitionLevelHelper.cpp
+++ b/Engine/Source/Runtime/Engine/Private/WorldPartition/WorldPartitionLevelHelper.cpp
@@ -87,15 +87,26 @@ FWorldPartitionLevelHelper::FWorldPartitionLevelHelper()
 
 void FWorldPartitionLevelHelper::PreGarbageCollect()
 {
-	for (TWeakObjectPtr<UPackage>& PackageToUnload : PreGCPackagesToUnload)
+	// Don't attempt to unload packages while AsyncLoading.
+	//   
+	// WorldPartitionlevelHelper releases the reference and adds packages to PreGCPackagesToUnload in FinalizeRuntimeLevel, right after the loading is done.
+	// However, AsyncLoading2 also tracks package references using GlobalImportStore. Package references are removed from the GlobalImportStore when FAsyncPackage2 is deleted,
+	// which can happen after PreGarbageCollect, due to usage of DeferredDeletePackages queue.
+	// If we get another request, which involves a package that has already been trashed (via FWorldPartitionPackageHelper::UnloadPackage)
+	// but not yet removed from the GlobalImportStore, AsyncLoading2 will attempt to reuse that package. Since it has already been trashed at that point, it'll lead to undesired behavior.
+	// To prevent this from happening, don't attempt to unload packages while AsyncLoading - IsAsyncLoading returns true until all AsyncPackages have not been deleted - see PackagesWithRemainingWorkCounter and DeleteAsyncPackage
+	if (!IsAsyncLoading())
 	{
-		// Test if WeakObjectPtr is valid since clean up could have happened outside of this helper
-		if (PackageToUnload.IsValid())
+		for (TWeakObjectPtr<UPackage>& PackageToUnload : PreGCPackagesToUnload)
 		{
-			FWorldPartitionPackageHelper::UnloadPackage(PackageToUnload.Get());
+			// Test if WeakObjectPtr is valid since clean up could have happened outside of this helper
+			if (PackageToUnload.IsValid())
+			{
+				FWorldPartitionPackageHelper::UnloadPackage(PackageToUnload.Get());
+			}
 		}
+		PreGCPackagesToUnload.Reset();
 	}
-	PreGCPackagesToUnload.Reset();
 }
 
 void FWorldPartitionLevelHelper::ApplyConstructionScriptPropertyOverridesFromAnnotation(AActor* InActor)
