Updated physics object ignoring to properly handle remote objects. This used to check if the object was on a remote server, but now checks if the object is owned by this server.

diff --git a/Engine/Source/Runtime/CoreUObject/Private/UObject/RemoteObjectTransfer.cpp b/Engine/Source/Runtime/CoreUObject/Private/UObject/RemoteObjectTransfer.cpp
--- a/Engine/Source/Runtime/CoreUObject/Private/UObject/RemoteObjectTransfer.cpp
+++ b/Engine/Source/Runtime/CoreUObject/Private/UObject/RemoteObjectTransfer.cpp
@@ -516,6 +516,10 @@ void MigrateObjectToRemoteServerWithExplicitPriority(FRemoteTransactionId Reques
 				SentObject->PostMigrate(PostMigrateParams);
 #endif
 				MarkAsRemote(SentObject, DestinationServerId);
+				if (bOwnershipChange)
+				{
+					UE::RemoteObject::Handle::ChangeOwnerServerId(SentObject, OwnerServerId);
+				}
 			}
 		}
 	}
diff --git a/Engine/Source/Runtime/Engine/Private/Collision/Collision.cpp b/Engine/Source/Runtime/Engine/Private/Collision/Collision.cpp
--- a/Engine/Source/Runtime/Engine/Private/Collision/Collision.cpp
+++ b/Engine/Source/Runtime/Engine/Private/Collision/Collision.cpp
@@ -108,18 +108,32 @@ void FCollisionQueryParams::Internal_AddIgnoredSourceObject(const FWeakObjectPtr
 {
 	if (WeakPtr.IsValid())
 	{
-		if (!WeakPtr.IsRemote())
+#if UE_WITH_REMOTE_OBJECT_HANDLE
+		// We only want to add the object to the ignore list if it's owned by the current server. 
+		// Even if it's owned, it might currently be remote which would cause a migration upon dereference. 
+		// We can grab the object directly, even if it's stale, to get the unique id and avoid a migration. 
+		// This unique id will always be valid to use.
+		if (UE::RemoteObject::Handle::IsOwned(WeakPtr.GetRemoteId()))
 		{
-			UObject* Obj = WeakPtr.Get();
-			// On rare occasions, this is null even with the above `IsValid` check
-			if (Obj != nullptr)
-			{
-				AddIgnoredSourceObject(WeakPtr.Get()->GetUniqueID());
-			}
-		}
+			constexpr bool bEvenIfGarbage = true;
+			constexpr bool bOnGameThread = true;
+			check(WeakPtr.IsValid(bEvenIfGarbage, bOnGameThread)); // this should always be true if the object is owned
 
-#if UE_WITH_REMOTE_OBJECT_HANDLE
+			// This is the UObject pointer, even if it's stale (e.g. contains some invalid data)
+			const UObject* DoNotMutateLocalCopy = StaticFindObjectFastInternal(WeakPtr.GetRemoteId());
+
+			// This is the array index in GUObjectArray which is the same as its UniqueId.
+			const int32 UniqueId = GUObjectArray.ObjectToIndex(DoNotMutateLocalCopy);
+			AddIgnoredSourceObject(UniqueId);
+		}
 		IgnoreSourceObjectReferences.Add(FRemoteObjectReference(WeakPtr));
+#else
+		UObject * Obj = WeakPtr.Get();
+		// On rare occasions, this is null even with the above `IsValid` check
+		if (Obj != nullptr)
+		{
+			AddIgnoredSourceObject(Obj->GetUniqueID());
+		}
 #endif
 	}
 }
@@ -237,15 +251,19 @@ void FCollisionQueryParams::AddIgnoredComponent(const TWeakObjectPtr<UPrimitiveC
 {
 	if (InIgnoreComponent.IsValid())
 	{
-		// Only add valid local objects, otherwise migration could happen
-		if (!InIgnoreComponent.IsRemote())
-		{
-			Internal_AddIgnoredComponent(InIgnoreComponent.Get());
-		}
+		bool bIsValidIgnore = true;
 
 #if UE_WITH_REMOTE_OBJECT_HANDLE
-		IgnoreComponentReferences.Add(FRemoteObjectReference(InIgnoreComponent));
+		FRemoteObjectReference RemoteObjectRef(InIgnoreComponent);
+		IgnoreComponentReferences.Add(RemoteObjectRef);
+		// Only ignore objects that belong to this server's physics world (only true if owned).
+		bIsValidIgnore = UE::RemoteObject::Handle::IsOwned(RemoteObjectRef.GetRemoteId());
 #endif
+
+		if (bIsValidIgnore)
+		{
+			Internal_AddIgnoredComponent(InIgnoreComponent.Get());
+		}
 	}
 }
 
@@ -254,7 +272,13 @@ void FCollisionQueryParams::AddIgnoredComponent(const FRemoteObjectReference& In
 	FWeakObjectPtr WeakPtr = InComponentRef.ToWeakPtr();
 	if (WeakPtr.IsValid())
 	{
-		if (!WeakPtr.IsRemote())
+		bool bIsValidIgnore = true;
+#if UE_WITH_REMOTE_OBJECT_HANDLE
+		IgnoreComponentReferences.Add(InComponentRef);
+		// Only ignore objects that belong to this server's physics world (only true if owned).
+		bIsValidIgnore = UE::RemoteObject::Handle::IsOwned(InComponentRef.GetRemoteId());
+#endif
+		if (bIsValidIgnore)
 		{
 			UObject* Obj = WeakPtr.Get();
 			// On rare occasions, this is null even with the above `IsValid` check
@@ -264,12 +288,7 @@ void FCollisionQueryParams::AddIgnoredComponent(const FRemoteObjectReference& In
 				check(PrimitiveComponent);
 				Internal_AddIgnoredComponent(PrimitiveComponent);
 			}
-			
 		}
-
-#if UE_WITH_REMOTE_OBJECT_HANDLE
-		IgnoreComponentReferences.Add(InComponentRef);
-#endif
 	}
 }
 
