Fix potential crash when loading maps with existing FImportanceTexture objects

The old data could be missing the new LinearTextureData field, so avoid crashing until it can be rebuilt.

#rnx
#tests Tested on QAGame/Rendering.TM-SobolNoise

diff --git a/Engine/Source/Runtime/Engine/Private/Kismet/ImportanceSamplingLibrary.cpp b/Engine/Source/Runtime/Engine/Private/Kismet/ImportanceSamplingLibrary.cpp
--- a/Engine/Source/Runtime/Engine/Private/Kismet/ImportanceSamplingLibrary.cpp
+++ b/Engine/Source/Runtime/Engine/Private/Kismet/ImportanceSamplingLibrary.cpp
@@ -49,6 +49,10 @@ FLinearColor FImportanceTexture::GetColorTrilinear(FVector2D Position, float Mip
 
 FLinearColor FImportanceTexture::GetColorBilinear(FVector2D Position, int32 Mip) const
 {
+	if (LinearTextureData.IsEmpty())
+	{
+		return FLinearColor(EForceInit());
+	}
 	int32 iMip = FMath::Clamp(Mip, 0, NumMips - 1);
 	FIntPoint MipSize(((Size.X - 1) >> iMip) + 1, ((Size.Y - 1) >> iMip) + 1);
 	size_t LevelStart = 4 * (Size.X * Size.Y - MipSize.X * MipSize.Y) / 3;
