Fix cached material instance data not being migrated for remote objects

diff --git a/Engine/Source/Runtime/Engine/Private/Materials/MaterialInstance.cpp b/Engine/Source/Runtime/Engine/Private/Materials/MaterialInstance.cpp
--- a/Engine/Source/Runtime/Engine/Private/Materials/MaterialInstance.cpp
+++ b/Engine/Source/Runtime/Engine/Private/Materials/MaterialInstance.cpp
@@ -3111,7 +3111,8 @@ void UMaterialInstance::Serialize(FArchive& Ar)
 		// we want to save the cached data when cooking or duplicating the object in a cooked game
 		const bool bWantToSaveCachedData = Ar.IsCooking();
 #else
-		const bool bDuplicatingObjectInACookedGame = FPlatformProperties::RequiresCookedData() && (Ar.GetPortFlags() & PPF_Duplicate);
+		// We want to copy the cached data when the material is duplicated either directly or during remote object migration
+		const bool bDuplicatingObjectInACookedGame = FPlatformProperties::RequiresCookedData() && (Ar.GetPortFlags() & (PPF_Duplicate | PPF_MigratingRemoteObjects));
 		const bool bWantToSaveCachedData = Ar.IsSaving() && bDuplicatingObjectInACookedGame;
 
 		// Workaround for materials being created by annotation data at runtime not having CachedData
