Substrate - fixed input name when using clear coat model on legacy shading node.

#tests PC editor
#rnx

diff --git a/Engine/Source/Runtime/Engine/Private/Materials/MaterialExpressionSubstrate.cpp b/Engine/Source/Runtime/Engine/Private/Materials/MaterialExpressionSubstrate.cpp
--- a/Engine/Source/Runtime/Engine/Private/Materials/MaterialExpressionSubstrate.cpp
+++ b/Engine/Source/Runtime/Engine/Private/Materials/MaterialExpressionSubstrate.cpp
@@ -504,7 +504,18 @@ FName UMaterialExpressionSubstrateShadingModels::GetInputName(int32 InputIndex)
 	else if (InputIndex == 14)	return TEXT("Water Absorption Coefficients");
 	else if (InputIndex == 15)	return TEXT("Water Phase G");
 	else if (InputIndex == 16)	return TEXT("Color Scale BehindWater");
-	else if (InputIndex == 17)	return TEXT("Iris Normal");
+	else if (InputIndex == 17)
+	{
+		if (!bShadingModelFromExpression && ShadingModelOverride == MSM_ClearCoat)
+		{
+			return TEXT("Clear Coat Bottom Normal");
+		}
+		else if (!bShadingModelFromExpression && ShadingModelOverride == MSM_Eye)
+		{
+			return TEXT("Iris Normal");
+		}
+		return TEXT("Unused");
+	}
 	else if (InputIndex == 18)
 	{
 		if (!bShadingModelFromExpression && ShadingModelOverride == MSM_Eye)
