Substrate - added bottom norml map to the simple clear coat node.

#tests PC editor dx12
#rnx

diff --git a/Engine/Source/Runtime/Engine/Private/Materials/HLSLMaterialTranslator.cpp b/Engine/Source/Runtime/Engine/Private/Materials/HLSLMaterialTranslator.cpp
--- a/Engine/Source/Runtime/Engine/Private/Materials/HLSLMaterialTranslator.cpp
+++ b/Engine/Source/Runtime/Engine/Private/Materials/HLSLMaterialTranslator.cpp
@@ -13923,7 +13923,7 @@ int32 FHLSLMaterialTranslator::SubstrateSlabBSDF(
 	int32 Roughness, int32 Anisotropy,
 	int32 SSSProfileId, int32 SSSMFP, int32 SSSMFPScale, int32 SSSPhaseAniso, int32 SSSType,
 	int32 EmissiveColor,
-	int32 SecondRoughness, int32 SecondRoughnessWeight, int32 SecondRoughnessAsSimpleClearCoat,
+	int32 SecondRoughness, int32 SecondRoughnessWeight, int32 SecondRoughnessAsSimpleClearCoat, int32 ClearCoatBottomNormal,
 	int32 FuzzAmount, int32 FuzzColor, int32 FuzzRoughness,
 	int32 Thickness,
 	int32 GlintValue, int32 GlintUV,
@@ -13937,6 +13937,8 @@ int32 FHLSLMaterialTranslator::SubstrateSlabBSDF(
 	const FString ThicknessCode = GetParameterCode(Thickness);
 	const bool bIsThinSurface = Material->IsThinSurface();
 
+	int32 ClearCoatUseSecondNormal = Constant(ClearCoatBottomNormal != Normal ? 1.0f : 0.0f);
+
 	if (PromoteToOperator)
 	{
 		if (PromoteToOperator->Index == INDEX_NONE || PromoteToOperator->BSDFIndex == INDEX_NONE)
@@ -13992,8 +13994,8 @@ int32 FHLSLMaterialTranslator::SubstrateSlabBSDF(
 			*SubstrateGetCastParameterCode(SecondRoughness,			MCT_Float),
 			*SubstrateGetCastParameterCode(SecondRoughnessWeight,	MCT_Float),
 			*SubstrateGetCastParameterCode(SecondRoughnessAsSimpleClearCoat, MCT_Float),
-			*SubstrateGetCastParameterCode(Constant(0.0f),			MCT_Float),
-			*SubstrateGetCastParameterCode(Constant3(0.0f, 0.0f, 0.0f), MCT_Float3),
+			*SubstrateGetCastParameterCode(ClearCoatUseSecondNormal,MCT_Float),
+			*SubstrateGetCastParameterCode(ClearCoatBottomNormal,	MCT_Float3),
 			*SubstrateGetCastParameterCode(FuzzAmount,				MCT_Float),
 			*SubstrateGetCastParameterCode(FuzzColor,				MCT_Float3),
 			*SubstrateGetCastParameterCode(FuzzRoughness,			MCT_Float1),
diff --git a/Engine/Source/Runtime/Engine/Private/Materials/HLSLMaterialTranslator.h b/Engine/Source/Runtime/Engine/Private/Materials/HLSLMaterialTranslator.h
--- a/Engine/Source/Runtime/Engine/Private/Materials/HLSLMaterialTranslator.h
+++ b/Engine/Source/Runtime/Engine/Private/Materials/HLSLMaterialTranslator.h
@@ -1316,7 +1316,7 @@ protected:
 		int32 Roughness, int32 Anisotropy,
 		int32 SSSProfileId, int32 SSSMFP, int32 SSSMFPScale, int32 SSSPhaseAniso, int32 SSSType,
 		int32 EmissiveColor, 
-		int32 SecondRoughness, int32 SecondRoughnessWeight, int32 SecondRoughnessAsSimpleClearCoat,
+		int32 SecondRoughness, int32 SecondRoughnessWeight, int32 SecondRoughnessAsSimpleClearCoat, int32 ClearCoatBottomNormal,
 		int32 FuzzAmount, int32 FuzzColor, int32 FuzzRoughness,
 		int32 Thickness,
 		int32 GlintValue, int32 GlintUV,
diff --git a/Engine/Source/Runtime/Engine/Private/Materials/MaterialExpressionSubstrate.cpp b/Engine/Source/Runtime/Engine/Private/Materials/MaterialExpressionSubstrate.cpp
--- a/Engine/Source/Runtime/Engine/Private/Materials/MaterialExpressionSubstrate.cpp
+++ b/Engine/Source/Runtime/Engine/Private/Materials/MaterialExpressionSubstrate.cpp
@@ -884,6 +884,7 @@ int32 UMaterialExpressionSubstrateSlabBSDF::Compile(class FMaterialCompiler* Com
 		SecondRoughnessCodeChunk,
 		SecondRoughnessWeightCodeChunk,
 		Compiler->Constant(0.0f),										// SecondRoughnessAsSimpleClearCoat
+		NormalCodeChunk,
 		FuzzAmountCodeChunk,
 		FuzzColorCodeChunk,
 		FuzzRoughnessCodeChunk,
@@ -954,6 +955,7 @@ int32 UMaterialExpressionSubstrateSlabBSDF::CompileDefaultSlab(FMaterialCompiler
 		SecondRoughnessCodeChunk,
 		SecondRoughnessWeightCodeChunk,
 		Compiler->Constant(0.0f),										// SecondRoughnessAsSimpleClearCoat
+		NormalCodeChunk,
 		FuzzAmountCodeChunk,
 		FuzzColorCodeChunk,
 		FuzzRoughnessCodeChunk,
@@ -1346,6 +1348,13 @@ int32 UMaterialExpressionSubstrateSimpleClearCoatBSDF::Compile(class FMaterialCo
 
 	int32 RoughnessCodeChunk = CompileWithDefaultFloat1(Compiler, Roughness, 0.5f);
 
+	int32 BottomNormalCodeChunk = NormalCodeChunk;
+	const bool bHasCoatBottomNormal = BottomNormal.IsConnected();
+	if (bHasCoatBottomNormal)
+	{
+		BottomNormalCodeChunk = Compiler->ForceCast(CompileWithDefaultNormalWS(Compiler, BottomNormal), MCT_Float3, MFCF_ExactMatch | MFCF_ReplicateValue);
+	}
+
 	int32 OutputCodeChunk = Compiler->SubstrateSlabBSDF(
 		CompileWithDefaultFloat3(Compiler, DiffuseAlbedo, 0.18f, 0.18f, 0.18f),		// DiffuseAlbedo
 		CompileWithDefaultFloat3(Compiler, F0, DefaultF0, DefaultF0, DefaultF0),	// F0
@@ -1361,6 +1370,7 @@ int32 UMaterialExpressionSubstrateSimpleClearCoatBSDF::Compile(class FMaterialCo
 		CompileWithDefaultFloat1(Compiler, ClearCoatRoughness, 0.1f),
 		CompileWithDefaultFloat1(Compiler, ClearCoatCoverage, 1.0f),
 		Compiler->Constant(1.0f),								// SecondRoughnessAsSimpleClearCoat == true for UMaterialExpressionSubstrateSimpleClearCoatBSDF
+		BottomNormalCodeChunk,									// ClearCoatBottomNormal
 		Compiler->Constant(0.0f),								// FuzzAmount
 		Compiler->Constant3(0.0f, 0.0f, 0.0f),					// FuzzColor
 		RoughnessCodeChunk,										// FuzzRoughness
@@ -1417,6 +1427,10 @@ uint32 UMaterialExpressionSubstrateSimpleClearCoatBSDF::GetInputType(int32 Input
 	{
 		return MCT_Float3; // Emissive Color
 	}
+	else if (InputIndex == 7)
+	{
+		return MCT_Float3; // Bottom Normal
+	}
 
 	check(false);
 	return MCT_Float1;
@@ -1452,6 +1466,10 @@ FName UMaterialExpressionSubstrateSimpleClearCoatBSDF::GetInputName(int32 InputI
 	{
 		return TEXT("Emissive Color");
 	}
+	else if (InputIndex == 7)
+	{
+		return TEXT("Bottom Normal");
+	}
 
 	return TEXT("Unknown");
 }
diff --git a/Engine/Source/Runtime/Engine/Public/MaterialCompiler.h b/Engine/Source/Runtime/Engine/Public/MaterialCompiler.h
--- a/Engine/Source/Runtime/Engine/Public/MaterialCompiler.h
+++ b/Engine/Source/Runtime/Engine/Public/MaterialCompiler.h
@@ -721,7 +721,7 @@ public:
 		int32 Roughness, int32 Anisotropy,
 		int32 SSSProfileId, int32 SSSMFP, int32 SSSMFPScale, int32 SSSPhaseAniso, int32 SSSType,
 		int32 EmissiveColor, 
-		int32 SecondRoughness, int32 SecondRoughnessWeight, int32 SecondRoughnessAsSimpleClearCoat,
+		int32 SecondRoughness, int32 SecondRoughnessWeight, int32 SecondRoughnessAsSimpleClearCoat, int32 ClearCoatBottomNormal,
 		int32 FuzzAmount, int32 FuzzColor, int32 FuzzRoughness,
 		int32 Thickness,
 		int32 GlintValue, int32 GlintUV,
@@ -1408,7 +1408,7 @@ public:
 		int32 Roughness, int32 Anisotropy,
 		int32 SSSProfileId, int32 SSSMFP, int32 SSSMFPScale, int32 SSSPhaseAniso, int32 SSSType,
 		int32 EmissiveColor, 
-		int32 SecondRoughness, int32 SecondRoughnessWeight, int32 SecondRoughnessAsSimpleClearCoat,
+		int32 SecondRoughness, int32 SecondRoughnessWeight, int32 SecondRoughnessAsSimpleClearCoat, int32 ClearCoatBottomNormal,
 		int32 FuzzAmount, int32 FuzzColor, int32 FuzzRoughness,
 		int32 Thickness,
 		int32 GlintValue, int32 GlintUV,
@@ -1422,7 +1422,7 @@ public:
 			Roughness, Anisotropy,
 			SSSProfileId, SSSMFP, SSSMFPScale, SSSPhaseAniso, SSSType,
 			EmissiveColor, 
-			SecondRoughness, SecondRoughnessWeight, SecondRoughnessAsSimpleClearCoat,
+			SecondRoughness, SecondRoughnessWeight, SecondRoughnessAsSimpleClearCoat, ClearCoatBottomNormal,
 			FuzzAmount, FuzzColor, FuzzRoughness,
 			Thickness,
 			GlintValue, GlintUV,
diff --git a/Engine/Source/Runtime/Engine/Public/Materials/MaterialExpressionSubstrate.h b/Engine/Source/Runtime/Engine/Public/Materials/MaterialExpressionSubstrate.h
--- a/Engine/Source/Runtime/Engine/Public/Materials/MaterialExpressionSubstrate.h
+++ b/Engine/Source/Runtime/Engine/Public/Materials/MaterialExpressionSubstrate.h
@@ -439,6 +439,12 @@ class UMaterialExpressionSubstrateSimpleClearCoatBSDF : public UMaterialExpressi
 	UPROPERTY()
 	FExpressionInput EmissiveColor;
 
+	/**
+	 * The bottom layer normal. Only used when r.ClearCoatNormal is 1. The normal is considered tangent or world space according to the space properties on the main material node. (type = float3, unit = unitless, defaults to vertex normal)
+	 */
+	UPROPERTY()
+	FExpressionInput BottomNormal;
+
 	//~ Begin UMaterialExpression Interface
 #if WITH_EDITOR
 	virtual int32 Compile(class FMaterialCompiler* Compiler, int32 OutputIndex) override;
