Fixed crash when cell transformer stack contains an empty entry

diff --git a/Engine/Source/Runtime/Engine/Private/WorldPartition/WorldPartition.cpp b/Engine/Source/Runtime/Engine/Private/WorldPartition/WorldPartition.cpp
--- a/Engine/Source/Runtime/Engine/Private/WorldPartition/WorldPartition.cpp
+++ b/Engine/Source/Runtime/Engine/Private/WorldPartition/WorldPartition.cpp
@@ -2491,13 +2491,16 @@ void UWorldPartition::ApplyRuntimeCellsTransformerStack(ULevel* InLevel)
 		{
 			for (const FRuntimeCellTransformerInstance& TransformerInstance : RuntimeCellsTransformerStack)
 			{
-				TPair<double, int32>& TotalStats = RuntimeCellsTransformerStackTimes.FindOrAdd(TransformerInstance.Class);
-				TotalStats.Key -= FPlatformTime::Seconds();
-		
-				Func(TransformerInstance, InLevel);
-		
-				TotalStats.Key += FPlatformTime::Seconds();
-				TotalStats.Value++;
+				if (TransformerInstance.Class)
+				{
+					TPair<double, int32>& TotalStats = RuntimeCellsTransformerStackTimes.FindOrAdd(TransformerInstance.Class);
+					TotalStats.Key -= FPlatformTime::Seconds();
+
+					Func(TransformerInstance, InLevel);
+
+					TotalStats.Key += FPlatformTime::Seconds();
+					TotalStats.Value++;
+				}
 			}
 		};
 
