[materials]
- fix the AppendCompileStateDebugInfo function to properly handle the materialinstance hierarchy (need to traverse up the parent hierarchy to find either the first MI with a static permutation resource, or the base material)
- add additional detail to cover the cases where DDC query is pending and no compilation has occurred yet, or there is no pending query or compile (either before compilation has started, or after it finished)
- add VeryVerbose logging of this state information, printed any time IsCachedCookedPlatformDataLoaded is called

diff --git a/Engine/Source/Runtime/Engine/Private/Materials/Material.cpp b/Engine/Source/Runtime/Engine/Private/Materials/Material.cpp
--- a/Engine/Source/Runtime/Engine/Private/Materials/Material.cpp
+++ b/Engine/Source/Runtime/Engine/Private/Materials/Material.cpp
@@ -107,6 +107,7 @@
 #include "Editor.h"
 #endif
 #include "ShaderCodeLibrary.h"
+#include "String/ParseTokens.h"
 #include "Curves/CurveLinearColorAtlas.h"
 #include "Misc/ScopedSlowTask.h"
 
@@ -4692,13 +4693,20 @@ void UMaterial::GetShaderTypes(EShaderPlatform ShaderPlatform, const ITargetPlat
 
 void UMaterial::AppendCompileStateDebugInfo(FStringBuilderBase& OutDebugInfo) const
 {
-	for (auto Iter = CachedMaterialResourcesForCooking.CreateConstIterator(); Iter; ++Iter)
+	if (!CachedMaterialResourcesForCooking.IsEmpty())
 	{
-		for (const FMaterialResourceForCooking& Resource : Iter.Value())
+		for (auto Iter = CachedMaterialResourcesForCooking.CreateConstIterator(); Iter; ++Iter)
 		{
-			Resource.Resource->AppendCompileStateDebugInfo(OutDebugInfo);
+			for (const FMaterialResourceForCooking& Resource : Iter.Value())
+			{
+				Resource.Resource->AppendCompileStateDebugInfo(OutDebugInfo);
+			}
 		}
 	}
+	else
+	{
+		OutDebugInfo << "No resources created\n";
+	}
 }
 #endif // WITH_EDITOR
 
@@ -4807,8 +4815,21 @@ void UMaterial::BeginCacheForCookedPlatformData( const ITargetPlatform *TargetPl
 bool UMaterial::IsCachedCookedPlatformDataLoaded( const ITargetPlatform* TargetPlatform ) 
 {
 	LLM_SCOPE(ELLMTag::Materials);
-	const TArray<FMaterialResourceForCooking>* CachedMaterialResourcesForPlatform =
-		CachedMaterialResourcesForCooking.Find(TargetPlatform);
+	const TArray<FMaterialResourceForCooking>* CachedMaterialResourcesForPlatform = CachedMaterialResourcesForCooking.Find(TargetPlatform);
+
+	if (UE_LOG_ACTIVE(LogMaterial, VeryVerbose))
+	{
+		TStringBuilder<2048> CookStateInfo;
+		AppendCompileStateDebugInfo(CookStateInfo);
+		UE_LOG(LogMaterial, VeryVerbose, TEXT("Material [%s] Cook State:"), *GetName());
+
+		UE::String::ParseTokens(CookStateInfo.ToView(), TEXT('\n'),
+			[](FStringView Line)
+			{
+				UE_LOG(LogMaterial, VeryVerbose, TEXT("%.*s"), Line.Len(), Line.GetData());
+			}, UE::String::EParseTokensOptions::SkipEmpty);
+	}
+	
 
 	// this should always succeed if BeginCacheForCookedPlatformData is called first
 	if (CachedMaterialResourcesForPlatform != nullptr)
diff --git a/Engine/Source/Runtime/Engine/Private/Materials/MaterialInstance.cpp b/Engine/Source/Runtime/Engine/Private/Materials/MaterialInstance.cpp
--- a/Engine/Source/Runtime/Engine/Private/Materials/MaterialInstance.cpp
+++ b/Engine/Source/Runtime/Engine/Private/Materials/MaterialInstance.cpp
@@ -65,6 +65,7 @@
 
 #if WITH_EDITOR
 #include "Cooker/CookDependency.h"
+#include "String/ParseTokens.h"
 #endif
 
 #if WITH_ODSC
@@ -2990,14 +2991,29 @@ void UMaterialInstance::BeginCacheForCookedPlatformData( const ITargetPlatform *
 bool UMaterialInstance::IsCachedCookedPlatformDataLoaded( const ITargetPlatform* TargetPlatform ) 
 {
 	LLM_SCOPE(ELLMTag::Materials);
-	const TArray<FMaterialResourceForCooking>* CachedMaterialResourcesForPlatform =
-		CachedMaterialResourcesForCooking.Find( TargetPlatform );
+
+	if (UE_LOG_ACTIVE(LogMaterial, VeryVerbose))
+	{
+		TStringBuilder<2048> CookStateInfo;
+		AppendCompileStateDebugInfo(CookStateInfo);
+		UE_LOG(LogMaterial, VeryVerbose, TEXT("MaterialInstance [%s] Cook State:"), *GetName());
+
+		UE::String::ParseTokens(CookStateInfo.ToView(), TEXT('\n'),
+			[](FStringView Line)
+			{
+				UE_LOG(LogMaterial, VeryVerbose, TEXT("%.*s"), Line.Len(), Line.GetData());
+			}, UE::String::EParseTokensOptions::SkipEmpty);
+	}
+
+	const TArray<FMaterialResourceForCooking>* CachedMaterialResourcesForPlatform = CachedMaterialResourcesForCooking.Find( TargetPlatform );
 	if ( CachedMaterialResourcesForPlatform != NULL )
 	{
 		for (const FMaterialResourceForCooking& MaterialResource : *CachedMaterialResourcesForPlatform)
 		{
 			if (MaterialResource.Resource->IsCompilationFinished() == false)
+			{
 				return false;
+			}
 		}
 
 		return true;
@@ -3037,6 +3053,35 @@ void UMaterialInstance::ClearAllCachedCookedPlatformData()
 	FMaterial::DeferredDeleteArray(MaterialsToDelete);
 }
 
+void UMaterialInstance::AppendCompileStateDebugInfo(FStringBuilderBase& OutDebugInfo) const
+{
+	if (bHasStaticPermutationResource)
+	{
+		bool bAnyResources = false;
+		if (!CachedMaterialResourcesForCooking.IsEmpty())
+		{
+			for (auto Iter = CachedMaterialResourcesForCooking.CreateConstIterator(); Iter; ++Iter)
+			{
+				for (const FMaterialResourceForCooking& CookResource : Iter.Value())
+				{
+					bAnyResources = true;
+					CookResource.Resource->AppendCompileStateDebugInfo(OutDebugInfo);
+				}
+			}
+		}
+
+		if (!bAnyResources)
+		{
+			OutDebugInfo << "No resources created\n";
+		}
+	}
+	else
+	{
+		check(Parent);
+		Parent->AppendCompileStateDebugInfo(OutDebugInfo);
+	}
+}
+
 #endif
 
 void UMaterialInstance::Serialize(FArchive& Ar)
diff --git a/Engine/Source/Runtime/Engine/Private/Materials/MaterialShared.cpp b/Engine/Source/Runtime/Engine/Private/Materials/MaterialShared.cpp
--- a/Engine/Source/Runtime/Engine/Private/Materials/MaterialShared.cpp
+++ b/Engine/Source/Runtime/Engine/Private/Materials/MaterialShared.cpp
@@ -286,10 +286,9 @@ static FDelayedAutoRegisterHelper GBlockedMaterialDebugDelegateRegister(EDelayed
 		UE::Cook::FDelegates::PackageBlocked.AddLambda([](const UObject* Obj, FStringBuilderBase& OutDebugInfo)
 			{
 				const UMaterialInterface* MaterialInterface = Cast<UMaterialInterface>(Obj);
-				const UMaterial* Material = MaterialInterface ? MaterialInterface->GetMaterial() : nullptr;
-				if (Material)
+				if (MaterialInterface)
 				{
-					Material->AppendCompileStateDebugInfo(OutDebugInfo);
+					MaterialInterface->AppendCompileStateDebugInfo(OutDebugInfo);
 				}
 			});
 	});
@@ -2263,10 +2262,15 @@ bool FMaterialResource::CheckInValidStateForCompilation(FMaterialCompiler* Compi
 	return Material && Material->CheckInValidStateForCompilation(Compiler);
 }
 
-void FMaterialResource::AppendCompileStateDebugInfo(FStringBuilderBase& OutDebugInfo) const
+void FMaterial::AppendCompileStateDebugInfo(FStringBuilderBase& OutDebugInfo) const
 {
 	check(IsInGameThread());
-	if (GetGameThreadCompilingShaderMapId() != 0u)
+
+	if (CacheShadersPending && !CacheShadersPending->IsReady())
+	{
+		OutDebugInfo << "Pending async DDC load\n";
+	}
+	else if (GetGameThreadCompilingShaderMapId() != 0u)
 	{
 		FMaterialShaderMap* CompilingShaderMap = FMaterialShaderMap::FindCompilingShaderMap(GetGameThreadCompilingShaderMapId());
 		if (CompilingShaderMap)
@@ -2274,6 +2278,11 @@ void FMaterialResource::AppendCompileStateDebugInfo(FStringBuilderBase& OutDebug
 			CompilingShaderMap->AppendCompileStateDebugInfo(OutDebugInfo);
 		}
 	}
+	else
+	{
+		FMaterialShaderMap* ShaderMap = GetGameThreadShaderMap();
+		OutDebugInfo << "Compilation not executing; shadermap is " << (!ShaderMap ? "null" : (ShaderMap->IsComplete(this, true) ? "complete" : "incomplete")) << "\n";
+	}
 }
 
 #endif
diff --git a/Engine/Source/Runtime/Engine/Public/Materials/Material.h b/Engine/Source/Runtime/Engine/Public/Materials/Material.h
--- a/Engine/Source/Runtime/Engine/Public/Materials/Material.h
+++ b/Engine/Source/Runtime/Engine/Public/Materials/Material.h
@@ -1879,7 +1879,7 @@ public:
 #if WITH_EDITOR
 	ENGINE_API virtual void GetShaderTypes(EShaderPlatform Platform, const ITargetPlatform* TargetPlatform, TArray<FDebugShaderTypeInfo>& OutShaderInfo) override;
 
-	void AppendCompileStateDebugInfo(FStringBuilderBase& OutCompileStateStr) const;
+	ENGINE_API virtual void AppendCompileStateDebugInfo(FStringBuilderBase& OutDebugInfo) const override;
 #endif // WITH_EDITOR
 
 	ENGINE_API bool IsPropertyConnected(EMaterialProperty Property) const;
diff --git a/Engine/Source/Runtime/Engine/Public/Materials/MaterialInstance.h b/Engine/Source/Runtime/Engine/Public/Materials/MaterialInstance.h
--- a/Engine/Source/Runtime/Engine/Public/Materials/MaterialInstance.h
+++ b/Engine/Source/Runtime/Engine/Public/Materials/MaterialInstance.h
@@ -942,6 +942,8 @@ public:
 	virtual ENGINE_API bool IsCachedCookedPlatformDataLoaded(const ITargetPlatform* TargetPlatform) override;
 	virtual ENGINE_API void ClearCachedCookedPlatformData(const ITargetPlatform *TargetPlatform) override;
 	virtual ENGINE_API void ClearAllCachedCookedPlatformData() override;
+
+	virtual ENGINE_API void AppendCompileStateDebugInfo(FStringBuilderBase& OutDebugInfo) const override;
 #endif
 	virtual ENGINE_API void Serialize(FArchive& Ar) override;
 	virtual ENGINE_API void PostLoad() override;
diff --git a/Engine/Source/Runtime/Engine/Public/Materials/MaterialInterface.h b/Engine/Source/Runtime/Engine/Public/Materials/MaterialInterface.h
--- a/Engine/Source/Runtime/Engine/Public/Materials/MaterialInterface.h
+++ b/Engine/Source/Runtime/Engine/Public/Materials/MaterialInterface.h
@@ -1079,6 +1079,7 @@ public:
 
 #if WITH_EDITOR
 	virtual void CacheGivenTypesForCooking(EShaderPlatform Platform, ERHIFeatureLevel::Type FeatureLevel, EMaterialQualityLevel::Type QualityLevel, const TArray<const FVertexFactoryType*>& VFTypes, const TArray<const FShaderPipelineType*> PipelineTypes, const TArray<const FShaderType*>& ShaderTypes) {}
+	ENGINE_API virtual void AppendCompileStateDebugInfo(FStringBuilderBase& OutDebugInfo) const {}
 #endif
 
 	/** @brief Checks to see if this material has all its shaders cached.
