Material Editor: Missed a few captions in a previous CL
- Remove Exponent from UMaterialExpressionPower's Caption
- Remove index from Sobol's and TemporalSobal's Captions

diff --git a/Engine/Source/Runtime/Engine/Private/Materials/MaterialExpressions.cpp b/Engine/Source/Runtime/Engine/Private/Materials/MaterialExpressions.cpp
--- a/Engine/Source/Runtime/Engine/Private/Materials/MaterialExpressions.cpp
+++ b/Engine/Source/Runtime/Engine/Private/Materials/MaterialExpressions.cpp
@@ -10651,14 +10651,7 @@ int32 UMaterialExpressionPower::Compile(class FMaterialCompiler* Compiler, int32
 
 void UMaterialExpressionPower::GetCaption(TArray<FString>& OutCaptions) const
 {
-	FString ret = TEXT("Power");
-
-	if (!Exponent.GetTracedInput().Expression)
-	{
-		ret += FString::Printf( TEXT("(X, %.4g)"), ConstExponent);
-	}
-
-	OutCaptions.Add(ret);
+	OutCaptions.Add(TEXT("Power"));
 }
 
 void UMaterialExpressionPower::GetExpressionToolTip(TArray<FString>& OutToolTip) 
@@ -19788,14 +19781,7 @@ int32 UMaterialExpressionSobol::Compile(class FMaterialCompiler* Compiler, int32
 
 void UMaterialExpressionSobol::GetCaption(TArray<FString>& OutCaptions) const
 {
-	FString Caption = TEXT("Sobol");
-
-	if (!Index.GetTracedInput().Expression)
-	{
-		Caption += FString::Printf(TEXT(" (%d)"), ConstIndex);;
-	}
-
-	OutCaptions.Add(Caption);
+	OutCaptions.Add(TEXT("Sobol"));
 }
 
 ///////////////////////////////////////////////////////////////////////////////
@@ -19811,14 +19797,7 @@ int32 UMaterialExpressionTemporalSobol::Compile(class FMaterialCompiler* Compile
 
 void UMaterialExpressionTemporalSobol::GetCaption(TArray<FString>& OutCaptions) const
 {
-	FString Caption = TEXT("Temporal Sobol");
-
-	if (!Index.GetTracedInput().Expression)
-	{
-		Caption += FString::Printf(TEXT(" (%d)"), ConstIndex);;
-	}
-
-	OutCaptions.Add(Caption);
+	OutCaptions.Add(TEXT("Temporal Sobol"));
 }
 #endif // WITH_EDITOR
 
