Replaced `IsValid(this)` under Engine/Source.

diff --git a/Engine/Source/Editor/Blutility/Private/ToolMenuWidget.cpp b/Engine/Source/Editor/Blutility/Private/ToolMenuWidget.cpp
--- a/Engine/Source/Editor/Blutility/Private/ToolMenuWidget.cpp
+++ b/Engine/Source/Editor/Blutility/Private/ToolMenuWidget.cpp
@@ -63,7 +63,7 @@ void UToolMenuWidget::PostInitProperties()
 {
 	Super::PostInitProperties();
 
-	if (IsValid(this))
+	if (IsValidChecked(this))
 	{
 		UpdateFullMenuName();
 	}
diff --git a/Engine/Source/Editor/UnrealEd/Private/GroupActor.cpp b/Engine/Source/Editor/UnrealEd/Private/GroupActor.cpp
--- a/Engine/Source/Editor/UnrealEd/Private/GroupActor.cpp
+++ b/Engine/Source/Editor/UnrealEd/Private/GroupActor.cpp
@@ -78,7 +78,7 @@ void AGroupActor::PostEditUndo()
 {
 	Super::PostEditUndo();
 
-	if (!IsValid(this))
+	if (!IsValidChecked(this))
 	{
 		GetWorld()->ActiveGroupActors.RemoveSwap(this);
 	}
diff --git a/Engine/Source/Editor/VREditor/Private/UI/VREditorFloatingUI.cpp b/Engine/Source/Editor/VREditor/Private/UI/VREditorFloatingUI.cpp
--- a/Engine/Source/Editor/VREditor/Private/UI/VREditorFloatingUI.cpp
+++ b/Engine/Source/Editor/VREditor/Private/UI/VREditorFloatingUI.cpp
@@ -197,7 +197,7 @@ void AVREditorFloatingUI::TickManually(float DeltaTime)
 
 void AVREditorFloatingUI::Destroyed()
 {
-	if (IsValid(this))
+	if (IsValidChecked(this))
 	{
 		CleanupWidgetReferences();
 	}
diff --git a/Engine/Source/Runtime/AIModule/Private/AIController.cpp b/Engine/Source/Runtime/AIModule/Private/AIController.cpp
--- a/Engine/Source/Runtime/AIModule/Private/AIController.cpp
+++ b/Engine/Source/Runtime/AIModule/Private/AIController.cpp
@@ -66,7 +66,7 @@ void AAIController::PostInitializeComponents()
 {
 	Super::PostInitializeComponents();
 
-	if (bWantsPlayerState && IsValid(this) && (GetNetMode() != NM_Client))
+	if (bWantsPlayerState && IsValidChecked(this) && (GetNetMode() != NM_Client))
 	{
 		InitPlayerState();
 	}
diff --git a/Engine/Source/Runtime/AIModule/Private/BehaviorTree/BehaviorTreeComponent.cpp b/Engine/Source/Runtime/AIModule/Private/BehaviorTree/BehaviorTreeComponent.cpp
--- a/Engine/Source/Runtime/AIModule/Private/BehaviorTree/BehaviorTreeComponent.cpp
+++ b/Engine/Source/Runtime/AIModule/Private/BehaviorTree/BehaviorTreeComponent.cpp
@@ -516,7 +516,7 @@ void UBehaviorTreeComponent::HandleMessage(const FAIMessage& Message)
 
 void UBehaviorTreeComponent::OnTaskFinished(const UBTTaskNode* TaskNode, EBTNodeResult::Type TaskResult)
 {
-	if (TaskNode == NULL || InstanceStack.Num() == 0 || !IsValid(this))
+	if (TaskNode == NULL || InstanceStack.Num() == 0 || !IsValidChecked(this))
 	{
 		return;
 	}
@@ -1751,7 +1751,7 @@ void UBehaviorTreeComponent::TickComponent(float DeltaTime, const ELevelTick Tic
 	FScopedCsvStatExclusive _ScopedCsvStatExclusive_BehaviorTreeTick(CSVTickStatName);
 #endif
 
-	check(IsValid(this));
+	check(IsValidChecked(this));
 	float NextNeededDeltaTime = UE::BehaviorTree::DisableTick;
 
 	checkf(PendingBranchActionRequests.Num() == 0, TEXT("Pending branches action requests should always be flushed immediately with the new system"))
@@ -3155,7 +3155,7 @@ void UBehaviorTreeComponent::ForEachChildTask(UBTCompositeNode& StartNode, int32
 #if ENABLE_VISUAL_LOG
 void UBehaviorTreeComponent::DescribeSelfToVisLog(FVisualLogEntry* Snapshot) const
 {
-	if (!IsValid(this))
+	if (!IsValidChecked(this))
 	{
 		return;
 	}
diff --git a/Engine/Source/Runtime/AIModule/Private/BrainComponent.cpp b/Engine/Source/Runtime/AIModule/Private/BrainComponent.cpp
--- a/Engine/Source/Runtime/AIModule/Private/BrainComponent.cpp
+++ b/Engine/Source/Runtime/AIModule/Private/BrainComponent.cpp
@@ -199,7 +199,7 @@ void UBrainComponent::DescribeSelfToVisLog(FVisualLogEntry* Snapshot) const
 {
 	const static UEnum* PriorityEnum = StaticEnum<EAIRequestPriority::Type>();
 
-	if (!IsValid(this))
+	if (!IsValidChecked(this))
 	{
 		return;
 	}
diff --git a/Engine/Source/Runtime/CoreUObject/Private/UObject/ScriptCore.cpp b/Engine/Source/Runtime/CoreUObject/Private/UObject/ScriptCore.cpp
--- a/Engine/Source/Runtime/CoreUObject/Private/UObject/ScriptCore.cpp
+++ b/Engine/Source/Runtime/CoreUObject/Private/UObject/ScriptCore.cpp
@@ -2007,7 +2007,7 @@ void UObject::ProcessEvent( UFunction* Function, void* Parms )
 #endif // TOTAL_OVERHEAD_SCRIPT_STATS
 
 	// Reject.
-	if (!IsValid(this))
+	if (!IsValidChecked(this))
 	{
 		return;
 	}
diff --git a/Engine/Source/Runtime/Engine/Private/ActorConstruction.cpp b/Engine/Source/Runtime/Engine/Private/ActorConstruction.cpp
--- a/Engine/Source/Runtime/Engine/Private/ActorConstruction.cpp
+++ b/Engine/Source/Runtime/Engine/Private/ActorConstruction.cpp
@@ -784,7 +784,9 @@ namespace
 
 bool AActor::ExecuteConstruction(const FTransform& Transform, const FRotationConversionCache* TransformRotationCache, const FComponentInstanceDataCache* InstanceDataCache, bool bIsDefaultTransform, ESpawnActorScaleMethod TransformScaleMethod)
 {
-	check(IsValid(this));
+	PRAGMA_DISABLE_DEPRECATION_WARNINGS
+	check(IsThisNotNull(this, "AActor::ExecuteConstruction") && IsValidChecked(this));
+	PRAGMA_ENABLE_DEPRECATION_WARNINGS
 	check(!HasAnyFlags(RF_BeginDestroyed|RF_FinishDestroyed));
 
 #if WITH_EDITOR
diff --git a/Engine/Source/Runtime/Engine/Private/ActorEditor.cpp b/Engine/Source/Runtime/Engine/Private/ActorEditor.cpp
--- a/Engine/Source/Runtime/Engine/Private/ActorEditor.cpp
+++ b/Engine/Source/Runtime/Engine/Private/ActorEditor.cpp
@@ -754,7 +754,7 @@ bool AActor::InternalPostEditUndo()
 	}
 
 	// Restore OwnedComponents array
-	if (IsValid(this))
+	if (IsValidChecked(this))
 	{
 		ResetOwnedComponents();
 
diff --git a/Engine/Source/Runtime/Engine/Private/CanvasRenderTarget2D.cpp b/Engine/Source/Runtime/Engine/Private/CanvasRenderTarget2D.cpp
--- a/Engine/Source/Runtime/Engine/Private/CanvasRenderTarget2D.cpp
+++ b/Engine/Source/Runtime/Engine/Private/CanvasRenderTarget2D.cpp
@@ -151,7 +151,7 @@ void UCanvasRenderTarget2D::RepaintCanvas()
 		});
 	}
 
-	if (IsValid(this) && OnCanvasRenderTargetUpdate.IsBound())
+	if (IsValidChecked(this) && OnCanvasRenderTargetUpdate.IsBound())
 	{
 		OnCanvasRenderTargetUpdate.Broadcast(Canvas, GetSurfaceWidth(), GetSurfaceHeight());
 	}
diff --git a/Engine/Source/Runtime/Engine/Private/Character.cpp b/Engine/Source/Runtime/Engine/Private/Character.cpp
--- a/Engine/Source/Runtime/Engine/Private/Character.cpp
+++ b/Engine/Source/Runtime/Engine/Private/Character.cpp
@@ -142,7 +142,7 @@ void ACharacter::PostInitializeComponents()
 
 	Super::PostInitializeComponents();
 
-	if (IsValid(this))
+	if (IsValidChecked(this))
 	{
 		if (Mesh)
 		{
diff --git a/Engine/Source/Runtime/Engine/Private/CharacterMovementComponentAsync.cpp b/Engine/Source/Runtime/Engine/Private/CharacterMovementComponentAsync.cpp
--- a/Engine/Source/Runtime/Engine/Private/CharacterMovementComponentAsync.cpp
+++ b/Engine/Source/Runtime/Engine/Private/CharacterMovementComponentAsync.cpp
@@ -2850,7 +2850,7 @@ bool FUpdatedComponentAsyncInput::MoveComponent(const FVector& Delta, const FQua
 
 	// Handle blocking hit notifications. Avoid if pending kill (which could happen after overlaps).
 	const bool bAllowHitDispatch = !BlockingHit.bStartPenetrating || !(MoveFlags & MOVECOMP_DisableBlockingOverlapDispatch);
-	if (BlockingHit.bBlockingHit && bAllowHitDispatch/* && IsValid(this)*/)
+	if (BlockingHit.bBlockingHit && bAllowHitDispatch)
 	{
 		check(bFilledHitResult);
 		/*if (IsDeferringMovementUpdates()) // TODO DeferredMovementUpdates
diff --git a/Engine/Source/Runtime/Engine/Private/Components/ActorComponent.cpp b/Engine/Source/Runtime/Engine/Private/Components/ActorComponent.cpp
--- a/Engine/Source/Runtime/Engine/Private/Components/ActorComponent.cpp
+++ b/Engine/Source/Runtime/Engine/Private/Components/ActorComponent.cpp
@@ -1043,7 +1043,7 @@ void UActorComponent::PreEditChange(FProperty* PropertyThatWillChange)
 	{
 		// The component or its outer could be pending kill when calling PreEditChange when applying a transaction.
 		// Don't do do a full recreate in this situation, and instead simply detach.
-		if( IsValid(this) )
+		if(IsValidChecked(this))
 		{
 			// One way this check can fail is that component subclass does not call Super::PostEditChangeProperty
 			checkf(!EditReregisterContexts.Find(this),
@@ -1078,7 +1078,7 @@ void UActorComponent::PostEditUndo()
 {
 	// Objects marked pending kill don't call PostEditChange() from UObject::PostEditUndo(),
 	// so they can leave an EditReregisterContexts entry around if they are deleted by an undo action.
-	if( !IsValid(this) )
+	if(!IsValidChecked(this))
 	{
 		// For the redo case, ensure that we're no longer in the OwnedComponents array.
 		if (AActor* OwningActor = GetOwner())
@@ -1190,7 +1190,7 @@ void UActorComponent::ConsolidatedPostEditChange(const FPropertyChangedEvent& Pr
 
 	// The component or its outer could be pending kill when calling PostEditChange when applying a transaction.
 	// Don't do do a full recreate in this situation, and instead simply detach.
-	if( !IsValid(this) )
+	if(!IsValidChecked(this))
 	{
 		// @todo james should this call UnregisterComponent instead to remove itself from the RegisteredComponents array on the owner?
 		ExecuteUnregisterEvents();
@@ -1309,7 +1309,7 @@ void UActorComponent::OnRegister()
 	checkf(!IsUnreachable(), TEXT("%s"), *GetDetailedInfo());
 	checkf(!GetOuter()->IsTemplate(), TEXT("'%s' (%s)"), *GetOuter()->GetFullName(), *GetDetailedInfo());
 	checkf(!IsTemplate(), TEXT("'%s' (%s)"), *GetOuter()->GetFullName(), *GetDetailedInfo() );
-	checkf(IsValid(this), TEXT("OnRegister: %s to %s"), *GetDetailedInfo(), GetOwner() ? *GetOwner()->GetFullName() : TEXT("*** No Owner ***") );
+	checkf(IsValidChecked(this), TEXT("OnRegister: %s to %s"), *GetDetailedInfo(), GetOwner() ? *GetOwner()->GetFullName() : TEXT("*** No Owner ***") );
 #endif
 	checkf(WorldPrivate, TEXT("OnRegister: %s to %s"), *GetDetailedInfo(), GetOwner() ? *GetOwner()->GetFullName() : TEXT("*** No Owner ***") );
 	checkf(!bRegistered, TEXT("OnRegister: %s to %s"), *GetDetailedInfo(), GetOwner() ? *GetOwner()->GetFullName() : TEXT("*** No Owner ***") );
@@ -1618,7 +1618,7 @@ void UActorComponent::PreRegisterComponentWithWorld(UWorld* InWorld)
 {
 	checkf(!IsUnreachable(), TEXT("%s"), *GetFullName());
 
-	if (!IsValid(this))
+	if (!IsValidChecked(this))
 	{
 		UE_LOG(LogActorComponent, Log, TEXT("Trying to pre-register component with IsValid() == false. Aborting."));
 		return;
@@ -1683,7 +1683,7 @@ void UActorComponent::RegisterComponentWithWorld(UWorld* InWorld, FRegisterCompo
 
 	checkf(!IsUnreachable(), TEXT("%s"), *GetFullName());
 
-	if(!IsValid(this))
+	if(!IsValidChecked(this))
 	{
 		UE_LOG(LogActorComponent, Log, TEXT("RegisterComponentWithWorld: (%s) Trying to register component with IsValid() == false. Aborting."), *GetPathName());
 		return;
@@ -2328,7 +2328,7 @@ void UActorComponent::DoDeferredRenderUpdates_Concurrent()
 
 	checkf(!IsUnreachable(), TEXT("%s"), *GetFullName());
 	checkf(!IsTemplate(), TEXT("%s"), *GetFullName());
-	checkf(IsValid(this), TEXT("%s"), *GetFullName());
+	checkf(IsValidChecked(this), TEXT("%s"), *GetFullName());
 
 	FScopeCycleCounterUObject ContextScope(this);
 	FScopeCycleCounterUObject AdditionalScope(STATS ? AdditionalStatObject() : nullptr);
diff --git a/Engine/Source/Runtime/Engine/Private/Components/PrimitiveComponent.cpp b/Engine/Source/Runtime/Engine/Private/Components/PrimitiveComponent.cpp
--- a/Engine/Source/Runtime/Engine/Private/Components/PrimitiveComponent.cpp
+++ b/Engine/Source/Runtime/Engine/Private/Components/PrimitiveComponent.cpp
@@ -2941,7 +2941,7 @@ bool UPrimitiveComponent::MoveComponentImpl( const FVector& Delta, const FQuat&
 #endif
 
 	// static things can move before they are registered (e.g. immediately after streaming), but not after.
-	if (!IsValid(this) || CheckStaticMobilityAndWarn(PrimitiveComponentStatics::MobilityWarnText))
+	if (!IsValidChecked(this) || CheckStaticMobilityAndWarn(PrimitiveComponentStatics::MobilityWarnText))
 	{
 		if (OutHit)
 		{
@@ -3216,7 +3216,7 @@ bool UPrimitiveComponent::MoveComponentImpl( const FVector& Delta, const FQuat&
 
 	// Handle blocking hit notifications. Avoid if pending kill (which could happen after overlaps).
 	const bool bAllowHitDispatch = !BlockingHit.bStartPenetrating || !(MoveFlags & MOVECOMP_DisableBlockingOverlapDispatch);
-	if (BlockingHit.bBlockingHit && bAllowHitDispatch && IsValid(this))
+	if (BlockingHit.bBlockingHit && bAllowHitDispatch && IsValidChecked(this))
 	{
 		check(bFilledHitResult);
 		if (IsDeferringMovementUpdates())
@@ -3797,7 +3797,9 @@ void UPrimitiveComponent::BeginComponentOverlap(const FOverlapInfo& OtherOverlap
 	SCOPE_CYCLE_COUNTER(STAT_BeginComponentOverlap);
 
 	// If pending kill, we should not generate any new overlaps
-	if (!IsValid(this))
+	PRAGMA_DISABLE_DEPRECATION_WARNINGS
+	if (!IsThisNotNull(this, "UPrimitiveComponent::BeginComponentOverlap") || !IsValidChecked(this))
+	PRAGMA_ENABLE_DEPRECATION_WARNINGS
 	{
 		return;
 	}
@@ -3825,10 +3827,7 @@ void UPrimitiveComponent::BeginComponentOverlap(const FOverlapInfo& OtherOverlap
 			if (bDoNotifies && ((World && World->HasBegunPlay()) || bLevelStreamingOverlap))
 			{
 				// first execute component delegates
-				if (IsValid(this))
-				{
-					OnComponentBeginOverlap.Broadcast(this, OtherActor, OtherComp, OtherOverlap.GetBodyIndex(), OtherOverlap.bFromSweep, OtherOverlap.OverlapInfo);
-				}
+				OnComponentBeginOverlap.Broadcast(this, OtherActor, OtherComp, OtherOverlap.GetBodyIndex(), OtherOverlap.bFromSweep, OtherOverlap.OverlapInfo);
 
 				if (IsValid(OtherComp))
 				{
@@ -4198,7 +4197,7 @@ bool UPrimitiveComponent::UpdateOverlapsImpl(const TOverlapArrayView* NewPending
 			TInlineOverlapPointerArray NewOverlappingComponentPtrs;
 
 			// If pending kill, we should not generate any new overlaps. Also not if overlaps were just disabled during BeginComponentOverlap.
-			if (IsValid(this) && GetGenerateOverlapEvents())
+			if (IsValidChecked(this) && GetGenerateOverlapEvents())
 			{
 				// Might be able to avoid testing for new overlaps at the end location.
 				if (OverlapsAtEndLocation != nullptr && PrimitiveComponentCVars::bAllowCachedOverlapsCVar && PrevTransform.Equals(GetComponentTransform()))
@@ -4477,7 +4476,7 @@ void UPrimitiveComponent::UpdateOcclusionBoundsSlack(float NewSlack)
 
 void UPrimitiveComponent::UpdatePhysicsVolume( bool bTriggerNotifiers )
 {
-	if (GetShouldUpdatePhysicsVolume() && IsValid(this))
+	if (GetShouldUpdatePhysicsVolume() && IsValidChecked(this))
 	{
 		SCOPE_CYCLE_COUNTER(STAT_UpdatePhysicsVolume);
 		if (UWorld* MyWorld = GetWorld())
@@ -4673,7 +4672,7 @@ void UPrimitiveComponent::DispatchOnClicked(FKey ButtonPressed)
 		}
 	}
 
-	if (IsValid(this))
+	if (IsValidChecked(this))
 	{
 		OnClicked.Broadcast(this, ButtonPressed);
 	}
@@ -4690,7 +4689,7 @@ void UPrimitiveComponent::DispatchOnReleased(FKey ButtonReleased)
 		}
 	}
 
-	if (IsValid(this))
+	if (IsValidChecked(this))
 	{
 		OnReleased.Broadcast(this, ButtonReleased);
 	}
@@ -4707,7 +4706,7 @@ void UPrimitiveComponent::DispatchOnInputTouchBegin(const ETouchIndex::Type Fing
 		}
 	}
 
-	if (IsValid(this))
+	if (IsValidChecked(this))
 	{
 		OnInputTouchBegin.Broadcast(FingerIndex, this);
 	}
@@ -4724,7 +4723,7 @@ void UPrimitiveComponent::DispatchOnInputTouchEnd(const ETouchIndex::Type Finger
 		}
 	}
 
-	if (IsValid(this))
+	if (IsValidChecked(this))
 	{
 		OnInputTouchEnd.Broadcast(FingerIndex, this);
 	}
diff --git a/Engine/Source/Runtime/Engine/Private/Components/SceneComponent.cpp b/Engine/Source/Runtime/Engine/Private/Components/SceneComponent.cpp
--- a/Engine/Source/Runtime/Engine/Private/Components/SceneComponent.cpp
+++ b/Engine/Source/Runtime/Engine/Private/Components/SceneComponent.cpp
@@ -1320,7 +1320,7 @@ void USceneComponent::OnComponentDestroyed(bool bDestroyingHierarchy)
 						// We've gotten in to a bad state where the Child's AttachParent doesn't jive with the AttachChildren array
 						// so instead of crashing, gracefully handle and output an error. 
 						// We skip outputting the error if something is pending kill because this is likely a undo/redo situation that is not concerning.
-						if (IsValid(this) && IsValid(Child))
+						if (IsValidChecked(this) && IsValid(Child))
 						{
 							UE_LOG(LogSceneComponent, Error, TEXT("Component '%s' has '%s' in its AttachChildren array, however, '%s' believes it is not attached to anything"), *GetFullName(), *Child->GetFullName(), *Child->GetFullName());
 						}
@@ -1402,7 +1402,7 @@ void USceneComponent::OnComponentDestroyed(bool bDestroyingHierarchy)
 						// We've gotten in to a bad state where the Child's AttachParent doesn't jive with the AttachChildren array
 						// so instead of crashing, gracefully handle and output an error. 
 						// We skip outputting the error if something is pending kill because this is likely a undo/redo situation that is not concerning.
-						if (IsValid(this) && IsValid(Child))
+						if (IsValidChecked(this) && IsValid(Child))
 						{
 							UE_LOG(LogSceneComponent, Error, TEXT("Component '%s' has '%s' in its AttachChildren array, however, '%s' believes it is not attached to anything"), *GetFullName(), *Child->GetFullName(), *Child->GetFullName());
 						}
@@ -2890,7 +2890,7 @@ APhysicsVolume* USceneComponent::GetPhysicsVolume() const
 
 void USceneComponent::UpdatePhysicsVolume( bool bTriggerNotifiers )
 {
-	if ( bShouldUpdatePhysicsVolume && IsValid(this) )
+	if (bShouldUpdatePhysicsVolume && IsValidChecked(this))
 	{
 		if (UWorld* MyWorld = GetWorld())
 		{
@@ -3185,7 +3185,7 @@ bool USceneComponent::MoveComponentImpl(const FVector& Delta, const FQuat& NewRo
 	SCOPE_CYCLE_COUNTER(STAT_MoveComponentSceneComponentTime);
 
 	// static things can move before they are registered (e.g. immediately after streaming), but not after.
-	if (!IsValid(this) || CheckStaticMobilityAndWarn(SceneComponentStatics::MobilityWarnText))
+	if (!IsValidChecked(this) || CheckStaticMobilityAndWarn(SceneComponentStatics::MobilityWarnText))
 	{
 		if (OutHit)
 		{
diff --git a/Engine/Source/Runtime/Engine/Private/Controller.cpp b/Engine/Source/Runtime/Engine/Private/Controller.cpp
--- a/Engine/Source/Runtime/Engine/Private/Controller.cpp
+++ b/Engine/Source/Runtime/Engine/Private/Controller.cpp
@@ -300,7 +300,7 @@ void AController::PostInitializeComponents()
 {
 	Super::PostInitializeComponents();
 
-	if ( IsValid(this) )
+	if (IsValidChecked(this))
 	{
 		GetWorld()->AddController( this );
 
diff --git a/Engine/Source/Runtime/Engine/Private/EdGraph/EdGraphNode.cpp b/Engine/Source/Runtime/Engine/Private/EdGraph/EdGraphNode.cpp
--- a/Engine/Source/Runtime/Engine/Private/EdGraph/EdGraphNode.cpp
+++ b/Engine/Source/Runtime/Engine/Private/EdGraph/EdGraphNode.cpp
@@ -531,7 +531,7 @@ FText UEdGraphNode::GetVisualWarningTooltipText() const
 class UEdGraph* UEdGraphNode::GetGraph() const
 {
 	UEdGraph* Graph = Cast<UEdGraph>(GetOuter());
-	if (Graph == nullptr && IsValid(this))
+	if (Graph == nullptr && IsValidChecked(this))
 	{
 		if(GetOuter() != GetTransientPackage())
 		{
diff --git a/Engine/Source/Runtime/Engine/Private/LevelTick.cpp b/Engine/Source/Runtime/Engine/Private/LevelTick.cpp
--- a/Engine/Source/Runtime/Engine/Private/LevelTick.cpp
+++ b/Engine/Source/Runtime/Engine/Private/LevelTick.cpp
@@ -401,7 +401,7 @@ void AController::TickActor( float DeltaSeconds, ELevelTick TickType, FActorTick
 		return;
 	}
 
-	if( IsValid(this) )
+	if(IsValidChecked(this))
 	{
 		Tick(DeltaSeconds);	// perform any tick functions unique to an actor subclass
 	}
diff --git a/Engine/Source/Runtime/Engine/Private/Materials/MaterialExpressions.cpp b/Engine/Source/Runtime/Engine/Private/Materials/MaterialExpressions.cpp
--- a/Engine/Source/Runtime/Engine/Private/Materials/MaterialExpressions.cpp
+++ b/Engine/Source/Runtime/Engine/Private/Materials/MaterialExpressions.cpp
@@ -2144,8 +2144,6 @@ void UMaterialExpression::ConnectToPreviewMaterial(UMaterial* InMaterial, int32
 
 void UMaterialExpression::ValidateState()
 {
-	// Disabled for now until issues can be tracked down
-	//check(IsValid(this));
 }
 
 #if WITH_EDITOR
diff --git a/Engine/Source/Runtime/Engine/Private/Pawn.cpp b/Engine/Source/Runtime/Engine/Private/Pawn.cpp
--- a/Engine/Source/Runtime/Engine/Private/Pawn.cpp
+++ b/Engine/Source/Runtime/Engine/Private/Pawn.cpp
@@ -136,7 +136,7 @@ void APawn::PostInitializeComponents()
 
 	Super::PostInitializeComponents();
 	
-	if (IsValid(this))
+	if (IsValidChecked(this))
 	{
 		UWorld* World = GetWorld();
 
diff --git a/Engine/Source/Runtime/Engine/Private/PlayerController.cpp b/Engine/Source/Runtime/Engine/Private/PlayerController.cpp
--- a/Engine/Source/Runtime/Engine/Private/PlayerController.cpp
+++ b/Engine/Source/Runtime/Engine/Private/PlayerController.cpp
@@ -1029,7 +1029,7 @@ void APlayerController::PostInitializeComponents()
 {
 	Super::PostInitializeComponents();
 
-	if ( IsValid(this) && (GetNetMode() != NM_Client) )
+	if (IsValidChecked(this) && (GetNetMode() != NM_Client) )
 	{
 		// create a new player replication info
 		InitPlayerState();
@@ -5032,7 +5032,7 @@ void APlayerController::TickActor( float DeltaSeconds, ELevelTick TickType, FAct
 		// Clear axis inputs from previous frame.
 		RotationInput = FRotator::ZeroRotator;
 
-		if (IsValid(this))
+		if (IsValidChecked(this))
 		{
 			Tick(DeltaSeconds);	// perform any tick functions unique to an actor subclass
 		}
@@ -5168,7 +5168,7 @@ void APlayerController::TickActor( float DeltaSeconds, ELevelTick TickType, FAct
 			PlayerTick(DeltaSeconds);
 		}
 
-		if (!IsValid(this))
+		if (!IsValidChecked(this))
 		{
 			return;
 		}
@@ -5199,7 +5199,7 @@ void APlayerController::TickActor( float DeltaSeconds, ELevelTick TickType, FAct
 		}
 	}
 
-	if (IsValid(this))
+	if (IsValidChecked(this))
 	{
 		QUICK_SCOPE_CYCLE_COUNTER(Tick);
 		Tick(DeltaSeconds);	// perform any tick functions unique to an actor subclass
diff --git a/Engine/Source/Runtime/Engine/Private/PrimitiveComponentPhysics.cpp b/Engine/Source/Runtime/Engine/Private/PrimitiveComponentPhysics.cpp
--- a/Engine/Source/Runtime/Engine/Private/PrimitiveComponentPhysics.cpp
+++ b/Engine/Source/Runtime/Engine/Private/PrimitiveComponentPhysics.cpp
@@ -189,7 +189,7 @@ void UPrimitiveComponent::SetConstraintMode(EDOFMode::Type ConstraintMode)
 {
 	FBodyInstance * RootBI = GetBodyInstance(NAME_None, false);
 
-	if (RootBI == NULL || !IsValid(this))
+	if (RootBI == NULL || !IsValidChecked(this))
 	{
 		return;
 	}
@@ -889,7 +889,7 @@ void UPrimitiveComponent::SyncComponentToRBPhysics()
 		}
 	}
 
-	if (!IsValid(this) || !IsSimulatingPhysics() || !RigidBodyIsAwake())
+	if (!IsValidChecked(this) || !IsSimulatingPhysics() || !RigidBodyIsAwake())
 	{
 		return;
 	}
diff --git a/Engine/Source/Runtime/Engine/Private/World.cpp b/Engine/Source/Runtime/Engine/Private/World.cpp
--- a/Engine/Source/Runtime/Engine/Private/World.cpp
+++ b/Engine/Source/Runtime/Engine/Private/World.cpp
@@ -6237,7 +6237,7 @@ void UWorld::RemoveOnPreUnregisterAllActorComponentsHandler(FDelegateHandle InHa
 void UWorld::NotifyPreUnregisterAllActorComponents(AActor* Actor)
 {
 	// This may be called on inactive/GCing worlds, if so ignore it.
-	if (WorldType != EWorldType::Inactive && IsValid(this))
+	if (WorldType != EWorldType::Inactive && IsValidChecked(this))
 	{
 		OnPreUnregisterAllActorComponents.Broadcast(Actor);
 	}
diff --git a/Engine/Source/Runtime/Experimental/Animation/Constraints/Private/ConstraintsManager.cpp b/Engine/Source/Runtime/Experimental/Animation/Constraints/Private/ConstraintsManager.cpp
--- a/Engine/Source/Runtime/Experimental/Animation/Constraints/Private/ConstraintsManager.cpp
+++ b/Engine/Source/Runtime/Experimental/Animation/Constraints/Private/ConstraintsManager.cpp
@@ -193,7 +193,7 @@ void UTickableConstraint::PostEditUndo()
 	Super::PostEditUndo();
 
 	// make sure we update ticking when undone
-	const bool bActiveTick = ::IsValid(this);
+	const bool bActiveTick = ::IsValidChecked(this);
 	for (TPair<TWeakObjectPtr<ULevel>, FConstraintTickFunction>& Pair : ConstraintTicks)
 	{
 		if (Pair.Key.IsValid())
diff --git a/Engine/Source/Runtime/GameplayTasks/Private/GameplayTask.cpp b/Engine/Source/Runtime/GameplayTasks/Private/GameplayTask.cpp
--- a/Engine/Source/Runtime/GameplayTasks/Private/GameplayTask.cpp
+++ b/Engine/Source/Runtime/GameplayTasks/Private/GameplayTask.cpp
@@ -152,7 +152,7 @@ void UGameplayTask::TaskOwnerEnded()
 	if (TaskState != EGameplayTaskState::Finished)
 	{
 		bOwnerFinished = true;
-		if (IsValid(this))
+		if (IsValidChecked(this))
 		{
 			OnDestroy(true);
 		}
@@ -172,7 +172,7 @@ void UGameplayTask::EndTask()
 
 	if (TaskState != EGameplayTaskState::Finished)
 	{
-		if (IsValid(this))
+		if (IsValidChecked(this))
 		{
 			OnDestroy(false);
 		}
@@ -207,7 +207,7 @@ void UGameplayTask::ExternalCancel()
 
 void UGameplayTask::OnDestroy(bool bInOwnerFinished)
 {
-	if (ensureMsgf(IsValid(this), TEXT("OnDestroy called on invalid gameplay task")))
+	if (ensureMsgf(IsValidChecked(this), TEXT("OnDestroy called on invalid gameplay task")))
 	{
 		const FString OwnerName = TaskOwner.IsValid() ? TaskOwner.GetObject()->GetName() : TEXT("Invalid GameplayTask Owner");
 		ensureMsgf(TaskState != EGameplayTaskState::Finished, TEXT("%s OnDestroy called, current state: %i, owner name: %s"), *GetName(), TaskState, *OwnerName);
diff --git a/Engine/Source/Runtime/GameplayTasks/Private/GameplayTasksComponent.cpp b/Engine/Source/Runtime/GameplayTasks/Private/GameplayTasksComponent.cpp
--- a/Engine/Source/Runtime/GameplayTasks/Private/GameplayTasksComponent.cpp
+++ b/Engine/Source/Runtime/GameplayTasks/Private/GameplayTasksComponent.cpp
@@ -698,7 +698,9 @@ void UGameplayTasksComponent::DescribeSelfToVisLog(FVisualLogEntry* Snapshot) co
 	static const FString PriorityQueueName = TEXT("Priority Queue");
 	static const FString OtherTasksName = TEXT("Other tasks");
 
-	if (!IsValid(this))
+	PRAGMA_DISABLE_DEPRECATION_WARNINGS
+	if (!IsThisNotNull(this, "UGameplayTasksComponent::DescribeSelfToVisLog") || !IsValidChecked(this))
+	PRAGMA_ENABLE_DEPRECATION_WARNINGS
 	{
 		return;
 	}
diff --git a/Engine/Source/Runtime/Landscape/Private/LandscapeCollision.cpp b/Engine/Source/Runtime/Landscape/Private/LandscapeCollision.cpp
--- a/Engine/Source/Runtime/Landscape/Private/LandscapeCollision.cpp
+++ b/Engine/Source/Runtime/Landscape/Private/LandscapeCollision.cpp
@@ -2604,7 +2604,7 @@ void ULandscapeHeightfieldCollisionComponent::PostLoad()
 
 #if WITH_EDITOR
 	// PostLoad of the landscape can decide to recreate collision, in which case this components checks are irrelevant
-	if (!HasAnyFlags(RF_ClassDefaultObject) && IsValid(this))
+	if (!HasAnyFlags(RF_ClassDefaultObject) && IsValidChecked(this))
 	{
 		bShouldSaveCookedDataToDDC[0] = true;
 		bShouldSaveCookedDataToDDC[1] = true;
diff --git a/Engine/Source/Runtime/Landscape/Private/LandscapeEdit.cpp b/Engine/Source/Runtime/Landscape/Private/LandscapeEdit.cpp
--- a/Engine/Source/Runtime/Landscape/Private/LandscapeEdit.cpp
+++ b/Engine/Source/Runtime/Landscape/Private/LandscapeEdit.cpp
@@ -1000,7 +1000,7 @@ void ULandscapeComponent::PostEditUndo()
 
 	Super::PostEditUndo();
 
-	if (IsValid(this))
+	if (IsValidChecked(this))
 	{
 		EditToolRenderData.UpdateSelectionMaterial(EditToolRenderData.SelectedType, this);
 		if (!Proxy->HasLayersContent())
@@ -1015,7 +1015,7 @@ void ULandscapeComponent::PostEditUndo()
 		// Updating the bounds after undo will re-sync CachedLocalBox with the heightmap data. This is necessary because ULandscapeComponent might be part of 
 		//  the transaction buffer, but the final heightmaps aren't, so undoing might desynchronize them. If the heightmap update that is requested right after 
 		//  ends up modifying heightmap data, then CachedLocalBox will get updated properly again so this is a 'safe' way to fix the issue
-		if (IsValid(this))
+		if (IsValidChecked(this))
 		{
 			UpdateCachedBounds();
 		}
diff --git a/Engine/Source/Runtime/MRMesh/Private/MRMeshComponent.cpp b/Engine/Source/Runtime/MRMesh/Private/MRMeshComponent.cpp
--- a/Engine/Source/Runtime/MRMesh/Private/MRMeshComponent.cpp
+++ b/Engine/Source/Runtime/MRMesh/Private/MRMeshComponent.cpp
@@ -698,7 +698,7 @@ void UMRMeshComponent::SendBrickData_Internal(IMRMesh::FSendBrickDataArgs Args)
 	UE_LOG(LogMrMesh, Log, TEXT("SendBrickData_Internal() processing brick %llu with %i triangles"), Args.BrickId, Args.Indices.Num() / 3);
 	
 	// Collision generation.
-	if (IsValid(this) && !bNeverCreateCollisionMesh)
+	if (IsValidChecked(this) && !bNeverCreateCollisionMesh)
 	{
 			SCOPE_CYCLE_COUNTER(STAT_UpdateCollision);
 
diff --git a/Engine/Source/Runtime/NavigationSystem/Private/NavigationData.cpp b/Engine/Source/Runtime/NavigationSystem/Private/NavigationData.cpp
--- a/Engine/Source/Runtime/NavigationSystem/Private/NavigationData.cpp
+++ b/Engine/Source/Runtime/NavigationSystem/Private/NavigationData.cpp
@@ -199,7 +199,7 @@ void ANavigationData::PostInitProperties()
 {
 	Super::PostInitProperties();
 
-	if (!IsValid(this))
+	if (!IsValidChecked(this))
 	{
 		return;
 	}
@@ -408,7 +408,9 @@ void ANavigationData::OnUnregistered()
 void ANavigationData::InstantiateAndRegisterRenderingComponent()
 {
 #if UE_ENABLE_DEBUG_DRAWING
-	if (IsValid(this) && !IsValid(RenderingComp))
+	PRAGMA_DISABLE_DEPRECATION_WARNINGS
+	if (IsThisNotNull(this, "ANavigationData::InstantiateAndRegisterRenderingComponent") && IsValidChecked(this) && !IsValid(RenderingComp))
+	PRAGMA_ENABLE_DEPRECATION_WARNINGS
 	{
 		const bool bRootIsRenderComp = (RenderingComp == RootComponent);
 		if (RenderingComp)
diff --git a/Engine/Source/Runtime/UMG/Private/Components/WidgetComponent.cpp b/Engine/Source/Runtime/UMG/Private/Components/WidgetComponent.cpp
--- a/Engine/Source/Runtime/UMG/Private/Components/WidgetComponent.cpp
+++ b/Engine/Source/Runtime/UMG/Private/Components/WidgetComponent.cpp
@@ -1782,7 +1782,9 @@ void UWidgetComponent::SetSlateWidget(const TSharedPtr<SWidget>& InSlateWidget)
 void UWidgetComponent::UpdateWidget()
 {
 	// Don't do any work if Slate is not initialized
-	if (FSlateApplication::IsInitialized() && IsValid(this))
+	PRAGMA_DISABLE_DEPRECATION_WARNINGS
+	if (FSlateApplication::IsInitialized() && IsValidChecked(this))
+	PRAGMA_ENABLE_DEPRECATION_WARNINGS
 	{
 		if (Space == EWidgetSpace::World)
 		{
