New material translator refactor:
- Renamed FAnalysisData to FMaterialIRValueAnalyzer for consistency and better naming. Also architected it as a class rather than just data.
- Renamed "UE::MIR" namespace to simply "MIR".

#rnx

diff --git a/Engine/Source/Runtime/Engine/Private/Materials/MaterialExpressionsIR.cpp b/Engine/Source/Runtime/Engine/Private/Materials/MaterialExpressionsIR.cpp
--- a/Engine/Source/Runtime/Engine/Private/Materials/MaterialExpressionsIR.cpp
+++ b/Engine/Source/Runtime/Engine/Private/Materials/MaterialExpressionsIR.cpp
@@ -240,8 +240,6 @@
 
 #include "Materials/MaterialIREmitter.h"
 
-namespace MIR = UE::MIR;
-
 /* Constants */
 
 void UMaterialExpression::Build(MIR::FEmitter& Emitter)
@@ -489,27 +487,27 @@ static void BuildUnaryArithmeticOperatorFromInput(
 
 // TODO: Add the missing expressions
 
-void UMaterialExpressionAbs::Build(UE::MIR::FEmitter& Emitter)
+void UMaterialExpressionAbs::Build(MIR::FEmitter& Emitter)
 {
 	BuildUnaryArithmeticOperatorFromInput(Emitter, MIR::UO_Abs, &Input);
 }
 
-void UMaterialExpressionCeil::Build(UE::MIR::FEmitter& Emitter)
+void UMaterialExpressionCeil::Build(MIR::FEmitter& Emitter)
 {
 	BuildUnaryArithmeticOperatorFromInput(Emitter, MIR::UO_Ceil, &Input);
 }
 
-void UMaterialExpressionFloor::Build(UE::MIR::FEmitter& Emitter)
+void UMaterialExpressionFloor::Build(MIR::FEmitter& Emitter)
 {
 	BuildUnaryArithmeticOperatorFromInput(Emitter, MIR::UO_Floor, &Input);
 }
 
-void UMaterialExpressionFrac::Build(UE::MIR::FEmitter& Emitter)
+void UMaterialExpressionFrac::Build(MIR::FEmitter& Emitter)
 {
 	BuildUnaryArithmeticOperatorFromInput(Emitter, MIR::UO_Frac, &Input);
 }
 
-void UMaterialExpressionLength::Build(UE::MIR::FEmitter& Emitter)
+void UMaterialExpressionLength::Build(MIR::FEmitter& Emitter)
 {
 	MIR::FValue* Value = Emitter.Input(&Input);
 	UE_MIR_CHECKPOINT();
@@ -517,77 +515,77 @@ void UMaterialExpressionLength::Build(UE::MIR::FEmitter& Emitter)
 	BuildUnaryArithmeticOperatorFromValue(Emitter, MIR::UO_Length, Value);
 }
 
-void UMaterialExpressionRound::Build(UE::MIR::FEmitter& Emitter)
+void UMaterialExpressionRound::Build(MIR::FEmitter& Emitter)
 {
 	BuildUnaryArithmeticOperatorFromInput(Emitter, MIR::UO_Round, &Input);
 }
 
-void UMaterialExpressionExponential::Build(UE::MIR::FEmitter& Emitter)
+void UMaterialExpressionExponential::Build(MIR::FEmitter& Emitter)
 {
 	BuildUnaryArithmeticOperatorFromInput(Emitter, MIR::UO_Exponential, &Input);
 }
 
-void UMaterialExpressionExponential2::Build(UE::MIR::FEmitter& Emitter)
+void UMaterialExpressionExponential2::Build(MIR::FEmitter& Emitter)
 {
 	BuildUnaryArithmeticOperatorFromInput(Emitter, MIR::UO_Exponential2, &Input);
 }
 
-void UMaterialExpressionLogarithm::Build(UE::MIR::FEmitter& Emitter)
+void UMaterialExpressionLogarithm::Build(MIR::FEmitter& Emitter)
 {
 	BuildUnaryArithmeticOperatorFromInput(Emitter, MIR::UO_Logarithm, &Input);
 }
 
-void UMaterialExpressionLogarithm2::Build(UE::MIR::FEmitter& Emitter)
+void UMaterialExpressionLogarithm2::Build(MIR::FEmitter& Emitter)
 {
 	BuildUnaryArithmeticOperatorFromInput(Emitter, MIR::UO_Logarithm2, &X);
 }
 
-void UMaterialExpressionLogarithm10::Build(UE::MIR::FEmitter& Emitter)
+void UMaterialExpressionLogarithm10::Build(MIR::FEmitter& Emitter)
 {
 	BuildUnaryArithmeticOperatorFromInput(Emitter, MIR::UO_Logarithm10, &X);
 }
 
-void UMaterialExpressionTruncate::Build(UE::MIR::FEmitter& Emitter)
+void UMaterialExpressionTruncate::Build(MIR::FEmitter& Emitter)
 {
 	BuildUnaryArithmeticOperatorFromInput(Emitter, MIR::UO_Truncate, &Input);
 }
 
-void UMaterialExpressionArccosine::Build(UE::MIR::FEmitter& Emitter)
+void UMaterialExpressionArccosine::Build(MIR::FEmitter& Emitter)
 {
 	BuildUnaryArithmeticOperatorFromInput(Emitter, MIR::UO_ACos, &Input);
 }
 
-void UMaterialExpressionArcsine::Build(UE::MIR::FEmitter& Emitter)
+void UMaterialExpressionArcsine::Build(MIR::FEmitter& Emitter)
 {
 	BuildUnaryArithmeticOperatorFromInput(Emitter, MIR::UO_ASin, &Input);
 }
 
-void UMaterialExpressionArctangent::Build(UE::MIR::FEmitter& Emitter)
+void UMaterialExpressionArctangent::Build(MIR::FEmitter& Emitter)
 {
 	BuildUnaryArithmeticOperatorFromInput(Emitter, MIR::UO_ATan, &Input);
 }
 
-void UMaterialExpressionCosine::Build(UE::MIR::FEmitter& Emitter)
+void UMaterialExpressionCosine::Build(MIR::FEmitter& Emitter)
 {
 	BuildUnaryArithmeticOperatorFromInput(Emitter, MIR::UO_Cos, &Input);
 }
 
-void UMaterialExpressionSine::Build(UE::MIR::FEmitter& Emitter)
+void UMaterialExpressionSine::Build(MIR::FEmitter& Emitter)
 {
 	BuildUnaryArithmeticOperatorFromInput(Emitter, MIR::UO_Sin, &Input);
 }
 
-void UMaterialExpressionSaturate::Build(UE::MIR::FEmitter& Emitter)
+void UMaterialExpressionSaturate::Build(MIR::FEmitter& Emitter)
 {
 	BuildUnaryArithmeticOperatorFromInput(Emitter, MIR::UO_Saturate, &Input);
 }
 
-void UMaterialExpressionSquareRoot::Build(UE::MIR::FEmitter& Emitter)
+void UMaterialExpressionSquareRoot::Build(MIR::FEmitter& Emitter)
 {
 	BuildUnaryArithmeticOperatorFromInput(Emitter, MIR::UO_Sqrt, &Input);
 }
 
-void UMaterialExpressionTangent::Build(UE::MIR::FEmitter& Emitter)
+void UMaterialExpressionTangent::Build(MIR::FEmitter& Emitter)
 {
 	BuildUnaryArithmeticOperatorFromInput(Emitter, MIR::UO_Tan, &Input);
 }
@@ -644,7 +642,7 @@ void UMaterialExpressionAdd::Build(MIR::FEmitter& Emitter)
 	BuildBinaryOperatorFromInputWithDefault(Emitter, MIR::BO_Add, &A, ConstA, &B, ConstB);
 }
 
-void UMaterialExpressionDesaturation::Build(UE::MIR::FEmitter& Emitter)
+void UMaterialExpressionDesaturation::Build(MIR::FEmitter& Emitter)
 {
 	MIR::FValue* ColorValue = Emitter.Input(&Input);
 	UE_MIR_CHECKPOINT();
@@ -712,12 +710,12 @@ void UMaterialExpressionMin::Build(MIR::FEmitter& Emitter)
 	BuildBinaryOperatorFromInputWithDefault(Emitter, MIR::BO_Min, &A, ConstA, &B, ConstB);
 }
 
-void UMaterialExpressionDotProduct::Build(UE::MIR::FEmitter& Emitter)
+void UMaterialExpressionDotProduct::Build(MIR::FEmitter& Emitter)
 {
 	BuildBinaryOperatorFromInput(Emitter, MIR::BO_Dot, &A, &B);
 }
 
-void UMaterialExpressionCrossProduct::Build(UE::MIR::FEmitter& Emitter)
+void UMaterialExpressionCrossProduct::Build(MIR::FEmitter& Emitter)
 {
 	BuildBinaryOperatorFromInput(Emitter, MIR::BO_Cross, &A, &B);
 }
@@ -811,12 +809,12 @@ void UMaterialExpressionIf::Build(MIR::FEmitter& Emitter)
 	Emitter.Output(0, OutputValue);
 }
 
-void UMaterialExpressionTextureObject::Build(UE::MIR::FEmitter& Emitter)
+void UMaterialExpressionTextureObject::Build(MIR::FEmitter& Emitter)
 {
 	Emitter.Output(0, Emitter.EmitTextureObject(Texture, SamplerType));
 }
 
-static void BuildTextureSampleExpression(MIR::FValue* Texture, UMaterialExpressionTextureSample* Expr, UE::MIR::FEmitter& Emitter)
+static void BuildTextureSampleExpression(MIR::FValue* Texture, UMaterialExpressionTextureSample* Expr, MIR::FEmitter& Emitter)
 {
 	MIR::FValue* TexCoords = Emitter.TryInput(&Expr->Coordinates);
 	if (!TexCoords)
@@ -833,7 +831,7 @@ static void BuildTextureSampleExpression(MIR::FValue* Texture, UMaterialExpressi
 	Emitter.Output(5, OutputValue); 
 }
 
-void UMaterialExpressionTextureSample::Build(UE::MIR::FEmitter& Emitter)
+void UMaterialExpressionTextureSample::Build(MIR::FEmitter& Emitter)
 {
 	MIR::FValue* TextureValue = Emitter.TryInput(&TextureObject);
 	if (TextureValue)
@@ -862,7 +860,7 @@ void UMaterialExpressionTextureSample::Build(UE::MIR::FEmitter& Emitter)
 	BuildTextureSampleExpression(TextureValue, this, Emitter);
 }
 
-static MIR::FValue* BuildTextureObjectParameter(UE::MIR::FEmitter& Emitter, UMaterialExpressionTextureSampleParameter* Expr)
+static MIR::FValue* BuildTextureObjectParameter(MIR::FEmitter& Emitter, UMaterialExpressionTextureSampleParameter* Expr)
 {
 	FMaterialParameterMetadata Param;
 	if (!Expr->GetParameterValue(Param))
@@ -881,21 +879,21 @@ static MIR::FValue* BuildTextureObjectParameter(UE::MIR::FEmitter& Emitter, UMat
 	return ParameterValue;
 }
 
-void UMaterialExpressionTextureSampleParameter::Build(UE::MIR::FEmitter& Emitter)
+void UMaterialExpressionTextureSampleParameter::Build(MIR::FEmitter& Emitter)
 {
 	MIR::FValue* ParameterValue = BuildTextureObjectParameter(Emitter, this);
 	UE_MIR_CHECKPOINT();
 	BuildTextureSampleExpression(ParameterValue, this, Emitter);
 }
 
-void UMaterialExpressionTextureObjectParameter::Build(UE::MIR::FEmitter& Emitter)
+void UMaterialExpressionTextureObjectParameter::Build(MIR::FEmitter& Emitter)
 {
 	MIR::FValue* ParameterValue = BuildTextureObjectParameter(Emitter, this);
 	UE_MIR_CHECKPOINT();
 	Emitter.Output(0, ParameterValue);
 }
 
-void UMaterialExpressionTextureCoordinate::Build(UE::MIR::FEmitter& Emitter)
+void UMaterialExpressionTextureCoordinate::Build(MIR::FEmitter& Emitter)
 {
 	if (UnMirrorU || UnMirrorV)
 	{
@@ -910,7 +908,7 @@ void UMaterialExpressionTextureCoordinate::Build(UE::MIR::FEmitter& Emitter)
 	Emitter.Output(GetOutput(0), OutputValue);
 }
 
-void UMaterialExpressionReroute::Build(UE::MIR::FEmitter& Emitter)
+void UMaterialExpressionReroute::Build(MIR::FEmitter& Emitter)
 {
 	MIR::FValue* Value = Emitter.Input(&Input);
 	Emitter.Output(0, Value);
@@ -937,7 +935,7 @@ void BuildTernaryArithmeticOperator(
 	Emitter.Output(0, Emitter.EmitOperator(Op, ValueA, ValueB, ValueC));
 }
 
-void UMaterialExpressionClamp::Build(UE::MIR::FEmitter& Emitter)
+void UMaterialExpressionClamp::Build(MIR::FEmitter& Emitter)
 {
 	MIR::FValue* InputValue = Emitter.Input(&Input);
 	MIR::FValue* MinValue = Emitter.DefaultInputToFloatScalar(&Min, MinDefault).Input(&Min);
@@ -960,12 +958,12 @@ void UMaterialExpressionClamp::Build(UE::MIR::FEmitter& Emitter)
 	Emitter.Output(0, OutputValue);
 }
 
-void UMaterialExpressionLinearInterpolate::Build(UE::MIR::FEmitter& Emitter)
+void UMaterialExpressionLinearInterpolate::Build(MIR::FEmitter& Emitter)
 {
 	BuildTernaryArithmeticOperator(Emitter, MIR::TO_Lerp, &A, ConstA, &B, ConstB, &Alpha, ConstAlpha);
 }
 
-void UMaterialExpressionSmoothStep::Build(UE::MIR::FEmitter& Emitter)
+void UMaterialExpressionSmoothStep::Build(MIR::FEmitter& Emitter)
 {
 	BuildTernaryArithmeticOperator(Emitter, MIR::TO_Smoothstep, &Min, ConstMin, &Max, ConstMax, &Value, ConstValue);
 }
diff --git a/Engine/Source/Runtime/Engine/Private/Materials/MaterialIR.cpp b/Engine/Source/Runtime/Engine/Private/Materials/MaterialIR.cpp
--- a/Engine/Source/Runtime/Engine/Private/Materials/MaterialIR.cpp
+++ b/Engine/Source/Runtime/Engine/Private/Materials/MaterialIR.cpp
@@ -6,7 +6,7 @@
 
 #if WITH_EDITOR
 
-namespace UE::MIR
+namespace MIR
 {
 
 const TCHAR* ValueKindToString(EValueKind Kind)
@@ -502,6 +502,6 @@ const TCHAR* OperatorToString(EOperator Op)
 	return TEXT("???");
 }
 
-} // namespace UE::MIR
+} // namespace MIR
 
 #endif // #if WITH_EDITOR
diff --git a/Engine/Source/Runtime/Engine/Private/Materials/MaterialIRDebug.cpp b/Engine/Source/Runtime/Engine/Private/Materials/MaterialIRDebug.cpp
--- a/Engine/Source/Runtime/Engine/Private/Materials/MaterialIRDebug.cpp
+++ b/Engine/Source/Runtime/Engine/Private/Materials/MaterialIRDebug.cpp
@@ -18,7 +18,7 @@ static TAutoConsoleVariable<bool> CVarDumpMaterialIRUseGraph_EnableNext(
 	TEXT("Whether the Material Module IR 'Uses' graph should also display 'Instruction Next' edges."),
 	ECVF_RenderThreadSafe);
 
-namespace UE::MIR
+namespace MIR
 {
 
 static void DumpValueInfo(const FValue* Value, FString& Out)
@@ -424,6 +424,6 @@ void DebugDumpIR(FStringView MaterialName, const FMaterialIRModule& Module)
 	FFileHelper::SaveStringToFile(State.Out, *FilePath);
 }
 
-} // namespace UE::MIR
+} // namespace MIR
 
 #endif // #if WITH_EDITOR
diff --git a/Engine/Source/Runtime/Engine/Private/Materials/MaterialIREmitter.cpp b/Engine/Source/Runtime/Engine/Private/Materials/MaterialIREmitter.cpp
--- a/Engine/Source/Runtime/Engine/Private/Materials/MaterialIREmitter.cpp
+++ b/Engine/Source/Runtime/Engine/Private/Materials/MaterialIREmitter.cpp
@@ -12,7 +12,7 @@
 
 #if WITH_EDITOR
 
-namespace UE::MIR
+namespace MIR
 {
 
 const TCHAR* VectorComponentToString(EVectorComponent Component)
@@ -1869,6 +1869,6 @@ uint32 FEmitter::FValueKeyFuncs::GetKeyHash(KeyInitType Key)
 	return Internal::HashBytes((char*)Key, Key->GetSizeInBytes());
 }
 
-} // namespace UE::MIR
+} // namespace MIR
 
 #endif // #if WITH_EDITOR
diff --git a/Engine/Source/Runtime/Engine/Private/Materials/MaterialIRInternal.cpp b/Engine/Source/Runtime/Engine/Private/Materials/MaterialIRInternal.cpp
--- a/Engine/Source/Runtime/Engine/Private/Materials/MaterialIRInternal.cpp
+++ b/Engine/Source/Runtime/Engine/Private/Materials/MaterialIRInternal.cpp
@@ -8,7 +8,7 @@
 
 #if WITH_EDITOR
 
-namespace UE::MIR::Internal {
+namespace MIR::Internal {
 
 bool IsMaterialPropertyShared(EMaterialProperty InProperty)
 {
@@ -121,6 +121,6 @@ uint32 HashBytes(const char* Ptr, uint32 Size)
 	return Hash;
 }
 
-} // namespace UE::MIR::Internal
+} // namespace MIR::Internal
 
 #endif
diff --git a/Engine/Source/Runtime/Engine/Private/Materials/MaterialIRInternal.h b/Engine/Source/Runtime/Engine/Private/Materials/MaterialIRInternal.h
--- a/Engine/Source/Runtime/Engine/Private/Materials/MaterialIRInternal.h
+++ b/Engine/Source/Runtime/Engine/Private/Materials/MaterialIRInternal.h
@@ -13,7 +13,7 @@ TArrayView<T> MakeTemporaryArray(FMemMark&, int Count)
 	return { Ptr, Count };
 }
 
-namespace UE::MIR::Internal {
+namespace MIR::Internal {
 
 //
 bool IsMaterialPropertyShared(EMaterialProperty InProperty);
@@ -51,7 +51,7 @@ bool Find(const TMap<TKey, TValue>& Map, const TKey& Key, TValue& OutValue)
 	return false;
 }
 
-} // namespace UE::MIR::Internal
+} // namespace MIR::Internal
 
 #endif // #if WITH_EDITOR
 
diff --git a/Engine/Source/Runtime/Engine/Private/Materials/MaterialIRModule.cpp b/Engine/Source/Runtime/Engine/Private/Materials/MaterialIRModule.cpp
--- a/Engine/Source/Runtime/Engine/Private/Materials/MaterialIRModule.cpp
+++ b/Engine/Source/Runtime/Engine/Private/Materials/MaterialIRModule.cpp
@@ -6,8 +6,6 @@
 
 #if WITH_EDITOR
 
-namespace MIR = UE::MIR;
-
 FMaterialIRModule::FMaterialIRModule()
 {
 	RootBlock = new MIR::FBlock;
diff --git a/Engine/Source/Runtime/Engine/Private/Materials/MaterialIRModuleBuilder.cpp b/Engine/Source/Runtime/Engine/Private/Materials/MaterialIRModuleBuilder.cpp
--- a/Engine/Source/Runtime/Engine/Private/Materials/MaterialIRModuleBuilder.cpp
+++ b/Engine/Source/Runtime/Engine/Private/Materials/MaterialIRModuleBuilder.cpp
@@ -4,7 +4,7 @@
 
 #if WITH_EDITOR
 
-#include "Materials/MaterialIRModuleBuilderAnalysis.h"
+#include "Materials/MaterialIRValueAnalyzer.h"
 #include "Materials/MaterialIRModule.h"
 #include "Materials/MaterialIRDebug.h"
 #include "Materials/MaterialIR.h"
@@ -22,8 +22,6 @@
 #include "Materials/MaterialInsights.h"
 #include "Engine/Texture.h"
 
-namespace MIR = UE::MIR;
-
 static TAutoConsoleVariable<int32> CVarMaterialIRDebugDumpLevel(
 	TEXT("r.Material.Translator.DebugDump"),
 	0,
@@ -38,8 +36,8 @@ struct FAnalysisContext
 	UMaterialExpressionMaterialFunctionCall* Call{};
 	TSet<UMaterialExpression*> BuiltExpressions{};
 	TArray<UMaterialExpression*> ExpressionStack{};
-	TMap<const FExpressionInput*, UE::MIR::FValue*> InputValues;
-	TMap<const FExpressionOutput*, UE::MIR::FValue*> OutputValues;
+	TMap<const FExpressionInput*, MIR::FValue*> InputValues;
+	TMap<const FExpressionOutput*, MIR::FValue*> OutputValues;
 
 	MIR::FValue* GetInputValue(const FExpressionInput* Input)
 	{
@@ -71,7 +69,7 @@ struct FMaterialIRModuleBuilderImpl
 	MIR::FEmitter* Emitter;
 	TArray<FAnalysisContext> AnalysisContextStack;
 	TArray<MIR::FInstruction*> InstructionStack{};
-	FAnalysisData AnalysisData;
+	FMaterialIRValueAnalyzer ValueAnalyzer;
 
 	void Step_Initialize()
 	{
@@ -86,7 +84,7 @@ struct FMaterialIRModuleBuilderImpl
 	{
 		// Prepare the array of FSetMaterialOutputInstr outputs from the material attributes inputs.
 		FMaterialInputDescription Input;
-		for (int32 Index = 0; UE::MIR::Internal::NextMaterialAttributeInput(Builder->Material, Index, Input); ++Index)
+		for (int32 Index = 0; MIR::Internal::NextMaterialAttributeInput(Builder->Material, Index, Input); ++Index)
 		{
 			EMaterialProperty Property = (EMaterialProperty)Index;
 
@@ -98,7 +96,7 @@ struct FMaterialIRModuleBuilderImpl
 			}
 			else if (!Input.Input->IsConnected())
 			{
-				Output->Arg = UE::MIR::Internal::CreateMaterialAttributeDefaultValue(*Emitter, Builder->Material, Property);
+				Output->Arg = MIR::Internal::CreateMaterialAttributeDefaultValue(*Emitter, Builder->Material, Property);
 			}
 			else
 			{
@@ -354,7 +352,7 @@ struct FMaterialIRModuleBuilderImpl
 						check(UseValue->GetUses().Num() == 0);
 
 						UseValue->SetFlags(MIR::VF_Analyzed);
-						AnalyzeValue(AnalysisData, UseValue);
+						ValueAnalyzer.Analyze(UseValue);
 					}
 					continue;
 				}
@@ -381,7 +379,7 @@ struct FMaterialIRModuleBuilderImpl
 			}
 
 			Instr->SetFlags(MIR::VF_Analyzed);
-			AnalyzeValue(AnalysisData, Instr);
+			ValueAnalyzer.Analyze(Instr);
 		}
 	}
 
@@ -522,10 +520,7 @@ bool FMaterialIRModuleBuilder::Build(FMaterialIRModule* TargetModule)
 
 	// Setup the Builder implementation
 	Impl.Emitter = &Emitter;
-	Impl.AnalysisData.Material = Material;
-	Impl.AnalysisData.Module = TargetModule;
-	Impl.AnalysisData.Insights = TargetInsights;
-	Impl.AnalysisData.CompilationOutput = &TargetModule->CompilationOutput;
+	Impl.ValueAnalyzer.Setup(Material, TargetModule, &TargetModule->CompilationOutput, TargetInsights);
 	
 	Impl.Step_Initialize();
 	Impl.Step_GenerateOutputInstructions();
@@ -544,14 +539,14 @@ bool FMaterialIRModuleBuilder::Build(FMaterialIRModule* TargetModule)
 	// Dump debugging information if requested 
 	switch (CVarMaterialIRDebugDumpLevel.GetValueOnGameThread())
 	{
-		case 2: UE::MIR::DebugDumpIRUseGraph(*TargetModule); // fallthrough
-		case 1: UE::MIR::DebugDumpIR(Material->GetFullName(), *TargetModule); // fallthrough
+		case 2: MIR::DebugDumpIRUseGraph(*TargetModule); // fallthrough
+		case 1: MIR::DebugDumpIR(Material->GetFullName(), *TargetModule); // fallthrough
 	}
 
 	return true;
 }
 
-namespace UE::MIR::Internal {
+namespace MIR::Internal {
 
 FValue* GetInputValue(FMaterialIRModuleBuilderImpl* Builder, const FExpressionInput* Input)
 {
@@ -568,6 +563,6 @@ void SetOutputValue(FMaterialIRModuleBuilderImpl* Builder, const FExpressionOutp
 	Builder->AnalysisContextStack.Last().SetOutputValue(Output, Value);
 }
 
-} // namespace UE::MIR::Internal
+} // namespace MIR::Internal
 
 #endif // #if WITH_EDITOR
diff --git a/Engine/Source/Runtime/Engine/Private/Materials/MaterialIRModuleBuilderAnalysis.cpp b/Engine/Source/Runtime/Engine/Private/Materials/MaterialIRModuleBuilderAnalysis.cpp
deleted file mode 100644
diff --git a/Engine/Source/Runtime/Engine/Private/Materials/MaterialIRModuleBuilderAnalysis.h b/Engine/Source/Runtime/Engine/Private/Materials/MaterialIRModuleBuilderAnalysis.h
deleted file mode 100644
diff --git a/Engine/Source/Runtime/Engine/Private/Materials/MaterialIRToHLSLTranslator.cpp b/Engine/Source/Runtime/Engine/Private/Materials/MaterialIRToHLSLTranslator.cpp
--- a/Engine/Source/Runtime/Engine/Private/Materials/MaterialIRToHLSLTranslator.cpp
+++ b/Engine/Source/Runtime/Engine/Private/Materials/MaterialIRToHLSLTranslator.cpp
@@ -20,8 +20,6 @@
 
 #if WITH_EDITOR
 
-namespace MIR = UE::MIR;
-
 enum ENoOp { NoOp };
 enum ENewLine { NewLine };
 enum EEndOfStatement { EndOfStatement };
diff --git a/Engine/Source/Runtime/Engine/Private/Materials/MaterialIRTypes.cpp b/Engine/Source/Runtime/Engine/Private/Materials/MaterialIRTypes.cpp
--- a/Engine/Source/Runtime/Engine/Private/Materials/MaterialIRTypes.cpp
+++ b/Engine/Source/Runtime/Engine/Private/Materials/MaterialIRTypes.cpp
@@ -4,7 +4,7 @@
 
 #if WITH_EDITOR
 
-namespace UE::MIR
+namespace MIR
 {
 
 const TCHAR* TypeKindToString(ETypeKind Kind)
@@ -380,6 +380,6 @@ FTypePtr FObjectType::GetTexture2D()
 	return &Instance;
 }
 
-} // namespace UE::MIR
+} // namespace MIR
 
 #endif // #if WITH_EDITOR
diff --git a/Engine/Source/Runtime/Engine/Private/Materials/MaterialIRValueAnalyzer.cpp b/Engine/Source/Runtime/Engine/Private/Materials/MaterialIRValueAnalyzer.cpp
new file mode 100644
--- /dev/null
+++ b/Engine/Source/Runtime/Engine/Private/Materials/MaterialIRValueAnalyzer.cpp
@@ -0,0 +1,197 @@
+// Copyright Epic Games, Inc. All Rights Reserved.
+
+#include "Runtime/Engine/Private/Materials/MaterialIRValueAnalyzer.h"
+
+#if WITH_EDITOR
+
+#include "Materials/MaterialIRModule.h"
+#include "Materials/MaterialIRDebug.h"
+#include "Materials/MaterialIR.h"
+#include "Materials/MaterialIRTypes.h"
+#include "Materials/MaterialIRInternal.h"
+#include "Materials/MaterialInsights.h"
+#include "MaterialShared.h"
+#include "Materials/MaterialInsights.h"
+#include "Engine/Texture.h"
+#include "Materials/Material.h"
+
+static void AnalyzeExternalInput(FMaterialIRValueAnalyzer& Analyzer, MIR::FExternalInput* ExternalInput)
+{
+	 Analyzer.Module->GetStatistics().ExternalInputUsedMask[SF_Vertex][(int)ExternalInput->Id] = true;
+	 Analyzer.Module->GetStatistics().ExternalInputUsedMask[SF_Pixel][(int)ExternalInput->Id] = true;
+}
+
+static void AnalyzeTextureObject(FMaterialIRValueAnalyzer& Analyzer, MIR::FTextureObject* TextureObject)
+{
+	EMaterialTextureParameterType ParamType = MIR::Internal::TextureMaterialValueTypeToParameterType(TextureObject->Texture->GetMaterialType());
+
+	FMaterialTextureParameterInfo ParamInfo {
+		.ParameterInfo = { "", EMaterialParameterAssociation::GlobalParameter, INDEX_NONE },
+		.TextureIndex = Analyzer.Material->GetReferencedTextures().Find(TextureObject->Texture),
+		.SamplerSource = SSM_FromTextureAsset, // TODO - Is this needed?
+	};
+			
+	check(ParamInfo.TextureIndex != INDEX_NONE);
+
+	TextureObject->Analysis_UniformParameterIndex = Analyzer.CompilationOutput->UniformExpressionSet.FindOrAddTextureParameter(ParamType, ParamInfo);
+}
+
+static void AnalyzeTextureUniformParameter(FMaterialIRValueAnalyzer& Analyzer, MIR::FUniformParameter* Parameter)
+{
+	EMaterialTextureParameterType ParamType = MIR::Internal::TextureMaterialValueTypeToParameterType(Parameter->Metadata.Value.Texture->GetMaterialType());
+	
+	FMaterialTextureParameterInfo ParamInfo {
+		.ParameterInfo = Parameter->Info,
+		.TextureIndex = Analyzer.Material->GetReferencedTextures().Find(Parameter->Metadata.Value.Texture),
+		.SamplerSource = SSM_FromTextureAsset, // TODO - Is this needed?
+		.VirtualTextureLayerIndex = 0xff,
+	};
+	
+	Parameter->Analysis_UniformParameterIndex = Analyzer.CompilationOutput->UniformExpressionSet.FindOrAddTextureParameter(ParamType, ParamInfo);
+}
+
+static void AnalyzePrimitiveUniformParameter(FMaterialIRValueAnalyzer& Analyzer, MIR::FUniformParameter* Parameter)
+{
+	FUniformExpressionSet& UniformExpressionSet = Analyzer.CompilationOutput->UniformExpressionSet;
+
+	UE::Shader::FValue DefaultValue;
+	switch (Parameter->Metadata.Value.Type)
+	{
+		case EMaterialParameterType::Scalar: DefaultValue = Parameter->Metadata.Value.AsScalar(); break;
+		case EMaterialParameterType::Vector: DefaultValue = Parameter->Metadata.Value.AsLinearColor(); break;
+		default: UE_MIR_TODO();
+	}
+
+	// TODO: GetParameterOverrideValueForCurrentFunction?
+
+	uint32 DefaultValueOffset;
+	if (!MIR::Internal::Find(Analyzer.UniformDefaultValueOffsets, DefaultValue, DefaultValueOffset))
+	{
+		DefaultValueOffset = UniformExpressionSet.AddDefaultParameterValue(DefaultValue);
+		Analyzer.UniformDefaultValueOffsets.Add(DefaultValue, DefaultValueOffset);
+	}
+
+	Parameter->Analysis_UniformParameterIndex = UniformExpressionSet.FindOrAddNumericParameter(Parameter->Metadata.Value.Type, Parameter->Info, DefaultValueOffset);
+
+	// Get the parameter type as primitive
+	MIR::FPrimitiveTypePtr ParameterType = Parameter->Type->AsPrimitive();
+
+	// Only int and float parameters supported for now (TODO)
+	check(ParameterType->ScalarKind == MIR::ScalarKind_Int || ParameterType->ScalarKind == MIR::ScalarKind_Float);
+
+	// Depending on the number of compoents this parameter has, try to find a previous allocation that may fit it.
+	uint32 ParameterUniformBufferOffset = UINT32_MAX;
+	uint32 NumComponents = ParameterType->GetNumComponents();
+	check(NumComponents > 0);
+	check(NumComponents <= 4); // Only vectors supported for now
+
+	uint32 UsedNumComponents;
+
+	for (UsedNumComponents = NumComponents; UsedNumComponents < 4; ++UsedNumComponents)
+	{
+		auto& FreeOffsetsForCurrentNumComponents = Analyzer.FreeOffsetsPerNumComponents[UsedNumComponents - 1];
+
+		// If there are no available slots at this nubmer of components, try with a larger one
+		if (FreeOffsetsForCurrentNumComponents.IsEmpty()) {
+			continue;
+		}
+		
+		// There is space for this number `NumComponents` in a a chunk with `i` free components
+		ParameterUniformBufferOffset = FreeOffsetsForCurrentNumComponents.Last();
+		FreeOffsetsForCurrentNumComponents.Pop();
+
+		break;
+	}
+	
+	// If UsedNumComponents is 4, it means we looked for all possible components sizes (1, 2 and 3) and
+	// could not find a chunk that would fit this parameter. Allocate a new chunk.
+	if (UsedNumComponents >= 4) {
+		check(UsedNumComponents == 4); 
+		ParameterUniformBufferOffset = UniformExpressionSet.AllocateFromUniformBuffer(16);
+	}
+
+	// Finally, check whether in the used slot there's any slack left, and if so, record it in the
+	// appropriate free-list.
+	int NumComponentsLeft = UsedNumComponents - NumComponents;
+	if (NumComponentsLeft > 0) {
+		int LeftOverOffset = ParameterUniformBufferOffset + NumComponents * 4;
+		Analyzer.FreeOffsetsPerNumComponents[NumComponentsLeft - 1].Push(LeftOverOffset);
+	}
+
+	// Add the parameter evaluation to the uniform data
+	UniformExpressionSet.AddNumericParameterEvaluation(Parameter->Analysis_UniformParameterIndex, ParameterUniformBufferOffset);
+
+	if (Analyzer.Insights)
+	{
+		// Push information about this parameter allocation to the insights.
+		// ParameterUniformBufferOffset is the byte offset of where the parameter will be inside
+		// the uniform buffer. Since the uniform buffer is an array of 4D vectors of 4-byte components,
+		// we first calculate UniformComponentIndexFromStart being the offset of the allocated
+		// parameter in the uniform buffer expressed in "components" by dividing the byte offset by 4 bytes.
+		// From this, we can work out the index of the 4D vector in the whole uniform buffer (stored in
+		// SlotInsight.UniformIndex) and the first component occupied (SlotInsight.UniformComponetOffset).
+		// Currently the allocation algorithm does not spread parameters across multiple uniform 4D constants,
+		// hence SlotInsight.ParameterComponentOffset is always zero.
+
+		uint32 UniformComponentIndexFromStart = ParameterUniformBufferOffset / 4;
+
+		FMaterialInsights::FUniformParameterAllocationInsight ParamInsight;
+		ParamInsight.UniformIndex = UniformComponentIndexFromStart / 4;
+		ParamInsight.FirstUniformComponent = UniformComponentIndexFromStart % 4;
+		ParamInsight.FirstParameterComponent = 0; // Explained above
+		ParamInsight.ComponentsCount = NumComponents;
+		ParamInsight.ParameterName = Parameter->Info.Name;
+
+		switch (ParameterType->ScalarKind)
+		{
+			case MIR::ScalarKind_Int: ParamInsight.ComponentType = FMaterialInsights::FUniformBufferSlotComponentType::CT_Int; break;
+			case MIR::ScalarKind_Float: ParamInsight.ComponentType = FMaterialInsights::FUniformBufferSlotComponentType::CT_Float; break;
+			default: UE_MIR_UNREACHABLE();
+		}
+
+		Analyzer.Insights->UniformParameterAllocationInsights.Push(ParamInsight);
+	}
+}
+
+static void AnalyzeUniformParameter(FMaterialIRValueAnalyzer& Analyzer, MIR::FUniformParameter* Parameter)
+{
+	if (Parameter->Type->IsObjectOfKind(MIR::ObjectKind_Texture2D))
+	{
+		AnalyzeTextureUniformParameter(Analyzer, Parameter);
+	}
+	else
+	{
+		check(Parameter->Type->AsPrimitive());
+		AnalyzePrimitiveUniformParameter(Analyzer, Parameter);
+	}
+}
+
+void FMaterialIRValueAnalyzer::Setup(UMaterial* InMaterial, FMaterialIRModule* InModule, FMaterialCompilationOutput* InCompilationOutput, FMaterialInsights* InInsights)
+{
+	Material = InMaterial;
+	Module = InModule;
+	CompilationOutput = InCompilationOutput;
+	Insights = InInsights;
+}
+
+void FMaterialIRValueAnalyzer::Analyze(MIR::FValue* Value)
+{
+	#define EXPAND_CASE(ValueType) \
+	case MIR::VK_##ValueType: \
+	{ \
+		Analyze##ValueType(*this, static_cast<MIR::F##ValueType*>(Value)); \
+		break; \
+	}
+
+	switch (Value->Kind)
+	{
+		EXPAND_CASE(ExternalInput);
+		EXPAND_CASE(TextureObject);
+		EXPAND_CASE(UniformParameter);
+		default: break;
+	}
+
+	#undef EXPAND_CASE
+}
+	
+#endif // #if WITH_EDITOR
diff --git a/Engine/Source/Runtime/Engine/Private/Materials/MaterialIRValueAnalyzer.h b/Engine/Source/Runtime/Engine/Private/Materials/MaterialIRValueAnalyzer.h
new file mode 100644
--- /dev/null
+++ b/Engine/Source/Runtime/Engine/Private/Materials/MaterialIRValueAnalyzer.h
@@ -0,0 +1,45 @@
+// Copyright Epic Games, Inc. All Rights Reserved.
+
+#pragma once
+
+#include "Materials/MaterialIRModuleBuilder.h"
+
+#if WITH_EDITOR
+
+#include "Materials/MaterialIRModule.h"
+#include "Materials/MaterialIRDebug.h"
+#include "Materials/MaterialIR.h"
+#include "Materials/MaterialIRTypes.h"
+#include "Materials/MaterialIRInternal.h"
+#include "MaterialShared.h"
+#include "Materials/MaterialInsights.h"
+#include "Engine/Texture.h"
+
+struct FMaterialIRValueAnalyzer
+{
+	//
+	void Setup(UMaterial* InMaterial, FMaterialIRModule* InModule, FMaterialCompilationOutput* InCompilationOutput, FMaterialInsights* InInsights = nullptr);
+		
+	//
+	void Analyze(MIR::FValue* Value);
+
+	//
+	UMaterial* Material{};
+	
+	//
+	FMaterialIRModule* Module{};
+	
+	//
+	FMaterialCompilationOutput* CompilationOutput{};
+		
+	//
+	TMap<UE::Shader::FValue, uint32> UniformDefaultValueOffsets{};
+		
+	//
+	TArray<uint32, TInlineAllocator<8>> FreeOffsetsPerNumComponents[3];
+		
+	//
+	FMaterialInsights* Insights{}; // Optional
+};
+
+#endif // #if WITH_EDITOR
diff --git a/Engine/Source/Runtime/Engine/Public/Materials/MaterialExpression.h b/Engine/Source/Runtime/Engine/Public/Materials/MaterialExpression.h
--- a/Engine/Source/Runtime/Engine/Public/Materials/MaterialExpression.h
+++ b/Engine/Source/Runtime/Engine/Public/Materials/MaterialExpression.h
@@ -27,7 +27,7 @@ class FMaterialExpressionKey;
 
 enum class EMaterialParameterType : uint8;
 
-namespace UE::MIR
+namespace MIR
 {
 	class FEmitter;
 }
@@ -275,7 +275,7 @@ class UMaterialExpression : public UObject
 	 * This function is responsible for doing the semantic analysis on this error, error checking and reporting and
 	 * implementing this expression's logic by emitting appropriate material IR values.
 	 */
-	ENGINE_API virtual void Build(UE::MIR::FEmitter& Emitter);
+	ENGINE_API virtual void Build(MIR::FEmitter& Emitter);
 
 	/* Material expressions can add shader tags that to a material.  Adding a shader tag helps to identify materials
 	   that use certain material expressions, and can be used to limit shader generation to only materials that have
diff --git a/Engine/Source/Runtime/Engine/Public/Materials/MaterialExpressionAbs.h b/Engine/Source/Runtime/Engine/Public/Materials/MaterialExpressionAbs.h
--- a/Engine/Source/Runtime/Engine/Public/Materials/MaterialExpressionAbs.h
+++ b/Engine/Source/Runtime/Engine/Public/Materials/MaterialExpressionAbs.h
@@ -27,7 +27,7 @@ public:
 
 	//~ Begin UMaterialExpression Interface
 #if WITH_EDITOR
-	virtual void Build(UE::MIR::FEmitter& Emitter) override;
+	virtual void Build(MIR::FEmitter& Emitter) override;
 	virtual int32 Compile(class FMaterialCompiler* Compiler, int32 OutputIndex) override;
 	virtual void GetCaption(TArray<FString>& OutCaptions) const override;
 #endif
diff --git a/Engine/Source/Runtime/Engine/Public/Materials/MaterialExpressionAdd.h b/Engine/Source/Runtime/Engine/Public/Materials/MaterialExpressionAdd.h
--- a/Engine/Source/Runtime/Engine/Public/Materials/MaterialExpressionAdd.h
+++ b/Engine/Source/Runtime/Engine/Public/Materials/MaterialExpressionAdd.h
@@ -32,7 +32,7 @@ public:
 
 	//~ Begin UMaterialExpression Interface
 #if WITH_EDITOR
-	virtual void Build(UE::MIR::FEmitter& Emitter) override;
+	virtual void Build(MIR::FEmitter& Emitter) override;
 	virtual int32 Compile(class FMaterialCompiler* Compiler, int32 OutputIndex) override;
 	virtual void GetCaption(TArray<FString>& OutCaptions) const override;
 	virtual FText GetKeywords() const override {return FText::FromString(TEXT("+"));}
diff --git a/Engine/Source/Runtime/Engine/Public/Materials/MaterialExpressionAppendVector.h b/Engine/Source/Runtime/Engine/Public/Materials/MaterialExpressionAppendVector.h
--- a/Engine/Source/Runtime/Engine/Public/Materials/MaterialExpressionAppendVector.h
+++ b/Engine/Source/Runtime/Engine/Public/Materials/MaterialExpressionAppendVector.h
@@ -25,7 +25,7 @@ public:
 
 	//~ Begin UMaterialExpression Interface
 #if WITH_EDITOR
-	virtual void Build(UE::MIR::FEmitter& Emitter) override;
+	virtual void Build(MIR::FEmitter& Emitter) override;
 	virtual int32 Compile(class FMaterialCompiler* Compiler, int32 OutputIndex) override;
 	virtual void GetCaption(TArray<FString>& OutCaptions) const override;
 #endif
diff --git a/Engine/Source/Runtime/Engine/Public/Materials/MaterialExpressionArccosine.h b/Engine/Source/Runtime/Engine/Public/Materials/MaterialExpressionArccosine.h
--- a/Engine/Source/Runtime/Engine/Public/Materials/MaterialExpressionArccosine.h
+++ b/Engine/Source/Runtime/Engine/Public/Materials/MaterialExpressionArccosine.h
@@ -20,7 +20,7 @@ public:
 
 	//~ Begin UMaterialExpression Interface
 #if WITH_EDITOR
-	virtual void Build(UE::MIR::FEmitter& Emitter) override;
+	virtual void Build(MIR::FEmitter& Emitter) override;
 	virtual int32 Compile(class FMaterialCompiler* Compiler, int32 OutputIndex) override;
 	virtual void GetCaption(TArray<FString>& OutCaptions) const override;
 	virtual void GetExpressionToolTip(TArray<FString>& OutToolTip) override;
diff --git a/Engine/Source/Runtime/Engine/Public/Materials/MaterialExpressionArcsine.h b/Engine/Source/Runtime/Engine/Public/Materials/MaterialExpressionArcsine.h
--- a/Engine/Source/Runtime/Engine/Public/Materials/MaterialExpressionArcsine.h
+++ b/Engine/Source/Runtime/Engine/Public/Materials/MaterialExpressionArcsine.h
@@ -20,7 +20,7 @@ public:
 
 	//~ Begin UMaterialExpression Interface
 #if WITH_EDITOR
-	virtual void Build(UE::MIR::FEmitter& Emitter) override;
+	virtual void Build(MIR::FEmitter& Emitter) override;
 	virtual int32 Compile(class FMaterialCompiler* Compiler, int32 OutputIndex) override;
 	virtual void GetCaption(TArray<FString>& OutCaptions) const override;
 	virtual void GetExpressionToolTip(TArray<FString>& OutToolTip) override;
diff --git a/Engine/Source/Runtime/Engine/Public/Materials/MaterialExpressionArctangent.h b/Engine/Source/Runtime/Engine/Public/Materials/MaterialExpressionArctangent.h
--- a/Engine/Source/Runtime/Engine/Public/Materials/MaterialExpressionArctangent.h
+++ b/Engine/Source/Runtime/Engine/Public/Materials/MaterialExpressionArctangent.h
@@ -20,7 +20,7 @@ public:
 
 	//~ Begin UMaterialExpression Interface
 #if WITH_EDITOR
-	virtual void Build(UE::MIR::FEmitter& Emitter) override;
+	virtual void Build(MIR::FEmitter& Emitter) override;
 	virtual int32 Compile(class FMaterialCompiler* Compiler, int32 OutputIndex) override;
 	virtual void GetCaption(TArray<FString>& OutCaptions) const override;
 	virtual void GetExpressionToolTip(TArray<FString>& OutToolTip) override;
diff --git a/Engine/Source/Runtime/Engine/Public/Materials/MaterialExpressionCeil.h b/Engine/Source/Runtime/Engine/Public/Materials/MaterialExpressionCeil.h
--- a/Engine/Source/Runtime/Engine/Public/Materials/MaterialExpressionCeil.h
+++ b/Engine/Source/Runtime/Engine/Public/Materials/MaterialExpressionCeil.h
@@ -22,7 +22,7 @@ public:
 
 	//~ Begin UMaterialExpression Interface
 #if WITH_EDITOR
-	virtual void Build(UE::MIR::FEmitter& Emitter) override;
+	virtual void Build(MIR::FEmitter& Emitter) override;
 	virtual int32 Compile(class FMaterialCompiler* Compiler, int32 OutputIndex) override;
 	virtual uint32 GetOutputType(int32 OutputIndex) override;
 	virtual void GetCaption(TArray<FString>& OutCaptions) const override;
diff --git a/Engine/Source/Runtime/Engine/Public/Materials/MaterialExpressionChannelMaskParameter.h b/Engine/Source/Runtime/Engine/Public/Materials/MaterialExpressionChannelMaskParameter.h
--- a/Engine/Source/Runtime/Engine/Public/Materials/MaterialExpressionChannelMaskParameter.h
+++ b/Engine/Source/Runtime/Engine/Public/Materials/MaterialExpressionChannelMaskParameter.h
@@ -34,7 +34,7 @@ class UMaterialExpressionChannelMaskParameter : public UMaterialExpressionVector
 	virtual bool SetParameterValue(FName InParameterName, FLinearColor InValue, EMaterialExpressionSetParameterValueFlags Flags) override;
 	virtual void PostEditChangeProperty(FPropertyChangedEvent& PropertyChangedEvent) override;
 
-	virtual void Build(UE::MIR::FEmitter& Emitter) override;
+	virtual void Build(MIR::FEmitter& Emitter) override;
 	virtual int32 Compile(class FMaterialCompiler* Compiler, int32 OutputIndex) override;
 	virtual void GetCaption(TArray<FString>& OutCaptions) const override;
 	
diff --git a/Engine/Source/Runtime/Engine/Public/Materials/MaterialExpressionClamp.h b/Engine/Source/Runtime/Engine/Public/Materials/MaterialExpressionClamp.h
--- a/Engine/Source/Runtime/Engine/Public/Materials/MaterialExpressionClamp.h
+++ b/Engine/Source/Runtime/Engine/Public/Materials/MaterialExpressionClamp.h
@@ -49,7 +49,7 @@ public:
 
 	//~ Begin UMaterialExpression Interface
 #if WITH_EDITOR
-	virtual void Build(UE::MIR::FEmitter& Emitter) override;
+	virtual void Build(MIR::FEmitter& Emitter) override;
 	virtual int32 Compile(class FMaterialCompiler* Compiler, int32 OutputIndex) override;
 	virtual void GetCaption(TArray<FString>& OutCaptions) const override;
 #endif
diff --git a/Engine/Source/Runtime/Engine/Public/Materials/MaterialExpressionConstant.h b/Engine/Source/Runtime/Engine/Public/Materials/MaterialExpressionConstant.h
--- a/Engine/Source/Runtime/Engine/Public/Materials/MaterialExpressionConstant.h
+++ b/Engine/Source/Runtime/Engine/Public/Materials/MaterialExpressionConstant.h
@@ -19,7 +19,7 @@ class UMaterialExpressionConstant : public UMaterialExpression
 
 	//~ Begin UMaterialExpression Interface
 #if WITH_EDITOR
-	virtual void Build(UE::MIR::FEmitter& Emitter) override;
+	virtual void Build(MIR::FEmitter& Emitter) override;
 	virtual int32 Compile(class FMaterialCompiler* Compiler, int32 OutputIndex) override;
 	virtual void GetCaption(TArray<FString>& OutCaptions) const override;
 	virtual FString GetDescription() const override;
diff --git a/Engine/Source/Runtime/Engine/Public/Materials/MaterialExpressionConstant2Vector.h b/Engine/Source/Runtime/Engine/Public/Materials/MaterialExpressionConstant2Vector.h
--- a/Engine/Source/Runtime/Engine/Public/Materials/MaterialExpressionConstant2Vector.h
+++ b/Engine/Source/Runtime/Engine/Public/Materials/MaterialExpressionConstant2Vector.h
@@ -23,7 +23,7 @@ class UMaterialExpressionConstant2Vector : public UMaterialExpression
 
 	//~ Begin UMaterialExpression Interface
 #if WITH_EDITOR
-	virtual void Build(UE::MIR::FEmitter& Emitter) override;
+	virtual void Build(MIR::FEmitter& Emitter) override;
 	virtual int32 Compile(class FMaterialCompiler* Compiler, int32 OutputIndex) override;
 	virtual void GetCaption(TArray<FString>& OutCaptions) const override;
 	virtual FString GetDescription() const override;
diff --git a/Engine/Source/Runtime/Engine/Public/Materials/MaterialExpressionConstant3Vector.h b/Engine/Source/Runtime/Engine/Public/Materials/MaterialExpressionConstant3Vector.h
--- a/Engine/Source/Runtime/Engine/Public/Materials/MaterialExpressionConstant3Vector.h
+++ b/Engine/Source/Runtime/Engine/Public/Materials/MaterialExpressionConstant3Vector.h
@@ -21,7 +21,7 @@ public:
 
 	//~ UMaterialExpression interface
 #if WITH_EDITOR
-	virtual void Build(UE::MIR::FEmitter& Emitter) override;
+	virtual void Build(MIR::FEmitter& Emitter) override;
 	virtual int32 Compile(class FMaterialCompiler* Compiler, int32 OutputIndex) override;
 	virtual void GetCaption(TArray<FString>& OutCaptions) const override;
 	virtual FString GetDescription() const override;
diff --git a/Engine/Source/Runtime/Engine/Public/Materials/MaterialExpressionConstant4Vector.h b/Engine/Source/Runtime/Engine/Public/Materials/MaterialExpressionConstant4Vector.h
--- a/Engine/Source/Runtime/Engine/Public/Materials/MaterialExpressionConstant4Vector.h
+++ b/Engine/Source/Runtime/Engine/Public/Materials/MaterialExpressionConstant4Vector.h
@@ -19,7 +19,7 @@ class UMaterialExpressionConstant4Vector : public UMaterialExpression
 
 	//~ Begin UMaterialExpression Interface
 #if WITH_EDITOR
-	virtual void Build(UE::MIR::FEmitter& Emitter) override;
+	virtual void Build(MIR::FEmitter& Emitter) override;
 	virtual int32 Compile(class FMaterialCompiler* Compiler, int32 OutputIndex) override;
 	virtual void GetCaption(TArray<FString>& OutCaptions) const override;
 	virtual FString GetDescription() const override;
diff --git a/Engine/Source/Runtime/Engine/Public/Materials/MaterialExpressionCosine.h b/Engine/Source/Runtime/Engine/Public/Materials/MaterialExpressionCosine.h
--- a/Engine/Source/Runtime/Engine/Public/Materials/MaterialExpressionCosine.h
+++ b/Engine/Source/Runtime/Engine/Public/Materials/MaterialExpressionCosine.h
@@ -24,7 +24,7 @@ public:
 
 	//~ Begin UMaterialExpression Interface
 #if WITH_EDITOR
-	virtual void Build(UE::MIR::FEmitter& Emitter) override;
+	virtual void Build(MIR::FEmitter& Emitter) override;
 	virtual int32 Compile(class FMaterialCompiler* Compiler, int32 OutputIndex) override;
 	virtual void GetCaption(TArray<FString>& OutCaptions) const override;
 	virtual FText GetKeywords() const override {return FText::FromString(TEXT("cos"));}
diff --git a/Engine/Source/Runtime/Engine/Public/Materials/MaterialExpressionCrossProduct.h b/Engine/Source/Runtime/Engine/Public/Materials/MaterialExpressionCrossProduct.h
--- a/Engine/Source/Runtime/Engine/Public/Materials/MaterialExpressionCrossProduct.h
+++ b/Engine/Source/Runtime/Engine/Public/Materials/MaterialExpressionCrossProduct.h
@@ -25,7 +25,7 @@ public:
 
 	//~ Begin UMaterialExpression Interface
 #if WITH_EDITOR
-	virtual void Build(UE::MIR::FEmitter& Emitter) override;
+	virtual void Build(MIR::FEmitter& Emitter) override;
 	virtual int32 Compile(class FMaterialCompiler* Compiler, int32 OutputIndex) override;
 	virtual void GetCaption(TArray<FString>& OutCaptions) const override;
 #endif
diff --git a/Engine/Source/Runtime/Engine/Public/Materials/MaterialExpressionDesaturation.h b/Engine/Source/Runtime/Engine/Public/Materials/MaterialExpressionDesaturation.h
--- a/Engine/Source/Runtime/Engine/Public/Materials/MaterialExpressionDesaturation.h
+++ b/Engine/Source/Runtime/Engine/Public/Materials/MaterialExpressionDesaturation.h
@@ -35,7 +35,7 @@ class UMaterialExpressionDesaturation : public UMaterialExpression
 
 	//~ Begin UMaterialExpression Interface
 #if WITH_EDITOR
-	virtual void Build(UE::MIR::FEmitter& Emitter) override;
+	virtual void Build(MIR::FEmitter& Emitter) override;
 	virtual int32 Compile(class FMaterialCompiler* Compiler, int32 OutputIndex) override;
 	virtual void GetCaption(TArray<FString>& OutCaptions) const override
 	{
diff --git a/Engine/Source/Runtime/Engine/Public/Materials/MaterialExpressionDistance.h b/Engine/Source/Runtime/Engine/Public/Materials/MaterialExpressionDistance.h
--- a/Engine/Source/Runtime/Engine/Public/Materials/MaterialExpressionDistance.h
+++ b/Engine/Source/Runtime/Engine/Public/Materials/MaterialExpressionDistance.h
@@ -25,7 +25,7 @@ public:
 
 	//~ Begin UMaterialExpression Interface
 #if WITH_EDITOR
-	virtual void Build(UE::MIR::FEmitter& Emitter) override;
+	virtual void Build(MIR::FEmitter& Emitter) override;
 	virtual int32 Compile(class FMaterialCompiler* Compiler, int32 OutputIndex) override;
 	virtual void GetCaption(TArray<FString>& OutCaptions) const override;
 #endif
diff --git a/Engine/Source/Runtime/Engine/Public/Materials/MaterialExpressionDivide.h b/Engine/Source/Runtime/Engine/Public/Materials/MaterialExpressionDivide.h
--- a/Engine/Source/Runtime/Engine/Public/Materials/MaterialExpressionDivide.h
+++ b/Engine/Source/Runtime/Engine/Public/Materials/MaterialExpressionDivide.h
@@ -33,7 +33,7 @@ public:
 
 	//~ Begin UMaterialExpression Interface
 #if WITH_EDITOR
-	virtual void Build(UE::MIR::FEmitter& Emitter) override;
+	virtual void Build(MIR::FEmitter& Emitter) override;
 	virtual int32 Compile(class FMaterialCompiler* Compiler, int32 OutputIndex) override;
 	virtual void GetCaption(TArray<FString>& OutCaptions) const override;
 	virtual FText GetKeywords() const override {return FText::FromString(TEXT("/"));}
diff --git a/Engine/Source/Runtime/Engine/Public/Materials/MaterialExpressionDotProduct.h b/Engine/Source/Runtime/Engine/Public/Materials/MaterialExpressionDotProduct.h
--- a/Engine/Source/Runtime/Engine/Public/Materials/MaterialExpressionDotProduct.h
+++ b/Engine/Source/Runtime/Engine/Public/Materials/MaterialExpressionDotProduct.h
@@ -25,7 +25,7 @@ public:
 
 	//~ Begin UMaterialExpression Interface
 #if WITH_EDITOR
-	virtual void Build(UE::MIR::FEmitter& Emitter) override;
+	virtual void Build(MIR::FEmitter& Emitter) override;
 	virtual int32 Compile(class FMaterialCompiler* Compiler, int32 OutputIndex) override;
 	virtual void GetCaption(TArray<FString>& OutCaptions) const override;
 #endif
diff --git a/Engine/Source/Runtime/Engine/Public/Materials/MaterialExpressionExponential.h b/Engine/Source/Runtime/Engine/Public/Materials/MaterialExpressionExponential.h
--- a/Engine/Source/Runtime/Engine/Public/Materials/MaterialExpressionExponential.h
+++ b/Engine/Source/Runtime/Engine/Public/Materials/MaterialExpressionExponential.h
@@ -18,7 +18,7 @@ public:
 
 	//~ Begin UMaterialExpression Interface
 #if WITH_EDITOR
-	virtual void Build(UE::MIR::FEmitter& Emitter) override;
+	virtual void Build(MIR::FEmitter& Emitter) override;
 	virtual int32 Compile(class FMaterialCompiler* Compiler, int32 OutputIndex) override;
 	virtual void GetCaption(TArray<FString>& OutCaptions) const override;
 	virtual void GetExpressionToolTip(TArray<FString>& OutToolTip) override;
diff --git a/Engine/Source/Runtime/Engine/Public/Materials/MaterialExpressionExponential2.h b/Engine/Source/Runtime/Engine/Public/Materials/MaterialExpressionExponential2.h
--- a/Engine/Source/Runtime/Engine/Public/Materials/MaterialExpressionExponential2.h
+++ b/Engine/Source/Runtime/Engine/Public/Materials/MaterialExpressionExponential2.h
@@ -18,7 +18,7 @@ public:
 
 	//~ Begin UMaterialExpression Interface
 #if WITH_EDITOR
-	virtual void Build(UE::MIR::FEmitter& Emitter) override;
+	virtual void Build(MIR::FEmitter& Emitter) override;
 	virtual int32 Compile(class FMaterialCompiler* Compiler, int32 OutputIndex) override;
 	virtual void GetCaption(TArray<FString>& OutCaptions) const override;
 	virtual void GetExpressionToolTip(TArray<FString>& OutToolTip) override;
diff --git a/Engine/Source/Runtime/Engine/Public/Materials/MaterialExpressionFloor.h b/Engine/Source/Runtime/Engine/Public/Materials/MaterialExpressionFloor.h
--- a/Engine/Source/Runtime/Engine/Public/Materials/MaterialExpressionFloor.h
+++ b/Engine/Source/Runtime/Engine/Public/Materials/MaterialExpressionFloor.h
@@ -22,7 +22,7 @@ public:
 
 	//~ Begin UMaterialExpression Interface
 #if WITH_EDITOR
-	virtual void Build(UE::MIR::FEmitter& Emitter) override;
+	virtual void Build(MIR::FEmitter& Emitter) override;
 	virtual int32 Compile(class FMaterialCompiler* Compiler, int32 OutputIndex) override;
 	virtual uint32 GetOutputType(int32 OutputIndex) override;
 	virtual void GetCaption(TArray<FString>& OutCaptions) const override;
diff --git a/Engine/Source/Runtime/Engine/Public/Materials/MaterialExpressionFmod.h b/Engine/Source/Runtime/Engine/Public/Materials/MaterialExpressionFmod.h
--- a/Engine/Source/Runtime/Engine/Public/Materials/MaterialExpressionFmod.h
+++ b/Engine/Source/Runtime/Engine/Public/Materials/MaterialExpressionFmod.h
@@ -25,7 +25,7 @@ public:
 
 	//~ Begin UMaterialExpression Interface
 #if WITH_EDITOR
-	virtual void Build(UE::MIR::FEmitter& Emitter) override;
+	virtual void Build(MIR::FEmitter& Emitter) override;
 	virtual int32 Compile(class FMaterialCompiler* Compiler, int32 OutputIndex) override;
 	virtual void GetCaption(TArray<FString>& OutCaptions) const override;
 	virtual FText GetKeywords() const override {return FText::FromString(TEXT("%"));}
diff --git a/Engine/Source/Runtime/Engine/Public/Materials/MaterialExpressionFrac.h b/Engine/Source/Runtime/Engine/Public/Materials/MaterialExpressionFrac.h
--- a/Engine/Source/Runtime/Engine/Public/Materials/MaterialExpressionFrac.h
+++ b/Engine/Source/Runtime/Engine/Public/Materials/MaterialExpressionFrac.h
@@ -22,7 +22,7 @@ public:
 
 	//~ Begin UMaterialExpression Interface
 #if WITH_EDITOR
-	virtual void Build(UE::MIR::FEmitter& Emitter) override;
+	virtual void Build(MIR::FEmitter& Emitter) override;
 	virtual int32 Compile(class FMaterialCompiler* Compiler, int32 OutputIndex) override;
 	virtual void GetCaption(TArray<FString>& OutCaptions) const override;
 #endif
diff --git a/Engine/Source/Runtime/Engine/Public/Materials/MaterialExpressionIf.h b/Engine/Source/Runtime/Engine/Public/Materials/MaterialExpressionIf.h
--- a/Engine/Source/Runtime/Engine/Public/Materials/MaterialExpressionIf.h
+++ b/Engine/Source/Runtime/Engine/Public/Materials/MaterialExpressionIf.h
@@ -44,7 +44,7 @@ public:
 
 	//~ Begin UMaterialExpression Interface
 #if WITH_EDITOR
-	virtual void Build(UE::MIR::FEmitter& Emitter) override;
+	virtual void Build(MIR::FEmitter& Emitter) override;
 	virtual int32 Compile(class FMaterialCompiler* Compiler, int32 OutputIndex) override;
 	virtual void GetCaption(TArray<FString>& OutCaptions) const override;
 	virtual uint32 GetInputType(int32 InputIndex) override;
diff --git a/Engine/Source/Runtime/Engine/Public/Materials/MaterialExpressionIfThenElse.h b/Engine/Source/Runtime/Engine/Public/Materials/MaterialExpressionIfThenElse.h
--- a/Engine/Source/Runtime/Engine/Public/Materials/MaterialExpressionIfThenElse.h
+++ b/Engine/Source/Runtime/Engine/Public/Materials/MaterialExpressionIfThenElse.h
@@ -28,7 +28,7 @@ public:
 	
 	//~ Begin UMaterialExpression Interface
 #if WITH_EDITOR
-	virtual void Build(UE::MIR::FEmitter& Emitter) override;
+	virtual void Build(MIR::FEmitter& Emitter) override;
 	virtual void GetCaption(TArray<FString>& OutCaptions) const override;
 	virtual uint32 GetInputType(int32 InputIndex) override;
 #endif
diff --git a/Engine/Source/Runtime/Engine/Public/Materials/MaterialExpressionLength.h b/Engine/Source/Runtime/Engine/Public/Materials/MaterialExpressionLength.h
--- a/Engine/Source/Runtime/Engine/Public/Materials/MaterialExpressionLength.h
+++ b/Engine/Source/Runtime/Engine/Public/Materials/MaterialExpressionLength.h
@@ -16,7 +16,7 @@ public:
 
 	//~ Begin UMaterialExpressionMaterialX Interface
 #if WITH_EDITOR
-	virtual void Build(UE::MIR::FEmitter& Emitter) override;
+	virtual void Build(MIR::FEmitter& Emitter) override;
 	virtual int32 Compile(class FMaterialCompiler* Compiler, int32 OutputIndex) override;
 	virtual void GetCaption(TArray<FString>& OutCaptions) const override;
 	virtual void GetExpressionToolTip(TArray<FString>& OutToolTip) override;
diff --git a/Engine/Source/Runtime/Engine/Public/Materials/MaterialExpressionLinearInterpolate.h b/Engine/Source/Runtime/Engine/Public/Materials/MaterialExpressionLinearInterpolate.h
--- a/Engine/Source/Runtime/Engine/Public/Materials/MaterialExpressionLinearInterpolate.h
+++ b/Engine/Source/Runtime/Engine/Public/Materials/MaterialExpressionLinearInterpolate.h
@@ -39,7 +39,7 @@ public:
 
 	//~ Begin UMaterialExpression Interface
 #if WITH_EDITOR
-	virtual void Build(UE::MIR::FEmitter& Emitter) override;
+	virtual void Build(MIR::FEmitter& Emitter) override;
 	virtual int32 Compile(class FMaterialCompiler* Compiler, int32 OutputIndex) override;
 	virtual void GetCaption(TArray<FString>& OutCaptions) const override;
 	virtual FText GetKeywords() const override {return FText::FromString(TEXT("lerp"));}
diff --git a/Engine/Source/Runtime/Engine/Public/Materials/MaterialExpressionLogarithm.h b/Engine/Source/Runtime/Engine/Public/Materials/MaterialExpressionLogarithm.h
--- a/Engine/Source/Runtime/Engine/Public/Materials/MaterialExpressionLogarithm.h
+++ b/Engine/Source/Runtime/Engine/Public/Materials/MaterialExpressionLogarithm.h
@@ -18,7 +18,7 @@ public:
 
 	//~ Begin UMaterialExpression Interface
 #if WITH_EDITOR
-	virtual void Build(UE::MIR::FEmitter& Emitter) override;
+	virtual void Build(MIR::FEmitter& Emitter) override;
 	virtual int32 Compile(class FMaterialCompiler* Compiler, int32 OutputIndex) override;
 	virtual void GetCaption(TArray<FString>& OutCaptions) const override;
 	virtual void GetExpressionToolTip(TArray<FString>& OutToolTip) override;
diff --git a/Engine/Source/Runtime/Engine/Public/Materials/MaterialExpressionLogarithm10.h b/Engine/Source/Runtime/Engine/Public/Materials/MaterialExpressionLogarithm10.h
--- a/Engine/Source/Runtime/Engine/Public/Materials/MaterialExpressionLogarithm10.h
+++ b/Engine/Source/Runtime/Engine/Public/Materials/MaterialExpressionLogarithm10.h
@@ -21,7 +21,7 @@ public:
 
 	//~ Begin UMaterialExpression Interface
 #if WITH_EDITOR
-	virtual void Build(UE::MIR::FEmitter& Emitter) override;
+	virtual void Build(MIR::FEmitter& Emitter) override;
 	virtual int32 Compile(class FMaterialCompiler* Compiler, int32 OutputIndex) override;
 	virtual void GetCaption(TArray<FString>& OutCaptions) const override;
 	virtual void GetExpressionToolTip(TArray<FString>& OutToolTip) override;
diff --git a/Engine/Source/Runtime/Engine/Public/Materials/MaterialExpressionLogarithm2.h b/Engine/Source/Runtime/Engine/Public/Materials/MaterialExpressionLogarithm2.h
--- a/Engine/Source/Runtime/Engine/Public/Materials/MaterialExpressionLogarithm2.h
+++ b/Engine/Source/Runtime/Engine/Public/Materials/MaterialExpressionLogarithm2.h
@@ -21,7 +21,7 @@ public:
 
 	//~ Begin UMaterialExpression Interface
 #if WITH_EDITOR
-	virtual void Build(UE::MIR::FEmitter& Emitter) override;
+	virtual void Build(MIR::FEmitter& Emitter) override;
 	virtual int32 Compile(class FMaterialCompiler* Compiler, int32 OutputIndex) override;
 	virtual void GetCaption(TArray<FString>& OutCaptions) const override;
 	virtual void GetExpressionToolTip(TArray<FString>& OutToolTip) override;
diff --git a/Engine/Source/Runtime/Engine/Public/Materials/MaterialExpressionMin.h b/Engine/Source/Runtime/Engine/Public/Materials/MaterialExpressionMin.h
--- a/Engine/Source/Runtime/Engine/Public/Materials/MaterialExpressionMin.h
+++ b/Engine/Source/Runtime/Engine/Public/Materials/MaterialExpressionMin.h
@@ -32,7 +32,7 @@ public:
 
 	//~ Begin UMaterialExpression Interface
 #if WITH_EDITOR
-	virtual void Build(UE::MIR::FEmitter& Emitter) override;
+	virtual void Build(MIR::FEmitter& Emitter) override;
 	virtual int32 Compile(class FMaterialCompiler* Compiler, int32 OutputIndex) override;
 	virtual void GetCaption(TArray<FString>& OutCaptions) const override;
 #endif
diff --git a/Engine/Source/Runtime/Engine/Public/Materials/MaterialExpressionMultiply.h b/Engine/Source/Runtime/Engine/Public/Materials/MaterialExpressionMultiply.h
--- a/Engine/Source/Runtime/Engine/Public/Materials/MaterialExpressionMultiply.h
+++ b/Engine/Source/Runtime/Engine/Public/Materials/MaterialExpressionMultiply.h
@@ -33,7 +33,7 @@ public:
 
 	//~ Begin UMaterialExpression Interface
 #if WITH_EDITOR
-	virtual void Build(UE::MIR::FEmitter& Emitter) override;
+	virtual void Build(MIR::FEmitter& Emitter) override;
 	virtual int32 Compile(class FMaterialCompiler* Compiler, int32 OutputIndex) override;
 	virtual void GetCaption(TArray<FString>& OutCaptions) const override;
 	virtual FText GetKeywords() const override {return FText::FromString(TEXT("*"));}
diff --git a/Engine/Source/Runtime/Engine/Public/Materials/MaterialExpressionOneMinus.h b/Engine/Source/Runtime/Engine/Public/Materials/MaterialExpressionOneMinus.h
--- a/Engine/Source/Runtime/Engine/Public/Materials/MaterialExpressionOneMinus.h
+++ b/Engine/Source/Runtime/Engine/Public/Materials/MaterialExpressionOneMinus.h
@@ -22,7 +22,7 @@ public:
 
 	//~ Begin UMaterialExpression Interface
 #if WITH_EDITOR
-	virtual void Build(UE::MIR::FEmitter& Emitter) override;
+	virtual void Build(MIR::FEmitter& Emitter) override;
 	virtual int32 Compile(class FMaterialCompiler* Compiler, int32 OutputIndex) override;
 	virtual void GetCaption(TArray<FString>& OutCaptions) const override;
 	virtual FText GetKeywords() const override {return FText::FromString(TEXT("1-x"));}
diff --git a/Engine/Source/Runtime/Engine/Public/Materials/MaterialExpressionParameter.h b/Engine/Source/Runtime/Engine/Public/Materials/MaterialExpressionParameter.h
--- a/Engine/Source/Runtime/Engine/Public/Materials/MaterialExpressionParameter.h
+++ b/Engine/Source/Runtime/Engine/Public/Materials/MaterialExpressionParameter.h
@@ -59,7 +59,7 @@ class UMaterialExpressionParameter : public UMaterialExpression
 		return true;
 	}
 
-	virtual void Build(UE::MIR::FEmitter& Emitter) override;
+	virtual void Build(MIR::FEmitter& Emitter) override;
 #endif
 	//~ End UMaterialExpression Interface
 
diff --git a/Engine/Source/Runtime/Engine/Public/Materials/MaterialExpressionReroute.h b/Engine/Source/Runtime/Engine/Public/Materials/MaterialExpressionReroute.h
--- a/Engine/Source/Runtime/Engine/Public/Materials/MaterialExpressionReroute.h
+++ b/Engine/Source/Runtime/Engine/Public/Materials/MaterialExpressionReroute.h
@@ -22,7 +22,7 @@ public:
 
 	//~ Begin UMaterialExpression Interface
 #if WITH_EDITOR
-	ENGINE_API virtual void Build(UE::MIR::FEmitter& Emitter) override;
+	ENGINE_API virtual void Build(MIR::FEmitter& Emitter) override;
 	ENGINE_API virtual int32 Compile(class FMaterialCompiler* Compiler, int32 OutputIndex) override;
 	ENGINE_API virtual int32 CompilePreview(FMaterialCompiler* Compiler, int32 OutputIndex) override;
 	ENGINE_API virtual void GetCaption(TArray<FString>& OutCaptions) const override;
diff --git a/Engine/Source/Runtime/Engine/Public/Materials/MaterialExpressionRound.h b/Engine/Source/Runtime/Engine/Public/Materials/MaterialExpressionRound.h
--- a/Engine/Source/Runtime/Engine/Public/Materials/MaterialExpressionRound.h
+++ b/Engine/Source/Runtime/Engine/Public/Materials/MaterialExpressionRound.h
@@ -20,7 +20,7 @@ public:
 
 	//~ Begin UMaterialExpression Interface
 #if WITH_EDITOR
-	virtual void Build(UE::MIR::FEmitter& Emitter) override;
+	virtual void Build(MIR::FEmitter& Emitter) override;
 	virtual int32 Compile(class FMaterialCompiler* Compiler, int32 OutputIndex) override;
 	virtual uint32 GetOutputType(int32 OutputIndex) override;
 	virtual void GetCaption(TArray<FString>& OutCaptions) const override;
diff --git a/Engine/Source/Runtime/Engine/Public/Materials/MaterialExpressionSaturate.h b/Engine/Source/Runtime/Engine/Public/Materials/MaterialExpressionSaturate.h
--- a/Engine/Source/Runtime/Engine/Public/Materials/MaterialExpressionSaturate.h
+++ b/Engine/Source/Runtime/Engine/Public/Materials/MaterialExpressionSaturate.h
@@ -20,7 +20,7 @@ public:
 
 	//~ Begin UMaterialExpression Interface
 #if WITH_EDITOR
-	virtual void Build(UE::MIR::FEmitter& Emitter) override;
+	virtual void Build(MIR::FEmitter& Emitter) override;
 	virtual int32 Compile(class FMaterialCompiler* Compiler, int32 OutputIndex) override;
 	virtual void GetCaption(TArray<FString>& OutCaptions) const override;
 	virtual void GetExpressionToolTip(TArray<FString>& OutToolTip) override;
diff --git a/Engine/Source/Runtime/Engine/Public/Materials/MaterialExpressionSine.h b/Engine/Source/Runtime/Engine/Public/Materials/MaterialExpressionSine.h
--- a/Engine/Source/Runtime/Engine/Public/Materials/MaterialExpressionSine.h
+++ b/Engine/Source/Runtime/Engine/Public/Materials/MaterialExpressionSine.h
@@ -24,7 +24,7 @@ public:
 
 	//~ Begin UMaterialExpression Interface
 #if WITH_EDITOR
-	virtual void Build(UE::MIR::FEmitter& Emitter) override;
+	virtual void Build(MIR::FEmitter& Emitter) override;
 	virtual int32 Compile(class FMaterialCompiler* Compiler, int32 OutputIndex) override;
 	virtual void GetCaption(TArray<FString>& OutCaptions) const override;
 	virtual FText GetKeywords() const override {return FText::FromString(TEXT("sin"));}
diff --git a/Engine/Source/Runtime/Engine/Public/Materials/MaterialExpressionSmoothStep.h b/Engine/Source/Runtime/Engine/Public/Materials/MaterialExpressionSmoothStep.h
--- a/Engine/Source/Runtime/Engine/Public/Materials/MaterialExpressionSmoothStep.h
+++ b/Engine/Source/Runtime/Engine/Public/Materials/MaterialExpressionSmoothStep.h
@@ -40,7 +40,7 @@ public:
 
 	//~ Begin UMaterialExpression Interface
 #if WITH_EDITOR
-	virtual void Build(UE::MIR::FEmitter& Emitter) override;
+	virtual void Build(MIR::FEmitter& Emitter) override;
 	virtual int32 Compile(class FMaterialCompiler* Compiler, int32 OutputIndex) override;
 	virtual void GetCaption(TArray<FString>& OutCaptions) const override;
 	virtual FText GetKeywords() const override { return FText::FromString(TEXT("SmoothStep")); }
diff --git a/Engine/Source/Runtime/Engine/Public/Materials/MaterialExpressionSquareRoot.h b/Engine/Source/Runtime/Engine/Public/Materials/MaterialExpressionSquareRoot.h
--- a/Engine/Source/Runtime/Engine/Public/Materials/MaterialExpressionSquareRoot.h
+++ b/Engine/Source/Runtime/Engine/Public/Materials/MaterialExpressionSquareRoot.h
@@ -22,7 +22,7 @@ public:
 
 	//~ Begin UMaterialExpression Interface
 #if WITH_EDITOR
-	virtual void Build(UE::MIR::FEmitter& Emitter) override;
+	virtual void Build(MIR::FEmitter& Emitter) override;
 	virtual int32 Compile(class FMaterialCompiler* Compiler, int32 OutputIndex) override;
 	virtual void GetCaption(TArray<FString>& OutCaptions) const override;
 	virtual FText GetKeywords() const override {return FText::FromString(TEXT("sqrt"));}
diff --git a/Engine/Source/Runtime/Engine/Public/Materials/MaterialExpressionStaticBool.h b/Engine/Source/Runtime/Engine/Public/Materials/MaterialExpressionStaticBool.h
--- a/Engine/Source/Runtime/Engine/Public/Materials/MaterialExpressionStaticBool.h
+++ b/Engine/Source/Runtime/Engine/Public/Materials/MaterialExpressionStaticBool.h
@@ -20,7 +20,7 @@ class UMaterialExpressionStaticBool : public UMaterialExpression
 
 	//~ Begin UMaterialExpression Interface
 #if WITH_EDITOR
-	virtual void Build(UE::MIR::FEmitter& Emitter) override;
+	virtual void Build(MIR::FEmitter& Emitter) override;
 	virtual int32 CompilePreview(class FMaterialCompiler* Compiler, int32 OutputIndex) override;
 	virtual int32 Compile(class FMaterialCompiler* Compiler, int32 OutputIndex) override;
 	virtual void GetCaption(TArray<FString>& OutCaptions) const override;
diff --git a/Engine/Source/Runtime/Engine/Public/Materials/MaterialExpressionStaticBoolParameter.h b/Engine/Source/Runtime/Engine/Public/Materials/MaterialExpressionStaticBoolParameter.h
--- a/Engine/Source/Runtime/Engine/Public/Materials/MaterialExpressionStaticBoolParameter.h
+++ b/Engine/Source/Runtime/Engine/Public/Materials/MaterialExpressionStaticBoolParameter.h
@@ -24,7 +24,7 @@ class UMaterialExpressionStaticBoolParameter : public UMaterialExpressionParamet
 public:
 	//~ Begin UMaterialExpression Interface
 #if WITH_EDITOR
-	virtual void Build(UE::MIR::FEmitter& Emitter) override;
+	virtual void Build(MIR::FEmitter& Emitter) override;
 	virtual int32 Compile(class FMaterialCompiler* Compiler, int32 OutputIndex) override;
 	virtual int32 CompilePreview(class FMaterialCompiler* Compiler, int32 OutputIndex) override;
 	virtual void GetCaption(TArray<FString>& OutCaptions) const override;
diff --git a/Engine/Source/Runtime/Engine/Public/Materials/MaterialExpressionStaticSwitch.h b/Engine/Source/Runtime/Engine/Public/Materials/MaterialExpressionStaticSwitch.h
--- a/Engine/Source/Runtime/Engine/Public/Materials/MaterialExpressionStaticSwitch.h
+++ b/Engine/Source/Runtime/Engine/Public/Materials/MaterialExpressionStaticSwitch.h
@@ -32,7 +32,7 @@ public:
 
 	//~ Begin UMaterialExpression Interface
 #if WITH_EDITOR
-	virtual void Build(UE::MIR::FEmitter& Emitter) override;
+	virtual void Build(MIR::FEmitter& Emitter) override;
 	virtual int32 Compile(class FMaterialCompiler* Compiler, int32 OutputIndex) override;
 	virtual void GetCaption(TArray<FString>& OutCaptions) const override;
 	virtual FName GetInputName(int32 InputIndex) const override;
diff --git a/Engine/Source/Runtime/Engine/Public/Materials/MaterialExpressionStaticSwitchParameter.h b/Engine/Source/Runtime/Engine/Public/Materials/MaterialExpressionStaticSwitchParameter.h
--- a/Engine/Source/Runtime/Engine/Public/Materials/MaterialExpressionStaticSwitchParameter.h
+++ b/Engine/Source/Runtime/Engine/Public/Materials/MaterialExpressionStaticSwitchParameter.h
@@ -25,7 +25,7 @@ public:
 
 	//~ Begin UMaterialExpression Interface
 #if WITH_EDITOR
-	virtual void Build(UE::MIR::FEmitter& Emitter) override;
+	virtual void Build(MIR::FEmitter& Emitter) override;
 	virtual FName GetInputName(int32 InputIndex) const override;
 	virtual int32 Compile(class FMaterialCompiler* Compiler, int32 OutputIndex) override;
 	virtual void GetCaption(TArray<FString>& OutCaptions) const override;
diff --git a/Engine/Source/Runtime/Engine/Public/Materials/MaterialExpressionSubtract.h b/Engine/Source/Runtime/Engine/Public/Materials/MaterialExpressionSubtract.h
--- a/Engine/Source/Runtime/Engine/Public/Materials/MaterialExpressionSubtract.h
+++ b/Engine/Source/Runtime/Engine/Public/Materials/MaterialExpressionSubtract.h
@@ -33,7 +33,7 @@ public:
 
 	//~ Begin UMaterialExpression Interface
 #if WITH_EDITOR
-	virtual void Build(UE::MIR::FEmitter& Emitter) override;
+	virtual void Build(MIR::FEmitter& Emitter) override;
 	virtual int32 Compile(class FMaterialCompiler* Compiler, int32 OutputIndex) override;
 	virtual void GetCaption(TArray<FString>& OutCaptions) const override;
 	virtual FText GetKeywords() const override {return FText::FromString(TEXT("-"));}
diff --git a/Engine/Source/Runtime/Engine/Public/Materials/MaterialExpressionTangent.h b/Engine/Source/Runtime/Engine/Public/Materials/MaterialExpressionTangent.h
--- a/Engine/Source/Runtime/Engine/Public/Materials/MaterialExpressionTangent.h
+++ b/Engine/Source/Runtime/Engine/Public/Materials/MaterialExpressionTangent.h
@@ -23,7 +23,7 @@ public:
 
 	//~ Begin UMaterialExpression Interface
 #if WITH_EDITOR
-	virtual void Build(UE::MIR::FEmitter& Emitter) override;
+	virtual void Build(MIR::FEmitter& Emitter) override;
 	virtual int32 Compile(class FMaterialCompiler* Compiler, int32 OutputIndex) override;
 	virtual void GetCaption(TArray<FString>& OutCaptions) const override;
 	virtual FText GetKeywords() const override {return FText::FromString(TEXT("tan"));}
diff --git a/Engine/Source/Runtime/Engine/Public/Materials/MaterialExpressionTextureCoordinate.h b/Engine/Source/Runtime/Engine/Public/Materials/MaterialExpressionTextureCoordinate.h
--- a/Engine/Source/Runtime/Engine/Public/Materials/MaterialExpressionTextureCoordinate.h
+++ b/Engine/Source/Runtime/Engine/Public/Materials/MaterialExpressionTextureCoordinate.h
@@ -36,7 +36,7 @@ class UMaterialExpressionTextureCoordinate : public UMaterialExpression
 
 	//~ Begin UMaterialExpression Interface
 #if WITH_EDITOR
-	virtual void Build(UE::MIR::FEmitter& Emitter) override;
+	virtual void Build(MIR::FEmitter& Emitter) override;
 	virtual int32 Compile(class FMaterialCompiler* Compiler, int32 OutputIndex) override;
 	virtual void GetCaption(TArray<FString>& OutCaptions) const override;
 	virtual FText GetKeywords() const override { return FText::FromString(TEXT("texcoord")); }
diff --git a/Engine/Source/Runtime/Engine/Public/Materials/MaterialExpressionTextureObject.h b/Engine/Source/Runtime/Engine/Public/Materials/MaterialExpressionTextureObject.h
--- a/Engine/Source/Runtime/Engine/Public/Materials/MaterialExpressionTextureObject.h
+++ b/Engine/Source/Runtime/Engine/Public/Materials/MaterialExpressionTextureObject.h
@@ -24,7 +24,7 @@ class UMaterialExpressionTextureObject : public UMaterialExpressionTextureBase
 #endif // WITH_EDITOR
 	//~ Begin UMaterialExpression Interface
 #if WITH_EDITOR
-	virtual void Build(UE::MIR::FEmitter& Emitter) override;
+	virtual void Build(MIR::FEmitter& Emitter) override;
 	virtual int32 Compile(class FMaterialCompiler* Compiler, int32 OutputIndex) override;
 	virtual int32 CompilePreview(class FMaterialCompiler* Compiler, int32 OutputIndex) override;
 	virtual void GetCaption(TArray<FString>& OutCaptions) const override;
diff --git a/Engine/Source/Runtime/Engine/Public/Materials/MaterialExpressionTextureObjectParameter.h b/Engine/Source/Runtime/Engine/Public/Materials/MaterialExpressionTextureObjectParameter.h
--- a/Engine/Source/Runtime/Engine/Public/Materials/MaterialExpressionTextureObjectParameter.h
+++ b/Engine/Source/Runtime/Engine/Public/Materials/MaterialExpressionTextureObjectParameter.h
@@ -26,7 +26,7 @@ class UMaterialExpressionTextureObjectParameter : public UMaterialExpressionText
 
 	//~ Begin UMaterialExpression Interface
 #if WITH_EDITOR
-	virtual void Build(UE::MIR::FEmitter& Emitter) override;
+	virtual void Build(MIR::FEmitter& Emitter) override;
 	virtual void GetCaption(TArray<FString>& OutCaptions) const override;	
 	virtual int32 Compile(class FMaterialCompiler* Compiler, int32 OutputIndex) override;
 	virtual int32 CompilePreview(class FMaterialCompiler* Compiler, int32 OutputIndex) override;
diff --git a/Engine/Source/Runtime/Engine/Public/Materials/MaterialExpressionTextureSample.h b/Engine/Source/Runtime/Engine/Public/Materials/MaterialExpressionTextureSample.h
--- a/Engine/Source/Runtime/Engine/Public/Materials/MaterialExpressionTextureSample.h
+++ b/Engine/Source/Runtime/Engine/Public/Materials/MaterialExpressionTextureSample.h
@@ -116,7 +116,7 @@ public:
 
 	//~ Begin UMaterialExpression Interface
 #if WITH_EDITOR
-	ENGINE_API virtual void Build(UE::MIR::FEmitter& Emitter) override;
+	ENGINE_API virtual void Build(MIR::FEmitter& Emitter) override;
 	ENGINE_API virtual TArrayView<FExpressionInput*> GetInputsView() override;
 	ENGINE_API virtual FExpressionInput* GetInput(int32 InputIndex) override;
 	ENGINE_API virtual FName GetInputName(int32 InputIndex) const override;
diff --git a/Engine/Source/Runtime/Engine/Public/Materials/MaterialExpressionTextureSampleParameter.h b/Engine/Source/Runtime/Engine/Public/Materials/MaterialExpressionTextureSampleParameter.h
--- a/Engine/Source/Runtime/Engine/Public/Materials/MaterialExpressionTextureSampleParameter.h
+++ b/Engine/Source/Runtime/Engine/Public/Materials/MaterialExpressionTextureSampleParameter.h
@@ -34,7 +34,7 @@ class UMaterialExpressionTextureSampleParameter : public UMaterialExpressionText
 
 	//~ Begin UMaterialExpression Interface
 #if WITH_EDITOR
-	ENGINE_API virtual void Build(UE::MIR::FEmitter& Emitter) override;
+	ENGINE_API virtual void Build(MIR::FEmitter& Emitter) override;
 	ENGINE_API virtual int32 Compile(class FMaterialCompiler* Compiler, int32 OutputIndex) override;
 	ENGINE_API virtual void GetCaption(TArray<FString>& OutCaptions) const override;
 	ENGINE_API virtual bool MatchesSearchQuery(const TCHAR* SearchQuery) override;
diff --git a/Engine/Source/Runtime/Engine/Public/Materials/MaterialExpressionTruncate.h b/Engine/Source/Runtime/Engine/Public/Materials/MaterialExpressionTruncate.h
--- a/Engine/Source/Runtime/Engine/Public/Materials/MaterialExpressionTruncate.h
+++ b/Engine/Source/Runtime/Engine/Public/Materials/MaterialExpressionTruncate.h
@@ -20,7 +20,7 @@ public:
 
 	//~ Begin UMaterialExpression Interface
 #if WITH_EDITOR
-	virtual void Build(UE::MIR::FEmitter& Emitter) override;
+	virtual void Build(MIR::FEmitter& Emitter) override;
 	virtual int32 Compile(class FMaterialCompiler* Compiler, int32 OutputIndex) override;
 	virtual uint32 GetOutputType(int32 OutputIndex) override;
 	virtual void GetCaption(TArray<FString>& OutCaptions) const override;
diff --git a/Engine/Source/Runtime/Engine/Public/Materials/MaterialExpressionVectorParameter.h b/Engine/Source/Runtime/Engine/Public/Materials/MaterialExpressionVectorParameter.h
--- a/Engine/Source/Runtime/Engine/Public/Materials/MaterialExpressionVectorParameter.h
+++ b/Engine/Source/Runtime/Engine/Public/Materials/MaterialExpressionVectorParameter.h
@@ -28,7 +28,7 @@ class UMaterialExpressionVectorParameter : public UMaterialExpressionParameter
 
 	//~ Begin UMaterialExpression Interface
 #if WITH_EDITOR
-	virtual void Build(UE::MIR::FEmitter& Emitter) override;
+	virtual void Build(MIR::FEmitter& Emitter) override;
 	virtual int32 Compile(class FMaterialCompiler* Compiler, int32 OutputIndex) override;
 	virtual void GetCaption(TArray<FString>& OutCaptions) const override;
 	virtual bool GetParameterValue(FMaterialParameterMetadata& OutMeta) const override
diff --git a/Engine/Source/Runtime/Engine/Public/Materials/MaterialIR.h b/Engine/Source/Runtime/Engine/Public/Materials/MaterialIR.h
--- a/Engine/Source/Runtime/Engine/Public/Materials/MaterialIR.h
+++ b/Engine/Source/Runtime/Engine/Public/Materials/MaterialIR.h
@@ -6,7 +6,7 @@
 
 #if WITH_EDITOR
 
-namespace UE::MIR {
+namespace MIR {
 
 enum EValueKind
 {
@@ -478,5 +478,5 @@ bool IsAnyNonNullAPoison(const FValue* Head, TTail... Tail)
 	}
 }
 
-} // namespace UE::MIR
+} // namespace MIR
 #endif
diff --git a/Engine/Source/Runtime/Engine/Public/Materials/MaterialIRCommon.h b/Engine/Source/Runtime/Engine/Public/Materials/MaterialIRCommon.h
--- a/Engine/Source/Runtime/Engine/Public/Materials/MaterialIRCommon.h
+++ b/Engine/Source/Runtime/Engine/Public/Materials/MaterialIRCommon.h
@@ -35,7 +35,7 @@ namespace UE::Shader
 
 namespace ERHIFeatureLevel { enum Type : int; }
 
-namespace UE::MIR
+namespace MIR
 {
 
 /* Types*/
diff --git a/Engine/Source/Runtime/Engine/Public/Materials/MaterialIRDebug.h b/Engine/Source/Runtime/Engine/Public/Materials/MaterialIRDebug.h
--- a/Engine/Source/Runtime/Engine/Public/Materials/MaterialIRDebug.h
+++ b/Engine/Source/Runtime/Engine/Public/Materials/MaterialIRDebug.h
@@ -5,7 +5,7 @@
 
 #if WITH_EDITOR
 
-namespace UE::MIR
+namespace MIR
 {
 
 // Dumps a Graphviz Dot file to "{ProjectLogDir}/MaterialIRDumpUseGraph.dot" containing the
@@ -17,6 +17,6 @@ void DebugDumpIRUseGraph(const FMaterialIRModule& Module);
 // Dumps the module to a human readable file in "{ProjectLogDir}/MaterialIRDump.txt"
 void DebugDumpIR(FStringView MaterialName, const FMaterialIRModule& Module);
 
-} // namespace UE::MIR
+} // namespace MIR
 
 #endif // #if WITH_EDITOR
diff --git a/Engine/Source/Runtime/Engine/Public/Materials/MaterialIREmitter.h b/Engine/Source/Runtime/Engine/Public/Materials/MaterialIREmitter.h
--- a/Engine/Source/Runtime/Engine/Public/Materials/MaterialIREmitter.h
+++ b/Engine/Source/Runtime/Engine/Public/Materials/MaterialIREmitter.h
@@ -9,7 +9,7 @@
 
 #define UE_MIR_CHECKPOINT() if (Emitter.IsInvalid()) return
 
-namespace UE::MIR {
+namespace MIR {
 
 enum class EVectorComponent : uint8_t
 {
@@ -207,6 +207,6 @@ private:
 };
 
 
-} // namespace UE::MIR
+} // namespace MIR
 
 #endif // #if WITH_EDITOR
diff --git a/Engine/Source/Runtime/Engine/Public/Materials/MaterialIRModule.h b/Engine/Source/Runtime/Engine/Public/Materials/MaterialIRModule.h
--- a/Engine/Source/Runtime/Engine/Public/Materials/MaterialIRModule.h
+++ b/Engine/Source/Runtime/Engine/Public/Materials/MaterialIRModule.h
@@ -39,8 +39,8 @@ public:
 	void Empty();
 	EShaderPlatform GetShaderPlatform() const { return ShaderPlatform; }
 	const FMaterialCompilationOutput& GetCompilationOutput() const { return CompilationOutput; }
-	TArrayView<const UE::MIR::FSetMaterialOutput* const> GetOutputs() const { return Outputs; }
-	const UE::MIR::FBlock& GetRootBlock() const { return *RootBlock; }
+	TArrayView<const MIR::FSetMaterialOutput* const> GetOutputs() const { return Outputs; }
+	const MIR::FBlock& GetRootBlock() const { return *RootBlock; }
 	TArrayView<const FError> GetErrors() const { return Errors; }
 	const FStatistics& GetStatistics() const { return Statistics; }
 	FStatistics& GetStatistics() { return Statistics; }
@@ -49,13 +49,13 @@ private:
 	EShaderPlatform ShaderPlatform;
 	FMaterialCompilationOutput CompilationOutput;
 	FMemStackBase Allocator{};
-	TArray<UE::MIR::FValue*> Values;
-	TArray<UE::MIR::FSetMaterialOutput*> Outputs;
+	TArray<MIR::FValue*> Values;
+	TArray<MIR::FSetMaterialOutput*> Outputs;
 	TArray<FError> Errors;
-	UE::MIR::FBlock* RootBlock;
+	MIR::FBlock* RootBlock;
 	FStatistics Statistics;
-
-	friend UE::MIR::FEmitter;
+	
+	friend MIR::FEmitter;
 	friend FMaterialIRModuleBuilder;
 	friend FMaterialIRModuleBuilderImpl;
 };
diff --git a/Engine/Source/Runtime/Engine/Public/Materials/MaterialIRTypes.h b/Engine/Source/Runtime/Engine/Public/Materials/MaterialIRTypes.h
--- a/Engine/Source/Runtime/Engine/Public/Materials/MaterialIRTypes.h
+++ b/Engine/Source/Runtime/Engine/Public/Materials/MaterialIRTypes.h
@@ -8,7 +8,7 @@
 
 #if WITH_EDITOR
 
-namespace UE::MIR
+namespace MIR
 {
 
 enum ETypeKind
@@ -146,6 +146,6 @@ struct FObjectType : FType
 	static FTypePtr GetTexture2D();
 };
 
-} // namespace UE::MIR
+} // namespace MIR
 
 #endif // #if WITH_EDITOR
