Revert the change to check overridden layers when loading cached data for layers as this actually was not necessary for the final fix (addressed elsewhere in the CL 39943095).

diff --git a/Engine/Source/Runtime/Engine/Private/Materials/MaterialCachedData.cpp b/Engine/Source/Runtime/Engine/Private/Materials/MaterialCachedData.cpp
--- a/Engine/Source/Runtime/Engine/Private/Materials/MaterialCachedData.cpp
+++ b/Engine/Source/Runtime/Engine/Private/Materials/MaterialCachedData.cpp
@@ -605,7 +605,7 @@ void FMaterialCachedExpressionData::UpdateForExpressions(const FMaterialCachedEx
 				MaterialLayers = Layers.GetRuntime();
 				EditorOnlyData->MaterialLayers = Layers.EditorOnly;
 				FMaterialLayersFunctions::Validate(MaterialLayers, EditorOnlyData->MaterialLayers);
-				LayersExpression->OverrideLayerGraph(&Layers);
+				LayersExpression->RebuildLayerGraph(false);
 			}
 		}
 		else if (UMaterialExpressionMaterialFunctionCall* FunctionCall = Cast<UMaterialExpressionMaterialFunctionCall>(Expression))
