Expanded option to visualize SVT in editor, which now has 3 options : Always, Never and Only In PIE.

diff --git a/Engine/Source/Runtime/Core/Public/UObject/FortniteMainBranchObjectVersions.inl b/Engine/Source/Runtime/Core/Public/UObject/FortniteMainBranchObjectVersions.inl
--- a/Engine/Source/Runtime/Core/Public/UObject/FortniteMainBranchObjectVersions.inl
+++ b/Engine/Source/Runtime/Core/Public/UObject/FortniteMainBranchObjectVersions.inl
@@ -577,5 +577,8 @@ FORTNITE_MAIN_VERSION(PCGInlineConstantDefaultValues, 187)
 // Add MaterialSubstrateSubsurfaceType type to UMaterialExpressionSubstrateSlabBSDF for replacing bUseSSSDifffusion
 FORTNITE_MAIN_VERSION(AddMaterialSubstrateSubsurfaceType, 188)
 
+// Added option to visualize runtime virtual textures' streamed mips only in PIE 
+FORTNITE_MAIN_VERSION(AddedRuntimeVirtualTextureUseStreamingMipsInEditorMode, 189)
+
 // -----<new versions can be added above this line>-------------------------------------------------
 #undef FORTNITE_MAIN_VERSION
\ No newline at end of file
diff --git a/Engine/Source/Runtime/Engine/Classes/Components/RuntimeVirtualTextureComponent.h b/Engine/Source/Runtime/Engine/Classes/Components/RuntimeVirtualTextureComponent.h
--- a/Engine/Source/Runtime/Engine/Classes/Components/RuntimeVirtualTextureComponent.h
+++ b/Engine/Source/Runtime/Engine/Classes/Components/RuntimeVirtualTextureComponent.h
@@ -27,6 +27,14 @@ enum class ERuntimeVirtualTextureMaterialQuality : uint8
 	Epic = 3,
 };
 
+UENUM()
+enum class ERuntimeVirtualTextureUseStreamingMipsInEditorMode : uint8
+{
+	Never UMETA(ToolTip = "Never show the streaming low mips in the editor"),
+	PIEOnly UMETA(ToolTip = "Only show the streaming low mips when Play In Editor is active"),
+	Always UMETA(ToolTip = "Always show the streaming low mips in the editor"),
+};
+
 /** Component used to place a URuntimeVirtualTexture in the world. */
 UCLASS(Blueprintable, ClassGroup = Rendering, HideCategories = (Activation, Collision, Cooking, HLOD, Mobility, LOD, Navigation, Object, Physics), MinimalAPI)
 class URuntimeVirtualTextureComponent : public USceneComponent
@@ -116,9 +124,15 @@ protected:
 	UPROPERTY(EditAnywhere, AdvancedDisplay, Category = StreamingVirtualTexture, meta = (DisplayName = "Use Only Streaming Texture"))
 	bool bUseStreamingMipsOnly = false;
 
-	/** Use streaming low mips when rendering in editor. Set true to view and debug the baked streaming low mips. */
+	/** 
+	 * Use streaming low mips when rendering this runtime virtual texture in the editor. Allows to visualize the baked streaming low mips. 
+	 * r.VT.RVT.StreamingMips.UseInEditor can also be used to allow this across all RVT components (for debugging purposes)
+	*/
 	UPROPERTY(EditAnywhere, AdvancedDisplay, Category = StreamingVirtualTexture, meta = (DisplayName = "View in Editor"))
-	bool bUseStreamingMipsInEditor = false;
+	ERuntimeVirtualTextureUseStreamingMipsInEditorMode UseStreamingMipsInEditorMode = ERuntimeVirtualTextureUseStreamingMipsInEditorMode::Never;
+
+	UPROPERTY()
+	bool bUseStreamingMipsInEditor_DEPRECATED = false;
 
 #if WITH_EDITOR
 	/** Delegate handle for our function called on PIE end. */
@@ -229,6 +243,7 @@ protected:
 	//~ Begin UObject Interface
 	ENGINE_API virtual void BeginDestroy() override;
 	ENGINE_API virtual bool IsReadyForFinishDestroy() override;
+	ENGINE_API virtual void Serialize(FArchive& Ar) override;
 #if WITH_EDITOR
 	ENGINE_API virtual bool CanEditChange(const FProperty* InProperty) const override;
 #endif
diff --git a/Engine/Source/Runtime/Engine/Private/Components/RuntimeVirtualTextureComponent.cpp b/Engine/Source/Runtime/Engine/Private/Components/RuntimeVirtualTextureComponent.cpp
--- a/Engine/Source/Runtime/Engine/Private/Components/RuntimeVirtualTextureComponent.cpp
+++ b/Engine/Source/Runtime/Engine/Private/Components/RuntimeVirtualTextureComponent.cpp
@@ -16,6 +16,7 @@
 #include "SceneInterface.h"
 #include "SceneUtils.h"
 #include "UnrealEngine.h"
+#include "UObject/FortniteMainBranchObjectVersion.h"
 #include "UObject/UObjectIterator.h"
 #include "UObject/UnrealType.h"
 #include "VT/RuntimeVirtualTexture.h"
@@ -96,6 +97,20 @@ bool URuntimeVirtualTextureComponent::IsActiveInWorld() const
 			|| (World->WorldType == EWorldType::PIE)));
 }
 
+void URuntimeVirtualTextureComponent::Serialize(FArchive& Ar)
+{
+	Ar.UsingCustomVersion(FFortniteMainBranchObjectVersion::GUID);
+
+	Super::Serialize(Ar);
+
+	if (Ar.IsLoading() && Ar.CustomVer(FFortniteMainBranchObjectVersion::GUID) < FFortniteMainBranchObjectVersion::AddedRuntimeVirtualTextureUseStreamingMipsInEditorMode)
+	{
+PRAGMA_DISABLE_DEPRECATION_WARNINGS
+		UseStreamingMipsInEditorMode = bUseStreamingMipsInEditor_DEPRECATED ? ERuntimeVirtualTextureUseStreamingMipsInEditorMode::Always : ERuntimeVirtualTextureUseStreamingMipsInEditorMode::Never;
+PRAGMA_ENABLE_DEPRECATION_WARNINGS
+	}
+}
+
 #if WITH_EDITOR
 
 void URuntimeVirtualTextureComponent::OnRegister()
@@ -306,10 +321,19 @@ bool URuntimeVirtualTextureComponent::IsStreamingLowMips(EShadingPath ShadingPat
 	if (GIsEditor)
 	{
 		const int32 ShowStreamingMipsInEditor = CVarVTStreamingMipsShowInEditor.GetValueOnAnyThread();
-		if (ShowStreamingMipsInEditor == 0 || (ShowStreamingMipsInEditor == 1 && !bUseStreamingMipsInEditor))
+		if (ShowStreamingMipsInEditor == 0)
 		{
 			return false;
 		}
+		else if (ShowStreamingMipsInEditor == 1)
+		{
+			UWorld* World = GetWorld();
+			if ((UseStreamingMipsInEditorMode == ERuntimeVirtualTextureUseStreamingMipsInEditorMode::Never)
+				|| ((UseStreamingMipsInEditorMode == ERuntimeVirtualTextureUseStreamingMipsInEditorMode::PIEOnly) && ((World == nullptr) || !World->IsPlayInEditor())))
+			{
+				return false;
+			}
+		}
 	}
 #endif
 	return VirtualTexture != nullptr && StreamingTexture != nullptr && StreamingTexture->GetVirtualTexture(ShadingPath) != nullptr && CVarVTStreamingMips.GetValueOnAnyThread();
@@ -517,7 +541,7 @@ void URuntimeVirtualTextureComponent::InitializeStreamingTexture(EShadingPath Sh
 bool URuntimeVirtualTextureComponent::CanEditChange(const FProperty* InProperty) const
 {
 	bool bCanEdit = Super::CanEditChange(InProperty);
-	if (InProperty->GetFName() == GET_MEMBER_NAME_CHECKED(URuntimeVirtualTextureComponent, bUseStreamingMipsInEditor) || 
+	if (InProperty->GetFName() == GET_MEMBER_NAME_CHECKED(URuntimeVirtualTextureComponent, UseStreamingMipsInEditorMode) ||
 		InProperty->GetFName() == GET_MEMBER_NAME_CHECKED(URuntimeVirtualTextureComponent, bUseStreamingMipsOnly))
 	{
 		bCanEdit &= GetVirtualTexture() != nullptr && GetStreamingTexture() != nullptr;
