Remove immediate enqueue of mesh deformer work when creating a skinned mesh scene proxy.
The use of immediate enqueue forces a slow path here.
We only do it this to force a first frame deformation update for the editor.
After this change components with a mesh deformer can appear in bind pose in editor.
This will be fixed in a subsequent CL.

diff --git a/Engine/Source/Runtime/Engine/Private/Components/SkinnedMeshComponent.cpp b/Engine/Source/Runtime/Engine/Private/Components/SkinnedMeshComponent.cpp
--- a/Engine/Source/Runtime/Engine/Private/Components/SkinnedMeshComponent.cpp
+++ b/Engine/Source/Runtime/Engine/Private/Components/SkinnedMeshComponent.cpp
@@ -1174,29 +1174,18 @@ void USkinnedMeshComponent::CreateRenderState_Concurrent(FRegisterComponentConte
 		}
 	}
 
+#if WITH_EDITORONLY_DATA
 	if (UMeshDeformerInstance* MeshDeformerInstance = GetMeshDeformerInstanceForLOD(UseLOD))
  	{
-
-#if WITH_EDITORONLY_DATA
-			for (TUniquePtr<FMeshDeformerGeometryReadbackRequest>& Request : GeometryReadbackRequests)
-			{
-				MeshDeformerInstance->RequestReadbackDeformerGeometry(MoveTemp(Request));
-			}
-#endif // WITH_EDITORONLY_DATA
-
-		// Enqueue immediate execution of work here to ensure that we have some deformer outputs written for the next frame.
-		UMeshDeformerInstance::FEnqueueWorkDesc Desc;
-		Desc.Scene = GetScene();
-		Desc.ExecutionGroup = UMeshDeformerInstance::ExecutionGroup_Immediate;
-		Desc.OwnerName = GetSkinnedAsset() != nullptr ? GetSkinnedAsset()->GetFName() : GetFName();
-		MeshDeformerInstance->EnqueueWork(Desc);
+		for (TUniquePtr<FMeshDeformerGeometryReadbackRequest>& Request : GeometryReadbackRequests)
+		{
+			MeshDeformerInstance->RequestReadbackDeformerGeometry(MoveTemp(Request));
+		}
 	}
 
-#if WITH_EDITORONLY_DATA
 	// Avoid unbound accumulation of readback requests, each readback request must be picked up by the deformer instance within the same frame or ignored otherwise
 	GeometryReadbackRequests.Reset();
-#endif // WITH_EDITORONLY_DATA 
-	
+#endif // WITH_EDITORONLY_DATA
 }
 
 void USkinnedMeshComponent::DestroyRenderState_Concurrent()
