[NetPhysics] Fix Delta Serialization issue where if the sender used Default as the delta source the receiver would use the latest delta source instead.

diff --git a/Engine/Source/Runtime/Engine/Private/PhysicsEngine/NetworkPhysicsComponent.cpp b/Engine/Source/Runtime/Engine/Private/PhysicsEngine/NetworkPhysicsComponent.cpp
--- a/Engine/Source/Runtime/Engine/Private/PhysicsEngine/NetworkPhysicsComponent.cpp
+++ b/Engine/Source/Runtime/Engine/Private/PhysicsEngine/NetworkPhysicsComponent.cpp
@@ -124,20 +124,32 @@ bool FNetworkPhysicsRewindDataProxy::NetSerializeBase(FArchive& Ar, class UPacka
 							{
 								Ar.SerializeIntPacked(DeltaSourceFrame);
 								
-								if (FNetworkPhysicsData* DeltaSourceData = GetDeltaSourceData(static_cast<int32>(DeltaSourceFrame) - 1))
+								FNetworkPhysicsData* DeltaSourceData = nullptr;
+								if (DeltaSourceFrame > 0)
 								{
-									NetData->SetDeltaSourceData(DeltaSourceData);
+									// Try get a valid delta source for frame
+									DeltaSourceData = GetDeltaSourceData(static_cast<int32>(DeltaSourceFrame) - 1);
 								}
 								else
 								{
-#if DEBUG_NETWORK_PHYSICS
-									UE_LOG(LogChaos, Warning, TEXT("Delta Serialization ISSUE, did not find delta source with frame: %d"), DeltaSourceFrame);
+									// Sender used default as delta source
+									DeltaSourceData = GetDeltaSourceData(/*Default*/ -2);
+								}
+								
+								if (!DeltaSourceData)
+								{
+#if DEBUG_NETWORK_PHYSICS_DELTASERIALIZATION
+									const int32 Frame = static_cast<int32>(DeltaSourceFrame) - 1;
+									UE_LOG(LogChaos, Warning, TEXT("[DEBUG Delta Serialization] %s ISSUE, did not find delta source with frame: %d  --  Name: %s")
+										, (Owner->HasServerWorld() ? TEXT("[SERVER]    ") : (Owner->IsLocallyControlled() ? TEXT("[AUTONOMOUS]") : TEXT("[SIMULATED] "))), Frame, *AActor::GetDebugName(Owner->GetOwner()));
 #endif
 									bDeltaSerializationIssue = true;
-									NetData->SetDeltaSourceData(GetDeltaSourceData(/*Default*/ -2));
+									DeltaSourceData = GetDeltaSourceData(/*Default*/ -2);
 								}
+								
+								NetData->SetDeltaSourceData(DeltaSourceData);
 							}
-							else
+							else // IsSaving
 							{
 								if (FNetworkPhysicsData* DeltaSourceData = GetDeltaSourceData(/*Latest*/ -1))
 								{
@@ -146,6 +158,13 @@ bool FNetworkPhysicsRewindDataProxy::NetSerializeBase(FArchive& Ar, class UPacka
 									ensure((DeltaSourceData->ServerFrame + 1) >= 0);
 									DeltaSourceFrame = static_cast<uint32>(DeltaSourceData->ServerFrame + 1); // +1 since ServerFrame has a default value of -1 and it needs to be serialized unsigned
 								}
+								else
+								{
+									ensureMsgf(false, TEXT("Delta Serialization failed to get the latest delta source when sending, should not happen. On the first send the latest index should be populated with a default value, not null."));
+									NetData->SetDeltaSourceData(GetDeltaSourceData(/*Default*/ -2));
+									DeltaSourceFrame = 0; // Set DeltaSourceFrame to 0 to indicate that default delta source was used
+								}
+
 								Ar.SerializeIntPacked(DeltaSourceFrame);
 							}
 						}
@@ -180,55 +199,100 @@ UE_NET_IMPLEMENT_NAMED_STRUCT_LASTRESORT_NETSERIALIZER_AND_REGISTRY_DELEGATES(Ne
 
 bool FNetworkPhysicsRewindDataInputProxy::NetSerialize(FArchive& Ar, class UPackageMap* Map, bool& bOutSuccess)
 {
-	return NetSerializeBase(Ar, Map, bOutSuccess, [this]() { return Owner->InputHelper->CreateUniqueRewindHistory(0); }, [this](const int32 Value) -> FNetworkPhysicsData* { return Owner->GetDeltaSourceInput(Value, /*bValueIsIndex*/ false); });
+	const bool bSuccess = NetSerializeBase(Ar, Map, bOutSuccess, [this]() { return Owner->InputHelper->CreateUniqueRewindHistory(0); }, [this](const int32 Value) -> FNetworkPhysicsData* { return Owner->GetDeltaSourceInput(Value, /*bValueIsIndex*/ false); });
+#if DEBUG_NETWORK_PHYSICS_DELTASERIALIZATION
+	if (bDeltaSerializationIssue)
+	{
+		UE_LOG(LogChaos, Warning, TEXT("		[DEBUG Delta Serialization] INPUT"));
+	}
+#endif
+	return bSuccess;
 }
 
 bool FNetworkPhysicsRewindDataRemoteInputProxy::NetSerialize(FArchive& Ar, class UPackageMap* Map, bool& bOutSuccess)
 {
-	return NetSerializeBase(Ar, Map, bOutSuccess, [this]() { return Owner->InputHelper->CreateUniqueRewindHistory(0); }, [this](const int32 Value) -> FNetworkPhysicsData* { return Owner->GetDeltaSourceInput(Value, /*bValueIsIndex*/ false); });
+	const bool bSuccess = NetSerializeBase(Ar, Map, bOutSuccess, [this]() { return Owner->InputHelper->CreateUniqueRewindHistory(0); }, [this](const int32 Value) -> FNetworkPhysicsData* { return Owner->GetDeltaSourceInput(Value, /*bValueIsIndex*/ false); });
+
+#if DEBUG_NETWORK_PHYSICS_DELTASERIALIZATION
+	if (bDeltaSerializationIssue) { UE_LOG(LogChaos, Warning, TEXT("		[DEBUG Delta Serialization] REMOTE INPUT")); }
+#endif
+
+	return bSuccess;
 }
 
 bool FNetworkPhysicsRewindDataStateProxy::NetSerialize(FArchive& Ar, class UPackageMap* Map, bool& bOutSuccess)
 {
-	return NetSerializeBase(Ar, Map, bOutSuccess, [this]() { return Owner->StateHelper->CreateUniqueRewindHistory(0); }, [this](const int32 Value) -> FNetworkPhysicsData* { return Owner->GetDeltaSourceState(Value, /*bValueIsIndex*/ false); });
+	const bool bSuccess = NetSerializeBase(Ar, Map, bOutSuccess, [this]() { return Owner->StateHelper->CreateUniqueRewindHistory(0); }, [this](const int32 Value) -> FNetworkPhysicsData* { return Owner->GetDeltaSourceState(Value, /*bValueIsIndex*/ false); });
+
+#if DEBUG_NETWORK_PHYSICS_DELTASERIALIZATION
+	if (bDeltaSerializationIssue) { UE_LOG(LogChaos, Warning, TEXT("		[DEBUG Delta Serialization] STATE")); }
+#endif
+
+	return bSuccess;
 }
 
 bool FNetworkPhysicsRewindDataImportantInputProxy::NetSerialize(FArchive& Ar, class UPackageMap* Map, bool& bOutSuccess)
 {
-	return NetSerializeBase(Ar, Map, bOutSuccess, [this]() { return Owner->InputHelper->CreateUniqueRewindHistory(0); }, [this](const int32 Value) -> FNetworkPhysicsData* { return Owner->GetDeltaSourceInput(Value, /*bValueIsIndex*/ false); });
+	const bool bSuccess = NetSerializeBase(Ar, Map, bOutSuccess, [this]() { return Owner->InputHelper->CreateUniqueRewindHistory(0); }, [this](const int32 Value) -> FNetworkPhysicsData* { return Owner->GetDeltaSourceInput(Value, /*bValueIsIndex*/ false); });
+
+#if DEBUG_NETWORK_PHYSICS_DELTASERIALIZATION
+	if (bDeltaSerializationIssue) { UE_LOG(LogChaos, Warning, TEXT("		[DEBUG Delta Serialization] IMPORTANT INPUT")); }
+#endif
+
+	return bSuccess;
 }
 
 bool FNetworkPhysicsRewindDataImportantStateProxy::NetSerialize(FArchive& Ar, class UPackageMap* Map, bool& bOutSuccess)
 {
-	return NetSerializeBase(Ar, Map, bOutSuccess, [this]() { return Owner->StateHelper->CreateUniqueRewindHistory(0); }, [this](const int32 Value) -> FNetworkPhysicsData* { return Owner->GetDeltaSourceState(Value, /*bValueIsIndex*/ false); });
+	const bool bSuccess = NetSerializeBase(Ar, Map, bOutSuccess, [this]() { return Owner->StateHelper->CreateUniqueRewindHistory(0); }, [this](const int32 Value) -> FNetworkPhysicsData* { return Owner->GetDeltaSourceState(Value, /*bValueIsIndex*/ false); });
+
+#if DEBUG_NETWORK_PHYSICS_DELTASERIALIZATION
+	if (bDeltaSerializationIssue) { UE_LOG(LogChaos, Warning, TEXT("		[DEBUG Delta Serialization] IMPORTANT STATE")); }
+#endif
+
+	return bSuccess;
 }
 
 bool FNetworkPhysicsRewindDataDeltaSourceInputProxy::NetSerialize(FArchive& Ar, class UPackageMap* Map, bool& bOutSuccess)
 {
+	bool bSuccess = true;
 	if (PhysicsReplicationCVars::ResimulationCVars::bUseDefaultDeltaForDeltaSourceReplication)
 	{
 		// Use default as base for delta serialization when sending delta source
-		return NetSerializeBase(Ar, Map, bOutSuccess, [this]() { return Owner->InputHelper->CreateUniqueRewindHistory(0); }, [this](const int32 Value) -> FNetworkPhysicsData* { return Owner->GetDeltaSourceInput(/*Default*/ -2, /*bValueIsIndex*/ false); });
+		bSuccess = NetSerializeBase(Ar, Map, bOutSuccess, [this]() { return Owner->InputHelper->CreateUniqueRewindHistory(0); }, [this](const int32 Value) -> FNetworkPhysicsData* { return Owner->GetDeltaSourceInput(/*Default*/ -2, /*bValueIsIndex*/ false); });
 	}
 	else
 	{
 		// Standard serialization for delta source
-		return NetSerializeBase(Ar, Map, bOutSuccess, [this]() { return Owner->InputHelper->CreateUniqueRewindHistory(0); }, nullptr);
+		bSuccess = NetSerializeBase(Ar, Map, bOutSuccess, [this]() { return Owner->InputHelper->CreateUniqueRewindHistory(0); }, nullptr);
 	}
+
+#if DEBUG_NETWORK_PHYSICS_DELTASERIALIZATION
+	if (bDeltaSerializationIssue) { UE_LOG(LogChaos, Warning, TEXT("		[DEBUG Delta Serialization] DELTA INPUT")); }
+#endif
+
+	return bSuccess;
 }
 
 bool FNetworkPhysicsRewindDataDeltaSourceStateProxy::NetSerialize(FArchive& Ar, class UPackageMap* Map, bool& bOutSuccess)
 {
+	bool bSuccess = true;
 	if (PhysicsReplicationCVars::ResimulationCVars::bUseDefaultDeltaForDeltaSourceReplication)
 	{
 		// Use default as base for delta serialization when sending delta source
-		return NetSerializeBase(Ar, Map, bOutSuccess, [this]() { return Owner->StateHelper->CreateUniqueRewindHistory(0); }, [this](const int32 Value) -> FNetworkPhysicsData* { return Owner->GetDeltaSourceState(/*Default*/ -2, /*bValueIsIndex*/ false); });
+		bSuccess = NetSerializeBase(Ar, Map, bOutSuccess, [this]() { return Owner->StateHelper->CreateUniqueRewindHistory(0); }, [this](const int32 Value) -> FNetworkPhysicsData* { return Owner->GetDeltaSourceState(/*Default*/ -2, /*bValueIsIndex*/ false); });
 	}
 	else
 	{
 		// Standard serialization for delta source
-		return NetSerializeBase(Ar, Map, bOutSuccess, [this]() { return Owner->StateHelper->CreateUniqueRewindHistory(0); }, nullptr);
+		bSuccess = NetSerializeBase(Ar, Map, bOutSuccess, [this]() { return Owner->StateHelper->CreateUniqueRewindHistory(0); }, nullptr);
 	}
+
+#if DEBUG_NETWORK_PHYSICS_DELTASERIALIZATION
+	if (bDeltaSerializationIssue) { UE_LOG(LogChaos, Warning, TEXT("		[DEBUG Delta Serialization] DELTA STATE")); }
+#endif
+
+	return bSuccess;
 }
 
 // --------------------------- Network Physics Callback ---------------------------
@@ -728,6 +792,7 @@ void UNetworkPhysicsComponent::OnRep_SetReplicatedDeltaSourceInput()
 {
 	if (ReplicatedDeltaSourceInput.bDeltaSerializationIssue)
 	{
+		ensureMsgf(false, TEXT("OnRep_SetReplicatedDeltaSourceInput failed delta serialization, should not happen."));
 		return;
 	}
 	AddDeltaSourceInput();
@@ -742,6 +807,11 @@ void UNetworkPhysicsComponent::ServerReceiveDeltaSourceInputFrame_Implementation
 		{
 			// Set latest delta source index acknowledged by the client so that we can start using this delta source
 			LatestAcknowledgedDeltaSourceInputIndex = I;
+
+#if DEBUG_NETWORK_PHYSICS_DELTASERIALIZATION
+			UE_LOG(LogChaos, Log, TEXT("[DEBUG Delta Serialization] [SERVER]     Owner acknowledged delta source INPUT frame: %d at index: %d  --  Name: %s"), Frame, LatestAcknowledgedDeltaSourceInputIndex, *AActor::GetDebugName(GetOwner()));
+#endif
+
 			break;
 		}
 	}
@@ -764,6 +834,9 @@ void UNetworkPhysicsComponent::AddDeltaSourceInput()
 			// On the client, set the latest index, to be used when sending inputs to the server
 			LatestAcknowledgedDeltaSourceInputIndex = Index;
 
+#if DEBUG_NETWORK_PHYSICS_DELTASERIALIZATION
+			UE_LOG(LogChaos, Log, TEXT("[DEBUG Delta Serialization] %s Received delta source INPUT for frame: %d at index: %d  --  Name: %s"), (IsLocallyControlled() ? TEXT("[AUTONOMOUS]") : TEXT("[SIMULATED] ")), ReplicatedDeltaSourceInput.History->GetLatestFrame(), LatestAcknowledgedDeltaSourceInputIndex, *AActor::GetDebugName(GetOwner()));
+#endif
 			if (IsLocallyControlled())
 			{
 				// If this client is the one controlling this entity, send back an acknowledgment to the server we have received the delta source for ServerFrame
@@ -776,6 +849,13 @@ void UNetworkPhysicsComponent::AddDeltaSourceInput()
 			LatestAcknowledgedDeltaSourceInputIndex = Index;
 		}
 
+#if DEBUG_NETWORK_PHYSICS_DELTASERIALIZATION
+		if (HasServerWorld())
+		{
+			UE_LOG(LogChaos, Log, TEXT("[DEBUG Delta Serialization] [SERVER]     Sent delta source INPUT for frame: %d at index: %d  --  Name: %s"), ReplicatedDeltaSourceInput.History->GetLatestFrame(), Index, *AActor::GetDebugName(GetOwner()));
+		}
+#endif
+
 		LatestCachedDeltaSourceInputIndex = Index;
 	}
 }
@@ -784,6 +864,7 @@ void UNetworkPhysicsComponent::OnRep_SetReplicatedDeltaSourceState()
 {
 	if (ReplicatedDeltaSourceState.bDeltaSerializationIssue)
 	{
+		ensureMsgf(false, TEXT("OnRep_SetReplicatedDeltaSourceState failed delta serialization, should not happen."));
 		return;
 	}
 	AddDeltaSourceState();
@@ -798,6 +879,11 @@ void UNetworkPhysicsComponent::ServerReceiveDeltaSourceStateFrame_Implementation
 		{
 			// Set latest delta source index acknowledged by the client so that we can start using this delta source
 			LatestAcknowledgedDeltaSourceStateIndex = I;
+
+#if DEBUG_NETWORK_PHYSICS_DELTASERIALIZATION
+			UE_LOG(LogChaos, Log, TEXT("[DEBUG Delta Serialization] [SERVER]     Owner acknowledged delta source STATE frame: %d at index: %d  --  Name: %s"), Frame, LatestAcknowledgedDeltaSourceStateIndex, *AActor::GetDebugName(GetOwner()));
+#endif
+
 			break;
 		}
 	}
@@ -820,6 +906,10 @@ void UNetworkPhysicsComponent::AddDeltaSourceState()
 			// On the client, set the latest index (unlike for DeltaSourceInput this latest index is not used on the client since the client doesn't send states towards the server, but set the value for the future)
 			LatestAcknowledgedDeltaSourceStateIndex = Index;
 
+#if DEBUG_NETWORK_PHYSICS_DELTASERIALIZATION
+			UE_LOG(LogChaos, Log, TEXT("[DEBUG Delta Serialization] %s Received delta source STATE for frame: %d at index: %d  --  Name: %s"), (IsLocallyControlled() ? TEXT("[AUTONOMOUS]") : TEXT("[SIMULATED] ")), ReplicatedDeltaSourceState.History->GetLatestFrame(), LatestAcknowledgedDeltaSourceStateIndex, *AActor::GetDebugName(GetOwner()));
+#endif
+
 			if (IsLocallyControlled())
 			{
 				// If this client is the one controlling this entity, send back an acknowledgment to the server we have received the delta source for ServerFrame
@@ -832,6 +922,13 @@ void UNetworkPhysicsComponent::AddDeltaSourceState()
 			LatestAcknowledgedDeltaSourceStateIndex = Index;
 		}
 
+#if DEBUG_NETWORK_PHYSICS_DELTASERIALIZATION
+		if (HasServerWorld())
+		{
+			UE_LOG(LogChaos, Log, TEXT("[DEBUG Delta Serialization] [SERVER]     Sent delta source STATE for frame: %d at index: %d  --  Name: %s"), ReplicatedDeltaSourceState.History->GetLatestFrame(), Index, *AActor::GetDebugName(GetOwner()));
+		}
+#endif
+
 		LatestCachedDeltaSourceStateIndex = Index;
 	}
 }
@@ -904,7 +1001,13 @@ FNetworkPhysicsData* UNetworkPhysicsComponent::GetDeltaSourceState(const int32 V
 
 void UNetworkPhysicsComponent::OnRep_SetReplicatedStates()
 {
-	if (!NetworkPhysicsComponent_Internal || !StateHelper || !ReplicatedStates.History || ReplicatedStates.bDeltaSerializationIssue)
+	if (ReplicatedStates.bDeltaSerializationIssue)
+	{
+		ensureMsgf(!IsLocallyControlled(), TEXT("OnRep_SetReplicatedStates failed delta serialization, should not happen on the owning client."));
+		return;
+	}
+
+	if (!NetworkPhysicsComponent_Internal || !StateHelper || !ReplicatedStates.History)
 	{
 		return;
 	}
@@ -924,7 +1027,13 @@ void UNetworkPhysicsComponent::OnRep_SetReplicatedStates()
 
 void UNetworkPhysicsComponent::OnRep_SetReplicatedInputs()
 {
-	if (!NetworkPhysicsComponent_Internal || !InputHelper || !ReplicatedInputs.History || ReplicatedInputs.bDeltaSerializationIssue)
+	if (ReplicatedInputs.bDeltaSerializationIssue)
+	{
+		ensureMsgf(!IsLocallyControlled(), TEXT("OnRep_SetReplicatedInputs failed delta serialization, should not happen on the owning client."));
+		return;
+	}
+
+	if (!NetworkPhysicsComponent_Internal || !InputHelper || !ReplicatedInputs.History)
 	{
 		return ;
 	}
@@ -944,7 +1053,13 @@ void UNetworkPhysicsComponent::OnRep_SetReplicatedInputs()
 
 void UNetworkPhysicsComponent::OnRep_SetReplicatedRemoteInputs()
 {
-	if (!NetworkPhysicsComponent_Internal || !InputHelper || !ReplicatedRemoteInputs.History || ReplicatedRemoteInputs.bDeltaSerializationIssue)
+	if (ReplicatedRemoteInputs.bDeltaSerializationIssue)
+	{
+		ensureMsgf(!IsLocallyControlled(), TEXT("OnRep_SetReplicatedRemoteInputs failed delta serialization, should not happen on the owning client."));
+		return;
+	}
+
+	if (!NetworkPhysicsComponent_Internal || !InputHelper || !ReplicatedRemoteInputs.History)
 	{
 		return ;
 	}
@@ -964,7 +1079,13 @@ void UNetworkPhysicsComponent::OnRep_SetReplicatedRemoteInputs()
 
 void UNetworkPhysicsComponent::ServerReceiveInputData_Implementation(const FNetworkPhysicsRewindDataInputProxy& ClientInputs)
 {
-	if (!NetworkPhysicsComponent_Internal || !InputHelper || !ClientInputs.History || ClientInputs.bDeltaSerializationIssue)
+	if (ClientInputs.bDeltaSerializationIssue)
+	{
+		ensureMsgf(false, TEXT("ServerReceiveInputData_Implementation failed delta serialization, should not happen on the server."));
+		return;
+	}
+
+	if (!NetworkPhysicsComponent_Internal || !InputHelper || !ClientInputs.History)
 	{
 		return;
 	}
@@ -996,7 +1117,13 @@ void UNetworkPhysicsComponent::ServerReceiveInputData_Implementation(const FNetw
 
 void UNetworkPhysicsComponent::ServerReceiveImportantInputData_Implementation(const FNetworkPhysicsRewindDataImportantInputProxy& ClientInputs)
 {
-	if (!NetworkPhysicsComponent_Internal || !ClientInputs.History || ClientInputs.bDeltaSerializationIssue)
+	if (ClientInputs.bDeltaSerializationIssue)
+	{
+		ensureMsgf(false, TEXT("ServerReceiveImportantInputData_Implementation failed delta serialization, should not happen on the server."));
+		return;
+	}
+
+	if (!NetworkPhysicsComponent_Internal || !ClientInputs.History)
 	{
 		return;
 	}
@@ -1025,7 +1152,13 @@ void UNetworkPhysicsComponent::MulticastReceiveImportantInputData_Implementation
 		return;
 	}
 
-	if (!NetworkPhysicsComponent_Internal || !ServerInputs.History || ServerInputs.bDeltaSerializationIssue)
+	if (ServerInputs.bDeltaSerializationIssue)
+	{
+		ensureMsgf(!IsLocallyControlled(), TEXT("MulticastReceiveImportantInputData_Implementation failed delta serialization, should not happen on the owning client."));
+		return;
+	}
+
+	if (!NetworkPhysicsComponent_Internal || !ServerInputs.History)
 	{
 		return;
 	}
@@ -1048,7 +1181,13 @@ void UNetworkPhysicsComponent::MulticastReceiveImportantStateData_Implementation
 		return;
 	}
 
-	if (!NetworkPhysicsComponent_Internal || !ServerStates.History || ServerStates.bDeltaSerializationIssue)
+	if (ServerStates.bDeltaSerializationIssue)
+	{
+		ensureMsgf(!IsLocallyControlled(), TEXT("MulticastReceiveImportantStateData_Implementation failed delta serialization, should not happen on the owning client."));
+		return;
+	}
+
+	if (!NetworkPhysicsComponent_Internal || !ServerStates.History)
 	{
 		return;
 	}
diff --git a/Engine/Source/Runtime/Engine/Public/Physics/NetworkPhysicsComponent.h b/Engine/Source/Runtime/Engine/Public/Physics/NetworkPhysicsComponent.h
--- a/Engine/Source/Runtime/Engine/Public/Physics/NetworkPhysicsComponent.h
+++ b/Engine/Source/Runtime/Engine/Public/Physics/NetworkPhysicsComponent.h
@@ -15,6 +15,10 @@
 
 #include "NetworkPhysicsComponent.generated.h"
 
+#ifndef DEBUG_NETWORK_PHYSICS_DELTASERIALIZATION
+#define DEBUG_NETWORK_PHYSICS_DELTASERIALIZATION 0
+#endif
+
 class FAsyncNetworkPhysicsComponent;
 
 DECLARE_MULTICAST_DELEGATE_OneParam(FOnPreProcessInputsInternal, const int32);
