New material translator:
- Fixed several issues with the implementation in the new translator of uniform parameters.

#rnx

diff --git a/Engine/Source/Runtime/Engine/Private/Materials/MaterialIRToHLSLTranslator.cpp b/Engine/Source/Runtime/Engine/Private/Materials/MaterialIRToHLSLTranslator.cpp
--- a/Engine/Source/Runtime/Engine/Private/Materials/MaterialIRToHLSLTranslator.cpp
+++ b/Engine/Source/Runtime/Engine/Private/Materials/MaterialIRToHLSLTranslator.cpp
@@ -406,14 +406,37 @@ struct FTranslator : FMaterialIRToHLSLTranslation
 	{
 		MIR::FPrimitiveTypePtr PrimitiveType = UniformParameter->Type->AsPrimitive();
 		check(PrimitiveType);
+		check(PrimitiveType->IsScalar() || PrimitiveType->IsVector()); // no matrices yet
 
 		if (PrimitiveType->ScalarKind == MIR::EScalarKind::ScalarKind_Int)
 		{
 			Printer << "asint(";
 		}
 
-		Printer << TEXT("Material.PreshaderBuffer[") << UniformParameter->Analysis_UniformParameterIndex << TEXT("]");
-		check(PrimitiveType->IsScalar() || PrimitiveType->IsVector()); // no matrices yet
+		const FUniformExpressionSet& UniformExpressionSet = Module->GetCompilationOutput().UniformExpressionSet;
+		
+		// Get the global float4 component index (e.g. if this is 13, it refer to PreshaderBuffer[3].y)
+		uint32 GlobalComponentOffset = UniformExpressionSet.GetNumericParameterEvaluationOffset(UniformParameter->Analysis_UniformParameterIndex);
+		
+		// Index of the float4 slot
+		uint32 BufferSlotIndex = GlobalComponentOffset / 4;
+
+		// Starting component of the float4 slot
+		uint32 BufferSlotOffset = GlobalComponentOffset % 4;
+
+		Printer << TEXT("Material.PreshaderBuffer[") << BufferSlotIndex << TEXT("]");
+
+		if (PrimitiveType->GetNumComponents() < 4)
+		{
+			Printer << TEXT(".");
+
+			static const TCHAR* Components = TEXT("xyzw");
+			for (int32 i = 0; i < PrimitiveType->GetNumComponents(); ++i)
+			{
+				check(BufferSlotOffset + i < sizeof(Components));
+				Printer.Buffer.AppendChar(Components[BufferSlotOffset + i]);
+			}
+		}
 
 		if (PrimitiveType->ScalarKind != MIR::EScalarKind::ScalarKind_Float)
 		{
diff --git a/Engine/Source/Runtime/Engine/Private/Materials/MaterialIRValueAnalyzer.cpp b/Engine/Source/Runtime/Engine/Private/Materials/MaterialIRValueAnalyzer.cpp
--- a/Engine/Source/Runtime/Engine/Private/Materials/MaterialIRValueAnalyzer.cpp
+++ b/Engine/Source/Runtime/Engine/Private/Materials/MaterialIRValueAnalyzer.cpp
@@ -79,8 +79,16 @@ static void AnalyzePrimitiveUniformParameter(FMaterialIRValueAnalyzer& Analyzer,
 	// Only int and float parameters supported for now (TODO)
 	check(ParameterType->ScalarKind == MIR::ScalarKind_Int || ParameterType->ScalarKind == MIR::ScalarKind_Float);
 
-	// Depending on the number of compoents this parameter has, try to find a previous allocation that may fit it.
-	uint32 ParameterUniformBufferOffset = UINT32_MAX;
+	// Following code calculates GlobalComponentOffset.
+	// GlobalComponentOffset is the i-th component in the array of float4s that make the uniform buffer.
+	// For example a GlobalComponentOffset of 13 references PreshaderBuffer[3].y.
+	// First, try to find an available sequence of free components in any previous allocation, in
+	// order to reduce the number of allocations and thus preshader buffer memory footprint.
+	// If the parameter type is too large and we can't find space for it in previous allocations,
+	// allocate a new uniform buffer slot (a float4, 16 bytes) and put any unused component
+	// in the appropriate freelist.
+
+	uint32 GlobalComponentOffset = UINT32_MAX;
 	uint32 NumComponents = ParameterType->GetNumComponents();
 	check(NumComponents > 0);
 	check(NumComponents <= 4); // Only vectors supported for now
@@ -92,12 +100,13 @@ static void AnalyzePrimitiveUniformParameter(FMaterialIRValueAnalyzer& Analyzer,
 		auto& FreeOffsetsForCurrentNumComponents = Analyzer.FreeOffsetsPerNumComponents[UsedNumComponents - 1];
 
 		// If there are no available slots at this nubmer of components, try with a larger one
-		if (FreeOffsetsForCurrentNumComponents.IsEmpty()) {
+		if (FreeOffsetsForCurrentNumComponents.IsEmpty())
+		{
 			continue;
 		}
 		
 		// There is space for this number `NumComponents` in a a chunk with `i` free components
-		ParameterUniformBufferOffset = FreeOffsetsForCurrentNumComponents.Last();
+		GlobalComponentOffset = FreeOffsetsForCurrentNumComponents.Last();
 		FreeOffsetsForCurrentNumComponents.Pop();
 
 		break;
@@ -105,40 +114,29 @@ static void AnalyzePrimitiveUniformParameter(FMaterialIRValueAnalyzer& Analyzer,
 	
 	// If UsedNumComponents is 4, it means we looked for all possible components sizes (1, 2 and 3) and
 	// could not find a chunk that would fit this parameter. Allocate a new chunk.
-	if (UsedNumComponents >= 4) {
+	if (UsedNumComponents >= 4)
+	{
 		check(UsedNumComponents == 4); 
-		ParameterUniformBufferOffset = UniformExpressionSet.AllocateFromUniformBuffer(16);
+		GlobalComponentOffset = UniformExpressionSet.AllocateFromUniformBuffer(16) / 4;
 	}
 
-	// Finally, check whether in the used slot there's any slack left, and if so, record it in the
-	// appropriate free-list.
+	// Finally, check whether in the used slot there's any slack left, and if so, record it in the appropriate free-list.
 	int NumComponentsLeft = UsedNumComponents - NumComponents;
-	if (NumComponentsLeft > 0) {
-		int LeftOverOffset = ParameterUniformBufferOffset + NumComponents * 4;
+	if (NumComponentsLeft > 0)
+	{
+		int LeftOverOffset = GlobalComponentOffset + NumComponents;
 		Analyzer.FreeOffsetsPerNumComponents[NumComponentsLeft - 1].Push(LeftOverOffset);
 	}
 
 	// Add the parameter evaluation to the uniform data
-	UniformExpressionSet.AddNumericParameterEvaluation(Parameter->Analysis_UniformParameterIndex, ParameterUniformBufferOffset);
+	UniformExpressionSet.AddNumericParameterEvaluation(Parameter->Analysis_UniformParameterIndex, GlobalComponentOffset);
 
 	if (Analyzer.Insights)
 	{
 		// Push information about this parameter allocation to the insights.
-		// ParameterUniformBufferOffset is the byte offset of where the parameter will be inside
-		// the uniform buffer. Since the uniform buffer is an array of 4D vectors of 4-byte components,
-		// we first calculate UniformComponentIndexFromStart being the offset of the allocated
-		// parameter in the uniform buffer expressed in "components" by dividing the byte offset by 4 bytes.
-		// From this, we can work out the index of the 4D vector in the whole uniform buffer (stored in
-		// SlotInsight.UniformIndex) and the first component occupied (SlotInsight.UniformComponetOffset).
-		// Currently the allocation algorithm does not spread parameters across multiple uniform 4D constants,
-		// hence SlotInsight.ParameterComponentOffset is always zero.
-
-		uint32 UniformComponentIndexFromStart = ParameterUniformBufferOffset / 4;
-
 		FMaterialInsights::FUniformParameterAllocationInsight ParamInsight;
-		ParamInsight.UniformIndex = UniformComponentIndexFromStart / 4;
-		ParamInsight.FirstUniformComponent = UniformComponentIndexFromStart % 4;
-		ParamInsight.FirstParameterComponent = 0; // Explained above
+		ParamInsight.BufferSlotIndex = GlobalComponentOffset / 4;
+		ParamInsight.BufferSlotOffset = GlobalComponentOffset % 4;
 		ParamInsight.ComponentsCount = NumComponents;
 		ParamInsight.ParameterName = Parameter->Info.Name;
 
diff --git a/Engine/Source/Runtime/Engine/Private/Materials/MaterialShared.cpp b/Engine/Source/Runtime/Engine/Private/Materials/MaterialShared.cpp
--- a/Engine/Source/Runtime/Engine/Private/Materials/MaterialShared.cpp
+++ b/Engine/Source/Runtime/Engine/Private/Materials/MaterialShared.cpp
@@ -3281,17 +3281,17 @@ static void EmitDebugInfoComment(const FMaterialInsights& Insights, FString& Sha
 	else {
 		for (auto& ParamInsight : Insights.UniformParameterAllocationInsights)
 		{
-			ShaderCode.Appendf(TEXT(" - UniformBuffer[%d]."), ParamInsight.UniformIndex);
+			ShaderCode.Appendf(TEXT(" - UniformBuffer[%d]."), ParamInsight.BufferSlotIndex);
 			for (uint32 i = 0; i < ParamInsight.ComponentsCount; ++i)
 			{
-				check(i >= 0 && i < 4);
-				ShaderCode.AppendChar(TEXT("xyzw")[ParamInsight.FirstUniformComponent + i]);
+				check(ParamInsight.BufferSlotOffset + i < 4);
+				ShaderCode.AppendChar(TEXT("xyzw")[ParamInsight.BufferSlotOffset + i]);
 			}
 			ShaderCode.Appendf(TEXT(" = %s."), *ParamInsight.ParameterName.ToString());
 			for (uint32 i = 0; i < ParamInsight.ComponentsCount; ++i)
 			{
-				check(i >= 0 && i < 4);
-				ShaderCode.AppendChar(TEXT("xyzw")[ParamInsight.FirstParameterComponent + i]);
+				check(i < 4);
+				ShaderCode.AppendChar(TEXT("xyzw")[i]);
 			}
 			ShaderCode.AppendChar('\n');
 		}
diff --git a/Engine/Source/Runtime/Engine/Private/Materials/MaterialUniformExpressions.cpp b/Engine/Source/Runtime/Engine/Private/Materials/MaterialUniformExpressions.cpp
--- a/Engine/Source/Runtime/Engine/Private/Materials/MaterialUniformExpressions.cpp
+++ b/Engine/Source/Runtime/Engine/Private/Materials/MaterialUniformExpressions.cpp
@@ -714,6 +714,11 @@ void FUniformExpressionSet::AddNumericParameterEvaluation(uint32 ParameterIndex,
 	});
 }
 
+uint32 FUniformExpressionSet::GetNumericParameterEvaluationOffset(uint32 ParameterIndex) const
+{
+	return UniformParameterEvaluations[ParameterIndex].BufferOffset;
+}
+
 uint32 FUniformExpressionSet::AllocateFromUniformBuffer(uint32 SizeInBytes)
 {
 	uint32 Temp = UniformPreshaderBufferSize;
diff --git a/Engine/Source/Runtime/Engine/Public/MaterialShared.h b/Engine/Source/Runtime/Engine/Public/MaterialShared.h
--- a/Engine/Source/Runtime/Engine/Public/MaterialShared.h
+++ b/Engine/Source/Runtime/Engine/Public/MaterialShared.h
@@ -696,7 +696,14 @@ public:
 	int32 FindOrAddExternalTextureParameter(const FMaterialExternalTextureParameterInfo& Info);
 	int32 FindOrAddNumericParameter(EMaterialParameterType Type, const FMaterialParameterInfo& ParameterInfo, uint32 DefaultValueOffset);
 	uint32 AddDefaultParameterValue(const UE::Shader::FValue& Value);
+	
+	/// Adds a new direct parameter evaluation for parameter with specified index (returned by FindOrAddNumericParameter)
+	/// BufferOffset is the location in the uniform buffer in which to evaluate the parameter (in multiple of float components).
 	void AddNumericParameterEvaluation(uint32 ParameterIndex, uint32 BufferOffset);
+	
+	/// Returns the offset of specified parameter given with AddNumericParameterEvaluation.
+	uint32 GetNumericParameterEvaluationOffset(uint32 ParameterIndex) const;
+	
 	uint32 AllocateFromUniformBuffer(uint32 SizeInBytes);
 	
 	TConstArrayView<FMaterialVirtualTextureStack> GetVTStacks() const { return VTStacks; }
diff --git a/Engine/Source/Runtime/Engine/Public/Materials/MaterialInsights.h b/Engine/Source/Runtime/Engine/Public/Materials/MaterialInsights.h
--- a/Engine/Source/Runtime/Engine/Public/Materials/MaterialInsights.h
+++ b/Engine/Source/Runtime/Engine/Public/Materials/MaterialInsights.h
@@ -45,12 +45,11 @@ struct FMaterialInsights
 	*/
 	struct FUniformParameterAllocationInsight
 	{
-		uint32 UniformIndex;							///< Index of the preshader buffer constant vec4 (e.g. PreshaderBuffer[0])
-		uint32 FirstUniformComponent;					///< First component of the vec4 above (e.g. 0 for .x, 1 for .y, etc)
-		FName ParameterName;							///< The parameter name
-		uint32 FirstParameterComponent;					///< First component of the parameter (eg. 0 for .x, etc)
-		uint32 ComponentsCount;							///< Number of components stored, same for the Uniform and the Parameter
-		FUniformBufferSlotComponentType ComponentType;  ///< Parameter component type
+		uint16 BufferSlotIndex;								///< Index of the preshader buffer constant float4 slot (e.g. PreshaderBuffer[0])
+		uint16 BufferSlotOffset;							///< First component of the float4 above (e.g. 0 for .x, 1 for .y, etc)
+		uint16 ComponentsCount;								///< Number of components stored, same for the Uniform and the Parameter
+		FUniformBufferSlotComponentType ComponentType : 8;  ///< Parameter component type
+		FName ParameterName;								///< The parameter name
 	};
 
 	/// Array of connection insights.
