PR #12868: Fix VisualLogger Deprecation messages

diff --git a/Engine/Source/Runtime/Engine/Public/VisualLogger/VisualLogger.h b/Engine/Source/Runtime/Engine/Public/VisualLogger/VisualLogger.h
--- a/Engine/Source/Runtime/Engine/Public/VisualLogger/VisualLogger.h
+++ b/Engine/Source/Runtime/Engine/Public/VisualLogger/VisualLogger.h
@@ -304,7 +304,7 @@ public:
 
 	// Arrow
 	template <typename FmtType, typename... Types>
-	UE_DEPRECATED(5.6, "Use ArrowLineLogfImpl")
+	UE_DEPRECATED(5.6, "Use ArrowLineLogf")
 	static void ArrowLogf(const UObject* LogOwner, const FLogCategoryBase& Category, ELogVerbosity::Type Verbosity, const FVector& Start, const FVector& End, const FColor& Color, const FmtType& Fmt, Types... Args)
 	{
 		static_assert(TIsArrayOrRefOfTypeByPredicate<FmtType, TIsCharEncodingCompatibleWithTCHAR>::Value, "Formatting string must be a TCHAR array.");
@@ -312,7 +312,7 @@ public:
 		ArrowLineLogfImpl(LogOwner, Category.GetCategoryName(), Verbosity, Start, End, Color, 8, (const TCHAR*)Fmt, Args...);
 	}
 	template <typename FmtType, typename... Types>
-	UE_DEPRECATED(5.6, "Use ArrowLineLogfImpl")
+	UE_DEPRECATED(5.6, "Use ArrowLineLogf")
 	static void ArrowLogf(const UObject* LogOwner, const FName& CategoryName, ELogVerbosity::Type Verbosity, const FVector& Start, const FVector& End, const FColor& Color, const FmtType& Fmt, Types... Args)
 	{
 		static_assert(TIsArrayOrRefOfTypeByPredicate<FmtType, TIsCharEncodingCompatibleWithTCHAR>::Value, "Formatting string must be a TCHAR array.");
@@ -341,7 +341,7 @@ public:
 
 	// Circle log
 	template <typename FmtType, typename... Types>
-	UE_DEPRECATED(5.6, "Use DiscLogfImpl")
+	UE_DEPRECATED(5.6, "Use DiscLogf")
 	static void CircleLogf(const UObject* LogOwner, const FLogCategoryBase& Category, ELogVerbosity::Type Verbosity, const FVector& Center, const FVector& UpAxis, const float Radius, const FColor& Color, const uint16 Thickness, const FmtType& Fmt, Types... Args)
 	{
 		static_assert(TIsArrayOrRefOfTypeByPredicate<FmtType, TIsCharEncodingCompatibleWithTCHAR>::Value, "Formatting string must be a TCHAR array.");
@@ -350,7 +350,7 @@ public:
 		DiscLogfImpl(LogOwner, Category.GetCategoryName(), Verbosity, Center, UpAxis, Radius, Color, Thickness, /*bWireframe = */false, (const TCHAR*)Fmt, Args...);
 	}
 	template <typename FmtType, typename... Types>
-	UE_DEPRECATED(5.6, "Use DiscLogfImpl")
+	UE_DEPRECATED(5.6, "Use DiscLogf")
 	static void CircleLogf(const UObject* LogOwner, const FName& CategoryName, ELogVerbosity::Type Verbosity, const FVector& Center, const FVector& UpAxis, const float Radius, const FColor& Color, const uint16 Thickness, const FmtType& Fmt, Types... Args)
 	{
 		static_assert(TIsArrayOrRefOfTypeByPredicate<FmtType, TIsCharEncodingCompatibleWithTCHAR>::Value, "Formatting string must be a TCHAR array.");
@@ -470,7 +470,7 @@ public:
 
 	// Cone log
 	template <typename FmtType, typename... Types>
-	UE_DEPRECATED(5.4, "Use ConeLogfImpl")
+	UE_DEPRECATED(5.4, "Use ConeLogf")
 	static void GeometryShapeLogf(const UObject* LogOwner, const FLogCategoryBase& Category, ELogVerbosity::Type Verbosity, const FVector& Origin, const FVector& Direction, const float Length, const float Angle, const FColor& Color, const FmtType& Fmt, Types... Args)
 	{
 		static_assert(TIsArrayOrRefOfTypeByPredicate<FmtType, TIsCharEncodingCompatibleWithTCHAR>::Value, "Formatting string must be a TCHAR array.");
@@ -479,7 +479,7 @@ public:
 		ConeLogfImpl(LogOwner, Category.GetCategoryName(), Verbosity, Origin, Direction, Length, Angle, Color, /*bWireframe = */false, (const TCHAR*)Fmt, Args...);
 	}
 	template <typename FmtType, typename... Types>
-	UE_DEPRECATED(5.4, "Use ConeLogfImpl")
+	UE_DEPRECATED(5.4, "Use ConeLogf")
 	static void GeometryShapeLogf(const UObject* LogOwner, const FName& CategoryName, ELogVerbosity::Type Verbosity, const FVector& Origin, const FVector& Direction, const float Length, const float Angle, const FColor& Color, const FmtType& Fmt, Types... Args)
 	{
 		static_assert(TIsArrayOrRefOfTypeByPredicate<FmtType, TIsCharEncodingCompatibleWithTCHAR>::Value, "Formatting string must be a TCHAR array.");
