- Added mode 3 to wp.Runtime.HashSet.ShowDebugDisplayMode which shows Level Streaming State Without Data Layers/Content Bundles.
- Added support for wp.Runtime.DebugFilterByStreamingStatus and wp.Runtime.DebugFilterByDataLayer in HashSet.
- Apply same debug filtering rules to the 3D debug display version.

diff --git a/Engine/Source/Runtime/Engine/Private/WorldPartition/RuntimeHashSet/WorldPartitionRuntimeHashSetDebugDraw.cpp b/Engine/Source/Runtime/Engine/Private/WorldPartition/RuntimeHashSet/WorldPartitionRuntimeHashSetDebugDraw.cpp
--- a/Engine/Source/Runtime/Engine/Private/WorldPartition/RuntimeHashSet/WorldPartitionRuntimeHashSetDebugDraw.cpp
+++ b/Engine/Source/Runtime/Engine/Private/WorldPartition/RuntimeHashSet/WorldPartitionRuntimeHashSetDebugDraw.cpp
@@ -28,7 +28,7 @@ static int32 GShowRuntimeHashSetDebugDisplayMode = 0;
 static FAutoConsoleVariableRef CVarShowRuntimeHashSetDebugDisplayMode(
 	TEXT("wp.Runtime.HashSet.ShowDebugDisplayMode"),
 	GShowRuntimeHashSetDebugDisplayMode,
-	TEXT("Used to choose what mode to display when showing runtime partitions (0=Level Streaming State, 1=Data Layers, 2=Content Bundles)."));
+	TEXT("Used to choose what mode to display when showing runtime partitions (0=Level Streaming State, 1=Data Layers, 2=Content Bundles, 3=Level Streaming State Without Data Layers/Content Bundles)."));
 
 static EWorldPartitionRuntimeCellVisualizeMode GetStreamingCellVisualizeMode()
 {
@@ -50,6 +50,15 @@ static TMap<FName, FColor> GetDataLayerDebugColors(const UWorldPartition* InWorl
 	return DebugColors;
 }
 
+static bool IsDebugCellShown(UWorldPartitionRuntimeCell* Cell)
+{
+	UWorldPartitionRuntimeCellData* RuntimeCellData = Cell->RuntimeCellData;
+	return ((RuntimeCellData->HierarchicalLevel >= GShowRuntimeHashSetDebugDisplayLevel) &&
+		(RuntimeCellData->HierarchicalLevel < (GShowRuntimeHashSetDebugDisplayLevel + GShowRuntimeHashSetDebugDisplayLevelCount)) &&
+		(GShowRuntimeHashSetDebugDisplayMode != 3 || (!Cell->GetDataLayers().Num() && !Cell->GetContentBundleID().IsValid())) &&
+		Cell->IsDebugShown());
+}
+
 bool UWorldPartitionRuntimeHashSet::Draw2D(FWorldPartitionDraw2DContext& DrawContext) const
 {
 	TRACE_CPUPROFILER_EVENT_SCOPE(UWorldPartitionRuntimeHashSet::Draw2D);
@@ -124,8 +133,7 @@ bool UWorldPartitionRuntimeHashSet::Draw2D(FWorldPartitionDraw2DContext& DrawCon
 
 		auto FilterCell = [&FilteredCells](UWorldPartitionRuntimeCell* Cell)
 		{
-			UWorldPartitionRuntimeCellData* RuntimeCellData = Cell->RuntimeCellData;
-			if ((RuntimeCellData->HierarchicalLevel >= GShowRuntimeHashSetDebugDisplayLevel) && (RuntimeCellData->HierarchicalLevel < (GShowRuntimeHashSetDebugDisplayLevel + GShowRuntimeHashSetDebugDisplayLevelCount)))
+			if (IsDebugCellShown(Cell))
 			{
 				FilteredCells.Add(Cell);
 			}
@@ -161,6 +169,7 @@ bool UWorldPartitionRuntimeHashSet::Draw2D(FWorldPartitionDraw2DContext& DrawCon
 				switch (GShowRuntimeHashSetDebugDisplayMode)
 				{
 				case 0:
+				case 3:
 					CellColors.Add(Cell->GetDebugColor(VisualizeMode));
 					CellOpacity = (VisualizeMode == EWorldPartitionRuntimeCellVisualizeMode::StreamingPriority) ? 0.75f : 0.25f / FMath::Max<float>(GShowRuntimeHashSetDebugDisplayLevelCount, 1);
 					break;
@@ -368,6 +377,11 @@ void UWorldPartitionRuntimeHashSet::Draw3D(const TArray<FWorldPartitionStreaming
 					{
 						StreamingData.SpatialIndex.Get()->ForEachIntersectingElement(InShape, [this, VisualizeMode, &DataLayerDebugColors, ContentBundleManager, OwningWorld, &WorldPartitionTransform](UWorldPartitionRuntimeCell* Cell)
 						{
+							if (!IsDebugCellShown(Cell))
+							{
+								return;
+							}
+
 							const FVector2D CellStreamingBoundsSize = FVector2D(Cell->GetStreamingBounds().GetSize());
 							const FVector2D CellStreamingBoundsMin = FVector2D(Cell->GetStreamingBounds().Min);
 
@@ -377,6 +391,7 @@ void UWorldPartitionRuntimeHashSet::Draw3D(const TArray<FWorldPartitionStreaming
 							switch (GShowRuntimeHashSetDebugDisplayMode)
 							{
 							case 0:
+							case 3:
 								CellColors.Add(Cell->GetDebugColor(VisualizeMode));
 								CellOpacity = 0.25f / FMath::Max<float>(GShowRuntimeHashSetDebugDisplayLevelCount, 1);
 								break;
