1-Only start precaching pso after the engine have loaded the default materials.

The underlying issue is that we are trying to precache PSOs when we are loading the default materials, LoadObject can potentially use default materials, and they should not be precache, since UMaterialInterface::InitDefaultMaterials will precache them after the recursive loads, by calling PrecacheDefaultMaterialPSOs. I had the same issue with  EPSOPrecacheMode::PreloadShader, that why IsPSOShaderPreloadingEnabled() is checking if we initialized the default materials before starting to preload shaders. I think the same logic should apply to the PSO caching mode

2-Fix GPuSkin PSO precaching when ManualVF is disabled.

Expose  InitStaticMeshVertexFactoryComponents in StaticMeshResources, and remove the extern declaration.

#tests replayruns (all platforms) + 672b8b72e07ed232abe35e53

diff --git a/Engine/Source/Runtime/Engine/Private/Components/StaticMeshComponent.cpp b/Engine/Source/Runtime/Engine/Private/Components/StaticMeshComponent.cpp
--- a/Engine/Source/Runtime/Engine/Private/Components/StaticMeshComponent.cpp
+++ b/Engine/Source/Runtime/Engine/Private/Components/StaticMeshComponent.cpp
@@ -1764,7 +1764,7 @@ void UStaticMeshComponent::CollectPSOPrecacheData(const FPSOPrecacheParams& Base
 		int32 LODLightMapCoordinateIndex = LightMapCoordinateIndex < NumTexCoords ? LightMapCoordinateIndex : NumTexCoords - 1;
 		bool bOverrideColorVertexBuffer = LODIndex < LODData.Num() && LODData[LODIndex].OverrideVertexColors != nullptr;
 		FLocalVertexFactory::FDataType Data;
-		LODRenderData.VertexBuffers.InitCompVF(nullptr /*VertexFactory*/, LODLightMapCoordinateIndex, bOverrideColorVertexBuffer, Data);
+		LODRenderData.VertexBuffers.InitComponentVF(nullptr /*VertexFactory*/, LODLightMapCoordinateIndex, bOverrideColorVertexBuffer, Data);
 		FLocalVertexFactory::GetVertexElements(GMaxRHIFeatureLevel, EVertexInputStreamType::Default, bSupportsManualVertexFetch, Data, Elements);
 	};
 	
diff --git a/Engine/Source/Runtime/Engine/Private/SkeletalRenderGPUSkin.cpp b/Engine/Source/Runtime/Engine/Private/SkeletalRenderGPUSkin.cpp
--- a/Engine/Source/Runtime/Engine/Private/SkeletalRenderGPUSkin.cpp
+++ b/Engine/Source/Runtime/Engine/Private/SkeletalRenderGPUSkin.cpp
@@ -2280,7 +2280,7 @@ void FSkeletalMeshObjectGPUSkin::GetUsedVertexFactoryData(
 			FVertexDeclarationElementList VertexElements;
 			bool bOverrideColorVertexBuffer = false;
 			FGPUSkinPassthroughVertexFactory::FDataType Data;
-			LODRenderData.StaticVertexBuffers.InitCompVF(nullptr /*VertexFactory*/, 0, bOverrideColorVertexBuffer, Data);
+			LODRenderData.StaticVertexBuffers.InitComponentVF(nullptr /*VertexFactory*/, 0, bOverrideColorVertexBuffer, Data);
 			FGPUSkinPassthroughVertexFactory::GetVertexElements(GMaxRHIFeatureLevel, EVertexInputStreamType::Default, bSupportsManualVertexFetch, Data, VertexElements);
 			VertexFactoryDataList.AddUnique(FPSOPrecacheVertexFactoryData(GPUSkinVFType, VertexElements));
 		}
diff --git a/Engine/Source/Runtime/Engine/Private/SkinnedAsset.cpp b/Engine/Source/Runtime/Engine/Private/SkinnedAsset.cpp
--- a/Engine/Source/Runtime/Engine/Private/SkinnedAsset.cpp
+++ b/Engine/Source/Runtime/Engine/Private/SkinnedAsset.cpp
@@ -175,7 +175,7 @@ void USkinnedAsset::PostLoad()
 	}
 }
 
-void USkinnedAsset::GetVertexFactoryTypesPerMaterialIndex(FPSOPrecacheVertexFactoryDataPerMaterialIndexList& OutList,
+FPSOPrecacheVertexFactoryDataPerMaterialIndexList USkinnedAsset::GetVertexFactoryTypesPerMaterialIndex(
 	USkinnedMeshComponent* SkinnedMeshComponent, int32 MinLODIndex, bool bCPUSkin, ERHIFeatureLevel::Type FeatureLevel)
 {
 	FSkeletalMeshRenderData* SkelMeshRenderData = GetResourceForRendering();
@@ -218,7 +218,7 @@ void USkinnedAsset::GetVertexFactoryTypesPerMaterialIndex(FPSOPrecacheVertexFact
 					FVertexDeclarationElementList VertexElements;
 					bool bOverrideColorVertexBuffer = false;
 					FLocalVertexFactory::FDataType Data;
-					LODRenderData.StaticVertexBuffers.InitCompVF(nullptr /*VertexFactory*/, 0, bOverrideColorVertexBuffer, Data);
+					LODRenderData.StaticVertexBuffers.InitComponentVF(nullptr /*VertexFactory*/, 0, bOverrideColorVertexBuffer, Data);
 					FLocalVertexFactory::GetVertexElements(GMaxRHIFeatureLevel, EVertexInputStreamType::Default, bSupportsManualVertexFetch, Data, VertexElements);
 					VFsPerMaterial->VertexFactoryDataList.AddUnique(FPSOPrecacheVertexFactoryData(CPUSkinVFType, VertexElements));
 				}
diff --git a/Engine/Source/Runtime/Engine/Private/StaticMesh.cpp b/Engine/Source/Runtime/Engine/Private/StaticMesh.cpp
--- a/Engine/Source/Runtime/Engine/Private/StaticMesh.cpp
+++ b/Engine/Source/Runtime/Engine/Private/StaticMesh.cpp
@@ -1130,7 +1130,7 @@ static inline void InitOrUpdateResource(FRHICommandListBase& RHICmdList, FRender
 	}
 }
 
-void FStaticMeshVertexBuffers::InitCompVF(FLocalVertexFactory* VertexFactory, int32 LightMapCoordinateIndex, bool bOverrideColorVertexBuffer, FLocalVertexFactory::FDataType& OutData) const
+void FStaticMeshVertexBuffers::InitComponentVF(FLocalVertexFactory* VertexFactory, int32 LightMapCoordinateIndex, bool bOverrideColorVertexBuffer, FLocalVertexFactory::FDataType& OutData) const
 {
 	PositionVertexBuffer.BindPositionVertexBuffer(VertexFactory, OutData);
 	StaticMeshVertexBuffer.BindTangentVertexBuffer(VertexFactory, OutData);
diff --git a/Engine/Source/Runtime/Engine/Public/StaticMeshResources.h b/Engine/Source/Runtime/Engine/Public/StaticMeshResources.h
--- a/Engine/Source/Runtime/Engine/Public/StaticMeshResources.h
+++ b/Engine/Source/Runtime/Engine/Public/StaticMeshResources.h
@@ -358,7 +358,7 @@ struct FStaticMeshVertexBuffers
 	}
 
 	void ENGINE_API InitModelBuffers(TArray<FModelVertex>& Vertices);
-	void ENGINE_API InitCompVF(FLocalVertexFactory* VertexFactory, int32 LightMapCoordinateIndex, bool bOverrideColorVertexBuffer, FLocalVertexFactory::FDataType& OutDat) const;
+	void ENGINE_API InitComponentVF(FLocalVertexFactory* VertexFactory, int32 LightMapCoordinateIndex, bool bOverrideColorVertexBuffer, FLocalVertexFactory::FDataType& OutDat) const;
 
 	inline void ENGINE_API InitModelVF(FRHICommandListBase& RHICmdList, FLocalVertexFactory* VertexFactory)
 	{
