Implement Download Progress for Streamable handles.

diff --git a/Engine/Source/Runtime/Engine/Classes/Engine/StreamableManager.h b/Engine/Source/Runtime/Engine/Classes/Engine/StreamableManager.h
--- a/Engine/Source/Runtime/Engine/Classes/Engine/StreamableManager.h
+++ b/Engine/Source/Runtime/Engine/Classes/Engine/StreamableManager.h
@@ -16,6 +16,16 @@
 #include "UObject/ObjectMacros.h"
 #include "UObject/SoftObjectPtr.h"
 
+enum class EStreamableProgressType : uint8
+{
+	/** Load progress */
+	Load,
+	/** Progress of downloading packages that were not installed when the request was issued. */
+	DownloadRelative,
+	/** Progress of the total percentage of installed packages for the request. */
+	DownloadAbsolute,
+};
+
 enum class EStreamableErrorType : uint8
 {
 	Unknown = 0,
@@ -36,7 +46,7 @@ struct FStreamableError
 		, ErrorType(InErrorType)
 	{}
 
-	FString ErrorMessage; // TODO: Can this just be a TCHAR* or string view ?
+	FString ErrorMessage;
 	EStreamableErrorType ErrorType = EStreamableErrorType::Unknown;
 };
 
@@ -376,7 +386,16 @@ struct FStreamableHandle : public TSharedFromThis<FStreamableHandle>
 	void ForEachLoadedAsset(CallableT Callable) const;
 
 	/** Returns progress as a value between 0.0 and 1.0. */
-	ENGINE_API float GetProgress() const;
+	ENGINE_API float GetProgress(EStreamableProgressType Type = EStreamableProgressType::Load) const;
+
+	/** Returns progress as a value between 0.0 and 1.0. */
+	float GetLoadProgress() const { return GetProgress(EStreamableProgressType::Load); }
+
+	/** Returns progress as a value between 0.0 and 1.0 representing progress of downloading packages that were not installed when the request was issued. */
+	float GetRelativeDownloadProgress() const { return GetProgress(EStreamableProgressType::DownloadRelative); }
+
+	/** Returns progress as a value between 0.0 and 1.0 representing the total percentage of installed packages for the request. */
+	float GetAbsoluteDownloadProgress() const { return GetProgress(EStreamableProgressType::DownloadAbsolute); }
 
 	/** Get the StreamableManager for this handle */
 	ENGINE_API struct FStreamableManager* GetOwningManager() const;
@@ -626,7 +645,7 @@ private:
 	/** If true, this handle will be released when it finishes loading */
 	bool bReleaseWhenLoaded;
 
-	/** If true, this handle will be released when it finishes loading */
+	/** If true, this is a combined handle that depends on child handles. */
 	bool bIsCombinedHandle;
 };
 
diff --git a/Engine/Source/Runtime/Engine/Private/StreamableManager.cpp b/Engine/Source/Runtime/Engine/Private/StreamableManager.cpp
--- a/Engine/Source/Runtime/Engine/Private/StreamableManager.cpp
+++ b/Engine/Source/Runtime/Engine/Private/StreamableManager.cpp
@@ -379,6 +379,7 @@ namespace UE::StreamableManager::Private
 	struct FDownloadContext
 	{
 		UE::IoStore::FOnDemandInstallRequest InstallRequest;
+		UE::IoStore::FOnDemandInstallProgress InstallProgress;
 		EStreamableDownloadPriority Priority = EStreamableDownloadPriority::Normal;
 		bool bInstallSoftReferences = false;
 
@@ -513,7 +514,12 @@ namespace UE::StreamableManager::Private
 
 		void OnInstallProgress(const TSharedRef<FStreamableHandle>& OwningHandle, const UE::IoStore::FOnDemandInstallProgress& InProgress)
 		{
-			// TODO: progress
+			InstallProgress = InProgress;
+
+			if (OwningHandle->OwningManager) // Make sure handle isn't canceled or released
+			{
+				OwningHandle->CallUpdateDelegate();
+			}
 		}
 	};
 
@@ -778,18 +784,56 @@ void FStreamableHandle::GetLoadedCount(int32& LoadedCount, int32& RequestedCount
 	}
 }
 
-float FStreamableHandle::GetProgress() const
+float FStreamableHandle::GetProgress(EStreamableProgressType Type) const
 {
-	if (HasLoadCompleted())
+	if (Type == EStreamableProgressType::Load)
 	{
-		return 1.0f;
+		if (HasLoadCompleted())
+		{
+			return 1.0f;
+		}
+
+		int32 Loaded, Total;
+
+		GetLoadedCount(Loaded, Total);
+
+		return (float)Loaded / Total;
 	}
 
-	int32 Loaded, Total;
+	int32 DownloadCount = 0;
+	float TotalProgress = 0.0f;
+
+	TArray<const FStreamableHandle*, TInlineAllocator<32>> Stack;
+	Stack.Push(this);
+
+	do
+	{
+		const FStreamableHandle* Current = Stack.Pop(EAllowShrinking::No);
+		if (Current->DownloadContext)
+		{
+			++DownloadCount;
+			if (Type == EStreamableProgressType::DownloadRelative)
+			{
+				TotalProgress += DownloadContext->InstallProgress.GetRelativeProgress();
+			}
+			else
+			{
+				TotalProgress += DownloadContext->InstallProgress.GetAbsoluteProgress();
+			}
+		}
+
+		for (const TSharedPtr<FStreamableHandle>& ChildHandle : Current->ChildHandles)
+		{
+			Stack.Push(ChildHandle.Get());
+		}
+	} while (!Stack.IsEmpty());
 
-	GetLoadedCount(Loaded, Total);
+	if (DownloadCount > 0)
+	{
+		return TotalProgress / DownloadCount;
+	}
 
-	return (float)Loaded / Total;
+	return 1.0f;
 }
 
 FStreamableManager* FStreamableHandle::GetOwningManager() const
@@ -1253,7 +1297,7 @@ TSharedPtr<FStreamableHandle> FStreamableHandle::CreateCombinedHandle(const TCon
 		TArray<TSharedPtr<FStreamableHandle>> HandlesToUse{OtherHandles};
 		HandlesToUse.AddUnique(ThisAsShared);
 
-		return Manager->CreateCombinedHandle(HandlesToUse, HandleDebugName_CombinedHandle, MergeOptions);		
+		return Manager->CreateCombinedHandle(HandlesToUse, HandleDebugName_CombinedHandle, MergeOptions);
 	}
 
 	return nullptr;
