Fix for Random Sound Node Crash in editor

diff --git a/Engine/Source/Runtime/Engine/Private/SoundNodeRandom.cpp b/Engine/Source/Runtime/Engine/Private/SoundNodeRandom.cpp
--- a/Engine/Source/Runtime/Engine/Private/SoundNodeRandom.cpp
+++ b/Engine/Source/Runtime/Engine/Private/SoundNodeRandom.cpp
@@ -228,7 +228,10 @@ void USoundNodeRandom::ParseNodes( FAudioDevice* AudioDevice, const UPTRINT Node
 		}
 
 		// set the node that has JUST played to be true so we don't repeat it
-		HasBeenUsed[NodeIndex] = true;
+		if (HasBeenUsed.IsValidIndex(NodeIndex))
+		{
+			HasBeenUsed[NodeIndex] = true;
+		}
 		NumRandomUsed = 1;
 	}
 
