Fix for shapes in welded bodies not having proxies set.

Shapes normally have proxies set by SetProxy called by FPBDRigidsSolver::RegisterObject. This does not take welded bodies into account, leaving all shapes but the root node in a hierarchy having proxies unset. This CL sets the shape proxies after welding.

diff --git a/Engine/Source/Runtime/Engine/Private/PhysicsEngine/BodyInstance.cpp b/Engine/Source/Runtime/Engine/Private/PhysicsEngine/BodyInstance.cpp
--- a/Engine/Source/Runtime/Engine/Private/PhysicsEngine/BodyInstance.cpp
+++ b/Engine/Source/Runtime/Engine/Private/PhysicsEngine/BodyInstance.cpp
@@ -1924,6 +1924,7 @@ bool FBodyInstance::Weld(FBodyInstance* TheirBody, const FTransform& TheirTM)
 
 			for (int32 ShapeIdx = 0; ShapeIdx < PNewShapes.Num(); ++ShapeIdx)
 			{
+				PNewShapes[ShapeIdx].Shape->SetProxy(ActorHandle);
 				ShapeToBodiesMap->Add(PNewShapes[ShapeIdx], FWeldInfo(TheirBody, RelativeTM));
 			}
 
