[MIR] Implement UMaterialExpressionExternalCodeBase in new material IR translator.

This adds FExternalCode as subtype of MIR::TInstruction to generate the IR for external code expressions in the new material translator.

#rnx

diff --git a/Engine/Source/Runtime/Engine/Private/Materials/MaterialExpressionsIR.cpp b/Engine/Source/Runtime/Engine/Private/Materials/MaterialExpressionsIR.cpp
--- a/Engine/Source/Runtime/Engine/Private/Materials/MaterialExpressionsIR.cpp
+++ b/Engine/Source/Runtime/Engine/Private/Materials/MaterialExpressionsIR.cpp
@@ -231,6 +231,7 @@
 #include "Materials/MaterialExpressionVolumetricAdvancedMaterialInput.h"
 #include "Materials/MaterialExpressionVolumetricAdvancedMaterialOutput.h"
 #include "Materials/MaterialExpressionWorldPosition.h"
+#include "Materials/MaterialExternalCodeRegistry.h"
 #include "Materials/MaterialFunctionInterface.h"
 #include "Materials/MaterialParameterCollection.h"
 #include "MaterialShared.h"
@@ -590,6 +591,19 @@ void UMaterialExpressionTangent::Build(MIR::FEmitter& Emitter)
 	BuildUnaryArithmeticOperatorFromInput(Emitter, MIR::UO_Tan, &Input);
 }
 
+void UMaterialExpressionExternalCodeBase::Build(MIR::FEmitter& Emitter)
+{
+	if (const FMaterialExternalCodeDeclaration* ExternalCodeDeclaration = MaterialExternalCodeRegistry::Get().FindExternalCode(ExternalCodeIdentifier))
+	{
+		MIR::FValue* Value = Emitter.EmitExternalCode(ExternalCodeDeclaration);
+		Emitter.Output(0, Value);
+	}
+	else
+	{
+		Emitter.Errorf(TEXT("Missing external code declaration for '%s'"), *ExternalCodeIdentifier.ToString());
+	}
+}
+
 /* Binary Operators */
 
 static void BuildBinaryOperatorFromValue(
diff --git a/Engine/Source/Runtime/Engine/Private/Materials/MaterialIREmitter.cpp b/Engine/Source/Runtime/Engine/Private/Materials/MaterialIREmitter.cpp
--- a/Engine/Source/Runtime/Engine/Private/Materials/MaterialIREmitter.cpp
+++ b/Engine/Source/Runtime/Engine/Private/Materials/MaterialIREmitter.cpp
@@ -4,6 +4,7 @@
 #include "Materials/MaterialIRModule.h"
 #include "Materials/MaterialIRModuleBuilder.h"
 #include "Materials/MaterialIRInternal.h"
+#include "Materials/MaterialExternalCodeRegistry.h"
 #include "Shader/ShaderTypes.h"
 #include "MaterialShared.h"
 #include "MaterialExpressionIO.h"
@@ -1793,6 +1794,15 @@ FValue* FEmitter::EmitTextureSample(FValue* Texture, FValue* TexCoord, ESamplerS
 	return EmitPrototype(this, Prototype);
 }
 
+FValue* FEmitter::EmitExternalCode(const FMaterialExternalCodeDeclaration* InExternalCodeDeclaration)
+{
+	check(InExternalCodeDeclaration != nullptr);
+	FTypePtr ReturnType = FType::FromMaterialValueType(InExternalCodeDeclaration->GetReturnTypeValue());
+	MIR::FExternalCode Prototype = MakePrototype<MIR::FExternalCode>(ReturnType);
+	Prototype.ExternalCodeDeclaration = InExternalCodeDeclaration;
+	return EmitPrototype(this, Prototype);
+}
+
 FTypePtr FEmitter::TryGetCommonType(FTypePtr A, FTypePtr B)
 {
 	// Trivial case: types are equal
diff --git a/Engine/Source/Runtime/Engine/Private/Materials/MaterialIRToHLSLTranslator.cpp b/Engine/Source/Runtime/Engine/Private/Materials/MaterialIRToHLSLTranslator.cpp
--- a/Engine/Source/Runtime/Engine/Private/Materials/MaterialIRToHLSLTranslator.cpp
+++ b/Engine/Source/Runtime/Engine/Private/Materials/MaterialIRToHLSLTranslator.cpp
@@ -4,6 +4,7 @@
 #include "Materials/MaterialIRModule.h"
 #include "Materials/MaterialIRTypes.h"
 #include "Materials/MaterialIR.h"
+#include "Materials/MaterialExternalCodeRegistry.h"
 #include "MaterialIRInternal.h"
 
 #include "ShaderCore.h"
@@ -339,6 +340,7 @@ struct FTranslator : FMaterialIRToHLSLTranslation
 			case MIR::VK_Branch: LowerBranch(static_cast<MIR::FBranch*>(Instr)); break;
 			case MIR::VK_Subscript: LowerSubscript(static_cast<MIR::FSubscript*>(Instr)); break;
 			case MIR::VK_TextureSample: LowerTextureSample(static_cast<MIR::FTextureSample*>(Instr)); break;
+			case MIR::VK_ExternalCode: LowerExternalCode(static_cast<const MIR::FExternalCode*>(Instr)); break;
 
 			default:
 				UE_MIR_UNREACHABLE();
@@ -385,6 +387,12 @@ struct FTranslator : FMaterialIRToHLSLTranslation
 		}
 	}
 
+	void LowerExternalCode(const MIR::FExternalCode* ExternalCode)
+	{
+		check(ExternalCode->ExternalCodeDeclaration != nullptr);
+		Printer.Buffer.Append(*ExternalCode->ExternalCodeDeclaration->Definition);
+	}
+
 	void LowerTextureObject(const MIR::FTextureObject* TextureObject)
 	{
 		LowerTextureReference(TextureObject->Texture->GetMaterialType(), TextureObject->Analysis_UniformParameterIndex);
diff --git a/Engine/Source/Runtime/Engine/Public/Materials/MaterialExpressionExternalCodeBase.h b/Engine/Source/Runtime/Engine/Public/Materials/MaterialExpressionExternalCodeBase.h
--- a/Engine/Source/Runtime/Engine/Public/Materials/MaterialExpressionExternalCodeBase.h
+++ b/Engine/Source/Runtime/Engine/Public/Materials/MaterialExpressionExternalCodeBase.h
@@ -18,6 +18,7 @@ public:
 
 	//~ Begin UMaterialExpression Interface
 #if WITH_EDITOR
+	virtual void Build(MIR::FEmitter& Emitter) override;
 	virtual int32 Compile(class FMaterialCompiler* Compiler, int32 OutputIndex) override;
 #endif
 	//~ End UMaterialExpression Interface
diff --git a/Engine/Source/Runtime/Engine/Public/Materials/MaterialIR.h b/Engine/Source/Runtime/Engine/Public/Materials/MaterialIR.h
--- a/Engine/Source/Runtime/Engine/Public/Materials/MaterialIR.h
+++ b/Engine/Source/Runtime/Engine/Public/Materials/MaterialIR.h
@@ -6,6 +6,8 @@
 
 #if WITH_EDITOR
 
+struct FMaterialExternalCodeDeclaration;
+
 namespace MIR {
 
 enum EValueKind
@@ -29,6 +31,7 @@ enum EValueKind
 	VK_Subscript,
 	VK_Cast,
 	VK_TextureSample,
+	VK_ExternalCode,
 
 	VK_InstructionEnd,
 };
@@ -322,6 +325,12 @@ struct FSetMaterialOutput : TInstruction<VK_SetMaterialOutput>
 	FValue* Arg;
 };
 
+// MIR instruction for external code declarations; See BaseMaterialExpressions.ini file.
+struct FExternalCode : TInstruction<VK_ExternalCode>
+{
+	const FMaterialExternalCodeDeclaration* ExternalCodeDeclaration;
+};
+
 // Note: if you modify this enum, update the implementations of the
 // helper functions below.
 enum EOperator
diff --git a/Engine/Source/Runtime/Engine/Public/Materials/MaterialIREmitter.h b/Engine/Source/Runtime/Engine/Public/Materials/MaterialIREmitter.h
--- a/Engine/Source/Runtime/Engine/Public/Materials/MaterialIREmitter.h
+++ b/Engine/Source/Runtime/Engine/Public/Materials/MaterialIREmitter.h
@@ -161,6 +161,7 @@ public:
 	FValue* EmitBranch(FValue* Condition, FValue* True, FValue* False);
 	FValue* EmitConstruct(FTypePtr Type, FValue* Initializer);
 	FValue* EmitTextureSample(FValue* Texture, FValue* TexCoord, ESamplerSourceMode SamplerSourceMode, ETextureMipValueMode MipValueMode, EMaterialSamplerType SamplerType);
+	FValue* EmitExternalCode(const FMaterialExternalCodeDeclaration* InExternalCodeDeclaration);
 
 	/* Types */
 
