Changing LogActor CompileTimeVerbosity to All rather than Warning.  The assumption is that setting it to Warning was an oversight, making it impossible to change the Verbosity of LogActor without recompiling the Engine.  It looks like some Logs were bumped up to Warning to compensate for that oversight, but we are attempting to rectify it now.

diff --git a/Engine/Source/Runtime/Engine/Classes/GameFramework/Actor.h b/Engine/Source/Runtime/Engine/Classes/GameFramework/Actor.h
--- a/Engine/Source/Runtime/Engine/Classes/GameFramework/Actor.h
+++ b/Engine/Source/Runtime/Engine/Classes/GameFramework/Actor.h
@@ -159,7 +159,7 @@ private:
 	friend struct FPrimitiveSceneInfoAdapter;
 };
 
-ENGINE_API DECLARE_LOG_CATEGORY_EXTERN(LogActor, Log, Warning);
+ENGINE_API DECLARE_LOG_CATEGORY_EXTERN(LogActor, Log, All);
 
 // Delegate signatures
 DECLARE_DYNAMIC_MULTICAST_SPARSE_DELEGATE_FiveParams( FTakeAnyDamageSignature, AActor, OnTakeAnyDamage, AActor*, DamagedActor, float, Damage, const class UDamageType*, DamageType, class AController*, InstigatedBy, AActor*, DamageCauser );
