Fix for CIS Issue 780449.

Adding the exported symbols to allow inheritance of FTexture2DDynamicResource.

diff --git a/Engine/Source/Runtime/Engine/Public/TextureResource.h b/Engine/Source/Runtime/Engine/Public/TextureResource.h
--- a/Engine/Source/Runtime/Engine/Public/TextureResource.h
+++ b/Engine/Source/Runtime/Engine/Public/TextureResource.h
@@ -249,10 +249,10 @@ public:
 	ENGINE_API FTexture2DDynamicResource(class UTexture2DDynamic* InOwner);
 
 	/** Returns the width of the texture in pixels. */
-	virtual uint32 GetSizeX() const override;
+	ENGINE_API virtual uint32 GetSizeX() const override;
 
 	/** Returns the height of the texture in pixels. */
-	virtual uint32 GetSizeY() const override;
+	ENGINE_API virtual uint32 GetSizeY() const override;
 
 	/** Called when the resource is initialized. This is only called by the rendering thread. */
 	ENGINE_API virtual void InitRHI(FRHICommandListBase& RHICmdList) override;
