WorldPartition: Added 3D display of 2D partition cells (when using wp.runtime.ToggleDrawRuntimeHash3D)

diff --git a/Engine/Source/Runtime/Engine/Private/WorldPartition/RuntimeHashSet/WorldPartitionRuntimeHashSetDebugDraw.cpp b/Engine/Source/Runtime/Engine/Private/WorldPartition/RuntimeHashSet/WorldPartitionRuntimeHashSetDebugDraw.cpp
--- a/Engine/Source/Runtime/Engine/Private/WorldPartition/RuntimeHashSet/WorldPartitionRuntimeHashSetDebugDraw.cpp
+++ b/Engine/Source/Runtime/Engine/Private/WorldPartition/RuntimeHashSet/WorldPartitionRuntimeHashSetDebugDraw.cpp
@@ -365,71 +365,76 @@ void UWorldPartitionRuntimeHashSet::Draw3D(const TArray<FWorldPartitionStreaming
 	const EWorldPartitionRuntimeCellVisualizeMode VisualizeMode = GetStreamingCellVisualizeMode();
 	const TMap<FName, FColor> DataLayerDebugColors = GetDataLayerDebugColors(WorldPartition);
 
-	ForEachStreamingData([this, &Sources, VisualizeMode, &DataLayerDebugColors, ContentBundleManager, OwningWorld, &WorldPartitionTransform](const FRuntimePartitionStreamingData& StreamingData)
+	auto DebugDrawCell3D = [WorldPartitionTransform, OwningWorld, ContentBundleManager, VisualizeMode, &DataLayerDebugColors](UWorldPartitionRuntimeCell* Cell)
+	{
+		if (!IsDebugCellShown(Cell))
+		{
+			return;
+		}
+
+		const FVector2D CellStreamingBoundsSize = FVector2D(Cell->GetStreamingBounds().GetSize());
+		const FVector2D CellStreamingBoundsMin = FVector2D(Cell->GetStreamingBounds().Min);
+
+		float CellOpacity = 0.0f;
+		TArray<FLinearColor> CellColors;
+
+		switch (GShowRuntimeHashSetDebugDisplayMode)
+		{
+		case 0:
+		case 3:
+			CellColors.Add(Cell->GetDebugColor(VisualizeMode));
+			CellOpacity = 0.25f / FMath::Max<float>(GShowRuntimeHashSetDebugDisplayLevelCount, 1);
+			break;
+		case 1:
+			if (DataLayerDebugColors.Num() && Cell->GetDataLayers().Num())
+			{
+				for (const FName& DataLayer : Cell->GetDataLayers())
+				{
+					CellColors.Add(DataLayerDebugColors[DataLayer]);
+				}
+				CellOpacity = 0.67f;
+			}
+			break;
+		case 2:
+			if (ContentBundleManager && Cell->GetContentBundleID().IsValid())
+			{
+				if (const FContentBundleBase* ContentBundle = ContentBundleManager->GetContentBundle(OwningWorld, Cell->GetContentBundleID()))
+				{
+					check(ContentBundle->GetDescriptor());
+					CellColors.Add(ContentBundle->GetDescriptor()->GetDebugColor());
+					CellOpacity = 0.67f;
+				}
+			}
+			break;
+		}
+
+		if (CellColors.IsEmpty())
+		{
+			CellColors.Add(FLinearColor::White);
+			CellOpacity = 0.1f;
+		}
+
+		// Draw Cell using its debug color
+		const FBox Box(Cell->GetStreamingBounds());
+		const FVector BoxCenter(Box.GetCenter());
+		const FColor BoxColor(CellColors[0].CopyWithNewOpacity(CellOpacity).ToFColor(true));
+		const FVector CellPos = WorldPartitionTransform.TransformPosition(BoxCenter);
+		DrawDebugBox(OwningWorld, CellPos, Box.GetExtent(), WorldPartitionTransform.GetRotation(), BoxColor.WithAlpha(255), false, -1.f, 255, 20.f);
+	};
+
+	ForEachStreamingData([this, &Sources, VisualizeMode, &DataLayerDebugColors, ContentBundleManager, OwningWorld, &WorldPartitionTransform, &DebugDrawCell3D](const FRuntimePartitionStreamingData& StreamingData)
 	{
 		if (StreamingData.SpatialIndex)
 		{
 			for (const FWorldPartitionStreamingSource& Source : Sources)
 			{
-				Source.ForEachShape(StreamingData.GetLoadingRange(), StreamingData.Name, false, [this, &StreamingData, VisualizeMode, &DataLayerDebugColors, ContentBundleManager, OwningWorld, &WorldPartitionTransform](const FSphericalSector& Shape)
+				Source.ForEachShape(StreamingData.GetLoadingRange(), StreamingData.Name, false, [this, &StreamingData, VisualizeMode, &DataLayerDebugColors, ContentBundleManager, OwningWorld, &WorldPartitionTransform, &DebugDrawCell3D](const FSphericalSector& Shape)
 				{
 					auto ForEachIntersectingElement = [&]<typename ShapeType>(const ShapeType& InShape)
 					{
-						StreamingData.SpatialIndex.Get()->ForEachIntersectingElement(InShape, [this, VisualizeMode, &DataLayerDebugColors, ContentBundleManager, OwningWorld, &WorldPartitionTransform](UWorldPartitionRuntimeCell* Cell)
+						StreamingData.SpatialIndex.Get()->ForEachIntersectingElement(InShape, [this, VisualizeMode, &DataLayerDebugColors, ContentBundleManager, OwningWorld, &WorldPartitionTransform, &DebugDrawCell3D](UWorldPartitionRuntimeCell* Cell)
 						{
-							if (!IsDebugCellShown(Cell))
-							{
-								return;
-							}
-
-							const FVector2D CellStreamingBoundsSize = FVector2D(Cell->GetStreamingBounds().GetSize());
-							const FVector2D CellStreamingBoundsMin = FVector2D(Cell->GetStreamingBounds().Min);
-
-							float CellOpacity = 0.0f;
-							TArray<FLinearColor> CellColors;
-
-							switch (GShowRuntimeHashSetDebugDisplayMode)
-							{
-							case 0:
-							case 3:
-								CellColors.Add(Cell->GetDebugColor(VisualizeMode));
-								CellOpacity = 0.25f / FMath::Max<float>(GShowRuntimeHashSetDebugDisplayLevelCount, 1);
-								break;
-							case 1:
-								if (DataLayerDebugColors.Num() && Cell->GetDataLayers().Num())
-								{
-									for (const FName& DataLayer : Cell->GetDataLayers())
-									{
-										CellColors.Add(DataLayerDebugColors[DataLayer]);
-									}
-									CellOpacity = 0.67f;
-								}
-								break;
-							case 2:
-								if (ContentBundleManager && Cell->GetContentBundleID().IsValid())
-								{
-									if (const FContentBundleBase* ContentBundle = ContentBundleManager->GetContentBundle(OwningWorld, Cell->GetContentBundleID()))
-									{
-										check(ContentBundle->GetDescriptor());
-										CellColors.Add(ContentBundle->GetDescriptor()->GetDebugColor());
-										CellOpacity = 0.67f;
-									}
-								}
-								break;
-							}
-
-							if (CellColors.IsEmpty())
-							{
-								CellColors.Add(FLinearColor::White);
-								CellOpacity = 0.1f;
-							}
-
-							// Draw Cell using its debug color
-							const FBox Box(Cell->GetStreamingBounds());
-							const FVector BoxCenter(Box.GetCenter());
-							const FColor BoxColor(CellColors[0].CopyWithNewOpacity(CellOpacity).ToFColor(true));
-							const FVector CellPos = WorldPartitionTransform.TransformPosition(BoxCenter);
-							DrawDebugBox(OwningWorld, CellPos, Box.GetExtent(), WorldPartitionTransform.GetRotation(), BoxColor.WithAlpha(255), false, -1.f, 255, 20.f);
+							DebugDrawCell3D(Cell);
 						});
 					};
 
@@ -446,6 +451,20 @@ void UWorldPartitionRuntimeHashSet::Draw3D(const TArray<FWorldPartitionStreaming
 				});
 			}
 		}
+
+		if (StreamingData.SpatialIndex2D.IsValid())
+		{
+			for (const FWorldPartitionStreamingSource& Source : Sources)
+			{
+				const bool bCalcBoundsIn2D = true;
+				const FBox SourceBounds = Source.CalcBounds(StreamingData.GetLoadingRange(), StreamingData.Name, bCalcBoundsIn2D);
+				const FBox2D SourceBounds2D(FVector2D(SourceBounds.Min.X, SourceBounds.Min.Y), FVector2D(SourceBounds.Max.X, SourceBounds.Max.Y));
+				StreamingData.SpatialIndex2D->ForEachIntersectingElement(SourceBounds2D, [&DebugDrawCell3D](UWorldPartitionRuntimeCell* Cell)
+				{
+					DebugDrawCell3D(Cell);
+				});
+			}
+		}
 		return true;
 	});
 }
\ No newline at end of file
