[WorldPartition] Do not flatten bounds (in Z) when using 2D hash set.

This was a remnant of the older 2D compatible hashset implementation which is now replaced by FSpatialIndexProfile2D when executing intersection checks.

Flattened bounds can affect distance based checks which are used when calculating priority and streaming metrics.

#rnx
#tests Level streaming tests on PF with the change.

diff --git a/Engine/Source/Runtime/Engine/Private/WorldPartition/RuntimeHashSet/RuntimePartitionLHGrid.cpp b/Engine/Source/Runtime/Engine/Private/WorldPartition/RuntimeHashSet/RuntimePartitionLHGrid.cpp
--- a/Engine/Source/Runtime/Engine/Private/WorldPartition/RuntimeHashSet/RuntimePartitionLHGrid.cpp
+++ b/Engine/Source/Runtime/Engine/Private/WorldPartition/RuntimeHashSet/RuntimePartitionLHGrid.cpp
@@ -198,6 +198,7 @@ bool URuntimePartitionLHGrid::GenerateStreaming(const FGenerateStreamingParams&
 		{
 			FBox ActorSetInstanceBounds = ActorSetInstance->Bounds;
 
+			// Ensure cached bounds are in 2D so as to not affect CellCoords.
 			if (bIs2D)
 			{
 				ActorSetInstanceBounds.Min.Z = 0;
@@ -223,13 +224,6 @@ bool URuntimePartitionLHGrid::GenerateStreaming(const FGenerateStreamingParams&
 		if (bIsSpatiallyLoaded)
 		{
 			CellDesc.CellBounds = FCellCoord::GetCellBounds(CellCoord, CellSize, FVector(Origin.X, Origin.Y, bIs2D ? 0 : Origin.Z));
-
-			if (bIs2D)
-			{
-				CellDesc.CellBounds.GetValue().Min.Z = 0;
-				CellDesc.CellBounds.GetValue().Max.Z = 0;
-			}
-
 			CellDesc.bIs2D = bIs2D;
 		}
 	}
diff --git a/Engine/Source/Runtime/Engine/Private/WorldPartition/RuntimeHashSet/WorldPartitionRuntimeHashSetStreamingGeneration.cpp b/Engine/Source/Runtime/Engine/Private/WorldPartition/RuntimeHashSet/WorldPartitionRuntimeHashSetStreamingGeneration.cpp
--- a/Engine/Source/Runtime/Engine/Private/WorldPartition/RuntimeHashSet/WorldPartitionRuntimeHashSetStreamingGeneration.cpp
+++ b/Engine/Source/Runtime/Engine/Private/WorldPartition/RuntimeHashSet/WorldPartitionRuntimeHashSetStreamingGeneration.cpp
@@ -184,12 +184,6 @@ bool UWorldPartitionRuntimeHashSet::GenerateStreaming(UWorldPartitionStreamingPo
 
 				UWorldPartitionRuntimeCellDataHashSet* RuntimeCellData = CastChecked<UWorldPartitionRuntimeCellDataHashSet>(RuntimeCell->RuntimeCellData);
 
-				if (RuntimeCellData->bIs2D)
-				{
-					RuntimeCell->RuntimeCellData->ContentBounds.Min.Z = 0.0;
-					RuntimeCell->RuntimeCellData->ContentBounds.Max.Z = 0.0;
-				}
-
 				if (CellDescInstance.CellBounds.IsSet())
 				{
 					if (!RuntimeCell->RuntimeCellData->HierarchicalLevel)
