Fix tick dependency bug where if a tick function with an early start and end group depended on a function with an early start but a late end group, it would deadlock forever. With this change it now defers the end group separately from the begin group

diff --git a/Engine/Source/Runtime/Engine/Private/TickTaskManager.cpp b/Engine/Source/Runtime/Engine/Private/TickTaskManager.cpp
--- a/Engine/Source/Runtime/Engine/Private/TickTaskManager.cpp
+++ b/Engine/Source/Runtime/Engine/Private/TickTaskManager.cpp
@@ -2367,7 +2367,8 @@ void FTickFunction::QueueTickFunction(FTickTaskSequencer& TTS, const struct FTic
 		InternalData->TickVisitedGFrameCounter.store(CurrentFrameCounter, std::memory_order_relaxed);
 		if (TickState != FTickFunction::ETickState::Disabled)
 		{
-			ETickingGroup MaxPrerequisiteTickGroup =  ETickingGroup(0);
+			ETickingGroup MaxStartTickGroup = ETickingGroup(0);
+			ETickingGroup MaxEndTickGroup = ETickingGroup(0);
 
 			TArray<FTickFunction*> RawPrerequisites;
 			for (int32 PrereqIndex = 0; PrereqIndex < Prerequisites.Num(); PrereqIndex++)
@@ -2393,14 +2394,15 @@ void FTickFunction::QueueTickFunction(FTickTaskSequencer& TTS, const struct FTic
 					}
 					else if (TTS.ShouldConsiderPrerequisite(this, Prereq))
 					{
-						MaxPrerequisiteTickGroup =  FMath::Max<ETickingGroup>(MaxPrerequisiteTickGroup, Prereq->InternalData->ActualStartTickGroup.GetValue());
+						MaxStartTickGroup = FMath::Max<ETickingGroup>(MaxStartTickGroup, Prereq->InternalData->ActualStartTickGroup.GetValue());
+						MaxEndTickGroup = FMath::Max<ETickingGroup>(MaxEndTickGroup, Prereq->InternalData->ActualEndTickGroup.GetValue());
 						RawPrerequisites.Add(Prereq);
 					}
 				}
 			}
 
 			// tick group is the max of the prerequisites, the current tick group, and the desired tick group
-			ETickingGroup MyActualTickGroup =  FMath::Max<ETickingGroup>(MaxPrerequisiteTickGroup, FMath::Max<ETickingGroup>(TickGroup.GetValue(),TickContext.TickGroup));
+			ETickingGroup MyActualTickGroup = FMath::Max<ETickingGroup>(MaxStartTickGroup, FMath::Max<ETickingGroup>(TickGroup.GetValue(), TickContext.TickGroup));
 			if (MyActualTickGroup != TickGroup)
 			{
 				// if the tick was "demoted", make sure it ends up in an ordinary tick group.
@@ -2411,11 +2413,15 @@ void FTickFunction::QueueTickFunction(FTickTaskSequencer& TTS, const struct FTic
 			}
 			InternalData->ActualStartTickGroup = MyActualTickGroup;
 			InternalData->ActualEndTickGroup = MyActualTickGroup;
-			if (EndTickGroup > MyActualTickGroup)
+
+			// Also check to see if the end tick group needs to be extended separately
+			ETickingGroup MyActualEndTickGroup = FMath::Max<ETickingGroup>(MaxEndTickGroup, FMath::Max<ETickingGroup>(EndTickGroup.GetValue(), MyActualTickGroup));
+
+			if (MyActualEndTickGroup > MyActualTickGroup)
 			{
-				check(EndTickGroup <= TG_NewlySpawned);
+				check(MyActualEndTickGroup <= TG_NewlySpawned);
 				ETickingGroup TestTickGroup = ETickingGroup(MyActualTickGroup + 1);
-				while (TestTickGroup <= EndTickGroup)
+				while (TestTickGroup <= MyActualEndTickGroup)
 				{
 					if (CanDemoteIntoTickGroup(TestTickGroup))
 					{
@@ -2450,7 +2456,8 @@ void FTickFunction::QueueTickFunctionParallel(const struct FTickContext& TickCon
 		check(IsTickFunctionRegistered());
 		if (TickState != FTickFunction::ETickState::Disabled)
 		{
-			ETickingGroup MaxPrerequisiteTickGroup = ETickingGroup(0);
+			ETickingGroup MaxStartTickGroup = ETickingGroup(0);
+			ETickingGroup MaxEndTickGroup = ETickingGroup(0);
 
 			FGraphEventArray TaskPrerequisites;
 			if (Prerequisites.Num())
@@ -2479,7 +2486,8 @@ void FTickFunction::QueueTickFunctionParallel(const struct FTickContext& TickCon
 						}
 						else
 						{
-							MaxPrerequisiteTickGroup = FMath::Max<ETickingGroup>(MaxPrerequisiteTickGroup, Prereq->InternalData->ActualStartTickGroup.GetValue());
+							MaxStartTickGroup = FMath::Max<ETickingGroup>(MaxStartTickGroup, Prereq->InternalData->ActualStartTickGroup.GetValue());
+							MaxEndTickGroup = FMath::Max<ETickingGroup>(MaxEndTickGroup, Prereq->InternalData->ActualEndTickGroup.GetValue());
 							TaskPrerequisites.Add(Prereq->GetCompletionHandle());
 						}
 					}
@@ -2488,7 +2496,7 @@ void FTickFunction::QueueTickFunctionParallel(const struct FTickContext& TickCon
 			}
 
 			// tick group is the max of the prerequisites, the current tick group, and the desired tick group
-			ETickingGroup MyActualTickGroup =  FMath::Max<ETickingGroup>(MaxPrerequisiteTickGroup, FMath::Max<ETickingGroup>(TickGroup,TickContext.TickGroup));
+			ETickingGroup MyActualTickGroup =  FMath::Max<ETickingGroup>(MaxStartTickGroup, FMath::Max<ETickingGroup>(TickGroup,TickContext.TickGroup));
 			if (MyActualTickGroup != TickGroup)
 			{
 				// if the tick was "demoted", make sure it ends up in an ordinary tick group.
@@ -2499,11 +2507,15 @@ void FTickFunction::QueueTickFunctionParallel(const struct FTickContext& TickCon
 			}
 			InternalData->ActualStartTickGroup = MyActualTickGroup;
 			InternalData->ActualEndTickGroup = MyActualTickGroup;
-			if (EndTickGroup > MyActualTickGroup)
+		
+			// Also check to see if the end tick group needs to be extended separately
+			ETickingGroup MyActualEndTickGroup = FMath::Max<ETickingGroup>(MaxEndTickGroup, FMath::Max<ETickingGroup>(EndTickGroup.GetValue(), MyActualTickGroup));
+
+			if (MyActualEndTickGroup > MyActualTickGroup)
 			{
-				check(EndTickGroup <= TG_NewlySpawned);
+				check(MyActualEndTickGroup <= TG_NewlySpawned);
 				ETickingGroup TestTickGroup = ETickingGroup(MyActualTickGroup + 1);
-				while (TestTickGroup <= EndTickGroup)
+				while (TestTickGroup <= MyActualEndTickGroup)
 				{
 					if (CanDemoteIntoTickGroup(TestTickGroup))
 					{
