Add edit condition to substrate layers project settings option.

diff --git a/Engine/Source/Runtime/Engine/Classes/Engine/RendererSettings.h b/Engine/Source/Runtime/Engine/Classes/Engine/RendererSettings.h
--- a/Engine/Source/Runtime/Engine/Classes/Engine/RendererSettings.h
+++ b/Engine/Source/Runtime/Engine/Classes/Engine/RendererSettings.h
@@ -1234,7 +1234,8 @@ class URendererSettings : public UDeveloperSettings
 	UPROPERTY(config, EditAnywhere, Category = Substrate, meta = (
 		ConsoleVariable = "r.Substrate.EnableLayerSupport", DisplayName = "Substrate enable material layer support (Experimental)",
 		ToolTip = "Enable Substrate material layering and UI(Experimental).Note: This support is one way; the legacy layer materials auto-upgrade and can only be reversed manually after assets are re - saved.",
-		ConfigRestartRequired = true))
+		ConfigRestartRequired = true,
+		EditCondition = "bEnableSubstrate"))
 		uint32 SubstrateEnableLayerSupport : 1;
 
 	/**
