Clarify SimpleVolume SSS type and fix MaterialSubsurfaceType's category.

diff --git a/Engine/Source/Runtime/Engine/Public/Materials/MaterialExpressionSubstrate.h b/Engine/Source/Runtime/Engine/Public/Materials/MaterialExpressionSubstrate.h
--- a/Engine/Source/Runtime/Engine/Public/Materials/MaterialExpressionSubstrate.h
+++ b/Engine/Source/Runtime/Engine/Public/Materials/MaterialExpressionSubstrate.h
@@ -20,7 +20,7 @@ enum EMaterialSubSurfaceType : int
 	MSS_TwoSidedWrap		= SSS_TYPE_TWO_SIDED_WRAP		UMETA(DisplayName="Two-Sided Wrap", ToolTip="Approximation using wrap-lighting and handling thin surface (e.g.: foliage)"),
 	MSS_Diffusion			= SSS_TYPE_DIFFUSION 			UMETA(DisplayName="Diffusion", ToolTip="Diffusion based sub-surface scattering"),
 	MSS_DiffusionProfile	= SSS_TYPE_DIFFUSION_PROFILE 	UMETA(Hidden),
-	MSS_SimpleVolume		= SSS_TYPE_SIMPLEVOLUME 		UMETA(DisplayName="Simple Volume", ToolTip="Approximation of optically thin slab (i.e., light visible through the material)"),
+	MSS_SimpleVolume		= SSS_TYPE_SIMPLEVOLUME 		UMETA(DisplayName="Simple Volume", ToolTip="Approximation of optically thin slab (e.g.: glass) where light is visible through the material"),
 	MSS_MAX
 };
 
@@ -359,7 +359,7 @@ class UMaterialExpressionSubstrateSlabBSDF : public UMaterialExpressionSubstrate
 	 *   * For slab not sitting at the bottom of the topology (e.g. slabs stacked with a vertial operator), only the SimpleVolume SSS type is available.
 	 *   * In non-opaque blend modes, SSS Diffusion and SSS Diffusion Profile are not available and will fallback onto Wrap mode.
 	 */
-	UPROPERTY(EditAnywhere, Category = ShadingModel, meta = (DisplayName = "Sub-Surface Type"))
+	UPROPERTY(EditAnywhere, Category = Material, meta = (DisplayName = "Sub-Surface Type"))
 	TEnumAsByte<enum EMaterialSubSurfaceType> SubSurfaceType = MSS_None;
 
 	//~ Begin UMaterialExpression Interface
