Enable virtual texture sampling for post process materials.
Previously this was kept behind a renderer setting, and virtual texture feedback was not supported.
Now virtual texture feedback has been added and the renderer setting removed

diff --git a/Engine/Shaders/Private/PostProcessMaterialShaders.usf b/Engine/Shaders/Private/PostProcessMaterialShaders.usf
--- a/Engine/Shaders/Private/PostProcessMaterialShaders.usf
+++ b/Engine/Shaders/Private/PostProcessMaterialShaders.usf
@@ -6,6 +6,7 @@
 
 #include "Common.ush"
 #include "ScreenPass.ush"
+#include "VirtualTextureCommon.ush"
 
 #pragma message("UESHADERMETADATA_VERSION 02E0E449-2223-45B4-8CD0-93AB087FBEE3")
 
@@ -354,6 +355,10 @@ void MainPS(
 	Alpha = SceneTextureLookup(float2(0,0), PPI_PropagatedAlpha, false).r;
 #endif
 
+#if MATERIAL_VIRTUALTEXTURE_FEEDBACK
+	FinalizeVirtualTextureFeedback(Parameters.VirtualTextureFeedback, Parameters.SvPosition, View.VTFeedbackBuffer);
+#endif
+
 	OutColor = float4(EmissiveColor, Alpha );
 
 #if !MATERIALBLENDING_MODULATE && POST_PROCESS_MATERIAL_BEFORE_TONEMAP && !POST_PROCESS_DISABLE_PRE_EXPOSURE_SCALE
@@ -534,6 +539,10 @@ void MainPS(
 
 #endif // SUBSTRATE_ENABLED
 
+#if MATERIAL_VIRTUALTEXTURE_FEEDBACK
+	FinalizeVirtualTextureFeedback(Parameters.VirtualTextureFeedback, Parameters.SvPosition, View.VTFeedbackBuffer);
+#endif
+
 	half4 FullSceneColor = half4(EmissiveColor, Alpha);
 	   
 #if POST_PROCESS_MATERIAL_BEFORE_TONEMAP
diff --git a/Engine/Source/Runtime/Engine/Classes/Engine/RendererSettings.h b/Engine/Source/Runtime/Engine/Classes/Engine/RendererSettings.h
--- a/Engine/Source/Runtime/Engine/Classes/Engine/RendererSettings.h
+++ b/Engine/Source/Runtime/Engine/Classes/Engine/RendererSettings.h
@@ -424,12 +424,6 @@ class URendererSettings : public UDeveloperSettings
 		ToolTip = "Relax restriction on virtual textures contributing to Opacity Mask. In some edge cases this can lead to low resolution shadow edges."))
 		uint32 bEnableVirtualTextureOpacityMask : 1;
 
-	UPROPERTY(config, EditAnywhere, Category = VirtualTextures, meta = (
-		EditCondition = "bVirtualTextures",
-		DisplayName = "Enable virtual textures for Post Processing Materials",
-		ToolTip = "Relax restriction on virtual textures being sampled in Post Processing Materials. There is no feedback from the Post Processing passes, so virtual textures will only have pages in the cache that are triggered by feedback in another pass."))
-		uint32 bEnableVirtualTexturePostProcessing : 1;
-
 	UPROPERTY(config, EditAnywhere, Category = VirtualTextures, meta = (
 		EditCondition = "bVirtualTextures",
 		ConsoleVariable = "r.VT.TileSize", DisplayName = "Tile size",
diff --git a/Engine/Source/Runtime/Engine/Private/Materials/HLSLMaterialTranslator.cpp b/Engine/Source/Runtime/Engine/Private/Materials/HLSLMaterialTranslator.cpp
--- a/Engine/Source/Runtime/Engine/Private/Materials/HLSLMaterialTranslator.cpp
+++ b/Engine/Source/Runtime/Engine/Private/Materials/HLSLMaterialTranslator.cpp
@@ -7197,12 +7197,15 @@ int32 FHLSLMaterialTranslator::TextureSample(
 	const bool bVirtualTexture = bMeshPaint || TextureType == MCT_TextureVirtual;
 	if (bVirtualTexture)
 	{
-		if (Material->GetMaterialDomain() != MD_Surface && Material->GetMaterialDomain() != MD_DeferredDecal)
+		const EMaterialDomain Domain = Material->GetMaterialDomain();
+		switch (Domain)
 		{
-			if (GetDefault<URendererSettings>()->bEnableVirtualTexturePostProcessing == false)
-			{
-				return Errorf(TEXT("Sampling a virtual texture is currently only supported for surface and decal materials."));
-			}
+			case MD_Surface:
+			case MD_DeferredDecal:
+			case MD_PostProcess:
+				break;
+			default:
+				return Errorf(TEXT("Sampling a virtual texture is not supported for the %s material domain."), *MaterialDomainString(Domain));
 		}
 	}
 
diff --git a/Engine/Source/Runtime/Engine/Private/RendererSettings.cpp b/Engine/Source/Runtime/Engine/Private/RendererSettings.cpp
--- a/Engine/Source/Runtime/Engine/Private/RendererSettings.cpp
+++ b/Engine/Source/Runtime/Engine/Private/RendererSettings.cpp
@@ -104,7 +104,6 @@ URendererSettings::URendererSettings(const FObjectInitializer& ObjectInitializer
 	WhiteChromaticityCoordinate = FVector2D::ZeroVector;
 	bUseLegacyLuminanceFactors = false;
 	bEnableVirtualTextureOpacityMask = false;
-	bEnableVirtualTexturePostProcessing = false;
 
 #if WITH_EDITOR
 	if (IsTemplate())
