New material translator push:
- Added uniform parameter registration to ModuleIRModule, allowing referenced uniform parameters to be registered once within the module. A numeric ID is now used by MIR::FUniformParameter to reference the parameter instead of the non-trivial FMaterialParameterInfo and FMaterialParameterMetadata.
- Added a static assert to ensure that all MIR::FValue types are trivial, meaning they have a trivial constructor and destructor.

#rnx

diff --git a/Engine/Source/Runtime/Engine/Private/Materials/MaterialExpressionsIR.cpp b/Engine/Source/Runtime/Engine/Private/Materials/MaterialExpressionsIR.cpp
--- a/Engine/Source/Runtime/Engine/Private/Materials/MaterialExpressionsIR.cpp
+++ b/Engine/Source/Runtime/Engine/Private/Materials/MaterialExpressionsIR.cpp
@@ -845,6 +845,20 @@ static void BuildTextureSampleExpression(MIR::FValue* Texture, UMaterialExpressi
 	Emitter.Output(5, OutputValue); 
 }
 
+// Returns this value's textures if it has one (nullptr otherwise).
+static EMaterialSamplerType GetValueMaterialSamplerType(const MIR::FValue* Value)
+{
+	if (auto TextureObject = MIR::As<MIR::FTextureObject>(Value))
+	{
+		return TextureObject->SamplerType;
+	}
+	else if (auto UniformParameter = MIR::As<MIR::FUniformParameter>(Value))
+	{
+		return UniformParameter->SamplerType;
+	}
+	return SAMPLERTYPE_MAX;
+}
+
 void UMaterialExpressionTextureSample::Build(MIR::FEmitter& Emitter)
 {
 	MIR::FValue* TextureValue = Emitter.TryInput(&TextureObject);
@@ -865,7 +879,7 @@ void UMaterialExpressionTextureSample::Build(MIR::FEmitter& Emitter)
 
 	UE_MIR_CHECKPOINT();
 
-	if (TextureValue && TextureValue->GetTexture()->Get<EMaterialSamplerType>() != SAMPLERTYPE_Color)
+	if (TextureValue && GetValueMaterialSamplerType(TextureValue) != SAMPLERTYPE_Color)
 	{
 		Emitter.Error(TEXT("Input texture sampler type must be color."));
 		return;
diff --git a/Engine/Source/Runtime/Engine/Private/Materials/MaterialIR.cpp b/Engine/Source/Runtime/Engine/Private/Materials/MaterialIR.cpp
--- a/Engine/Source/Runtime/Engine/Private/Materials/MaterialIR.cpp
+++ b/Engine/Source/Runtime/Engine/Private/Materials/MaterialIR.cpp
@@ -238,21 +238,6 @@ bool FValue::AreAllNearlyOne() const
 	return false;
 }
 
-TOptional<TPair<UTexture*, EMaterialSamplerType>> FValue::GetTexture()
-{
-	using FPair = TPair<UTexture*, EMaterialSamplerType>;
-	if (auto TextureObject = As<MIR::FTextureObject>())
-	{
-		return FPair{ TextureObject->Texture, TextureObject->SamplerType };
-	}
-	else if (auto UniformParameter = As<MIR::FUniformParameter>())
-	{
-		check(UniformParameter->Metadata.Value.Type == EMaterialParameterType::Texture);
-		return FPair{ UniformParameter->Metadata.Value.Texture, UniformParameter->SamplerType };
-	}
-	return {};
-}
-
 FInstruction* AsInstruction(FValue* Value)
 {
 	return Value && (Value->Kind > VK_InstructionBegin && Value->Kind < VK_InstructionEnd) ? static_cast<FInstruction*>(Value) : nullptr;
diff --git a/Engine/Source/Runtime/Engine/Private/Materials/MaterialIRDebug.cpp b/Engine/Source/Runtime/Engine/Private/Materials/MaterialIRDebug.cpp
--- a/Engine/Source/Runtime/Engine/Private/Materials/MaterialIRDebug.cpp
+++ b/Engine/Source/Runtime/Engine/Private/Materials/MaterialIRDebug.cpp
@@ -175,6 +175,9 @@ void DebugDumpIRUseGraph(const FMaterialIRModule& Module)
 /// Helper struct to wrap the state used during IR to text dumping.
 struct FDebugDumpIRState
 {
+	/// The module we are printing the IR for.
+	const FMaterialIRModule* Module{};
+
 	/// Output string containing the generated result.
 	FString Out; 
 
@@ -332,7 +335,8 @@ struct FDebugDumpIRState
 		}
 		else if (const FUniformParameter* UniformParameter = Value->As<FUniformParameter>())
 		{
-			Out.Appendf(TEXT("[Parameter #%d \"%s\"]"), UniformParameter->Analysis_UniformParameterIndex, *UniformParameter->Info.Name.ToString());
+			FName ParameterName = Module->GetParameterInfo(UniformParameter->ParameterIdInModule).Name;
+			Out.Appendf(TEXT("[Parameter #%d \"%s\"]"), UniformParameter->Analysis_UniformParameterIndex, *ParameterName.ToString());
 			ReferencedParameters.AddUnique({ UniformParameter->Analysis_UniformParameterIndex, UniformParameter });
 		}
 		else
@@ -389,10 +393,12 @@ struct FDebugDumpIRState
 			const FValue* Value = ParamPair.Value;
 			if (const FUniformParameter* Param = Value->As<FUniformParameter>())
 			{
+				FName ParameterName = Module->GetParameterInfo(Param->ParameterIdInModule).Name;
+				EMaterialParameterType ParameterType = Module->GetParameterMetadata(Param->ParameterIdInModule).Value.Type;
 				Out.Appendf(TEXT("#%d = Name=\"%s\" Type=\"%s\"\n"),
 					Param->Analysis_UniformParameterIndex,
-					*Param->Info.Name.ToString(),
-					*MaterialParameterTypeToString(Param->Metadata.Value.Type));
+					*ParameterName.ToString(),
+					*MaterialParameterTypeToString(ParameterType));
 			}
 			else if (const FTextureObject* TextureObject = Value->As<FTextureObject>())
 			{
@@ -408,6 +414,7 @@ struct FDebugDumpIRState
 FString DebugDumpIR(FStringView MaterialName, const FMaterialIRModule& Module)
 {
 	FDebugDumpIRState State{};
+	State.Module = &Module;
 	State.Out.Append(TEXT("; Material IR module dump.\n"));
 	State.Out.Appendf(TEXT(";    Material: %s\n"), MaterialName.GetData());
 	State.Out.Append(TEXT("\n"));
diff --git a/Engine/Source/Runtime/Engine/Private/Materials/MaterialIREmitter.cpp b/Engine/Source/Runtime/Engine/Private/Materials/MaterialIREmitter.cpp
--- a/Engine/Source/Runtime/Engine/Private/Materials/MaterialIREmitter.cpp
+++ b/Engine/Source/Runtime/Engine/Private/Materials/MaterialIREmitter.cpp
@@ -123,6 +123,9 @@ static FValue* EmitNew(FEmitter* Emitter, FValue* Value)
 template <typename T>
 static T MakePrototype(FTypePtr InType)
 {
+	static_assert(std::is_trivially_constructible_v<T> && std::is_trivially_destructible_v<T> &&  std::is_trivially_copy_constructible_v<T> && std::is_trivially_copy_assignable_v<T>,
+		"FValues are expected to be trivial types.");
+
 	T Value;
 	FMemory::Memzero(Value);
 	Value.Kind = T::TypeKind;
@@ -692,6 +695,25 @@ FValue* FEmitter::EmitTextureObject(UTexture* Texture, EMaterialSamplerType Samp
 
 FValue* FEmitter::EmitParameter(FName Name, FMaterialParameterMetadata& Metadata, EMaterialSamplerType SamplerType)
 {
+	// Helper local function that register a parameter (info and metadata) to the module, and returns some uint32 ID.
+	auto RegisterParameter = [] (FMaterialIRModule* Module, FMaterialParameterInfo Info, const FMaterialParameterMetadata& Metadata) -> uint32
+	{
+		uint32 Id;
+		if (Internal::Find(Module->ParameterInfoToId, Info, Id))
+		{
+			check(Module->ParameterIdToData[Id].Get<FMaterialParameterMetadata>().Value == Metadata.Value);
+			return Id;
+		}
+
+		Id = Module->ParameterIdToData.Num();
+		Module->ParameterInfoToId.Add(Info, Id);
+		Module->ParameterIdToData.Push({ Info, Metadata });
+
+		return Id;
+	};
+
+	FMaterialParameterInfo Info{ Name };
+
 	switch (Metadata.Value.Type)
 	{
 		case EMaterialParameterType::Scalar:
@@ -702,8 +724,7 @@ FValue* FEmitter::EmitParameter(FName Name, FMaterialParameterMetadata& Metadata
 			}
 
 			MIR::FUniformParameter Prototype = MakePrototype<MIR::FUniformParameter>(FPrimitiveType::GetFloat1());
-			Prototype.Info.Name = Name;
-			Prototype.Metadata = Metadata;
+			Prototype.ParameterIdInModule = RegisterParameter(Module, Info, Metadata);
 			return EmitPrototype(this, Prototype);
 		}
 
@@ -719,8 +740,7 @@ FValue* FEmitter::EmitParameter(FName Name, FMaterialParameterMetadata& Metadata
 			}
 
 			MIR::FUniformParameter Prototype = MakePrototype<MIR::FUniformParameter>(FPrimitiveType::GetFloat4());
-			Prototype.Info.Name = Name;
-			Prototype.Metadata = Metadata;
+			Prototype.ParameterIdInModule = RegisterParameter(Module, Info, Metadata);
 			return EmitPrototype(this, Prototype);
 		}
 
@@ -733,10 +753,7 @@ FValue* FEmitter::EmitParameter(FName Name, FMaterialParameterMetadata& Metadata
 			}
 
 			MIR::FUniformParameter Prototype = MakePrototype<MIR::FUniformParameter>(FObjectType::GetTexture2D());
-			Prototype.Info.Name = Name;
-			Prototype.Info.Index = -1;
-			Prototype.Info.Association = GlobalParameter;
-			Prototype.Metadata = Metadata;
+			Prototype.ParameterIdInModule = RegisterParameter(Module, Info, Metadata);
 			Prototype.SamplerType = SamplerType;
 			return EmitPrototype(this, Prototype);
 		}
diff --git a/Engine/Source/Runtime/Engine/Private/Materials/MaterialIRToHLSLTranslator.cpp b/Engine/Source/Runtime/Engine/Private/Materials/MaterialIRToHLSLTranslator.cpp
--- a/Engine/Source/Runtime/Engine/Private/Materials/MaterialIRToHLSLTranslator.cpp
+++ b/Engine/Source/Runtime/Engine/Private/Materials/MaterialIRToHLSLTranslator.cpp
@@ -402,7 +402,8 @@ struct FTranslator : FMaterialIRToHLSLTranslation
 	{
 		if (UniformParameter->Type->IsObjectOfKind(MIR::ObjectKind_Texture2D))
 		{
-			LowerTextureReference(UniformParameter->Metadata.Value.Texture->GetMaterialType(), UniformParameter->Analysis_UniformParameterIndex);
+		 	EMaterialValueType TextureType = Module->GetParameterMetadata(UniformParameter->ParameterIdInModule).Value.Texture->GetMaterialType();
+			LowerTextureReference(TextureType, UniformParameter->Analysis_UniformParameterIndex);
 		}
 		else
 		{
diff --git a/Engine/Source/Runtime/Engine/Private/Materials/MaterialIRValueAnalyzer.cpp b/Engine/Source/Runtime/Engine/Private/Materials/MaterialIRValueAnalyzer.cpp
--- a/Engine/Source/Runtime/Engine/Private/Materials/MaterialIRValueAnalyzer.cpp
+++ b/Engine/Source/Runtime/Engine/Private/Materials/MaterialIRValueAnalyzer.cpp
@@ -38,14 +38,16 @@ static void AnalyzeTextureObject(FMaterialIRValueAnalyzer& Analyzer, MIR::FTextu
 
 static void AnalyzeTextureUniformParameter(FMaterialIRValueAnalyzer& Analyzer, MIR::FUniformParameter* Parameter)
 {
-	EMaterialTextureParameterType ParamType = MIR::Internal::TextureMaterialValueTypeToParameterType(Parameter->Metadata.Value.Texture->GetMaterialType());
-	
+	UTexture* Texture = Analyzer.Module->GetParameterMetadata(Parameter->ParameterIdInModule).Value.Texture;
+
 	FMaterialTextureParameterInfo ParamInfo {
-		.ParameterInfo = Parameter->Info,
-		.TextureIndex = Analyzer.Material->GetReferencedTextures().Find(Parameter->Metadata.Value.Texture),
+		.ParameterInfo = Analyzer.Module->GetParameterInfo(Parameter->ParameterIdInModule),
+		.TextureIndex = Analyzer.Material->GetReferencedTextures().Find(Texture),
 		.SamplerSource = SSM_FromTextureAsset, // TODO - Is this needed?
 		.VirtualTextureLayerIndex = 0xff,
 	};
+
+	EMaterialTextureParameterType ParamType = MIR::Internal::TextureMaterialValueTypeToParameterType(Texture->GetMaterialType());
 	
 	Parameter->Analysis_UniformParameterIndex = Analyzer.CompilationOutput->UniformExpressionSet.FindOrAddTextureParameter(ParamType, ParamInfo);
 }
@@ -53,12 +55,14 @@ static void AnalyzeTextureUniformParameter(FMaterialIRValueAnalyzer& Analyzer, M
 static void AnalyzePrimitiveUniformParameter(FMaterialIRValueAnalyzer& Analyzer, MIR::FUniformParameter* Parameter)
 {
 	FUniformExpressionSet& UniformExpressionSet = Analyzer.CompilationOutput->UniformExpressionSet;
+	FMaterialParameterInfo ParameterInfo = Analyzer.Module->GetParameterInfo(Parameter->ParameterIdInModule);
+	FMaterialParameterMetadata ParameterMetadata = Analyzer.Module->GetParameterMetadata(Parameter->ParameterIdInModule);
 
 	UE::Shader::FValue DefaultValue;
-	switch (Parameter->Metadata.Value.Type)
+	switch (ParameterMetadata.Value.Type)
 	{
-		case EMaterialParameterType::Scalar: DefaultValue = Parameter->Metadata.Value.AsScalar(); break;
-		case EMaterialParameterType::Vector: DefaultValue = Parameter->Metadata.Value.AsLinearColor(); break;
+		case EMaterialParameterType::Scalar: DefaultValue = ParameterMetadata.Value.AsScalar(); break;
+		case EMaterialParameterType::Vector: DefaultValue = ParameterMetadata.Value.AsLinearColor(); break;
 		default: UE_MIR_TODO();
 	}
 
@@ -71,7 +75,7 @@ static void AnalyzePrimitiveUniformParameter(FMaterialIRValueAnalyzer& Analyzer,
 		Analyzer.UniformDefaultValueOffsets.Add(DefaultValue, DefaultValueOffset);
 	}
 
-	Parameter->Analysis_UniformParameterIndex = UniformExpressionSet.FindOrAddNumericParameter(Parameter->Metadata.Value.Type, Parameter->Info, DefaultValueOffset);
+	Parameter->Analysis_UniformParameterIndex = UniformExpressionSet.FindOrAddNumericParameter(ParameterMetadata.Value.Type, ParameterInfo, DefaultValueOffset);
 
 	// Get the parameter type as primitive
 	MIR::FPrimitiveTypePtr ParameterType = Parameter->Type->AsPrimitive();
@@ -138,7 +142,7 @@ static void AnalyzePrimitiveUniformParameter(FMaterialIRValueAnalyzer& Analyzer,
 		ParamInsight.BufferSlotIndex = GlobalComponentOffset / 4;
 		ParamInsight.BufferSlotOffset = GlobalComponentOffset % 4;
 		ParamInsight.ComponentsCount = NumComponents;
-		ParamInsight.ParameterName = Parameter->Info.Name;
+		ParamInsight.ParameterName = ParameterInfo.Name;
 
 		switch (ParameterType->ScalarKind)
 		{
diff --git a/Engine/Source/Runtime/Engine/Public/Materials/MaterialIR.h b/Engine/Source/Runtime/Engine/Public/Materials/MaterialIR.h
--- a/Engine/Source/Runtime/Engine/Public/Materials/MaterialIR.h
+++ b/Engine/Source/Runtime/Engine/Public/Materials/MaterialIR.h
@@ -49,9 +49,9 @@ enum EValueFlags
 //
 struct FValue
 {
-	EValueKind Kind : 8{};
-	EValueFlags Flags : 8{};
-	FTypePtr  Type{};
+	EValueKind Kind : 8;
+	EValueFlags Flags : 8;
+	FTypePtr  Type;
 
 	// Returns the size in bytes of this value instance.
 	uint32 GetSizeInBytes() const;
@@ -98,9 +98,6 @@ struct FValue
 	// Returns whether this value is arithmetic and approximately one.
 	bool AreAllNearlyOne() const;
 	
-	// Returns this value's textures if it has one (nullptr otherwise).
-	TOptional<TPair<UTexture*, EMaterialSamplerType>> GetTexture();
-
 	// Tries to cast this value to specified type T and returns the casted pointer, if possible (nullptr otherwise).
 	template <typename T>
 	T* As() { return IsA(T::TypeKind) ? static_cast<T*>(this) : nullptr; }
@@ -250,17 +247,20 @@ struct FExternalInput : TValue<VK_ExternalInput>
 struct FTextureObject : TValue<VK_TextureObject>
 {
 	UTexture* Texture;
+
 	EMaterialSamplerType SamplerType;
 
 	/// Index of this parameter in the uniform expression set
 	uint32 Analysis_UniformParameterIndex;
 };
 
-//
+// 
 struct FUniformParameter: TValue<VK_UniformParameter>
 {
-	FMaterialParameterInfo Info;
-	FMaterialParameterMetadata Metadata;
+	// Index of the parameter as registered in the module.
+	uint32 ParameterIdInModule;
+
+	// Eventual sampler type to use to sample this parameter's texture (if it is one)
 	EMaterialSamplerType SamplerType;
 
 	/// Index of this parameter in the uniform expression set
@@ -273,9 +273,9 @@ struct FUniformParameter: TValue<VK_UniformParameter>
 //
 struct FBlock
 {
-	FBlock* Parent{};
-	FInstruction* Instructions{};
-	int32 Level{};
+	FBlock* Parent;
+	FInstruction* Instructions;
+	int32 Level;
 
 	FBlock* FindCommonParentWith(MIR::FBlock* Other);
 };
@@ -283,10 +283,10 @@ struct FBlock
 //
 struct FInstruction : FValue
 {
-	FInstruction* Next{};
-	FBlock* Block{};
-	uint32 NumUsers{};
-	uint32 NumProcessedUsers{};
+	FInstruction* Next;
+	FBlock* Block;
+	uint32 NumUsers;
+	uint32 NumProcessedUsers;
 
 	FBlock* GetDesiredBlockForUse(int32 UseIndex); 
 };
@@ -424,20 +424,20 @@ const TCHAR* OperatorToString(EOperator Op);
 //
 struct FOperator : TInstruction<VK_Operator>
 {
-	EOperator Op = O_Invalid;
-	FValue* AArg{};
-	FValue* BArg{};
-	FValue* CArg{};
+	EOperator Op;
+	FValue* AArg;
+	FValue* BArg;
+	FValue* CArg;
 };
 
 //
 struct FBranch : TInstruction<VK_Branch>
 {
-	FValue* ConditionArg{};
-	FValue* TrueArg{};
-	FValue* FalseArg{};
-	FBlock TrueBlock{};
-	FBlock FalseBlock{};
+	FValue* ConditionArg;
+	FValue* TrueArg;
+	FValue* FalseArg;
+	FBlock TrueBlock;
+	FBlock FalseBlock;
 };
 
 //
@@ -450,7 +450,7 @@ struct FSubscript : TInstruction<VK_Subscript>
 //
 struct FCast : TInstruction<VK_Cast>
 {
-	FValue* Arg{};
+	FValue* Arg;
 };
 
 //
diff --git a/Engine/Source/Runtime/Engine/Public/Materials/MaterialIRModule.h b/Engine/Source/Runtime/Engine/Public/Materials/MaterialIRModule.h
--- a/Engine/Source/Runtime/Engine/Public/Materials/MaterialIRModule.h
+++ b/Engine/Source/Runtime/Engine/Public/Materials/MaterialIRModule.h
@@ -6,7 +6,6 @@
 
 #if WITH_EDITOR
 
-
 /// This class represents the intermediate representation (IR) of a material build.
 /// The `IRModule` includes an IR value graph, produced through expression analysis,
 /// as well as metadata on resource usage and reflection. The IR graph serves as an
@@ -44,6 +43,8 @@ public:
 	TArrayView<const FError> GetErrors() const { return Errors; }
 	const FStatistics& GetStatistics() const { return Statistics; }
 	FStatistics& GetStatistics() { return Statistics; }
+	const FMaterialParameterInfo& GetParameterInfo(uint32 ParameterId) const { return ParameterIdToData[ParameterId].Key; }
+	const FMaterialParameterMetadata& GetParameterMetadata(uint32 ParameterId) const { return ParameterIdToData[ParameterId].Value; }
 
 private:
 	EShaderPlatform ShaderPlatform;
@@ -54,7 +55,9 @@ private:
 	TArray<FError> Errors;
 	MIR::FBlock* RootBlock;
 	FStatistics Statistics;
-	
+	TMap<FMaterialParameterInfo, uint32> ParameterInfoToId;
+	TArray<TPair<FMaterialParameterInfo, FMaterialParameterMetadata>> ParameterIdToData;
+
 	friend MIR::FEmitter;
 	friend FMaterialIRModuleBuilder;
 	friend FMaterialIRModuleBuilderImpl;
diff --git a/Engine/Source/Runtime/Engine/Public/Materials/MaterialIRToHLSLTranslator.h b/Engine/Source/Runtime/Engine/Public/Materials/MaterialIRToHLSLTranslator.h
--- a/Engine/Source/Runtime/Engine/Public/Materials/MaterialIRToHLSLTranslator.h
+++ b/Engine/Source/Runtime/Engine/Public/Materials/MaterialIRToHLSLTranslator.h
@@ -10,16 +10,16 @@
 struct FMaterialIRToHLSLTranslation
 {
 	//
-	const FMaterialIRModule* Module;
+	const FMaterialIRModule* Module{};
 
 	//
-	const FMaterial* Material;
+	const FMaterial* Material{};
 	
 	//
-	const FStaticParameterSet* StaticParameters;
+	const FStaticParameterSet* StaticParameters{};
 	
 	//
-	const ITargetPlatform* TargetPlatform;
+	const ITargetPlatform* TargetPlatform{};
 	
 	//
 	void Run(TMap<FString, FString>& OutParameters, FShaderCompilerEnvironment& OutEnvironment);
