Build fix for CL 40026304.

#rnx

diff --git a/Engine/Source/Runtime/Engine/Private/Materials/MaterialIREmitter.cpp b/Engine/Source/Runtime/Engine/Private/Materials/MaterialIREmitter.cpp
--- a/Engine/Source/Runtime/Engine/Private/Materials/MaterialIREmitter.cpp
+++ b/Engine/Source/Runtime/Engine/Private/Materials/MaterialIREmitter.cpp
@@ -696,19 +696,18 @@ FValue* FEmitter::EmitTextureObject(UTexture* Texture, EMaterialSamplerType Samp
 FValue* FEmitter::EmitParameter(FName Name, FMaterialParameterMetadata& Metadata, EMaterialSamplerType SamplerType)
 {
 	// Helper local function that register a parameter (info and metadata) to the module, and returns some uint32 ID.
-	auto RegisterParameter = [] (FMaterialIRModule* Module, FMaterialParameterInfo Info, const FMaterialParameterMetadata& Metadata) -> uint32
+	auto RegisterParameter = [] (FMaterialIRModule* InModule, FMaterialParameterInfo InInfo, const FMaterialParameterMetadata& InMetadata) -> uint32
 	{
 		uint32 Id;
-		if (Internal::Find(Module->ParameterInfoToId, Info, Id))
+		if (Internal::Find(InModule->ParameterInfoToId, InInfo, Id))
 		{
-			check(Module->ParameterIdToData[Id].Get<FMaterialParameterMetadata>().Value == Metadata.Value);
+			check(InModule->ParameterIdToData[Id].Get<FMaterialParameterMetadata>().Value == InMetadata.Value);
 			return Id;
 		}
 
-		Id = Module->ParameterIdToData.Num();
-		Module->ParameterInfoToId.Add(Info, Id);
-		Module->ParameterIdToData.Push({ Info, Metadata });
-
+		Id = InModule->ParameterIdToData.Num();
+		InModule->ParameterInfoToId.Add(InInfo, Id);
+		InModule->ParameterIdToData.Push({ InInfo, InMetadata });
 		return Id;
 	};
 
