ULevel: Unregister an actor from TEDS when it is unloaded, this removes any actors that are unloaded but haven't been GC'd yet since we do not want to track them in TEDS anymore.

diff --git a/Engine/Source/Runtime/Engine/Private/Level.cpp b/Engine/Source/Runtime/Engine/Private/Level.cpp
--- a/Engine/Source/Runtime/Engine/Private/Level.cpp
+++ b/Engine/Source/Runtime/Engine/Private/Level.cpp
@@ -54,6 +54,8 @@ Level.cpp: Level-related functions
 #include "Kismet2/BlueprintEditorUtils.h"
 #include "PieFixupSerializer.h"
 #include "Editor.h"
+#include "Elements/Common/EditorDataStorageFeatures.h"
+#include "Elements/Interfaces/TypedElementDataStorageCompatibilityInterface.h"
 #include "Subsystems/EditorActorSubsystem.h"
 #include "Widgets/Notifications/SNotificationList.h"
 #include "Settings/LevelEditorMiscSettings.h"
@@ -1188,6 +1190,15 @@ void ULevel::RemoveLoadedActors(const TArray<AActor*>& ActorList, const FTransfo
 
 			OnLoadedActorRemovedFromLevelEvent.Broadcast(*Actor);
 
+			using namespace UE::Editor::DataStorage;
+			// NOTE: we do have to remove the actor from TEDS explicitly here as an actor unloading doesn't always trigger GC immediately,
+			// but for our purposes it is considered dead already
+			if (IEditorDataStorageCompatibilityProvider* TedsCompat =
+					GetMutableDataStorageFeature<IEditorDataStorageCompatibilityProvider>(CompatibilityFeatureName))
+			{
+				TedsCompat->RemoveCompatibleObject(Actor);
+			}
+
 			if (TransformToRemove)
 			{
 				const FTransform TransformToRemoveInverse = TransformToRemove->Inverse();
