Copy the known alpha information in to the VT tiles so that they don't scan the intermediate mips unnecessarily. Fixes CitySample spam.

#rnx

diff --git a/Engine/Source/Runtime/Engine/Private/VT/VirtualTextureDataBuilder.cpp b/Engine/Source/Runtime/Engine/Private/VT/VirtualTextureDataBuilder.cpp
--- a/Engine/Source/Runtime/Engine/Private/VT/VirtualTextureDataBuilder.cpp
+++ b/Engine/Source/Runtime/Engine/Private/VT/VirtualTextureDataBuilder.cpp
@@ -754,6 +754,9 @@ void FVirtualTextureDataBuilder::BuildBlockTiles(uint32 LayerIndex, uint32 Block
 		TBSettings.bOodlePreserveExtremes = BuildSettingsForLayer.bOodlePreserveExtremes;
 		TBSettings.OodleTextureSdkVersion = BuildSettingsForLayer.OodleTextureSdkVersion;
 
+		TBSettings.bKnowAlphaTransparency = BuildSettingsForLayer.bKnowAlphaTransparency;
+		TBSettings.bHasTransparentAlpha = BuildSettingsForLayer.bHasTransparentAlpha;
+
 		check(TBSettings.GetDestGammaSpace() == BuildSettingsForLayer.GetDestGammaSpace());
 
 		// Mip levels start at Block.MipBias in case provided Texture is smaller than block size used in VT
