[LLM] Mark more callstacks to corresponding tags and assets

#tests cooked build
#rnx

diff --git a/Engine/Source/Runtime/CoreUObject/Private/Serialization/AsyncLoading2.cpp b/Engine/Source/Runtime/CoreUObject/Private/Serialization/AsyncLoading2.cpp
--- a/Engine/Source/Runtime/CoreUObject/Private/Serialization/AsyncLoading2.cpp
+++ b/Engine/Source/Runtime/CoreUObject/Private/Serialization/AsyncLoading2.cpp
@@ -8125,6 +8125,10 @@ EEventLoadNodeExecutionResult FAsyncPackage2::Event_DeferredPostLoadExportBundle
 						PackageScope.ThreadContext.CurrentlyPostLoadedObjectByALT = Object;
 						{
 							FScopeCycleCounterUObject ConstructorScope(Object, GET_STATID(STAT_FAsyncPackage_PostLoadObjectsGameThread));
+							LLM_SCOPE(ELLMTag::UObject);
+							LLM_SCOPE_DYNAMIC_STAT_OBJECTPATH(Object->GetPackage(), ELLMTagSet::Assets);
+							LLM_SCOPE_DYNAMIC_STAT_OBJECTPATH(Object->GetClass(), ELLMTagSet::AssetClasses);
+							UE_TRACE_METADATA_SCOPE_ASSET(Object, Object->GetClass());
 							Object->ConditionalPostLoad();
 						}
 						PackageScope.ThreadContext.CurrentlyPostLoadedObjectByALT = nullptr;
diff --git a/Engine/Source/Runtime/Engine/Private/Actor.cpp b/Engine/Source/Runtime/Engine/Private/Actor.cpp
--- a/Engine/Source/Runtime/Engine/Private/Actor.cpp
+++ b/Engine/Source/Runtime/Engine/Private/Actor.cpp
@@ -61,6 +61,8 @@
 #include "UObject/MetaData.h"
 #include "Engine/DamageEvents.h"
 #include "Engine/SCS_Node.h"
+#include "HAL/LowLevelMemStats.h"
+#include "ProfilingDebugging/AssetMetadataTrace.h"
 
 #if UE_WITH_REMOTE_OBJECT_HANDLE
 #include "UObject/RemoteExecutor.h"
@@ -1912,6 +1914,9 @@ void AActor::CallPreReplication(UNetDriver* NetDriver)
 	{
 		return;
 	}
+	LLM_SCOPE_DYNAMIC_STAT_OBJECTPATH(GetPackage(), ELLMTagSet::Assets);
+	LLM_SCOPE_DYNAMIC_STAT_OBJECTPATH(GetClass(), ELLMTagSet::AssetClasses);
+	UE_TRACE_METADATA_SCOPE_ASSET(this, GetClass());
 	
 	const bool bPreReplication = ShouldCallPreReplication();
 	const bool bPreReplicationForReplay = ShouldCallPreReplicationForReplay();
@@ -3078,6 +3083,8 @@ void AActor::RouteEndPlay(const EEndPlayReason::Type EndPlayReason)
 	{
 		if (ActorHasBegunPlay == EActorBeginPlayState::HasBegunPlay)
 		{
+			SCOPE_CYCLE_UOBJECT(EndPlay, this);
+
 			EndPlay(EndPlayReason);
 			ensureMsgf(ActorHasBegunPlay == EActorBeginPlayState::HasNotBegunPlay, TEXT("EndPlay on %s failed. Make sure to call Super::EndPlay() in your override function."), *GetName());
 		}
@@ -3104,7 +3111,10 @@ void AActor::RouteEndPlay(const EEndPlayReason::Type EndPlayReason)
 		}
 	}
 
-	UninitializeComponents();
+	{
+		SCOPE_CYCLE_UOBJECT(UninitializeComponents, this);
+		UninitializeComponents();
+	}
 }
 
 void AActor::EndPlay(const EEndPlayReason::Type EndPlayReason)
@@ -3983,6 +3993,9 @@ static void DispatchOnComponentsCreated(AActor* NewActor)
 	TInlineComponentArray<UActorComponent*> Components;
 	NewActor->GetComponents(Components);
 
+	LLM_SCOPE_DYNAMIC_STAT_OBJECTPATH(NewActor->GetPackage(), ELLMTagSet::Assets);
+	LLM_SCOPE_DYNAMIC_STAT_OBJECTPATH(NewActor->GetClass(), ELLMTagSet::AssetClasses);
+	UE_TRACE_METADATA_SCOPE_ASSET(NewActor, NewActor->GetClass());
 	for (UActorComponent* ActorComp : Components)
 	{
 		if (!ActorComp->HasBeenCreated())
@@ -4266,6 +4279,9 @@ void AActor::FinishSpawning(const FTransform& UserTransform, bool bIsDefaultTran
 
 void AActor::PostActorConstruction()
 {
+	LLM_SCOPE_DYNAMIC_STAT_OBJECTPATH(GetPackage(), ELLMTagSet::Assets);
+	LLM_SCOPE_DYNAMIC_STAT_OBJECTPATH(GetClass(), ELLMTagSet::AssetClasses);
+	UE_TRACE_METADATA_SCOPE_ASSET(this, GetClass());
 	UWorld* const World = GetWorld();
 	bool const bActorsInitialized = World && World->AreActorsInitialized();
 
@@ -4560,6 +4576,11 @@ void AActor::DispatchBeginPlay(bool bFromLevelStreaming)
 
 	if (World)
 	{
+		SCOPE_CYCLE_UOBJECT(BeginPlay, this);
+		LLM_SCOPE_DYNAMIC_STAT_OBJECTPATH(GetPackage(), ELLMTagSet::Assets);
+		LLM_SCOPE_DYNAMIC_STAT_OBJECTPATH(GetClass(), ELLMTagSet::AssetClasses);
+		UE_TRACE_METADATA_SCOPE_ASSET(this, GetClass());
+
 		ensureMsgf(ActorHasBegunPlay == EActorBeginPlayState::HasNotBegunPlay, TEXT("BeginPlay was called on actor %s which was in state %d"), *GetPathName(), (int32)ActorHasBegunPlay);
 		const uint32 CurrentCallDepth = BeginPlayCallDepth++;
 
@@ -5963,6 +5984,9 @@ void AActor::IncrementalPreRegisterComponents()
 
 bool AActor::IncrementalRegisterComponents(int32 NumComponentsToRegister, FRegisterComponentContext* Context)
 {
+	LLM_SCOPE_DYNAMIC_STAT_OBJECTPATH(GetPackage(), ELLMTagSet::Assets);
+	LLM_SCOPE_DYNAMIC_STAT_OBJECTPATH(GetClass(), ELLMTagSet::AssetClasses);
+	UE_TRACE_METADATA_SCOPE_ASSET(this, GetClass());
 	if (NumComponentsToRegister == 0)
 	{
 		// 0 - means register all components
diff --git a/Engine/Source/Runtime/Engine/Private/DataChannel.cpp b/Engine/Source/Runtime/Engine/Private/DataChannel.cpp
--- a/Engine/Source/Runtime/Engine/Private/DataChannel.cpp
+++ b/Engine/Source/Runtime/Engine/Private/DataChannel.cpp
@@ -1908,17 +1908,23 @@ void UControlChannel::ReceivedBunch( FInBunch& Bunch )
 			FString MessageStr;
 			if (FNetControlMessage<NMT_GameSpecific>::Receive(Bunch, MessageByte, MessageStr))
 			{
+				UGameInstance* GameInstance = nullptr;
 				if (Connection->Driver->World != NULL && Connection->Driver->World->GetGameInstance() != NULL)
 				{
-					Connection->Driver->World->GetGameInstance()->HandleGameNetControlMessage(Connection, MessageByte, MessageStr);
+					GameInstance = Connection->Driver->World->GetGameInstance();
 				}
 				else
 				{
 					FWorldContext* Context = GEngine->GetWorldContextFromPendingNetGameNetDriver(Connection->Driver);
-					if (Context != NULL && Context->OwningGameInstance != NULL)
-					{
-						Context->OwningGameInstance->HandleGameNetControlMessage(Connection, MessageByte, MessageStr);
-					}
+					GameInstance = Context ? Context->OwningGameInstance.Get() : nullptr;
+				}
+				if (GameInstance)
+				{
+					LLM_SCOPE_BYNAME(TEXT("Networking/GameControlMessage"));
+					LLM_SCOPE_DYNAMIC_STAT_OBJECTPATH(GameInstance->GetPackage(), ELLMTagSet::Assets);
+					LLM_SCOPE_DYNAMIC_STAT_OBJECTPATH(GameInstance->GetClass(), ELLMTagSet::AssetClasses);
+					UE_TRACE_METADATA_SCOPE_ASSET(GameInstance, GameInstance->GetClass());
+					GameInstance->HandleGameNetControlMessage(Connection, MessageByte, MessageStr);
 				}
 			}
 		}
diff --git a/Engine/Source/Runtime/Engine/Private/NetDriver.cpp b/Engine/Source/Runtime/Engine/Private/NetDriver.cpp
--- a/Engine/Source/Runtime/Engine/Private/NetDriver.cpp
+++ b/Engine/Source/Runtime/Engine/Private/NetDriver.cpp
@@ -799,7 +799,7 @@ void UNetDriver::AssertValid()
 
 void UNetDriver::AddNetworkActor(AActor* Actor)
 {
-	LLM_SCOPE(ELLMTag::Networking);
+	LLM_SCOPE_BYTAG(NetDriver);
 	ensureMsgf(Actor == nullptr || !(Actor->HasAnyFlags(RF_ClassDefaultObject | RF_ArchetypeObject)), TEXT("%s is a CDO or Archetype and should not be replicated."), *GetFullNameSafe(Actor));
 
 	UE_LOG(LogNetLifecycle, Log, TEXT("Registering for replication: Actor: %s"), *GetFullNameSafe(Actor));
@@ -4108,6 +4108,7 @@ void UNetDriver::NotifySubObjectDestroyed(UObject* SubObject)
 
 void UNetDriver::RemoveNetworkActor(AActor* Actor)
 {
+	LLM_SCOPE_BYTAG(NetDriver);
 	UE_LOG(LogNetLifecycle, Log, TEXT("Deregistering for replication: Actor: %s"), *GetFullNameSafe(Actor));
 
 	// This is a hack: Sometimes world streaming will stream out levels even though some players can still see them
diff --git a/Engine/Source/Runtime/Engine/Private/PSOPrecacheMaterial.cpp b/Engine/Source/Runtime/Engine/Private/PSOPrecacheMaterial.cpp
--- a/Engine/Source/Runtime/Engine/Private/PSOPrecacheMaterial.cpp
+++ b/Engine/Source/Runtime/Engine/Private/PSOPrecacheMaterial.cpp
@@ -16,6 +16,8 @@
 #include "SceneInterface.h"
 #include "ShaderCodeArchive.h"
 #include "ODSC/ODSCManager.h"
+#include "HAL/LowLevelMemTracker.h"
+#include "ProfilingDebugging/MetadataTrace.h"
 
 int32 GPSOUseBackgroundThreadForCollection = 1;
 static FAutoConsoleVariableRef CVarPSOUseBackgroundThreadForCollection(
@@ -185,6 +187,10 @@ public:
 
 	FMaterialPSOPrecacheRequestID PrecachePSOs(const FMaterialPSOPrecacheParams& Params, EPSOPrecachePriority Priority, FGraphEventArray& OutGraphEvents)
 	{
+		LLM_SCOPE(ELLMTag::PSO);
+		LLM_TAGSET_SCOPE_CLEAR(ELLMTagSet::Assets);
+		LLM_TAGSET_SCOPE_CLEAR(ELLMTagSet::AssetClasses);
+		UE_TRACE_METADATA_CLEAR_SCOPE();
 		FMaterialPSOPrecacheRequestID RequestID = INDEX_NONE;		
 
 		if (GetPSOPrecacheMode() == EPSOPrecacheMode::PreloadShader)
@@ -350,6 +356,10 @@ public:
 
 	void MarkCollectionComplete(const FMaterialPSOPrecacheParams& Params, const FPSOPrecacheDataArray& PrecacheData, const FPSOPrecacheRequestResultArray& PrecacheRequestResults, uint32 RequestLifecycleID)
 	{
+		LLM_SCOPE(ELLMTag::PSO);
+		LLM_TAGSET_SCOPE_CLEAR(ELLMTagSet::Assets);
+		LLM_TAGSET_SCOPE_CLEAR(ELLMTagSet::AssetClasses);
+		UE_TRACE_METADATA_CLEAR_SCOPE();
 		FRWScopeLock WriteLock(RWLock, SLT_Write);
 		
 		// Ignore requests not coming from current life cycle ID
@@ -384,6 +394,10 @@ public:
 
 	void ReleasePrecacheData(FMaterialPSOPrecacheRequestID MaterialPSORequestID)
 	{
+		LLM_SCOPE(ELLMTag::PSO);
+		LLM_TAGSET_SCOPE_CLEAR(ELLMTagSet::Assets);
+		LLM_TAGSET_SCOPE_CLEAR(ELLMTagSet::AssetClasses);
+		UE_TRACE_METADATA_CLEAR_SCOPE();
 		check(MaterialPSORequestID != INDEX_NONE);
 
 		FRWScopeLock WriteLock(RWLock, SLT_Write);
@@ -396,6 +410,10 @@ public:
 
 	void BoostPriority(EPSOPrecachePriority NewPri, FMaterialPSOPrecacheRequestID MaterialPSORequestID)
 	{
+		LLM_SCOPE(ELLMTag::PSO);
+		LLM_TAGSET_SCOPE_CLEAR(ELLMTagSet::Assets);
+		LLM_TAGSET_SCOPE_CLEAR(ELLMTagSet::AssetClasses);
+		UE_TRACE_METADATA_CLEAR_SCOPE();
 		check(MaterialPSORequestID != INDEX_NONE);
 
 		{
@@ -427,6 +445,10 @@ public:
 
 	void ClearMaterialPSORequests()
 	{
+		LLM_SCOPE(ELLMTag::PSO);
+		LLM_TAGSET_SCOPE_CLEAR(ELLMTagSet::Assets);
+		LLM_TAGSET_SCOPE_CLEAR(ELLMTagSet::AssetClasses);
+		UE_TRACE_METADATA_CLEAR_SCOPE();
 		check(IsInGameThread());
 
 		FRWScopeLock WriteLock(RWLock, SLT_Write);
@@ -561,6 +583,10 @@ FMaterialPSORequestManager GMaterialPSORequestManager;
 
 void FMaterialPSOPrecacheCollectionTask::DoTask(ENamedThreads::Type CurrentThread, const FGraphEventRef& MyCompletionGraphEvent)
 {
+	LLM_SCOPE(ELLMTag::PSO);
+	LLM_TAGSET_SCOPE_CLEAR(ELLMTagSet::Assets);
+	LLM_TAGSET_SCOPE_CLEAR(ELLMTagSet::AssetClasses);
+	UE_TRACE_METADATA_CLEAR_SCOPE();
 	TRACE_CPUPROFILER_EVENT_SCOPE(FMaterialPSOPrecacheCollectionTask);
 
 #if WITH_ODSC
@@ -608,6 +634,10 @@ void FMaterialPSOPrecacheCollectionTask::DoTask(ENamedThreads::Type CurrentThrea
 
 void FShaderMapPreloadTask::DoTask(ENamedThreads::Type CurrentThread, const FGraphEventRef& MyCompletionGraphEvent)
 {
+	LLM_SCOPE(ELLMTag::PSO);
+	LLM_TAGSET_SCOPE_CLEAR(ELLMTagSet::Assets);
+	LLM_TAGSET_SCOPE_CLEAR(ELLMTagSet::AssetClasses);
+	UE_TRACE_METADATA_CLEAR_SCOPE();
 	TRACE_CPUPROFILER_EVENT_SCOPE(FShaderMapPreloadTask);
 
 	FTaskTagScope ParallelGTScope(ETaskTag::EParallelGameThread);
@@ -627,6 +657,10 @@ void FShaderMapPreloadTask::DoTask(ENamedThreads::Type CurrentThread, const FGra
 
 void FShaderPreloadCollectionTask::DoTask(ENamedThreads::Type CurrentThread, const FGraphEventRef& MyCompletionGraphEvent)
 {
+	LLM_SCOPE(ELLMTag::PSO);
+	LLM_TAGSET_SCOPE_CLEAR(ELLMTagSet::Assets);
+	LLM_TAGSET_SCOPE_CLEAR(ELLMTagSet::AssetClasses);
+	UE_TRACE_METADATA_CLEAR_SCOPE();
 	TRACE_CPUPROFILER_EVENT_SCOPE(FShaderPreloadCollectionTask);
 
 	FTaskTagScope ParallelGTScope(ETaskTag::EParallelGameThread);
diff --git a/Engine/Source/Runtime/Engine/Private/World.cpp b/Engine/Source/Runtime/Engine/Private/World.cpp
--- a/Engine/Source/Runtime/Engine/Private/World.cpp
+++ b/Engine/Source/Runtime/Engine/Private/World.cpp
@@ -139,6 +139,8 @@
 #endif // UE_WITH_IRIS
 
 #include "ChaosSolversModule.h"
+#include "HAL/LowLevelMemStats.h"
+#include "ProfilingDebugging/AssetMetadataTrace.h"
 
 #include UE_INLINE_GENERATED_CPP_BY_NAME(World)
 
@@ -2653,6 +2655,9 @@ void UWorld::RemoveActor(AActor* Actor, bool bShouldModifyLevel) const
 
 			CheckLevel->Actors[ActorListIndex] = nullptr;
 
+			LLM_SCOPE_DYNAMIC_STAT_OBJECTPATH(GetPackage(), ELLMTagSet::Assets);
+			LLM_SCOPE_DYNAMIC_STAT_OBJECTPATH(GetClass(), ELLMTagSet::AssetClasses);
+			UE_TRACE_METADATA_SCOPE_ASSET(this, GetClass());
 			CheckLevel->ActorsForGC.RemoveSwap(Actor);
 		}
 	}
diff --git a/Engine/Source/Runtime/Experimental/Iris/Core/IrisCore.Build.cs b/Engine/Source/Runtime/Experimental/Iris/Core/IrisCore.Build.cs
--- a/Engine/Source/Runtime/Experimental/Iris/Core/IrisCore.Build.cs
+++ b/Engine/Source/Runtime/Experimental/Iris/Core/IrisCore.Build.cs
@@ -4,7 +4,7 @@ namespace UnrealBuildTool.Rules
 {
 	public class IrisCore : ModuleRules
 	{
-        public IrisCore(ReadOnlyTargetRules Target) : base(Target)
+		public IrisCore(ReadOnlyTargetRules Target) : base(Target)
 		{
 			PublicDependencyModuleNames.AddRange(
 				new string[]
@@ -12,8 +12,9 @@ namespace UnrealBuildTool.Rules
 					"Core",
 					"CoreUObject",
 					"NetCore",
-                }
-            );
+					"TraceLog"
+				}
+			);
 
 			if (Target.bUseIris == true)
 			{
diff --git a/Engine/Source/Runtime/Experimental/Iris/Core/Private/Iris/ReplicationSystem/DequantizeAndApplyHelper.cpp b/Engine/Source/Runtime/Experimental/Iris/Core/Private/Iris/ReplicationSystem/DequantizeAndApplyHelper.cpp
--- a/Engine/Source/Runtime/Experimental/Iris/Core/Private/Iris/ReplicationSystem/DequantizeAndApplyHelper.cpp
+++ b/Engine/Source/Runtime/Experimental/Iris/Core/Private/Iris/ReplicationSystem/DequantizeAndApplyHelper.cpp
@@ -3,6 +3,7 @@
 #include "DequantizeAndApplyHelper.h"
 #include "Containers/ArrayView.h"
 #include "HAL/IConsoleManager.h"
+#include "HAL/LowLevelMemStats.h"
 #include "Iris/Core/IrisProfiler.h"
 #include "Net/Core/NetBitArray.h"
 #include "Iris/ReplicationState/ReplicationStateDescriptor.h"
@@ -14,6 +15,8 @@
 #include "Iris/Serialization/NetSerializer.h"
 #include "Misc/MemStack.h"
 #include "ProfilingDebugging/CsvProfiler.h"
+#include "ProfilingDebugging/AssetMetadataTrace.h"
+#include "UObject/Package.h"
 
 namespace UE::Net::Private
 {
@@ -51,6 +54,10 @@ static void CallLegacyFunctionForEachOwner(const FDequantizeAndApplyHelper::FCon
 	{
 		for (UObject* Object : MakeArrayView(Context.OwnerCollector->GetOwners(), Context.OwnerCollector->GetOwnerCount()))
 		{
+			LLM_SCOPE_BYNAME(TEXT("UObject/NetworkReceive"));
+			LLM_SCOPE_DYNAMIC_STAT_OBJECTPATH(Object->GetPackage(), ELLMTagSet::Assets);
+			LLM_SCOPE_DYNAMIC_STAT_OBJECTPATH(Object->GetClass(), ELLMTagSet::AssetClasses);
+			UE_TRACE_METADATA_SCOPE_ASSET(Object, Object->GetClass());
 			Functor(Object);
 		}
 	}
