MultiServer Borrowing Changes:
- Only register the Ticks on the owning Server.
- Deal with ShapeComponent mutating the CDO.
- Make sure Destroying a remote Actor/PrimitiveComponent (such as during Migrations) never results in ClearComponentOverlaps callbacks.

#rnx

diff --git a/Engine/Source/Runtime/Engine/Private/Actor.cpp b/Engine/Source/Runtime/Engine/Private/Actor.cpp
--- a/Engine/Source/Runtime/Engine/Private/Actor.cpp
+++ b/Engine/Source/Runtime/Engine/Private/Actor.cpp
@@ -1237,10 +1237,11 @@ void AActor::PostMigrate(const FPostMigrateParams& PostMigrateParams)
 	Super::PostMigrate(PostMigrateParams);
 
 	// We are using this as a proxy for having ownership until we sort out those mechanics
-	const bool bHasOwnership = bTickFunctionsRegistered;
-	FRemoteServerId LocalServerId = UE::RemoteObject::GetGlobalServerId();
+	const FRemoteServerId LocalServerId = UE::RemoteObject::GetGlobalServerId();
+	const bool bHasOwnership = UE::RemoteObject::Handle::IsOwned(this);
+	const bool bRequestOwnershipChange = EnumHasAnyFlags(PostMigrateParams.MigrationFlags, EObjectMigrationFlags::ChangeOwnership);
 
-	UE_LOG(LogActor, Verbose, TEXT("%hs: %s %d (LocalSrvId: %s. ReqSrvId: %s. DstSrvId: %s. HasOwnership: %d)"), __func__, *GetName(), static_cast<int>(PostMigrateParams.MigrationFlags), *LocalServerId.ToString(), *PostMigrateParams.RequestingServerId.ToString(), *PostMigrateParams.DestinationServerId.ToString(), bHasOwnership);
+	UE_LOG(LogActor, Verbose, TEXT("%hs: %s %d (LocalSrvId: %s. ReqSrvId: %s. DstSrvId: %s. HasOwnership: %d. ChangeOwnership: %d)"), __func__, *GetName(), static_cast<int>(PostMigrateParams.MigrationFlags), *LocalServerId.ToString(), *PostMigrateParams.RequestingServerId.ToString(), *PostMigrateParams.DestinationServerId.ToString(), bHasOwnership, bRequestOwnershipChange);
 
 	UWorld* const World = GetWorld();
 	checkf(World, TEXT("Migrated Actor %s (%s) doesn't have a World"), *GetPathName(), *FRemoteObjectId(this).ToString());
@@ -1250,14 +1251,13 @@ void AActor::PostMigrate(const FPostMigrateParams& PostMigrateParams)
 	{
 		ensureMsgf(PostMigrateParams.DestinationServerId != LocalServerId, TEXT("%hs: Sending Object %s from %s to %s but destination is the same as the Local Server Id."), __func__, *GetName(), *PostMigrateParams.RequestingServerId.ToString(), *PostMigrateParams.DestinationServerId.ToString());
 
-		const bool bRelinquishOwnership = !bHasOwnership || EnumHasAnyFlags(PostMigrateParams.MigrationFlags, EObjectMigrationFlags::ChangeOwnership);
+		const bool bRelinquishOwnership = !bHasOwnership || bRequestOwnershipChange;
 		if (bRelinquishOwnership)
 		{
 			// Remove the actor from the world on this server as it is about to be migrated to another server.
 			// Migrated actors are unique to the server instance that currently owns their memory and are not replicated to other servers.
 			if (World)
 			{
-				ClearComponentOverlaps();
 				UE::Net::Private::FNetPropertyConditionManager::Get().NotifyObjectDestroyed(this);
 
 				// Notify net drivers that this actor has been destroyed.
@@ -1288,48 +1288,45 @@ void AActor::PostMigrate(const FPostMigrateParams& PostMigrateParams)
 
 			// Unregister the actor's tick function
 			RegisterAllActorTickFunctions(/*bRegisterTickFunctions*/ false, /*bIncludeComponents*/ true);
-
-			// Mark the actor and its direct components as pending kill.
-			MarkComponentsAsGarbage();
-			MarkAsGarbage();
 		}
 	}
 	else if (EnumHasAnyFlags(PostMigrateParams.MigrationFlags, EObjectMigrationFlags::Receive))
 	{
-		ensureMsgf(PostMigrateParams.DestinationServerId == UE::RemoteObject::GetGlobalServerId(), TEXT("Unexpected DestinationServerId (%s was not the local ServerId: %s)"), *PostMigrateParams.DestinationServerId.ToString(), *LocalServerId.ToString());
+		ensureMsgf(PostMigrateParams.DestinationServerId == LocalServerId, TEXT("Unexpected DestinationServerId (%s was not the local ServerId: %s)"), *PostMigrateParams.DestinationServerId.ToString(), *LocalServerId.ToString());
 
 		// Add the actor that was migrated from another server to this server's world
-		GetLevel()->TryAddActorToList(this, /*bAddUnique*/false);
+		constexpr bool bOnlyAddIfUnique = true;
+		GetLevel()->TryAddActorToList(this, bOnlyAddIfUnique);
 
 		// Set the actor's world transform if it has a native rootcomponent.
 		const USceneComponent* SceneRootComponent = FixupNativeActorComponents(this);
 		ensureMsgf(SceneRootComponent == RootComponent, TEXT("%hs: We were expecting RootComponent (%s) to be equal to FixupNativeActorComponents (%s). Transform fix-up code would be required to handle this case."), __func__, *GetNameSafe(RootComponent), *GetNameSafe(SceneRootComponent));
 
-		ResetOwnedComponents();
-		RegisterAllComponents();
-
-		bHasFinishedSpawning = true;
-
-		PreInitializeComponents();
-		InitializeComponents();
-		PostInitializeComponents();
+		// If we don't have Ownership, these components were never registered
+		if (!bHasOwnership)
+		{
+			ResetOwnedComponents();
+			RegisterAllComponents();
 
-		BuildReplicatedComponentsInfo();
+			PreInitializeComponents();
+			InitializeComponents();
+			PostInitializeComponents();
+		}
 
-#if UE_WITH_IRIS
-		BeginReplication();
-#endif // UE_WITH_IRIS
+		bHasFinishedSpawning = true;
 
-		const bool bAcquiredOwnership = !bHasOwnership && EnumHasAnyFlags(PostMigrateParams.MigrationFlags, EObjectMigrationFlags::ChangeOwnership);
+		// If we're being passed ownership, register us with the global systems
+		const bool bAcquiredOwnership = !bHasOwnership && bRequestOwnershipChange;
 		if (bAcquiredOwnership)
 		{
-			RegisterAllActorTickFunctions(true, true);
-		}
+			BuildReplicatedComponentsInfo();
 
-		ActorHasBegunPlay = EActorBeginPlayState::HasBegunPlay;
+			#if UE_WITH_IRIS
+			BeginReplication();
+			#endif // UE_WITH_IRIS
+
+			RegisterAllActorTickFunctions(true, true);
 
-		if (bAcquiredOwnership)
-		{
 			// We are going to assume that the Physics queries on the previous owning server had taken multi-server physics into account
 			// thus we do not actually need to check for Initial Overlaps anymore
 			// ForEachSceneComponent: SceneComponent->ClearSkipUpdateOverlaps();
@@ -1337,6 +1334,8 @@ void AActor::PostMigrate(const FPostMigrateParams& PostMigrateParams)
 
 			World->AddNetworkActor(this);
 		}
+
+		ActorHasBegunPlay = EActorBeginPlayState::HasBegunPlay;
 	}
 }
 #endif // UE_WITH_REMOTE_OBJECT_HANDLE
diff --git a/Engine/Source/Runtime/Engine/Private/Components/PrimitiveComponent.cpp b/Engine/Source/Runtime/Engine/Private/Components/PrimitiveComponent.cpp
--- a/Engine/Source/Runtime/Engine/Private/Components/PrimitiveComponent.cpp
+++ b/Engine/Source/Runtime/Engine/Private/Components/PrimitiveComponent.cpp
@@ -1785,8 +1785,16 @@ void UPrimitiveComponent::OnComponentDestroyed(bool bDestroyingHierarchy)
 	// Prevent future overlap events. Any later calls to UpdateOverlaps will only allow this to end overlaps.
 	SetGenerateOverlapEvents(false);
 
+#if UE_WITH_REMOTE_OBJECT_HANDLE
+	// We can be destroying an object set to be Remote if we're evicting it (e.g. migrating or GC), or reconstructing a received object in-place (ConstructRemoteObject)
+	// In that case, the we don't want to call any End Overlaps because it's still overlapping on another server.
+	const bool bClearComponentOverlaps = (OverlappingComponents.Num() > 0) && !HasAnyInternalFlags(EInternalObjectFlags::Remote);
+#else
+	const bool bClearComponentOverlaps = (OverlappingComponents.Num() > 0);
+#endif
+
 	// End all current overlaps
-	if (OverlappingComponents.Num() > 0)
+	if (bClearComponentOverlaps)
 	{
 		const bool bDoNotifies = true;
 		const bool bSkipNotifySelf = false;
diff --git a/Engine/Source/Runtime/Engine/Private/Components/ShapeComponent.cpp b/Engine/Source/Runtime/Engine/Private/Components/ShapeComponent.cpp
--- a/Engine/Source/Runtime/Engine/Private/Components/ShapeComponent.cpp
+++ b/Engine/Source/Runtime/Engine/Private/Components/ShapeComponent.cpp
@@ -107,6 +107,20 @@ void UShapeComponent::Serialize(FArchive& Ar)
 {
 	Super::Serialize(Ar);
 
+#if UE_WITH_REMOTE_OBJECT_HANDLE
+	if (Ar.IsMigratingRemoteObjects() && Ar.IsLoading())
+	{
+		// ShapeBodySetup exists outside of this hierarchy (it is a runtime instanced SubObject on the CDO)
+		// So null it out here so we don't end up with that runtime CDO SubObject from another server.
+		// If we migrated, it would show up as a RemoteObject and crash since its Outer is the CDO.
+		// Note: We should ban CDO mutation on runtime objects. This is never safe due to delta serialization.
+		if (ShapeBodySetup.IsRemote())
+		{
+			ShapeBodySetup = nullptr;
+		}
+	}
+#endif
+
 #if WITH_EDITOR
 	Ar.UsingCustomVersion(FShapeComponentCustomVersion::GUID);
 
