Utility functions and console commands that print UObject state information:
- PrintObjectsInOuter.
- PrintObjectsWithName.
- PrintObjectProperties.
- PrintObjectArchetype.
- PrintObjectIDO.

[RN] Utility functions and console commands that print UObject state information.

diff --git a/Engine/Source/Runtime/Engine/Private/PrintObjectUtils.cpp b/Engine/Source/Runtime/Engine/Private/PrintObjectUtils.cpp
new file mode 100644
--- /dev/null
+++ b/Engine/Source/Runtime/Engine/Private/PrintObjectUtils.cpp
@@ -0,0 +1,455 @@
+// Copyright Epic Games, Inc. All Rights Reserved.
+
+#include "PrintObjectUtils.h"
+
+#if !UE_BUILD_SHIPPING
+
+#include "Containers/Set.h"
+#include "HAL/IConsoleManager.h"
+#include "UObject/Object.h"
+#include "UObject/Class.h"
+#include "UObject/UnrealType.h"
+#include "UObject/UObjectHash.h"
+
+#if WITH_EDITORONLY_DATA
+#include "UObject/InstanceDataObjectUtils.h"
+#include "UObject/PropertyBagRepository.h"
+#endif
+
+namespace UE
+{
+
+void PrintObjectsInOuter(UObject* Object, FOutputDevice* InOutputDevice)
+{
+	FOutputDevice* OutputDevice = InOutputDevice ? InOutputDevice : GLog;
+	if (OutputDevice)
+	{
+		if (Object)
+		{
+			TArray<UObject*> ChildObjects;
+			GetObjectsWithOuter(Object, ChildObjects, /*bIncludeNestedObjects*/true);
+
+			OutputDevice->Logf(TEXT("Printing %d object(s) under: '%s' [%s] (%s)"),
+				ChildObjects.Num(), *Object->GetName(), *Object->GetClass()->GetName(), *Object->GetPathName());
+
+			for (UObject* ChildObject : ChildObjects)
+			{
+				FString ChildClassName = ChildObject->GetClass() ? *ChildObject->GetClass()->GetName() : TEXT("NULL Class");
+
+				OutputDevice->Logf(TEXT("%s [%s] (%s)"), *ChildObject->GetName(), *ChildClassName, *ChildObject->GetPathName());
+			}
+		}
+		else
+		{
+			OutputDevice->Logf(TEXT("PrintObjectsInOuter: NULL object"));
+		}
+	}
+}
+
+void PrintObjectsWithName(const TCHAR* ObjectName, FOutputDevice* InOutputDevice)
+{
+	FOutputDevice* OutputDevice = InOutputDevice ? InOutputDevice : GLog;
+	if (OutputDevice)
+	{
+		if (ObjectName)
+		{
+			TArray<UObject*> Objects;
+			bool bDidFindAny = StaticFindAllObjectsSafe(Objects, UObject::StaticClass(), ObjectName);
+
+			if (bDidFindAny)
+			{
+				OutputDevice->Logf(TEXT("Printing %d object(s) with name: '%s'"), Objects.Num(), ObjectName);
+
+				for (UObject* Object : Objects)
+				{
+					FString ClassName = Object->GetClass() ? *Object->GetClass()->GetName() : TEXT("NULL Class");
+					OutputDevice->Logf(TEXT("%s [%s] (%s)"), *Object->GetName(), *ClassName, *Object->GetPathName());
+				}
+			}
+			else
+			{
+				OutputDevice->Logf(TEXT("PrintObjectsWithName: failed to find any objects with name: '%s'"), ObjectName);
+			}
+		}
+		else
+		{
+			OutputDevice->Logf(TEXT("PrintObjectsWithName: NULL object name"));
+		}
+	}
+}
+
+void PrintStructProperties(UStruct* Struct, void* StructData, bool bIncludeinitializationState, FOutputDevice* InOutputDevice)
+{
+	FOutputDevice* OutputDevice = InOutputDevice ? InOutputDevice : GLog;
+	if (!OutputDevice)
+	{
+		return;
+	}
+
+	if (!Struct)
+	{
+		OutputDevice->Logf(TEXT("PrintStructProperties: NULL struct"));
+	}
+	if (!StructData)
+	{
+		OutputDevice->Logf(TEXT("PrintStructProperties: NULL StructData"));
+	}
+
+	if (Struct && StructData)
+	{
+		OutputDevice->Logf(TEXT("Printing properties for struct of type: '%s'"), *Struct->GetName());
+
+		// Keep track of visited property-owner pairs to avoid referencing cycles.
+		TSet<TTuple<const FProperty*, void*>> VisitedPropOwners;
+
+		Struct->Visit(StructData, [&VisitedPropOwners, bIncludeinitializationState, OutputDevice](const FPropertyVisitorPath& PropertyPath, const FPropertyVisitorData& Data)->EPropertyVisitorControlFlow
+		{
+			const FProperty* Property = PropertyPath.Top().Property;
+			void* Owner = Data.ParentStructData;
+			TTuple<const FProperty*, void*> PropOwner(Property, Owner);
+
+			if (!Property || VisitedPropOwners.Contains(PropOwner))
+			{
+				return EPropertyVisitorControlFlow::StepOver;
+			}
+
+			const UStruct* OwnerType = PropertyPath.Top().ParentStructType;
+
+			UObject* OwnerObject = nullptr;
+
+			if (OwnerType && OwnerType->IsChildOf<UObject>())
+			{
+				OwnerObject = static_cast<UObject*>(Owner);
+			}
+
+			void* PropData = Property->ContainerPtrToValuePtr<void>(Owner);
+			FString PropValueAsString;
+			Property->ExportText_Direct(PropValueAsString, PropData, PropData, OwnerObject, PPF_None);
+
+			FString OwnerPath = OwnerObject ? OwnerObject->GetPathName() : FString();
+
+			FString InitializationStateText;
+
+			if (bIncludeinitializationState)
+			{
+				#if WITH_EDITORONLY_DATA
+				UE::FInitializedPropertyValueState InitializedState(OwnerType, Owner);
+
+				if (InitializedState.IsTracking())
+				{
+					bool bIsInitialized = !Property->HasAnyPropertyFlags(CPF_RequiredParm) || InitializedState.IsSet(Property);
+
+					if (bIsInitialized)
+					{
+						InitializationStateText = TEXT(", (initialized)");
+					}
+					else
+					{
+						InitializationStateText = TEXT(", (uninitialized)");
+					}
+				}
+				#endif
+			}
+
+			OutputDevice->Logf(TEXT("%s.%s: [%s] %s%s"), *OwnerPath, *Property->GetName(), *Property->GetClass()->GetName(), *PropValueAsString, *InitializationStateText);
+
+			VisitedPropOwners.Add(PropOwner);
+
+			return EPropertyVisitorControlFlow::StepOver;
+		});
+	}
+}
+
+void PrintObjectProperties(UObject* Object, bool bIncludeinitializationState, FOutputDevice* InOutputDevice)
+{
+	FOutputDevice* OutputDevice = InOutputDevice ? InOutputDevice : GLog;
+	if (!OutputDevice)
+	{
+		return;
+	}
+
+	if (Object)
+	{
+		FString ObjectPath = Object->GetPathName();
+
+		OutputDevice->Logf(TEXT("Printing properties of object '%s' [%s] (%s)"), *Object->GetName(), *Object->GetClass()->GetName(), *ObjectPath);
+
+		// Keep track of visited property-owner pairs to avoid referencing cycles.
+		TSet<TTuple<const FProperty*, void*>> VisitedPropOwners;
+
+		Object->GetClass()->Visit(Object, [&VisitedPropOwners, Object, &ObjectPath, bIncludeinitializationState, OutputDevice](const FPropertyVisitorPath& PropertyPath, const FPropertyVisitorData& Data)->EPropertyVisitorControlFlow
+		{
+			const FProperty* Property = PropertyPath.Top().Property;
+			void* Owner = Data.ParentStructData;
+			TTuple<const FProperty*, void*> PropOwner(Property, Owner);
+
+			if (!Property || VisitedPropOwners.Contains(PropOwner))
+			{
+				return EPropertyVisitorControlFlow::StepOver;
+			}
+
+			const UStruct* OwnerType = PropertyPath.Top().ParentStructType;
+
+			bool bIsInRootObject = true;
+
+			if (OwnerType && OwnerType->IsChildOf<UObject>())
+			{
+				UObject* OwnerObject = static_cast<UObject*>(Owner);
+
+				if (OwnerObject)
+				{
+					void* PropData = Property->ContainerPtrToValuePtr<void>(OwnerObject);
+					FString PropValueAsString;
+					Property->ExportText_Direct(PropValueAsString, PropData, PropData, OwnerObject, PPF_None);
+
+					FString OwnerPath = OwnerObject->GetPathName();
+					FString OwnerRelPath = OwnerPath.Contains(ObjectPath) ? OwnerPath.RightChop(ObjectPath.Len()) : OwnerPath;
+
+					FString InitializationStateText;
+
+					if (bIncludeinitializationState)
+					{
+						#if WITH_EDITORONLY_DATA
+						UE::FInitializedPropertyValueState InitializedState(OwnerObject);
+
+						if (InitializedState.IsTracking())
+						{
+							bool bIsInitialized = !Property->HasAnyPropertyFlags(CPF_RequiredParm) || InitializedState.IsSet(Property);
+
+							if (bIsInitialized)
+							{
+								InitializationStateText = TEXT(", (initialized)");
+							}
+							else
+							{
+								InitializationStateText = TEXT(", (uninitialized)");
+							}
+						}
+						#endif
+					}
+
+					OutputDevice->Logf(TEXT("%s.%s: [%s] %s%s"), *OwnerRelPath, *Property->GetName(), *Property->GetClass()->GetName(), *PropValueAsString, *InitializationStateText);
+
+					bIsInRootObject = OwnerObject->IsInOuter(Object);
+				}
+			}
+
+			VisitedPropOwners.Add(PropOwner);
+
+			if (bIsInRootObject)
+			{
+				return EPropertyVisitorControlFlow::StepInto;
+			}
+			else
+			{
+				// Don't step into external object references.
+				return EPropertyVisitorControlFlow::StepOver;
+			}
+		});
+	}
+	else
+	{
+		OutputDevice->Logf(TEXT("PrintObjectProperties: NULL object"));
+	}
+}
+
+void PrintObjectArchetype(UObject* Object, FOutputDevice* InOutputDevice)
+{
+	FOutputDevice* OutputDevice = InOutputDevice ? InOutputDevice : GLog;
+	if (OutputDevice)
+	{
+		if (Object)
+		{
+			FString ObjectPath = Object->GetPathName();
+
+			UObject* Archetype = Object->GetArchetype();
+			FString ArchetypeName = Archetype ? Archetype->GetName() : TEXT("<NULL>");
+			const UClass* ArchetypeClass = Archetype ? Archetype->GetClass() : nullptr;
+			FString ArchetypeClassName = ArchetypeClass ? ArchetypeClass->GetName() : FString();
+			FString ArchetypePath = Archetype ? Archetype->GetPathName() : FString();
+
+			OutputDevice->Logf(TEXT("Printing archetype for object '%s' [%s] (%s)"),
+				*Object->GetName(), *Object->GetClass()->GetName(), *ObjectPath);
+
+			OutputDevice->Logf(TEXT("'%s' [%s] (%s)"), *ArchetypeName, *ArchetypeClassName, *ArchetypePath);
+		}
+		else
+		{
+			OutputDevice->Logf(TEXT("PrintObjectArchetype: NULL object"));
+		}
+	}
+}
+
+void PrintObjectIDO(UObject* Object, FOutputDevice* InOutputDevice)
+{
+	FOutputDevice* OutputDevice = InOutputDevice ? InOutputDevice : GLog;
+	if (OutputDevice)
+	{
+		#if WITH_EDITORONLY_DATA
+		if (Object)
+		{
+			FString ObjectPath = Object->GetPathName();
+
+			UObject* IDO = UE::FPropertyBagRepository::Get().FindInstanceDataObject(Object);
+
+			if (IDO)
+			{
+				const UClass* IDOClass = IDO->GetClass();
+				FString IDOClassName = IDOClass ? IDOClass->GetName() : FString();
+
+				OutputDevice->Logf(TEXT("Printing IDO (Instance Data Object) for object '%s' [%s] (%s)"),
+					*Object->GetName(), *Object->GetClass()->GetName(), *ObjectPath);
+
+				OutputDevice->Logf(TEXT("'%s' [%s] (%s)"), *IDO->GetName(), *IDOClassName, *IDO->GetPathName());
+			}
+			else
+			{
+				OutputDevice->Logf(TEXT("PrintObjectIDO: No IDO (Instance Data Object) found for object '%s' [%s] (%s)"),
+					*Object->GetName(), *Object->GetClass()->GetName(), *ObjectPath);
+			}
+		}
+		else
+		{
+			OutputDevice->Logf(TEXT("PrintObjectIDO: NULL object"));
+		}
+		#else
+		OutputDevice->Logf(TEXT("PrintObjectIDO: IDOs (Instance Data Objects) not supported in the current build"));
+		#endif
+	}
+}
+
+
+FAutoConsoleCommandWithArgsAndOutputDevice CVarCommandPrintObjectsInOuter(
+	TEXT("Obj.DumpObjectsInOuter"),
+	TEXT("Lists all objects under a specified parent (the parent object must be specified as a path, e.g., /MyLevel/MyLevel.MyLevel:PersistentLevel)."),
+	FConsoleCommandWithArgsAndOutputDeviceDelegate::CreateLambda([](const TArray<FString>& Args, FOutputDevice& OutputDevice)
+	{
+		if (Args.Num() > 0)
+		{
+			FString ObjectPath = Args[0];
+			UObject* Object = FindFirstObjectSafe<UObject>(*ObjectPath);
+
+			if (Object)
+			{
+				PrintObjectsInOuter(Object, &OutputDevice);
+			}
+			else
+			{
+				OutputDevice.Logf(TEXT("PrintObjectsInOuter: failed to find any objects for path: '%s'"), *ObjectPath);
+			}
+		}
+		else
+		{
+			OutputDevice.Logf(TEXT("PrintObjectsInOuter: no object path specified (example usage: Obj.PrintObjectsInOuter /MyLevel/MyLevel.MyLevel:PersistentLevel)"));
+		}
+	}));
+
+FAutoConsoleCommandWithArgsAndOutputDevice CVarCommandPrintObjectsWithName(
+	TEXT("Obj.DumpObjectsWithName"),
+	TEXT("Lists all objects with a given name."),
+	FConsoleCommandWithArgsAndOutputDeviceDelegate::CreateLambda([](const TArray<FString>& Args, FOutputDevice& OutputDevice)
+	{
+		if (Args.Num() > 0)
+		{
+			FString ObjectName = Args[0];
+
+			PrintObjectsWithName(*ObjectName, &OutputDevice);
+		}
+		else
+		{
+			OutputDevice.Logf(TEXT("PrintObjectsWithName: no object name specified (example usage: Obj.PrintObjectsWithName PersistentLevel)"));
+		}
+	}));
+
+FAutoConsoleCommandWithArgsAndOutputDevice CVarCommandPrintObjectProperties(
+	TEXT("Obj.DumpObjectProperties"),
+	TEXT("Lists the properties of an object (the object must be specified as a path, e.g., /MyLevel/MyLevel.MyLevel:PersistentLevel.MyActor). Pass InitState=true to iclude the properties' initialization state."),
+	FConsoleCommandWithArgsAndOutputDeviceDelegate::CreateLambda([](const TArray<FString>& Args, FOutputDevice& OutputDevice)
+	{
+		if (Args.Num() > 0)
+		{
+			FString ObjectPath;
+			bool bIncludeinitializationState = false;
+
+			for (FString Arg : Args)
+			{
+				if (Arg.Equals(TEXT("InitState=true"), ESearchCase::IgnoreCase))
+				{
+					bIncludeinitializationState = true;
+				}
+				else
+				{
+					ObjectPath = Arg;
+				}
+			}
+
+			UObject* Object = FindFirstObjectSafe<UObject>(*ObjectPath);
+
+			if (Object)
+			{
+				PrintObjectProperties(Object, bIncludeinitializationState, &OutputDevice);
+			}
+			else
+			{
+				OutputDevice.Logf(TEXT("PrintObjectProperties: failed to find any objects for path: '%s'"), *ObjectPath);
+			}
+		}
+		else
+		{
+			OutputDevice.Logf(TEXT("PrintObjectProperties: no object path specified (example usage: Obj.PrintObjectProperties /MyLevel/MyLevel.MyLevel:PersistentLevel.MyActor)"));
+		}
+	}));
+
+FAutoConsoleCommandWithArgsAndOutputDevice CVarCommandPrintObjectArchetype(
+	TEXT("Obj.DumpObjectArchetype"),
+	TEXT("Outputs an object's archetype (the object must be specified as a path, e.g., /MyLevel/MyLevel.MyLevel:PersistentLevel.MyActor)."),
+	FConsoleCommandWithArgsAndOutputDeviceDelegate::CreateLambda([](const TArray<FString>& Args, FOutputDevice& OutputDevice)
+	{
+		if (Args.Num() > 0)
+		{
+			FString ObjectPath = Args[0];
+			UObject* Object = FindFirstObjectSafe<UObject>(*ObjectPath);
+
+			if (Object)
+			{
+				PrintObjectArchetype(Object, &OutputDevice);
+			}
+			else
+			{
+				OutputDevice.Logf(TEXT("PrintObjectArchetype: failed to find any objects for path: '%s'"), *ObjectPath);
+			}
+		}
+		else
+		{
+			OutputDevice.Logf(TEXT("PrintObjectArchetype: no object path specified (example usage: Obj.PrintObjectArchetype /MyLevel/MyLevel.MyLevel:PersistentLevel.MyActor)"));
+		}
+	}));
+
+FAutoConsoleCommandWithArgsAndOutputDevice CVarCommandPrintObjectIDO(
+	TEXT("Obj.DumpObjectIDO"),
+	TEXT("Outputs an object's IDO (Instance Data Object) (the object must be specified as a path, e.g., /MyLevel/MyLevel.MyLevel:PersistentLevel.MyActor)."),
+	FConsoleCommandWithArgsAndOutputDeviceDelegate::CreateLambda([](const TArray<FString>& Args, FOutputDevice& OutputDevice)
+	{
+		if (Args.Num() > 0)
+		{
+			FString ObjectPath = Args[0];
+			UObject* Object = FindFirstObjectSafe<UObject>(*ObjectPath);
+
+			if (Object)
+			{
+				PrintObjectIDO(Object, &OutputDevice);
+			}
+			else
+			{
+				OutputDevice.Logf(TEXT("PrintObjectIDO: failed to find any objects for path: '%s'"), *ObjectPath);
+			}
+		}
+		else
+		{
+			OutputDevice.Logf(TEXT("PrintObjectIDO: no object path specified (example usage: Obj.PrintObjectArchetype /MyLevel/MyLevel.MyLevel:PersistentLevel.MyActor)"));
+		}
+	}));
+
+} // namespace UE
+
+#endif // !UE_BUILD_SHIPPING
diff --git a/Engine/Source/Runtime/Engine/Public/PrintObjectUtils.h b/Engine/Source/Runtime/Engine/Public/PrintObjectUtils.h
new file mode 100644
--- /dev/null
+++ b/Engine/Source/Runtime/Engine/Public/PrintObjectUtils.h
@@ -0,0 +1,74 @@
+// Copyright Epic Games, Inc. All Rights Reserved.
+
+#include "HAL/Platform.h"			// For TCHAR
+#include "Misc/Build.h"				// For UE_BUILD_SHIPPING
+#include "Misc/CoreMiscDefines.h"	// For UE_INTERNAL
+
+#if !UE_BUILD_SHIPPING
+
+class UObject;
+class UStruct;
+class FOutputDevice;
+
+/**
+ *	Utility functions and console commands that print UObject state information.
+ *	Can be used in a debugger as well.
+**/
+
+namespace UE
+{
+
+/**
+ * Lists all objects under a specified parent.
+ * 
+ * @param Object The outer object to use.
+ * @param OutputDevice The output device to print to (can be NULL).
+ */
+UE_INTERNAL void PrintObjectsInOuter(UObject* Object, FOutputDevice* OutputDevice = nullptr);
+
+/**
+ * Lists all objects with a given name.
+ *
+ * @param ObjectName The name to search for.
+ * @param OutputDevice The output device to print to (can be NULL).
+ */
+UE_INTERNAL void PrintObjectsWithName(const TCHAR* ObjectName, FOutputDevice* OutputDevice = nullptr);
+
+/**
+ * Lists the properties of a UStruct instance.
+ *
+ * @param Struct The type of the UStruct instance.
+ * @param StructData The data of the UStruct instance.
+ * @param bIncludeinitializationState If true the function will also output the initialization state of the properties.
+ * @param OutputDevice The output device to print to (can be NULL).
+ */
+UE_INTERNAL void PrintStructProperties(UStruct* Struct, void* StructData, bool bIncludeinitializationState = false, FOutputDevice* OutputDevice = nullptr);
+
+/**
+ * Lists the property values of an object.
+ *
+ * @param Object The object to use.
+ * @param bIncludeinitializationState If true the function will also output the initialization state of the properties.
+ * @param OutputDevice The output device to print to (can be NULL).
+ */
+UE_INTERNAL void PrintObjectProperties(UObject* Object, bool bIncludeinitializationState = false, FOutputDevice* OutputDevice = nullptr);
+
+/**
+ * Outputs an object's archetype.
+ *
+ * @param Object The object to use.
+ * @param OutputDevice The output device to print to (can be NULL).
+ */
+UE_INTERNAL void PrintObjectArchetype(UObject* Object, FOutputDevice* OutputDevice = nullptr);
+
+/**
+ * Outputs an object's IDO (Instance Data Object).
+ *
+ * @param Object The object to use.
+ * @param OutputDevice The output device to print to (can be NULL).
+ */
+UE_INTERNAL void PrintObjectIDO(UObject* Object, FOutputDevice* OutputDevice = nullptr);
+
+} // namespace UE
+
+#endif // !UE_BUILD_SHIPPING
