[NewMaterialTranslator] - Fix a bug where operator precedence wasn't taken into account. Fix this by adding parenthesis around all arithematic operators. This functionality is exactly equivilant to the old translator (see FHLSLMaterialTranslator::Add)

#rnx

diff --git a/Engine/Source/Runtime/Engine/Private/Materials/MaterialIRToHLSLTranslator.cpp b/Engine/Source/Runtime/Engine/Private/Materials/MaterialIRToHLSLTranslator.cpp
--- a/Engine/Source/Runtime/Engine/Private/Materials/MaterialIRToHLSLTranslator.cpp
+++ b/Engine/Source/Runtime/Engine/Private/Materials/MaterialIRToHLSLTranslator.cpp
@@ -479,6 +479,8 @@ struct FTranslator : FMaterialIRToHLSLTranslation
 	{
 		if (IsOperatorInfix(Operator->Op))
 		{
+			Printer << "(";
+
 			LowerValue(Operator->AArg);
 
 			const TCHAR* OpString;
@@ -499,6 +501,8 @@ struct FTranslator : FMaterialIRToHLSLTranslation
 			Printer << OpString;
 
 			LowerValue(Operator->BArg);
+
+			Printer << ")";
 		}
 		else
 		{
