New function in AssetManager - GetPakStringToChunkIdMapping, calls GetPakChunkIdToStringMapping to get map with named pak as the key and chunk id as the value.

diff --git a/Engine/Source/Runtime/Engine/Classes/Engine/AssetManager.h b/Engine/Source/Runtime/Engine/Classes/Engine/AssetManager.h
--- a/Engine/Source/Runtime/Engine/Classes/Engine/AssetManager.h
+++ b/Engine/Source/Runtime/Engine/Classes/Engine/AssetManager.h
@@ -706,6 +706,13 @@ public:
 	 * @param ChunkIdStringOverride - Out parameter for the override information of a chunkID to string.
 	 */
 	virtual void GetPakChunkIdToStringMapping(TMap<int32, FString>& ChunkIdStringOverride) {}
+	/**
+	 * Helper to make it easy to get the ChunkId based on the string override for a pakchunk.
+	 * see "GetPakChunkIdToStringMapping" for more information
+	 *
+	 * @param ChunkIdStringOverride - Out parameter for getting the chunkID based on the string override.
+	 */
+	ENGINE_API void GetPakStringToChunkIdMapping(TMap<FString, int32>& StringOverrideToChunkId);
 
 protected:
 	friend class FAssetManagerEditorModule;
diff --git a/Engine/Source/Runtime/Engine/Private/AssetManager.cpp b/Engine/Source/Runtime/Engine/Private/AssetManager.cpp
--- a/Engine/Source/Runtime/Engine/Private/AssetManager.cpp
+++ b/Engine/Source/Runtime/Engine/Private/AssetManager.cpp
@@ -3056,6 +3056,17 @@ FSoftObjectPath UAssetManager::GetAssetPathForData(const FAssetData& AssetData)
 	return Result;
 }
 
+void UAssetManager::GetPakStringToChunkIdMapping(TMap<FString, int32>& StringOverrideToChunkId)
+{
+	TMap<int32, FString> ChunkIdStringOverride;
+	GetPakChunkIdToStringMapping(ChunkIdStringOverride);
+
+	for (const TPair<int32, FString>& IdAndStringPair : ChunkIdStringOverride)
+	{
+		StringOverrideToChunkId.Add(IdAndStringPair.Value, IdAndStringPair.Key);
+	}
+}
+
 void UAssetManager::GetAssetDataForPathInternal(IAssetRegistry& AssetRegistry, const FString& AssetPath, OUT FAssetData& OutAssetData) const
 {
 	// We're a class if our path is foo.foo_C
