[MIR] Added missing debug output for FExternalCode MIR instruction (Follow-up fix to 39995668).

#rnx

diff --git a/Engine/Source/Runtime/Engine/Private/Materials/MaterialIR.cpp b/Engine/Source/Runtime/Engine/Private/Materials/MaterialIR.cpp
--- a/Engine/Source/Runtime/Engine/Private/Materials/MaterialIR.cpp
+++ b/Engine/Source/Runtime/Engine/Private/Materials/MaterialIR.cpp
@@ -15,6 +15,7 @@ const TCHAR* ValueKindToString(EValueKind Kind)
 	{
 		case VK_Constant: return TEXT("Constant");
 		case VK_ExternalInput: return TEXT("ExternalInput");
+		case VK_ExternalCode: return TEXT("ExternalCode");
 		case VK_TextureObject: return TEXT("TextureObject");
 		case VK_UniformParameter: return TEXT("UniformParameter");
 		case VK_Dimensional: return TEXT("Dimensional");
@@ -34,6 +35,7 @@ uint32 FValue::GetSizeInBytes() const
 	{
 		case VK_Constant: return sizeof(FConstant);
 		case VK_ExternalInput: return sizeof(FExternalInput);
+		case VK_ExternalCode: return sizeof(FExternalCode);
 		case VK_TextureObject: return sizeof(FTextureObject);
 		case VK_UniformParameter: return sizeof(FUniformParameter);
 		case VK_SetMaterialOutput: return sizeof(FSetMaterialOutput);
@@ -132,6 +134,11 @@ TArrayView<FValue*> FValue::GetUses()
 			return { &This->TextureObject, 4 };
 		}
 
+		case VK_ExternalCode:
+		{
+			return {}; // @todo-laura.hermanns - Add input arguments as soon as external code expressions support parameters
+		}
+
 		default: UE_MIR_UNREACHABLE();
 	}
 }
diff --git a/Engine/Source/Runtime/Engine/Private/Materials/MaterialIRDebug.cpp b/Engine/Source/Runtime/Engine/Private/Materials/MaterialIRDebug.cpp
--- a/Engine/Source/Runtime/Engine/Private/Materials/MaterialIRDebug.cpp
+++ b/Engine/Source/Runtime/Engine/Private/Materials/MaterialIRDebug.cpp
@@ -6,6 +6,7 @@
 #include "Materials/MaterialIRTypes.h"
 #include "Materials/MaterialIRInternal.h"
 #include "Materials/MaterialAttributeDefinitionMap.h"
+#include "Materials/MaterialExternalCodeRegistry.h"
 #include "Misc/Paths.h"
 #include "Misc/FileHelper.h"
 #include "Engine/Texture.h"
@@ -289,6 +290,17 @@ struct FDebugDumpIRState
 			Out.Appendf(TEXT(" MipValueMode=\"%s\""), *StaticEnum<ETextureMipValueMode>()->GetDisplayNameTextByValue(TextureSample->MipValueMode).ToString());
 			Out.Appendf(TEXT(" SamplerType=\"%s\""), *StaticEnum<EMaterialSamplerType>()->GetDisplayNameTextByValue(TextureSample->SamplerType).ToString());
 		}
+		else if (const FExternalCode* ExternalCode = Instr->As<FExternalCode>())
+		{
+			if (const FMaterialExternalCodeDeclaration* ExternalCodeDeclaration = ExternalCode->ExternalCodeDeclaration)
+			{
+				Out.Appendf(TEXT(" \"%s\""), *ExternalCodeDeclaration->Definition);
+			}
+			else
+			{
+				Out.Append(TEXT(" <Missing Definition>"));
+			}
+		}
 	}
 
 	/// Appends a reference to the specified value. This will look like "%x" if Value is an
