[AutoRTFM] Swap critical section for transactionally-safe variant.

This code is reachable from a transaction.

diff --git a/Engine/Source/Runtime/Engine/Classes/Engine/CurveTable.h b/Engine/Source/Runtime/Engine/Classes/Engine/CurveTable.h
--- a/Engine/Source/Runtime/Engine/Classes/Engine/CurveTable.h
+++ b/Engine/Source/Runtime/Engine/Classes/Engine/CurveTable.h
@@ -3,6 +3,7 @@
 #pragma once
 
 #include "CoreMinimal.h"
+#include "Misc/TransactionallySafeCriticalSection.h"
 #include "UObject/ObjectMacros.h"
 #include "UObject/Object.h"
 #include "UObject/Class.h"
@@ -74,7 +75,7 @@ protected:
 	 */
 	TMap<FName, FRealCurve*>	RowMap;
 
-	static FCriticalSection& GetCurveTableChangeCriticalSection();
+	static FTransactionallySafeCriticalSection& GetCurveTableChangeCriticalSection();
 
 public:
 	//~ Begin UObject Interface.
diff --git a/Engine/Source/Runtime/Engine/Private/CompositeCurveTable.cpp b/Engine/Source/Runtime/Engine/Private/CompositeCurveTable.cpp
--- a/Engine/Source/Runtime/Engine/Private/CompositeCurveTable.cpp
+++ b/Engine/Source/Runtime/Engine/Private/CompositeCurveTable.cpp
@@ -144,7 +144,7 @@ void UCompositeCurveTable::UpdateCachedRowMap(bool bWarnOnInvalidChildren)
 				}
 			};
 
-			FScopeLock ScopeLock(&UCurveTable::GetCurveTableChangeCriticalSection());
+			UE::TScopeLock ScopeLock(UCurveTable::GetCurveTableChangeCriticalSection());
 
 			// If we are using simple curves we know all our parents are also simple
 			if (CurveTableMode == ECurveTableMode::SimpleCurves)
diff --git a/Engine/Source/Runtime/Engine/Private/CurveTable.cpp b/Engine/Source/Runtime/Engine/Private/CurveTable.cpp
--- a/Engine/Source/Runtime/Engine/Private/CurveTable.cpp
+++ b/Engine/Source/Runtime/Engine/Private/CurveTable.cpp
@@ -1,6 +1,7 @@
 // Copyright Epic Games, Inc. All Rights Reserved.
 
 #include "Engine/CurveTable.h"
+#include "Misc/TransactionallySafeCriticalSection.h"
 #include "Serialization/JsonSerializer.h"
 #include "Serialization/Csv/CsvParser.h"
 #include "HAL/IConsoleManager.h"
@@ -54,16 +55,16 @@ namespace
 		static TMap<UCurveTable*, int32> ScopeCount;
 
 	public:
-		static FCriticalSection CriticalSection;
+		static FTransactionallySafeCriticalSection CriticalSection;
 	};
 
 	TMap<UCurveTable*, int32> FScopedCurveTableChange::ScopeCount;
-	FCriticalSection FScopedCurveTableChange::CriticalSection;
+	FTransactionallySafeCriticalSection FScopedCurveTableChange::CriticalSection;
 
 #define CURVETABLE_CHANGE_SCOPE()	FScopedCurveTableChange ActiveScope(this);
 }
 
-FCriticalSection& UCurveTable::GetCurveTableChangeCriticalSection()
+FTransactionallySafeCriticalSection& UCurveTable::GetCurveTableChangeCriticalSection()
 {
 	return FScopedCurveTableChange::CriticalSection;
 }
