SceneGraph
* EntityReplicator's location on client will now be synced to the entity's global transform
* This should fix some replay bugs where entities were getting culled out by the DemoNetDriver due to being out of relevancy
* Added a seperate AlwaysRelevant entity unit test
* Updated the SpatialEntitiesWithChild to show that some replicated child entities are kept in uninitialized state while relevant. Will be fixed later via UE-243598

#rnx

diff --git a/Engine/Source/Runtime/Engine/Public/Tests/NetTestHelpers.h b/Engine/Source/Runtime/Engine/Public/Tests/NetTestHelpers.h
--- a/Engine/Source/Runtime/Engine/Public/Tests/NetTestHelpers.h
+++ b/Engine/Source/Runtime/Engine/Public/Tests/NetTestHelpers.h
@@ -209,6 +209,12 @@ struct FTestWorlds
 	 */
 	ENGINE_API UObject* FindReplicatedObjectOnClient(UObject* ServerObject, uint32 ClientIndex) const;
 
+	template<typename T>
+	T* FindReplicatedObjectOnClient(UObject* ServerObject, uint32 ClientIndex) const
+	{
+		return Cast<T>(FindReplicatedObjectOnClient(ServerObject, ClientIndex));
+	}
+
 	/**
 	 * Returns if the the remote instance of a replicated object exists on the client or not.
 	 * @param ServerObject The server object you want to see if it is replicated
