[Materials] Fix crash in HLSLMaterialTranslator when CameraVector() does not provide a valid code chunk.

#rnx

diff --git a/Engine/Source/Runtime/Engine/Private/Materials/HLSLMaterialTranslator.cpp b/Engine/Source/Runtime/Engine/Private/Materials/HLSLMaterialTranslator.cpp
--- a/Engine/Source/Runtime/Engine/Private/Materials/HLSLMaterialTranslator.cpp
+++ b/Engine/Source/Runtime/Engine/Private/Materials/HLSLMaterialTranslator.cpp
@@ -3206,8 +3206,9 @@ FString FHLSLMaterialTranslator::GetParameterCode(int32 Index, const TCHAR* Defa
 // GetParameterCode
 FString FHLSLMaterialTranslator::GetParameterCodeRaw(int32 Index, const TCHAR* Default)
 {
-	if(Index == INDEX_NONE && Default)
+	if(Index == INDEX_NONE)
 	{
+		checkf(Default != nullptr, TEXT("Invalid material expression and no default value provided"));
 		return Default;
 	}
 
@@ -12192,7 +12193,12 @@ int32 FHLSLMaterialTranslator::SkyAtmosphereViewLuminance(int32 WorldDirectionOv
 {
 	if (WorldDirectionOverrideCodeChunk == INDEX_NONE)
 	{
-		WorldDirectionOverrideCodeChunk = Mul(Constant3(-1.0f, -1.0f, -1.0f), CameraVector());
+		const int32 CameraVectorCode = CameraVector();
+		if (CameraVectorCode == INDEX_NONE)
+		{
+			return INDEX_NONE;
+		}
+		WorldDirectionOverrideCodeChunk = Mul(Constant3(-1.0f, -1.0f, -1.0f), CameraVectorCode);
 	}
 
 	bUsesSkyAtmosphere = true;
@@ -14402,8 +14408,14 @@ int32 FHLSLMaterialTranslator::SubstrateHorizontalMixingParameterBlending(
 		return INDEX_NONE;
 	}
 
-	const int32 BackgroundNoV = Saturate(Dot(BackgroundBSDFNormalCodeChunk, CameraVector()));
-	const int32 ForegroundNoV = Saturate(Dot(ForegroundBSDFNormalCodeChunk, CameraVector()));
+	const int32 CameraVectorCode = CameraVector();
+	if (CameraVectorCode == INDEX_NONE)
+	{
+		return INDEX_NONE;
+	}
+
+	const int32 BackgroundNoV = Saturate(Dot(BackgroundBSDFNormalCodeChunk, CameraVectorCode));
+	const int32 ForegroundNoV = Saturate(Dot(ForegroundBSDFNormalCodeChunk, CameraVectorCode));
 
 	if (PromoteToOperator)
 	{
@@ -14467,8 +14479,14 @@ int32 FHLSLMaterialTranslator::SubstrateVerticalLayeringParameterBlending(int32
 
 	const FString ThicknessCode = Thickness != INDEX_NONE ? GetParameterCode(Thickness) : TEXT("NONE");
 
-	const int32 TopNoV  = Saturate(Dot(TopBSDFNormalCodeChunk, CameraVector()));
-	const int32 BaseNoV = Saturate(Dot(BaseBSDFNormalCodeChunk, CameraVector()));
+	const int32 CameraVectorCode = CameraVector();
+	if (CameraVectorCode == INDEX_NONE)
+	{
+		return INDEX_NONE;
+	}
+
+	const int32 TopNoV  = Saturate(Dot(TopBSDFNormalCodeChunk, CameraVectorCode));
+	const int32 BaseNoV = Saturate(Dot(BaseBSDFNormalCodeChunk, CameraVectorCode));
 
 	if (PromoteToOperator)
 	{
@@ -14524,8 +14542,14 @@ int32 FHLSLMaterialTranslator::SubstrateAddParameterBlending(int32 A, int32 B, i
 		return INDEX_NONE;
 	}
 
-	const int32 ANoV = Saturate(Dot(ABSDFNormalCodeChunk, CameraVector()));
-	const int32 BNoV = Saturate(Dot(BBSDFNormalCodeChunk, CameraVector()));
+	const int32 CameraVectorCode = CameraVector();
+	if (CameraVectorCode == INDEX_NONE)
+	{
+		return INDEX_NONE;
+	}
+
+	const int32 ANoV = Saturate(Dot(ABSDFNormalCodeChunk, CameraVectorCode));
+	const int32 BNoV = Saturate(Dot(BBSDFNormalCodeChunk, CameraVectorCode));
 
 	if (PromoteToOperator)
 	{
@@ -14730,9 +14754,14 @@ int32 FHLSLMaterialTranslator::SubstrateThinFilm(int32 NormalCodeChunk, int32 Sp
 		Errorf(TEXT("Substrate error: SubstrateThinFilm output index is invalid (asset: %s).\r\n"), *Material->GetDebugName(), *Material->GetAssetPath().ToString());
 		return INDEX_NONE;
 	}
+	const int32 CameraVectorCode = CameraVector();
+	if (CameraVectorCode == INDEX_NONE)
+	{
+		return INDEX_NONE;
+	}
 	return AddCodeChunk(
 		MCT_Float3, TEXT("SubstrateGetThinFilmF0F90(%s, %s, %s, %s, %s)%s"),
-		*GetParameterCode(Dot(NormalCodeChunk, CameraVector())),
+		*GetParameterCode(Dot(NormalCodeChunk, CameraVectorCode)),
 		*GetParameterCode(SpecularColorCodeChunk),
 		*GetParameterCode(EdgeSpecularColorCodeChunk),
 		*GetParameterCode(ThicknessCodeChunk),
