[XR] Deprecate the IXrInput interface. Nothing in the engine actually used this. Use a custom player controller subclass to override the input behavior if desired.

This will deprecate the IXRInput.h file so that you will get a compilation error if you included this to notify you of the deprecation, since virtual functions are difficult to deprecate and wouldn't give end users warnings.

You can include LegacyIXRInput.h as a temporary solution if your project depends on this somehow.

#rnx

diff --git a/Engine/Source/Runtime/Engine/Classes/GameFramework/PlayerController.h b/Engine/Source/Runtime/Engine/Classes/GameFramework/PlayerController.h
--- a/Engine/Source/Runtime/Engine/Classes/GameFramework/PlayerController.h
+++ b/Engine/Source/Runtime/Engine/Classes/GameFramework/PlayerController.h
@@ -1746,19 +1746,11 @@ public:
 	UFUNCTION(BlueprintCallable, Category = Input)
 	ENGINE_API TSubclassOf<UPlayerInput> GetOverridePlayerInputClass() const;
 
-	/** Handles a key press */
-	UE_DEPRECATED(5.0, "This version of InputKey has been deprecated, please use the version that takes FInputKeyParams instead")
-	ENGINE_API virtual bool InputKey(FKey Key, EInputEvent EventType, float AmountDepressed, bool bGamepad);
-
 	/** Handles a key press */
 	ENGINE_API virtual bool InputKey(const FInputKeyParams& Params);
 
 	/** Handles a touch screen action */
 	ENGINE_API virtual bool InputTouch(uint32 Handle, ETouchType::Type Type, const FVector2D& TouchLocation, float Force, FDateTime DeviceTimestamp, uint32 TouchpadIndex);
-
-	/** Handles a controller axis input */
-	UE_DEPRECATED(5.0, "InputAxis has been deprecated, please use InputKey instead.")
-	ENGINE_API virtual bool InputAxis(FKey Key, float Delta, float DeltaTime, int32 NumSamples, bool bGamepad);
 	
 	/** Handles motion control */
 	ENGINE_API virtual bool InputMotion(const FVector& Tilt, const FVector& RotationRate, const FVector& Gravity, const FVector& Acceleration);
diff --git a/Engine/Source/Runtime/Engine/Private/PlayerController.cpp b/Engine/Source/Runtime/Engine/Private/PlayerController.cpp
--- a/Engine/Source/Runtime/Engine/Private/PlayerController.cpp
+++ b/Engine/Source/Runtime/Engine/Private/PlayerController.cpp
@@ -42,7 +42,7 @@
 #include "IHeadMountedDisplay.h"
 #include "IXRTrackingSystem.h"
 #include "IXRCamera.h"
-#include "IXRInput.h"
+#include "LegacyIXRInput.h"
 #include "GameFramework/TouchInterface.h"
 #include "DisplayDebugHelpers.h"
 #include "MoviePlayerProxy.h"
@@ -2323,17 +2323,6 @@ TSubclassOf<UPlayerInput> APlayerController::GetOverridePlayerInputClass() const
 	return OverridePlayerInputClass;
 }
 
-bool APlayerController::InputKey(FKey Key, EInputEvent EventType, float AmountDepressed, bool bGamepad)
-{
-	FInputKeyParams Params;
-	Params.Key = Key;
-	Params.Event = EventType;
-	Params.Delta.X = AmountDepressed;
-	Params.bIsGamepadOverride = bGamepad;
-	
-	return InputKey(Params);
-}
-
 bool APlayerController::InputKey(const FInputKeyParams& Params)
 {
 	bool bResult = false;
@@ -2356,14 +2345,16 @@ bool APlayerController::InputKey(const FInputKeyParams& Params)
 	// But we need special case XR handling for non-analog values...
 	else
 	{
+		PRAGMA_DISABLE_DEPRECATION_WARNINGS
 		if (GEngine->XRSystem.IsValid())
 		{
-			auto XRInput = GEngine->XRSystem->GetXRInput();
+			IXRInput* XRInput = GEngine->XRSystem->GetXRInput();
 			if (XRInput && XRInput->HandleInputKey(PlayerInput, Params.Key, Params.Event, Params.Delta.X, Params.IsGamepad()))
 			{
 				return true;
 			}
 		}
+		PRAGMA_ENABLE_DEPRECATION_WARNINGS
 
 		if (PlayerInput)
 		{
@@ -2424,28 +2415,18 @@ bool APlayerController::InputKey(const FInputKeyParams& Params)
 	return bResult;
 }
 
-bool APlayerController::InputAxis(FKey Key, float Delta, float DeltaTime, int32 NumSamples, bool bGamepad)
-{
-	FInputKeyParams Params;
-	Params.Key = Key;
-	Params.Delta = FVector(static_cast<double>(Delta), 0.0, 0.0);
-	Params.NumSamples = NumSamples;
-	Params.DeltaTime = DeltaTime;
-	Params.bIsGamepadOverride = bGamepad;
-
-	return InputKey(Params);
-}
-
 bool APlayerController::InputTouch(uint32 Handle, ETouchType::Type Type, const FVector2D& TouchLocation, float Force, FDateTime DeviceTimestamp, uint32 TouchpadIndex)
 {
+	PRAGMA_DISABLE_DEPRECATION_WARNINGS
 	if (GEngine->XRSystem.IsValid())
 	{
-		auto XRInput = GEngine->XRSystem->GetXRInput();
+		IXRInput* XRInput = GEngine->XRSystem->GetXRInput();
 		if(XRInput && XRInput->HandleInputTouch(Handle, Type, TouchLocation, DeviceTimestamp, TouchpadIndex))
 		{
 			return true;
 		}
 	}
+	PRAGMA_ENABLE_DEPRECATION_WARNINGS
 
 	bool bResult = false;
 	if (PlayerInput)
diff --git a/Engine/Source/Runtime/HeadMountedDisplay/Public/IXRInput.h b/Engine/Source/Runtime/HeadMountedDisplay/Public/IXRInput.h
--- a/Engine/Source/Runtime/HeadMountedDisplay/Public/IXRInput.h
+++ b/Engine/Source/Runtime/HeadMountedDisplay/Public/IXRInput.h
@@ -2,29 +2,9 @@
 
 #pragma once
 
-#include "CoreMinimal.h"
-#include "HeadMountedDisplayTypes.h"
-#include "InputCoreTypes.h"
-#include "Engine/EngineBaseTypes.h"
+#include "Misc/CoreMiscDefines.h"
 
-/** 
- * Optional interface returned from IXRTrackingSystem if the plugin requires being able to grab touch or keyboard input events.
- */
-class IXRInput
-{
-public:
-	/**
-	* Passing key events to HMD.
-	* If returns 'false' then key will be handled by PlayerController;
-	* otherwise, key won't be handled by the PlayerController.
-	*/
-	virtual bool HandleInputKey(class UPlayerInput*, const struct FKey& Key, EInputEvent EventType, float AmountDepressed, bool bGamepad) { return false; }
-
-	/**
-	* Passing touch events to HMD.
-	* If returns 'false' then touch will be handled by PlayerController;
-	* otherwise, touch won't be handled by the PlayerController.
-	*/
-	virtual bool HandleInputTouch(uint32 Handle, ETouchType::Type Type, const FVector2D& TouchLocation, FDateTime DeviceTimestamp, uint32 TouchpadIndex) { return false; }
-};
+// IXRInput has been deprecated and will be removed in a future release.
+UE_DEPRECATED_HEADER(5.6, "Use LegacyIXRInput.h instead of IXRInput.h")
 
+#include "LegacyIXRInput.h"
\ No newline at end of file
diff --git a/Engine/Source/Runtime/HeadMountedDisplay/Public/IXRTrackingSystem.h b/Engine/Source/Runtime/HeadMountedDisplay/Public/IXRTrackingSystem.h
--- a/Engine/Source/Runtime/HeadMountedDisplay/Public/IXRTrackingSystem.h
+++ b/Engine/Source/Runtime/HeadMountedDisplay/Public/IXRTrackingSystem.h
@@ -7,7 +7,6 @@
 #include "IIdentifiableXRDevice.h"
 #include "UObject/ObjectMacros.h"
 #include "Features/IModularFeature.h"
-#include "IXRInput.h"
 #include "StereoRendering.h"
 
 class IXRCamera;
@@ -334,11 +333,14 @@ public:
 	}
 
 	/**
+	 * Deprecated, and will be removed in a future release.
+	 
 	 * Access optional HMD input override interface.
 	 *
 	 * @return a IXRInput pointer or a nullptr if not supported
 	 */
-	virtual IXRInput* GetXRInput() { return nullptr; }
+	UE_DEPRECATED(5.6, "IXRInput is deprecated and will be removed in a future release.")
+	virtual class IXRInput* GetXRInput() { return nullptr; }
 
 
 	/**
diff --git a/Engine/Source/Runtime/HeadMountedDisplay/Public/LegacyIXRInput.h b/Engine/Source/Runtime/HeadMountedDisplay/Public/LegacyIXRInput.h
new file mode 100644
--- /dev/null
+++ b/Engine/Source/Runtime/HeadMountedDisplay/Public/LegacyIXRInput.h
@@ -0,0 +1,32 @@
+// Copyright Epic Games, Inc. All Rights Reserved.
+
+#pragma once
+
+#include "CoreMinimal.h"
+#include "HeadMountedDisplayTypes.h"
+#include "InputCoreTypes.h"
+#include "Engine/EngineBaseTypes.h"
+
+/** 
+ * Deprecated. This type will be removed in a future release
+ * 
+ * Optional interface returned from IXRTrackingSystem if the plugin requires being able to grab touch or keyboard input events.
+ */
+class UE_DEPRECATED(5.6, "IXRInput is deprecated and will be removed in a future release.") IXRInput
+{
+public:
+	/**
+	* Passing key events to HMD.
+	* If returns 'false' then key will be handled by PlayerController;
+	* otherwise, key won't be handled by the PlayerController.
+	*/
+	virtual bool HandleInputKey(class UPlayerInput*, const struct FKey& Key, EInputEvent EventType, float AmountDepressed, bool bGamepad) { return false; }
+
+	/**
+	* Passing touch events to HMD.
+	* If returns 'false' then touch will be handled by PlayerController;
+	* otherwise, touch won't be handled by the PlayerController.
+	*/
+	virtual bool HandleInputTouch(uint32 Handle, ETouchType::Type Type, const FVector2D& TouchLocation, FDateTime DeviceTimestamp, uint32 TouchpadIndex) { return false; }
+};
+
