Fix SetOneWayInteraction not applying after toggling physics on body.

SetOneWayInteraction will not apply by default if SimulatePhysics is disabled on a body because the function checks that the body is not kinematic. This CL updates the state once SimulatePhysics is reenabled on the body in FBodyInstance::UpdateInstanceSimulatePhysics to properly apply the property.

diff --git a/Engine/Source/Runtime/Engine/Private/PhysicsEngine/BodyInstance.cpp b/Engine/Source/Runtime/Engine/Private/PhysicsEngine/BodyInstance.cpp
--- a/Engine/Source/Runtime/Engine/Private/PhysicsEngine/BodyInstance.cpp
+++ b/Engine/Source/Runtime/Engine/Private/PhysicsEngine/BodyInstance.cpp
@@ -2526,6 +2526,7 @@ void FBodyInstance::UpdateInstanceSimulatePhysics()
 		FPhysicsInterface::SetIsKinematic_AssumesLocked(Actor, !bUseSimulate);
 		FPhysicsInterface::SetCcdEnabled_AssumesLocked(Actor, bUseCCD);
 		FPhysicsInterface::SetMACDEnabled_AssumesLocked(Actor, bUseMACD);
+		FPhysicsInterface::SetOneWayInteraction_AssumesLocked(Actor, bOneWayInteraction);
 		FPhysicsInterface::SetPositionSolverIterationCount_AssumesLocked(Actor, PositionSolverIterationCount);
 		FPhysicsInterface::SetVelocitySolverIterationCount_AssumesLocked(Actor, VelocitySolverIterationCount);
 		FPhysicsInterface::SetProjectionSolverIterationCount_AssumesLocked(Actor, ProjectionSolverIterationCount);
