[AutoRTFM] Make the FNetworkProfiler methods called by FRepLayout transaction-safe

Like RPCs, these functions defer the work until the transaction is committed.

diff --git a/Engine/Source/Runtime/Engine/Private/NetworkProfiler.cpp b/Engine/Source/Runtime/Engine/Private/NetworkProfiler.cpp
--- a/Engine/Source/Runtime/Engine/Private/NetworkProfiler.cpp
+++ b/Engine/Source/Runtime/Engine/Private/NetworkProfiler.cpp
@@ -5,6 +5,7 @@
 =============================================================================*/
 
 #include "Net/NetworkProfiler.h"
+#include "AutoRTFM/AutoRTFM.h"
 #include "HAL/FileManager.h"
 #include "HAL/IConsoleManager.h"
 #include "Misc/Paths.h"
@@ -594,6 +595,16 @@ void FNetworkProfiler::TrackReplicatePropertiesMetadata(const UObject* Object, T
 	}
 }
 
+void FNetworkProfiler::SerializeReplicateProperty( uint32 NameTableIndex, uint16 NumBits, UNetConnection* Connection )
+{
+	SetCurrentConnection( Connection );
+
+	uint8 Type = NPTYPE_ReplicateProperty;
+	(*FileWriter) << Type;
+	(*FileWriter).SerializeIntPacked( NameTableIndex );
+	(*FileWriter) << NumBits;
+}
+
 /**
  * Track property being replicated.
  *
@@ -604,33 +615,57 @@ void FNetworkProfiler::TrackReplicateProperty( const FProperty* Property, uint16
 {
 	if(bIsTrackingEnabled && !!!IgnorePropertyCount)
 	{
-		SCOPE_LOCK_REF(CriticalSection);
-
-		SetCurrentConnection( Connection );
+		UE::TScopeLock ScopeLock(CriticalSection);
 
 		uint32 NameTableIndex = GetNameTableIndex( Property->GetName() );
 
-		uint8 Type = NPTYPE_ReplicateProperty;
-		(*FileWriter) << Type;
-		(*FileWriter).SerializeIntPacked( NameTableIndex );
-		(*FileWriter) << NumBits;
+		if (AutoRTFM::IsClosed())
+		{
+			UE_AUTORTFM_ONCOMMIT(this, NameTableIndex, NumBits, Connection)
+			{
+				UE::TScopeLock ScopeLock2(CriticalSection);
+
+				SerializeReplicateProperty(NameTableIndex, NumBits, Connection);
+			};
+		}
+		else
+		{
+			SerializeReplicateProperty(NameTableIndex, NumBits, Connection);
+		}
 	}
 }
 
+void FNetworkProfiler::SerializeWritePropertyHeader( uint32 NameTableIndex, uint16 NumBits, UNetConnection* Connection )
+{
+	SetCurrentConnection( Connection );
+
+	uint8 Type = NPTYPE_WritePropertyHeader;
+	(*FileWriter) << Type;
+	(*FileWriter).SerializeIntPacked( NameTableIndex );
+	(*FileWriter) << NumBits;
+}
+
 void FNetworkProfiler::TrackWritePropertyHeader( const FProperty* Property, uint16 NumBits, UNetConnection* Connection )
 {
 	if( bIsTrackingEnabled )
 	{
-		SCOPE_LOCK_REF(CriticalSection);
-
-		SetCurrentConnection( Connection );
+		UE::TScopeLock ScopeLock(CriticalSection);
 
 		uint32 NameTableIndex = GetNameTableIndex( Property->GetName() );
 
-		uint8 Type = NPTYPE_WritePropertyHeader;
-		(*FileWriter) << Type;
-		(*FileWriter).SerializeIntPacked( NameTableIndex );
-		(*FileWriter) << NumBits;
+		if (AutoRTFM::IsClosed())
+		{
+			UE_AUTORTFM_ONCOMMIT(this, NameTableIndex, NumBits, Connection)
+			{
+				UE::TScopeLock ScopeLock2(CriticalSection);
+
+				SerializeWritePropertyHeader(NameTableIndex, NumBits, Connection);
+			};
+		}
+		else
+		{
+			SerializeWritePropertyHeader(NameTableIndex, NumBits, Connection);
+		}
 	}
 }
 
@@ -795,52 +830,93 @@ void FNetworkProfiler::TrackExportBunch( uint16 NumBits, UNetConnection* Connect
 	}
 }
 
- 
+void FNetworkProfiler::SerializeBeginContentBlock( uint32 NameTableIndex, uint16 NumBits, UNetConnection* Connection )
+{
+	SetCurrentConnection( Connection );
+
+	uint8 Type = NPTYPE_BeginContentBlock;
+	(*FileWriter) << Type;
+	(*FileWriter).SerializeIntPacked( NameTableIndex );
+	(*FileWriter) << NumBits;
+}
+
 void FNetworkProfiler::TrackBeginContentBlock( UObject* Object, uint16 NumBits, UNetConnection* Connection )
 {
 	if ( bIsTrackingEnabled )
 	{
-		SCOPE_LOCK_REF(CriticalSection);
-
-		SetCurrentConnection( Connection );
+		UE::TScopeLock ScopeLock(CriticalSection);
 
 		uint32 NameTableIndex = GetNameTableIndex( Object != nullptr ? Object->GetName() : UnknownName );
 
-		uint8 Type = NPTYPE_BeginContentBlock;
-		(*FileWriter) << Type;
-		(*FileWriter).SerializeIntPacked( NameTableIndex );
-		(*FileWriter) << NumBits;
+		if (AutoRTFM::IsClosed())
+		{
+			UE_AUTORTFM_ONCOMMIT(this, NameTableIndex, NumBits, Connection)
+			{
+				UE::TScopeLock ScopeLock2(CriticalSection);
+
+				SerializeBeginContentBlock(NameTableIndex, NumBits, Connection);
+			};
+		}
+		else
+		{
+			SerializeBeginContentBlock(NameTableIndex, NumBits, Connection);
+		}
 	 }
 }
 
+void FNetworkProfiler::SerializeEndContentBlock( uint32 NameTableIndex, uint16 NumBits, UNetConnection* Connection )
+{
+	SetCurrentConnection( Connection );
+
+	uint8 Type = NPTYPE_EndContentBlock;
+	(*FileWriter) << Type;
+	(*FileWriter).SerializeIntPacked( NameTableIndex );
+	(*FileWriter) << NumBits;
+}
+
 void FNetworkProfiler::TrackEndContentBlock( UObject* Object, uint16 NumBits, UNetConnection* Connection )
 {
 	if ( bIsTrackingEnabled )
 	{
-		SCOPE_LOCK_REF(CriticalSection);
-
-		SetCurrentConnection( Connection );
+		UE::TScopeLock ScopeLock(CriticalSection);
 
 		uint32 NameTableIndex = GetNameTableIndex( Object != nullptr ? Object->GetName() : UnknownName );
 
-		uint8 Type = NPTYPE_EndContentBlock;
-		(*FileWriter) << Type;
-		(*FileWriter).SerializeIntPacked( NameTableIndex );
-		(*FileWriter) << NumBits;
+		if (AutoRTFM::IsClosed())
+		{
+			UE_AUTORTFM_ONCOMMIT(this, NameTableIndex, NumBits, Connection)
+			{
+				UE::TScopeLock ScopeLock2(CriticalSection);
+
+				SerializeEndContentBlock(NameTableIndex, NumBits, Connection);
+			};
+		}
+		else
+		{
+			SerializeEndContentBlock(NameTableIndex, NumBits, Connection);
+		}
 	}
 }
 
+void FNetworkProfiler::SerializeWritePropertyHandle( uint16 NumBits, UNetConnection* Connection )
+{
+	SetCurrentConnection( Connection );
+
+	uint8 Type = NPTYPE_WritePropertyHandle;
+	(*FileWriter) << Type;
+	(*FileWriter) << NumBits;
+}
+
 void FNetworkProfiler::TrackWritePropertyHandle( uint16 NumBits, UNetConnection* Connection )
 {
 	if (bIsTrackingEnabled && !!!IgnorePropertyCount)
 	{
-		SCOPE_LOCK_REF(CriticalSection);
-
-		SetCurrentConnection( Connection );
+		UE_AUTORTFM_ONCOMMIT(this, NumBits, Connection)
+		{
+			SCOPE_LOCK_REF(CriticalSection);
 
-		uint8 Type = NPTYPE_WritePropertyHandle;
-		(*FileWriter) << Type;
-		(*FileWriter) << NumBits;
+			SerializeWritePropertyHandle(NumBits, Connection);
+		};
 	}
 }
 
diff --git a/Engine/Source/Runtime/Engine/Private/Tests/AutoRTFM/AutoRTFMNetworkProfilerTests.cpp b/Engine/Source/Runtime/Engine/Private/Tests/AutoRTFM/AutoRTFMNetworkProfilerTests.cpp
new file mode 100644
--- /dev/null
+++ b/Engine/Source/Runtime/Engine/Private/Tests/AutoRTFM/AutoRTFMNetworkProfilerTests.cpp
@@ -0,0 +1,98 @@
+// Copyright Epic Games, Inc. All Rights Reserved.
+
+#include "AutoRTFM/AutoRTFMTesting.h"
+#include "AutoRTFMTestActor.h"
+#include "Engine/ActorChannel.h"
+#include "Engine/DemoNetDriver.h"
+#include "Engine/NetDriver.h"
+#include "Misc/AutomationTest.h"
+#include "AutoRTFM/AutoRTFM.h"
+#include "Net/NetworkProfiler.h"
+#include "UObject/UnrealType.h"
+
+#if WITH_DEV_AUTOMATION_TESTS
+
+IMPLEMENT_SIMPLE_AUTOMATION_TEST(FAutoRTFMNetProfilerTests, "AutoRTFM + FNetworkProfiler", EAutomationTestFlags::EngineFilter | EAutomationTestFlags::ClientContext | EAutomationTestFlags::ServerContext | EAutomationTestFlags::CommandletContext)
+
+bool FAutoRTFMNetProfilerTests::RunTest(const FString & Parameters)
+{
+	if (!AutoRTFM::ForTheRuntime::IsAutoRTFMRuntimeEnabled())
+	{
+		ExecutionInfo.AddEvent(FAutomationEvent(EAutomationEventType::Info, TEXT("SKIPPED 'FAutoRTFMNetProfilerTests' test. AutoRTFM disabled.")));
+		return true;
+	}
+
+	UNetDriver* const Driver = NewObject<UDemoNetDriver>();
+
+	UNetConnection* const Connection = NewObject<UDemoNetConnection>();
+	Connection->Driver = Driver;
+	Driver->AddClientConnection(Connection);
+
+	UActorChannel* const ActorChannel = NewObject<UActorChannel>();
+	ActorChannel->OpenedLocally = true;
+	ActorChannel->Connection = Connection;
+	Connection->Channels.Add(ActorChannel);
+	Connection->OpenChannels.Add(ActorChannel);
+	
+	AAutoRTFMTestActor* const Actor = NewObject<AAutoRTFMTestActor>();
+
+	FProperty* const Property = AAutoRTFMTestActor::StaticClass()->FindPropertyByName(GET_MEMBER_NAME_CHECKED(AAutoRTFMTestActor, MyProperty));
+
+	FNetworkProfiler Profiler;
+	Profiler.EnableTracking(/* bShouldEnableTracking */ true);
+	Profiler.TrackSessionChange(/* bShouldContinueTracking */ true, FURL(TEXT("FAutoRTFMNetProfilerTests")));
+
+	AutoRTFM::Testing::Commit([&]
+	{
+		Profiler.TrackWritePropertyHandle(16, Connection);
+	});
+	AutoRTFM::Testing::Abort([&]
+	{
+		Profiler.TrackWritePropertyHandle(16, Connection);
+		AutoRTFM::AbortTransaction();
+	});
+
+	AutoRTFM::Testing::Commit([&]
+	{
+		Profiler.TrackWritePropertyHeader(Property, 16, Connection);
+	});
+	AutoRTFM::Testing::Abort([&]
+	{
+		Profiler.TrackWritePropertyHeader(Property, 16, Connection);
+		AutoRTFM::AbortTransaction();
+	});
+
+	AutoRTFM::Testing::Commit([&]
+	{
+		Profiler.TrackReplicateProperty(Property, 16, Connection);
+	});
+	AutoRTFM::Testing::Abort([&]
+	{
+		Profiler.TrackReplicateProperty(Property, 16, Connection);
+		AutoRTFM::AbortTransaction();
+	});
+
+	AutoRTFM::Testing::Commit([&]
+	{
+		Profiler.TrackBeginContentBlock(Actor, 16, Connection);
+	});
+	AutoRTFM::Testing::Abort([&]
+	{
+		Profiler.TrackBeginContentBlock(Actor, 16, Connection);
+		AutoRTFM::AbortTransaction();
+	});
+
+	AutoRTFM::Testing::Commit([&]
+	{
+		Profiler.TrackEndContentBlock(Actor, 16, Connection);
+	});
+	AutoRTFM::Testing::Abort([&]
+	{
+		Profiler.TrackEndContentBlock(Actor, 16, Connection);
+		AutoRTFM::AbortTransaction();
+	});
+
+	return true;
+}
+
+#endif //WITH_DEV_AUTOMATION_TESTS
diff --git a/Engine/Source/Runtime/Engine/Private/Tests/AutoRTFM/AutoRTFMTestActor.h b/Engine/Source/Runtime/Engine/Private/Tests/AutoRTFM/AutoRTFMTestActor.h
--- a/Engine/Source/Runtime/Engine/Private/Tests/AutoRTFM/AutoRTFMTestActor.h
+++ b/Engine/Source/Runtime/Engine/Private/Tests/AutoRTFM/AutoRTFMTestActor.h
@@ -6,6 +6,7 @@
 #include "Misc/MTTransactionallySafeAccessDetector.h"
 #include "Misc/TransactionallySafeCriticalSection.h"
 #include "Misc/TransactionallySafeRWLock.h"
+#include "UObject/ObjectMacros.h"
 #include "AutoRTFMTestActor.generated.h"
 
 UCLASS()
@@ -14,8 +15,11 @@ class AAutoRTFMTestActor : public AActor
     GENERATED_BODY()
 
 public:
-    int Value = 42;
-    
+	int Value = 42;
+
+	UPROPERTY()
+	int32 MyProperty = 0;
+
 	// Various AutoRTFM-complex member fields, so that the tests exercise their
 	// ctor / dtors in interesting ways.
 	FTransactionallySafeCriticalSection CriticalSection;
diff --git a/Engine/Source/Runtime/Engine/Public/Net/NetworkProfiler.h b/Engine/Source/Runtime/Engine/Public/Net/NetworkProfiler.h
--- a/Engine/Source/Runtime/Engine/Public/Net/NetworkProfiler.h
+++ b/Engine/Source/Runtime/Engine/Public/Net/NetworkProfiler.h
@@ -201,6 +201,49 @@ private:
 
 	ENGINE_API void AutoStopTracking();
 
+	/**
+	 * Sets the current connection to Connection and serializes the property.
+	 * The CriticalSection lock must be held while calling.
+	 *
+	 * @param	NameTableIndex	Name of the property being replicated
+	 * @param	NumBits			Number of bits used to replicate this property
+	 */
+	void SerializeReplicateProperty( uint32 NameTableIndex, uint16 NumBits, UNetConnection* Connection );
+	
+	/**
+	 * Sets the current connection to Connection and serializes the property header.
+	 * The CriticalSection lock must be held while calling.
+	 *
+	 * @param	NameTableIndex	Name of the property being replicated
+	 * @param	NumBits			Number of bits used in the header
+	 */
+	void SerializeWritePropertyHeader( uint32 NameTableIndex, uint16 NumBits, UNetConnection* Connection );
+
+	/**
+	 * Sets the current connection to Connection and serializes the property handle.
+	 * The CriticalSection lock must be held while calling.
+	 *
+	 * @param NumBits Number of bits in the property handle
+	 */
+	void SerializeWritePropertyHandle( uint16 NumBits, UNetConnection* Connection );
+
+	/**
+	 * Sets the current connection to Connection and serializes the start of actor content block headers
+	 *
+	 * @param NameTableIndex the name of the object being replicated (might be a subobject of the actor)
+	 * @param NumBits the number of bits in the content block header
+	 */
+	void SerializeBeginContentBlock( uint32 NameTableIndex, uint16 NumBits, UNetConnection* Connection );
+
+	/**
+	 * Sets the current connection to Connection and serializes the end of actor content block headers
+	 *
+	 * @param NameTableIndex the name of the object being replicated (might be a subobject of the actor)
+	 * @param NumBits the number of bits in the content block header
+	 */
+	void SerializeEndContentBlock( uint32 NameTableIndex, uint16 NumBits, UNetConnection* Connection );
+ 
+	
 	/** The filename to use next time a file writer is opened.  If empty then an autogenerated name will be used. */
 	FString NextFileName;
 
@@ -474,7 +517,8 @@ public:
 	 */
 	ENGINE_API void TrackEndContentBlock( UObject* Object, uint16 NumBits, UNetConnection* Connection );
 
-	/** Track property handles
+	/**
+	 * Track property handles
 	 *
 	 * @param NumBits Number of bits in the property handle
 	 */
