Protect against bad  access of audio parameters
#rnx

diff --git a/Engine/Plugins/Runtime/Metasound/Source/MetasoundFrontend/Private/MetasoundParameterTransmitter.cpp b/Engine/Plugins/Runtime/Metasound/Source/MetasoundFrontend/Private/MetasoundParameterTransmitter.cpp
--- a/Engine/Plugins/Runtime/Metasound/Source/MetasoundFrontend/Private/MetasoundParameterTransmitter.cpp
+++ b/Engine/Plugins/Runtime/Metasound/Source/MetasoundFrontend/Private/MetasoundParameterTransmitter.cpp
@@ -251,13 +251,6 @@ namespace Metasound
 	{
 	}
 
-
-	bool FMetaSoundParameterTransmitter::Reset()
-	{
-		OnDeleteActiveSound();
-		return true;
-	}
-
 	void FMetaSoundParameterTransmitter::AddAvailableParameter(FName InName)
 	{
 		AvailableParameterNames.Add(InName);
@@ -270,6 +263,8 @@ namespace Metasound
 
 	bool FMetaSoundParameterTransmitter::SetVirtualizedParameters(TArray<FAudioParameter>&& InParameters)
 	{
+		check(IsInAudioThread());
+
 		bool bSuccess = true;
 
 		// Remove triggers
@@ -296,9 +291,11 @@ namespace Metasound
 
 	bool FMetaSoundParameterTransmitter::SetParameters(TArray<FAudioParameter>&& InParameters)
 	{
+		check(IsInAudioThread());
+
 		// Don't set parameters directly if the active sound 
 		// is currently virtualized (to prevent accumulation of unneeded updates) 
-		if (bIsVirtualized)
+		if (IsVirtualized())
 		{
 			return SetVirtualizedParameters(MoveTemp(InParameters));
 		}
diff --git a/Engine/Plugins/Runtime/Metasound/Source/MetasoundFrontend/Public/MetasoundParameterTransmitter.h b/Engine/Plugins/Runtime/Metasound/Source/MetasoundFrontend/Public/MetasoundParameterTransmitter.h
--- a/Engine/Plugins/Runtime/Metasound/Source/MetasoundFrontend/Public/MetasoundParameterTransmitter.h
+++ b/Engine/Plugins/Runtime/Metasound/Source/MetasoundFrontend/Public/MetasoundParameterTransmitter.h
@@ -86,18 +86,6 @@ namespace Metasound
 			/** Shared queue used to communicate with running MetaSound instance. */
 			TSharedPtr<TSpscQueue<FParameter>> DataChannel;
 
-			UE_DEPRECATED(5.3, "This constructor is deprecated and will be removed.")
-			FInitParams(const FOperatorSettings& InSettings, uint64 InInstanceID, TArray<FAudioParameter>&& InDefaultParams, const TArray<FSendInfo>& InInfos=TArray<FSendInfo>())
-			: OperatorSettings(InSettings)
-			, InstanceID(InInstanceID)
-			, DefaultParams(MoveTemp(InDefaultParams))
-			{
-				for (const FSendInfo& SendInfo : InInfos)
-				{
-					ValidParameterNames.Add(SendInfo.ParameterName);
-				}
-			}
-
 			FInitParams(const FOperatorSettings& InSettings, uint64 InInstanceID, TArray<FAudioParameter>&& InDefaultParams, TArray<FName>&& InValidParameterNames, TSharedPtr<TSpscQueue<FParameter>> InDataChannel)
 			: OperatorSettings(InSettings)
 			, InstanceID(InInstanceID)
@@ -119,10 +107,6 @@ namespace Metasound
 		FMetaSoundParameterTransmitter(FMetaSoundParameterTransmitter::FInitParams&& InInitParams);
 		virtual ~FMetaSoundParameterTransmitter() = default;
 
-
-		UE_DEPRECATED(5.2, "Use ResetParameters() or OnDeleteActiveSound() instead depending on use case.")
-		bool Reset() override;
-		
 		void AddAvailableParameter(FName InName);
 		void RemoveAvailableParameter(FName InName);
 
diff --git a/Engine/Source/Runtime/AudioExtensions/Private/IAudioParameterTransmitter.cpp b/Engine/Source/Runtime/AudioExtensions/Private/IAudioParameterTransmitter.cpp
--- a/Engine/Source/Runtime/AudioExtensions/Private/IAudioParameterTransmitter.cpp
+++ b/Engine/Source/Runtime/AudioExtensions/Private/IAudioParameterTransmitter.cpp
@@ -12,18 +12,26 @@ namespace Audio
 		return { };
 	}
 
-	FParameterTransmitterBase::FParameterTransmitterBase(TArray<FAudioParameter>&& InDefaultParams)
+	void ILegacyParameterTransmitter::AddReferencedObjects(FReferenceCollector& InCollector)
+	{
+	}
+
+PRAGMA_DISABLE_DEPRECATION_WARNINGS // deprecation is because AudioParameters will become private
+	FParameterTransmitterBase::FParameterTransmitterBase(TArray<FAudioParameter> InDefaultParams)
 		: AudioParameters(MoveTemp(InDefaultParams))
 		, bIsVirtualized(false)
 	{
 	}
 
 	FParameterTransmitterBase::~FParameterTransmitterBase() = default;
+PRAGMA_ENABLE_DEPRECATION_WARNINGS
 
 	bool FParameterTransmitterBase::GetParameter(FName InName, FAudioParameter& OutValue) const
 	{
-		FScopeLock Lock(&AudioParameterCritSec);
+		check(IsInAudioThread());
+PRAGMA_DISABLE_DEPRECATION_WARNINGS // deprecation is because AudioParameters will become private
 		if (const FAudioParameter* Param = FAudioParameter::FindParam(AudioParameters, InName))
+PRAGMA_ENABLE_DEPRECATION_WARNINGS
 		{
 			OutValue = *Param;
 			return true;
@@ -34,39 +42,85 @@ namespace Audio
 
 	void FParameterTransmitterBase::ResetParameters() 
 	{
-		FScopeLock Lock(&AudioParameterCritSec);
+		check(IsInAudioThread());
 
+PRAGMA_DISABLE_DEPRECATION_WARNINGS // deprecation is because AudioParameters will become private
 		AudioParameters.Reset();
-	}
-	
-	bool FParameterTransmitterBase::Reset()
-	{
-		ResetParameters();
-		return true;
+PRAGMA_ENABLE_DEPRECATION_WARNINGS
 	}
 
 	const TArray<FAudioParameter>& FParameterTransmitterBase::GetParameters() const
 	{
+		check(IsInAudioThread());
+PRAGMA_DISABLE_DEPRECATION_WARNINGS // deprecation is because AudioParameters will become private
 		return AudioParameters;
+PRAGMA_ENABLE_DEPRECATION_WARNINGS
 	}
 
-	bool FParameterTransmitterBase::SetParameters(TArray<FAudioParameter>&& InParameters)
+	void FParameterTransmitterBase::CopyParameters(TArray<FAudioParameter>& OutParameters) const
 	{
-		FScopeLock Lock(&AudioParameterCritSec);
+PRAGMA_DISABLE_DEPRECATION_WARNINGS // deprecation is because AudioParameters will become private
+		OutParameters = AudioParameters;
+PRAGMA_ENABLE_DEPRECATION_WARNINGS
+	}
 
+	bool FParameterTransmitterBase::SetParameters(TArray<FAudioParameter>&& InParameters)
+	{
+		check(IsInAudioThread());
+PRAGMA_DISABLE_DEPRECATION_WARNINGS // deprecation is because AudioParameters will become private
 		FAudioParameter::Merge(MoveTemp(InParameters), AudioParameters);
+PRAGMA_ENABLE_DEPRECATION_WARNINGS
 		return true;
 	}
 
+	void FParameterTransmitterBase::AddReferencedObjects(FReferenceCollector& InCollector)
+	{
+PRAGMA_DISABLE_DEPRECATION_WARNINGS // deprecation is because AudioParameters will become private
+		AddReferencedObjectsFromParameters(InCollector, AudioParameters);
+PRAGMA_ENABLE_DEPRECATION_WARNINGS
+	}
+
 	void FParameterTransmitterBase::OnVirtualizeActiveSound()
 	{ 
+		check(IsInAudioThread());
+PRAGMA_DISABLE_DEPRECATION_WARNINGS // deprecation is because bIsVirtualized will become private
 		bIsVirtualized = true;
-		ResetParameters();
+PRAGMA_ENABLE_DEPRECATION_WARNINGS
 	}
 
 	 void FParameterTransmitterBase::OnRealizeVirtualizedActiveSound(TArray<FAudioParameter>&& InParameters)
 	 { 
-		 bIsVirtualized = false;
-		 SetParameters(MoveTemp(InParameters));
+		check(IsInAudioThread());
+PRAGMA_DISABLE_DEPRECATION_WARNINGS // deprecation is because bIsVirtualized will become private
+		bIsVirtualized = false;
+PRAGMA_ENABLE_DEPRECATION_WARNINGS
+		SetParameters(MoveTemp(InParameters));
 	 }
+
+	 bool FParameterTransmitterBase::IsVirtualized() const
+	 {
+		check(IsInAudioThread());
+PRAGMA_DISABLE_DEPRECATION_WARNINGS // deprecation is because bIsVirtualized will become private
+		return bIsVirtualized;
+PRAGMA_ENABLE_DEPRECATION_WARNINGS
+	 }
+
+	void FParameterTransmitterBase::AddReferencedObjectsFromParameters(FReferenceCollector& InCollector, TArrayView<FAudioParameter> InParameters) const
+	{
+		for (FAudioParameter& Param : InParameters)
+		{
+			if (Param.ObjectParam)
+			{
+				InCollector.AddReferencedObject(Param.ObjectParam);
+			}
+
+			for (TObjectPtr<UObject>& Object : Param.ArrayObjectParam)
+			{
+				if (Object)
+				{
+					InCollector.AddReferencedObject(Object);
+				}
+			}
+		}
+	}
 } // namespace Audio
diff --git a/Engine/Source/Runtime/AudioExtensions/Public/IAudioParameterTransmitter.h b/Engine/Source/Runtime/AudioExtensions/Public/IAudioParameterTransmitter.h
--- a/Engine/Source/Runtime/AudioExtensions/Public/IAudioParameterTransmitter.h
+++ b/Engine/Source/Runtime/AudioExtensions/Public/IAudioParameterTransmitter.h
@@ -11,6 +11,7 @@
 
 // Forward Declarations
 class UObject;
+class FReferenceCollector;
 
 namespace Audio
 {
@@ -43,7 +44,10 @@ namespace Audio
 		public:
 			virtual ~ILegacyParameterTransmitter() = default;
 
+			UE_DEPRECATED(5.6, "Use AddReferencedObjects(...) instead")
 			AUDIOEXTENSIONS_API virtual TArray<const TObjectPtr<UObject>*> GetReferencedObjects() const;
+
+			AUDIOEXTENSIONS_API virtual void AddReferencedObjects(FReferenceCollector& InCollector);
 	};
 
 	/** Interface for a audio instance transmitter.
@@ -57,9 +61,6 @@ namespace Audio
 
 			virtual ~IParameterTransmitter() = default;
 
-			UE_DEPRECATED(5.2, "Use ResetParameters() or OnDeleteActiveSound() instead depending on use case.")
-			virtual bool Reset() { ResetParameters(); return true; }
-
 			// Reset parameters which stored on the transmitter.
 			virtual void ResetParameters() {}
 
@@ -72,8 +73,11 @@ namespace Audio
 			virtual bool GetParameter(FName InName, FAudioParameter& OutParam) const = 0;
 
 			// Return reference to the cached parameter array.
+			UE_DEPRECATED(5.6, "Access to internal parameters is no longer supported")
 			virtual const TArray<FAudioParameter>& GetParameters() const = 0;
 
+			virtual void CopyParameters(TArray<FAudioParameter>& OutParameters) const = 0;
+
 			// Parameter Setters
 			virtual bool SetParameters(TArray<FAudioParameter>&& InParameters) = 0;
 
@@ -91,21 +95,33 @@ namespace Audio
 	class FParameterTransmitterBase : public IParameterTransmitter
 	{
 	public:
-		AUDIOEXTENSIONS_API FParameterTransmitterBase(TArray<FAudioParameter>&& InDefaultParams);
+		AUDIOEXTENSIONS_API FParameterTransmitterBase(TArray<FAudioParameter> InDefaultParams);
 		AUDIOEXTENSIONS_API virtual ~FParameterTransmitterBase();
 
 		AUDIOEXTENSIONS_API virtual bool GetParameter(FName InName, FAudioParameter& OutParam) const override;
 		AUDIOEXTENSIONS_API virtual void ResetParameters() override;
+
+		UE_DEPRECATED(5.6, "Access to internal parameter storage is no longer supported. Use CopyParameters(...) instead")
 		AUDIOEXTENSIONS_API virtual const TArray<FAudioParameter>& GetParameters() const override;
-		UE_DEPRECATED(5.2, "Use ResetParameters() or OnDeleteActiveSound() instead depending on use case.")
-		AUDIOEXTENSIONS_API virtual bool Reset() override;
+
+		AUDIOEXTENSIONS_API virtual void CopyParameters(TArray<FAudioParameter>& OutParameters) const override;
+
 		AUDIOEXTENSIONS_API virtual bool SetParameters(TArray<FAudioParameter>&& InParameters) override;
+
+		AUDIOEXTENSIONS_API virtual void AddReferencedObjects(FReferenceCollector& InCollector) override;
+
 		AUDIOEXTENSIONS_API virtual void OnVirtualizeActiveSound() override;
 		AUDIOEXTENSIONS_API virtual void OnRealizeVirtualizedActiveSound(TArray<FAudioParameter>&& InParameters);
+		AUDIOEXTENSIONS_API bool IsVirtualized() const;
 
 	protected:
-		mutable FCriticalSection AudioParameterCritSec; 
+
+		AUDIOEXTENSIONS_API void AddReferencedObjectsFromParameters(FReferenceCollector& InCollector, TArrayView<FAudioParameter> InParameters) const;
+
+		UE_DEPRECATED(5.6, "Subclasses of FParameterTransmitterBase should not access AudioParameters member directly. These will be moved to private visibility soon.")
 		TArray<FAudioParameter> AudioParameters;
+		UE_DEPRECATED(5.6, "Subclasses of FParameterTransmitterBase should not access bIsVirtualized member directly. These will be moved to private visibility soon.")
 		bool bIsVirtualized;
+
 	};
 } // namespace Audio
diff --git a/Engine/Source/Runtime/AudioMixer/Private/AudioMixerSource.cpp b/Engine/Source/Runtime/AudioMixer/Private/AudioMixerSource.cpp
--- a/Engine/Source/Runtime/AudioMixer/Private/AudioMixerSource.cpp
+++ b/Engine/Source/Runtime/AudioMixer/Private/AudioMixerSource.cpp
@@ -1184,8 +1184,10 @@ namespace Audio
 			bActiveSoundIsPreviewSound = ActiveSound->bIsPreviewSound;
 			if (Audio::IParameterTransmitter* Transmitter = ActiveSound->GetTransmitter())
 			{
-				DefaultParameters = Transmitter->GetParameters();
-				SoundWave.InitParameters(DefaultParameters);
+				// This copying of parameters is for the case of virtual loop realization. 
+				// The most up-to-date parameters exist on the instance transmitter. 
+				Transmitter->CopyParameters(DefaultParameters);
+ 				SoundWave.InitParameters(DefaultParameters);
 			}
 		}
 
diff --git a/Engine/Source/Runtime/Engine/Classes/Sound/SoundCue.h b/Engine/Source/Runtime/Engine/Classes/Sound/SoundCue.h
--- a/Engine/Source/Runtime/Engine/Classes/Sound/SoundCue.h
+++ b/Engine/Source/Runtime/Engine/Classes/Sound/SoundCue.h
@@ -349,19 +349,31 @@ private:
 
 class FSoundCueParameterTransmitter : public Audio::FParameterTransmitterBase
 {
+	using Super = Audio::FParameterTransmitterBase;
 public:
-	FSoundCueParameterTransmitter(Audio::FParameterTransmitterInitParams&& InParams)
-		: Audio::FParameterTransmitterBase(MoveTemp(InParams.DefaultParams))
-	{
-	}
+	ENGINE_API FSoundCueParameterTransmitter(Audio::FParameterTransmitterInitParams&& InParams);
 
-	virtual ~FSoundCueParameterTransmitter() = default;
+	ENGINE_API virtual ~FSoundCueParameterTransmitter();
 
+	UE_DEPRECATED(5.6, "Use AddReferencedObjects(...) instead")
 	ENGINE_API TArray<const TObjectPtr<UObject>*> GetReferencedObjects() const override;
 
+	ENGINE_API virtual void AddReferencedObjects(FReferenceCollector& InCollector) override;
+
 	ENGINE_API virtual bool SetParameters(TArray<FAudioParameter>&& InParameters) override;
 
+	ENGINE_API void CreateChildTransmitter(UPTRINT InWaveInstanceHash, USoundBase& InSound, Audio::FParameterTransmitterInitParams InTransmitterInitParams);
+
+	ENGINE_API void RemoveChildTransmitter(UPTRINT InWaveInstanceHash);
+
+	/** The SoundCueParameter keeps a running track of parameter changes from SetParameters(...). 
+	 * When this method is called, the accumulated parameters are returned and the internal
+	 * tracking of accumulated parameters is reset. */
+	ENGINE_API TArray<FAudioParameter> ReleaseAccumulatedParameterUpdates();
+
+	UE_DEPRECATED(5.6, "Do not access ParamsToSet. This will be moved to be a private member.")
 	TArray<FAudioParameter> ParamsToSet;
 
+	UE_DEPRECATED(5.6, "Do not access Transmitters. This will be moved to be a private member.")
 	TMap<UPTRINT, TSharedPtr<Audio::IParameterTransmitter>> Transmitters;
 };
diff --git a/Engine/Source/Runtime/Engine/Private/ActiveSound.cpp b/Engine/Source/Runtime/Engine/Private/ActiveSound.cpp
--- a/Engine/Source/Runtime/Engine/Private/ActiveSound.cpp
+++ b/Engine/Source/Runtime/Engine/Private/ActiveSound.cpp
@@ -268,10 +268,7 @@ void FActiveSound::AddReferencedObjects(FReferenceCollector& Collector)
 
 	if (InstanceTransmitter.IsValid())
 	{
-		for (auto* Object : InstanceTransmitter->GetReferencedObjects())
-		{
-			Collector.AddReferencedObject(const_cast<TObjectPtr<UObject>&>(*Object));
-		}
+		InstanceTransmitter->AddReferencedObjects(Collector);
 	}
 
 	// Adding ref'd objects in AttenuationSettings
diff --git a/Engine/Source/Runtime/Engine/Private/AudioDevice.cpp b/Engine/Source/Runtime/Engine/Private/AudioDevice.cpp
--- a/Engine/Source/Runtime/Engine/Private/AudioDevice.cpp
+++ b/Engine/Source/Runtime/Engine/Private/AudioDevice.cpp
@@ -5342,27 +5342,9 @@ void FAudioDevice::InitSoundParams(FActiveSound& InOutActiveSound, TArray<FAudio
 	}
 	else if (bInIsVirtualLoopRealizing)
 	{
-		// Gather parameters from audio component to set on virtual loop when it restarts
-		TArray<FAudioParameter> InstanceParams;
-		uint64 AudioComponentID = InOutActiveSound.GetAudioComponentID();
-		if (AudioComponentID > 0)
-		{
-			if (UAudioComponent* AudioComponent = UAudioComponent::GetAudioComponentFromID(AudioComponentID))
-			{
-				if (ensure(AudioComponent->Sound != nullptr))
-				{
-					// Use instance parameters which will not contain transient parameters
-					InstanceParams = AudioComponent->GetInstanceParameters();
-
-					// create proxies from the UObject references in InstanceParams
-					AudioComponent->Sound->InitParameters(InstanceParams);
-				}
-			}
-		}
-		
-		FAudioParameter::Merge(MoveTemp(InDefaultParams), InstanceParams);
-		
-		InOutActiveSound.InstanceTransmitter->OnRealizeVirtualizedActiveSound(MoveTemp(InstanceParams));
+		// The default parameters should be empty when realizing a virtualized loop. Prior
+		// parameters are already set on the instance transmitter. 
+		InOutActiveSound.InstanceTransmitter->OnRealizeVirtualizedActiveSound(MoveTemp(InDefaultParams));
 	}
 	else
 	{
diff --git a/Engine/Source/Runtime/Engine/Private/AudioDeviceManager.cpp b/Engine/Source/Runtime/Engine/Private/AudioDeviceManager.cpp
--- a/Engine/Source/Runtime/Engine/Private/AudioDeviceManager.cpp
+++ b/Engine/Source/Runtime/Engine/Private/AudioDeviceManager.cpp
@@ -5,6 +5,7 @@
 #include "Audio/AudioDebug.h"
 #include "AudioAnalytics.h"
 #include "AudioMixerDevice.h"
+#include "CoreGlobals.h"
 #include "Engine/Engine.h"
 #include "Engine/World.h"
 #include "Features/IModularFeatures.h"
@@ -831,6 +832,21 @@ void FAudioDeviceManager::IterateOverAllDevices(TUniqueFunction<void(Audio::FDev
 
 void FAudioDeviceManager::AddReferencedObjects(FReferenceCollector& Collector)
 {
+#if !WITH_EDITORONLY_DATA
+	// Audio object references are updated while the audio thread is running and so
+	// we need to make sure that the audio thread is not running while we collect
+	// object references. The AudioThread generally only exists in packaged games
+	// and it has mechanisms for pausing/resuming. It is triggered by the start 
+	// and stop of GarbageCollection and so for most scenarios we are safe here. 
+	//
+	// One exception to the rule is StandaloneGame which launches the AudioThread 
+	// but is not packaged. There are several BP fixup and redirector calls that
+	// occur on level load which hit this callstack even though they are not part
+	// of GarbageCollection. This inherent race condition has existed for at least
+	// 10 years and has not caused known issues. It remains as tech-debt. 
+	checkf(!IsAudioThreadRunning(), TEXT("The audio thread must be disabled or suspended while collecting object references"));
+#endif
+
 	IterateOverAllDevices(
 		[&Collector](Audio::FDeviceId, FAudioDevice* InDevice)
 		{
diff --git a/Engine/Source/Runtime/Engine/Private/SoundCue.cpp b/Engine/Source/Runtime/Engine/Private/SoundCue.cpp
--- a/Engine/Source/Runtime/Engine/Private/SoundCue.cpp
+++ b/Engine/Source/Runtime/Engine/Private/SoundCue.cpp
@@ -667,6 +667,8 @@ bool USoundCue::IsPlayWhenSilent() const
 
 void USoundCue::Parse(FAudioDevice* AudioDevice, const UPTRINT NodeWaveInstanceHash, FActiveSound& ActiveSound, const FSoundParseParameters& ParseParams, TArray<FWaveInstance*>& WaveInstances)
 {
+	using namespace Audio;
+
 	if (FirstNode)
 	{
 		FirstNode->ParseNodes(AudioDevice, (UPTRINT)FirstNode, ActiveSound, ParseParams, WaveInstances);
@@ -674,11 +676,26 @@ void USoundCue::Parse(FAudioDevice* AudioDevice, const UPTRINT NodeWaveInstanceH
 
 	if (FSoundCueParameterTransmitter* Transmitter = static_cast<FSoundCueParameterTransmitter*>(ActiveSound.GetTransmitter()))
 	{
-		if (Transmitter->ParamsToSet.IsEmpty())
+		TArray<FAudioParameter> ParameterUpdates = Transmitter->ReleaseAccumulatedParameterUpdates();
+		if (ParameterUpdates.IsEmpty())
 		{
 			return;
 		}
 
+		auto UpdateChildParameters = [](IParameterTransmitter& InChildTransmitter, USoundBase* InSound, TArray<FAudioParameter> InParameters)
+		{
+			if (InSound)
+			{
+				InSound->InitParameters(InParameters);
+			}
+
+			if (!InParameters.IsEmpty())
+			{
+				InChildTransmitter.SetParameters(MoveTemp(InParameters));
+			}
+		};
+
+		const FWaveInstance* LastInstance = WaveInstances.Last();
 		for (const FWaveInstance* Instance : WaveInstances)
 		{
 			if (!Instance)
@@ -686,28 +703,27 @@ void USoundCue::Parse(FAudioDevice* AudioDevice, const UPTRINT NodeWaveInstanceH
 				continue;
 			}
 			
+PRAGMA_DISABLE_DEPRECATION_WARNINGS	
 			if (TSharedPtr<Audio::IParameterTransmitter>* ChildTransmitterPtr = Transmitter->Transmitters.Find(Instance->WaveInstanceHash))
+PRAGMA_ENABLE_DEPRECATION_WARNINGS
 			{
 				TSharedPtr<Audio::IParameterTransmitter>& ChildTransmitter = *ChildTransmitterPtr;
 				
-				if (ChildTransmitter.IsValid())
+				if (ChildTransmitterPtr->IsValid())
 				{
-					TArray<FAudioParameter> Params = Transmitter->ParamsToSet;
-
-					if (USoundWave* Sound = Instance->WaveData)
+					if (Instance == LastInstance)
 					{
-						Sound->InitParameters(Params);
+						// Avoid extra copy on last wave instance by moving parameters. 
+						// This also avoids a copy on simple case where there is only one wave instance. 
+						UpdateChildParameters(*ChildTransmitterPtr->Get(), Instance->WaveData, MoveTemp(ParameterUpdates));
 					}
-
-					if (!Params.IsEmpty())
+					else
 					{
-						ChildTransmitter->SetParameters(MoveTemp(Params));
+						UpdateChildParameters(*ChildTransmitterPtr->Get(), Instance->WaveData, ParameterUpdates);
 					}
 				}
 			}
 		}
-		
-		Transmitter->ParamsToSet.Reset();
 	}
 }
 
@@ -906,27 +922,38 @@ TSharedPtr<ISoundCueAudioEditor> USoundCue::GetSoundCueAudioEditor()
 	return SoundCueAudioEditor;
 }
 #endif // WITH_EDITOR
+	   
+PRAGMA_DISABLE_DEPRECATION_WARNINGS // Disabling deprecation warnings for construction/destruction of deprecated member variables. 
+FSoundCueParameterTransmitter::FSoundCueParameterTransmitter(Audio::FParameterTransmitterInitParams&& InParams)
+	: Audio::FParameterTransmitterBase(MoveTemp(InParams.DefaultParams))
+{
+}
+PRAGMA_ENABLE_DEPRECATION_WARNINGS
+
+PRAGMA_DISABLE_DEPRECATION_WARNINGS // Disabling deprecation warnings for construction/destruction of deprecated member variables.
+FSoundCueParameterTransmitter::~FSoundCueParameterTransmitter()
+{
+}
+PRAGMA_ENABLE_DEPRECATION_WARNINGS
 
 TArray<const TObjectPtr<UObject>*> FSoundCueParameterTransmitter::GetReferencedObjects() const
 {
-	TArray<const TObjectPtr<UObject>*> Objects;
-	for (const FAudioParameter& Param : AudioParameters)
+	static bool bHasLoggedError = false;
+	if (!bHasLoggedError)
 	{
-		if (Param.ObjectParam)
-		{
-			Objects.Add(&Param.ObjectParam);
-		}
-
-		for (const auto& Object : Param.ArrayObjectParam)
-		{
-			if (Object)
-			{
-				Objects.Add(&Object);
-			}
-		}
+		UE_LOG(LogAudio, Error, TEXT("Do not call FSoundCueParameterTransmitter::GetReferencedObjects() as it is not thread safe."));
+		bHasLoggedError = true;
 	}
 
-	return Objects;
+	return {};
+}
+
+void FSoundCueParameterTransmitter::AddReferencedObjects(FReferenceCollector& InCollector)
+{
+PRAGMA_DISABLE_DEPRECATION_WARNINGS // Disable deprecations because ParamsToSet is being moved to private
+	AddReferencedObjectsFromParameters(InCollector, ParamsToSet);
+PRAGMA_ENABLE_DEPRECATION_WARNINGS
+	Super::AddReferencedObjects(InCollector);
 }
 
 bool FSoundCueParameterTransmitter::SetParameters(TArray<FAudioParameter>&& InParameters)
@@ -946,22 +973,65 @@ bool FSoundCueParameterTransmitter::SetParameters(TArray<FAudioParameter>&& InPa
 		}
 	};
 
-	if (bIsVirtualized)
+	if (IsVirtualized())
 	{
 		RemoveTriggerParameters();
 
 		TArray<FAudioParameter> TempParams = InParameters;
+PRAGMA_DISABLE_DEPRECATION_WARNINGS // Disable deprecations because ParamsToSet is being moved to private
 		FAudioParameter::Merge(MoveTemp(TempParams), ParamsToSet);
+PRAGMA_ENABLE_DEPRECATION_WARNINGS
 
-		return Audio::FParameterTransmitterBase::SetParameters(MoveTemp(InParameters));
+		return Super::SetParameters(MoveTemp(InParameters));
 	}
 	else
 	{
 		TArray<FAudioParameter> TempParams = InParameters;
+PRAGMA_DISABLE_DEPRECATION_WARNINGS // Disable deprecations because ParamsToSet is being moved to private
 		FAudioParameter::Merge(MoveTemp(TempParams), ParamsToSet);
+PRAGMA_ENABLE_DEPRECATION_WARNINGS
 
 		RemoveTriggerParameters();
 
-		return Audio::FParameterTransmitterBase::SetParameters(MoveTemp(InParameters));
+		return Super::SetParameters(MoveTemp(InParameters));
+	}
+}
+
+void FSoundCueParameterTransmitter::CreateChildTransmitter(UPTRINT InWaveInstanceHash, USoundBase& InSound, Audio::FParameterTransmitterInitParams InTransmitterInitParams)
+{
+	check(IsInAudioThread());
+
+	// Override any default parameters with the existing parameters. 
+	TArray<FAudioParameter> CurrentAudioParameters;
+	CopyParameters(CurrentAudioParameters);
+	FAudioParameter::Merge(MoveTemp(CurrentAudioParameters), InTransmitterInitParams.DefaultParams);
+
+	InSound.InitParameters(InTransmitterInitParams.DefaultParams);
+	// Create child transmitter and store if valid. 
+	TSharedPtr<Audio::IParameterTransmitter> ChildTransmitter = InSound.CreateParameterTransmitter(MoveTemp(InTransmitterInitParams));
+
+	if (ChildTransmitter.IsValid())
+	{
+PRAGMA_DISABLE_DEPRECATION_WARNINGS // Disable deprecations because transmitters is being moved to private
+		Transmitters.Add(InWaveInstanceHash, ChildTransmitter.ToSharedRef());
+PRAGMA_ENABLE_DEPRECATION_WARNINGS 
 	}
 }
+
+void FSoundCueParameterTransmitter::RemoveChildTransmitter(UPTRINT InWaveInstanceHash)
+{
+	check(IsInAudioThread());
+PRAGMA_DISABLE_DEPRECATION_WARNINGS // Disable deprecations because transmitters is being moved to private
+	Transmitters.Remove(InWaveInstanceHash);
+PRAGMA_ENABLE_DEPRECATION_WARNINGS 
+}
+
+TArray<FAudioParameter> FSoundCueParameterTransmitter::ReleaseAccumulatedParameterUpdates()
+{
+PRAGMA_DISABLE_DEPRECATION_WARNINGS // Disable deprecations because ParamsToSet is being moved to private
+	TArray<FAudioParameter> AccumulatedParameters(MoveTemp(ParamsToSet));
+	ParamsToSet.Reset();
+PRAGMA_ENABLE_DEPRECATION_WARNINGS
+	return MoveTemp(AccumulatedParameters);
+}
+
diff --git a/Engine/Source/Runtime/Engine/Private/SoundNodeWaveParam.cpp b/Engine/Source/Runtime/Engine/Private/SoundNodeWaveParam.cpp
--- a/Engine/Source/Runtime/Engine/Private/SoundNodeWaveParam.cpp
+++ b/Engine/Source/Runtime/Engine/Private/SoundNodeWaveParam.cpp
@@ -41,22 +41,14 @@ void USoundNodeWaveParam::ParseNodes(FAudioDevice* AudioDevice, const UPTRINT No
 			if (FSoundCueParameterTransmitter* SoundCueTransmitter = static_cast<FSoundCueParameterTransmitter*>(ActiveSound.GetTransmitter()))
 			{
 				// removing here prevents waste in the case that a metasound is replaced with a soundwave
-				SoundCueTransmitter->Transmitters.Remove(GetNodeWaveInstanceHash(NodeWaveInstanceHash, (UPTRINT)PrevWave, 0));
+				SoundCueTransmitter->RemoveChildTransmitter(GetNodeWaveInstanceHash(NodeWaveInstanceHash, (UPTRINT)PrevWave, 0));
 				
 				Audio::FParameterTransmitterInitParams Params;
-				Params.DefaultParams = ActiveSound.GetTransmitter()->GetParameters();
 				Params.InstanceID = Audio::GetTransmitterID(ActiveSound.GetAudioComponentID(), WaveHash, ActiveSound.GetPlayOrder()); 
 				Params.SampleRate = AudioDevice->GetSampleRate();
 				Params.AudioDeviceID = AudioDevice->DeviceID;
-				
-				NewWave->InitParameters(Params.DefaultParams);
 
-				const TSharedPtr<Audio::IParameterTransmitter> SoundWaveTransmitter = NewWave->CreateParameterTransmitter(MoveTemp(Params));
-			
-				if (SoundWaveTransmitter.IsValid())
-				{
-					SoundCueTransmitter->Transmitters.Add(WaveHash, SoundWaveTransmitter);
-				}
+				SoundCueTransmitter->CreateChildTransmitter(WaveHash, *NewWave, MoveTemp(Params));
 			}
 
 			PrevWave = NewWave;
diff --git a/Engine/Source/Runtime/Engine/Private/SoundNodeWavePlayer.cpp b/Engine/Source/Runtime/Engine/Private/SoundNodeWavePlayer.cpp
--- a/Engine/Source/Runtime/Engine/Private/SoundNodeWavePlayer.cpp
+++ b/Engine/Source/Runtime/Engine/Private/SoundNodeWavePlayer.cpp
@@ -170,19 +170,11 @@ void USoundNodeWavePlayer::ParseNodes( FAudioDevice* AudioDevice, const UPTRINT
 			if (FSoundCueParameterTransmitter* SoundCueTransmitter = static_cast<FSoundCueParameterTransmitter*>(ActiveSound.GetTransmitter()))
 			{
 				Audio::FParameterTransmitterInitParams Params;
-				Params.DefaultParams = ActiveSound.GetTransmitter()->GetParameters();
 				Params.InstanceID = Audio::GetTransmitterID(ActiveSound.GetAudioComponentID(), NodeWaveInstanceHash, ActiveSound.GetPlayOrder()); 
 				Params.SampleRate = AudioDevice->GetSampleRate();
 				Params.AudioDeviceID = AudioDevice->DeviceID;
 
-				SoundWave->InitParameters(Params.DefaultParams);
-				
-				const TSharedPtr<Audio::IParameterTransmitter> SoundWaveTransmitter = SoundWave->CreateParameterTransmitter(MoveTemp(Params));
-				
-				if (SoundWaveTransmitter.IsValid())
-				{
-					SoundCueTransmitter->Transmitters.Add(NodeWaveInstanceHash, SoundWaveTransmitter);
-				}
+				SoundCueTransmitter->CreateChildTransmitter(NodeWaveInstanceHash, *SoundWave, MoveTemp(Params));
 			}
 			
 			*RequiresInitialization = 0;
