Network test FTestWorldInstances now use the default RHI feature level.

diff --git a/Engine/Source/Runtime/Engine/Private/Tests/NetTestHelpers.cpp b/Engine/Source/Runtime/Engine/Private/Tests/NetTestHelpers.cpp
--- a/Engine/Source/Runtime/Engine/Private/Tests/NetTestHelpers.cpp
+++ b/Engine/Source/Runtime/Engine/Private/Tests/NetTestHelpers.cpp
@@ -51,7 +51,7 @@ FTestWorldInstance FTestWorldInstance::CreateServer(const TCHAR* InURL, const ER
 	ServerParams.bStartInSpectatorMode = false;
 	ServerParams.bRunAsDedicated = true;
 	ServerParams.bIsPrimaryPIEClient = false;
-	ServerParams.WorldFeatureLevel = ERHIFeatureLevel::SM5;
+	ServerParams.WorldFeatureLevel = ERHIFeatureLevel::Num;
 	ServerParams.EditorPlaySettings = nullptr;
 	ServerParams.NetMode = PIE_ListenServer;
 	ServerParams.OverrideMapURL = TEXT("/Engine/Maps/Entry"); // always start from an empty map. Otherwise it will clone the current editor map to start the server in.
@@ -78,7 +78,7 @@ FTestWorldInstance FTestWorldInstance::CreateClient(int32 ServerPort)
 	ClientParams.bStartInSpectatorMode = false;
 	ClientParams.bRunAsDedicated = false;
 	ClientParams.bIsPrimaryPIEClient = false;
-	ClientParams.WorldFeatureLevel = ERHIFeatureLevel::SM5;
+	ClientParams.WorldFeatureLevel = ERHIFeatureLevel::Num;
 	ClientParams.EditorPlaySettings = nullptr;
 	ClientParams.NetMode = PIE_Client;
 
@@ -111,7 +111,7 @@ FTestWorldInstance FTestWorldInstance::CreateProxy()
 	ServerParams.bStartInSpectatorMode = false;
 	ServerParams.bRunAsDedicated = true;
 	ServerParams.bIsPrimaryPIEClient = false;
-	ServerParams.WorldFeatureLevel = ERHIFeatureLevel::SM5;
+	ServerParams.WorldFeatureLevel = ERHIFeatureLevel::Num;
 	ServerParams.EditorPlaySettings = nullptr;
 	ServerParams.NetMode = PIE_ListenServer;
 	ServerParams.OverrideMapURL = TEXT("/Engine/Maps/Entry"); // always start from an empty map. Otherwise it will clone the current editor map to start the server in.
