Empty the skeletal mesh updater packets immediately after completing work so that scopes can be destructed sooner.

diff --git a/Engine/Source/Runtime/Engine/Private/SkeletalMeshUpdater.cpp b/Engine/Source/Runtime/Engine/Private/SkeletalMeshUpdater.cpp
--- a/Engine/Source/Runtime/Engine/Private/SkeletalMeshUpdater.cpp
+++ b/Engine/Source/Runtime/Engine/Private/SkeletalMeshUpdater.cpp
@@ -276,6 +276,8 @@ struct FSkeletalMeshUpdater::FTaskData
 		{
 			RHICmdList->FinishRecording();
 		}
+
+		Packets.Empty();
 	}
 
 	void ProcessForeground()
