Optimized the resource pool and added a scope to take an external lock so that multiple allocations / deallocations can avoid it.

diff --git a/Engine/Source/Runtime/Engine/Public/ResourcePool.h b/Engine/Source/Runtime/Engine/Public/ResourcePool.h
--- a/Engine/Source/Runtime/Engine/Public/ResourcePool.h
+++ b/Engine/Source/Runtime/Engine/Public/ResourcePool.h
@@ -10,27 +10,22 @@
 #include "RenderResource.h"
 #include "TickableObjectRenderThread.h"
 #include "RHICommandList.h"
+#include "Async/Mutex.h"
 
 /** A templated pool for resources that can only be freed at a 'safe' point in the frame. */
 template<typename ResourceType, class ResourcePoolPolicy, class ResourceCreationArguments>
 class TResourcePool
 {
 public:
-	/** Constructor */
-	TResourcePool()
-	{
-		
-	}
-	
+	TResourcePool() = default;
+
 	/** Constructor with policy argument
 	 * @param InPolicy An initialised policy object
 	 */
 	TResourcePool(ResourcePoolPolicy InPolicy)
 	: Policy(InPolicy)
-	{
-		
-	}
-	
+	{}
+
 	/** Destructor */
 	virtual ~TResourcePool()
 	{
@@ -58,12 +53,12 @@ public:
 		TArray<FPooledResource>& PoolBucket = ResourceBuckets[BucketIndex];
 
 		{
-			FScopeLock Lock(&CS);
+			UE::TConditionalScopeLock Lock(CS, !bLocked);
 
 			if (PoolBucket.Num() > 0)
 			{
 				// Reuse the last entry in this size bucket
-				return PoolBucket.Pop().Resource;
+				return PoolBucket.Pop(EAllowShrinking::No).Resource;
 			}
 		}
 
@@ -74,20 +69,19 @@ public:
 	/** Release a resource back into the pool.
 	 * @param Resource The resource to return to the pool
 	 */
-	void ReleasePooledResource(const ResourceType& Resource)
+	void ReleasePooledResource(ResourceType& Resource)
 	{
 		FPooledResource NewEntry;
-		NewEntry.Resource = Resource;
-		NewEntry.FrameFreed = GFrameNumberRenderThread;
 		NewEntry.CreationArguments = Policy.GetCreationArguments(Resource);
+		NewEntry.Resource = MoveTemp(Resource);
+		NewEntry.FrameFreed = GFrameNumberRenderThread;
+		NewEntry.BucketIndex = Policy.GetPoolBucketIndex(NewEntry.CreationArguments);
 		
 		// Add to this frame's array of free resources
 		const int32 SafeFrameIndex = GFrameNumberRenderThread % ResourcePoolPolicy::NumSafeFrames;
-		const uint32 BucketIndex = Policy.GetPoolBucketIndex(NewEntry.CreationArguments);
 
-		FScopeLock Lock(&CS);
-	
-		SafeResourceBuckets[SafeFrameIndex][BucketIndex].Add(NewEntry);
+		UE::TConditionalScopeLock Lock(CS, !bLocked);
+		SafeResourceBuckets[SafeFrameIndex].Emplace(MoveTemp(NewEntry));
 	}
 	
 	/** Drain the pool of freed resources that need to be culled or prepared for reuse.
@@ -98,30 +92,29 @@ public:
 		uint32 NumToCleanThisFrame = ResourcePoolPolicy::NumToDrainPerFrame;
 		uint32 CullAfterFramesNum = ResourcePoolPolicy::CullAfterFramesNum;
 
-		FScopeLock Lock(&CS);
+		check(!bLocked);
+		UE::TScopeLock Lock(CS);
 
 		if(!bForceDrainAll)
 		{
 			// Index of the bucket that is now old enough to be reused
 			const int32 SafeFrameIndex = (GFrameNumberRenderThread + 1) % ResourcePoolPolicy::NumSafeFrames;
-			
-			// Merge the bucket into the free pool array
-			for (int32 BucketIndex = 0; BucketIndex < ResourcePoolPolicy::NumPoolBuckets; BucketIndex++)
+
+			for (FPooledResource& PoolEntry : SafeResourceBuckets[SafeFrameIndex])
 			{
-				ResourceBuckets[BucketIndex].Append(SafeResourceBuckets[SafeFrameIndex][BucketIndex]);
-				SafeResourceBuckets[SafeFrameIndex][BucketIndex].Reset();
+				ResourceBuckets[PoolEntry.BucketIndex].Emplace(MoveTemp(PoolEntry));
 			}
+			SafeResourceBuckets[SafeFrameIndex].Reset();
 		}
 		else
 		{
-			for(int32 FrameIndex = 0; FrameIndex < ResourcePoolPolicy::NumSafeFrames; FrameIndex++)
+			for (int32 FrameIndex = 0; FrameIndex < ResourcePoolPolicy::NumSafeFrames; FrameIndex++)
 			{
-				// Merge the bucket into the free pool array
-				for (int32 BucketIndex = 0; BucketIndex < ResourcePoolPolicy::NumPoolBuckets; BucketIndex++)
+				for (FPooledResource& PoolEntry : SafeResourceBuckets[FrameIndex])
 				{
-					ResourceBuckets[BucketIndex].Append(SafeResourceBuckets[FrameIndex][BucketIndex]);
-					SafeResourceBuckets[FrameIndex][BucketIndex].Reset();
+					ResourceBuckets[PoolEntry.BucketIndex].Emplace(MoveTemp(PoolEntry));
 				}
+				SafeResourceBuckets[FrameIndex].Reset();
 			}
 		}
 		
@@ -137,7 +130,7 @@ public:
 				{
 					Policy.FreeResource(ResourceBuckets[BucketIndex][EntryIndex].Resource);
 					
-					ResourceBuckets[BucketIndex].RemoveAtSwap(EntryIndex);
+					ResourceBuckets[BucketIndex].RemoveAtSwap(EntryIndex, EAllowShrinking::No);
 					
 					--NumToCleanThisFrame;
 					if (!bForceDrainAll && NumToCleanThisFrame == 0)
@@ -153,6 +146,28 @@ public:
 			}
 		}
 	}
+
+	/** A scope that takes a single lock so that individual allocations / deallocations can skip it. */
+	class FLockScope
+	{
+	public:
+		FLockScope(TResourcePool& InPool)
+			: Pool(InPool)
+			, Lock(Pool.CS)
+		{
+			check(!Pool.bLocked);
+			Pool.bLocked = true;
+		}
+
+		~FLockScope()
+		{
+			Pool.bLocked = false;
+		}
+
+	private:
+		TResourcePool& Pool;
+		UE::TScopeLock<UE::FMutex> Lock;
+	};
 	
 private:
 	/** Pooling policy for this resource */
@@ -167,15 +182,17 @@ private:
 		ResourceCreationArguments CreationArguments;
 		/** The frame the resource was freed */
 		uint32 FrameFreed;
+		uint32 BucketIndex;
 	};
 
-	FCriticalSection CS;
+	UE::FMutex CS;
+	bool bLocked = false;
 
 	// Pool of free Resources, indexed by bucket for constant size search time.
 	TArray<FPooledResource> ResourceBuckets[ResourcePoolPolicy::NumPoolBuckets];
 	
 	// Resources that have been freed more recently than NumSafeFrames ago.
-	TArray<FPooledResource> SafeResourceBuckets[ResourcePoolPolicy::NumSafeFrames][ResourcePoolPolicy::NumPoolBuckets];
+	TArray<FPooledResource> SafeResourceBuckets[ResourcePoolPolicy::NumSafeFrames];
 };
 
 /** A resource pool that automatically handles render-thread resources */
@@ -235,7 +252,7 @@ public:
 	/** Release a resource back into the pool.
 	 * @param Resource The resource to return to the pool
 	 */
-	void ReleasePooledResource(ResourceType Resource)
+	void ReleasePooledResource(ResourceType& Resource)
 	{
 		if (IsInitialized())
 		{
