Fix edge cases with skeletal mesh components with disabled animation

Reference pose initialization now correctly recalcs required bones, so we dont get collapsed meshes

diff --git a/Engine/Source/Runtime/Engine/Private/Components/SkeletalMeshComponent.cpp b/Engine/Source/Runtime/Engine/Private/Components/SkeletalMeshComponent.cpp
--- a/Engine/Source/Runtime/Engine/Private/Components/SkeletalMeshComponent.cpp
+++ b/Engine/Source/Runtime/Engine/Private/Components/SkeletalMeshComponent.cpp
@@ -914,6 +914,7 @@ void USkeletalMeshComponent::OnRegister()
 	}
 	else
 	{
+		RecalcRequiredBones(GetPredictedLODLevel());
 		ResetToRefPose();
 	}
 
@@ -3303,7 +3304,15 @@ void USkeletalMeshComponent::SetSkeletalMesh(USkeletalMesh* InSkelMesh, bool bRe
 		// RecalcRequiredBones will be called by InitAnim below
 		bRequiredBonesUpToDate = false;
 		
-		InitAnim(bReinitPose);
+		if(bEnableAnimation)
+		{
+			InitAnim(bReinitPose);
+		}
+		else if(IsRegistered())
+		{
+			RecalcRequiredBones(GetPredictedLODLevel());
+			ResetToRefPose();
+		}
 
 		RecreateClothingActors();
 	}
