Fix for GetMaterialAttributes node returning the wrong output type on non-Substrate nodes when material layers are enabled.

diff --git a/Engine/Source/Runtime/Engine/Private/Materials/MaterialExpressions.cpp b/Engine/Source/Runtime/Engine/Private/Materials/MaterialExpressions.cpp
--- a/Engine/Source/Runtime/Engine/Private/Materials/MaterialExpressions.cpp
+++ b/Engine/Source/Runtime/Engine/Private/Materials/MaterialExpressions.cpp
@@ -6151,7 +6151,16 @@ bool UMaterialExpressionGetMaterialAttributes::IsResultSubstrateMaterial(int32 O
 {	
 	if (MaterialAttributes.Expression)
 	{
-		return MaterialAttributes.Expression->IsResultSubstrateMaterial(MaterialAttributes.OutputIndex);
+		EMaterialValueType OutputType = IsResultMaterialAttributes(OutputIndex) ? MCT_MaterialAttributes : FMaterialAttributeDefinitionMap::GetValueType(AttributeGetTypes[OutputIndex - 1]);
+		switch (OutputType)
+		{
+			case MCT_Substrate: return true;
+			case MCT_MaterialAttributes:
+			{
+				return MaterialAttributes.Expression->IsResultSubstrateMaterial(MaterialAttributes.OutputIndex);
+			}
+			default:break;
+		}
 	}
 	return false;
 }
