Show actor labels for debug name of post process volumes when ShowFlag.VisualizePostProcessStack is enabled.

diff --git a/Engine/Source/Runtime/Engine/Classes/Engine/PostProcessVolume.h b/Engine/Source/Runtime/Engine/Classes/Engine/PostProcessVolume.h
--- a/Engine/Source/Runtime/Engine/Classes/Engine/PostProcessVolume.h
+++ b/Engine/Source/Runtime/Engine/Classes/Engine/PostProcessVolume.h
@@ -66,7 +66,7 @@ class APostProcessVolume : public AVolume, public IInterface_PostProcessVolume
 #if DEBUG_POST_PROCESS_VOLUME_ENABLE
 	virtual FString GetDebugName() const override
 	{
-		return GetName();
+		return GetActorNameOrLabel();
 	}
 #endif
 	//~ End IInterface_PostProcessVolume Interface
