Improve UDataTable::HandleDataTableChanged performance when modifying a single row. There is still a heavy cost when loading a table in the editor since the method will iterate over all rows to potentially call an empty function in FTableRowBase.

Adding 200,000 items:
- Before 2m23s
- After 114ms

Adding 1,000,000 items now takes ~550ms

#rnx

diff --git a/Engine/Source/Runtime/Engine/Classes/Engine/DataTable.h b/Engine/Source/Runtime/Engine/Classes/Engine/DataTable.h
--- a/Engine/Source/Runtime/Engine/Classes/Engine/DataTable.h
+++ b/Engine/Source/Runtime/Engine/Classes/Engine/DataTable.h
@@ -423,11 +423,12 @@ protected:
 	struct FScopedDataTableChange
 	{
 		FScopedDataTableChange(UDataTable* InTable);
+		FScopedDataTableChange(UDataTable* InTable,FName Row);
 		~FScopedDataTableChange();
 
 	private:
 		UDataTable* Table;
-
+		FName RowName;
 		static TMap<UDataTable*, int32> ScopeCount;
 		static FTransactionallySafeCriticalSection CriticalSection;
 	};
diff --git a/Engine/Source/Runtime/Engine/Private/DataTable.cpp b/Engine/Source/Runtime/Engine/Private/DataTable.cpp
--- a/Engine/Source/Runtime/Engine/Private/DataTable.cpp
+++ b/Engine/Source/Runtime/Engine/Private/DataTable.cpp
@@ -50,6 +50,19 @@ UDataTable::FScopedDataTableChange::FScopedDataTableChange(UDataTable* InTable)
 {
 	UE::TScopeLock Lock(CriticalSection);
 	int32& Count = ScopeCount.FindOrAdd(Table);
+	RowName = NAME_None;
+
+	++Count;
+}
+
+UDataTable::FScopedDataTableChange::FScopedDataTableChange(UDataTable* InTable, FName Row)
+	: Table(InTable)
+{
+	UE::TScopeLock Lock(CriticalSection);
+	int32& Count = ScopeCount.FindOrAdd(Table);
+	
+	//Recording RowName so we can emit HandleDataTableChanged with the name when adding\removing item without nested scopes.
+	RowName = Row;
 	++Count;
 }
 
@@ -60,7 +73,7 @@ UDataTable::FScopedDataTableChange::~FScopedDataTableChange()
 	--Count;
 	if (Count == 0)
 	{
-		Table->HandleDataTableChanged();
+		Table->HandleDataTableChanged(RowName);
 		ScopeCount.Remove(Table);
 	}
 }
@@ -70,6 +83,9 @@ FTransactionallySafeCriticalSection UDataTable::FScopedDataTableChange::Critical
 
 #define DATATABLE_CHANGE_SCOPE()	UDataTable::FScopedDataTableChange ActiveScope(this);
 
+//This macro should only be used for innermost "locks" like adding\removing a single item. 
+#define DATATABLE_CHANGE_SCOPE_SINGLE_ROW(Name)	UDataTable::FScopedDataTableChange ActiveScope(this, Name);
+
 UDataTable::UDataTable(const FObjectInitializer& ObjectInitializer)
 	: Super(ObjectInitializer)
 {
@@ -218,15 +234,23 @@ void UDataTable::HandleDataTableChanged(FName ChangedRowName)
 
 		if (bIsNativeRowStruct)
 		{
-			for (const TPair<FName, uint8*>& TableRowPair : RowMap)
+			//Avoid iterating on the full table when a specific row was modified. 
+			if (ChangedRowName != NAME_None)
 			{
-				if (ChangedRowName != NAME_None && ChangedRowName != TableRowPair.Key)
+				FTableRowBase* CurRow = reinterpret_cast<FTableRowBase*>(RowMap[ChangedRowName]);
+				
+				if (CurRow)
 				{
-					continue;
+					CurRow->OnDataTableChanged(this, ChangedRowName);
+				}
+			}
+			else
+			{
+				for (const TPair<FName, uint8*>& TableRowPair : RowMap)
+				{
+					FTableRowBase* CurRow = reinterpret_cast<FTableRowBase*>(TableRowPair.Value);
+					CurRow->OnDataTableChanged(this, TableRowPair.Key);
 				}
-
-				FTableRowBase* CurRow = reinterpret_cast<FTableRowBase*>(TableRowPair.Value);
-				CurRow->OnDataTableChanged(this, TableRowPair.Key);
 			}
 		}
 	}
@@ -482,7 +506,7 @@ void UDataTable::EmptyTable()
 
 void UDataTable::RemoveRow(FName RowName)
 {
-	DATATABLE_CHANGE_SCOPE();
+	DATATABLE_CHANGE_SCOPE_SINGLE_ROW(RowName);
 
 	RemoveRowInternal(RowName);
 }
@@ -503,7 +527,7 @@ void UDataTable::RemoveRowInternal(FName RowName)
 
 void UDataTable::AddRow(FName RowName, const FTableRowBase& RowData)
 {
-	DATATABLE_CHANGE_SCOPE();
+	DATATABLE_CHANGE_SCOPE_SINGLE_ROW(RowName);
 
 	UScriptStruct& EmptyUsingStruct = GetEmptyUsingStruct();
 
@@ -521,7 +545,7 @@ void UDataTable::AddRow(FName RowName, const FTableRowBase& RowData)
 
 void UDataTable::AddRow(FName RowName, const uint8* RowData, const UScriptStruct* RowType)
 {
-	DATATABLE_CHANGE_SCOPE();
+	DATATABLE_CHANGE_SCOPE_SINGLE_ROW(RowName);
 
 	UScriptStruct& EmptyUsingStruct = GetEmptyUsingStruct();
 
