Removed WorldPartition post-GC reachability analysis callback since it's not required anymore following CL 40093672.

diff --git a/Engine/Source/Runtime/Engine/Private/WorldPartition/WorldPartition.cpp b/Engine/Source/Runtime/Engine/Private/WorldPartition/WorldPartition.cpp
--- a/Engine/Source/Runtime/Engine/Private/WorldPartition/WorldPartition.cpp
+++ b/Engine/Source/Runtime/Engine/Private/WorldPartition/WorldPartition.cpp
@@ -361,8 +361,6 @@ UWorldPartition::UWorldPartition(const FObjectInitializer& ObjectInitializer)
 	, WorldPartitionEditor(nullptr)
 	, bStreamingWasEnabled(true)
 	, bShouldCheckEnableStreamingWarning(false)
-	, bForceGarbageCollection(false)
-	, bForceGarbageCollectionPurge(false)
 	, bForceRefreshAlwaysLoaded(false)
 	, bForceRefreshEditor(false)
 	, bEnablingStreamingJustified(false)
@@ -389,52 +387,6 @@ UWorldPartition::UWorldPartition(const FObjectInitializer& ObjectInitializer)
 }
 
 #if WITH_EDITOR
-void UWorldPartition::OnGCPostReachabilityAnalysis()
-{
-	if (!IsEngineExitRequested())
-	{
-		TRACE_CPUPROFILER_EVENT_SCOPE(UWorldPartition::OnGCPostReachabilityAnalysis);
-		UE_SCOPED_TIMER(TEXT("UWorldPartition::OnGCPostReachabilityAnalysis"), LogWorldPartition, Display);
-		const TIndirectArray<FWorldContext>& WorldContextList = GEngine->GetWorldContexts();
-
-		// Avoid running this process while a game world is live
-		for (const FWorldContext& WorldContext : WorldContextList)
-		{
-			if (WorldContext.World() != nullptr && WorldContext.World()->IsGameWorld())
-			{
-				return;
-			}
-		}
-
-		for (FRawObjectIterator It; It; ++It)
-		{
-			if (It->IsUnreachable())
-			{
-				if (AActor* Actor = Cast<AActor>(static_cast<UObject*>(It->Object)))
-				{
-					if (Actor->IsPackageExternal() && Actor->IsMainPackageActor() && !Actor->GetClass()->HasAnyClassFlags(CLASS_NewerVersionExists))
-					{
-						ForEachObjectWithPackage(Actor->GetPackage(), [this, Actor](UObject* Object)
-						{
-							if (Object->HasAnyFlags(RF_Standalone))
-							{
-								UE_LOG(LogWorldPartition, Log, TEXT("Actor %s is unreachable without properly detaching object %s in its package"), *Actor->GetPathName(), *Object->GetPathName());
-
-								Object->ClearFlags(RF_Standalone);
-
-								// Make sure we trigger a second GC at the next tick to properly destroy packages there were fixed in this pass
-								bForceGarbageCollection = true;
-								bForceGarbageCollectionPurge = true;
-							}
-							return true;
-						}, false);
-					}
-				}
-			}
-		}
-	}
-}
-
 // Returns whether the memory package is part of the known/valid package names
 // used by World Partition for PIE/-game streaming.
 bool UWorldPartition::IsValidPackageName(const FString& InPackageName)
@@ -1100,10 +1052,6 @@ void UWorldPartition::RegisterDelegates()
 
 		if (!IsRunningCommandlet())
 		{
-			if (IsMainWorldPartition())
-			{
-				FCoreUObjectDelegates::PostReachabilityAnalysis.AddUObject(this, &UWorldPartition::OnGCPostReachabilityAnalysis);
-			}
 			ExternalDirtyActorsTracker = MakeUnique<FWorldPartitionExternalDirtyActorsTracker>(this);
 		}
 	}
@@ -1149,14 +1097,6 @@ void UWorldPartition::UnregisterDelegates()
 
 		if (!IsRunningCommandlet())
 		{
-			if (!IsEngineExitRequested())
-			{
-				if (IsMainWorldPartition())
-				{
-					FCoreUObjectDelegates::PostReachabilityAnalysis.RemoveAll(this);
-				}
-			}
-
 			ExternalDirtyActorsTracker.Reset();
 		}
 	}
@@ -1841,14 +1781,6 @@ void UWorldPartition::Tick(float DeltaSeconds)
 		bForceRefreshEditor = false;
 	}
 
-	if (bForceGarbageCollection)
-	{
-		GEngine->ForceGarbageCollection(bForceGarbageCollectionPurge);
-
-		bForceGarbageCollection = false;
-		bForceGarbageCollectionPurge = false;
-	}
-
 	if (bShouldCheckEnableStreamingWarning)
 	{
 		bShouldCheckEnableStreamingWarning = false;
diff --git a/Engine/Source/Runtime/Engine/Public/WorldPartition/WorldPartition.h b/Engine/Source/Runtime/Engine/Public/WorldPartition/WorldPartition.h
--- a/Engine/Source/Runtime/Engine/Public/WorldPartition/WorldPartition.h
+++ b/Engine/Source/Runtime/Engine/Public/WorldPartition/WorldPartition.h
@@ -173,7 +173,6 @@ public:
 private:
 	ENGINE_API void SavePerUserSettings();
 		
-	ENGINE_API void OnGCPostReachabilityAnalysis();
 	ENGINE_API void OnPackageDirtyStateChanged(UPackage* Package);
 
 	// PIE/Game
@@ -609,8 +608,6 @@ private:
 	TObjectPtr<UWorld> World;
 
 #if WITH_EDITOR
-	bool bForceGarbageCollection;
-	bool bForceGarbageCollectionPurge;
 	bool bForceRefreshAlwaysLoaded;
 	bool bForceRefreshEditor;
 	bool bEnablingStreamingJustified;
