Fixed issue where landscape grass disappears in viewport after reimporting a FGrassVariety static mesh

diff --git a/Engine/Source/Runtime/Engine/Classes/Components/HierarchicalInstancedStaticMeshComponent.h b/Engine/Source/Runtime/Engine/Classes/Components/HierarchicalInstancedStaticMeshComponent.h
--- a/Engine/Source/Runtime/Engine/Classes/Components/HierarchicalInstancedStaticMeshComponent.h
+++ b/Engine/Source/Runtime/Engine/Classes/Components/HierarchicalInstancedStaticMeshComponent.h
@@ -354,8 +354,8 @@ protected:
 	ENGINE_API virtual void BuildComponentInstanceData(ERHIFeatureLevel::Type FeatureLevel, FInstanceUpdateComponentDesc& OutData) override;
 
 
-	ENGINE_API void BuildTree();
-	ENGINE_API void BuildTreeAsync();
+	ENGINE_API virtual void BuildTree();
+	ENGINE_API virtual void BuildTreeAsync();
 	ENGINE_API void ApplyBuildTree(FClusterBuilder& Builder, const bool bWasAsyncBuild);
 	ENGINE_API void ApplyEmpty();
 	ENGINE_API void SetPerInstanceLightMapAndEditorData(FStaticMeshInstanceData& PerInstanceData, const TArray<TRefCountPtr<HHitProxy>>& HitProxies);
diff --git a/Engine/Source/Runtime/Engine/Private/HierarchicalInstancedStaticMesh.cpp b/Engine/Source/Runtime/Engine/Private/HierarchicalInstancedStaticMesh.cpp
--- a/Engine/Source/Runtime/Engine/Private/HierarchicalInstancedStaticMesh.cpp
+++ b/Engine/Source/Runtime/Engine/Private/HierarchicalInstancedStaticMesh.cpp
@@ -2631,6 +2631,7 @@ void UHierarchicalInstancedStaticMeshComponent::BuildTree()
 	TRACE_CPUPROFILER_EVENT_SCOPE(UHierarchicalInstancedStaticMeshComponent::BuildTree);
 
 	check(!HasAnyFlags(RF_ClassDefaultObject | RF_ArchetypeObject));
+	check(PrimitiveInstanceDataManager.GetMode() != FPrimitiveInstanceDataManager::EMode::ExternalLegacyData)
 	checkSlow(IsInGameThread());
 
 	// If we try to build the tree with the static mesh not fully loaded, we can end up in an inconsistent state which ends in a crash later
@@ -2842,6 +2843,7 @@ void UHierarchicalInstancedStaticMeshComponent::BuildTreeAsync()
 	TRACE_CPUPROFILER_EVENT_SCOPE(UHierarchicalInstancedStaticMeshComponent::BuildTreeAsync);
 
 	check(!HasAnyFlags(RF_ClassDefaultObject | RF_ArchetypeObject));
+	check(PrimitiveInstanceDataManager.GetMode() != FPrimitiveInstanceDataManager::EMode::ExternalLegacyData)
 	check(IsInGameThread());
 
 	// If we try to build the tree with the static mesh not fully loaded, we can end up in an inconsistent state which ends in a crash later
diff --git a/Engine/Source/Runtime/Foliage/Private/InstancedGrass.cpp b/Engine/Source/Runtime/Foliage/Private/InstancedGrass.cpp
--- a/Engine/Source/Runtime/Foliage/Private/InstancedGrass.cpp
+++ b/Engine/Source/Runtime/Foliage/Private/InstancedGrass.cpp
@@ -90,6 +90,16 @@ void UGrassInstancedStaticMeshComponent::AcceptPrebuiltTree(TArray<FClusterNode>
 	MarkRenderStateDirty();
 }
 
+void UGrassInstancedStaticMeshComponent::BuildTree()
+{
+	// Grass instances managed externally and handled in AcceptPrebuiltTree()
+}
+
+void UGrassInstancedStaticMeshComponent::BuildTreeAsync()
+{
+	// Grass instances managed externally and handled in AcceptPrebuiltTree()
+}
+
 void UGrassInstancedStaticMeshComponent::BuildComponentInstanceData(ERHIFeatureLevel::Type FeatureLevel, FInstanceUpdateComponentDesc& OutData)
 {
 	OutData.PrimitiveLocalToWorld = GetRenderMatrix();
diff --git a/Engine/Source/Runtime/Foliage/Public/GrassInstancedStaticMeshComponent.h b/Engine/Source/Runtime/Foliage/Public/GrassInstancedStaticMeshComponent.h
--- a/Engine/Source/Runtime/Foliage/Public/GrassInstancedStaticMeshComponent.h
+++ b/Engine/Source/Runtime/Foliage/Public/GrassInstancedStaticMeshComponent.h
@@ -21,6 +21,8 @@ public:
 	int32 GetNumRenderInstances() const override { return NumBuiltRenderInstances; }
 
 private:
+	virtual void BuildTree() override final;
+	virtual void BuildTreeAsync() override final;
 	virtual void BuildComponentInstanceData(ERHIFeatureLevel::Type FeatureLevel, FInstanceUpdateComponentDesc& OutData) override final;
 };
 
