New material translator changes:
- Added the ability of executing IR instructions in one or more execution contexts called "stage". This allows us to support VertexWorldOffset which is evaluated in the vertex shader.
- Added support for VertexWorldOffset.
- Added some comments and documention.
- Refactored IR use graph debug dump generator code into a helper struct.

#rnx

diff --git a/Engine/Source/Runtime/Engine/Private/Materials/MaterialIR.cpp b/Engine/Source/Runtime/Engine/Private/Materials/MaterialIR.cpp
--- a/Engine/Source/Runtime/Engine/Private/Materials/MaterialIR.cpp
+++ b/Engine/Source/Runtime/Engine/Private/Materials/MaterialIR.cpp
@@ -9,6 +9,11 @@
 namespace MIR
 {
 
+EStage GetMaterialOutputStage(EMaterialProperty Property)
+{
+	return Property == MP_WorldPositionOffset ? MIR::Stage_Vertex : MIR::Stage_Fragment;
+}
+
 const TCHAR* ValueKindToString(EValueKind Kind)
 {
 	switch (Kind)
@@ -29,6 +34,11 @@ const TCHAR* ValueKindToString(EValueKind Kind)
 	};
 }
 
+bool FValue::IsAnalyzed(EStage Stage) const
+{
+	return (AnalyzedInStage & (1 << (uint32)Stage)) != 0;
+}
+
 uint32 FValue::GetSizeInBytes() const
 {
 	switch (Kind)
@@ -49,11 +59,6 @@ uint32 FValue::GetSizeInBytes() const
 	}
 }
 
-void FValue::SetFlags(EValueFlags InFlags)
-{
-	Flags = (EValueFlags)(Flags | InFlags);
-}
-
 bool FValue::IsA(EValueKind InKind) const
 {
 	return Kind == InKind;
@@ -257,7 +262,12 @@ const FInstruction* AsInstruction(const FValue* Value)
 
 FPoison* FPoison::Get()
 {
-	static FPoison Poison{ VK_Poison, VF_None, FType::GetPoison() };
+	static FPoison Poison = [] {
+		FPoison P;
+		P.Kind = VK_Poison;
+		P.Type = FType::GetPoison();
+		return P;
+	} ();
 	return &Poison;
 }
 
@@ -396,21 +406,21 @@ bool FDimensional::AreComponentsConstant() const
 	return true;
 }
 
-FBlock* FInstruction::GetDesiredBlockForUse(int32 UseIndex)
+FBlock* FInstruction::GetDesiredBlockForUse(EStage Stage, int32 UseIndex)
 {
 	if (auto Branch = As<FBranch>())
 	{
 		switch (UseIndex)
 		{
-			case 0: return Block; // ConditionArg goes into the same block as this instruction's
-			case 1: return &Branch->TrueBlock; // TrueArg
-			case 2: return &Branch->FalseBlock; // FalseArg
+			case 0: return Block[Stage]; 				// ConditionArg goes into the same block as this instruction's
+			case 1: return &Branch->TrueBlock[Stage]; 	// TrueArg
+			case 2: return &Branch->FalseBlock[Stage]; 	// FalseArg
 			default: UE_MIR_UNREACHABLE();
 		}
 	}
 
 	// By default, dependencies can go in the same block as this instruction
-	return Block;
+	return Block[Stage];
 }
 
 bool IsComparisonOperator(EOperator Op)
diff --git a/Engine/Source/Runtime/Engine/Private/Materials/MaterialIRDebug.cpp b/Engine/Source/Runtime/Engine/Private/Materials/MaterialIRDebug.cpp
--- a/Engine/Source/Runtime/Engine/Private/Materials/MaterialIRDebug.cpp
+++ b/Engine/Source/Runtime/Engine/Private/Materials/MaterialIRDebug.cpp
@@ -22,114 +22,56 @@ static TAutoConsoleVariable<bool> CVarDumpMaterialIRUseGraph_EnableNext(
 namespace MIR
 {
 
-static void DumpValueInfo(const FValue* Value, FString& Out)
+struct FDebugDumpIRUseGraphState
 {
-	if (const FConstant* Constant = Value->As<FConstant>())
-	{
-		switch (Constant->Type->AsPrimitive()->ScalarKind)
-		{
-			case ScalarKind_Bool:	Out.Append(Constant->Boolean ? TEXT("true") : TEXT("false")); break;
-			case ScalarKind_Int:	Out.Appendf(TEXT("%lld"), Constant->Integer); break;
-			case ScalarKind_Float: Out.Appendf(TEXT("%f"), Constant->Float); break;
-			default: UE_MIR_UNREACHABLE();
-		}
-	}
-	else if (const FExternalInput* ExternalInput = Value->As<FExternalInput>())
-	{
-		Out.Append(ExternalInputToString(ExternalInput->Id));
-	}
-	else if (const FSetMaterialOutput* SetMaterailOutput = Value->As<FSetMaterialOutput>())
-	{
-		const FString& PropertyName = (SetMaterailOutput->Property == MP_SubsurfaceColor)
-			? TEXT("Subsurface")
-			: FMaterialAttributeDefinitionMap::GetAttributeName(SetMaterailOutput->Property);
-		
-		Out.Append(PropertyName);
-	}
-	else if (const FSubscript* Subscript = Value->As<FSubscript>())
-	{
-		if (Subscript->Arg->Type->AsVector())
-		{
-			static const TCHAR* Suffix[] = { TEXT(".x"), TEXT(".y"), TEXT(".z"), TEXT(".w") };
-			check(Subscript->Index < 4); 
-			Out.Append(Suffix[Subscript->Index]);
-		}
-		else
-		{
-			Out.Appendf(TEXT("Index: %d"), Subscript->Index);
-		}
-	}
-	else if (const FOperator* Operator = Value->As<FOperator>())
-	{
-		Out.Append(OperatorToString(Operator->Op));
-	}
-}				
+	FString Out;
+	TSet<const FValue*> Crawled;
+	TArray<const FValue*> ValueStack;
 
-static void DumpUseInfo(const FValue* Value, const FValue* Use, int32 UseIndex, FString& Out)
-{
-	if (const FDimensional* Dimensional = Value->As<FDimensional>())
+	void DumpModule(const FMaterialIRModule& Module)
 	{
-		if (Dimensional->Type->AsPrimitive()->IsVector())
-		{
-			check(UseIndex < 4);
-			Out.AppendChar(TEXT("xyzw")[UseIndex]);
-		}
-		else
+		Out.Appendf(TEXT(
+			"digraph G {\n\n"
+			"rankdir=LR\n"
+			"node [shape=box,fontname=\"Consolas\"]\n"
+			"edge [fontname=\"Consolas\"]\n\n"
+		));
+
+		for (int32 StageIndex = 0; StageIndex < MIR::NumStages; ++StageIndex)
 		{
-			Out.AppendInt(UseIndex);
-		}
-	}
-	else if (const FBranch* If = Value->As<FBranch>())
-	{
-		static const TCHAR* Uses[] = { TEXT("condition"), TEXT("true"), TEXT("false") };
-		Out.Append(Uses[UseIndex]);
-	}
-	else if (const FOperator* Operator = Value->As<FOperator>())
-	{
-		static const TCHAR* Uses[] = { TEXT("a"), TEXT("b"), TEXT("c") };
-		Out.Append(Uses[UseIndex]);
-	}
-}
-
-void DebugDumpIRUseGraph(const FMaterialIRModule& Module)
-{
-	FString Content;
-	TSet<const FValue*> Crawled;
-	TArray<const FValue*> ValueStack;
+			for (const FSetMaterialOutput* Output : Module.GetOutputs((MIR::EStage)StageIndex))
+			{
+				ValueStack.Push(Output);
+			}
 
-	Content.Appendf(TEXT(
-		"digraph G {\n\n"
-		"rankdir=LR\n"
-		"node [shape=box,fontname=\"Consolas\"]\n"
-		"edge [fontname=\"Consolas\"]\n\n"
-	));
+			while (!ValueStack.IsEmpty())
+			{
+				DumpValue((MIR::EStage)StageIndex, ValueStack.Pop());
+			}
+		}
 
-	for (const FSetMaterialOutput* Output : Module.GetOutputs())
-	{
-		ValueStack.Push(Output);
+		Out.Appendf(TEXT("\n}\n"));
 	}
 
-	bool DumpInstructionSequence = CVarDumpMaterialIRUseGraph_EnableNext.GetValueOnAnyThread();
-
-	while (!ValueStack.IsEmpty())
+	void DumpValue(MIR::EStage Stage, const FValue* Value)
 	{
-		const FValue* Value = ValueStack.Pop();
+		const bool bDumpInstructionSequence = CVarDumpMaterialIRUseGraph_EnableNext.GetValueOnAnyThread();
 
 		// Begin the node declaration
-		Content.Appendf(TEXT("\"%p\" [label=< <b>%s</b>  (%s) <br/> "),
+		Out.Appendf(TEXT("\"%p\" [label=< <b>%s</b>  (%s) <br/> "),
 						Value,
 						ValueKindToString(Value->Kind),
 						Value->Type ? Value->Type->GetSpelling().GetData() : TEXT("???"));
 
-		DumpValueInfo(Value, Content);
+		DumpValueInfo(Value);
 
 		// End the node declaration
-		Content.Append(TEXT(">]\n"));
+		Out.Append(TEXT(">]\n"));
 
 		const FInstruction* Instr = AsInstruction(Value);
-		if (DumpInstructionSequence && Instr && Instr->Next)
+		if (bDumpInstructionSequence && Instr && Instr->Next[Stage])
 		{
-			Content.Appendf(TEXT("\"%p\" -> \"%p\" [color=\"red\"]\n"), Instr, Instr->Next);
+			Out.Appendf(TEXT("\"%p\" -> \"%p\" [color=\"red\"]\n"), Instr, Instr->Next[Stage]);
 		}
 
 		int32 UseIndex = -1;
@@ -142,11 +84,11 @@ void DebugDumpIRUseGraph(const FMaterialIRModule& Module)
 				continue;
 			}
 			
-			Content.Appendf(TEXT("\"%p\" -> \"%p\" [label=\""), Value, Use);
+			Out.Appendf(TEXT("\"%p\" -> \"%p\" [label=\""), Value, Use);
 
-			DumpUseInfo(Value, Use, UseIndex, Content);
+			DumpUseInfo(Value, Use, UseIndex);
 
-			Content.Appendf(TEXT("\"]\n"));
+			Out.Appendf(TEXT("\"]\n"));
 
 			if (!Crawled.Contains(Use))
 			{
@@ -154,21 +96,94 @@ void DebugDumpIRUseGraph(const FMaterialIRModule& Module)
 				ValueStack.Push(Use);
 			}
 
-			if (DumpInstructionSequence && Instr)
+			if (bDumpInstructionSequence && Instr)
 			{
 				const FInstruction* UseInstr = AsInstruction(Use);
-				if (UseInstr && UseInstr->Block != Instr->Block)
+				if (UseInstr && UseInstr->Block[Stage] != Instr->Block[Stage])
 				{
-					Content.Appendf(TEXT("\"%p\" -> \"%p\" [color=\"red\", style=\"dashed\"]\n"), UseInstr, Instr);
+					Out.Appendf(TEXT("\"%p\" -> \"%p\" [color=\"red\", style=\"dashed\"]\n"), UseInstr, Instr);
 				}
 			}
 		}
 	}
 
-	Content.Appendf(TEXT("\n}\n"));
-	
+	void DumpValueInfo(const FValue* Value)
+	{
+		if (const FConstant* Constant = Value->As<FConstant>())
+		{
+			switch (Constant->Type->AsPrimitive()->ScalarKind)
+			{
+				case ScalarKind_Bool:	Out.Append(Constant->Boolean ? TEXT("true") : TEXT("false")); break;
+				case ScalarKind_Int:	Out.Appendf(TEXT("%lld"), Constant->Integer); break;
+				case ScalarKind_Float: Out.Appendf(TEXT("%f"), Constant->Float); break;
+				default: UE_MIR_UNREACHABLE();
+			}
+		}
+		else if (const FExternalInput* ExternalInput = Value->As<FExternalInput>())
+		{
+			Out.Append(ExternalInputToString(ExternalInput->Id));
+		}
+		else if (const FSetMaterialOutput* SetMaterailOutput = Value->As<FSetMaterialOutput>())
+		{
+			const FString& PropertyName = (SetMaterailOutput->Property == MP_SubsurfaceColor)
+				? TEXT("Subsurface")
+				: FMaterialAttributeDefinitionMap::GetAttributeName(SetMaterailOutput->Property);
+		
+			Out.Append(PropertyName);
+		}
+		else if (const FSubscript* Subscript = Value->As<FSubscript>())
+		{
+			if (Subscript->Arg->Type->AsVector())
+			{
+				static const TCHAR* Suffix[] = { TEXT(".x"), TEXT(".y"), TEXT(".z"), TEXT(".w") };
+				check(Subscript->Index < 4); 
+				Out.Append(Suffix[Subscript->Index]);
+			}
+			else
+			{
+				Out.Appendf(TEXT("Index: %d"), Subscript->Index);
+			}
+		}
+		else if (const FOperator* Operator = Value->As<FOperator>())
+		{
+			Out.Append(OperatorToString(Operator->Op));
+		}
+	}
+
+	void DumpUseInfo(const FValue* Value, const FValue* Use, int32 UseIndex)
+	{
+		if (const FDimensional* Dimensional = Value->As<FDimensional>())
+		{
+			if (Dimensional->Type->AsPrimitive()->IsVector())
+			{
+				check(UseIndex < 4);
+				Out.AppendChar(TEXT("xyzw")[UseIndex]);
+			}
+			else
+			{
+				Out.AppendInt(UseIndex);
+			}
+		}
+		else if (const FBranch* If = Value->As<FBranch>())
+		{
+			static const TCHAR* Uses[] = { TEXT("condition"), TEXT("true"), TEXT("false") };
+			Out.Append(Uses[UseIndex]);
+		}
+		else if (const FOperator* Operator = Value->As<FOperator>())
+		{
+			static const TCHAR* Uses[] = { TEXT("a"), TEXT("b"), TEXT("c") };
+			Out.Append(Uses[UseIndex]);
+		}
+	}
+};
+
+void DebugDumpIRUseGraph(const FMaterialIRModule& Module)
+{
+	FDebugDumpIRUseGraphState State;
+	State.DumpModule(Module);
+
 	FString FilePath = FPaths::Combine(FPaths::ProjectSavedDir(), "Materials", TEXT("IRDumpUseGraph.dot"));
-	FFileHelper::SaveStringToFile(Content, *FilePath);
+	FFileHelper::SaveStringToFile(State.Out, *FilePath);
 }
 
 /* Module IR to textual representation dumping */
@@ -194,11 +209,14 @@ struct FDebugDumpIRState
 	/// Array of encountered parameters. Used later on to generate the a recap of all referenced parameters.
 	TArray<TPair<uint32, const FValue*>> ReferencedParameters;
 
+	/// Stage we're currently emitting to.
+	MIR::EStage CurrentStage;
+
 	/// Prints a block of instructions to `Out`. Indentation indicates how many levels of indentation to put
 	/// to the left of printed instruction.
 	void AppendBlock(const FBlock& Block, int32 Indentation)
 	{
-		for (FInstruction* Instr = Block.Instructions; Instr; Instr = Instr->Next)
+		for (FInstruction* Instr = Block.Instructions; Instr; Instr = Instr->Next[CurrentStage])
 		{
 			// Format the left column (e.g. "%4 = ") string if this instruction is referenceable.
 			Temp.Empty();
@@ -237,8 +255,8 @@ struct FDebugDumpIRState
 				Out.Appendf(TEXT("%s "), Use->Type->GetSpelling().GetData());
 
 				// If this use is in a block different from current's, dump the block in "{}" first.
-				const FBlock* UseBlock = Instr->GetDesiredBlockForUse(UseIndex);
-				if (UseBlock != Instr->Block && UseBlock->Instructions)
+				const FBlock* UseBlock = Instr->GetDesiredBlockForUse(CurrentStage, UseIndex);
+				if (UseBlock != Instr->Block[CurrentStage] && UseBlock->Instructions)
 				{
 					Out.Append(TEXT("{\n"));
 					AppendBlock(*UseBlock, Indentation + 1);
@@ -429,11 +447,15 @@ FString DebugDumpIR(FStringView MaterialName, const FMaterialIRModule& Module)
 	State.Module = &Module;
 	State.Out.Append(TEXT("; Material IR module dump.\n"));
 	State.Out.Appendf(TEXT(";    Material: %s\n"), MaterialName.GetData());
-	State.Out.Append(TEXT("\n"));
 
 	// Dump the IR instructions in the root block.
 	static const TCHAR* Header = TEXT("; Material \"%s\" translated IR dump\n\n");
-	State.AppendBlock(Module.GetRootBlock(), 0);
+	for (int32 StageIndex = 0; StageIndex < MIR::NumStages; ++StageIndex)
+	{
+		State.CurrentStage = (MIR::EStage)StageIndex;
+		State.Out.Appendf(TEXT("\n; Stage %d\n"), StageIndex);
+		State.AppendBlock(Module.GetRootBlock((MIR::EStage)StageIndex), 0);
+	}
 
 	// Print referenced material parameters recap if any
 	State.DumpReferencedParameters();
diff --git a/Engine/Source/Runtime/Engine/Private/Materials/MaterialIREmitter.cpp b/Engine/Source/Runtime/Engine/Private/Materials/MaterialIREmitter.cpp
--- a/Engine/Source/Runtime/Engine/Private/Materials/MaterialIREmitter.cpp
+++ b/Engine/Source/Runtime/Engine/Private/Materials/MaterialIREmitter.cpp
@@ -784,13 +784,15 @@ FValue* FEmitter::EmitCustomPrimitiveData(uint32 PrimitiveDataIndex)
 
 FSetMaterialOutput* FEmitter::EmitSetMaterialOutput(EMaterialProperty InProperty, FValue* Arg)
 {
+	EStage PropertyStage = GetMaterialOutputStage(InProperty);
+
 	FSetMaterialOutput Proto = MakePrototype<FSetMaterialOutput>(nullptr);
-	Proto.Block 	= Module->RootBlock;
-	Proto.Property  = InProperty;
-	Proto.Arg 		= Arg;
+	Proto.Block[PropertyStage]	= Module->RootBlock[PropertyStage];
+	Proto.Property  			= InProperty;
+	Proto.Arg 					= Arg;
 
 	FSetMaterialOutput* Instr = static_cast<FSetMaterialOutput*>(EmitPrototype(this, Proto));
-	Module->Outputs.Add(Instr);
+	Module->Outputs[PropertyStage].Add(Instr);
 
 	return Instr;
 }
diff --git a/Engine/Source/Runtime/Engine/Private/Materials/MaterialIRInternal.cpp b/Engine/Source/Runtime/Engine/Private/Materials/MaterialIRInternal.cpp
--- a/Engine/Source/Runtime/Engine/Private/Materials/MaterialIRInternal.cpp
+++ b/Engine/Source/Runtime/Engine/Private/Materials/MaterialIRInternal.cpp
@@ -10,7 +10,7 @@
 
 namespace MIR::Internal {
 
-bool IsMaterialPropertyShared(EMaterialProperty InProperty)
+bool IsMaterialPropertyEnabled(EMaterialProperty InProperty)
 {
 	switch (InProperty)
 	{
@@ -32,6 +32,7 @@ bool IsMaterialPropertyShared(EMaterialProperty InProperty)
 	case MP_SurfaceThickness:
 	case MP_FrontMaterial:
 	case MP_Displacement:
+	case MP_WorldPositionOffset:
 		return true;
 	default:
 		return false;
@@ -43,7 +44,7 @@ bool NextMaterialAttributeInput(UMaterial* BaseMaterial, int32& PropertyIndex, F
 	for (; PropertyIndex < MP_MAX; ++PropertyIndex)
 	{
 		EMaterialProperty Property = (EMaterialProperty)PropertyIndex;
-		if (MIR::Internal::IsMaterialPropertyShared(Property)
+		if (MIR::Internal::IsMaterialPropertyEnabled(Property) 
 			&& Property != MP_SubsurfaceColor
 			&& Property != MP_FrontMaterial
 			&& BaseMaterial->GetExpressionInputDescription(Property, Input))
diff --git a/Engine/Source/Runtime/Engine/Private/Materials/MaterialIRInternal.h b/Engine/Source/Runtime/Engine/Private/Materials/MaterialIRInternal.h
--- a/Engine/Source/Runtime/Engine/Private/Materials/MaterialIRInternal.h
+++ b/Engine/Source/Runtime/Engine/Private/Materials/MaterialIRInternal.h
@@ -16,7 +16,7 @@ TArrayView<T> MakeTemporaryArray(FMemMark&, int Count)
 namespace MIR::Internal {
 
 //
-bool IsMaterialPropertyShared(EMaterialProperty InProperty);
+bool IsMaterialPropertyEnabled(EMaterialProperty InProperty);
 
 //
 bool NextMaterialAttributeInput(UMaterial* BaseMaterial, int32& PropertyIndex, FMaterialInputDescription& Input);
diff --git a/Engine/Source/Runtime/Engine/Private/Materials/MaterialIRModule.cpp b/Engine/Source/Runtime/Engine/Private/Materials/MaterialIRModule.cpp
--- a/Engine/Source/Runtime/Engine/Private/Materials/MaterialIRModule.cpp
+++ b/Engine/Source/Runtime/Engine/Private/Materials/MaterialIRModule.cpp
@@ -8,18 +8,29 @@
 
 FMaterialIRModule::FMaterialIRModule()
 {
-	RootBlock = new MIR::FBlock;
+	for (int32 i = 0; i < MIR::NumStages; ++i)
+	{
+		RootBlock[i] = new MIR::FBlock;
+	}
 }
 
 FMaterialIRModule::~FMaterialIRModule()
 {
 	Empty();
-	delete RootBlock;
+
+	for (int32 i = 0; i < MIR::NumStages; ++i)
+	{
+		delete RootBlock[i];
+	}
 }
 
 void FMaterialIRModule::Empty()
 {
-	RootBlock->Instructions = nullptr;
+	for (int32 i = 0; i < MIR::NumStages; ++i)
+	{
+		RootBlock[i]->Instructions = nullptr;
+		Outputs[i].Empty();
+	}
 
 	for (MIR::FValue* Value : Values)
 	{
@@ -27,10 +38,6 @@ void FMaterialIRModule::Empty()
 	}
 
 	Values.Empty();
-	Outputs.Empty();
-
-	// Allocator.~FMemStackBase();
-	// new (&Allocator) FMemStackBase;
 
 	// Reset module statistics.
 	for (int i = 0; i < (int)SF_NumFrequencies; ++i)
@@ -40,6 +47,8 @@ void FMaterialIRModule::Empty()
 
 	Statistics.NumVertexTexCoords = 0;
 	Statistics.NumPixelTexCoords = 0;
+
+	Allocator.Flush();
 }
 
 #endif // #if WITH_EDITOR
diff --git a/Engine/Source/Runtime/Engine/Private/Materials/MaterialIRModuleBuilder.cpp b/Engine/Source/Runtime/Engine/Private/Materials/MaterialIRModuleBuilder.cpp
--- a/Engine/Source/Runtime/Engine/Private/Materials/MaterialIRModuleBuilder.cpp
+++ b/Engine/Source/Runtime/Engine/Private/Materials/MaterialIRModuleBuilder.cpp
@@ -69,7 +69,6 @@ struct FMaterialIRModuleBuilderImpl
 	FMaterialIRModule* Module;
 	MIR::FEmitter* Emitter;
 	TArray<FAnalysisContext> AnalysisContextStack;
-	TArray<MIR::FInstruction*> InstructionStack{};
 	FMaterialIRValueAnalyzer ValueAnalyzer;
 
 	void Step_Initialize()
@@ -292,159 +291,166 @@ struct FMaterialIRModuleBuilderImpl
 	{
 		FAnalysisContext& Context = AnalysisContextStack.Last();
 
-		for (MIR::FSetMaterialOutput* Output : Module->Outputs)
+		for (uint32 StageIndex = 0; StageIndex < MIR::NumStages; ++StageIndex)
 		{
-			FMaterialInputDescription Input;
-			ensure(Builder->Material->GetExpressionInputDescription(Output->Property, Input));
-
-			if (!Output->Arg)
+			for (MIR::FSetMaterialOutput* Output : Module->Outputs[StageIndex])
 			{
-				MIR::FValue** ValuePtr = Context.OutputValues.Find(Input.Input->GetConnectedOutput());
-				check(ValuePtr && *ValuePtr);
+				FMaterialInputDescription Input;
+				ensure(Builder->Material->GetExpressionInputDescription(Output->Property, Input));
 
-				MIR::Internal::SetInputValue(this, Input.Input, *ValuePtr);
+				if (!Output->Arg)
+				{
+					MIR::FValue** ValuePtr = Context.OutputValues.Find(Input.Input->GetConnectedOutput());
+					check(ValuePtr && *ValuePtr);
 
-				MIR::FTypePtr OutputArgType = MIR::FType::FromShaderType(Input.Type);
-				Output->Arg = Emitter->EmitConstruct(OutputArgType, *ValuePtr);
-			}
+					MIR::Internal::SetInputValue(this, Input.Input, *ValuePtr);
 
-			if (Builder->TargetInsights)
-			{
-				check(Output->Arg);
-				PushConnectionInsight(Builder->Material, (int)Output->Property, Input.Input->Expression, Input.Input->OutputIndex, Output->Arg->Type);
+					MIR::FTypePtr OutputArgType = MIR::FType::FromShaderType(Input.Type);
+					Output->Arg = Emitter->EmitConstruct(OutputArgType, *ValuePtr);
+				}
+
+				if (Builder->TargetInsights)
+				{
+					check(Output->Arg);
+					PushConnectionInsight(Builder->Material, (int)Output->Property, Input.Input->Expression, Input.Input->OutputIndex, Output->Arg->Type);
+				}
 			}
 		}
 	}
 
 	void Step_AnalyzeIRGraph()
 	{
-		InstructionStack.Empty(64);
+		TArray<MIR::FValue*> ValueStack{};
 
-		for (MIR::FSetMaterialOutput* Output : Module->Outputs)
+		for (uint32 StageIndex = 0; StageIndex < MIR::NumStages; ++StageIndex)
 		{
-			InstructionStack.Push(Output);
-		}
+			MIR::EStage CurrentStage = (MIR::EStage)StageIndex;
 
-		while (!InstructionStack.IsEmpty())
-		{
-			MIR::FInstruction* Instr = InstructionStack.Last();
+			// Clear the value stack but preserve its allocated memory.
+			ValueStack.Empty(ValueStack.Max());
 
-			if (Instr->Flags & MIR::VF_Analyzed)
+			// Push each output in the current stage to the value stack.
+			for (MIR::FSetMaterialOutput* Output : Module->Outputs[StageIndex])
 			{
-				InstructionStack.Pop();
-				continue;
+				ValueStack.Push(Output);
 			}
 
-			for (MIR::FValue* UseValue : Instr->GetUses())
+			// Process until the value stack is empty.
+			while (!ValueStack.IsEmpty())
 			{
-				// If this argument value is null, ignore it
-				if (!UseValue)
+				MIR::FValue* Value = ValueStack.Last();
+				
+				// Module building should have interrupted before if poison values were generated.
+				check(!Value->IsPoison());
+
+				// If this instruction has already been analyzed for this stage, nothing else is left to do for it. Continue.
+				if (Value->IsAnalyzed(CurrentStage))
 				{
+					ValueStack.Pop();
 					continue;
 				}
 
-				// This value is used. If it hasn't been analyzed yet, analyze it now.
-				MIR::FInstruction* Use = MIR::AsInstruction(UseValue);
-				if (!Use)
+				// Before analyzing this value, make sure all used values are analyzed first.
+				for (MIR::FValue* Use : Value->GetUses())
 				{
-					if (!(UseValue->Flags & MIR::VF_Analyzed))
+					if (Use && !Use->IsAnalyzed(CurrentStage))
 					{
-						// Only instructions can have arguments. Verify this here for safety.
-						check(UseValue->GetUses().Num() == 0);
-
-						UseValue->SetFlags(MIR::VF_Analyzed);
-						ValueAnalyzer.Analyze(UseValue);
+						ValueStack.Push(Use);
 					}
-					continue;
 				}
 
-				if (!(Use->Flags & MIR::VF_Analyzed))
+				// If any other value has been pushed to the stack, it means we have a dependency to analyze first.
+				if (ValueStack.Last() != Value)
 				{
-					InstructionStack.Push(Use);
+					continue;
 				}
-			}
 
-			if (InstructionStack.Last() != Instr)
-			{
-				continue;
-			}
+				// All dependencies of this value has been analyzed, we can proceed analyzing this value now.
+				ValueStack.Pop();
 
-			InstructionStack.Pop();
+				// Go through each use instruction and increment its counter of users (this instruction).
+				for (MIR::FValue* Use : Value->GetUses())
+				{
+					// If this used value is an instruction, update its counter of users (in current stage).
+					if (MIR::FInstruction* UseInstr = MIR::AsInstruction(Use))
+					{
+						UseInstr->NumUsers[StageIndex] += 1;
+					}
+				}
 
-			for (MIR::FValue* UseValue : Instr->GetUses())
-			{
-				if (MIR::FInstruction* UseAsInstruction = MIR::AsInstruction(UseValue))
+				// If this is the first time this value is analyzed, let the analyzer process it.
+				// Note that individual value processing is independent from the stage it runs on so we can perform it only once.
+				if (Value->AnalyzedInStage == 0)
 				{
-					UseAsInstruction->NumUsers += 1;
+					ValueAnalyzer.Analyze(Value);
 				}
-			}
 
-			Instr->SetFlags(MIR::VF_Analyzed);
-			ValueAnalyzer.Analyze(Instr);
+				// Mark the used instruction as analyzed for this stage.
+				Value->AnalyzedInStage |= 1 << StageIndex;
+			}
 		}
 	}
 
 	void Step_LinkInstructions()
 	{
-		// This function walks the instruction graph and puts each instruction into the inner most possible block.
-		InstructionStack.Empty(InstructionStack.Max());
+		TArray<MIR::FInstruction*> InstructionStack{};
 
-		for (MIR::FSetMaterialOutput* Output : Module->Outputs)
+		for (uint32 StageIndex = 0; StageIndex < MIR::NumStages; ++StageIndex)
 		{
-			Output->Block = Module->RootBlock;
-			InstructionStack.Add(Output);
-		}
+			// This function walks the instruction graph and puts each instruction into the inner most possible block.
+			InstructionStack.Empty(InstructionStack.Max());
 
-		while (!InstructionStack.IsEmpty())
-		{
-			MIR::FInstruction* Instr = InstructionStack.Pop();
-			if (MIR::FSetMaterialOutput* Output = Instr->As<MIR::FSetMaterialOutput>())
+			// Initialize each output's block in this stage and add it to the instruction stack.
+			for (MIR::FSetMaterialOutput* Output : Module->Outputs[StageIndex])
 			{
-				if (Output->Property == EMaterialProperty::MP_BaseColor)
-				{
-					static int l = 0;
-					++l;
-				}
-			}
+				Output->Block[StageIndex] = Module->RootBlock[StageIndex];
 
-			// Push the instruction to its block in reverse order (push front)
-			Instr->Next = Instr->Block->Instructions;
-			Instr->Block->Instructions = Instr;
+				InstructionStack.Add(Output);
+			}
 
-			TArrayView<MIR::FValue*> Uses = Instr->GetUses();
-			for (int32 UseIndex = 0; UseIndex < Uses.Num(); ++UseIndex)
+			while (!InstructionStack.IsEmpty())
 			{
-				MIR::FValue* Use = Uses[UseIndex];
-				MIR::FInstruction* UseInstr = MIR::AsInstruction(Use);
-				if (!UseInstr)
-				{
-					continue;
-				}
+				MIR::FInstruction* Instr = InstructionStack.Pop();
 
-				// Get the block into which the dependency instruction should go.
-				MIR::FBlock* TargetBlock = Instr->GetDesiredBlockForUse(UseIndex);
+				// Push the instruction to its block in reverse order (push front)
+				Instr->Next[StageIndex] = Instr->Block[StageIndex] ->Instructions;
+				Instr->Block[StageIndex]->Instructions = Instr;
 
-				// Update dependency's block to be a child of current instruction's block.
-				if (TargetBlock != Instr->Block)
+				TArrayView<MIR::FValue*> Uses = Instr->GetUses();
+				for (int32 UseIndex = 0; UseIndex < Uses.Num(); ++UseIndex)
 				{
-					TargetBlock->Parent = Instr->Block;
-					TargetBlock->Level = Instr->Block->Level + 1;
-				}
+					MIR::FValue* Use = Uses[UseIndex];
+					MIR::FInstruction* UseInstr = MIR::AsInstruction(Use);
+					if (!UseInstr)
+					{
+						continue;
+					}
 
-				// Set the dependency's block to the common block betwen its current block and this one.
-				UseInstr->Block = UseInstr->Block
-					? UseInstr->Block->FindCommonParentWith(TargetBlock)
-					: TargetBlock;
+					// Get the block into which the dependency instruction should go.
+					MIR::FBlock* TargetBlock = Instr->GetDesiredBlockForUse((MIR::EStage)StageIndex, UseIndex);
 
-				// Increase the number of times this dependency instruction has been considered.
-				// When all of its users have processed, we can carry on visiting this instruction.
-				++UseInstr->NumProcessedUsers;
-				check(UseInstr->NumProcessedUsers <= UseInstr->NumUsers);
+					// Update dependency's block to be a child of current instruction's block.
+					if (TargetBlock != Instr->Block[StageIndex])
+					{
+						TargetBlock->Parent = Instr->Block[StageIndex];
+						TargetBlock->Level = Instr->Block[StageIndex]->Level + 1;
+					}
 
-				// If all dependants have been processed, we can carry the processing from this dependency.
-				if (UseInstr->NumProcessedUsers == UseInstr->NumUsers)
-				{
-					InstructionStack.Push(UseInstr);
+					// Set the dependency's block to the common block betwen its current block and this one.
+					UseInstr->Block[StageIndex] = UseInstr->Block[StageIndex]
+						? UseInstr->Block[StageIndex]->FindCommonParentWith(TargetBlock)
+						: TargetBlock;
+
+					// Increase the number of times this dependency instruction has been considered.
+					// When all of its users have processed, we can carry on visiting this instruction.
+					++UseInstr->NumProcessedUsers[StageIndex];
+					check(UseInstr->NumProcessedUsers[StageIndex] <= UseInstr->NumUsers[StageIndex]);
+
+					// If all dependants have been processed, we can carry the processing from this dependency.
+					if (UseInstr->NumProcessedUsers[StageIndex] == UseInstr->NumUsers[StageIndex])
+					{
+						InstructionStack.Push(UseInstr);
+					}
 				}
 			}
 		}
diff --git a/Engine/Source/Runtime/Engine/Private/Materials/MaterialIRToHLSLTranslator.cpp b/Engine/Source/Runtime/Engine/Private/Materials/MaterialIRToHLSLTranslator.cpp
--- a/Engine/Source/Runtime/Engine/Private/Materials/MaterialIRToHLSLTranslator.cpp
+++ b/Engine/Source/Runtime/Engine/Private/Materials/MaterialIRToHLSLTranslator.cpp
@@ -186,11 +186,11 @@ static const TCHAR* GetShadingModelParameterName(EMaterialShadingModel InModel)
 	}
 }
 
-static bool IsFoldable(const MIR::FInstruction* Instr)
+static bool IsFoldable(const MIR::FInstruction* Instr, MIR::EStage Stage)
 {
 	if (auto Branch = Instr->As<MIR::FBranch>())
 	{
-		return !Branch->TrueBlock.Instructions && !Branch->FalseBlock.Instructions;
+		return !Branch->TrueBlock[Stage].Instructions && !Branch->FalseBlock[Stage].Instructions;
 	}
 
 	return true;
@@ -226,26 +226,18 @@ struct FTranslator : FMaterialIRToHLSLTranslation
 {
 	int32 NumLocals{};
 	TMap<const MIR::FInstruction*, FString> LocalIdentifier;
+	MIR::EStage CurrentStage;
 	FHLSLPrinter Printer;
 	FString PixelAttributesHLSL;
+	FString WorldPositionOffsetHLSL;
 	FString EvaluateOtherMaterialAttributesHLSL;
 
-	void GenerateHLSL()
+	void GeneratePixelAttributesHLSL()
 	{
-		Printer.Tabs = 1;
-		Printer << Indentation;
-
-		LowerBlock(Module->GetRootBlock());
-
-		Printer << TEXT("PixelMaterialInputs.FrontMaterial = GetInitialisedSubstrateData()") << EndOfStatement;
-		Printer << TEXT("PixelMaterialInputs.Subsurface = 0") << EndOfStatement;
-
-		EvaluateOtherMaterialAttributesHLSL = MoveTemp(Printer.Buffer);
-
 		for (int32 PropertyIndex = 0; PropertyIndex < MP_MAX; ++PropertyIndex)
 		{
 			EMaterialProperty Property = (EMaterialProperty)PropertyIndex;
-			if (!MIR::Internal::IsMaterialPropertyShared(Property))
+			if (!MIR::Internal::IsMaterialPropertyEnabled(Property) || MIR::GetMaterialOutputStage(Property) != MIR::Stage_Fragment)
 			{
 				continue;
 			}
@@ -261,17 +253,36 @@ struct FTranslator : FMaterialIRToHLSLTranslation
 		}
 	}
 	
+	FString GenerateStageHLSL(MIR::EStage Stage)
+	{
+		Printer = {};
+		Printer.Tabs = 1;
+		Printer << Indentation;
+
+		CurrentStage = Stage;
+
+		LowerBlock(Module->GetRootBlock(CurrentStage));
+
+		if (Stage == MIR::Stage_Fragment)
+		{
+			Printer << TEXT("PixelMaterialInputs.FrontMaterial = GetInitialisedSubstrateData()") << EndOfStatement;
+			Printer << TEXT("PixelMaterialInputs.Subsurface = 0") << EndOfStatement;
+		}
+
+		return MoveTemp(Printer.Buffer);
+	}
+	
 	ENoOp LowerBlock(const MIR::FBlock& Block)
 	{
 		int OldNumLocals = NumLocals;
-        for (MIR::FInstruction* Instr = Block.Instructions; Instr; Instr = Instr->Next)
+		for (MIR::FInstruction* Instr = Block.Instructions; Instr; Instr = Instr->Next[CurrentStage])
 		{
-            if (Instr->NumUsers == 1 && IsFoldable(Instr))
+            if (Instr->NumUsers[CurrentStage] == 1 && IsFoldable(Instr, CurrentStage))
 			{
                 continue;
             }
             
-			if (Instr->NumUsers >= 1)
+			if (Instr->NumUsers[CurrentStage] >= 1)
 			{
                 FString LocalStr = FString::Printf(TEXT("_%d"), NumLocals);
                 ++NumLocals;
@@ -279,7 +290,7 @@ struct FTranslator : FMaterialIRToHLSLTranslation
                 Printer << LowerType(Instr->Type) << TEXT(" ") << LocalStr;
 
                 LocalIdentifier.Add(Instr, MoveTemp(LocalStr));
-                if (IsFoldable(Instr))
+                if (IsFoldable(Instr, CurrentStage))
 				{
                     Printer << TEXT(" = ";)
                 }
@@ -306,7 +317,7 @@ struct FTranslator : FMaterialIRToHLSLTranslation
 	{
 		if (MIR::FInstruction* Instr = MIR::AsInstruction(InValue))
 		{
-			if (Instr->NumUsers <= 1 && IsFoldable(Instr))
+			if (Instr->NumUsers[CurrentStage] <= 1 && IsFoldable(Instr, CurrentStage))
 			{
 				Printer << LowerInstruction(Instr);
 			}
@@ -472,7 +483,17 @@ struct FTranslator : FMaterialIRToHLSLTranslation
 	{
 		// Special case MP_SubsurfaceColor as the actual property is a combination of the color and the profile but we don't want to expose the profile
 		const FString& PropertyName = (Output->Property == MP_SubsurfaceColor) ? "Subsurface" : FMaterialAttributeDefinitionMap::GetAttributeName(Output->Property);
-		Printer << TEXT("PixelMaterialInputs.") << PropertyName << TEXT(" = ") << LowerValue(Output->Arg);
+
+		if (Output->Property == MP_WorldPositionOffset)
+		{
+			Printer << TEXT("return ");
+		}
+		else
+		{
+			Printer << TEXT("PixelMaterialInputs.") << PropertyName << TEXT(" = ");
+		}
+
+		LowerValue(Output->Arg);
 	}
 
 	void LowerOperator(const MIR::FOperator* Operator)
@@ -562,7 +583,7 @@ struct FTranslator : FMaterialIRToHLSLTranslation
 
 	void LowerBranch(const MIR::FBranch* Branch)
 	{
-		if (IsFoldable(Branch))
+		if (IsFoldable(Branch, CurrentStage))
 		{
 			Printer << LowerValue(Branch->ConditionArg)
 				<< TEXT(" ? ") << LowerValue(Branch->TrueArg)
@@ -572,11 +593,11 @@ struct FTranslator : FMaterialIRToHLSLTranslation
 		{
 			Printer << EndOfStatement;
 			Printer << TEXT("if (") << LowerValue(Branch->ConditionArg) << TEXT(")") << NewLine << OpenBrace;
-			Printer << LowerBlock(Branch->TrueBlock);
+			Printer << LowerBlock(Branch->TrueBlock[CurrentStage]);
 			Printer << LocalIdentifier[Branch] << " = " << LowerValue(Branch->TrueArg) << EndOfStatement;
 			Printer << CloseBrace << NewLine;
 			Printer << TEXT("else") << NewLine << OpenBrace;
-			Printer << LowerBlock(Branch->FalseBlock);
+			Printer << LowerBlock(Branch->FalseBlock[CurrentStage]);
 			Printer << LocalIdentifier[Branch] << " = " << LowerValue(Branch->FalseArg) << EndOfStatement;
 			Printer << CloseBrace;
 		}
@@ -770,8 +791,9 @@ struct FTranslator : FMaterialIRToHLSLTranslation
 		Params.Add(TEXT("get_material_translucent_multiple_scattering_extinction"), FString::Printf(TEXT(TAB "return MaterialFloat3(%.5f, %.5f, %.5f)"), Extinction.R, Extinction.G, Extinction.B));
 
 		SetParamReturnFloat(TEXT("get_material_opacity_mask_clip_value"), Material->GetOpacityMaskClipValue());
-		Params.Add(TEXT("get_material_world_position_offset_raw"), TEXT(TAB "return 0; // todo"));
-		Params.Add(TEXT("get_material_previous_world_position_offset_raw"), TEXT(TAB "return 0; // todo"));
+
+		Params.Add(TEXT("get_material_world_position_offset_raw"), WorldPositionOffsetHLSL);
+		Params.Add(TEXT("get_material_previous_world_position_offset_raw"), WorldPositionOffsetHLSL);
 	
 		// CustomData0/1 are named ClearCoat/ClearCoatRoughness
 		Params.Add(TEXT("get_material_custom_data0"), TEXT(TAB "return 1.0f; // todo"));
@@ -938,7 +960,9 @@ void FMaterialIRToHLSLTranslation::Run(TMap<FString, FString>& OutParameters, FS
 	OutParameters.Empty();
 
 	FTranslator Translator{ *this };
-	Translator.GenerateHLSL();
+	Translator.GeneratePixelAttributesHLSL();
+	Translator.WorldPositionOffsetHLSL = Translator.GenerateStageHLSL(MIR::Stage_Vertex);
+	Translator.EvaluateOtherMaterialAttributesHLSL = Translator.GenerateStageHLSL(MIR::Stage_Fragment);
 	Translator.GenerateTemplateStringParameters(OutParameters);
 	Translator.GetShaderCompilerEnvironment(OutEnvironment);
 }
diff --git a/Engine/Source/Runtime/Engine/Public/Materials/MaterialIR.h b/Engine/Source/Runtime/Engine/Public/Materials/MaterialIR.h
--- a/Engine/Source/Runtime/Engine/Public/Materials/MaterialIR.h
+++ b/Engine/Source/Runtime/Engine/Public/Materials/MaterialIR.h
@@ -10,6 +10,17 @@ struct FMaterialExternalCodeDeclaration;
 
 namespace MIR {
 
+// Identifies the block of execution in which an instruction runs.
+enum EStage
+{
+	Stage_Vertex,
+	Stage_Fragment,
+	NumStages
+};
+
+// Returns the stage in which specified material property shall be evaluated.
+EStage GetMaterialOutputStage(EMaterialProperty Property);
+
 enum EValueKind
 {
 	/* Values */
@@ -40,29 +51,48 @@ const TCHAR* ValueKindToString(EValueKind Kind);
 
 /* Values */
 
-enum EValueFlags
-{
-	VF_None = 0,
-	VF_Analyzed = 1,       // todo
-};
-
-//
+// Base entity of all IR graph nodes.
+// An IR module is a graph of values, connected by their "uses" relations. The graph of IR
+// values is built by the `MaterialIRModuleBuilder` as the result of crawling through and 
+// analyizing the MaterialExpression graph contained in the translated Material. During 
+// this processing, IR values are emitted, and linked together. After the graph is 
+// constructed, it is itself analyzed: the builder will call `FMaterialIRValueAnalyzer::Analyze()`
+// in each active* (i.e. truly used) value in the graph, making sure a value is analyzed only 
+// after its dependencies have been analyzed.
+// A few notes:
+//  - FValues are automatically zero-initialized.
+//  - FValues are intended to be simple and inert data records. They cannot have non-trivia
+//    ctor, dtor or copy operators.
+//  - The ModuleBuilder relies on this property to efficiently hashing values so that it
+//    will reuse the same value instance instead of creating multiple instances of the
+//    same computation (for a more efficient output shader).
+//  - All values have a MIR type.
+//  - Pure FValue instances are values that do not have other dependencies (called "uses").
+//  - If a value has some other value as dependency, it means that it is the result of a
+//    calculation on those values. Values that have dependencies are Instructions (they
+//    derive from FInstruction).
 struct FValue
 {
+	// Used to discern the concrete C++ type of this value (e.g. Subscript)
 	EValueKind Kind : 8;
-	EValueFlags Flags : 8;
+
+	// Whether this value was analyzed, which occurs after its dependencies have been analyzed.
+	uint32 AnalyzedInStage : MIR::NumStages;
+
+	// The runtime type this value has.
 	FTypePtr  Type;
 
+	// Utility function to tell whether this value has been analyzed for specified stage.
+	bool IsAnalyzed(EStage State) const;
+
 	// Returns the size in bytes of this value instance.
 	uint32 GetSizeInBytes() const;
 	
-	// Sets the specified value flags.
-	void SetFlags(EValueFlags InFlags);
-
 	// Returns whether this value is of specified kind.
 	bool IsA(EValueKind InKind) const;
 
-	// Returns the array of this value's uses. An use is another value referenced by this one (e.g. the operands of a binary expression).
+	// Returns the array of this value's uses. An use is another value referenced by this 
+	// one (e.g. the operands of a binary expression).
 	TArrayView<const FValue*> GetUses() const;
 
 	// Returns the array of this value's uses. An use is another value referenced by this one (e.g. the operands of a binary expression).
@@ -250,11 +280,11 @@ struct FTextureObject : TValue<VK_TextureObject>
 
 	EMaterialSamplerType SamplerType;
 
-	/// Index of this parameter in the uniform expression set
+	// Index of this parameter in the uniform expression set
 	uint32 Analysis_UniformParameterIndex;
 };
 
-// 
+//
 struct FUniformParameter: TValue<VK_UniformParameter>
 {
 	// Index of the parameter as registered in the module.
@@ -263,7 +293,7 @@ struct FUniformParameter: TValue<VK_UniformParameter>
 	// Eventual sampler type to use to sample this parameter's texture (if it is one)
 	EMaterialSamplerType SamplerType;
 
-	/// Index of this parameter in the uniform expression set
+	// Index of this parameter in the uniform expression set
 	uint32 Analysis_UniformParameterIndex;
 };
 
@@ -283,12 +313,13 @@ struct FBlock
 //
 struct FInstruction : FValue
 {
-	FInstruction* Next;
-	FBlock* Block;
-	uint32 NumUsers;
-	uint32 NumProcessedUsers;
+	FInstruction* Next[NumStages];
+	FBlock* Block[NumStages];
+	uint32 NumUsers[NumStages];
+	uint32 NumProcessedUsers[NumStages];
 
-	FBlock* GetDesiredBlockForUse(int32 UseIndex); 
+	// Returns the block in which use with specified index should execute.
+	FBlock* GetDesiredBlockForUse(EStage Stage, int32 UseIndex); 
 };
 
 template <EValueKind TTypeKind>
@@ -302,13 +333,13 @@ struct FDimensional : TInstruction<VK_Dimensional>
 {
 	static constexpr int MaxNumComponents = 16;
 
-	/// Returns the constant array of component values. 
+	// Returns the constant array of component values. 
 	TArrayView<FValue* const> GetComponents() const;
 
-	/// Returns the mutable array of component values. 
+	// Returns the mutable array of component values. 
 	TArrayView<FValue*> GetComponents();
 
-	/// Returns whether all components are constant.
+	// Returns whether all components are constant.
 	bool AreComponentsConstant() const;
 };
 
@@ -436,8 +467,8 @@ struct FBranch : TInstruction<VK_Branch>
 	FValue* ConditionArg;
 	FValue* TrueArg;
 	FValue* FalseArg;
-	FBlock TrueBlock;
-	FBlock FalseBlock;
+	FBlock TrueBlock[NumStages];
+	FBlock FalseBlock[NumStages];
 };
 
 //
@@ -473,7 +504,7 @@ static inline bool IsAnyNonNullAPoison()
 	return false;
 }
 
-/// Returns whether there is a poison in the values specified, which could be null.
+// Returns whether there is a poison in the values specified, which could be null.
 template <typename... TTail>
 bool IsAnyNonNullAPoison(const FValue* Head, TTail... Tail)
 {
diff --git a/Engine/Source/Runtime/Engine/Public/Materials/MaterialIRCommon.h b/Engine/Source/Runtime/Engine/Public/Materials/MaterialIRCommon.h
--- a/Engine/Source/Runtime/Engine/Public/Materials/MaterialIRCommon.h
+++ b/Engine/Source/Runtime/Engine/Public/Materials/MaterialIRCommon.h
@@ -49,10 +49,8 @@ using FObjectTypePtr = const FObjectType*;
 
 /* IR */
 struct FValue;
-struct FGlobalInput;
 struct FInstruction;
 struct FSetMaterialOutput;
-using FInstructionPtr = const FInstruction*;
 enum class EExternalInput;
 
 /* Others */
diff --git a/Engine/Source/Runtime/Engine/Public/Materials/MaterialIRModule.h b/Engine/Source/Runtime/Engine/Public/Materials/MaterialIRModule.h
--- a/Engine/Source/Runtime/Engine/Public/Materials/MaterialIRModule.h
+++ b/Engine/Source/Runtime/Engine/Public/Materials/MaterialIRModule.h
@@ -1,21 +1,21 @@
 // Copyright Epic Games, Inc. All Rights Reserved.
 #pragma once
 
-#include "Materials/MaterialIRCommon.h"
+#include "Materials/MaterialIR.h"
 #include "MaterialShared.h"
 
 #if WITH_EDITOR
 
-/// This class represents the intermediate representation (IR) of a material build.
-/// The `IRModule` includes an IR value graph, produced through expression analysis,
-/// as well as metadata on resource usage and reflection. The IR graph serves as an
-/// abstract representation of the material and must be translated into a target backend
-/// such as HLSL or specific Preshader opcodes for execution.
-///
-/// This class is designed to be backend-agnostic, meaning it does not contain any
-/// HLSL code nor does it configure a `MaterialCompilationOutput` instance. The data
-/// stored within this class should be sufficient to enable translation to any supported
-/// backend without requiring additional processing or validation.
+// This class represents the intermediate representation (IR) of a material build.
+// The `IRModule` includes an IR value graph, produced through expression analysis,
+// as well as metadata on resource usage and reflection. The IR graph serves as an
+// abstract representation of the material and must be translated into a target backend
+// such as HLSL or specific Preshader opcodes for execution.
+//
+// This class is designed to be backend-agnostic, meaning it does not contain any
+// HLSL code nor does it configure a `MaterialCompilationOutput` instance. The data
+// stored within this class should be sufficient to enable translation to any supported
+// backend without requiring additional processing or validation.
 class FMaterialIRModule
 {
 public:
@@ -38,8 +38,8 @@ public:
 	void Empty();
 	EShaderPlatform GetShaderPlatform() const { return ShaderPlatform; }
 	const FMaterialCompilationOutput& GetCompilationOutput() const { return CompilationOutput; }
-	TArrayView<const MIR::FSetMaterialOutput* const> GetOutputs() const { return Outputs; }
-	const MIR::FBlock& GetRootBlock() const { return *RootBlock; }
+	TArrayView<const MIR::FSetMaterialOutput* const> GetOutputs(MIR::EStage Stage) const { return Outputs[Stage]; }
+	const MIR::FBlock& GetRootBlock(MIR::EStage Stage) const { return *RootBlock[Stage]; }
 	TArrayView<const FError> GetErrors() const { return Errors; }
 	const FStatistics& GetStatistics() const { return Statistics; }
 	FStatistics& GetStatistics() { return Statistics; }
@@ -51,9 +51,9 @@ private:
 	FMaterialCompilationOutput CompilationOutput;
 	FMemStackBase Allocator{};
 	TArray<MIR::FValue*> Values;
-	TArray<MIR::FSetMaterialOutput*> Outputs;
+	TArray<MIR::FSetMaterialOutput*> Outputs[MIR::NumStages];
 	TArray<FError> Errors;
-	MIR::FBlock* RootBlock;
+	MIR::FBlock* RootBlock[MIR::NumStages];
 	FStatistics Statistics;
 	TMap<FMaterialParameterInfo, uint32> ParameterInfoToId;
 	TArray<TPair<FMaterialParameterInfo, FMaterialParameterMetadata>> ParameterIdToData;
