[Viewport Toolbar] Tweak blueprint editor background color to contrast with the viewport toolbar

- Change default color of Preview Scene background, to improve contrast with the new Viewport Toolbar

diff --git a/Engine/Source/Runtime/Engine/Private/PreviewScene.cpp b/Engine/Source/Runtime/Engine/Private/PreviewScene.cpp
--- a/Engine/Source/Runtime/Engine/Private/PreviewScene.cpp
+++ b/Engine/Source/Runtime/Engine/Private/PreviewScene.cpp
@@ -336,6 +336,5 @@ void FPreviewScene::SaveSettings(const TCHAR* Section)
 
 FLinearColor FPreviewScene::GetBackgroundColor() const
 {
-	FLinearColor BackgroundColor = FColor(55, 55, 55);
-	return BackgroundColor;
+	return FLinearColor(0.05f, 0.05f, 0.05f);
 }
