Culled unused tangent and previous frame buffers from GPU skin cache. Fixed GPU skin cache in GPU stats by injecting the dispatch work into RDG.

diff --git a/Engine/Source/Runtime/Engine/Private/GPUSkinCache.cpp b/Engine/Source/Runtime/Engine/Private/GPUSkinCache.cpp
--- a/Engine/Source/Runtime/Engine/Private/GPUSkinCache.cpp
+++ b/Engine/Source/Runtime/Engine/Private/GPUSkinCache.cpp
@@ -29,14 +29,16 @@ GPUSkinCache.cpp: Performs skinning on a compute shader into a buffer to avoid v
 #include "ComponentRecreateRenderStateContext.h"
 #include "RenderGraphBuilder.h"
 
+DECLARE_DWORD_COUNTER_STAT(TEXT("Num Dispatches"), STAT_GPUSkinCache_NumDispatches, STATGROUP_GPUSkinCache);
+DECLARE_DWORD_COUNTER_STAT(TEXT("Num Buffers"), STAT_GPUSkinCache_NumBuffers, STATGROUP_GPUSkinCache);
+DECLARE_DWORD_COUNTER_STAT(TEXT("Num RayTracing Dispatches"), STAT_GPUSkinCache_NumRayTracingDispatches, STATGROUP_GPUSkinCache);
+DECLARE_DWORD_COUNTER_STAT(TEXT("Num RayTracing Buffers"), STAT_GPUSkinCache_NumRayTracingBuffers, STATGROUP_GPUSkinCache);
 DEFINE_STAT(STAT_GPUSkinCache_TotalNumChunks);
 DEFINE_STAT(STAT_GPUSkinCache_TotalNumVertices);
 DEFINE_STAT(STAT_GPUSkinCache_TotalMemUsed);
 DEFINE_STAT(STAT_GPUSkinCache_TangentsIntermediateMemUsed);
 DEFINE_STAT(STAT_GPUSkinCache_NumTrianglesForRecomputeTangents);
 DEFINE_STAT(STAT_GPUSkinCache_NumSectionsProcessed);
-DEFINE_STAT(STAT_GPUSkinCache_NumSetVertexStreams);
-DEFINE_STAT(STAT_GPUSkinCache_NumPreGDME);
 DEFINE_LOG_CATEGORY_STATIC(LogSkinCache, Log, All);
 
 /** Exec helper to handle GPU Skin Cache related commands. */
@@ -265,11 +267,24 @@ ENGINE_API bool GPUSkinCacheNeedsDuplicatedVertices()
 }
 
 // determine if during DispatchUpdateSkinning caching should occur
-enum class EGPUSkinCacheDispatchFlags
+enum class EGPUSkinCacheDispatchFlags : uint8
 {
-	DispatchPrevPosition	= 1 << 0,
-	DispatchPosition		= 1 << 1,
+	None                 = 0,
+	Position             = 1 << 0,
+	PositionPrevious     = 1 << 1,
+	Tangents             = 1 << 2
 };
+ENUM_CLASS_FLAGS(EGPUSkinCacheDispatchFlags);
+
+enum class EGPUSkinCacheBufferBits : uint8
+{
+	None                 = 0,
+	Tangents             = 1 << 0,
+	IntermediateTangents = 1 << 1,
+	PositionPrevious     = 1 << 2,
+	AllTangents          = Tangents | IntermediateTangents
+};
+ENUM_CLASS_FLAGS(EGPUSkinCacheBufferBits);
 
 class FGPUSkinCacheEntry
 {
@@ -311,8 +326,7 @@ public:
 		// 0:normal, 1:with morph target, 2:with APEX cloth (not yet implemented)
 		uint16 SkinType = 0;
 
-		// See EGPUSkinCacheDispatchFlags
-		uint16 DispatchFlags = 0;
+		EGPUSkinCacheDispatchFlags DispatchFlags = EGPUSkinCacheDispatchFlags::None;
 
 		uint32 UpdatedFrameNumber = 0;
 		
@@ -377,6 +391,7 @@ public:
 
 		inline FGPUSkinCache::FSkinCacheRWBuffer* GetTangentRWBuffer() const
 		{
+			check(TangentBuffer);
 			return TangentBuffer;
 		}
 
@@ -394,13 +409,15 @@ public:
 
 		void UpdateVertexFactoryDeclaration(FRHICommandListBase& RHICmdList)
 		{
+			FRHIShaderResourceView* RHITangentSRV = TangentBuffer ? TangentBuffer->Buffer.SRV : TangentBufferSRV;
+
 			FGPUSkinPassthroughVertexFactory::FAddVertexAttributeDesc Desc;
 			Desc.FrameNumber = SourceVertexFactory->GetShaderData().UpdatedFrameNumber;
 			Desc.StreamBuffers[FGPUSkinPassthroughVertexFactory::EVertexAttribute::VertexPosition] = GetPositionRWBuffer()->Buffer.Buffer;
-			Desc.StreamBuffers[FGPUSkinPassthroughVertexFactory::EVertexAttribute::VertexTangent] = GetTangentRWBuffer()->Buffer.Buffer;
+			Desc.StreamBuffers[FGPUSkinPassthroughVertexFactory::EVertexAttribute::VertexTangent] = RHITangentSRV->GetBuffer();
 			Desc.SRVs[FGPUSkinPassthroughVertexFactory::EShaderResource::Position] = GetPositionRWBuffer()->Buffer.SRV;
 			Desc.SRVs[FGPUSkinPassthroughVertexFactory::EShaderResource::PreviousPosition] = GetPreviousPositionRWBuffer()->Buffer.SRV;
-			Desc.SRVs[FGPUSkinPassthroughVertexFactory::EShaderResource::Tangent] = GetTangentRWBuffer()->Buffer.SRV;
+			Desc.SRVs[FGPUSkinPassthroughVertexFactory::EShaderResource::Tangent] = RHITangentSRV;
 			TargetVertexFactory->SetVertexAttributes(RHICmdList, SourceVertexFactory, Desc);
 		}
 	};
@@ -536,8 +553,6 @@ public:
 			}
 		}
 
-		//INC_DWORD_STAT(STAT_GPUSkinCache_TotalNumChunks);
-
 		// SkinType 0:normal, 1:with morph target, 2:with cloth
 		Data.SkinType = ClothVertexBuffer ? 2 : (bMorph ? 1 : 0);
 		Data.InputStreamStart = InputStreamStart;
@@ -1222,10 +1237,10 @@ void FGPUSkinCache::DispatchUpdateSkinTangents(FRHICommandList& RHICmdList, FGPU
 	}
 }
 
-FGPUSkinCache::FRWBuffersAllocation* FGPUSkinCache::TryAllocBuffer(uint32 NumVertices, bool WithTangents, bool UseIntermediateTangents, uint32 InIntermediateAccumulatedTangentsSize, FRHICommandList& RHICmdList, const FName& OwnerName)
+FGPUSkinCache::FRWBuffersAllocation* FGPUSkinCache::TryAllocBuffer(FRHICommandList& RHICmdList, EGPUSkinCacheBufferBits BufferBits, uint32 NumVertices, uint32 InIntermediateAccumulatedTangentsSize, const FName& OwnerName)
 {
 	uint64 MaxSizeInBytes = (uint64)(GSkinCacheSceneMemoryLimitInMB * 1024.0f * 1024.0f);
-	uint64 RequiredMemInBytes = FRWBuffersAllocation::CalculateRequiredMemory(NumVertices, WithTangents, UseIntermediateTangents, InIntermediateAccumulatedTangentsSize);
+	uint64 RequiredMemInBytes = FRWBuffersAllocation::CalculateRequiredMemory(BufferBits, NumVertices, InIntermediateAccumulatedTangentsSize);
 	if (bRequiresMemoryLimit && UsedMemoryInBytes + RequiredMemInBytes >= MaxSizeInBytes)
 	{
 		ExtraRequiredMemory += RequiredMemInBytes;
@@ -1234,7 +1249,7 @@ FGPUSkinCache::FRWBuffersAllocation* FGPUSkinCache::TryAllocBuffer(uint32 NumVer
 		return nullptr;
 	}
 
-	FRWBuffersAllocation* NewAllocation = new FRWBuffersAllocation(NumVertices, WithTangents, UseIntermediateTangents, InIntermediateAccumulatedTangentsSize, RHICmdList, OwnerName);
+	FRWBuffersAllocation* NewAllocation = new FRWBuffersAllocation(RHICmdList, BufferBits, NumVertices, InIntermediateAccumulatedTangentsSize, OwnerName);
 	Allocations.Add(NewAllocation);
 
 	UsedMemoryInBytes += RequiredMemInBytes;
@@ -1336,31 +1351,49 @@ UE::Tasks::FTask FGPUSkinCache::DoDispatch(FRDGBuilder& GraphBuilder, const UE::
 	if (!GraphBuilder.IsParallelSetupEnabled())
 	{
 		PrerequisitesTask.Wait();
+		SCOPED_GPU_STAT(GraphBuilder.RHICmdList, GPUSkinCache);
 		DoDispatch(GraphBuilder.RHICmdList);
 		return {};
 	}
 
-	FRHICommandList* RHICmdList = new FRHICommandList();
-	GraphBuilder.RHICmdList.QueueAsyncCommandListSubmit(RHICmdList);
-
+	RDG_GPU_STAT_SCOPE(GraphBuilder, GPUSkinCache);
 	FDispatchSetupContext* SetupContext = GraphBuilder.AllocObject<FDispatchSetupContext>();
+	FRHICommandList* RHICmdListSetup = new FRHICommandList;
 
-	DispatchSetupTask = GraphBuilder.AddSetupTask([this, RHICmdList, SetupContext]
+	DispatchSetupTask = GraphBuilder.AddSetupTask([this, RHICmdListSetup, SetupContext]
 	{
-		RHICmdList->SwitchPipeline(ERHIPipeline::Graphics);
-		SetupDispatch(*RHICmdList, *SetupContext);
+		RHICmdListSetup->SwitchPipeline(ERHIPipeline::Graphics);
+		DispatchSetupForeground(*RHICmdListSetup, *SetupContext);
 
 	}, PrerequisitesTask, UE::Tasks::ETaskPriority::High);
 
-	GraphBuilder.AddSetupTask([this, RHICmdList, SetupContext]
+	GraphBuilder.AddSetupTask([this, RHICmdListSetup, SetupContext]
 	{
 		FTaskTagScope TagScope(ETaskTag::EParallelRenderingThread);
-		DoDispatch(*RHICmdList, *SetupContext);
-		RHICmdList->FinishRecording();
+		DispatchSetupBackground(*RHICmdListSetup, *SetupContext);
+		RHICmdListSetup->FinishRecording();
+
+		// Execute the dispatches in a separate command list that is enqueued into RDG.
+		FRHICommandList*& RHICmdListDispatch = SetupContext->RHICmdListDispatch;
+		RHICmdListDispatch = new FRHICommandList();
+		RHICmdListDispatch->SwitchPipeline(ERHIPipeline::Graphics);
+		DispatchPassExecute(*RHICmdListDispatch, *SetupContext);
+		RHICmdListDispatch->FinishRecording();
 
 	}, DispatchSetupTask, UE::Tasks::ETaskPriority::BackgroundHigh);
 
-	GraphBuilder.AddPostExecuteCallback([this] { DispatchSetupTask = {}; });
+	// Submit the setup command list here to start processing the RHI work ASAP.
+	GraphBuilder.RHICmdList.QueueAsyncCommandListSubmit(RHICmdListSetup);
+
+	// We use a dispatch pass here to insert the work into RDG, but the work is recorded during setup rather than kicking a task within the dispatch pass.
+	// This is because we can't record the dispatches asynchronously after the RDG frame ends because new render commands will modify the skin cache. Putting
+	// the work into RDG allows for stat scopes to work properly.
+	FEmptyShaderParameters* PassParameters = GraphBuilder.AllocParameters<FEmptyShaderParameters>();
+	GraphBuilder.AddDispatchPass(RDG_EVENT_NAME("GPUSkinCache"), PassParameters, ERDGPassFlags::Compute | ERDGPassFlags::NeverCull, [this, SetupContext](FRDGDispatchPassBuilder& DispatchPassBuilder)
+	{
+		DispatchPassBuilder.AddExistingCommandList(SetupContext->RHICmdListDispatch);
+		DispatchSetupTask = {};
+	});
 
 	return DispatchSetupTask;
 }
@@ -1372,13 +1405,14 @@ void FGPUSkinCache::DoDispatch(FRHICommandList& RHICmdList)
 	if (!BatchDispatches.IsEmpty())
 	{
 		FGPUSkinCache::FDispatchSetupContext Context;
-		SetupDispatch(RHICmdList, Context);
-		DoDispatch(RHICmdList, Context);
+		DispatchSetupForeground(RHICmdList, Context);
+		DispatchSetupBackground(RHICmdList, Context);
+		DispatchPassExecute(RHICmdList, Context);
 		BatchDispatches.Reset();
 	}
 }
 
-void FGPUSkinCache::SetupDispatch(FRHICommandList& RHICmdList, FDispatchSetupContext& Context)
+void FGPUSkinCache::DispatchSetupForeground(FRHICommandList& RHICmdList, FDispatchSetupContext& Context)
 {
 	TRACE_CPUPROFILER_EVENT_SCOPE(GPUSkinCache_SetupDispatch);
 
@@ -1386,13 +1420,14 @@ void FGPUSkinCache::SetupDispatch(FRHICommandList& RHICmdList, FDispatchSetupCon
 
 	for (FDispatchEntry& DispatchItem : BatchDispatches)
 	{
-		PrepareUpdateSkinning(DispatchItem.SkinCacheEntry, DispatchItem.Section, DispatchItem.SkinCacheEntry->DispatchData[DispatchItem.Section].RevisionNumber, &Context.BuffersToTransition);
-
-		// Clear the flag that this is queued for dispatch.
-		DispatchItem.SkinCacheEntry->bQueuedForDispatch = false;
-		DispatchItem.SkinCacheEntry->DispatchData[DispatchItem.Section].RevisionNumber = 0;
+		PrepareUpdateSkinning(DispatchItem.SkinCacheEntry, DispatchItem.Section, Context);
 	}
 
+	SET_DWORD_STAT(STAT_GPUSkinCache_NumDispatches, BatchDispatches.Num());
+	SET_DWORD_STAT(STAT_GPUSkinCache_NumRayTracingDispatches, Context.Stats.NumRayTracingDispatches);
+	SET_DWORD_STAT(STAT_GPUSkinCache_NumBuffers, Context.Stats.NumBuffers);
+	SET_DWORD_STAT(STAT_GPUSkinCache_NumRayTracingBuffers, Context.Stats.NumRayTracingBuffers);
+
 #if RHI_RAYTRACING
 	if (IsGPUSkinCacheRayTracingSupported())
 	{
@@ -1406,7 +1441,7 @@ void FGPUSkinCache::SetupDispatch(FRHICommandList& RHICmdList, FDispatchSetupCon
 #endif
 }
 
-void FGPUSkinCache::DoDispatch(FRHICommandList& RHICmdList, FDispatchSetupContext& Context)
+void FGPUSkinCache::DispatchSetupBackground(FRHICommandList& RHICmdList, FDispatchSetupContext& Context)
 {
 	int32 BatchCount = BatchDispatches.Num();
 
@@ -1415,73 +1450,88 @@ void FGPUSkinCache::DoDispatch(FRHICommandList& RHICmdList, FDispatchSetupContex
 		return;
 	}
 
-	INC_DWORD_STAT_BY(STAT_GPUSkinCache_TotalNumChunks, BatchCount);
-	TRACE_CPUPROFILER_EVENT_SCOPE(FGPUSkinCache::DoDispatch);
-
-	bool bCapture = BatchCount > 0 && GNumDispatchesToCapture > 0;
-	RenderCaptureInterface::FScopedCapture RenderCapture(bCapture, &RHICmdList);
-	GNumDispatchesToCapture -= bCapture ? 1 : 0;
-
-	RHI_BREADCRUMB_EVENT_STAT(RHICmdList, GPUSkinCache, "GPUSkinCache");
-	SCOPED_GPU_STAT(RHICmdList, GPUSkinCache);
-
-	TArray<FSkinCacheRWBuffer*> BuffersToTransitionToRead;
-
-	TArray<FRHIUnorderedAccessView*> SkinningBuffersToOverlap;
-	TransitionBufferUAVs(RHICmdList, Context.BuffersToTransition, SkinningBuffersToOverlap);
-	RHICmdList.BeginUAVOverlap(SkinningBuffersToOverlap);
+	Context.SortedDispatches.Reserve(BatchCount);
 	{
-		TRACE_CPUPROFILER_EVENT_SCOPE(GPUSkinCache_UpdateSkinningBatches);
-		SCOPED_DRAW_EVENT(RHICmdList, GPUSkinCache_UpdateSkinningBatches);
+		TRACE_CPUPROFILER_EVENT_SCOPE(BuildSortedDispatchList);
 
-		auto* GlobalShaderMap = GetGlobalShaderMap(GetFeatureLevel());
+		int32 NumRayTracingDispatches = 0;
 
-		TArray<FSortedDispatchEntry> SortedDispatches;
-		SortedDispatches.Reserve(BatchCount);
+		for (int32 i = 0; i < BatchCount; ++i)
 		{
-			TRACE_CPUPROFILER_EVENT_SCOPE(BuildSortedDispatchList);
+			FDispatchEntry& DispatchItem = BatchDispatches[i];
+			FGPUSkinCacheEntry* Entry = DispatchItem.SkinCacheEntry;
+			int32 Section = DispatchItem.Section;
+			FGPUSkinCacheEntry::FSectionDispatchData& DispatchData = Entry->DispatchData[Section];
 
-			for (int32 i = 0; i < BatchCount; ++i)
+			if (EnumHasAnyFlags(DispatchData.DispatchFlags, EGPUSkinCacheDispatchFlags::PositionPrevious | EGPUSkinCacheDispatchFlags::Position))
 			{
-				FDispatchEntry& DispatchItem = BatchDispatches[i];
-				FGPUSkinCacheEntry* Entry = DispatchItem.SkinCacheEntry;
-				int32 Section = DispatchItem.Section;
-				FGPUSkinCacheEntry::FSectionDispatchData& DispatchData = Entry->DispatchData[Section];
-
-				if ((DispatchData.DispatchFlags & ((uint32)EGPUSkinCacheDispatchFlags::DispatchPrevPosition | (uint32)EGPUSkinCacheDispatchFlags::DispatchPosition)) != 0)
+				// For 'unlimited' bone indexes, we pass in the index and weight sizes via a shader parameter and so we
+				// can re-use the same shader permutation as for 8-bit indexes.
+				bool bUse16BitBoneIndex = Entry->bUse16BitBoneIndex;
+				bool bUse16BitBoneWeight = Entry->bUse16BitBoneWeight;
+				if (Entry->BoneInfluenceType == 2)
 				{
-					// For 'unlimited' bone indexes, we pass in the index and weight sizes via a shader parameter and so we
-					// can re-use the same shader permutation as for 8-bit indexes.
-					bool bUse16BitBoneIndex = Entry->bUse16BitBoneIndex;
-					bool bUse16BitBoneWeight = Entry->bUse16BitBoneWeight;
-					if (Entry->BoneInfluenceType == 2)
-					{
-						bUse16BitBoneIndex = bUse16BitBoneWeight = false;
-					}
+					bUse16BitBoneIndex = bUse16BitBoneWeight = false;
+				}
 
-					FSortedDispatchEntry SortedEntry;
+				FSortedDispatchEntry SortedEntry;
 
-					SortedEntry.ShaderIndex = SKIN_CACHE_SHADER_IDX(
-						static_cast<int32>(bUse16BitBoneWeight),
-						static_cast<int32>(bUse16BitBoneIndex),
-						Entry->BoneInfluenceType,
-						DispatchData.SkinType);
-					SortedEntry.BatchIndex = i;
+				SortedEntry.ShaderIndex = SKIN_CACHE_SHADER_IDX(
+					static_cast<int32>(bUse16BitBoneWeight),
+					static_cast<int32>(bUse16BitBoneIndex),
+					Entry->BoneInfluenceType,
+					DispatchData.SkinType);
+				SortedEntry.BatchIndex = i;
 
-					SortedDispatches.Add(SortedEntry);
-				}
+				Context.SortedDispatches.Add(SortedEntry);
 			}
+		}
 
-			Algo::Sort(SortedDispatches,
-				[](const FSortedDispatchEntry& A, const FSortedDispatchEntry& B)
-				{
-					if (A.ShaderIndex != B.ShaderIndex)
-					{
-						return A.ShaderIndex < B.ShaderIndex;
-					}
+		Algo::Sort(Context.SortedDispatches, [](const FSortedDispatchEntry& A, const FSortedDispatchEntry& B)
+		{
+			if (A.ShaderIndex != B.ShaderIndex)
+			{
+				return A.ShaderIndex < B.ShaderIndex;
+			}
 			return A.BatchIndex < B.BatchIndex;
-				});
+		});
+	}
+
+	TransitionBufferUAVs(RHICmdList, Context.BuffersToTransition, Context.SkinningBuffersToOverlap);
+
+	{
+		TRACE_CPUPROFILER_EVENT_SCOPE(UpdateVertexFactoryDeclarations);
+
+		for (int32 i = 0; i < BatchCount; ++i)
+		{
+			FDispatchEntry& DispatchItem = BatchDispatches[i];
+			if (DispatchItem.SkinCacheEntry->Mode == EGPUSkinCacheEntryMode::Raster)
+			{
+				DispatchItem.SkinCacheEntry->UpdateVertexFactoryDeclaration(RHICmdList, DispatchItem.Section);
+			}
 		}
+	}
+}
+
+void FGPUSkinCache::DispatchPassExecute(FRHICommandList& RHICmdList, FDispatchSetupContext& Context)
+{
+	const int32 BatchCount = BatchDispatches.Num();
+
+	INC_DWORD_STAT_BY(STAT_GPUSkinCache_TotalNumChunks, BatchCount);
+	TRACE_CPUPROFILER_EVENT_SCOPE(FGPUSkinCache::DoDispatch);
+
+	bool bCapture = BatchCount > 0 && GNumDispatchesToCapture > 0;
+	RenderCaptureInterface::FScopedCapture RenderCapture(bCapture, &RHICmdList);
+	GNumDispatchesToCapture -= bCapture ? 1 : 0;
+
+	TArray<FSkinCacheRWBuffer*, FConcurrentLinearArrayAllocator> BuffersToTransitionToRead;
+	RHICmdList.BeginUAVOverlap(Context.SkinningBuffersToOverlap);
+
+	{
+		TRACE_CPUPROFILER_EVENT_SCOPE(GPUSkinCache_UpdateSkinningBatches);
+		SCOPED_DRAW_EVENT(RHICmdList, GPUSkinCache_UpdateSkinningBatches);
+
+		auto* GlobalShaderMap = GetGlobalShaderMap(GetFeatureLevel());
 
 		int32 LastShaderIndex = -1;
 		TShaderRef<FBaseGPUSkinCacheCS> Shader;
@@ -1491,10 +1541,10 @@ void FGPUSkinCache::DoDispatch(FRHICommandList& RHICmdList, FDispatchSetupContex
 		SKIN_CACHE_SHADER_ALL();
 #undef SKIN_CACHE_SHADER
 
-		int32 SortedCount = SortedDispatches.Num();
+		int32 SortedCount = Context.SortedDispatches.Num();
 		for (int32 i = 0; i < SortedCount; ++i)
 		{
-			const FSortedDispatchEntry& SortedEntry = SortedDispatches[i];
+			const FSortedDispatchEntry& SortedEntry = Context.SortedDispatches[i];
 			if (SortedEntry.ShaderIndex != LastShaderIndex)
 			{
 				TRACE_CPUPROFILER_EVENT_SCOPE(ChangeShader);
@@ -1538,7 +1588,7 @@ void FGPUSkinCache::DoDispatch(FRHICommandList& RHICmdList, FDispatchSetupContex
 
 			uint32 VertexCountAlign64 = FMath::DivideAndRoundUp(DispatchData.NumVertices, (uint32)64);
 
-			if ((DispatchData.DispatchFlags & (uint32)EGPUSkinCacheDispatchFlags::DispatchPrevPosition) != 0)
+			if (EnumHasAnyFlags(DispatchData.DispatchFlags, EGPUSkinCacheDispatchFlags::PositionPrevious))
 			{
 				const FVertexBufferAndSRV& PrevBoneBuffer = ShaderData.GetBoneBufferForReading(true);
 
@@ -1558,7 +1608,7 @@ void FGPUSkinCache::DoDispatch(FRHICommandList& RHICmdList, FDispatchSetupContex
 				BuffersToTransitionToRead.Add(DispatchData.GetPreviousPositionRWBuffer());
 			}
 
-			if ((DispatchData.DispatchFlags & (uint32)EGPUSkinCacheDispatchFlags::DispatchPosition) != 0)
+			if (EnumHasAnyFlags(DispatchData.DispatchFlags, EGPUSkinCacheDispatchFlags::Position))
 			{
 				const FVertexBufferAndSRV& BoneBuffer = ShaderData.GetBoneBufferForReading(false);
 
@@ -1584,7 +1634,11 @@ void FGPUSkinCache::DoDispatch(FRHICommandList& RHICmdList, FDispatchSetupContex
 				LastShaderIndex = -1;
 			}
 
-			BuffersToTransitionToRead.Add(DispatchData.GetTangentRWBuffer());
+			if (EnumHasAnyFlags(DispatchData.DispatchFlags, EGPUSkinCacheDispatchFlags::Tangents))
+			{
+				BuffersToTransitionToRead.Add(DispatchData.GetTangentRWBuffer());
+			}
+
 			check(DispatchData.PreviousPositionBuffer != DispatchData.PositionBuffer);
 		}
 
@@ -1593,84 +1647,88 @@ void FGPUSkinCache::DoDispatch(FRHICommandList& RHICmdList, FDispatchSetupContex
 			UnsetShaderParametersLegacyCS(RHICmdList, Shader);
 		}
 	}
-	RHICmdList.EndUAVOverlap(SkinningBuffersToOverlap);
 
-	// Do necessary buffer transitions before recomputing tangents
+	RHICmdList.EndUAVOverlap(Context.SkinningBuffersToOverlap);
+
 	TArray<FSkinCacheRWBuffer*> BuffersToSRVForRecomputeTangents;
 	TArray<FSkinCacheRWBuffer*> IntermediateAccumulatedTangentBuffers;
-	for (int32 i = 0; i < BatchCount; ++i)
+
+	for (FDispatchEntry& DispatchItem : BatchDispatches)
 	{
-		FDispatchEntry& DispatchItem = BatchDispatches[i];
 		FGPUSkinCacheEntry::FSectionDispatchData& DispatchData = DispatchItem.SkinCacheEntry->DispatchData[DispatchItem.Section];
-		if (DispatchData.IndexBuffer)
+
+		if (EnumHasAnyFlags(DispatchData.DispatchFlags, EGPUSkinCacheDispatchFlags::Tangents))
 		{
 			BuffersToSRVForRecomputeTangents.Add(DispatchData.GetPositionRWBuffer());
 			BuffersToSRVForRecomputeTangents.Add(DispatchData.GetActiveTangentRWBuffer());
+			BuffersToTransitionToRead.Add(DispatchData.GetPositionRWBuffer());
+
 			if (GRecomputeTangentsParallelDispatch)
 			{
 				IntermediateAccumulatedTangentBuffers.Add(DispatchData.GetIntermediateAccumulatedTangentBuffer());
 			}
-			BuffersToTransitionToRead.Add(DispatchData.GetPositionRWBuffer());
 		}
 	}
-	TransitionBuffers(RHICmdList, BuffersToSRVForRecomputeTangents, ERHIAccess::SRVCompute);
 
-	TArray<FRHIUnorderedAccessView*> IntermediateAccumulatedTangentBuffersToOverlap;
-	TransitionBufferUAVs(RHICmdList, IntermediateAccumulatedTangentBuffers, IntermediateAccumulatedTangentBuffersToOverlap);
-	RHICmdList.BeginUAVOverlap(IntermediateAccumulatedTangentBuffersToOverlap);
+	if (!BuffersToSRVForRecomputeTangents.IsEmpty())
 	{
-		SCOPED_DRAW_EVENT(RHICmdList, GPUSkinCache_RecomputeTangentsBatches);
-		FSkinCacheRWBuffer* StagingBuffer = nullptr;
-		for (int32 i = 0; i < BatchCount; ++i)
+		TArray<FRHIUnorderedAccessView*> IntermediateAccumulatedTangentBuffersToOverlap;
+		TransitionBuffers(RHICmdList, BuffersToSRVForRecomputeTangents, ERHIAccess::SRVCompute);
+		TransitionBufferUAVs(RHICmdList, IntermediateAccumulatedTangentBuffers, IntermediateAccumulatedTangentBuffersToOverlap);
+		RHICmdList.BeginUAVOverlap(IntermediateAccumulatedTangentBuffersToOverlap);
+
 		{
-			FDispatchEntry& DispatchItem = BatchDispatches[i];
-			if (DispatchItem.SkinCacheEntry->DispatchData[DispatchItem.Section].IndexBuffer)
+			SCOPED_DRAW_EVENT(RHICmdList, GPUSkinCache_RecomputeTangentsBatches);
+			FSkinCacheRWBuffer* StagingBuffer = nullptr;
+
+			for (FDispatchEntry& DispatchItem : BatchDispatches)
 			{
-				DispatchUpdateSkinTangents(RHICmdList, DispatchItem.SkinCacheEntry, DispatchItem.Section, StagingBuffer, true);
-				if (!GRecomputeTangentsParallelDispatch)
+				FGPUSkinCacheEntry::FSectionDispatchData& DispatchData = DispatchItem.SkinCacheEntry->DispatchData[DispatchItem.Section];
+				if (EnumHasAnyFlags(DispatchData.DispatchFlags, EGPUSkinCacheDispatchFlags::Tangents))
 				{
-					// When parallel dispatching is off, triangle pass and vertex pass are dispatched interleaved.
-					DispatchUpdateSkinTangents(RHICmdList, DispatchItem.SkinCacheEntry, DispatchItem.Section, StagingBuffer, false);
+					DispatchUpdateSkinTangents(RHICmdList, DispatchItem.SkinCacheEntry, DispatchItem.Section, StagingBuffer, true);
+
+					if (!GRecomputeTangentsParallelDispatch)
+					{
+						// When parallel dispatching is off, triangle pass and vertex pass are dispatched interleaved.
+						DispatchUpdateSkinTangents(RHICmdList, DispatchItem.SkinCacheEntry, DispatchItem.Section, StagingBuffer, false);
+					}
 				}
 			}
-		}
-		if (GRecomputeTangentsParallelDispatch)
-		{
-			// Do necessary buffer transitions before vertex pass dispatches
-			TArray<FSkinCacheRWBuffer*> TangentBuffers;
-			for (int32 i = 0; i < BatchCount; ++i)
-			{
-				FDispatchEntry& DispatchItem = BatchDispatches[i];
-				FGPUSkinCacheEntry::FSectionDispatchData& DispatchData = DispatchItem.SkinCacheEntry->DispatchData[DispatchItem.Section];
-				TangentBuffers.Add(DispatchData.GetTangentRWBuffer());
-				BuffersToTransitionToRead.Add(DispatchData.GetTangentRWBuffer());
-			}
-			TransitionBuffers(RHICmdList, IntermediateAccumulatedTangentBuffers, ERHIAccess::UAVCompute);
-		
-			TArray<FRHIUnorderedAccessView*> TangentBuffersToOverlap;
-			TransitionBufferUAVs(RHICmdList, TangentBuffers, TangentBuffersToOverlap);
-			RHICmdList.BeginUAVOverlap(TangentBuffersToOverlap);
-			for (int32 i = 0; i < BatchCount; ++i)
+
+			if (GRecomputeTangentsParallelDispatch)
 			{
-				FDispatchEntry& DispatchItem = BatchDispatches[i];
-				if (DispatchItem.SkinCacheEntry->DispatchData[DispatchItem.Section].IndexBuffer)
+				TArray<FSkinCacheRWBuffer*> TangentBuffers;
+
+				for (FDispatchEntry& DispatchItem : BatchDispatches)
 				{
-					DispatchUpdateSkinTangents(RHICmdList, DispatchItem.SkinCacheEntry, DispatchItem.Section, StagingBuffer, false);
+					FGPUSkinCacheEntry::FSectionDispatchData& DispatchData = DispatchItem.SkinCacheEntry->DispatchData[DispatchItem.Section];
+					if (EnumHasAnyFlags(DispatchData.DispatchFlags, EGPUSkinCacheDispatchFlags::Tangents))
+					{
+						TangentBuffers.Add(DispatchData.GetTangentRWBuffer());
+						BuffersToTransitionToRead.Add(DispatchData.GetTangentRWBuffer());
+					}
 				}
-			}
-			RHICmdList.EndUAVOverlap(TangentBuffersToOverlap);
-		}
-	}
-	RHICmdList.EndUAVOverlap(IntermediateAccumulatedTangentBuffersToOverlap);
 
-	{
-		TRACE_CPUPROFILER_EVENT_SCOPE(UpdateVertexFactoryDeclarations);
+				TArray<FRHIUnorderedAccessView*> TangentBuffersToOverlap;
+				TransitionBuffers(RHICmdList, IntermediateAccumulatedTangentBuffers, ERHIAccess::UAVCompute);
+				TransitionBufferUAVs(RHICmdList, TangentBuffers, TangentBuffersToOverlap);
+				RHICmdList.BeginUAVOverlap(TangentBuffersToOverlap);
 
-		for (int32 i = 0; i < BatchCount; ++i)
-		{
-			FDispatchEntry& DispatchItem = BatchDispatches[i];
-			DispatchItem.SkinCacheEntry->UpdateVertexFactoryDeclaration(RHICmdList, DispatchItem.Section);
+				for (FDispatchEntry& DispatchItem : BatchDispatches)
+				{
+					FGPUSkinCacheEntry::FSectionDispatchData& DispatchData = DispatchItem.SkinCacheEntry->DispatchData[DispatchItem.Section];
+					if (EnumHasAnyFlags(DispatchData.DispatchFlags, EGPUSkinCacheDispatchFlags::Tangents))
+					{
+						DispatchUpdateSkinTangents(RHICmdList, DispatchItem.SkinCacheEntry, DispatchItem.Section, StagingBuffer, false);
+					}
+				}
+
+				RHICmdList.EndUAVOverlap(TangentBuffersToOverlap);
+			}
 		}
+
+		RHICmdList.EndUAVOverlap(IntermediateAccumulatedTangentBuffersToOverlap);
 	}
 
 	{
@@ -1716,27 +1774,38 @@ bool FGPUSkinCache::ProcessEntry(
 		InvalidateAllEntries();
 	}
 
-	int32 RecomputeTangentsMode = GSkinCacheRecomputeTangents;
-	bool bShouldRecomputeTangent = false;
+	EGPUSkinCacheBufferBits BufferBits = EGPUSkinCacheBufferBits::None;
 
 	// IntermediateAccumulatedTangents buffer is needed if mesh has at least one section needing recomputing tangents.
-	uint32 InterAccumTangentBufferSize = 0;
-	uint32 CurrInterAccumTangentBufferOffset = 0;
-	if (RecomputeTangentsMode > 0)
-	{
-		for (int32 i = 0; i < LodData.RenderSections.Num(); ++i)
-		{			
-			const FSkelMeshRenderSection& RenderSection = LodData.RenderSections[i];
-			if (RecomputeTangentsMode == 1 || RenderSection.bRecomputeTangent)
+	uint32 IntermediateAccumulatedTangentBufferSize = 0;
+	uint32 IntermediateAccumulatedTangentBufferOffset = 0;
+
+	if (Mode == EGPUSkinCacheEntryMode::Raster)
+	{
+		if (GSkinCacheRecomputeTangents > 0)
+		{
+			for (int32 Index = 0; Index < LodData.RenderSections.Num(); ++Index)
 			{
-				bShouldRecomputeTangent = true;
-				InterAccumTangentBufferSize += RenderSection.GetNumVertices();
-				if (i < Section)
+				const FSkelMeshRenderSection& RenderSection = LodData.RenderSections[Index];
+				if (GSkinCacheRecomputeTangents == 1 || RenderSection.bRecomputeTangent)
 				{
-					CurrInterAccumTangentBufferOffset += RenderSection.GetNumVertices();
+					if (Index < Section)
+					{
+						IntermediateAccumulatedTangentBufferOffset += RenderSection.GetNumVertices();
+					}
+
+					if (RenderSection.RecomputeTangentsVertexMaskChannel < ESkinVertexColorChannel::None)
+					{
+						BufferBits |= EGPUSkinCacheBufferBits::IntermediateTangents;
+					}
+
+					BufferBits |= EGPUSkinCacheBufferBits::Tangents;
+					IntermediateAccumulatedTangentBufferSize += RenderSection.GetNumVertices();
 				}
 			}
 		}
+
+		BufferBits |= EGPUSkinCacheBufferBits::PositionPrevious;
 	}
 
 	if (InOutEntry)
@@ -1752,8 +1821,18 @@ bool FGPUSkinCache::ProcessEntry(
 			if (!InOutEntry->IsSectionValid(Section) || !InOutEntry->IsSourceFactoryValid(Section, VertexFactory))
 			{
 				// This section might not be valid yet, so set it up
-				InOutEntry->SetupSection(Section, InOutEntry->PositionAllocation, &LodData.RenderSections[Section], MorphVertexBuffer, ClothVertexBuffer, NumVertices, InputStreamStart, 
-											VertexFactory, TargetVertexFactory, CurrInterAccumTangentBufferOffset, SimData);
+				InOutEntry->SetupSection(
+					Section,
+					InOutEntry->PositionAllocation,
+					&LodData.RenderSections[Section],
+					MorphVertexBuffer,
+					ClothVertexBuffer,
+					NumVertices,
+					InputStreamStart,
+					VertexFactory,
+					TargetVertexFactory,
+					IntermediateAccumulatedTangentBufferOffset,
+					SimData);
 			}
 		}
 	}
@@ -1764,29 +1843,15 @@ bool FGPUSkinCache::ProcessEntry(
 		// If something caused the existing entry to be invalid, disable recreate logic for the rest of the function
 		bRecreating = false;
 
-		const bool WithTangents = true;
-		int32 TotalNumVertices = VertexFactory->GetNumVertices();
-		
 		// IntermediateTangents buffer is needed if mesh has at least one section using vertex color as recompute tangents blending mask
-		bool bEntryUseIntermediateTangents = false;
-		if (bShouldRecomputeTangent)
-		{
-			for (const FSkelMeshRenderSection& RenderSection : LodData.RenderSections)
-			{
-				if (RenderSection.RecomputeTangentsVertexMaskChannel < ESkinVertexColorChannel::None)
-				{
-					bEntryUseIntermediateTangents = true;
-					break;
-				}
-			}
-		}
+		int32 TotalNumVertices = VertexFactory->GetNumVertices();
 
-		FRWBuffersAllocation* NewPositionAllocation = TryAllocBuffer(TotalNumVertices, WithTangents, bEntryUseIntermediateTangents, InterAccumTangentBufferSize, RHICmdList, Skin->GetAssetPathName(LODIndex));
+		FRWBuffersAllocation* NewPositionAllocation = TryAllocBuffer(RHICmdList, BufferBits, TotalNumVertices, IntermediateAccumulatedTangentBufferSize, Skin->GetAssetPathName(LODIndex));
 		if (!NewPositionAllocation)
 		{
 			if (GSkinCachePrintMemorySummary > 0)
 			{
-				uint64 RequiredMemInBytes = FRWBuffersAllocation::CalculateRequiredMemory(TotalNumVertices, WithTangents, bEntryUseIntermediateTangents, InterAccumTangentBufferSize);
+				uint64 RequiredMemInBytes = FRWBuffersAllocation::CalculateRequiredMemory(BufferBits, TotalNumVertices, IntermediateAccumulatedTangentBufferSize);
 				UE_LOG(LogSkinCache, Warning, TEXT("FGPUSkinCache::ProcessEntry%s failed to allocate %.3fMB for mesh %s LOD%d, extra required memory increased to %.3fMB")
 					, Mode == EGPUSkinCacheEntryMode::RayTracing ? TEXT("[RT]") : TEXT("")
 					, RequiredMemInBytes / MBSize
@@ -1803,8 +1868,19 @@ bool FGPUSkinCache::ProcessEntry(
 		InOutEntry = new FGPUSkinCacheEntry(this, Skin, NewPositionAllocation, LODIndex, Mode);
 		InOutEntry->GPUSkin = Skin;
 
-		InOutEntry->SetupSection(Section, NewPositionAllocation, &LodData.RenderSections[Section], MorphVertexBuffer, ClothVertexBuffer, NumVertices, InputStreamStart, 
-									VertexFactory, TargetVertexFactory, CurrInterAccumTangentBufferOffset, SimData);
+		InOutEntry->SetupSection(
+			Section,
+			NewPositionAllocation,
+			&LodData.RenderSections[Section],
+			MorphVertexBuffer,
+			ClothVertexBuffer,
+			NumVertices,
+			InputStreamStart,
+			VertexFactory,
+			TargetVertexFactory,
+			IntermediateAccumulatedTangentBufferOffset,
+			SimData);
+
 		Entries.Add(InOutEntry);
 	}
 
@@ -1901,7 +1977,7 @@ bool FGPUSkinCache::ProcessEntry(
 
 	// Need to update the previous bone buffer pointer, so logic that checks if the bone buffers changed (FGPUSkinCache::FRWBufferTracker::Find)
 	// doesn't invalidate the previous frame position data.  Recreating the render state will have generated new bone buffers.
-	if (bRecreating)
+	if (bRecreating && Mode == EGPUSkinCacheEntryMode::Raster)
 	{
 		FGPUBaseSkinVertexFactory::FShaderDataType& ShaderData = VertexFactory->GetShaderData();
 		if (ShaderData.HasBoneBufferForReading(true))
@@ -2019,7 +2095,7 @@ void FGPUSkinCache::Release(FGPUSkinCacheEntry*& SkinCacheEntry)
 	}
 }
 
-void FGPUSkinCache::PrepareUpdateSkinning(FGPUSkinCacheEntry* Entry, int32 Section, uint32 RevisionNumber, TArray<FSkinCacheRWBuffer*>* OverlappedUAVs)
+void FGPUSkinCache::PrepareUpdateSkinning(FGPUSkinCacheEntry* Entry, int32 Section, FDispatchSetupContext& Context)
 {
 	FGPUSkinCacheEntry::FSectionDispatchData& DispatchData = Entry->DispatchData[Section];
 	FGPUBaseSkinVertexFactory::FShaderDataType& ShaderData = DispatchData.SourceVertexFactory->GetShaderData();
@@ -2030,34 +2106,62 @@ void FGPUSkinCache::PrepareUpdateSkinning(FGPUSkinCacheEntry* Entry, int32 Secti
 	uint32 CurrentRevision = ShaderData.GetRevisionNumber(false);
 	uint32 PreviousRevision = ShaderData.GetRevisionNumber(true);
 
-	DispatchData.DispatchFlags = 0;
+	DispatchData.DispatchFlags = EGPUSkinCacheDispatchFlags::None;
 
-	// Find existing buffers
-	DispatchData.PositionBuffer = DispatchData.PositionTracker.Find(BoneBuffer, CurrentRevision);
-	DispatchData.PreviousPositionBuffer = DispatchData.PositionTracker.Find(PrevBoneBuffer, PreviousRevision);
+	if (DispatchData.PositionTracker.Allocation->HasPreviousBuffer())
+	{
+		// Find existing buffers
+		DispatchData.PositionBuffer = DispatchData.PositionTracker.Find(BoneBuffer, CurrentRevision);
+		DispatchData.PreviousPositionBuffer = DispatchData.PositionTracker.Find(PrevBoneBuffer, PreviousRevision);
 
-	// Allocate buffers if not found, excluding buffers already in use.  Or make the current buffer distinct if it happens to equal previous.
-	if (!DispatchData.PositionBuffer || DispatchData.PositionBuffer == DispatchData.PreviousPositionBuffer)
+		// Allocate buffers if not found, excluding buffers already in use.  Or make the current buffer distinct if it happens to equal previous.
+		if (!DispatchData.PositionBuffer || DispatchData.PositionBuffer == DispatchData.PreviousPositionBuffer)
+		{
+			DispatchData.PositionBuffer = DispatchData.PositionTracker.AllocateUnused(BoneBuffer, CurrentRevision, DispatchData.PreviousPositionBuffer, &Context.BuffersToTransition);
+			DispatchData.DispatchFlags |= EGPUSkinCacheDispatchFlags::Position;
+		}
+		if (!DispatchData.PreviousPositionBuffer)
+		{
+			DispatchData.PreviousPositionBuffer = DispatchData.PositionTracker.AllocateUnused(PrevBoneBuffer, PreviousRevision, DispatchData.PositionBuffer, &Context.BuffersToTransition);
+			DispatchData.DispatchFlags |= EGPUSkinCacheDispatchFlags::PositionPrevious;
+		}
+	}
+	else
 	{
-		DispatchData.PositionBuffer = DispatchData.PositionTracker.AllocateUnused(BoneBuffer, CurrentRevision, DispatchData.PreviousPositionBuffer, OverlappedUAVs);
-		DispatchData.DispatchFlags |= (uint32)EGPUSkinCacheDispatchFlags::DispatchPosition;
+		DispatchData.PositionBuffer = &DispatchData.PositionTracker.Allocation->GetPositionBuffer();
+		DispatchData.DispatchFlags |= EGPUSkinCacheDispatchFlags::Position;
 	}
-	if (!DispatchData.PreviousPositionBuffer)
+
+	if (auto* TangentBuffer = DispatchData.PositionTracker.GetTangentBuffer())
 	{
-		DispatchData.PreviousPositionBuffer = DispatchData.PositionTracker.AllocateUnused(PrevBoneBuffer, PreviousRevision, DispatchData.PositionBuffer, OverlappedUAVs);
-		DispatchData.DispatchFlags |= (uint32)EGPUSkinCacheDispatchFlags::DispatchPrevPosition;
+		DispatchData.TangentBuffer = TangentBuffer;
+		DispatchData.DispatchFlags |= EGPUSkinCacheDispatchFlags::Tangents;
 	}
 
-	DispatchData.TangentBuffer = DispatchData.PositionTracker.GetTangentBuffer();
 	DispatchData.IntermediateTangentBuffer = DispatchData.PositionTracker.GetIntermediateTangentBuffer();
 	DispatchData.IntermediateAccumulatedTangentBuffer = DispatchData.PositionTracker.GetIntermediateAccumulatedTangentBuffer();
 
-	if (OverlappedUAVs && DispatchData.DispatchFlags != 0 && DispatchData.GetActiveTangentRWBuffer())
+	int32 NumBuffers = 0;
+	NumBuffers += DispatchData.PositionBuffer ? 1 : 0;
+	NumBuffers += DispatchData.PreviousPositionBuffer ? 1 : 0;
+	NumBuffers += DispatchData.TangentBuffer ? 1 : 0;
+	NumBuffers += DispatchData.IntermediateTangentBuffer ? 1 : 0;
+	NumBuffers += DispatchData.IntermediateAccumulatedTangentBuffer ? 1 : 0;
+
+	Context.Stats.NumBuffers += NumBuffers;
+	Context.Stats.NumRayTracingBuffers += Entry->Mode == EGPUSkinCacheEntryMode::RayTracing ? NumBuffers : 0;
+	Context.Stats.NumRayTracingDispatches += Entry->Mode == EGPUSkinCacheEntryMode::RayTracing ? 1 : 0;
+
+	if (DispatchData.DispatchFlags != EGPUSkinCacheDispatchFlags::None && DispatchData.GetActiveTangentRWBuffer())
 	{
-		(*OverlappedUAVs).Add(DispatchData.GetActiveTangentRWBuffer());
+		Context.BuffersToTransition.Add(DispatchData.GetActiveTangentRWBuffer());
 	}
 
 	check(DispatchData.PreviousPositionBuffer != DispatchData.PositionBuffer);
+
+	// Clear the flag that this is queued for dispatch.
+	Entry->bQueuedForDispatch = false;
+	DispatchData.RevisionNumber = 0;
 }
 
 void FGPUSkinCache::DispatchUpdateSkinning(FRHICommandList& RHICmdList, FGPUSkinCacheEntry* Entry, int32 Section, TArray<FSkinCacheRWBuffer*>& BuffersToTransitionToRead)
@@ -2108,8 +2212,8 @@ void FGPUSkinCache::DispatchUpdateSkinning(FRHICommandList& RHICmdList, FGPUSkin
 
 	check(Shader.IsValid());
 
-	const bool bDispatchPrevPosition = (DispatchData.DispatchFlags & (uint32)EGPUSkinCacheDispatchFlags::DispatchPrevPosition) != 0;
-	const bool bDispatchPosition = (DispatchData.DispatchFlags & (uint32)EGPUSkinCacheDispatchFlags::DispatchPosition) != 0;
+	const bool bDispatchPrevPosition = EnumHasAnyFlags(DispatchData.DispatchFlags, EGPUSkinCacheDispatchFlags::PositionPrevious);
+	const bool bDispatchPosition = EnumHasAnyFlags(DispatchData.DispatchFlags, EGPUSkinCacheDispatchFlags::Position);
 
 	if (bDispatchPrevPosition || bDispatchPosition)
 	{
@@ -2575,15 +2679,19 @@ void FGPUSkinCache::DrawVisualizationInfoText(const FName& GPUSkinCacheVisualiza
 	}
 }
 
-FGPUSkinCache::FRWBuffersAllocation::FRWBuffersAllocation(uint32 InNumVertices, bool InWithTangents, bool InUseIntermediateTangents, uint32 InIntermediateAccumulatedTangentsSize, FRHICommandList& RHICmdList, const FName& OwnerName)
-	: NumVertices(InNumVertices), WithTangents(InWithTangents), UseIntermediateTangents(InUseIntermediateTangents), IntermediateAccumulatedTangentsSize(InIntermediateAccumulatedTangentsSize)
+FGPUSkinCache::FRWBuffersAllocation::FRWBuffersAllocation(FRHICommandList& RHICmdList, EGPUSkinCacheBufferBits InBufferBits, uint32 InNumVertices, uint32 InIntermediateAccumulatedTangentsSize, const FName& OwnerName)
+	: NumVertices(InNumVertices)
+	, IntermediateAccumulatedTangentsSize(InIntermediateAccumulatedTangentsSize)
+	, BufferBits(InBufferBits)
 {
 	const static FLazyName PositionsName(TEXT("SkinCachePositions"));
 	const static FLazyName TangentsName(TEXT("SkinCacheTangents"));
 	const static FLazyName IntermediateTangentsName(TEXT("SkinCacheIntermediateTangents"));
 	const static FLazyName IntermediateAccumulatedTangentsName(TEXT("SkinCacheIntermediateAccumulatedTangents"));
 
-	for (int32 Index = 0; Index < NUM_BUFFERS; ++Index)
+	const int32 NumBuffers = EnumHasAnyFlags(InBufferBits, EGPUSkinCacheBufferBits::PositionPrevious) ? NUM_BUFFERS : 1;
+
+	for (int32 Index = 0; Index < NumBuffers; ++Index)
 	{
 		PositionBuffers[Index].Buffer.ClassName = PositionsName;
 		PositionBuffers[Index].Buffer.OwnerName = OwnerName;
@@ -2591,7 +2699,8 @@ FGPUSkinCache::FRWBuffersAllocation::FRWBuffersAllocation(uint32 InNumVertices,
 		PositionBuffers[Index].Buffer.Buffer->SetOwnerName(OwnerName);
 		PositionBuffers[Index].AccessState = ERHIAccess::Unknown;
 	}
-	if (WithTangents)
+
+	if (EnumHasAnyFlags(InBufferBits, EGPUSkinCacheBufferBits::Tangents))
 	{
 		// OpenGL ES does not support writing to RGBA16_SNORM images, instead pack data into SINT in the shader
 		const EPixelFormat TangentsFormat = IsOpenGLPlatform(GMaxRHIShaderPlatform) ? PF_R16G16B16A16_SINT : PF_R16G16B16A16_SNORM;
@@ -2601,7 +2710,8 @@ FGPUSkinCache::FRWBuffersAllocation::FRWBuffersAllocation(uint32 InNumVertices,
 		Tangents.Buffer.Initialize(RHICmdList, TEXT("SkinCacheTangents"), TangentBufferBytesPerElement, NumVertices * 2, TangentsFormat, BUF_Static);
 		Tangents.Buffer.Buffer->SetOwnerName(OwnerName);
 		Tangents.AccessState = ERHIAccess::Unknown;
-		if (UseIntermediateTangents)
+
+		if (EnumHasAnyFlags(InBufferBits, EGPUSkinCacheBufferBits::IntermediateTangents))
 		{
 			IntermediateTangents.Buffer.ClassName = IntermediateTangentsName;
 			IntermediateTangents.Buffer.OwnerName = OwnerName;
@@ -2610,6 +2720,7 @@ FGPUSkinCache::FRWBuffersAllocation::FRWBuffersAllocation(uint32 InNumVertices,
 			IntermediateTangents.AccessState = ERHIAccess::Unknown;
 		}
 	}
+
 	if (IntermediateAccumulatedTangentsSize > 0)
 	{
 		IntermediateAccumulatedTangents.Buffer.ClassName = IntermediateAccumulatedTangentsName;
@@ -2617,6 +2728,7 @@ FGPUSkinCache::FRWBuffersAllocation::FRWBuffersAllocation(uint32 InNumVertices,
 		IntermediateAccumulatedTangents.Buffer.Initialize(RHICmdList, TEXT("SkinCacheIntermediateAccumulatedTangents"), sizeof(int32), IntermediateAccumulatedTangentsSize * FGPUSkinCache::IntermediateAccumBufferNumInts, PF_R32_SINT, BUF_UnorderedAccess);
 		IntermediateAccumulatedTangents.Buffer.Buffer->SetOwnerName(OwnerName);
 		IntermediateAccumulatedTangents.AccessState = ERHIAccess::Unknown;
+
 		// The UAV must be zero-filled. We leave it zeroed after each round (see RecomputeTangentsPerVertexPass.usf), so this is only needed on when the buffer is first created.
 		RHICmdList.ClearUAVUint(IntermediateAccumulatedTangents.Buffer.UAV, FUintVector4(0, 0, 0, 0));
 	}
@@ -2628,27 +2740,44 @@ FGPUSkinCache::FRWBuffersAllocation::~FRWBuffersAllocation()
 	{
 		PositionBuffers[Index].Release();
 	}
-	if (WithTangents)
+
+	if (EnumHasAnyFlags(BufferBits, EGPUSkinCacheBufferBits::Tangents))
 	{
 		Tangents.Release();
 		IntermediateTangents.Release();
 	}
+
 	if (IntermediateAccumulatedTangentsSize > 0)
 	{
 		IntermediateAccumulatedTangents.Release();
 	}
 }
 
-uint64 FGPUSkinCache::FRWBuffersAllocation::CalculateRequiredMemory(uint32 InNumVertices, bool InWithTangents, bool InUseIntermediateTangents, uint32 InIntermediateAccumulatedTangentsSize)
+FGPUSkinCache::FSkinCacheRWBuffer* FGPUSkinCache::FRWBuffersAllocation::GetTangentBuffer()
 {
-	uint64 PositionBufferSize = PosBufferBytesPerElement * InNumVertices * 3 * NUM_BUFFERS;
-	uint64 TangentBufferSize = InWithTangents ? TangentBufferBytesPerElement * InNumVertices * 2 : 0;
-	uint64 IntermediateTangentBufferSize = 0;
-	if (InUseIntermediateTangents)
-	{
-		IntermediateTangentBufferSize = InWithTangents ? TangentBufferBytesPerElement * InNumVertices * 2 : 0;
-	}
-	uint64 AccumulatedTangentBufferSize = InIntermediateAccumulatedTangentsSize * FGPUSkinCache::IntermediateAccumBufferNumInts * sizeof(int32);
+	return EnumHasAllFlags(BufferBits, EGPUSkinCacheBufferBits::Tangents) ? &Tangents : nullptr;
+}
+
+FGPUSkinCache::FSkinCacheRWBuffer* FGPUSkinCache::FRWBuffersAllocation::GetIntermediateTangentBuffer()
+{
+	return EnumHasAllFlags(BufferBits, EGPUSkinCacheBufferBits::AllTangents) ? &IntermediateTangents : nullptr;
+}
+
+bool FGPUSkinCache::FRWBuffersAllocation::HasPreviousBuffer() const
+{
+	return EnumHasAllFlags(BufferBits, EGPUSkinCacheBufferBits::PositionPrevious);
+}
+
+uint64 FGPUSkinCache::FRWBuffersAllocation::CalculateRequiredMemory(EGPUSkinCacheBufferBits BufferBits, uint32 NumVertices, uint32 IntermediateAccumulatedTangentsSize)
+{
+	const bool bTangents = EnumHasAnyFlags(BufferBits, EGPUSkinCacheBufferBits::Tangents);
+	const bool bIntermediateTangents = EnumHasAnyFlags(BufferBits, EGPUSkinCacheBufferBits::IntermediateTangents);
+	const bool bPositionPrevious = EnumHasAnyFlags(BufferBits, EGPUSkinCacheBufferBits::PositionPrevious);
+
+	uint64 PositionBufferSize = PosBufferBytesPerElement * NumVertices * 3 * (bPositionPrevious ? NUM_BUFFERS : 1);
+	uint64 TangentBufferSize =  bTangents ? TangentBufferBytesPerElement * NumVertices * 2 : 0;
+	uint64 IntermediateTangentBufferSize = bIntermediateTangents ? TangentBufferBytesPerElement * NumVertices * 2 : 0;
+	uint64 AccumulatedTangentBufferSize = IntermediateAccumulatedTangentsSize * FGPUSkinCache::IntermediateAccumBufferNumInts * sizeof(int32);
 	return TangentBufferSize + IntermediateTangentBufferSize + PositionBufferSize + AccumulatedTangentBufferSize;
 }
 
diff --git a/Engine/Source/Runtime/Engine/Public/GPUSkinCache.h b/Engine/Source/Runtime/Engine/Public/GPUSkinCache.h
--- a/Engine/Source/Runtime/Engine/Public/GPUSkinCache.h
+++ b/Engine/Source/Runtime/Engine/Public/GPUSkinCache.h
@@ -47,6 +47,9 @@
 #include "DataDrivenShaderPlatformInfo.h"
 #include "Engine/SkinnedAssetCommon.h"
 
+enum class EGPUSkinCacheDispatchFlags : uint8;
+enum class EGPUSkinCacheBufferBits : uint8;
+
 class FGPUSkinPassthroughVertexFactory;
 class FGPUBaseSkinVertexFactory;
 class FMorphVertexBuffer;
@@ -118,7 +121,7 @@ public:
 	struct FDispatchEntry
 	{
 		FGPUSkinCacheEntry* SkinCacheEntry = nullptr;
-		uint32 Section = 0;	
+		uint32 Section = 0;
 	};
 
 	struct FSortedDispatchEntry
@@ -204,31 +207,34 @@ public:
 	{
 		friend struct FRWBufferTracker;
 
-		FRWBuffersAllocation(uint32 InNumVertices, bool InWithTangents, bool InUseIntermediateTangents, uint32 InIntermediateAccumulatedTangentsSize, FRHICommandList& RHICmdList, const FName& OwnerName);
+		static uint64 CalculateRequiredMemory(EGPUSkinCacheBufferBits BufferBits, uint32 NumVertices, uint32 IntermediateAccumulatedTangentsSize);
+
+		FRWBuffersAllocation(FRHICommandList& RHICmdList, EGPUSkinCacheBufferBits BufferBits, uint32 NumVertices, uint32 IntermediateAccumulatedTangentsSize, const FName& OwnerName);
 		~FRWBuffersAllocation();
 
-		static uint64 CalculateRequiredMemory(uint32 InNumVertices, bool InWithTangents, bool InUseIntermediateTangents, uint32 InIntermediateAccumulatedTangentsSize);
+		void AllocateTangents();
 
 		uint64 GetNumBytes() const
 		{
-			return CalculateRequiredMemory(NumVertices, WithTangents, UseIntermediateTangents, IntermediateAccumulatedTangentsSize);
+			return CalculateRequiredMemory(BufferBits, NumVertices, IntermediateAccumulatedTangentsSize);
 		}
 
-		FSkinCacheRWBuffer* GetTangentBuffer()
-		{
-			return WithTangents ? &Tangents : nullptr;
-		}
+		FSkinCacheRWBuffer* GetTangentBuffer();
 
-		FSkinCacheRWBuffer* GetIntermediateTangentBuffer()
-		{
-			return (WithTangents && UseIntermediateTangents) ? &IntermediateTangents : nullptr;
-		}
+		FSkinCacheRWBuffer* GetIntermediateTangentBuffer();
 
 		FSkinCacheRWBuffer* GetIntermediateAccumulatedTangentBuffer()
 		{
 			return IntermediateAccumulatedTangentsSize > 0 ? &IntermediateAccumulatedTangents : nullptr;
 		}
 
+		bool HasPreviousBuffer() const;
+
+		FSkinCacheRWBuffer& GetPositionBuffer()
+		{
+			return PositionBuffers[0];
+		}
+
 		void RemoveAllFromTransitionArray(TSet<FSkinCacheRWBuffer*>& BuffersToTransition);
 
 	private:
@@ -240,9 +246,8 @@ public:
 		FSkinCacheRWBuffer IntermediateAccumulatedTangents;	// Intermediate buffer used to accumulate results of triangle pass to be passed onto vertex pass
 
 		const uint32 NumVertices;
-		const bool WithTangents;
-		const bool UseIntermediateTangents;
 		const uint32 IntermediateAccumulatedTangentsSize;
+		EGPUSkinCacheBufferBits BufferBits;
 
 		static const uint32 PosBufferBytesPerElement = 4;
 		static const uint32 TangentBufferBytesPerElement = 8;
@@ -355,14 +360,6 @@ public:
 	inline bool HasWork() const { return !BatchDispatches.IsEmpty(); }
 
 protected:
-	struct FDispatchSetupContext
-	{
-		TArray<FSkinCacheRWBuffer*> BuffersToTransition;
-	};
-
-	void SetupDispatch(FRHICommandList& RHICmdList, FDispatchSetupContext& SetupContext);
-	void DoDispatch(FRHICommandList& RHICmdList, FDispatchSetupContext& SetupContext);
-
 	void WaitForDispatchSetupTask()
 	{
 		check(IsInRenderingThread());
@@ -374,6 +371,27 @@ protected:
 		}
 	}
 
+	struct FDispatchSetupContext
+	{
+		TArray<FSkinCacheRWBuffer*> BuffersToTransition;
+		TArray<FSortedDispatchEntry> SortedDispatches;
+		TArray<FRHIUnorderedAccessView*> SkinningBuffersToOverlap;
+
+		FRHICommandList* RHICmdListDispatch = nullptr;
+
+		struct
+		{
+			int32 NumRayTracingDispatches = 0;
+			int32 NumRayTracingBuffers = 0;
+			int32 NumBuffers = 0;
+
+		} Stats;
+	};
+
+	void DispatchSetupForeground(FRHICommandList& RHICmdList, FDispatchSetupContext& SetupContext);
+	void DispatchSetupBackground(FRHICommandList& RHICmdList, FDispatchSetupContext& SetupContext);
+	void DispatchPassExecute(FRHICommandList& RHICmdList, FDispatchSetupContext& SetupContext);
+
 	void TransitionBuffers(FRHICommandList& RHICmdList, TArray<FSkinCacheRWBuffer*>& Buffers, ERHIAccess ToState);
 	void TransitionBufferUAVs(FRHICommandList& RHICmdList, const TArray<FSkinCacheRWBuffer*>& InBuffers, TArray<FRHIUnorderedAccessView*>& OutUAVs);
 	void TransitionAllToReadable(FRHICommandList& RHICmdList, TConstArrayView<FSkinCacheRWBuffer*> BuffersToTransitionToRead);
@@ -383,15 +401,10 @@ protected:
 	TSet<FGPUSkinCacheEntry*> PendingProcessRTGeometryEntries;
 	TArray<FDispatchEntry> BatchDispatches;
 
-	FRWBuffersAllocation* TryAllocBuffer(uint32 NumVertices, bool WithTangnents, bool UseIntermediateTangents, uint32 InIntermediateAccumulatedTangentsSize, FRHICommandList& RHICmdList, const FName& OwnerName);
+	FRWBuffersAllocation* TryAllocBuffer(FRHICommandList& RHICmdList, EGPUSkinCacheBufferBits BufferBits, uint32 NumVertices, uint32 InIntermediateAccumulatedTangentsSize, const FName& OwnerName);
 	void DispatchUpdateSkinTangents(FRHICommandList& RHICmdList, FGPUSkinCacheEntry* Entry, int32 SectionIndex, FSkinCacheRWBuffer*& StagingBuffer, bool bTrianglePass);
 
-	void PrepareUpdateSkinning(
-		FGPUSkinCacheEntry* Entry, 
-		int32 Section, 
-		uint32 RevisionNumber, 
-		TArray<FSkinCacheRWBuffer*>* OverlappedUAVs
-		);
+	void PrepareUpdateSkinning(FGPUSkinCacheEntry* Entry, int32 Section, FDispatchSetupContext& Context);
 
 	void DispatchUpdateSkinning(
 		FRHICommandList& RHICmdList, 
@@ -438,5 +451,3 @@ DECLARE_MEMORY_STAT_EXTERN(TEXT("Total Memory Bytes Used"), STAT_GPUSkinCache_To
 DECLARE_MEMORY_STAT_EXTERN(TEXT("Intermediate buffer for Recompute Tangents"), STAT_GPUSkinCache_TangentsIntermediateMemUsed, STATGROUP_GPUSkinCache, );
 DECLARE_DWORD_COUNTER_STAT_EXTERN(TEXT("Num Triangles for Recompute Tangents"), STAT_GPUSkinCache_NumTrianglesForRecomputeTangents, STATGROUP_GPUSkinCache, );
 DECLARE_DWORD_COUNTER_STAT_EXTERN(TEXT("Num Sections Processed"), STAT_GPUSkinCache_NumSectionsProcessed, STATGROUP_GPUSkinCache, );
-DECLARE_DWORD_COUNTER_STAT_EXTERN(TEXT("Num SetVertexStreams"), STAT_GPUSkinCache_NumSetVertexStreams, STATGROUP_GPUSkinCache, );
-DECLARE_DWORD_COUNTER_STAT_EXTERN(TEXT("Num PreGDME"), STAT_GPUSkinCache_NumPreGDME, STATGROUP_GPUSkinCache, );
