Corrected the destruction order of the NetTest world instance to prevent some warnings while cleaning up actors

#rnx

diff --git a/Engine/Source/Runtime/Engine/Private/Tests/NetTestHelpers.cpp b/Engine/Source/Runtime/Engine/Private/Tests/NetTestHelpers.cpp
--- a/Engine/Source/Runtime/Engine/Private/Tests/NetTestHelpers.cpp
+++ b/Engine/Source/Runtime/Engine/Private/Tests/NetTestHelpers.cpp
@@ -177,8 +177,9 @@ void FTestWorldInstance::Shutdown()
 	if (World)
 	{
 		GEngine->ShutdownWorldNetDriver(World);
+		constexpr bool bInformEngineOfDestroyedWorld = true;
+		World->DestroyWorld(bInformEngineOfDestroyedWorld);
 		GEngine->DestroyWorldContext(World);
-		World->CleanupWorld();
 	}
 }
 
