Modulation Concurrency - Track all WaveInstances' Modulation data when ActiveSound is virtualized

diff --git a/Engine/Source/Runtime/AudioMixer/Private/AudioMixerSource.cpp b/Engine/Source/Runtime/AudioMixer/Private/AudioMixerSource.cpp
--- a/Engine/Source/Runtime/AudioMixer/Private/AudioMixerSource.cpp
+++ b/Engine/Source/Runtime/AudioMixer/Private/AudioMixerSource.cpp
@@ -1999,7 +1999,8 @@ namespace Audio
 		ActiveSound->bModulationRoutingUpdated = false;
 
 		// Query a modulation value for the active sound to use during concurrency evaluation
-		ActiveSound->SourceModulationVolume = MixerSourceVoice->GetVolumeModulationValue();
+		const float SourceModVolume = MixerSourceVoice->GetVolumeModulationValue();
+		ActiveSound->MaxSourceModulationValue = FMath::Max(SourceModVolume, ActiveSound->MaxSourceModulationValue);
 	}
 
 	void FMixerSource::UpdateSourceBusSends()
diff --git a/Engine/Source/Runtime/Engine/Private/ActiveSound.cpp b/Engine/Source/Runtime/Engine/Private/ActiveSound.cpp
--- a/Engine/Source/Runtime/Engine/Private/ActiveSound.cpp
+++ b/Engine/Source/Runtime/Engine/Private/ActiveSound.cpp
@@ -142,7 +142,7 @@ FActiveSound::FActiveSound()
 	, SubtitlePriority(DEFAULT_SUBTITLE_PRIORITY)
 	, Priority(1.0f)
 	, VolumeConcurrency(0.0f)
-	, SourceModulationVolume(0.0f)
+	, MaxSourceModulationValue(-1.0f)
 	, OcclusionCheckInterval(0.0f)
 	, LastOcclusionCheckTime(TNumericLimits<float>::Lowest())
 	, MaxDistance(FAudioDevice::GetMaxWorldDistance())
@@ -1063,7 +1063,16 @@ void FActiveSound::UpdateWaveInstances(TArray<FWaveInstance*> &InWaveInstances,
 
 			if (bIncludeModulationInConcurrencyEvaluationCVar)
 			{
-				VolumeConcurrency *= SourceModulationVolume;
+				// If there are no modulators set on any wave instance, then MaxSourceModulationValue may still be negative (see line below),
+				// so we take the absolute value in that case (1.0).
+				VolumeConcurrency *= FMath::Abs(MaxSourceModulationValue);
+
+				// We set the value as negative to mark it as "read", so the next place to update this value will always update it.
+				// By Multiplying by -1, we can reset the "read state" while preserving the previous value.
+				if (MaxSourceModulationValue >= 0.0f)
+				{
+					MaxSourceModulationValue *= -1.0f;
+				}
 			}
 
 			// Remove concurrency volume scalars as this can cause ping-ponging to occur with virtualization and loops
diff --git a/Engine/Source/Runtime/Engine/Private/Audio/AudioModulationInheritance.cpp b/Engine/Source/Runtime/Engine/Private/Audio/AudioModulationInheritance.cpp
--- a/Engine/Source/Runtime/Engine/Private/Audio/AudioModulationInheritance.cpp
+++ b/Engine/Source/Runtime/Engine/Private/Audio/AudioModulationInheritance.cpp
@@ -22,28 +22,10 @@ namespace AudioModulationInheritance
 		return InValueA;
 	}
 
-	// This function gets all the Modulators which apply to a sound, factoring in the inheritance path
-	// ie. Audio Component inherits modulation settings from Sound, and sounds inherit from sound class.
-	TSet<TObjectPtr<USoundModulatorBase>> GetEffectiveVolumeModulators(const FActiveSound& InActiveSound)
+	TSet<TObjectPtr<USoundModulatorBase>> GetEffectiveVolumeModulatorsFromWaveInstance(const FWaveInstance& InWaveInstance, const FActiveSound& InActiveSound)
 	{
-		FWaveInstance* WaveInstance = nullptr;
-		// For the "Sound" level of the inheritance path, we just grab the first WaveInstance we find
-		// Todo: How to resolve multiple wave instances with different modulation data? eg. Sound Cues
-		for (const TPair<UPTRINT, FWaveInstance*>& Pair : InActiveSound.GetWaveInstances())
-		{
-			WaveInstance = Pair.Value;
-			break;
-		}
-
-		EModulationRouting InWaveRouting = EModulationRouting::Inherit;
-		FSoundModulationDestinationSettings InWaveSettings = { };
-
-		if (WaveInstance)
-		{
-			UE_LOG(LogAudio, Verbose, TEXT("Using the first WaveInstance ('%s') on the ActiveSound to determine list of Modulators."), *WaveInstance->GetName());
-			InWaveRouting = WaveInstance->WaveData->ModulationSettings.VolumeRouting;
-			InWaveSettings = WaveInstance->WaveData->ModulationSettings.VolumeModulationDestination;
-		}
+		const EModulationRouting& InWaveRouting = InWaveInstance.WaveData->ModulationSettings.VolumeRouting;
+		const FSoundModulationDestinationSettings& InWaveSettings = InWaveInstance.WaveData->ModulationSettings.VolumeModulationDestination;
 
 		const USoundClass* InSoundClass = InActiveSound.GetSoundClass();
 		const Audio::FModulationParameter& InParam = Audio::GetModulationParameter("Volume");
@@ -72,7 +54,6 @@ namespace AudioModulationInheritance
 				UnionSoundClassSettings(UnionSettings);
 				return UnionSettings.Modulators;
 			}
-			break;
 
 			case EModulationRouting::Inherit:
 			{
@@ -94,7 +75,6 @@ namespace AudioModulationInheritance
 
 			return UnionSettings.Modulators;
 		}
-		break;
 
 		case EModulationRouting::Inherit:
 		{
@@ -106,7 +86,6 @@ namespace AudioModulationInheritance
 				UnionSoundClassSettings(UnionSettings);
 				return UnionSettings.Modulators;
 			}
-			break;
 
 			case EModulationRouting::Inherit:
 			{
@@ -121,7 +100,6 @@ namespace AudioModulationInheritance
 			{
 				return InWaveSettings.Modulators;
 			}
-			break;
 
 			case EModulationRouting::Disable:
 			default:
@@ -134,7 +112,6 @@ namespace AudioModulationInheritance
 		{
 			return InActiveSound.ModulationRouting.VolumeModulationDestination.Modulators;
 		}
-		break;
 
 		case EModulationRouting::Disable:
 		default:
@@ -143,4 +120,25 @@ namespace AudioModulationInheritance
 
 		return TSet<TObjectPtr<USoundModulatorBase>>();
 	}
-}
+
+	// This function gets all the Modulators which apply to a sound, factoring in the inheritance path
+	// i.e. Audio Component inherits modulation settings from Sound, and sounds inherit from sound class.
+	TArray<TSet<TObjectPtr<USoundModulatorBase>>> GetEffectiveVolumeModulatorsFromActiveSound(const FActiveSound& InActiveSound)
+	{
+		TArray<TSet<TObjectPtr<USoundModulatorBase>>> ModulatorSets;
+		for (const TPair<UPTRINT, FWaveInstance*>& Pair : InActiveSound.GetWaveInstances())
+		{
+			if (const FWaveInstance* WaveInstance = Pair.Value)
+			{
+				TSet<TObjectPtr<USoundModulatorBase>> VolumeModulators = GetEffectiveVolumeModulatorsFromWaveInstance(*WaveInstance, InActiveSound);
+				if (!VolumeModulators.IsEmpty())
+				{
+					ModulatorSets.Add(VolumeModulators);
+				}
+			}
+		}
+
+		return ModulatorSets;
+	}
+	
+} // namespace AudioModulationInheritance
diff --git a/Engine/Source/Runtime/Engine/Private/Audio/AudioModulationInheritance.h b/Engine/Source/Runtime/Engine/Private/Audio/AudioModulationInheritance.h
--- a/Engine/Source/Runtime/Engine/Private/Audio/AudioModulationInheritance.h
+++ b/Engine/Source/Runtime/Engine/Private/Audio/AudioModulationInheritance.h
@@ -1,4 +1,6 @@
 // Copyright Epic Games, Inc. All Rights Reserved.
+#pragma once
+
 #include "ActiveSound.h"
 #include "IAudioModulation.h"
 
@@ -6,6 +8,9 @@ namespace AudioModulationInheritance
 {
 	float MixInRoutedValue(const Audio::FModulationParameter& InParam, float InValueA, float InValueB);
 
-	TSet<TObjectPtr<USoundModulatorBase>> GetEffectiveVolumeModulators(const FActiveSound& InActiveSound);
+	TArray<TSet<TObjectPtr<USoundModulatorBase>>> GetEffectiveVolumeModulatorsFromActiveSound(const FActiveSound& InActiveSound);
+
+	TSet<TObjectPtr<USoundModulatorBase>> GetEffectiveVolumeModulatorsFromWaveInstance(const FWaveInstance& InWaveInstance, const FActiveSound& InActiveSound);
+
 }
 
diff --git a/Engine/Source/Runtime/Engine/Private/AudioVirtualLoop.cpp b/Engine/Source/Runtime/Engine/Private/AudioVirtualLoop.cpp
--- a/Engine/Source/Runtime/Engine/Private/AudioVirtualLoop.cpp
+++ b/Engine/Source/Runtime/Engine/Private/AudioVirtualLoop.cpp
@@ -132,8 +132,15 @@ bool FAudioVirtualLoop::Virtualize(const FActiveSound& InActiveSound, FAudioDevi
 
 	// Once virtualized, Modulation stops getting updated at the Mixer Source level. Because of this, we need to grab
 	// Modulation data now so that Volume modulation can be used to determine if a sound is loud enough to realize.
-	OutVirtualLoop.VolumeConcurrencyDestination.Init(InAudioDevice.DeviceID, AudioVirtualLoopPrivate::VolumeModulationParameterName, false /* bInIsBuffered */, true /* bInValueNormalized */);
-	OutVirtualLoop.VolumeConcurrencyDestination.UpdateModulators(AudioModulationInheritance::GetEffectiveVolumeModulators(InActiveSound));
+	// It's possible that the ActiveSound contains multiple WaveInstances.
+	// Because of this, we need to store one Modulation Destination for each WaveInstance.
+	TArray<TSet<TObjectPtr<USoundModulatorBase>>> ModulatorSets = AudioModulationInheritance::GetEffectiveVolumeModulatorsFromActiveSound(InActiveSound);
+	for (int32 DestinationIndex = 0; DestinationIndex < ModulatorSets.Num(); ++DestinationIndex)
+	{
+		OutVirtualLoop.VolumeConcurrencyDestinations.Add(Audio::FModulationDestination());
+		OutVirtualLoop.VolumeConcurrencyDestinations[DestinationIndex].Init(InAudioDevice.DeviceID, AudioVirtualLoopPrivate::VolumeModulationParameterName, false /* bInIsBuffered */, true /* bInValueNormalized */);
+		OutVirtualLoop.VolumeConcurrencyDestinations[DestinationIndex].UpdateModulators(ModulatorSets[DestinationIndex]);
+	}
 
 #if UE_AUDIO_PROFILERTRACE_ENABLED
 	const bool bChannelEnabled = UE_TRACE_CHANNELEXPR_IS_ENABLED(AudioChannel);
@@ -277,9 +284,21 @@ bool FAudioVirtualLoop::Update(float DeltaTime, bool bForceUpdate)
 		TimeSinceLastUpdate = 0.0f;
 	}
 
-	// ProcessControl only takes values in the Parameter's unit space, so we pass in 0dB as the default value.
-	VolumeConcurrencyDestination.ProcessControl(0.0f);
-	ActiveSound->SourceModulationVolume = VolumeConcurrencyDestination.GetValue();
+	if (!VolumeConcurrencyDestinations.IsEmpty())
+	{
+		float MaxModulationVolume = 0.0f;
+		for (Audio::FModulationDestination& Destination : VolumeConcurrencyDestinations)
+		{
+			// ProcessControl() only takes values in unit space (dB)
+			Destination.ProcessControl(0.0f);
+			const float DestVolume = Destination.GetValue();
+			if (DestVolume > MaxModulationVolume)
+			{
+				MaxModulationVolume = DestVolume;
+			}
+		}
+		ActiveSound->MaxSourceModulationValue = MaxModulationVolume;
+	}
 
 #if ENABLE_AUDIO_DEBUG
 	Audio::FAudioDebugger::DrawDebugInfo(*this);
diff --git a/Engine/Source/Runtime/Engine/Public/ActiveSound.h b/Engine/Source/Runtime/Engine/Public/ActiveSound.h
--- a/Engine/Source/Runtime/Engine/Public/ActiveSound.h
+++ b/Engine/Source/Runtime/Engine/Public/ActiveSound.h
@@ -597,8 +597,8 @@ public:
 	float VolumeConcurrency;
 
 	/** The source's final mixed Modulation value, used to evaluate VolumeConcurrency */
-	float SourceModulationVolume;
-
+	float MaxSourceModulationValue;
+	
 	/** The time in seconds with which to check for occlusion from its closest listener */
 	float OcclusionCheckInterval;
 
diff --git a/Engine/Source/Runtime/Engine/Public/AudioVirtualLoop.h b/Engine/Source/Runtime/Engine/Public/AudioVirtualLoop.h
--- a/Engine/Source/Runtime/Engine/Public/AudioVirtualLoop.h
+++ b/Engine/Source/Runtime/Engine/Public/AudioVirtualLoop.h
@@ -17,8 +17,11 @@ private:
 
 	FActiveSound* ActiveSound;
 
-	/** Modulation Destination for virtualized ActiveSounds to use when evaluating volume for concurrency */
-	Audio::FModulationDestination VolumeConcurrencyDestination;
+	/**
+	 * Array of Modulation Destinations for virtualized ActiveSounds to use when evaluating volume for concurrency
+	 * For sounds with multiple wave instance (e.g. Sound Cues), each WaveInstance needs its own Modulation Destination.
+	 */
+	TArray<Audio::FModulationDestination> VolumeConcurrencyDestinations;
 
 	/**
 	  * Check if provided active sound is in audible range.
