Skip attempting to set r.ScreenPercentage by scalability if sg.ResolutionQuality = 0.0. This prevents a warning from being issued on all projects that set r.ScreenPercentage in their project settings.

diff --git a/Engine/Source/Runtime/Engine/Private/Scalability.cpp b/Engine/Source/Runtime/Engine/Private/Scalability.cpp
--- a/Engine/Source/Runtime/Engine/Private/Scalability.cpp
+++ b/Engine/Source/Runtime/Engine/Private/Scalability.cpp
@@ -521,6 +521,7 @@ static void SetResolutionQualityLevel(float InResolutionQualityLevel)
 	if (InResolutionQualityLevel == 0.0)
 	{
 		// NOP: just use the project's default screen percentage.
+		return;
 	}
 	else
 	{
