Fix for missing inline on AnimationMath::Damper
Use FORCEINLINE instead of inline for other methods

diff --git a/Engine/Source/Runtime/Engine/Internal/Animation/AnimationMath.h b/Engine/Source/Runtime/Engine/Internal/Animation/AnimationMath.h
--- a/Engine/Source/Runtime/Engine/Internal/Animation/AnimationMath.h
+++ b/Engine/Source/Runtime/Engine/Internal/Animation/AnimationMath.h
@@ -10,18 +10,18 @@ namespace AnimationMath
 {
 	// Reference https://theorangeduck.com/page/spring-roll-call 
 
-	inline float SpringHalfLifeToDamping(float HalfLife, float eps = 1e-5f)
+	FORCEINLINE float SpringHalfLifeToDamping(float HalfLife, float eps = 1e-5f)
 	{
 		return (4.0f * 0.69314718056f) / (HalfLife + eps);
 	}
 
-	inline FVector QuatToScaledAngleAxis(const FQuat& Q)
+	FORCEINLINE FVector QuatToScaledAngleAxis(const FQuat& Q)
 	{
 		FQuat QLog = Q.Log();
 		return 2.0f * FVector(QLog.X, QLog.Y, QLog.Z);
 	}
 
-	inline FQuat QuatFromScaledAngleAxis(FVector V)
+	FORCEINLINE FQuat QuatFromScaledAngleAxis(FVector V)
 	{
 		V /= 2.0f;
 		FQuat QLog(V.X, V.Y, V.Z, 0.0f);
@@ -30,7 +30,7 @@ namespace AnimationMath
 
 	// Simplified version of FMath::CriticallyDampedSmoothing where v_goal is assumed to be 0
 	template<typename T>
-	inline void SimpleSpringDamper(
+	FORCEINLINE void SimpleSpringDamper(
 		T& x,
 		T& v,
 		T x_goal,
@@ -47,7 +47,7 @@ namespace AnimationMath
 	}
 
 	// Specialized quaternion version, uses FVector for velocity
-	inline void SimpleSpringDamperQuat(
+	FORCEINLINE void SimpleSpringDamperQuat(
 		FQuat& InOutRotation,
 		FVector& InOutAngularVelocity,
 		const FQuat& TargetRotation,
@@ -68,7 +68,7 @@ namespace AnimationMath
 	}
 
 	// Specialized quaternion damper, similar to FMath::ExponentialSmoothingApprox but for quats
-	FQuat Damper(const FQuat& CurrentRotation, const FQuat& TargetRotation, float HalfLife, float DeltaTime, float Eps = 1e-5f)
+	FORCEINLINE FQuat Damper(const FQuat& CurrentRotation, const FQuat& TargetRotation, float HalfLife, float DeltaTime, float Eps = 1e-5f)
 	{
 		if (HalfLife < Eps)
 			return TargetRotation;
