Add an option to disable entity replication
Disable replication on interop entities (transform_component was creating the replication_component)
#okforversepublic
#rnx

diff --git a/Engine/Source/Runtime/Engine/Private/Tests/NetTestHelpers.cpp b/Engine/Source/Runtime/Engine/Private/Tests/NetTestHelpers.cpp
--- a/Engine/Source/Runtime/Engine/Private/Tests/NetTestHelpers.cpp
+++ b/Engine/Source/Runtime/Engine/Private/Tests/NetTestHelpers.cpp
@@ -531,6 +531,43 @@ bool FTestWorlds::DoesReplicatedObjectExistOnClient(UObject* ServerObject, uint3
 	return IsValid(FindReplicatedObjectOnClient(ServerObject, ClientIndex));
 }
 
+bool FTestWorlds::IsServerObjectReplicated(const UObject* ServerObject, uint32 ClientIndex) const
+{
+	using namespace UE::Net;
+
+	if (ServerObject->GetWorld() != Server.GetWorld() && ServerObject->GetOuter() != GetTransientPackage())
+	{
+		ensureMsgf(false, TEXT("FTestWorlds::FindReplicatedObjectOnClient received object %s not part of the Server world"), *GetFullNameSafe(ServerObject));
+		return false;
+	}
+
+	const FTestWorldInstance::FContext ServerContext = Server.GetTestContext();
+	const FTestWorldInstance::FContext ClientContext = Clients[ClientIndex].GetTestContext();
+	
+	if (ServerContext.NetDriver->IsUsingIrisReplication())
+	{
+#if UE_WITH_IRIS
+		const FNetRefHandle NetHandle = ServerContext.IrisRepBridge->GetReplicatedRefHandle(ServerObject);
+		if (!NetHandle.IsValid())
+		{
+			return false;
+		}
+		return ClientContext.IrisRepBridge->GetReplicatedObject(NetHandle) != nullptr;
+#else
+		return false;
+#endif
+	}
+	else
+	{
+		const FNetworkGUID NetGUID = ServerContext.NetDriver->GetNetGuidCache()->GetNetGUID(ServerObject);
+		if (!NetGUID.IsValid())
+		{
+			return false;
+		}
+		return ClientContext.NetDriver->GetNetGuidCache()->GetObjectFromNetGUID(NetGUID, false) != nullptr;
+	}
+}
+
 //------------------------------------------------------------------------
 
 FScopedCVarOverrideInt::FScopedCVarOverrideInt(const TCHAR* VariableName, int32 Value)
diff --git a/Engine/Source/Runtime/Engine/Public/Tests/NetTestHelpers.h b/Engine/Source/Runtime/Engine/Public/Tests/NetTestHelpers.h
--- a/Engine/Source/Runtime/Engine/Public/Tests/NetTestHelpers.h
+++ b/Engine/Source/Runtime/Engine/Public/Tests/NetTestHelpers.h
@@ -216,12 +216,18 @@ struct FTestWorlds
 	}
 
 	/**
-	 * Returns if the the remote instance of a replicated object exists on the client or not.
+	 * Returns if the remote instance of a replicated object exists on the client or not.
 	 * @param ServerObject The server object you want to see if it is replicated
 	 * @param ClientIndex The client you want to test on
 	 */
 	ENGINE_API bool DoesReplicatedObjectExistOnClient(UObject* ServerObject, uint32 ClientIndex) const;
 
+	/**
+	 * Returns if the object is replicated on clients or not.
+	 * @param ServerObject The server object you want to see if it is replicated
+	 */
+	ENGINE_API bool IsServerObjectReplicated(const UObject* ServerObject, uint32 ClientIndex) const;
+
 public:
 
 	/** Server and Client Worlds */
