Various improvements to PrintObjectUtils:
- Using TStringBuilder to copy object state (e.g., object and class names, paths).
- Passing flags to every function for specifying optional behaviour (e.g., verbose output).
- Added an option to print an object's full archetype chain.

#rnx

diff --git a/Engine/Source/Runtime/Engine/Private/PrintObjectUtils.cpp b/Engine/Source/Runtime/Engine/Private/PrintObjectUtils.cpp
--- a/Engine/Source/Runtime/Engine/Private/PrintObjectUtils.cpp
+++ b/Engine/Source/Runtime/Engine/Private/PrintObjectUtils.cpp
@@ -10,6 +10,7 @@
 #include "UObject/Class.h"
 #include "UObject/UnrealType.h"
 #include "UObject/UObjectHash.h"
+#include "Misc/StringBuilder.h"
 
 #if WITH_EDITORONLY_DATA
 #include "UObject/InstanceDataObjectUtils.h"
@@ -19,7 +20,46 @@
 namespace UE
 {
 
-void PrintObjectsInOuter(UObject* Object, FOutputDevice* InOutputDevice)
+static void LogObjectState(const UObject* Object, FOutputDevice* OutputDevice, bool bVerbose = false, const TCHAR* Preamble = TEXT(""))
+{
+	TStringBuilder<32> NameBuilder;
+	TStringBuilder<32> ClassNameBuilder;
+	TStringBuilder<256> PathBuilder;
+
+	if (Object)
+	{
+		Object->GetFName().ToString(NameBuilder);
+		Object->GetPathName(/*StopOuter*/nullptr, PathBuilder);
+	}
+	else
+	{
+		NameBuilder.Append(TEXT("<NULL>"));
+		PathBuilder.Append(TEXT(""));
+	}
+
+	const UClass* ObjectClass = Object ? Object->GetClass() : nullptr;
+	if (ObjectClass)
+	{
+		ObjectClass->GetFName().ToString(ClassNameBuilder);
+	}
+	else
+	{
+		ClassNameBuilder.Append(TEXT("<NULL Class>"));
+	}
+
+	OutputDevice->Logf(TEXT("%s'%.*s' [%.*s] (%.*s)"), 
+		Preamble, NameBuilder.Len(), NameBuilder.GetData(), ClassNameBuilder.Len(), ClassNameBuilder.GetData(), PathBuilder.Len(), PathBuilder.GetData());
+
+	if (bVerbose)
+	{
+		EObjectFlags ObjectFlags = Object ? Object->GetFlags() : EObjectFlags::RF_NoFlags;
+		OutputDevice->Logf(TEXT("\tFlags: %s"), *LexToString(ObjectFlags));
+
+		OutputDevice->Logf(TEXT("\tAddress: 0x%x"), Object);
+	}
+}
+
+void PrintObjectsInOuter(UObject* Object, EPrintObjectFlag Flags, FOutputDevice* InOutputDevice)
 {
 	FOutputDevice* OutputDevice = InOutputDevice ? InOutputDevice : GLog;
 	if (OutputDevice)
@@ -29,14 +69,14 @@ void PrintObjectsInOuter(UObject* Object, FOutputDevice* InOutputDevice)
 			TArray<UObject*> ChildObjects;
 			GetObjectsWithOuter(Object, ChildObjects, /*bIncludeNestedObjects*/true);
 
-			OutputDevice->Logf(TEXT("Printing %d object(s) under: '%s' [%s] (%s)"),
-				ChildObjects.Num(), *Object->GetName(), *Object->GetClass()->GetName(), *Object->GetPathName());
+			bool bVerbose = EnumHasAnyFlags(Flags, EPrintObjectFlag::Verbose);
+
+			FString Preamble = FString::Printf(TEXT("Printing %d object(s) under: "), ChildObjects.Num());
+			LogObjectState(Object, OutputDevice, bVerbose, *Preamble);
 
 			for (UObject* ChildObject : ChildObjects)
 			{
-				FString ChildClassName = ChildObject->GetClass() ? *ChildObject->GetClass()->GetName() : TEXT("NULL Class");
-
-				OutputDevice->Logf(TEXT("%s [%s] (%s)"), *ChildObject->GetName(), *ChildClassName, *ChildObject->GetPathName());
+				LogObjectState(ChildObject, OutputDevice, bVerbose);
 			}
 		}
 		else
@@ -46,7 +86,7 @@ void PrintObjectsInOuter(UObject* Object, FOutputDevice* InOutputDevice)
 	}
 }
 
-void PrintObjectsWithName(const TCHAR* ObjectName, FOutputDevice* InOutputDevice)
+void PrintObjectsWithName(const TCHAR* ObjectName, EPrintObjectFlag Flags, FOutputDevice* InOutputDevice)
 {
 	FOutputDevice* OutputDevice = InOutputDevice ? InOutputDevice : GLog;
 	if (OutputDevice)
@@ -60,10 +100,11 @@ void PrintObjectsWithName(const TCHAR* ObjectName, FOutputDevice* InOutputDevice
 			{
 				OutputDevice->Logf(TEXT("Printing %d object(s) with name: '%s'"), Objects.Num(), ObjectName);
 
+				bool bVerbose = EnumHasAnyFlags(Flags, EPrintObjectFlag::Verbose);
+
 				for (UObject* Object : Objects)
 				{
-					FString ClassName = Object->GetClass() ? *Object->GetClass()->GetName() : TEXT("NULL Class");
-					OutputDevice->Logf(TEXT("%s [%s] (%s)"), *Object->GetName(), *ClassName, *Object->GetPathName());
+					LogObjectState(Object, OutputDevice, bVerbose);
 				}
 			}
 			else
@@ -78,7 +119,7 @@ void PrintObjectsWithName(const TCHAR* ObjectName, FOutputDevice* InOutputDevice
 	}
 }
 
-void PrintStructProperties(UStruct* Struct, void* StructData, bool bIncludeinitializationState, FOutputDevice* InOutputDevice)
+void PrintStructProperties(UStruct* Struct, void* StructData, EPrintObjectFlag Flags, FOutputDevice* InOutputDevice)
 {
 	FOutputDevice* OutputDevice = InOutputDevice ? InOutputDevice : GLog;
 	if (!OutputDevice)
@@ -102,6 +143,8 @@ void PrintStructProperties(UStruct* Struct, void* StructData, bool bIncludeiniti
 		// Keep track of visited property-owner pairs to avoid referencing cycles.
 		TSet<TTuple<const FProperty*, void*>> VisitedPropOwners;
 
+		bool bIncludeinitializationState = EnumHasAnyFlags(Flags, EPrintObjectFlag::PropertyInitializationState);
+
 		Struct->Visit(StructData, [&VisitedPropOwners, bIncludeinitializationState, OutputDevice](const FPropertyVisitorPath& PropertyPath, const FPropertyVisitorData& Data)->EPropertyVisitorControlFlow
 		{
 			const FProperty* Property = PropertyPath.Top().Property;
@@ -160,7 +203,7 @@ void PrintStructProperties(UStruct* Struct, void* StructData, bool bIncludeiniti
 	}
 }
 
-void PrintObjectProperties(UObject* Object, bool bIncludeinitializationState, FOutputDevice* InOutputDevice)
+void PrintObjectProperties(UObject* Object, EPrintObjectFlag Flags, FOutputDevice* InOutputDevice)
 {
 	FOutputDevice* OutputDevice = InOutputDevice ? InOutputDevice : GLog;
 	if (!OutputDevice)
@@ -170,13 +213,16 @@ void PrintObjectProperties(UObject* Object, bool bIncludeinitializationState, FO
 
 	if (Object)
 	{
-		FString ObjectPath = Object->GetPathName();
+		bool bVerbose = EnumHasAnyFlags(Flags, EPrintObjectFlag::Verbose);
+		bool bIncludeinitializationState = EnumHasAnyFlags(Flags, EPrintObjectFlag::PropertyInitializationState);
 
-		OutputDevice->Logf(TEXT("Printing properties of object '%s' [%s] (%s)"), *Object->GetName(), *Object->GetClass()->GetName(), *ObjectPath);
+		LogObjectState(Object, OutputDevice, bVerbose, TEXT("Printing properties of object "));
 
 		// Keep track of visited property-owner pairs to avoid referencing cycles.
 		TSet<TTuple<const FProperty*, void*>> VisitedPropOwners;
 
+		FString ObjectPath = Object->GetPathName();
+
 		Object->GetClass()->Visit(Object, [&VisitedPropOwners, Object, &ObjectPath, bIncludeinitializationState, OutputDevice](const FPropertyVisitorPath& PropertyPath, const FPropertyVisitorData& Data)->EPropertyVisitorControlFlow
 		{
 			const FProperty* Property = PropertyPath.Top().Property;
@@ -253,25 +299,34 @@ void PrintObjectProperties(UObject* Object, bool bIncludeinitializationState, FO
 	}
 }
 
-void PrintObjectArchetype(UObject* Object, FOutputDevice* InOutputDevice)
+void PrintObjectArchetype(UObject* Object, EPrintObjectFlag Flags, FOutputDevice* InOutputDevice)
 {
 	FOutputDevice* OutputDevice = InOutputDevice ? InOutputDevice : GLog;
 	if (OutputDevice)
 	{
 		if (Object)
 		{
-			FString ObjectPath = Object->GetPathName();
-
 			UObject* Archetype = Object->GetArchetype();
-			FString ArchetypeName = Archetype ? Archetype->GetName() : TEXT("<NULL>");
-			const UClass* ArchetypeClass = Archetype ? Archetype->GetClass() : nullptr;
-			FString ArchetypeClassName = ArchetypeClass ? ArchetypeClass->GetName() : FString();
-			FString ArchetypePath = Archetype ? Archetype->GetPathName() : FString();
 
-			OutputDevice->Logf(TEXT("Printing archetype for object '%s' [%s] (%s)"),
-				*Object->GetName(), *Object->GetClass()->GetName(), *ObjectPath);
+			bool bVerbose = EnumHasAnyFlags(Flags, EPrintObjectFlag::Verbose);
+			bool bShowFullChain = EnumHasAnyFlags(Flags, EPrintObjectFlag::FullArchetypeChain);
 
-			OutputDevice->Logf(TEXT("'%s' [%s] (%s)"), *ArchetypeName, *ArchetypeClassName, *ArchetypePath);
+			if (!bShowFullChain)
+			{
+				LogObjectState(Object, OutputDevice, bVerbose, TEXT("Printing archetype for object: "));
+				LogObjectState(Archetype, OutputDevice, bVerbose);
+			}
+			else
+			{
+				LogObjectState(Object, OutputDevice, bVerbose, TEXT("Printing archetype chain for object: "));
+
+				while (Archetype != nullptr)
+				{
+					LogObjectState(Archetype, OutputDevice, bVerbose);
+
+					Archetype = Archetype->GetArchetype();
+				}
+			}
 		}
 		else
 		{
@@ -280,7 +335,7 @@ void PrintObjectArchetype(UObject* Object, FOutputDevice* InOutputDevice)
 	}
 }
 
-void PrintObjectIDO(UObject* Object, FOutputDevice* InOutputDevice)
+void PrintObjectIDO(UObject* Object, EPrintObjectFlag Flags, FOutputDevice* InOutputDevice)
 {
 	FOutputDevice* OutputDevice = InOutputDevice ? InOutputDevice : GLog;
 	if (OutputDevice)
@@ -292,20 +347,18 @@ void PrintObjectIDO(UObject* Object, FOutputDevice* InOutputDevice)
 
 			UObject* IDO = UE::FPropertyBagRepository::Get().FindInstanceDataObject(Object);
 
+			bool bVerbose = EnumHasAnyFlags(Flags, EPrintObjectFlag::Verbose);
+
 			if (IDO)
 			{
-				const UClass* IDOClass = IDO->GetClass();
-				FString IDOClassName = IDOClass ? IDOClass->GetName() : FString();
-
-				OutputDevice->Logf(TEXT("Printing IDO (Instance Data Object) for object '%s' [%s] (%s)"),
-					*Object->GetName(), *Object->GetClass()->GetName(), *ObjectPath);
+				UObject* Archetype = Object->GetArchetype();
 
-				OutputDevice->Logf(TEXT("'%s' [%s] (%s)"), *IDO->GetName(), *IDOClassName, *IDO->GetPathName());
+				LogObjectState(Object, OutputDevice, bVerbose, TEXT("Printing IDO (Instance Data Object) for object: "));
+				LogObjectState(IDO, OutputDevice, bVerbose);
 			}
 			else
 			{
-				OutputDevice->Logf(TEXT("PrintObjectIDO: No IDO (Instance Data Object) found for object '%s' [%s] (%s)"),
-					*Object->GetName(), *Object->GetClass()->GetName(), *ObjectPath);
+				LogObjectState(Object, OutputDevice, bVerbose, TEXT("PrintObjectIDO: No IDO (Instance Data Object) found for object: "));
 			}
 		}
 		else
@@ -318,6 +371,39 @@ void PrintObjectIDO(UObject* Object, FOutputDevice* InOutputDevice)
 	}
 }
 
+static void ParseObjectIdAndFlags(const TArray<FString>& Args, FString& OutObjectIdentifier, EPrintObjectFlag& OutFlags)
+{
+	OutFlags = EPrintObjectFlag::None;
+
+	for (FString Arg : Args)
+	{
+		if (Arg.Equals(TEXT("Verbose=true"), ESearchCase::IgnoreCase))
+		{
+			EnumAddFlags(OutFlags, EPrintObjectFlag::Verbose);
+		}
+		else if (Arg.Equals(TEXT("InitState=true"), ESearchCase::IgnoreCase))
+		{
+			EnumAddFlags(OutFlags, EPrintObjectFlag::PropertyInitializationState);
+		}
+		else if (Arg.Equals(TEXT("ArchetypeChain=true"), ESearchCase::IgnoreCase))
+		{
+			EnumAddFlags(OutFlags, EPrintObjectFlag::FullArchetypeChain);
+		}
+		else
+		{
+			OutObjectIdentifier = Arg;
+		}
+	}
+}
+
+static void LogHelp(FOutputDevice& OutputDevice)
+{
+	OutputDevice.Logf(TEXT("Optional flags:"));
+	OutputDevice.Logf(TEXT("Verbose=true - Include verbose information"));
+	OutputDevice.Logf(TEXT("InitState=true - Include the initialization state for properties (only relevant for functions that print properties)."));
+	OutputDevice.Logf(TEXT("ArchetypeChain=true - Show the full archetype hierarchy (only relevant for functions that print archetypes)."));
+}
+
 
 FAutoConsoleCommandWithArgsAndOutputDevice CVarCommandPrintObjectsInOuter(
 	TEXT("Obj.DumpObjectsInOuter"),
@@ -326,12 +412,15 @@ FAutoConsoleCommandWithArgsAndOutputDevice CVarCommandPrintObjectsInOuter(
 	{
 		if (Args.Num() > 0)
 		{
-			FString ObjectPath = Args[0];
+			FString ObjectPath;
+			EPrintObjectFlag Flags;
+			ParseObjectIdAndFlags(Args, ObjectPath, Flags);
+
 			UObject* Object = FindFirstObjectSafe<UObject>(*ObjectPath);
 
 			if (Object)
 			{
-				PrintObjectsInOuter(Object, &OutputDevice);
+				PrintObjectsInOuter(Object, Flags, &OutputDevice);
 			}
 			else
 			{
@@ -340,7 +429,8 @@ FAutoConsoleCommandWithArgsAndOutputDevice CVarCommandPrintObjectsInOuter(
 		}
 		else
 		{
-			OutputDevice.Logf(TEXT("PrintObjectsInOuter: no object path specified (example usage: Obj.PrintObjectsInOuter /MyLevel/MyLevel.MyLevel:PersistentLevel)"));
+			OutputDevice.Logf(TEXT("PrintObjectsInOuter: no object path specified (example usage: Obj.DumpObjectsInOuter /MyLevel/MyLevel.MyLevel:PersistentLevel)"));
+			LogHelp(OutputDevice);
 		}
 	}));
 
@@ -351,43 +441,35 @@ FAutoConsoleCommandWithArgsAndOutputDevice CVarCommandPrintObjectsWithName(
 	{
 		if (Args.Num() > 0)
 		{
-			FString ObjectName = Args[0];
+			FString ObjectName;
+			EPrintObjectFlag Flags;
+			ParseObjectIdAndFlags(Args, ObjectName, Flags);
 
-			PrintObjectsWithName(*ObjectName, &OutputDevice);
+			PrintObjectsWithName(*ObjectName, Flags, &OutputDevice);
 		}
 		else
 		{
-			OutputDevice.Logf(TEXT("PrintObjectsWithName: no object name specified (example usage: Obj.PrintObjectsWithName PersistentLevel)"));
+			OutputDevice.Logf(TEXT("PrintObjectsWithName: no object name specified (example usage: Obj.DumpObjectsWithName PersistentLevel)"));
+			LogHelp(OutputDevice);
 		}
 	}));
 
 FAutoConsoleCommandWithArgsAndOutputDevice CVarCommandPrintObjectProperties(
-	TEXT("Obj.DumpObjectProperties"),
+	TEXT("Obj.DumpProperties"),
 	TEXT("Lists the properties of an object (the object must be specified as a path, e.g., /MyLevel/MyLevel.MyLevel:PersistentLevel.MyActor). Pass InitState=true to iclude the properties' initialization state."),
 	FConsoleCommandWithArgsAndOutputDeviceDelegate::CreateLambda([](const TArray<FString>& Args, FOutputDevice& OutputDevice)
 	{
 		if (Args.Num() > 0)
 		{
 			FString ObjectPath;
-			bool bIncludeinitializationState = false;
-
-			for (FString Arg : Args)
-			{
-				if (Arg.Equals(TEXT("InitState=true"), ESearchCase::IgnoreCase))
-				{
-					bIncludeinitializationState = true;
-				}
-				else
-				{
-					ObjectPath = Arg;
-				}
-			}
+			EPrintObjectFlag Flags;
+			ParseObjectIdAndFlags(Args, ObjectPath, Flags);
 
 			UObject* Object = FindFirstObjectSafe<UObject>(*ObjectPath);
 
 			if (Object)
 			{
-				PrintObjectProperties(Object, bIncludeinitializationState, &OutputDevice);
+				PrintObjectProperties(Object, Flags, &OutputDevice);
 			}
 			else
 			{
@@ -396,23 +478,27 @@ FAutoConsoleCommandWithArgsAndOutputDevice CVarCommandPrintObjectProperties(
 		}
 		else
 		{
-			OutputDevice.Logf(TEXT("PrintObjectProperties: no object path specified (example usage: Obj.PrintObjectProperties /MyLevel/MyLevel.MyLevel:PersistentLevel.MyActor)"));
+			OutputDevice.Logf(TEXT("PrintObjectProperties: no object path specified (example usage: Obj.DumpProperties /MyLevel/MyLevel.MyLevel:PersistentLevel.MyActor)"));
+			LogHelp(OutputDevice);
 		}
 	}));
 
 FAutoConsoleCommandWithArgsAndOutputDevice CVarCommandPrintObjectArchetype(
-	TEXT("Obj.DumpObjectArchetype"),
+	TEXT("Obj.DumpArchetype"),
 	TEXT("Outputs an object's archetype (the object must be specified as a path, e.g., /MyLevel/MyLevel.MyLevel:PersistentLevel.MyActor)."),
 	FConsoleCommandWithArgsAndOutputDeviceDelegate::CreateLambda([](const TArray<FString>& Args, FOutputDevice& OutputDevice)
 	{
 		if (Args.Num() > 0)
 		{
-			FString ObjectPath = Args[0];
+			FString ObjectPath;
+			EPrintObjectFlag Flags;
+			ParseObjectIdAndFlags(Args, ObjectPath, Flags);
+
 			UObject* Object = FindFirstObjectSafe<UObject>(*ObjectPath);
 
 			if (Object)
 			{
-				PrintObjectArchetype(Object, &OutputDevice);
+				PrintObjectArchetype(Object, Flags, &OutputDevice);
 			}
 			else
 			{
@@ -421,23 +507,27 @@ FAutoConsoleCommandWithArgsAndOutputDevice CVarCommandPrintObjectArchetype(
 		}
 		else
 		{
-			OutputDevice.Logf(TEXT("PrintObjectArchetype: no object path specified (example usage: Obj.PrintObjectArchetype /MyLevel/MyLevel.MyLevel:PersistentLevel.MyActor)"));
+			OutputDevice.Logf(TEXT("PrintObjectArchetype: no object path specified (example usage: Obj.DumpArchetype /MyLevel/MyLevel.MyLevel:PersistentLevel.MyActor)"));
+			LogHelp(OutputDevice);
 		}
 	}));
 
 FAutoConsoleCommandWithArgsAndOutputDevice CVarCommandPrintObjectIDO(
-	TEXT("Obj.DumpObjectIDO"),
+	TEXT("Obj.DumpIDO"),
 	TEXT("Outputs an object's IDO (Instance Data Object) (the object must be specified as a path, e.g., /MyLevel/MyLevel.MyLevel:PersistentLevel.MyActor)."),
 	FConsoleCommandWithArgsAndOutputDeviceDelegate::CreateLambda([](const TArray<FString>& Args, FOutputDevice& OutputDevice)
 	{
 		if (Args.Num() > 0)
 		{
-			FString ObjectPath = Args[0];
+			FString ObjectPath;
+			EPrintObjectFlag Flags;
+			ParseObjectIdAndFlags(Args, ObjectPath, Flags);
+			
 			UObject* Object = FindFirstObjectSafe<UObject>(*ObjectPath);
 
 			if (Object)
 			{
-				PrintObjectIDO(Object, &OutputDevice);
+				PrintObjectIDO(Object, Flags, &OutputDevice);
 			}
 			else
 			{
@@ -446,7 +536,8 @@ FAutoConsoleCommandWithArgsAndOutputDevice CVarCommandPrintObjectIDO(
 		}
 		else
 		{
-			OutputDevice.Logf(TEXT("PrintObjectIDO: no object path specified (example usage: Obj.PrintObjectArchetype /MyLevel/MyLevel.MyLevel:PersistentLevel.MyActor)"));
+			OutputDevice.Logf(TEXT("PrintObjectIDO: no object path specified (example usage: Obj.DumpArchetype /MyLevel/MyLevel.MyLevel:PersistentLevel.MyActor)"));
+			LogHelp(OutputDevice);
 		}
 	}));
 
diff --git a/Engine/Source/Runtime/Engine/Public/PrintObjectUtils.h b/Engine/Source/Runtime/Engine/Public/PrintObjectUtils.h
--- a/Engine/Source/Runtime/Engine/Public/PrintObjectUtils.h
+++ b/Engine/Source/Runtime/Engine/Public/PrintObjectUtils.h
@@ -3,6 +3,7 @@
 #include "HAL/Platform.h"			// For TCHAR
 #include "Misc/Build.h"				// For UE_BUILD_SHIPPING
 #include "Misc/CoreMiscDefines.h"	// For UE_INTERNAL
+#include "Misc/EnumClassFlags.h"
 
 #if !UE_BUILD_SHIPPING
 
@@ -18,56 +19,76 @@ class FOutputDevice;
 namespace UE
 {
 
+/**
+ * Flags for specifying optional behaviors.
+ */
+enum class EPrintObjectFlag : uint32
+{
+	/** No optional behaviors. */
+	None = 0,
+	/** Include verbose information. */
+	Verbose = (1 << 0),
+	/** Include the initialization state for properties (only relevant for functions that print properties). */
+	PropertyInitializationState = (1 << 1),
+	/** Show the full archetype hierarchy (only relevant for functions that print archetypes). */
+	FullArchetypeChain = (1 << 2),
+};
+
+ENUM_CLASS_FLAGS(EPrintObjectFlag);
+
 /**
  * Lists all objects under a specified parent.
  * 
  * @param Object The outer object to use.
  * @param OutputDevice The output device to print to (can be NULL).
  */
-UE_INTERNAL void PrintObjectsInOuter(UObject* Object, FOutputDevice* OutputDevice = nullptr);
+UE_INTERNAL void PrintObjectsInOuter(UObject* Object, EPrintObjectFlag Flags = EPrintObjectFlag::None, FOutputDevice* OutputDevice = nullptr);
 
 /**
  * Lists all objects with a given name.
  *
  * @param ObjectName The name to search for.
+ * @param Flags For specifying optional behaviors.
  * @param OutputDevice The output device to print to (can be NULL).
  */
-UE_INTERNAL void PrintObjectsWithName(const TCHAR* ObjectName, FOutputDevice* OutputDevice = nullptr);
+UE_INTERNAL void PrintObjectsWithName(const TCHAR* ObjectName, EPrintObjectFlag Flags = EPrintObjectFlag::None, FOutputDevice* OutputDevice = nullptr);
 
 /**
  * Lists the properties of a UStruct instance.
  *
  * @param Struct The type of the UStruct instance.
  * @param StructData The data of the UStruct instance.
- * @param bIncludeinitializationState If true the function will also output the initialization state of the properties.
+ * @param Flags For specifying optional behaviors.
  * @param OutputDevice The output device to print to (can be NULL).
  */
-UE_INTERNAL void PrintStructProperties(UStruct* Struct, void* StructData, bool bIncludeinitializationState = false, FOutputDevice* OutputDevice = nullptr);
+UE_INTERNAL void PrintStructProperties(UStruct* Struct, void* StructData, EPrintObjectFlag Flags = EPrintObjectFlag::None, FOutputDevice* OutputDevice = nullptr);
 
 /**
  * Lists the property values of an object.
  *
  * @param Object The object to use.
- * @param bIncludeinitializationState If true the function will also output the initialization state of the properties.
+ * @param Flags For specifying optional behaviors.
  * @param OutputDevice The output device to print to (can be NULL).
  */
-UE_INTERNAL void PrintObjectProperties(UObject* Object, bool bIncludeinitializationState = false, FOutputDevice* OutputDevice = nullptr);
+UE_INTERNAL void PrintObjectProperties(UObject* Object, EPrintObjectFlag Flags = EPrintObjectFlag::None, FOutputDevice* OutputDevice = nullptr);
 
 /**
  * Outputs an object's archetype.
  *
  * @param Object The object to use.
+ * @param Flags For specifying optional behaviors.
  * @param OutputDevice The output device to print to (can be NULL).
  */
-UE_INTERNAL void PrintObjectArchetype(UObject* Object, FOutputDevice* OutputDevice = nullptr);
+UE_INTERNAL void PrintObjectArchetype(UObject* Object, EPrintObjectFlag Flags = EPrintObjectFlag::None, FOutputDevice* OutputDevice = nullptr);
 
 /**
  * Outputs an object's IDO (Instance Data Object).
  *
  * @param Object The object to use.
+ * @param Flags For specifying optional behaviors.
  * @param OutputDevice The output device to print to (can be NULL).
  */
-UE_INTERNAL void PrintObjectIDO(UObject* Object, FOutputDevice* OutputDevice = nullptr);
+UE_INTERNAL void PrintObjectIDO(UObject* Object, EPrintObjectFlag Flags = EPrintObjectFlag::None, FOutputDevice* OutputDevice = nullptr);
 
 } // namespace UE
 
