Fixed stale deprecated data layers after removing a converted data layer instance which could lead to non-determinism cooks.

diff --git a/Engine/Source/Runtime/Engine/Private/WorldPartition/DataLayer/WorldDataLayers.cpp b/Engine/Source/Runtime/Engine/Private/WorldPartition/DataLayer/WorldDataLayers.cpp
--- a/Engine/Source/Runtime/Engine/Private/WorldPartition/DataLayer/WorldDataLayers.cpp
+++ b/Engine/Source/Runtime/Engine/Private/WorldPartition/DataLayer/WorldDataLayers.cpp
@@ -738,6 +738,17 @@ int32 AWorldDataLayers::RemoveDataLayers(const TArray<UDataLayerInstance*>& InDa
 			{
 				DeprecatedDataLayerNameToDataLayerInstance.Remove(DataLayerInstance->GetDataLayerFName());
 			}
+			else if (!DeprecatedDataLayerNameToDataLayerInstance.IsEmpty())
+			{
+				for (TMap<FName, TWeakObjectPtr<UDataLayerInstance>>::TIterator It = DeprecatedDataLayerNameToDataLayerInstance.CreateIterator(); It; ++It)
+				{
+					if (It.Value().Get() == DataLayerInstance)
+					{
+						It.RemoveCurrent();
+						break;
+					}
+				}
+			}
 			if (DataLayerInstance->IsPackageExternal())
 			{
 				DataLayerInstance->MarkAsGarbage();
