New material translator:
- "Weakened" our typing system in order to make arithmetic operator semantics behave like old translator instead of HLSL. Binary and ternary arithmetic operators will cast their arguments to the common type, so that "float + float3" results in a float3, as well as "lerp(float, float2, float4)" becomes the "largest type" float4.
- Small improvement to generated HLSL code that replaces dimensional values with the same component repeated over with a single cast. Now "float3(0.0, 0.0, 0.0)" will generate as "(float3)0.0".

#rnx

diff --git a/Engine/Source/Runtime/Engine/Private/Materials/MaterialIREmitter.cpp b/Engine/Source/Runtime/Engine/Private/Materials/MaterialIREmitter.cpp
--- a/Engine/Source/Runtime/Engine/Private/Materials/MaterialIREmitter.cpp
+++ b/Engine/Source/Runtime/Engine/Private/Materials/MaterialIREmitter.cpp
@@ -974,8 +974,7 @@ static FValue* TrySimplifyOperator(FEmitter* Emitter, EOperator& Op, FValue*& A,
 			}
 			else if (A->AreAllNearlyZero())
 			{
-				// return Emitter->EmitUnaryOperator(UO_Minus, Rhs);
-				UE_MIR_TODO();
+				return Emitter->EmitOperator(UO_Negate, A);
 			}
 			break;
 
@@ -1040,7 +1039,7 @@ static FValue* TryFoldOperatorScalar(FEmitter* Emitter, EOperator Op, FValue* A,
 	const FConstant* AConstant = As<FConstant>(A);
 	const FConstant* BConstant = As<FConstant>(B);
 	const FConstant* CConstant = As<FConstant>(C);
-	if (!AConstant ||  (IsBinaryOperator(Op) && !BConstant) || (IsTernaryOperator(Op) && !CConstant))
+	if (!AConstant || (IsBinaryOperator(Op) && !BConstant) || (IsTernaryOperator(Op) && (!BConstant || !CConstant)))
 	{
 		return nullptr;
 	}
@@ -1109,14 +1108,16 @@ static FValue* TryFoldOperatorScalar(FEmitter* Emitter, EOperator Op, FValue* A,
 /// Used to filter what parameter *primitive* types operators can take.
 enum EOperatorParameterTypeFilter
 {
-	OPTF_Unknown = 0xff,				///< Unspecified
-	OPTF_Any = 0,						///< Any primitive type (no filters)
-	OPTF_Boolean = 1 << 0,				///< Boolean primitive type of any dimension
-	OPTF_Arithmetic = 1 << 1,			///< Arithmetic primitive type of any dimension (i.e. that supports arithmetic operations)
-	OPTF_AnyFloat = 1 << 2,				///< Floating point primitive type of any dimension
-	OPTF_AsFirstArgument = 1 << 3,		///< Must be the same type as the first argument
-	OPTF_NotMatrix = 1 << 4,			///< Any primitive type except matrices
-	OPTF_Vector3 = 1 << 5,				///< 3D vector of any scalar type
+	OPTF_Unknown = 0xff,					///< Unspecified
+	OPTF_Any = 0,							///< Any primitive type (no filters)
+	OPTF_Boolean = 1 << 0,					///< Boolean primitive type of any dimension
+	OPTF_Arithmetic = 1 << 1,				///< Arithmetic primitive type of any dimension (i.e. that supports arithmetic operations)
+	OPTF_AnyFloat = 1 << 2,					///< Floating point primitive type of any dimension
+	OPTF_AsFirstArgument = 1 << 3,			///< Must be the same type as the first argument
+	OPTF_NotMatrix = 1 << 4,				///< Any primitive type except matrices
+	OPTF_Vector3 = 1 << 5,					///< 3D vector of any scalar type
+	OPTF_CastToCommon = 1 << 6,	 			///< An argument with this flag will be automatically cast to the common type between arguments.
+	OPTF_CastToCommonArithmetic = OPTF_Arithmetic | OPTF_CastToCommon,
 };
 
 inline EOperatorParameterTypeFilter operator|(EOperatorParameterTypeFilter A, EOperatorParameterTypeFilter B)
@@ -1146,8 +1147,8 @@ static const FOperatorSignature* GetOperatorSignature(EOperator Op)
 	static const FOperatorSignature* Signatures = [] ()
 	{
 		const FOperatorSignature UnaryArithmetic = { { OPTF_Arithmetic }, ORT_FirstArgumentType };
-		const FOperatorSignature BinaryArithmetic = { { OPTF_Arithmetic, OPTF_AsFirstArgument }, ORT_FirstArgumentType };
-		const FOperatorSignature TernaryArithmetic = { { OPTF_Arithmetic, OPTF_AsFirstArgument, OPTF_AsFirstArgument }, ORT_FirstArgumentType };
+		const FOperatorSignature BinaryArithmetic = { { OPTF_CastToCommonArithmetic, OPTF_CastToCommonArithmetic }, ORT_FirstArgumentType };
+		const FOperatorSignature TernaryArithmetic = { { OPTF_CastToCommonArithmetic, OPTF_CastToCommonArithmetic, OPTF_CastToCommonArithmetic }, ORT_FirstArgumentType };
 		const FOperatorSignature BinaryComparison = { { OPTF_Any, OPTF_AsFirstArgument }, ORT_BooleanWithFirstArgumentDimensions };
 		const FOperatorSignature BinaryLogical = { { OPTF_Boolean, OPTF_AsFirstArgument }, ORT_FirstArgumentType };
 
@@ -1202,7 +1203,7 @@ static const FOperatorSignature* GetOperatorSignature(EOperator Op)
  
 /// Validates that the types of the arguments are valid for specified operator.
 /// If valid, it returns the type of the result. Otherwise if it is not valid, it returns nullptr.
-static FPrimitiveTypePtr ValidateOperatorAndGetResultType(FEmitter* Emitter, EOperator Op, FValue* A, FValue*& B, FValue*& C)
+static FPrimitiveTypePtr ValidateOperatorAndGetResultType(FEmitter* Emitter, EOperator Op, FValue*& A, FValue*& B, FValue*& C)
 {
 	// Argument A must have always been provided.
 	check(A);
@@ -1221,6 +1222,7 @@ static FPrimitiveTypePtr ValidateOperatorAndGetResultType(FEmitter* Emitter, EOp
 	FPrimitiveTypePtr FirstArgumentPrimitiveType = A->Type->AsPrimitive();
 	static const TCHAR* ArgumentsStr[] = { TEXT("first"), TEXT("second"), TEXT("third") };
 	const FOperatorSignature* Signature = GetOperatorSignature(Op);
+	FTypePtr ArgumentsCommonType = FirstArgumentPrimitiveType;
 
 	bool bValid = true;
 	for (int32 i = 0; Arguments[i]; ++i)
@@ -1234,7 +1236,7 @@ static FPrimitiveTypePtr ValidateOperatorAndGetResultType(FEmitter* Emitter, EOp
 			Emitter->Errorf(TEXT("Operator %s's %s argument does not have primitive type."), OperatorToString(Op), ArgumentsStr[i]);
 			return nullptr;
 		}
-
+		
 		EOperatorParameterTypeFilter Filter = Signature->ParameterTypes[i];
 		check(Filter != OPTF_Unknown); // No signature specified for this operator.
 
@@ -1291,6 +1293,24 @@ static FPrimitiveTypePtr ValidateOperatorAndGetResultType(FEmitter* Emitter, EOp
 				bValid = false;
 			}
 		}
+		
+		// Update the common type.
+		if (i >= 1)
+		{
+			ArgumentsCommonType = Emitter->TryGetCommonType(ArgumentsCommonType, ArgumentPrimitiveType);
+		}
+	}
+	
+	// Now that we know the common type, cast each argument that requires it to it.
+	for (int32 i = 0; Arguments[i]; ++i)
+	{
+		EOperatorParameterTypeFilter Filter = Signature->ParameterTypes[i];
+		if (Filter & OPTF_CastToCommon)
+		{
+			check(ArgumentsCommonType && ArgumentsCommonType->AsPrimitive());
+			Arguments[i] = Emitter->EmitConstruct(ArgumentsCommonType, Arguments[i]);
+			bValid &= !Arguments[i]->IsPoison();
+		}
 	}
 
 	// Arguments might have changed, update the references.
@@ -1303,6 +1323,9 @@ static FPrimitiveTypePtr ValidateOperatorAndGetResultType(FEmitter* Emitter, EOp
 		return nullptr;
 	}
 	
+	// Update the first argument type, as it might have changed.
+	FirstArgumentPrimitiveType = Arguments[0]->Type->AsPrimitive();
+
 	// Finally, determine operator result type.
 	switch (Signature->ReturnType)
 	{
diff --git a/Engine/Source/Runtime/Engine/Private/Materials/MaterialIRToHLSLTranslator.cpp b/Engine/Source/Runtime/Engine/Private/Materials/MaterialIRToHLSLTranslator.cpp
--- a/Engine/Source/Runtime/Engine/Private/Materials/MaterialIRToHLSLTranslator.cpp
+++ b/Engine/Source/Runtime/Engine/Private/Materials/MaterialIRToHLSLTranslator.cpp
@@ -466,17 +466,33 @@ struct FTranslator : FMaterialIRToHLSLTranslation
 
 	void LowerDimensional(const MIR::FDimensional* Dimensional) 
 	{
+		TArrayView<MIR::FValue* const> Components = Dimensional->GetComponents();
 		MIR::FPrimitiveTypePtr ArithmeticType = Dimensional->Type->AsPrimitive();
-		check(ArithmeticType && ArithmeticType->IsVector());
-
-		Printer << ScalarKindToString(ArithmeticType->ScalarKind) << ArithmeticType->NumRows << BeginArgs;
-
-		for (MIR::FValue* Component : Dimensional->GetComponents())
+		check(ArithmeticType && !ArithmeticType->IsScalar());
+		
+		// In order to generate smaller and tidier HLSL, first check whether all components
+		// of this dimensional are actually the same. If so, we can simply emit the 
+		// component and cast it to the type.
+		bool bSameComponents = true;
+		for (int32 i = 1; bSameComponents && i < Components.Num(); ++i)
+		{
+			bSameComponents &= (Components[i] == Components[0]);
+		}
+		
+		if (bSameComponents)
 		{
-			Printer << ListSeparator << LowerValue(Component);
+			Printer << TEXT("(") << LowerType(ArithmeticType) << TEXT(")") << LowerValue(Components[0]);
+		}
+		else
+		{
+			Printer << LowerType(ArithmeticType) << BeginArgs;
+			for (MIR::FValue* Component : Components)
+			{
+				Printer << ListSeparator << LowerValue(Component);
+			}
+			Printer << EndArgs;
 		}
 
-		Printer << EndArgs;
 	}
 
 	void LowerSetMaterialOutput(const MIR::FSetMaterialOutput* Output)
