Fixed null dereference of previous bone vertex buffer for ray tracing skin cache entries.

diff --git a/Engine/Source/Runtime/Engine/Private/SkeletalRenderGPUSkin.cpp b/Engine/Source/Runtime/Engine/Private/SkeletalRenderGPUSkin.cpp
--- a/Engine/Source/Runtime/Engine/Private/SkeletalRenderGPUSkin.cpp
+++ b/Engine/Source/Runtime/Engine/Private/SkeletalRenderGPUSkin.cpp
@@ -814,9 +814,14 @@ void FSkeletalMeshObjectGPUSkin::UpdateBoneData(FRHICommandList& RHICmdList, EGP
 	const TArray<FSkelMeshRenderSection>& Sections = GetRenderSections(LODIndex);
 	const FName OwnerName = GetAssetPathName(LODIndex);
 
-	TConstArrayView<FMatrix44f> PreviousReferenceToLocalMatrices = DynamicData->GetPreviousReferenceToLocal(Mode);
+	TConstArrayView<FMatrix44f> PreviousReferenceToLocalMatrices;
 	TConstArrayView<FMatrix44f> ReferenceToLocalMatrices = DynamicData->GetReferenceToLocal(Mode);
-	
+
+	if (Mode == EGPUSkinCacheEntryMode::Raster)
+	{
+		PreviousReferenceToLocalMatrices = DynamicData->GetPreviousReferenceToLocal(Mode);
+	}
+
 	for (int32 SectionIdx = 0; SectionIdx < Sections.Num(); SectionIdx++)
 	{
 		const FSkelMeshRenderSection& Section = Sections[SectionIdx];
