LevelInstances: Reset loaders for all external packages for a given world, instead of having separate ResetLoaders call for each. This improves performance.
#rnx

diff --git a/Engine/Source/Runtime/Engine/Private/LevelInstance/LevelInstanceLevelStreaming.cpp b/Engine/Source/Runtime/Engine/Private/LevelInstance/LevelInstanceLevelStreaming.cpp
--- a/Engine/Source/Runtime/Engine/Private/LevelInstance/LevelInstanceLevelStreaming.cpp
+++ b/Engine/Source/Runtime/Engine/Private/LevelInstance/LevelInstanceLevelStreaming.cpp
@@ -331,13 +331,15 @@ void ULevelStreamingLevelInstance::ResetLevelInstanceLoaders()
 			LinkerLoad->DetachLoader();
 		}
 
-		ForEachObjectWithOuter(LoadedLevel, [](UObject* InObject)
+		TArray<UObject*> PackagesToReset;
+		ForEachObjectWithOuter(LoadedLevel, [&PackagesToReset](UObject* InObject)
 		{
 			if (InObject && InObject->IsPackageExternal())
 			{
-				ResetLoaders(InObject->GetExternalPackage());
+				PackagesToReset.Add(InObject->GetExternalPackage());
 			}
 		}, /*bIncludeNestedObjects*/ true);
+		ResetLoaders(PackagesToReset);
 
 		bResetLoadersCalled = true;
 	}
