Backed out change to insert GPUSkinCache as it breaks the PSO cache consolidation due to the command list being recorded prior and submitted after.

diff --git a/Engine/Source/Runtime/Engine/Private/GPUSkinCache.cpp b/Engine/Source/Runtime/Engine/Private/GPUSkinCache.cpp
--- a/Engine/Source/Runtime/Engine/Private/GPUSkinCache.cpp
+++ b/Engine/Source/Runtime/Engine/Private/GPUSkinCache.cpp
@@ -1346,8 +1346,6 @@ void FGPUSkinCache::TransitionBufferUAVs(FRHICommandList& RHICmdList, const TArr
 
 UE::Tasks::FTask FGPUSkinCache::DoDispatch(FRDGBuilder& GraphBuilder, const UE::Tasks::FTask& PrerequisitesTask)
 {
-	check(!DispatchSetupTask.IsValid());
-
 	if (!GraphBuilder.IsParallelSetupEnabled())
 	{
 		PrerequisitesTask.Wait();
@@ -1371,30 +1369,14 @@ UE::Tasks::FTask FGPUSkinCache::DoDispatch(FRDGBuilder& GraphBuilder, const UE::
 	{
 		FTaskTagScope TagScope(ETaskTag::EParallelRenderingThread);
 		DispatchSetupBackground(*RHICmdListSetup, *SetupContext);
+		DispatchPassExecute(*RHICmdListSetup, *SetupContext);
 		RHICmdListSetup->FinishRecording();
 
-		// Execute the dispatches in a separate command list that is enqueued into RDG.
-		FRHICommandList*& RHICmdListDispatch = SetupContext->RHICmdListDispatch;
-		RHICmdListDispatch = new FRHICommandList();
-		RHICmdListDispatch->SwitchPipeline(ERHIPipeline::Graphics);
-		DispatchPassExecute(*RHICmdListDispatch, *SetupContext);
-		RHICmdListDispatch->FinishRecording();
-
 	}, DispatchSetupTask, UE::Tasks::ETaskPriority::BackgroundHigh);
 
 	// Submit the setup command list here to start processing the RHI work ASAP.
 	GraphBuilder.RHICmdList.QueueAsyncCommandListSubmit(RHICmdListSetup);
 
-	// We use a dispatch pass here to insert the work into RDG, but the work is recorded during setup rather than kicking a task within the dispatch pass.
-	// This is because we can't record the dispatches asynchronously after the RDG frame ends because new render commands will modify the skin cache. Putting
-	// the work into RDG allows for stat scopes to work properly.
-	FEmptyShaderParameters* PassParameters = GraphBuilder.AllocParameters<FEmptyShaderParameters>();
-	GraphBuilder.AddDispatchPass(RDG_EVENT_NAME("GPUSkinCache"), PassParameters, ERDGPassFlags::Compute | ERDGPassFlags::NeverCull, [this, SetupContext](FRDGDispatchPassBuilder& DispatchPassBuilder)
-	{
-		DispatchPassBuilder.AddExistingCommandList(SetupContext->RHICmdListDispatch);
-		DispatchSetupTask = {};
-	});
-
 	return DispatchSetupTask;
 }
 
diff --git a/Engine/Source/Runtime/RenderCore/Private/RenderGraphPass.cpp b/Engine/Source/Runtime/RenderCore/Private/RenderGraphPass.cpp
--- a/Engine/Source/Runtime/RenderCore/Private/RenderGraphPass.cpp
+++ b/Engine/Source/Runtime/RenderCore/Private/RenderGraphPass.cpp
@@ -138,25 +138,6 @@ FRHICommandList* FRDGDispatchPassBuilder::CreateCommandList()
 	return RHICmdList;
 }
 
-void FRDGDispatchPassBuilder::AddExistingCommandList(FRHICommandList* RHICmdList)
-{
-#if RDG_ENABLE_DEBUG
-	for (FRHICommandListImmediate::FQueuedCommandList& QueuedCmdList : Pass->CommandLists)
-	{
-		checkf(QueuedCmdList.CmdList != RHICmdList, TEXT("Command list was already added!"));
-	}
-#endif
-
-	// When parallel executing we have to create a prologue command list to absorb the barriers and scopes.
-	if (Pass->bParallelExecute && Pass->CommandLists.IsEmpty())
-	{
-		FRHICommandList* RHICmdListPrologue = CreateCommandList();
-		RHICmdListPrologue->FinishRecording();
-	}
-
-	Pass->CommandLists.Emplace(RHICmdList);
-}
-
 void FRDGDispatchPassBuilder::Finish()
 {
 	// With serial execution the pass commands are embedded in the immediate command list instead.
diff --git a/Engine/Source/Runtime/RenderCore/Public/RenderGraphPass.h b/Engine/Source/Runtime/RenderCore/Public/RenderGraphPass.h
--- a/Engine/Source/Runtime/RenderCore/Public/RenderGraphPass.h
+++ b/Engine/Source/Runtime/RenderCore/Public/RenderGraphPass.h
@@ -752,12 +752,6 @@ public:
 	/** Create a new command list to record into and inserts it. Call FinishRecording() on the task when done. */
 	RENDERCORE_API FRHICommandList* CreateCommandList();
 
-	/** Insert a user-allocated command list into the dispatch pass. This version is useful if the command list
-     *  is recorded in a task that was launched prior to inserting the dispatch pass. Prefer CreateCommandList
-     *  if possible for performance. Call FinishRecording() on the task when done.
-     */
-	RENDERCORE_API void AddExistingCommandList(FRHICommandList* RHICmdList);
-
 private:
 	FRDGDispatchPassBuilder(FRDGDispatchPass* InPass)
 		: Pass(InPass)
