Fixed assert when running serial GPU skin cache path.

diff --git a/Engine/Source/Runtime/Engine/Private/GPUSkinCache.cpp b/Engine/Source/Runtime/Engine/Private/GPUSkinCache.cpp
--- a/Engine/Source/Runtime/Engine/Private/GPUSkinCache.cpp
+++ b/Engine/Source/Runtime/Engine/Private/GPUSkinCache.cpp
@@ -1346,6 +1346,8 @@ void FGPUSkinCache::TransitionBufferUAVs(FRHICommandList& RHICmdList, const TArr
 
 UE::Tasks::FTask FGPUSkinCache::DoDispatch(FRDGBuilder& GraphBuilder, const UE::Tasks::FTask& PrerequisitesTask)
 {
+	check(!DispatchSetupTask.IsValid());
+
 	if (!GraphBuilder.IsParallelSetupEnabled())
 	{
 		PrerequisitesTask.Wait();
@@ -1376,6 +1378,7 @@ UE::Tasks::FTask FGPUSkinCache::DoDispatch(FRDGBuilder& GraphBuilder, const UE::
 
 	// Submit the setup command list here to start processing the RHI work ASAP.
 	GraphBuilder.RHICmdList.QueueAsyncCommandListSubmit(RHICmdListSetup);
+	GraphBuilder.AddPostExecuteCallback([this] { DispatchSetupTask = {}; });
 
 	return DispatchSetupTask;
 }
