ForEachSubsystem improvement for FAudioSubsystemCollection

#rnx

diff --git a/Engine/Source/Runtime/Engine/Public/Subsystems/AudioEngineSubsystem.h b/Engine/Source/Runtime/Engine/Public/Subsystems/AudioEngineSubsystem.h
--- a/Engine/Source/Runtime/Engine/Public/Subsystems/AudioEngineSubsystem.h
+++ b/Engine/Source/Runtime/Engine/Public/Subsystems/AudioEngineSubsystem.h
@@ -85,12 +85,10 @@ public:
 	template<class InterfaceToCastTo>
 	void ForEachSubsystem(TFunctionRef<bool(InterfaceToCastTo*)> InFunction) const
 	{
-		FSubsystemCollection::ForEachSubsystem([Function=MoveTemp(InFunction)](USubsystem* Subsystem)
+		UClass* InterfaceClass = InterfaceToCastTo::UClassType::StaticClass();
+		ForEachSubsystemOfClass(InterfaceClass, [Function=MoveTemp(InFunction)](USubsystem* Subsystem)
 		{
-			if (InterfaceToCastTo* CastedSystem = Cast<InterfaceToCastTo>(Subsystem))
-			{
-				Function(CastedSystem);
-			}
+			Function(CastChecked<InterfaceToCastTo>(Subsystem));
 		});
 	}
 };
