Fixed some non-deterministic cook with HierarchicalInstancedStaticMesh components that get renamed after adding instances. The auto-generated seed is deterministic based on the full path, so being renamed before or after the actual rebuild would generate a different seed.

diff --git a/Engine/Source/Runtime/Engine/Private/HierarchicalInstancedStaticMesh.cpp b/Engine/Source/Runtime/Engine/Private/HierarchicalInstancedStaticMesh.cpp
--- a/Engine/Source/Runtime/Engine/Private/HierarchicalInstancedStaticMesh.cpp
+++ b/Engine/Source/Runtime/Engine/Private/HierarchicalInstancedStaticMesh.cpp
@@ -2458,6 +2458,9 @@ int32 UHierarchicalInstancedStaticMeshComponent::AddInstance(const FTransform& I
 {
 	SCOPE_CYCLE_COUNTER(STAT_HISMCAddInstance);
 
+	// Initialize the seed here even if the rebuild doesn't happen now to prevent races with being renamed by the time the mesh compiles
+	InitializeInstancingRandomSeed();
+
 	int32 InstanceIndex = UInstancedStaticMeshComponent::AddInstance(InstanceTransform, bWorldSpace);
 
 	// The tree will be fully rebuilt once the static mesh compilation is finished, no need for incremental update in that case.
@@ -2489,6 +2492,9 @@ TArray<int32> UHierarchicalInstancedStaticMeshComponent::AddInstances(const TArr
 {
 	SCOPE_CYCLE_COUNTER(STAT_HISMCAddInstances);
 
+	// Initialize the seed here even if the rebuild doesn't happen now to prevent races with being renamed by the time the mesh compiles
+	InitializeInstancingRandomSeed();
+
 	TArray<int32> InstanceIndices = UInstancedStaticMeshComponent::AddInstances(InstanceTransforms, true, bWorldSpace, bUpdateNavigation);
 
 	// The tree will be fully rebuilt once the static mesh compilation is finished, no need for incremental update in that case.
