VirtualCamera: Added asymmetrical overscan property to camera component which allows each edge of the view frustum to be overscanned by an independent amount, stacks with uniform overscan. As it is experimental and meant for use with VCam at the moment, not exposed to the editor or blueprints. Added blueprint functions to VCam library to allow VCam to support asymmetric overscan.

#virtualized

diff --git a/Engine/Plugins/VirtualProduction/VirtualCamera/Source/VirtualCamera/Private/FunctionLibraries/VCamBlueprintFunctionLibrary.cpp b/Engine/Plugins/VirtualProduction/VirtualCamera/Source/VirtualCamera/Private/FunctionLibraries/VCamBlueprintFunctionLibrary.cpp
--- a/Engine/Plugins/VirtualProduction/VirtualCamera/Source/VirtualCamera/Private/FunctionLibraries/VCamBlueprintFunctionLibrary.cpp
+++ b/Engine/Plugins/VirtualProduction/VirtualCamera/Source/VirtualCamera/Private/FunctionLibraries/VCamBlueprintFunctionLibrary.cpp
@@ -728,6 +728,24 @@ bool UVCamBlueprintFunctionLibrary::IsAssetDirty(const FAssetData& AssetData)
 	return false;
 }
 
+void UVCamBlueprintFunctionLibrary::SetCameraOverscan(UCineCameraComponent* InCameraComponent, const FVector4f& InAsymmetricOverscan, bool bScaleResolution)
+{
+	if (InCameraComponent)
+	{
+		InCameraComponent->AsymmetricOverscan = InAsymmetricOverscan;
+		InCameraComponent->bScaleResolutionWithOverscan = bScaleResolution;
+	}
+}
+
+void UVCamBlueprintFunctionLibrary::ClearCameraOverscan(UCineCameraComponent* InCameraComponent)
+{
+	if (InCameraComponent)
+	{
+		InCameraComponent->AsymmetricOverscan = FVector4f::Zero();
+		InCameraComponent->bScaleResolutionWithOverscan = false;
+	}
+}
+
 #if WITH_EDITOR
 TWeakPtr<ISequencer> UVCamBlueprintFunctionLibrary::GetSequencer()
 {
diff --git a/Engine/Plugins/VirtualProduction/VirtualCamera/Source/VirtualCamera/Public/FunctionLibraries/VCamBlueprintFunctionLibrary.h b/Engine/Plugins/VirtualProduction/VirtualCamera/Source/VirtualCamera/Public/FunctionLibraries/VCamBlueprintFunctionLibrary.h
--- a/Engine/Plugins/VirtualProduction/VirtualCamera/Source/VirtualCamera/Public/FunctionLibraries/VCamBlueprintFunctionLibrary.h
+++ b/Engine/Plugins/VirtualProduction/VirtualCamera/Source/VirtualCamera/Public/FunctionLibraries/VCamBlueprintFunctionLibrary.h
@@ -247,6 +247,21 @@ public:
 	UFUNCTION(BlueprintPure, Category = "VirtualCamera")
 	static bool IsAssetDirty(const FAssetData& AssetData);
 
+	/**
+	 * Adds asymmetric overscan to the camera.
+	 *
+	 * @param InCameraComponent - The camera to add overscan to
+	 * @param InAsymmetricOverscan - The amount of asymmetric overscan to apply, from 0.0 meaning no overscan to 1.0 meaning 100% overscan. By convention,
+	 *								 X is the left overscan, Y is the right overscan, Z is the top overscan, and W is the bottom overscan.
+	 *	@param bScaleResolution - Indicates if the render resolution should be scaled up so that the original frustum of the camera remains the same resolution
+	 */
+	UFUNCTION(BlueprintCallable, Category = "VirtualCamera")
+	static void SetCameraOverscan(UCineCameraComponent* InCameraComponent, const FVector4f& InAsymmetricOverscan, bool bScaleResolution = false);
+
+	/** Clears any overscan from the camera */
+	UFUNCTION(BlueprintCallable, Category = "VirtualCamera")
+	static void ClearCameraOverscan(UCineCameraComponent* InCameraComponent);
+	
 private:
 
 #if WITH_EDITOR
diff --git a/Engine/Source/Editor/UnrealEd/Private/EditorViewportClient.cpp b/Engine/Source/Editor/UnrealEd/Private/EditorViewportClient.cpp
--- a/Engine/Source/Editor/UnrealEd/Private/EditorViewportClient.cpp
+++ b/Engine/Source/Editor/UnrealEd/Private/EditorViewportClient.cpp
@@ -1330,6 +1330,7 @@ FSceneView* FEditorViewportClient::CalcSceneView(FSceneViewFamily* ViewFamily, c
 		ViewInitOptions.FirstPersonParams = FFirstPersonParameters(ControllingActorViewInfo.CalculateFirstPersonFOVCorrectionFactor(), ControllingActorViewInfo.FirstPersonScale, ControllingActorViewInfo.bUseFirstPersonParameters);
 		ViewInitOptions.OverscanResolutionFraction = ControllingActorViewInfo.OverscanResolutionFraction;
 		ViewInitOptions.CropFraction = ControllingActorViewInfo.CropFraction;
+		ViewInitOptions.AsymmetricCropFraction = ControllingActorViewInfo.AsymmetricCropFraction;
 	}
 
 	ViewInitOptions.OverrideFarClippingPlaneDistance = FarPlane;
diff --git a/Engine/Source/Runtime/CinematicCamera/Private/CineCameraComponent.cpp b/Engine/Source/Runtime/CinematicCamera/Private/CineCameraComponent.cpp
--- a/Engine/Source/Runtime/CinematicCamera/Private/CineCameraComponent.cpp
+++ b/Engine/Source/Runtime/CinematicCamera/Private/CineCameraComponent.cpp
@@ -296,7 +296,7 @@ float UCineCameraComponent::GetHorizontalFieldOfViewInternal(bool bIncludeOversc
 			}
 		}
 
-		const float OverscanScalar = bIncludeOverscan ? (1.0f + Overscan) : 1.0f;
+		const float OverscanScalar = bIncludeOverscan ? (1.0f + Overscan) * 0.5f * (AsymmetricOverscan.X + AsymmetricOverscan.Y + 2.0f) : 1.0f;
 		return FMath::RadiansToDegrees(2.f * FMath::Atan(CropedSensorWidth * OverscanScalar / (2.f * CurrentFocalLength)));
 	}
 
@@ -322,7 +322,7 @@ float UCineCameraComponent::GetVerticalFieldOfViewInternal(bool bIncludeOverscan
 			}
 		}
 
-		const float OverscanScalar = bIncludeOverscan ? (1.0f + Overscan) : 1.0f;
+		const float OverscanScalar = bIncludeOverscan ? (1.0f + Overscan) * 0.5f * (AsymmetricOverscan.Z + AsymmetricOverscan.W + 2.0f) : 1.0f;
 		return FMath::RadiansToDegrees(2.f * FMath::Atan(CropedSensorHeight * OverscanScalar / (2.f * CurrentFocalLength)));
 	}
 
@@ -341,7 +341,9 @@ float UCineCameraComponent::GetHorizontalProjectionOffset() const
 		}
 	}
 
-	return 2.0f * Filmback.SensorHorizontalOffset / (CroppedSensorWidth * (1.0f + Overscan));
+	const float OverscanScalar = (1.0f + Overscan) * 0.5f * (AsymmetricOverscan.X + AsymmetricOverscan.Y + 2.0f);
+	const float AsymmetricOverscanOffset = (AsymmetricOverscan.Y - AsymmetricOverscan.X) / (AsymmetricOverscan.X + AsymmetricOverscan.Y + 2.0f);
+	return 2.0f * Filmback.SensorHorizontalOffset / (CroppedSensorWidth * OverscanScalar) + AsymmetricOverscanOffset;
 }
 
 float UCineCameraComponent::GetVerticalProjectionOffset() const
@@ -356,7 +358,9 @@ float UCineCameraComponent::GetVerticalProjectionOffset() const
 		}
 	}
 
-	return 2.0f * Filmback.SensorVerticalOffset / (CroppedSensorHeight * (1.0f + Overscan));
+	const float OverscanScalar = (1.0f + Overscan) * 0.5f * (AsymmetricOverscan.Z + AsymmetricOverscan.W + 2.0f);
+	const float AsymmetricOverscanOffset = (AsymmetricOverscan.Z - AsymmetricOverscan.W) / (AsymmetricOverscan.Z + AsymmetricOverscan.W + 2.0f);
+	return 2.0f * Filmback.SensorVerticalOffset / (CroppedSensorHeight * OverscanScalar) + AsymmetricOverscanOffset;
 }
 
 FString UCineCameraComponent::GetFilmbackPresetName() const
@@ -743,11 +747,15 @@ void UCineCameraComponent::UpdateCameraLens(float DeltaTime, FMinimalViewInfo& D
 		DesiredView.PostProcessSettings.bOverride_DepthOfFieldFocalDistance = true;
 		DesiredView.PostProcessSettings.DepthOfFieldFocalDistance = CurrentFocusDistance;
 
+		const float OverscanScalar = (1.0f + Overscan) * 0.5f * (AsymmetricOverscan.X + AsymmetricOverscan.Y + 2.0f);
 		DesiredView.PostProcessSettings.bOverride_DepthOfFieldSensorWidth = true;
-		DesiredView.PostProcessSettings.DepthOfFieldSensorWidth = Filmback.SensorWidth * (1.0f + Overscan);
+		DesiredView.PostProcessSettings.DepthOfFieldSensorWidth = Filmback.SensorWidth * OverscanScalar;
 
+		const float AsymmetricSqueezeFactor = bCropOverscan ? (AsymmetricOverscan.X + AsymmetricOverscan.Y + 2.0f) / (AsymmetricOverscan.Z + AsymmetricOverscan.W + 2.0f) : 1.0f;
 		DesiredView.PostProcessSettings.bOverride_DepthOfFieldSqueezeFactor = true;
-		DesiredView.PostProcessSettings.DepthOfFieldSqueezeFactor = LensSettings.SqueezeFactor;
+		DesiredView.PostProcessSettings.DepthOfFieldSqueezeFactor = LensSettings.SqueezeFactor * AsymmetricSqueezeFactor;
+		DesiredView.PostProcessSettings.bOverride_DepthOfFieldAspectRatioScalar = true;
+		DesiredView.PostProcessSettings.DepthOfFieldAspectRatioScalar = AsymmetricSqueezeFactor;
 	}
 }
 
diff --git a/Engine/Source/Runtime/Engine/Classes/Camera/CameraComponent.h b/Engine/Source/Runtime/Engine/Classes/Camera/CameraComponent.h
--- a/Engine/Source/Runtime/Engine/Classes/Camera/CameraComponent.h
+++ b/Engine/Source/Runtime/Engine/Classes/Camera/CameraComponent.h
@@ -136,6 +136,16 @@ class UCameraComponent : public USceneComponent
 	float Overscan;
 	UFUNCTION(BlueprintCallable, Category = Camera)
 	void SetOverscan(float InOverscan) { Overscan = InOverscan; }
+
+	/**
+	 * Experimental: Amount to increase each side of the view frustum by, from 0.0 for no increase to 1.0 for 100% increase.
+	 * By convention, X is the left overscan, Y is the right overscan, Z is the top overscan, and W is the bottom overscan. Stacks with uniform Overscan.
+	 * Not currently exposed to the editor or blueprints, intended for internal use for now
+	 */
+	UPROPERTY(meta = (UIMin="0.0", ClampMin="0.0", UIMax="1.0", ClampMax="1.0"))
+	FVector4f AsymmetricOverscan = FVector4f::Zero();
+	UFUNCTION()
+	void SetAsymmetricOverscan(const FVector4f& InAsymmetricOverscan) { AsymmetricOverscan = InAsymmetricOverscan; }
 	
 	/**
 	 * Indicates that the resolution should be scaled by the overscan amount so that the original view frustum remains the same resolution.
diff --git a/Engine/Source/Runtime/Engine/Classes/Camera/CameraTypes.h b/Engine/Source/Runtime/Engine/Classes/Camera/CameraTypes.h
--- a/Engine/Source/Runtime/Engine/Classes/Camera/CameraTypes.h
+++ b/Engine/Source/Runtime/Engine/Classes/Camera/CameraTypes.h
@@ -138,10 +138,23 @@ struct FMinimalViewInfo
 	/** The fraction between 0.0 and 1.0 of the view to crop to during the final post process upscale, with 1.0 meaning no crop */
 	UPROPERTY(EditAnywhere, BlueprintReadWrite, Category=Camera, meta=(ClampMin=0.0, UIMin=0.0, ClampMax=1.0, UIMax=1.0))
 	float CropFraction;
+
+	/**
+	 * Experimental: The fraction for each edge between 0.0 and 1.0 of the view to crop to during the final post process upscale, with 1.0 meaning no crop.
+	 * By convention, X is the left edge, Y is the right edge, Z is the top edge, and W is the bottom edge. Stacks with uniform CropFraction.
+	 */
+	UPROPERTY(EditAnywhere, BlueprintReadWrite, Category=Camera, meta=(ClampMin=0.0, UIMin=0.0, ClampMax=1.0, UIMax=1.0))
+	FVector4f AsymmetricCropFraction;
 	
 private:
 	/** The amount of overscan that has been applied to the view's frustum, with 0.0 meaning no overscan and 1.0 meaning 100% overscan */
 	float Overscan;
+
+	/**
+	 * Experimental: The amount of asymmetric overscan that has been applied to the view's frustum, with 0.0 meaning no overscan and 1.0 meaning 100% overscan.
+	 * By convention, X is the left overscan, Y is the right overscan, Z is the top overscan, and W is the bottom overscan. Stacks with uniform Overscan.
+	 */
+	FVector4f AsymmetricOverscan;
 	
 	// Only used for Ortho camera auto plane calculations, tells the Near plane of the extra distance that needs to be added.
 	FVector CameraToViewTarget;
@@ -172,7 +185,9 @@ public:
 		, OffCenterProjectionOffset(ForceInitToZero)
 		, OverscanResolutionFraction(1.0f)
 		, CropFraction(1.0f)
+		, AsymmetricCropFraction(FVector4f::One())
 		, Overscan(0.0f)
+		, AsymmetricOverscan(FVector4f::Zero())
 		, CameraToViewTarget(FVector::ZeroVector)
 	{
 	}
@@ -238,11 +253,22 @@ public:
 	 */
 	ENGINE_API void ApplyOverscan(float InOverscan, bool bScaleResolutionWithOverscan = false, bool bCropOverscan = false);
 
+	/**
+	 * Experimental: Apply asymmetric overscan to the view info, which scales the field of view, ortho width, aspect ratio, and off-center projection to
+	 * simulate expanding the view frustum asymmetrically.
+	 * 
+	 * @param InAsymmetricOverscan - The amount of asymmetric overscan to apply, from 0.0 meaning no overscan to 1.0 meaning 100% overscan. By convention,
+	 *								 X is the left overscan, Y is the right overscan, Z is the top overscan, and W is the bottom overscan.
+	 * @param bScaleResolutionWithOverscan - Indicates that the view's resolution should be scaled with the amount of overscan, so that the original frustum remains the same resolution
+	 * @param bCropOverscan - Indicates that the view should be cropped during the final post process pass to remove the overscanned pixels
+	 */
+	ENGINE_API void ApplyAsymmetricOverscan(const FVector4f& InAsymmetricOverscan, bool bScaleResolutionWithOverscan = false, bool bCropOverscan = false);
+	
 	/** Gets the total amount of overscan that has been applied to the view's frustum, with 0.0 meaning no overscan and 1.0 meaning 100% overscan */
 	ENGINE_API float GetOverscan() const { return Overscan; }
 	
 	/**
-	 * Removes all overscan from the view info.
+	 * Removes all overscan (uniform and asymmetric) from the view info.
 	 */
 	ENGINE_API void ClearOverscan();
 };
diff --git a/Engine/Source/Runtime/Engine/Classes/Engine/Scene.h b/Engine/Source/Runtime/Engine/Classes/Engine/Scene.h
--- a/Engine/Source/Runtime/Engine/Classes/Engine/Scene.h
+++ b/Engine/Source/Runtime/Engine/Classes/Engine/Scene.h
@@ -1117,6 +1117,9 @@ struct FPostProcessSettings
 	UPROPERTY(EditAnywhere, BlueprintReadWrite, Category=Overrides, meta=(PinHiddenByDefault, InlineEditConditionToggle))
 	uint8 bOverride_DepthOfFieldPetzvalExclusionBoxRadius:1;
 
+	UPROPERTY(EditAnywhere, BlueprintReadWrite, Category=Overrides, meta=(PinHiddenByDefault, InlineEditConditionToggle))
+	uint8 bOverride_DepthOfFieldAspectRatioScalar:1;
+	
 	UPROPERTY(EditAnywhere, BlueprintReadWrite, Category=Overrides, meta=(PinHiddenByDefault, InlineEditConditionToggle))
 	uint8 bOverride_DepthOfFieldDepthBlurAmount:1;
 
@@ -2150,6 +2153,10 @@ struct FPostProcessSettings
 	UPROPERTY(interp, EditAnywhere, BlueprintReadWrite, Category="Lens|Depth of Field", meta=(ClampMin = "0.0", ClampMax = "1.0", editcondition = "bOverride_DepthOfFieldPetzvalExclusionBoxRadius", DisplayName = "Petzval Bokeh Exclusion Box Radius"))
 	float DepthOfFieldPetzvalExclusionBoxRadius;
 
+	/** Amount to scale the output viewport's aspect ratio by when computing depth of field properties.  */
+	UPROPERTY(Interp, BlueprintReadWrite, Category="Lens|Depth of Field", meta=(ClampMin = "0.0", editcondition = "bOverride_DepthOfFieldAspectRatioScalar", DisplayName = "Aspect Ratio Scalar"))
+	float DepthOfFieldAspectRatioScalar;
+	
 	/** Artificial region where all content is in focus, starting after DepthOfFieldFocalDistance, in unreal units  (cm) */
 	UPROPERTY(interp, EditAnywhere, BlueprintReadWrite, Category="Lens|Mobile Depth of Field", meta=(UIMin = "0.0", UIMax = "10000.0", editcondition = "bOverride_DepthOfFieldFocalRegion", DisplayName = "Focal Region"))
 	float DepthOfFieldFocalRegion;
diff --git a/Engine/Source/Runtime/Engine/Private/Camera/CameraComponent.cpp b/Engine/Source/Runtime/Engine/Private/Camera/CameraComponent.cpp
--- a/Engine/Source/Runtime/Engine/Private/Camera/CameraComponent.cpp
+++ b/Engine/Source/Runtime/Engine/Private/Camera/CameraComponent.cpp
@@ -466,6 +466,7 @@ void UCameraComponent::GetCameraView(float DeltaTime, FMinimalViewInfo& DesiredV
 	DesiredView.bUseFirstPersonParameters = bEnableFirstPersonFieldOfView || bEnableFirstPersonScale;
 	
 	DesiredView.ApplyOverscan(Overscan, bScaleResolutionWithOverscan, bCropOverscan);
+	DesiredView.ApplyAsymmetricOverscan(AsymmetricOverscan, bScaleResolutionWithOverscan, bCropOverscan);
 	
 	if (bAutoCalculateOrthoPlanes)
 	{
diff --git a/Engine/Source/Runtime/Engine/Private/Camera/CameraStackTypes.cpp b/Engine/Source/Runtime/Engine/Private/Camera/CameraStackTypes.cpp
--- a/Engine/Source/Runtime/Engine/Private/Camera/CameraStackTypes.cpp
+++ b/Engine/Source/Runtime/Engine/Private/Camera/CameraStackTypes.cpp
@@ -375,7 +375,11 @@ void FMinimalViewInfo::CalculateProjectionMatrixGivenViewRectangle(FMinimalViewI
 
 void FMinimalViewInfo::CalculateProjectionMatrixGivenView(FMinimalViewInfo& ViewInfo, TEnumAsByte<enum EAspectRatioAxisConstraint> AspectRatioAxisConstraint, FViewport* Viewport, FSceneViewProjectionData& InOutProjectionData)
 {
-	FIntRect ViewExtents = Viewport->CalculateViewExtents(ViewInfo.AspectRatio, InOutProjectionData.GetViewRect());
+	// Factor in any asymmetric crop, which can change the output aspect ratio
+	const float CropAspectRatio = (ViewInfo.AsymmetricCropFraction.X + ViewInfo.AsymmetricCropFraction.Y) / (ViewInfo.AsymmetricCropFraction.Z + ViewInfo.AsymmetricCropFraction.W);
+	const float AspectRatio = ViewInfo.AspectRatio * CropAspectRatio;
+
+	FIntRect ViewExtents = Viewport->CalculateViewExtents(AspectRatio, InOutProjectionData.GetViewRect());
 	CalculateProjectionMatrixGivenViewRectangle(ViewInfo, AspectRatioAxisConstraint, ViewExtents, InOutProjectionData);
 }
 
@@ -515,11 +519,68 @@ void FMinimalViewInfo::ApplyOverscan(float InOverscan, bool bScaleResolutionWith
 		
 		OrthoWidth *= OverscanScalar;
 
-		OverscanResolutionFraction *= bScaleResolutionWithOverscan ? OverscanScalar : 1.0;
+		if (bScaleResolutionWithOverscan)
+		{
+			// Ensure that the resolution fraction stays between 1.0 and 2.0
+			OverscanResolutionFraction = FMath::Clamp(OverscanResolutionFraction * OverscanScalar, 1.0f, 2.0f);
+		}
+
 		CropFraction *= bCropOverscan ? 1.0f / OverscanScalar : 1.0f;
 	}
 }
 
+void FMinimalViewInfo::ApplyAsymmetricOverscan(const FVector4f& InAsymmetricOverscan, bool bScaleResolutionWithOverscan, bool bCropOverscan)
+{
+	if (!FMath::IsNearlyZero(InAsymmetricOverscan.Size()))
+	{
+		// Clamp the incoming overscan so that the new total overscan can never be less than zero
+		const FVector4f InverseAsymmetricOverscan = FVector4f(
+			(AsymmetricOverscan.Y - AsymmetricOverscan.X + 2.0f) / (AsymmetricOverscan.X + AsymmetricOverscan.Y + 2.0f) - 1.0f,
+			(AsymmetricOverscan.X - AsymmetricOverscan.Y + 2.0f) / (AsymmetricOverscan.X + AsymmetricOverscan.Y + 2.0f) - 1.0f,
+			(AsymmetricOverscan.W - AsymmetricOverscan.Z + 2.0f) / (AsymmetricOverscan.Z + AsymmetricOverscan.W + 2.0f) - 1.0f,
+			(AsymmetricOverscan.Z - AsymmetricOverscan.W + 2.0f) / (AsymmetricOverscan.Z + AsymmetricOverscan.W + 2.0f) - 1.0f);
+		
+		const FVector4f ClampedAsymmetricOverscan = FVector4f(
+			FMath::Max(InAsymmetricOverscan.X, InverseAsymmetricOverscan.X),
+			FMath::Max(InAsymmetricOverscan.Y, InverseAsymmetricOverscan.Y),
+			FMath::Max(InAsymmetricOverscan.Z, InverseAsymmetricOverscan.Z),
+			FMath::Max(InAsymmetricOverscan.W, InverseAsymmetricOverscan.W));
+	
+		// Keep track of the total amount of asymmetric overscan that has been applied to the view. Mathematically,
+		// this formula is derived from 1 + TotalOverscan = (1 + Overscan) * (1 + InOverscan)
+		AsymmetricOverscan = AsymmetricOverscan * (ClampedAsymmetricOverscan + 1.0f) + ClampedAsymmetricOverscan;
+
+		// By convention, 0.0 means no overscan, so add 1 to compute the scalar needed for altering projection values
+		FVector4f AsymmetricOverscanScalar = ClampedAsymmetricOverscan + 1.0f;
+		
+		// Overscan directly scales the view frustum, but can be accomplished by scaling the FOV.
+		// However, must scale the tangent of the half-FOV to accomplish the same mathematical transform.
+		const float HalfFOVInRadians = FMath::DegreesToRadians(0.5f * FOV);
+		const float OverscannedFOV = FMath::Atan(0.5f * (AsymmetricOverscanScalar.X + AsymmetricOverscanScalar.Y) * FMath::Tan(HalfFOVInRadians));
+		FOV = 2.0f * FMath::RadiansToDegrees(OverscannedFOV);
+
+		OrthoWidth *= 0.5f * (AsymmetricOverscanScalar.X + AsymmetricOverscanScalar.Y);
+	
+		AspectRatio *= (AsymmetricOverscanScalar.X + AsymmetricOverscanScalar.Y) / (AsymmetricOverscanScalar.Z + AsymmetricOverscanScalar.W);
+		OffCenterProjectionOffset.X += (AsymmetricOverscanScalar.Y - AsymmetricOverscanScalar.X) / (AsymmetricOverscanScalar.X + AsymmetricOverscanScalar.Y);
+		OffCenterProjectionOffset.Y += (AsymmetricOverscanScalar.Z - AsymmetricOverscanScalar.W) / (AsymmetricOverscanScalar.Z + AsymmetricOverscanScalar.W);
+
+		FVector4f InvAsymmetricOverscanScalar = FVector4f(
+			(AsymmetricOverscanScalar.Y - AsymmetricOverscanScalar.X + 2.0f) / (AsymmetricOverscanScalar.X + AsymmetricOverscanScalar.Y),
+			(AsymmetricOverscanScalar.X - AsymmetricOverscanScalar.Y + 2.0f) / (AsymmetricOverscanScalar.X + AsymmetricOverscanScalar.Y),
+			(AsymmetricOverscanScalar.W - AsymmetricOverscanScalar.Z + 2.0f) / (AsymmetricOverscanScalar.Z + AsymmetricOverscanScalar.W),
+			(AsymmetricOverscanScalar.Z - AsymmetricOverscanScalar.W + 2.0f) / (AsymmetricOverscanScalar.Z + AsymmetricOverscanScalar.W));
+
+		if (bScaleResolutionWithOverscan)
+		{
+			// Ensure that the resolution fraction stays between 1.0 and 2.0
+			const float MaxResScale = FMath::Max(0.5f * (AsymmetricOverscanScalar.X + AsymmetricOverscanScalar.Y), 0.5f * (AsymmetricOverscanScalar.Z + AsymmetricOverscanScalar.W));
+			OverscanResolutionFraction = FMath::Clamp(OverscanResolutionFraction * MaxResScale, 1.0f, 2.0f);
+		}
+		AsymmetricCropFraction *= bCropOverscan ? InvAsymmetricOverscanScalar : FVector4f::One();
+	}
+}
+
 void FMinimalViewInfo::ClearOverscan()
 {
 	if (Overscan > 0.0f)
@@ -528,8 +589,21 @@ void FMinimalViewInfo::ClearOverscan()
 		// Inverse overscan derived from (1 + Overscan) * (1 + InverseOverscan) = 1
 		const float InverseOverscan = - Overscan / (1.0f + Overscan);
 		ApplyOverscan(InverseOverscan);
+	}
+	
+	if (AsymmetricOverscan.Size() > 0)
+	{
+		// Apply the inverse overscan to the view frustum to obtain the original frustum values (field of view, ortho width, etc)
+		const FVector4f InverseAsymmetricOverscan = FVector4f(
+			(AsymmetricOverscan.Y - AsymmetricOverscan.X + 2.0f) / (AsymmetricOverscan.X + AsymmetricOverscan.Y + 2.0f) - 1.0f,
+			(AsymmetricOverscan.X - AsymmetricOverscan.Y + 2.0f) / (AsymmetricOverscan.X + AsymmetricOverscan.Y + 2.0f) - 1.0f,
+			(AsymmetricOverscan.W - AsymmetricOverscan.Z + 2.0f) / (AsymmetricOverscan.Z + AsymmetricOverscan.W + 2.0f) - 1.0f,
+			(AsymmetricOverscan.Z - AsymmetricOverscan.W + 2.0f) / (AsymmetricOverscan.Z + AsymmetricOverscan.W + 2.0f) - 1.0f);
 
-		OverscanResolutionFraction = 1.0f;
-		CropFraction = 1.0f;
+		ApplyAsymmetricOverscan(InverseAsymmetricOverscan);
 	}
+	
+	OverscanResolutionFraction = 1.0f;
+	CropFraction = 1.0f;
+	AsymmetricCropFraction = FVector4f::One();
 }
diff --git a/Engine/Source/Runtime/Engine/Private/LocalPlayer.cpp b/Engine/Source/Runtime/Engine/Private/LocalPlayer.cpp
--- a/Engine/Source/Runtime/Engine/Private/LocalPlayer.cpp
+++ b/Engine/Source/Runtime/Engine/Private/LocalPlayer.cpp
@@ -844,6 +844,7 @@ FSceneView* ULocalPlayer::CalcSceneView( class FSceneViewFamily* ViewFamily,
 	ViewInitOptions.FirstPersonParams = FFirstPersonParameters(ViewInfo.CalculateFirstPersonFOVCorrectionFactor(), ViewInfo.FirstPersonScale, ViewInfo.bUseFirstPersonParameters);
 	ViewInitOptions.OverscanResolutionFraction = ViewInfo.OverscanResolutionFraction;
 	ViewInitOptions.CropFraction = ViewInfo.CropFraction;
+	ViewInitOptions.AsymmetricCropFraction = ViewInfo.AsymmetricCropFraction;
 	
 	// Fill out the rest of the view init options
 	ViewInitOptions.ViewFamily = ViewFamily;
diff --git a/Engine/Source/Runtime/Engine/Private/PostProcessUtils.cpp b/Engine/Source/Runtime/Engine/Private/PostProcessUtils.cpp
--- a/Engine/Source/Runtime/Engine/Private/PostProcessUtils.cpp
+++ b/Engine/Source/Runtime/Engine/Private/PostProcessUtils.cpp
@@ -158,6 +158,7 @@ bool FPostProcessUtils::OverridePostProcessSettings(FPostProcessSettings& ThisFr
 		UE_SET_PP(DepthOfFieldOcclusion);
 		UE_SET_PP(DepthOfFieldSkyFocusDistance);
 		UE_SET_PP(DepthOfFieldVignetteSize);
+		UE_SET_PP(DepthOfFieldAspectRatioScalar);
 		UE_SET_PP(MotionBlurAmount);
 		UE_SET_PP(MotionBlurMax);
 		UE_SET_PP(MotionBlurPerObjectSize);
@@ -448,6 +449,7 @@ bool FPostProcessUtils::BlendPostProcessSettings(FPostProcessSettings& ThisFrom,
 		UE_LERP_PP(DepthOfFieldOcclusion);
 		UE_LERP_PP(DepthOfFieldSkyFocusDistance);
 		UE_LERP_PP(DepthOfFieldVignetteSize);
+		UE_LERP_PP(DepthOfFieldAspectRatioScalar);
 		UE_LERP_PP(MotionBlurAmount);
 		UE_LERP_PP(MotionBlurMax);
 		UE_LERP_PP(MotionBlurPerObjectSize);
diff --git a/Engine/Source/Runtime/Engine/Private/Scene.cpp b/Engine/Source/Runtime/Engine/Private/Scene.cpp
--- a/Engine/Source/Runtime/Engine/Private/Scene.cpp
+++ b/Engine/Source/Runtime/Engine/Private/Scene.cpp
@@ -569,6 +569,7 @@ FPostProcessSettings::FPostProcessSettings()
 	DepthOfFieldPetzvalBokehFalloff = 1.0f;
 	DepthOfFieldPetzvalExclusionBoxExtents = FVector2f::ZeroVector;
 	DepthOfFieldPetzvalExclusionBoxRadius = 0.0f;
+	DepthOfFieldAspectRatioScalar = 1.0f;
 	DepthOfFieldFocalRegion = 0.0f;
 	DepthOfFieldNearTransitionRegion = 300.0f;
 	DepthOfFieldFarTransitionRegion = 500.0f;
diff --git a/Engine/Source/Runtime/Engine/Private/SceneView.cpp b/Engine/Source/Runtime/Engine/Private/SceneView.cpp
--- a/Engine/Source/Runtime/Engine/Private/SceneView.cpp
+++ b/Engine/Source/Runtime/Engine/Private/SceneView.cpp
@@ -1730,6 +1730,7 @@ void FSceneView::OverridePostProcessSettings(const FPostProcessSettings& Src, fl
 		LERP_PP(DepthOfFieldOcclusion);
 		LERP_PP(DepthOfFieldSkyFocusDistance);
 		LERP_PP(DepthOfFieldVignetteSize);
+		LERP_PP(DepthOfFieldAspectRatioScalar);
 		LERP_PP(MotionBlurAmount);
 		LERP_PP(MotionBlurMax);
 		LERP_PP(MotionBlurPerObjectSize);
diff --git a/Engine/Source/Runtime/Engine/Public/SceneView.h b/Engine/Source/Runtime/Engine/Public/SceneView.h
--- a/Engine/Source/Runtime/Engine/Public/SceneView.h
+++ b/Engine/Source/Runtime/Engine/Public/SceneView.h
@@ -248,6 +248,12 @@ struct FSceneViewInitOptions : public FSceneViewProjectionData
 
 	/** Fraction of the view to crop to during the secondary upscale pass, with 1.0 meaning no crop */
 	float CropFraction = 1.0f;
+
+	/**
+	 * The fraction for each edge between 0.0 and 1.0 of the view to crop to during the final post process upscale, with 1.0 meaning no crop.
+	 * By convention, X is the left edge, Y is the right edge, Z is the top edge, and W is the bottom edge. Stacks with uniform CropFraction.
+	 */
+	FVector4f AsymmetricCropFraction = FVector4f::One();
 	
 #if WITH_EDITOR
 	/** default to 0'th view index, which is a bitfield of 1 */
diff --git a/Engine/Source/Runtime/Renderer/Private/PostProcess/DiaphragmDOFUtils.cpp b/Engine/Source/Runtime/Renderer/Private/PostProcess/DiaphragmDOFUtils.cpp
--- a/Engine/Source/Runtime/Renderer/Private/PostProcess/DiaphragmDOFUtils.cpp
+++ b/Engine/Source/Runtime/Renderer/Private/PostProcess/DiaphragmDOFUtils.cpp
@@ -113,7 +113,7 @@ void DiaphragmDOF::FPhysicalCocModel::Compile(const FViewInfo& View)
 		PetzvalExclusionBoxExtents = View.FinalPostProcessSettings.DepthOfFieldPetzvalExclusionBoxExtents;
 		PetzvalExclusionBoxRadius = View.FinalPostProcessSettings.DepthOfFieldPetzvalExclusionBoxRadius;
 
-		RenderingAspectRatio = float(View.UnscaledViewRect.Width()) / float(View.UnscaledViewRect.Height());
+		RenderingAspectRatio = float(View.UnscaledViewRect.Width()) / float(View.UnscaledViewRect.Height()) * View.FinalPostProcessSettings.DepthOfFieldAspectRatioScalar;
 		const float HorizontalHalfFOV = FMath::Atan(1.0f / View.ViewMatrices.GetProjectionMatrix().M[0][0]);
 		const float VerticalHalfFOV = FMath::Atan(1.0f / View.ViewMatrices.GetProjectionMatrix().M[1][1]);
 
diff --git a/Engine/Source/Runtime/Renderer/Private/SceneRendering.cpp b/Engine/Source/Runtime/Renderer/Private/SceneRendering.cpp
--- a/Engine/Source/Runtime/Renderer/Private/SceneRendering.cpp
+++ b/Engine/Source/Runtime/Renderer/Private/SceneRendering.cpp
@@ -1290,16 +1290,20 @@ FIntRect FViewInfo::GetSecondaryViewCropRect() const
 	const FIntPoint SecondaryVewRectSize = GetSecondaryViewRectSize();
 
 	// Clamp the crop fraction to sensible values to ensure crop rect is always a valid rectangle
-	const float CropFrac = FMath::Clamp(SceneViewInitOptions.CropFraction, 0.0f, 1.0f);
+	const FVector4f CropFrac = FVector4f(
+		FMath::Clamp(SceneViewInitOptions.AsymmetricCropFraction.X * SceneViewInitOptions.CropFraction, 0.0f, 1.0f),
+		FMath::Clamp(SceneViewInitOptions.AsymmetricCropFraction.Y * SceneViewInitOptions.CropFraction, 0.0f, 1.0f),
+		FMath::Clamp(SceneViewInitOptions.AsymmetricCropFraction.Z * SceneViewInitOptions.CropFraction, 0.0f, 1.0f),
+		FMath::Clamp(SceneViewInitOptions.AsymmetricCropFraction.W * SceneViewInitOptions.CropFraction, 0.0f, 1.0f));
 	
 	FIntRect CropRect;
 	CropRect.Min = FIntPoint(
-			FMath::FloorToInt(0.5f * (1.0f - CropFrac) * SecondaryVewRectSize.X),
-			FMath::FloorToInt(0.5f * (1.0f - CropFrac) * SecondaryVewRectSize.Y));
+			FMath::FloorToInt(0.5f * (1.0f - CropFrac.X) * SecondaryVewRectSize.X),
+			FMath::FloorToInt(0.5f * (1.0f - CropFrac.Z) * SecondaryVewRectSize.Y));
 		
 	CropRect.Max = FIntPoint(
-		FMath::CeilToInt(0.5f * (1.0f + CropFrac) * SecondaryVewRectSize.X),
-		FMath::CeilToInt(0.5f * (1.0f + CropFrac) * SecondaryVewRectSize.Y));
+		FMath::CeilToInt(0.5f * (1.0f + CropFrac.Y) * SecondaryVewRectSize.X),
+		FMath::CeilToInt(0.5f * (1.0f + CropFrac.W) * SecondaryVewRectSize.Y));
 
 	return CropRect;
 }
