[MIR] Convert MCT_Float1/2/3/4 explicitly in Material IR, since those types are bitmasks.

#rnx

diff --git a/Engine/Source/Runtime/Engine/Private/Materials/MaterialIRTypes.cpp b/Engine/Source/Runtime/Engine/Private/Materials/MaterialIRTypes.cpp
--- a/Engine/Source/Runtime/Engine/Private/Materials/MaterialIRTypes.cpp
+++ b/Engine/Source/Runtime/Engine/Private/Materials/MaterialIRTypes.cpp
@@ -58,22 +58,28 @@ FTypePtr FType::FromMaterialValueType(EMaterialValueType Type)
 			return FType::GetVoid();
 
 		case MCT_Float1:
+			return FPrimitiveType::GetVector(ScalarKind_Float, 1);
 		case MCT_Float2:
+			return FPrimitiveType::GetVector(ScalarKind_Float, 2);
 		case MCT_Float3:
+			return FPrimitiveType::GetVector(ScalarKind_Float, 3);
 		case MCT_Float4:
-			return FPrimitiveType::GetVector(ScalarKind_Float, (int)Type - (int)MCT_Float1 + 1);
+			return FPrimitiveType::GetVector(ScalarKind_Float, 4);
 
 		case MCT_Float:
 			return FPrimitiveType::GetVector(ScalarKind_Float, 4);
 
 		case MCT_UInt1:
+			return FPrimitiveType::GetVector(ScalarKind_Int, 1);
 		case MCT_UInt2:
+			return FPrimitiveType::GetVector(ScalarKind_Int, 2);
 		case MCT_UInt3:
+			return FPrimitiveType::GetVector(ScalarKind_Int, 3);
 		case MCT_UInt4:
-			return FPrimitiveType::GetVector(ScalarKind_Int, (int)Type - (int)MCT_UInt1 + 1);
+			return FPrimitiveType::GetVector(ScalarKind_Int, 4);
 
 		case MCT_Bool:
-			return FPrimitiveType::GetVector(ScalarKind_Bool, (int)Type - (int)UE::Shader::EValueType::Int1 + 1);
+			return FPrimitiveType::GetVector(ScalarKind_Bool, 1);
 
 		default:
 			UE_MIR_UNREACHABLE();
