[Materials] Convert further material experessions to data-driven approach and add domain specific restriction.

- Convert DecalColor, DecalLifetimeOpacity, and DeltaTime to data-driven material expressions.
- Add list of material domains to FMaterialExternalCodeDeclaration, that restrict in what domains an external code declaration can be used.

#rnx

diff --git a/Engine/Config/BaseMaterialExpressions.ini b/Engine/Config/BaseMaterialExpressions.ini
--- a/Engine/Config/BaseMaterialExpressions.ini
+++ b/Engine/Config/BaseMaterialExpressions.ini
@@ -220,6 +220,7 @@ bShaderInputData=true
 [/Script/Engine.MaterialExpressionDecalColor]
 +MenuCategories=NSLOCTEXT("MaterialExpression", "Utils", "Utils")
 bShaderInputData=true
++ExternalCodeIdentifiers="DecalColor"
 
 [/Script/Engine.MaterialExpressionDecalDerivative]
 +MenuCategories=NSLOCTEXT("MaterialExpression", "Utils", "Utils")
@@ -228,6 +229,7 @@ bShaderInputData=true
 [/Script/Engine.MaterialExpressionDecalLifetimeOpacity]
 +MenuCategories=NSLOCTEXT("MaterialExpression", "Utils", "Utils")
 bShaderInputData=true
++ExternalCodeIdentifiers="DecalLifetimeOpacity"
 
 [/Script/Engine.MaterialExpressionDecalMipmapLevel]
 +MenuCategories=NSLOCTEXT("MaterialExpression", "Utils", "Utils")
@@ -235,6 +237,7 @@ bShaderInputData=true
 [/Script/Engine.MaterialExpressionDeltaTime]
 +MenuCategories=NSLOCTEXT("MaterialExpression", "Constants", "Constants")
 bShaderInputData=true
++ExternalCodeIdentifiers="DeltaTime"
 
 [/Script/Engine.MaterialExpressionDepthFade]
 +MenuCategories=NSLOCTEXT("MaterialExpression", "Depth", "Depth")
@@ -962,6 +965,9 @@ bShaderInputData=true
 
 [/Script/Engine.MaterialExternalCodeCollection]
 +ExternalCodeDeclarations=(ReturnType=Float3, Name="CameraVector", Definition="Parameters.CameraVector", ShaderFrequency=Pixel|Compute, bIsInlined=true)
++ExternalCodeDeclarations=(ReturnType=Float4, Name="DecalColor", Definition="DecalColor()", Derivative=Zero, ShaderFrequency=Pixel, Domains=(MD_DeferredDecal))
++ExternalCodeDeclarations=(ReturnType=Float, Name="DecalLifetimeOpacity", Definition="DecalLifetimeOpacity()", Derivative=Zero, ShaderFrequency=Pixel, Domains=(MD_DeferredDecal))
++ExternalCodeDeclarations=(ReturnType=Float, Name="DeltaTime", Definition="View.DeltaTime", Derivative=Zero, bIsInlined=true)
 +ExternalCodeDeclarations=(ReturnType=Float, Name="EyeAdaptation", Definition="EyeAdaptationLookup()", Derivative=Zero, ShaderFrequency=Pixel, bIsInlined=true)
 +ExternalCodeDeclarations=(ReturnType=Float4, Name="FontSignedDistanceData", Definition="GetFontSignedDistanceData(Parameters)", ShaderFrequency=Pixel)
 +ExternalCodeDeclarations=(ReturnType=Float, Name="IsOrthographic", Definition="IsOrthoProjectionFloat()", Derivative=Zero, bIsInlined=true)
@@ -977,7 +983,7 @@ bShaderInputData=true
 +ExternalCodeDeclarations=(ReturnType=Float, Name="ParticleRandom", Definition="Parameters.Particle.Random", Derivative=Zero, bIsInlined=true)
 +ExternalCodeDeclarations=(ReturnType=Float, Name="ParticleSpeed", Definition="Parameters.Particle.Velocity.w", Derivative=Zero, bIsInlined=true)
 +ExternalCodeDeclarations=(ReturnType=Float2, Name="ParticleSize", Definition="Parameters.Particle.Size", Derivative=Zero, bIsInlined=true)
-+ExternalCodeDeclarations=(RetunrType=Float2, Name="ParticleSpriteRotation", Definition="float2(Parameters.Particle.SpriteRotation, Parameters.Particle.SpriteRotation * 57.2957795131f)", Derivative=Zero, bIsInlined=true)
++ExternalCodeDeclarations=(ReturnType=Float2, Name="ParticleSpriteRotation", Definition="float2(Parameters.Particle.SpriteRotation, Parameters.Particle.SpriteRotation * 57.2957795131f)", Derivative=Zero, bIsInlined=true)
 +ExternalCodeDeclarations=(ReturnType=Float3, Name="PreSkinnedPosition", Definition="Parameters.PreSkinnedPosition", ShaderFrequency=Vertex, bIsInlined=true)
 +ExternalCodeDeclarations=(ReturnType=Float3, Name="PreSkinnedNormal", Definition="Parameters.PreSkinnedNormal", ShaderFrequency=Vertex, bIsInlined=true)
 +ExternalCodeDeclarations=(ReturnType=Float3, Name="ReflectionVector", Definition="Parameters.ReflectionVector", ShaderFrequency=Pixel|Compute, bIsInlined=true)
diff --git a/Engine/Source/Runtime/Engine/Private/Materials/HLSLMaterialTranslator.cpp b/Engine/Source/Runtime/Engine/Private/Materials/HLSLMaterialTranslator.cpp
--- a/Engine/Source/Runtime/Engine/Private/Materials/HLSLMaterialTranslator.cpp
+++ b/Engine/Source/Runtime/Engine/Private/Materials/HLSLMaterialTranslator.cpp
@@ -3848,11 +3848,29 @@ int32 FHLSLMaterialTranslator::AddExternalCodeChunk(const FName& InExternalCodeI
 
 int32 FHLSLMaterialTranslator::AddExternalCodeChunk(const FMaterialExternalCodeDeclaration& InExternalCode)
 {
+	// Validate this external code can be used in the current shader stage
 	if (((int32)InExternalCode.ShaderFrequency & (1 << (int32)ShaderFrequency)) == 0)
 	{
 		FName ShaderFrequencyName = StaticEnum<EMaterialShaderFrequency>()->GetNameByValue(1ull << (int32)ShaderFrequency);
 		return Errorf(TEXT("External code identifier '%s' is not available in the %s shader."), *InExternalCode.Name.ToString(), *ShaderFrequencyName.ToString());
 	}
+
+	// Validate this external code can be used for the current material domain. Empty list implies no restriction on material domains.
+	if (!InExternalCode.Domains.IsEmpty() && InExternalCode.Domains.Find(Material->GetMaterialDomain()) == INDEX_NONE)
+	{
+		TStringBuilder<1024> SupportedMaterialDomainsString;
+		for (EMaterialDomain Domain : InExternalCode.Domains)
+		{
+			if (SupportedMaterialDomainsString.Len() > 0)
+			{
+				SupportedMaterialDomainsString.Append(TEXT(", "));
+			}
+			SupportedMaterialDomainsString.Append(*MaterialDomainString(Domain));
+		}
+		return Errorf(TEXT("External code '%s' (asset: %s) is only available in the following material domains: %s"), *InExternalCode.Name.ToString(), *Material->GetAssetPath().ToString(), SupportedMaterialDomainsString.ToString());
+	}
+
+	// Emit external code definition as new code chunk
 	return AddCodeChunkInner(InExternalCode.GetReturnTypeValue(), InExternalCode.Derivative, InExternalCode.bIsInlined, TEXT("%s"), *InExternalCode.Definition);
 }
 
@@ -5455,7 +5473,7 @@ int32 FHLSLMaterialTranslator::DeltaTime()
 	// explicitly avoid trying to return previous frame's delta time for bCompilingPreviousFrame here
 	// DeltaTime expression is designed to be used when generating custom motion vectors, by using world position offset along with previous frame switch
 	// in this context, we will technically be evaluating the previous frame, but we want to use the current frame's delta tick in order to offset the vector used to create previous position
-	return AddInlinedCodeChunkZeroDeriv(MCT_Float, TEXT("View.DeltaTime"));
+	return AddExternalCodeChunk(TEXT("DeltaTime"));
 }
 
 int32 FHLSLMaterialTranslator::PeriodicHint(int32 PeriodicCode)
@@ -5521,10 +5539,10 @@ int32 FHLSLMaterialTranslator::Cosine(int32 X)
 		}
 		else
 		{
-		return AddCodeChunk(GetParameterType(X),TEXT("cos(%s)"),*GetParameterCode(X));
+			return AddCodeChunk(GetParameterType(X),TEXT("cos(%s)"),*GetParameterCode(X));
+		}
 	}
 }
-}
 
 int32 FHLSLMaterialTranslator::Tangent(int32 X)
 {
@@ -5546,10 +5564,10 @@ int32 FHLSLMaterialTranslator::Tangent(int32 X)
 		}
 		else
 		{
-		return AddCodeChunk(GetParameterType(X),TEXT("tan(%s)"),*GetParameterCode(X));
+			return AddCodeChunk(GetParameterType(X),TEXT("tan(%s)"),*GetParameterCode(X));
+		}
 	}
 }
-}
 
 int32 FHLSLMaterialTranslator::Arcsine(int32 X)
 {
@@ -5571,10 +5589,10 @@ int32 FHLSLMaterialTranslator::Arcsine(int32 X)
 		}
 		else
 		{
-		return AddCodeChunk(GetParameterType(X),TEXT("asin(%s)"),*GetParameterCode(X));
+			return AddCodeChunk(GetParameterType(X),TEXT("asin(%s)"),*GetParameterCode(X));
+		}
 	}
 }
-}
 
 int32 FHLSLMaterialTranslator::ArcsineFast(int32 X)
 {
@@ -5994,11 +6012,11 @@ int32 FHLSLMaterialTranslator::Abs(int32 X)
 	else
 	{
 		if (IsAnalyticDerivEnabled())
-	{
+		{
 			return DerivativeAutogen.GenerateExpressionFunc1(*this, FMaterialDerivativeAutogen::EFunc1::Abs, X);
-	}
-	else
-	{
+		}
+		else
+		{
 			return AddCodeChunk(GetParameterType(X), TEXT("abs(%s)"), *GetParameterCode(X));
 		}
 	}
@@ -7968,38 +7986,12 @@ int32 FHLSLMaterialTranslator::TextureDecalDerivative(bool bDDY)
 
 int32 FHLSLMaterialTranslator::DecalColor()
 {
-	if (Material->GetMaterialDomain() != MD_DeferredDecal)
-	{
-		return Errorf(TEXT("Decal color is only available in the decal material domain."));
-	}
-
-	if (ShaderFrequency != SF_Pixel)
-	{
-		return Errorf(TEXT("Decal color is only available in the pixel shader."));
-	}
-
-	return AddCodeChunkZeroDeriv(
-		MCT_Float4,
-		TEXT("DecalColor()")
-	);
+	return AddExternalCodeChunk(TEXT("DecalColor"));
 }
 
 int32 FHLSLMaterialTranslator::DecalLifetimeOpacity()
 {
-	if (Material->GetMaterialDomain() != MD_DeferredDecal)
-	{
-		return Errorf(TEXT("Decal lifetime fade is only available in the decal material domain."));
-	}
-
-	if (ShaderFrequency != SF_Pixel)
-	{
-		return Errorf(TEXT("Decal lifetime fade is only available in the pixel shader."));
-	}
-
-	return AddCodeChunkZeroDeriv(
-		MCT_Float,
-		TEXT("DecalLifetimeOpacity()")
-		);
+	return AddExternalCodeChunk(TEXT("DecalLifetimeOpacity"));
 }
 
 int32 FHLSLMaterialTranslator::PixelDepth()
diff --git a/Engine/Source/Runtime/Engine/Private/Materials/MaterialExpressions.cpp b/Engine/Source/Runtime/Engine/Private/Materials/MaterialExpressions.cpp
--- a/Engine/Source/Runtime/Engine/Private/Materials/MaterialExpressions.cpp
+++ b/Engine/Source/Runtime/Engine/Private/Materials/MaterialExpressions.cpp
@@ -10245,11 +10245,6 @@ void UMaterialExpressionIsOrthographic::GetCaption(TArray<FString>& OutCaptions)
 #endif // WITH_EDITOR
 
 #if WITH_EDITOR
-int32 UMaterialExpressionDeltaTime::Compile(class FMaterialCompiler* Compiler, int32 OutputIndex)
-{
-	return Compiler->DeltaTime();
-}
-
 void UMaterialExpressionDeltaTime::GetCaption(TArray<FString>& OutCaptions) const
 {
 	OutCaptions.Add(TEXT("DeltaTime"));
@@ -20470,16 +20465,6 @@ UMaterialExpressionDecalColor::UMaterialExpressionDecalColor(const FObjectInitia
 }
 
 #if WITH_EDITOR
-int32 UMaterialExpressionDecalColor::Compile(class FMaterialCompiler* Compiler, int32 OutputIndex)
-{
-	if (Material && Material->MaterialDomain != MD_DeferredDecal)
-	{
-		return CompilerError(Compiler, TEXT("Node only works for the deferred decal material domain."));
-	}
-
-	return Compiler->DecalColor();
-}
-
 void UMaterialExpressionDecalColor::GetCaption(TArray<FString>& OutCaptions) const
 {
 	OutCaptions.Add(TEXT("Decal Color"));
@@ -20499,16 +20484,6 @@ UMaterialExpressionDecalLifetimeOpacity::UMaterialExpressionDecalLifetimeOpacity
 }
 
 #if WITH_EDITOR
-int32 UMaterialExpressionDecalLifetimeOpacity::Compile(class FMaterialCompiler* Compiler, int32 OutputIndex)
-{
-	if (Material && Material->MaterialDomain != MD_DeferredDecal)
-	{
-		return CompilerError(Compiler, TEXT("Node only works for the deferred decal material domain."));
-	}
-
-	return Compiler->DecalLifetimeOpacity();
-}
-
 void UMaterialExpressionDecalLifetimeOpacity::GetCaption(TArray<FString>& OutCaptions) const
 {
 	OutCaptions.Add(TEXT("Decal Lifetime Opacity"));
diff --git a/Engine/Source/Runtime/Engine/Private/Materials/MaterialExternalCodeRegistry.h b/Engine/Source/Runtime/Engine/Private/Materials/MaterialExternalCodeRegistry.h
--- a/Engine/Source/Runtime/Engine/Private/Materials/MaterialExternalCodeRegistry.h
+++ b/Engine/Source/Runtime/Engine/Private/Materials/MaterialExternalCodeRegistry.h
@@ -93,6 +93,10 @@ public:
 	UPROPERTY(Config)
 	EMaterialShaderFrequency ShaderFrequency = EMaterialShaderFrequency::Any;
 
+	/** List of material domains this external code can be used with. If this is empty, all material domains are accepted. */
+	UPROPERTY(Config)
+	TArray<TEnumAsByte<EMaterialDomain>> Domains;
+
 public:
 	/**
 	 * Function return type as material value type.
diff --git a/Engine/Source/Runtime/Engine/Public/Materials/MaterialExpressionDecalColor.h b/Engine/Source/Runtime/Engine/Public/Materials/MaterialExpressionDecalColor.h
--- a/Engine/Source/Runtime/Engine/Public/Materials/MaterialExpressionDecalColor.h
+++ b/Engine/Source/Runtime/Engine/Public/Materials/MaterialExpressionDecalColor.h
@@ -5,17 +5,16 @@
 
 #include "CoreMinimal.h"
 #include "UObject/ObjectMacros.h"
-#include "Materials/MaterialExpression.h"
+#include "Materials/MaterialExpressionExternalCodeBase.h"
 #include "MaterialExpressionDecalColor.generated.h"
 
 UCLASS(collapsecategories, hidecategories = Object)
-class UMaterialExpressionDecalColor : public UMaterialExpression
+class UMaterialExpressionDecalColor : public UMaterialExpressionExternalCodeBase
 {
 	GENERATED_UCLASS_BODY()
 
 	//~ Begin UMaterialExpression Interface
 #if WITH_EDITOR
-	virtual int32 Compile(class FMaterialCompiler* Compiler, int32 OutputIndex) override;
 	virtual void GetCaption(TArray<FString>& OutCaptions) const override;
 #endif
 	//~ End UMaterialExpression Interface
diff --git a/Engine/Source/Runtime/Engine/Public/Materials/MaterialExpressionDecalLifetimeOpacity.h b/Engine/Source/Runtime/Engine/Public/Materials/MaterialExpressionDecalLifetimeOpacity.h
--- a/Engine/Source/Runtime/Engine/Public/Materials/MaterialExpressionDecalLifetimeOpacity.h
+++ b/Engine/Source/Runtime/Engine/Public/Materials/MaterialExpressionDecalLifetimeOpacity.h
@@ -5,17 +5,16 @@
 
 #include "CoreMinimal.h"
 #include "UObject/ObjectMacros.h"
-#include "Materials/MaterialExpression.h"
+#include "Materials/MaterialExpressionExternalCodeBase.h"
 #include "MaterialExpressionDecalLifetimeOpacity.generated.h"
 
 UCLASS(collapsecategories, hidecategories = Object)
-class UMaterialExpressionDecalLifetimeOpacity : public UMaterialExpression
+class UMaterialExpressionDecalLifetimeOpacity : public UMaterialExpressionExternalCodeBase
 {
 	GENERATED_UCLASS_BODY()
 
 	//~ Begin UMaterialExpression Interface
 #if WITH_EDITOR
-	virtual int32 Compile(class FMaterialCompiler* Compiler, int32 OutputIndex) override;
 	virtual void GetCaption(TArray<FString>& OutCaptions) const override;
 #endif
 	//~ End UMaterialExpression Interface
diff --git a/Engine/Source/Runtime/Engine/Public/Materials/MaterialExpressionDeltaTime.h b/Engine/Source/Runtime/Engine/Public/Materials/MaterialExpressionDeltaTime.h
--- a/Engine/Source/Runtime/Engine/Public/Materials/MaterialExpressionDeltaTime.h
+++ b/Engine/Source/Runtime/Engine/Public/Materials/MaterialExpressionDeltaTime.h
@@ -5,11 +5,11 @@
 
 #include "CoreMinimal.h"
 #include "UObject/ObjectMacros.h"
-#include "Materials/MaterialExpression.h"
+#include "Materials/MaterialExpressionExternalCodeBase.h"
 #include "MaterialExpressionDeltaTime.generated.h"
 
 UCLASS(collapsecategories, hidecategories=Object)
-class UMaterialExpressionDeltaTime : public UMaterialExpression
+class UMaterialExpressionDeltaTime : public UMaterialExpressionExternalCodeBase
 {
 	GENERATED_BODY()
 
@@ -17,7 +17,6 @@ public:
 
 	//~ Begin UMaterialExpression Interface
 #if WITH_EDITOR
-	virtual int32 Compile(class FMaterialCompiler* Compiler, int32 OutputIndex) override;
 	virtual void GetCaption(TArray<FString>& OutCaptions) const override;
 #endif
 	//~ End UMaterialExpression Interface
