World Partition - Execute cell transformers on HLOD cells during PIE sessions

diff --git a/Engine/Source/Runtime/Engine/Private/WorldPartition/WorldPartitionLevelStreamingDynamic.cpp b/Engine/Source/Runtime/Engine/Private/WorldPartition/WorldPartitionLevelStreamingDynamic.cpp
--- a/Engine/Source/Runtime/Engine/Private/WorldPartition/WorldPartitionLevelStreamingDynamic.cpp
+++ b/Engine/Source/Runtime/Engine/Private/WorldPartition/WorldPartitionLevelStreamingDynamic.cpp
@@ -512,6 +512,7 @@ void UWorldPartitionLevelStreamingDynamic::FinalizeRuntimeLevel()
 	UWorld* OuterWorld = RuntimeLevel->GetTypedOuter<UWorld>();
 	OuterWorld->bIsNameStableForNetworking = true;
 
+	// Remap soft object paths
 	if (StreamingCell.IsValid() && !StreamingCell->GetIsHLOD())
 	{
 		TRACE_CPUPROFILER_EVENT_SCOPE(FixupIDs);
@@ -534,7 +535,12 @@ void UWorldPartitionLevelStreamingDynamic::FinalizeRuntimeLevel()
 			// Remap Runtime Level's SoftObjectPaths
 			FWorldPartitionLevelHelper::RemapLevelSoftObjectPaths(RuntimeLevel, OuterWorldPartition.Get());
 		}
+	}
 
+	// Execute the cells transformers on the cell's level
+	if (StreamingCell.IsValid() && OuterWorld->IsGameWorld())
+	{
+		TRACE_CPUPROFILER_EVENT_SCOPE(ApplyRuntimeCellsTransformerStack);
 		OuterWorldPartition->ApplyRuntimeCellsTransformerStack(RuntimeLevel);
 	}
 
