Clang fix - Remove dll export from cpp

#rnx

diff --git a/Engine/Source/Runtime/Engine/Private/World.cpp b/Engine/Source/Runtime/Engine/Private/World.cpp
--- a/Engine/Source/Runtime/Engine/Private/World.cpp
+++ b/Engine/Source/Runtime/Engine/Private/World.cpp
@@ -6473,7 +6473,7 @@ float UWorld::GetDefaultGravityZ() const
 }
 
 /** This is our global function for retrieving the current MapName **/
-ENGINE_API const FString GetMapNameStatic()
+const FString GetMapNameStatic()
 {
 	FString Retval;
 
