Enables fix for tangent space normal hack in materials, precedently submitted to CL 39675439.

[RN] Fixed issue causing inconsistent handling of tangent vs world space normal space in the material depending on whether the Normal material attribute was connected.

diff --git a/Engine/Source/Runtime/Core/Private/UObject/FortniteMainBranchObjectVersion.cpp b/Engine/Source/Runtime/Core/Private/UObject/FortniteMainBranchObjectVersion.cpp
--- a/Engine/Source/Runtime/Core/Private/UObject/FortniteMainBranchObjectVersion.cpp
+++ b/Engine/Source/Runtime/Core/Private/UObject/FortniteMainBranchObjectVersion.cpp
@@ -14,6 +14,6 @@ TMap<FGuid, FGuid> FFortniteMainBranchObjectVersion::GetSystemGuids()
 	SystemGuids.Add(DevGuids.Niagara_LatestScriptCompileVersion, FGuid("670A2EBDEE1A4794BB77BBFF331EA7EF"));
 	SystemGuids.Add(DevGuids.SkeletalMeshDerivedDataVersion, FGuid("5E93192FCF3647EEB7F9D4FFF1750050"));
 	SystemGuids.Add(DevGuids.STATICMESH_DERIVEDDATA_VER, FGuid("656E338C78C34641A5D3BFAE303997AF"));
-	SystemGuids.Add(DevGuids.MaterialTranslationDDCVersion, FGuid("3DC4760CBA5A40E58DF2A2365A8C4F28"));
+	SystemGuids.Add(DevGuids.MaterialTranslationDDCVersion, FGuid("0195246791767C5080E1D4C3ED151B18"));
 	return SystemGuids;
 }
diff --git a/Engine/Source/Runtime/Engine/Private/Materials/Material.cpp b/Engine/Source/Runtime/Engine/Private/Materials/Material.cpp
--- a/Engine/Source/Runtime/Engine/Private/Materials/Material.cpp
+++ b/Engine/Source/Runtime/Engine/Private/Materials/Material.cpp
@@ -130,13 +130,6 @@ static TAutoConsoleVariable<bool> CVarMaterialsDuplicateVerbatim(
 	TEXT("Default: false"),
 	ECVF_Default);
 
-TAutoConsoleVariable<bool> CVarMaterialsTentativeNormalFix(
-	TEXT("r.Material.TentativeNormalFix"),
-	false,
-	TEXT("It enables a fix that removes a hack regarding tangent normal space handling when the Normal material attribute is connected.\n")
-	TEXT("Default: false"),
-	ECVF_Default);
-
 namespace MaterialImpl
 {
 	// Filtered list of properties that we follow on mobile platforms.
@@ -6974,7 +6967,7 @@ int32 UMaterial::CompilePropertyEx( FMaterialCompiler* Compiler, const FGuid& At
 		case MP_Normal:
 			// If material is not using tangent space normals, and nothing is connected to the normal pin,
 			// instead of setting the normal to the constant use the vertex normal.
-			if (CVarMaterialsTentativeNormalFix.GetValueOnAnyThread() && !this->bTangentSpaceNormal && EditorOnly->Normal.UseConstant)
+			if (!this->bTangentSpaceNormal && EditorOnly->Normal.UseConstant)
 			{
 				return Compiler->VertexNormal();
 			}
diff --git a/Engine/Source/Runtime/Engine/Private/Materials/MaterialShared.cpp b/Engine/Source/Runtime/Engine/Private/Materials/MaterialShared.cpp
--- a/Engine/Source/Runtime/Engine/Private/Materials/MaterialShared.cpp
+++ b/Engine/Source/Runtime/Engine/Private/Materials/MaterialShared.cpp
@@ -111,8 +111,6 @@ static TAutoConsoleVariable<bool> CVarDisableVRSForCustomPrimitiveData(
 	TEXT("Controls whether or not to disable VRS for materials using Custom Primitive Data."),
 	ECVF_RenderThreadSafe | ECVF_ReadOnly);
 
-extern TAutoConsoleVariable<bool> CVarMaterialsTentativeNormalFix;
-
 bool IsUsingNewMaterialTranslatorPrototype()
 {
 	return CVarUsingUseNewMaterialTranslatorPrototype.GetValueOnAnyThread();
@@ -1588,14 +1586,7 @@ void FMaterial::DiscardShaderMap()
 }
 
 EMaterialDomain FMaterialResource::GetMaterialDomain() const { return Material->MaterialDomain; }
-bool FMaterialResource::IsTangentSpaceNormal() const 
-{ 
-	if (IsSubstrateMaterial() || CVarMaterialsTentativeNormalFix.GetValueOnAnyThread())
-	{
-		return Material->bTangentSpaceNormal; // We do not need to check MP_Normal with Substrate as this cannot be specified on the root node anymore.
-	}
-	return Material->bTangentSpaceNormal || (!Material->IsPropertyConnected(MP_Normal) && !Material->bUseMaterialAttributes);
-}
+bool FMaterialResource::IsTangentSpaceNormal() const { return Material->bTangentSpaceNormal; }
 bool FMaterialResource::ShouldGenerateSphericalParticleNormals() const { return Material->bGenerateSphericalParticleNormals; }
 bool FMaterialResource::ShouldDisableDepthTest() const { return Material->bDisableDepthTest; }
 bool FMaterialResource::ShouldWriteOnlyAlpha() const { return Material->bWriteOnlyAlpha; }
