Bring back enqueue of mesh deformer work when creating scene proxy that was recently removed because it caused render pipe flush.
Now the work is enqueued to ComputeFramework EndOfFrame execution instead of Immediate execution which avoids the flush.
It looks like the Immediate execution was used to avoid a bug where queued work wasn't being retracted when the scene proxy is immediately recreated. That could lead to work happening after the GPUSkin Mesh Obect is destroyed, which can cause a crash.
We do mark the render data as dirty when we recreate mesh deformer instances in BeginPlay(), and this is one case where we would cause a crash like this.
The fix is to make sure we release the mesh deformer instance resources, which does the retraction for us.

diff --git a/Engine/Source/Runtime/Engine/Classes/Components/SkinnedMeshComponent.h b/Engine/Source/Runtime/Engine/Classes/Components/SkinnedMeshComponent.h
--- a/Engine/Source/Runtime/Engine/Classes/Components/SkinnedMeshComponent.h
+++ b/Engine/Source/Runtime/Engine/Classes/Components/SkinnedMeshComponent.h
@@ -334,6 +334,8 @@ public:
 private:
 	/** Initializes the MeshDeformerInstances property */
 	void CreateMeshDeformerInstances(const FMeshDeformerSet& DeformerSet);
+	/** Release resource contents of MeshDeformerInstances */
+	void ReleaseMeshDeformerInstances();
 
 public:
 	UFUNCTION(BlueprintPure, Category = "Components|SkinnedMesh")
diff --git a/Engine/Source/Runtime/Engine/Private/Components/SkinnedMeshComponent.cpp b/Engine/Source/Runtime/Engine/Private/Components/SkinnedMeshComponent.cpp
--- a/Engine/Source/Runtime/Engine/Private/Components/SkinnedMeshComponent.cpp
+++ b/Engine/Source/Runtime/Engine/Private/Components/SkinnedMeshComponent.cpp
@@ -1174,15 +1174,24 @@ void USkinnedMeshComponent::CreateRenderState_Concurrent(FRegisterComponentConte
 		}
 	}
 
-#if WITH_EDITORONLY_DATA
 	if (UMeshDeformerInstance* MeshDeformerInstance = GetMeshDeformerInstanceForLOD(UseLOD))
  	{
+
+#if WITH_EDITORONLY_DATA
 		for (TUniquePtr<FMeshDeformerGeometryReadbackRequest>& Request : GeometryReadbackRequests)
 		{
 			MeshDeformerInstance->RequestReadbackDeformerGeometry(MoveTemp(Request));
 		}
+#endif // WITH_EDITORONLY_DATA
+
+		// Enqueue work here to ensure that we have some deformer outputs written for the first frame.
+		UMeshDeformerInstance::FEnqueueWorkDesc Desc;
+		Desc.Scene = GetScene();
+		Desc.OwnerName = GetSkinnedAsset() != nullptr ? GetSkinnedAsset()->GetFName() : GetFName();
+		MeshDeformerInstance->EnqueueWork(Desc);
 	}
 
+#if WITH_EDITORONLY_DATA
 	// Avoid unbound accumulation of readback requests, each readback request must be picked up by the deformer instance within the same frame or ignored otherwise
 	GeometryReadbackRequests.Reset();
 #endif // WITH_EDITORONLY_DATA
@@ -1196,13 +1205,7 @@ void USkinnedMeshComponent::DestroyRenderState_Concurrent()
 
 	Super::DestroyRenderState_Concurrent();
 
-	for (UMeshDeformerInstance* MeshDeformerInstance : MeshDeformerInstances.DeformerInstances)
-	{
-		if (MeshDeformerInstance != nullptr)
-		{
-			MeshDeformerInstance->ReleaseResources();
-		}
-	}
+	ReleaseMeshDeformerInstances();
 
 	if(MeshObject)
 	{
@@ -2536,6 +2539,8 @@ FMeshDeformerSet USkinnedMeshComponent::GetActiveMeshDeformers() const
 
 void USkinnedMeshComponent::CreateMeshDeformerInstances(const FMeshDeformerSet& DeformerSet)
 {
+	ReleaseMeshDeformerInstances();
+
 	// Only one deformer is supported for now
 	check(DeformerSet.Deformers.Num() <= 1);
 	UMeshDeformer* ActiveMeshDeformer = DeformerSet.Deformers.Num() > 0 ? DeformerSet.Deformers[0] : nullptr;
@@ -2552,6 +2557,17 @@ void USkinnedMeshComponent::CreateMeshDeformerInstances(const FMeshDeformerSet&
 	}
 }
 
+void USkinnedMeshComponent::ReleaseMeshDeformerInstances()
+{
+	for (UMeshDeformerInstance* MeshDeformerInstance : MeshDeformerInstances.DeformerInstances)
+	{
+		if (MeshDeformerInstance != nullptr)
+		{
+			MeshDeformerInstance->ReleaseResources();
+		}
+	}
+}
+
 UMeshDeformerInstance* USkinnedMeshComponent::GetMeshDeformerInstance() const
 {
 	return MeshDeformerInstances.DeformerInstances.Num() > 0 ? MeshDeformerInstances.DeformerInstances[0]->GetInstanceForSourceDeformer() : nullptr;
