Spline Component: Encapsulate legacy spline data so that we can move forward with the new implementation.

This change primarily replaces SplineCurves with Spline and makes the changes necessary as a consequence of this swap. SplineCurves is still the default backing data structure, though it is now deprecated. Before we remove it, we will flip the CVar ```SplineComponent.UseSplineCurves```, as this is what defines the authoritative backing structure.

The new, opt-in backing structure is FSpline. This is designed to have one of three possible implementations: empty, legacy, and new. The empty implementation is the default, it will not be functional. The legacy implementation is implemented by FLegacySpline, a struct defined inside FSpline. The new implementation is not yet implemented. The implementation is controlled by the CVar ```Spline.Implementation```

diff --git a/Engine/Plugins/2D/Paper2D/Source/Paper2D/Private/Terrain/PaperTerrainComponent.cpp b/Engine/Plugins/2D/Paper2D/Source/Paper2D/Private/Terrain/PaperTerrainComponent.cpp
--- a/Engine/Plugins/2D/Paper2D/Source/Paper2D/Private/Terrain/PaperTerrainComponent.cpp
+++ b/Engine/Plugins/2D/Paper2D/Source/Paper2D/Private/Terrain/PaperTerrainComponent.cpp
@@ -249,11 +249,11 @@ FTransform UPaperTerrainComponent::GetTransformAtDistance(float InDistance) cons
 	const float SplineLength = AssociatedSpline->GetSplineLength();
 	InDistance = FMath::Clamp<float>(InDistance, 0.0f, SplineLength);
 
-	const float Param = AssociatedSpline->SplineCurves.ReparamTable.Eval(InDistance, 0.0f);
-	const FVector Position3D = AssociatedSpline->SplineCurves.Position.Eval(Param, FVector::ZeroVector);
+	const float Param = AssociatedSpline->GetInputKeyValueAtDistanceAlongSpline(InDistance);
+	const FVector Position3D = AssociatedSpline->GetLocationAtSplineInputKey(Param, ESplineCoordinateSpace::Local);
 
-	const FVector Tangent = AssociatedSpline->SplineCurves.Position.EvalDerivative(Param, FVector(1.0f, 0.0f, 0.0f)).GetSafeNormal();
-	const FVector NormalEst = AssociatedSpline->SplineCurves.Position.EvalSecondDerivative(Param, FVector(0.0f, 1.0f, 0.0f)).GetSafeNormal();
+	const FVector Tangent = AssociatedSpline->GetTangentAtSplineInputKey(Param, ESplineCoordinateSpace::Local).GetSafeNormal();
+	const FVector NormalEst = StaticCastPtr<const UPaperTerrainSplineComponent, UPaperTerrainSplineComponent>(AssociatedSpline)->GetSplinePointsPosition().EvalSecondDerivative(Param, FVector(0.0f, 1.0f, 0.0f)).GetSafeNormal();
 	const FVector Bitangent = FVector::CrossProduct(Tangent, NormalEst);
 	const FVector Normal = FVector::CrossProduct(Bitangent, Tangent);
 	const FVector Floop = FVector::CrossProduct(PaperAxisZ, Tangent);
@@ -368,16 +368,20 @@ void UPaperTerrainComponent::ConstrainSplinePointsToXZ()
 	if (AssociatedSpline != nullptr)
 	{
 		bool bSplineChanged = false;
-		auto& Points = AssociatedSpline->SplineCurves.Position.Points;
-		int32 NumPoints = Points.Num();
+		const int32 NumPoints = AssociatedSpline->GetNumberOfSplinePoints();
 		for (int PointIndex = 0; PointIndex < NumPoints; ++PointIndex)
 		{
-			auto& CurrentPoint = Points[PointIndex];
-			if (CurrentPoint.ArriveTangent.Y != 0 || CurrentPoint.LeaveTangent.Y != 0 || CurrentPoint.OutVal.Y != 0)
+			const FSplinePoint Point = AssociatedSpline->GetSplinePointAt(PointIndex, ESplineCoordinateSpace::Local);
+			
+			if (Point.ArriveTangent.Y != 0 || Point.LeaveTangent.Y != 0 || Point.Position.Y != 0)
 			{
-				CurrentPoint.ArriveTangent.Y = 0;
-				CurrentPoint.LeaveTangent.Y = 0;
-				CurrentPoint.OutVal.Y = 0;
+				FVector NewArriveTangent = Point.ArriveTangent; NewArriveTangent.Y = 0;
+				FVector NewLeaveTangent = Point.LeaveTangent; NewLeaveTangent.Y = 0;
+				AssociatedSpline->SetTangentsAtSplinePoint(PointIndex, NewArriveTangent, NewLeaveTangent, ESplineCoordinateSpace::Local, false);
+				
+				FVector NewPosition = Point.Position; NewPosition.Y = 0;
+				AssociatedSpline->SetLocationAtSplinePoint(PointIndex, NewPosition, ESplineCoordinateSpace::Local, false);
+				
 				bSplineChanged = true;
 			}
 		}
@@ -418,8 +422,6 @@ void UPaperTerrainComponent::OnSplineEdited()
 
 		FRandomStream RandomStream(RandomSeed);
 
-		const FInterpCurveVector& SplineInfo = AssociatedSpline->SplineCurves.Position;
-
 		float SplineLength = AssociatedSpline->GetSplineLength();
 
 
@@ -620,8 +622,9 @@ void UPaperTerrainComponent::OnSplineEdited()
 				float CurrentTime = 0.0f;
 				while (CurrentTime < SplineLength)
 				{
-					const float Param = AssociatedSpline->SplineCurves.ReparamTable.Eval(CurrentTime, 0.0f);
-					const FVector Position3D = AssociatedSpline->SplineCurves.Position.Eval(Param, FVector::ZeroVector);
+					// note [nickolas.drake]: SplineLength is not measured in time, I am adapting this code in spite of the fact that it seems wrong at a glance.
+					const float Param = AssociatedSpline->GetInputKeyValueAtDistanceAlongSpline(CurrentTime);
+					const FVector Position3D = AssociatedSpline->GetLocationAtSplineInputKey(Param, ESplineCoordinateSpace::Local);
 					const FVector2D Position2D = FVector2D(FVector::DotProduct(Position3D, PaperAxisX), FVector::DotProduct(Position3D, PaperAxisY));
 
 					SplineBounds += Position2D;
diff --git a/Engine/Plugins/Experimental/Buoyancy/Source/Runtime/Private/BuoyancySubsystem.cpp b/Engine/Plugins/Experimental/Buoyancy/Source/Runtime/Private/BuoyancySubsystem.cpp
--- a/Engine/Plugins/Experimental/Buoyancy/Source/Runtime/Private/BuoyancySubsystem.cpp
+++ b/Engine/Plugins/Experimental/Buoyancy/Source/Runtime/Private/BuoyancySubsystem.cpp
@@ -366,7 +366,7 @@ void UBuoyancySubsystem::UpdateSplineData()
 			UWaterSplineMetadata* SplineMetadata = WaterBodyComponent->GetWaterSplineMetadata();
 			UBakedShallowWaterSimulationComponent *BakedSim = WaterBodyComponent->GetBakedShallowWaterSimulation();
 			
-			if (UWaterSplineComponent* SplineComponent = WaterBodyComponent->GetWaterSpline())
+			if (const UWaterSplineComponent* SplineComponent = WaterBodyComponent->GetWaterSpline())
 			{
 				// Copy out water spline data into a shared ptr, to be associated with all
 				// child particles and marshaled to PT.
@@ -375,7 +375,7 @@ void UBuoyancySubsystem::UpdateSplineData()
 				const TOptional<FShallowWaterSimulationGrid> EmptyOptionalSimGrid;
 				TSharedPtr<FBuoyancyWaterSplineData> WaterSplineData = MakeShared<FBuoyancyWaterSplineData>(
 					WaterTransform,
-					SplineComponent->SplineCurves.Position,
+					SplineComponent->GetSplinePointsPosition(),
 					WaterBodyComponent->GetWaterBodyType(),
 					SplineMetadata ? SplineMetadata->RiverWidth : EmptyOptionalFloat,
 					SplineMetadata ? SplineMetadata->WaterVelocityScalar : EmptyOptionalFloat,
diff --git a/Engine/Plugins/Experimental/CineCameraRigs/Source/CineCameraRigs/Private/CineSplineComponent.cpp b/Engine/Plugins/Experimental/CineCameraRigs/Source/CineCameraRigs/Private/CineSplineComponent.cpp
--- a/Engine/Plugins/Experimental/CineCameraRigs/Source/CineCameraRigs/Private/CineSplineComponent.cpp
+++ b/Engine/Plugins/Experimental/CineCameraRigs/Source/CineCameraRigs/Private/CineSplineComponent.cpp
@@ -7,9 +7,8 @@
 UCineSplineComponent::UCineSplineComponent(const FObjectInitializer& ObjectInitializer)
 	: Super(ObjectInitializer)
 {
-	SplineCurves.Position.Points.Reset(10);
-	SplineCurves.Rotation.Points.Reset(10);
-	SplineCurves.Scale.Points.Reset(10);
+	ClearSplinePoints(false);
+	
 	CineSplineMetadata = ObjectInitializer.CreateDefaultSubobject<UCineSplineMetadata>(this, TEXT("CineSplineMetadata"));
 	CineSplineMetadata->Reset(10);
 
@@ -55,7 +54,7 @@ TStructOnScope<FActorComponentInstanceData> UCineSplineComponent::GetComponentIn
 	if (bSplineHasBeenEdited)
 	{
 		SplineInstanceData->CineSplineMetadata = CineSplineMetadata;
-		SplineInstanceData->SplineCurves = SplineCurves;
+		SplineInstanceData->SplineCurves = GetSplineCurves();
 	}
 
 	SplineInstanceData->bSplineHasBeenEdited = bSplineHasBeenEdited;
@@ -82,7 +81,7 @@ void UCineSplineComponent::ApplyComponentInstanceData(FCineSplineInstanceData* S
 			// If we are restoring the saved state, unmark the SplineCurves property as 'modified'.
 			// We don't want to consider that these changes have been made through the UCS.
 			TArray<FProperty*> Properties;
-			Properties.Emplace(FindFProperty<FProperty>(UCineSplineComponent::StaticClass(), GET_MEMBER_NAME_CHECKED(UCineSplineComponent, SplineCurves)));
+			Properties.Emplace(FindFProperty<FProperty>(UCineSplineComponent::StaticClass(), GetSplinePropertyName()));
 			RemoveUCSModifiedProperties(Properties);
 		}
 	}
diff --git a/Engine/Plugins/Experimental/CineCameraRigs/Source/CineCameraRigsEditor/Private/CineSplineMetadataDetails.cpp b/Engine/Plugins/Experimental/CineCameraRigs/Source/CineCameraRigsEditor/Private/CineSplineMetadataDetails.cpp
--- a/Engine/Plugins/Experimental/CineCameraRigs/Source/CineCameraRigsEditor/Private/CineSplineMetadataDetails.cpp
+++ b/Engine/Plugins/Experimental/CineCameraRigs/Source/CineCameraRigsEditor/Private/CineSplineMetadataDetails.cpp
@@ -97,7 +97,7 @@ void SetValues(FCineSplineMetadataDetails& Details, TArray<FInterpCurvePoint<T>>
 
 	Details.SplineComp->UpdateSpline();
 	Details.SplineComp->bSplineHasBeenEdited = true;
-	static FProperty* SplineCurvesProperty = FindFProperty<FProperty>(USplineComponent::StaticClass(), GET_MEMBER_NAME_CHECKED(USplineComponent, SplineCurves));
+	static FProperty* SplineCurvesProperty = FindFProperty<FProperty>(USplineComponent::StaticClass(), USplineComponent::GetSplinePropertyName());
 	FComponentVisualizer::NotifyPropertyModified(Details.SplineComp, SplineCurvesProperty);
 	Details.Update(Details.SplineComp, Details.SelectedKeys);
 
diff --git a/Engine/Plugins/Experimental/MeshModelingToolsetExp/Source/MeshModelingToolsEditorOnlyExp/Private/DrawSplineTool.cpp b/Engine/Plugins/Experimental/MeshModelingToolsetExp/Source/MeshModelingToolsEditorOnlyExp/Private/DrawSplineTool.cpp
--- a/Engine/Plugins/Experimental/MeshModelingToolsetExp/Source/MeshModelingToolsEditorOnlyExp/Private/DrawSplineTool.cpp
+++ b/Engine/Plugins/Experimental/MeshModelingToolsetExp/Source/MeshModelingToolsEditorOnlyExp/Private/DrawSplineTool.cpp
@@ -744,7 +744,7 @@ FVector3d UDrawSplineTool::GetUpVectorToUse(const FVector3d& HitLocation, const
 			// If there's only one point, GetUpVectorAtSplinePoint is unreliable because it seeks to build a
 			// quaternion from the tangent and the set up vector, and the tangent is zero. We want to use
 			// the "stored" up vector directly.
-			FVector3d LocalUpVector = WorkingSpline->SplineCurves.Rotation.Points[0].OutVal.RotateVector(WorkingSpline->DefaultUpVector);
+			FVector3d LocalUpVector = WorkingSpline->GetRotationAtSplinePoint(0, ESplineCoordinateSpace::Local).RotateVector(WorkingSpline->DefaultUpVector);
 			UpVectorToUse = WorkingSpline->GetComponentTransform().TransformVectorNoScale(LocalUpVector);
 		}
 		break;
diff --git a/Engine/Plugins/Experimental/MeshModelingToolsetExp/Source/MeshModelingToolsExp/Private/Splines/BaseMeshFromSplinesTool.cpp b/Engine/Plugins/Experimental/MeshModelingToolsetExp/Source/MeshModelingToolsExp/Private/Splines/BaseMeshFromSplinesTool.cpp
--- a/Engine/Plugins/Experimental/MeshModelingToolsetExp/Source/MeshModelingToolsExp/Private/Splines/BaseMeshFromSplinesTool.cpp
+++ b/Engine/Plugins/Experimental/MeshModelingToolsetExp/Source/MeshModelingToolsExp/Private/Splines/BaseMeshFromSplinesTool.cpp
@@ -65,7 +65,7 @@ void UBaseMeshFromSplinesTool::PollSplineUpdates()
 	int32 SplineIdx = 0;
 	EnumerateSplines([this, &bSplinesUpdated, &SplineIdx](USplineComponent* SplineComponent)
 	{
-		int32 Version = SplineComponent->SplineCurves.Version;
+		int32 Version = SplineComponent->GetVersion();
 		FTransform Transform = SplineComponent->GetComponentTransform();
 		if (SplineIdx >= LastSplineVersions.Num())
 		{
diff --git a/Engine/Plugins/Experimental/VirtualProductionUtilities/Source/VPUtilities/Private/Libraries/VPCameraBlueprintLibrary.cpp b/Engine/Plugins/Experimental/VirtualProductionUtilities/Source/VPUtilities/Private/Libraries/VPCameraBlueprintLibrary.cpp
--- a/Engine/Plugins/Experimental/VirtualProductionUtilities/Source/VPUtilities/Private/Libraries/VPCameraBlueprintLibrary.cpp
+++ b/Engine/Plugins/Experimental/VirtualProductionUtilities/Source/VPUtilities/Private/Libraries/VPCameraBlueprintLibrary.cpp
@@ -72,7 +72,11 @@ ACameraRig_Rail* UVPCameraBlueprintLibrary::SpawnDollyTrackFromPoints(UObject* W
 	}
 
 	SplineComponent->UpdateSpline();
-	SplineComponent->SplineCurves.ReparamTable.AutoSetTangents(1.0, true);
+
+	// Note: We previously called AutoSetTangents on the reparameterization table here.
+	// With the recent refactor deprecating SplineCurves, there was no way to replace this code.
+	// It did have the effect of setting tangents, but these tangents were not actually
+	// used because the reparam points were all CIM_Linear, so evaluation did not require tangents.
 
 	return Cast<ACameraRig_Rail>(UGameplayStatics::FinishSpawningActor(DollyTrack, Points[0]));
 }
@@ -86,7 +90,6 @@ ACameraRig_Rail* UVPCameraBlueprintLibrary::SpawnDollyTrackFromPointsSmooth(UObj
 	{
 		USplineComponent* SplineComponent = DollyTrack->GetRailSplineComponent();
 
-		const float TotalLength = SplineComponent->SplineCurves.GetSplineLength();
 		const int32 Substeps = SplineComponent->ReparamStepsPerSegment;
 		const int32 NumPoints = Points.Num();
 		const float Base = static_cast<float>(NumPoints * Substeps - 1);
@@ -121,7 +124,11 @@ ACameraRig_Rail* UVPCameraBlueprintLibrary::SpawnDollyTrackFromPointsSmooth(UObj
 
 		SplineComponent->ClearSplinePoints(false);
 		SplineComponent->AddPoints(SplinePoints, /* bUpdateSpline= */ true);
-		SplineComponent->SplineCurves.ReparamTable.AutoSetTangents(1.0, true);
+		
+		// Note: We previously called AutoSetTangents on the reparameterization table here.
+		// With the recent refactor deprecating SplineCurves, there was no way to replace this code.
+		// It did have the effect of setting tangents, but these tangents were not actually
+		// used because the reparam points were all CIM_Linear, so evaluation did not require tangents.
 	}
 
 	return DollyTrack;
@@ -171,17 +178,17 @@ ACameraRig_Rail* UVPCameraBlueprintLibrary::SpawnCameraRigFromPoints(UObject* Wo
 
 	if (Params.LinearApproximationMode != EVPCameraRigSpawnLinearApproximationMode::None)
 	{
-		const FSplineCurves& Curves = SplineComponent->SplineCurves;
+		const FSplineCurves CachedLegacySplineData = SplineComponent->GetSplineCurves();
 		TArray<FSplinePositionLinearApproximation> OutPoints;
 		float Density = Params.LinearApproximationParam;
 
 		if (Params.LinearApproximationMode == EVPCameraRigSpawnLinearApproximationMode::IntegrationStep)
 		{
 			// Convert integration step to density.
-			Density = FMath::CeilToFloat(1.f / (Curves.GetSplineLength() / FMath::Max(KINDA_SMALL_NUMBER, Density)));
+			Density = FMath::CeilToFloat(1.f / (CachedLegacySplineData.GetSplineLength() / FMath::Max(KINDA_SMALL_NUMBER, Density)));
 		}
 
-		FSplinePositionLinearApproximation::Build(Curves, OutPoints, Density);
+		FSplinePositionLinearApproximation::Build(CachedLegacySplineData, OutPoints, Density);
 
 		TArray<FVector> NewPoints;
 		NewPoints.Reserve(OutPoints.Num());
diff --git a/Engine/Plugins/Experimental/Water/Source/Editor/Private/WaterSplineComponentVisualizer.cpp b/Engine/Plugins/Experimental/Water/Source/Editor/Private/WaterSplineComponentVisualizer.cpp
--- a/Engine/Plugins/Experimental/Water/Source/Editor/Private/WaterSplineComponentVisualizer.cpp
+++ b/Engine/Plugins/Experimental/Water/Source/Editor/Private/WaterSplineComponentVisualizer.cpp
@@ -473,8 +473,7 @@ float FWaterSplineComponentVisualizer::ComputeDelta(UWaterSplineComponent* Water
 	check(WaterSplineComp);
 	check(SelectionState != nullptr);
 
-	FInterpCurveVector& SplinePosition = WaterSplineComp->GetSplinePointsPosition();
-	const int32 NumPoints = SplinePosition.Points.Num();
+	const int32 NumPoints = WaterSplineComp->GetNumberOfSplinePoints();
 	int32 LastKeyIndexSelected = SelectionState->GetVerifiedLastKeyIndexSelected(NumPoints);
 
 	const FTransform SplinePointToComponent = WaterSplineComp->GetTransformAtSplineInputKey(LastKeyIndexSelected, ESplineCoordinateSpace::Local, true);
diff --git a/Engine/Plugins/Experimental/Water/Source/Editor/Private/WaterSplineMetadataDetails.cpp b/Engine/Plugins/Experimental/Water/Source/Editor/Private/WaterSplineMetadataDetails.cpp
--- a/Engine/Plugins/Experimental/Water/Source/Editor/Private/WaterSplineMetadataDetails.cpp
+++ b/Engine/Plugins/Experimental/Water/Source/Editor/Private/WaterSplineMetadataDetails.cpp
@@ -110,7 +110,7 @@ void SetValues(FWaterSplineMetadataDetails& Details, TArray<FInterpCurvePoint<T>
 
 	Details.SplineComp->UpdateSpline();
 	Details.SplineComp->bSplineHasBeenEdited = true;
-	static FProperty* SplineCurvesProperty = FindFProperty<FProperty>(USplineComponent::StaticClass(), GET_MEMBER_NAME_CHECKED(USplineComponent, SplineCurves));
+	static FProperty* SplineCurvesProperty = FindFProperty<FProperty>(USplineComponent::StaticClass(), USplineComponent::GetSplinePropertyName());
 	EPropertyChangeType::Type PropertyChangeType = CommitInfo == ETextCommit::OnEnter ? EPropertyChangeType::ValueSet : EPropertyChangeType::Interactive;
 	FComponentVisualizer::NotifyPropertyModified(Details.SplineComp, SplineCurvesProperty, PropertyChangeType);
 	Details.Update(Details.SplineComp, Details.SelectedKeys);
diff --git a/Engine/Plugins/Experimental/Water/Source/Runtime/Private/WaterBodyActor.cpp b/Engine/Plugins/Experimental/Water/Source/Runtime/Private/WaterBodyActor.cpp
--- a/Engine/Plugins/Experimental/Water/Source/Runtime/Private/WaterBodyActor.cpp
+++ b/Engine/Plugins/Experimental/Water/Source/Runtime/Private/WaterBodyActor.cpp
@@ -405,11 +405,12 @@ void AWaterBody::DeprecateData()
 	{
 		if (GetLinkerCustomVersion(FFortniteMainBranchObjectVersion::GUID) < FFortniteMainBranchObjectVersion::MoveWaterMetadataToActor)
 		{
+PRAGMA_DISABLE_DEPRECATION_WARNINGS
 			if (SplineComp->SplineCurves.Metadata_DEPRECATED)
 			{
 				UWaterSplineMetadata* OldSplineMetadata = Cast<UWaterSplineMetadata>(SplineComp->SplineCurves.Metadata_DEPRECATED);
 				SplineComp->SplineCurves.Metadata_DEPRECATED = nullptr;
-
+PRAGMA_ENABLE_DEPRECATION_WARNINGS
 				if (WaterSplineMetadata)
 				{
 					WaterSplineMetadata->Depth = OldSplineMetadata->Depth;
diff --git a/Engine/Plugins/Experimental/Water/Source/Runtime/Private/WaterBodyRiverComponent.cpp b/Engine/Plugins/Experimental/Water/Source/Runtime/Private/WaterBodyRiverComponent.cpp
--- a/Engine/Plugins/Experimental/Water/Source/Runtime/Private/WaterBodyRiverComponent.cpp
+++ b/Engine/Plugins/Experimental/Water/Source/Runtime/Private/WaterBodyRiverComponent.cpp
@@ -107,7 +107,7 @@ static void AddVerticesForRiverSplineStep(
 	const FVector Normal = FVector::CrossProduct(Tangent, Up).GetSafeNormal();
 	const FVector Pos = SplineComp->GetLocationAtDistanceAlongSpline(DistanceAlongSpline, ESplineCoordinateSpace::Local);
 
-	const float Key = SplineComp->SplineCurves.ReparamTable.Eval(DistanceAlongSpline, 0.f);
+	const float Key = SplineComp->GetInputKeyValueAtDistanceAlongSpline(DistanceAlongSpline);
 	const float HalfWidth = WaterSplineMetadata->RiverWidth.Eval(Key) / 2;
 	float Velocity = WaterSplineMetadata->WaterVelocityScalar.Eval(Key);
 
@@ -235,7 +235,7 @@ static void AddTerminalVerticesForRiverSpline(
 	const FVector Normal = FVector::CrossProduct(Tangent, Up).GetSafeNormal();
 	const FVector Pos = SplineComp->GetLocationAtDistanceAlongSpline(DistanceAlongSpline, ESplineCoordinateSpace::Local);
 	
-	const float Key = SplineComp->SplineCurves.ReparamTable.Eval(DistanceAlongSpline, 0.f);
+	const float Key = SplineComp->GetInputKeyValueAtDistanceAlongSpline(DistanceAlongSpline);
 	const float HalfWidth = WaterSplineMetadata->RiverWidth.Eval(Key) / 2;
 
 	const float DilationAmount = Component->ShapeDilation;
diff --git a/Engine/Plugins/Experimental/Water/Source/Runtime/Private/WaterBodyTypes.cpp b/Engine/Plugins/Experimental/Water/Source/Runtime/Private/WaterBodyTypes.cpp
--- a/Engine/Plugins/Experimental/Water/Source/Runtime/Private/WaterBodyTypes.cpp
+++ b/Engine/Plugins/Experimental/Water/Source/Runtime/Private/WaterBodyTypes.cpp
@@ -97,7 +97,7 @@ FWaterSplineDataPhysics::FWaterSplineDataPhysics(UWaterSplineComponent* WaterSpl
 	if (WaterSplineComponent)
 	{
 		ComponentTransform = WaterSplineComponent->GetComponentTransform();
-		SplineCurves = WaterSplineComponent->SplineCurves;
+		SplineCurves = WaterSplineComponent->GetSplineCurves();
 		DefaultUpVector = WaterSplineComponent->GetDefaultUpVector(ESplineCoordinateSpace::World);
 	}
 }
@@ -107,7 +107,7 @@ FWaterSplineDataPhysics& FWaterSplineDataPhysics::operator=(const UWaterSplineCo
 	if (WaterSplineComponent)
 	{
 		ComponentTransform = WaterSplineComponent->GetComponentTransform();
-		SplineCurves = WaterSplineComponent->SplineCurves;
+		SplineCurves = WaterSplineComponent->GetSplineCurves();
 		DefaultUpVector = WaterSplineComponent->GetDefaultUpVector(ESplineCoordinateSpace::World);
 	}
 	return *this;
diff --git a/Engine/Plugins/Experimental/Water/Source/Runtime/Private/WaterSplineComponent.cpp b/Engine/Plugins/Experimental/Water/Source/Runtime/Private/WaterSplineComponent.cpp
--- a/Engine/Plugins/Experimental/Water/Source/Runtime/Private/WaterSplineComponent.cpp
+++ b/Engine/Plugins/Experimental/Water/Source/Runtime/Private/WaterSplineComponent.cpp
@@ -14,22 +14,11 @@ UWaterSplineComponent::UWaterSplineComponent(const FObjectInitializer& ObjectIni
 	//@todo_water: Remove once AWaterBody is not Blueprintable
 	{
 		// Add default spline points
-		SplineCurves.Position.Points.Empty(3);
-		SplineCurves.Rotation.Points.Empty(3);
-		SplineCurves.Scale.Points.Empty(3);
-
-		SplineCurves.Position.Points.Emplace(0.0f, FVector(0, 0, 0), FVector::ZeroVector, FVector::ZeroVector, CIM_CurveAuto);
-		SplineCurves.Rotation.Points.Emplace(0.0f, FQuat::Identity, FQuat::Identity, FQuat::Identity, CIM_CurveAuto);
-		SplineCurves.Scale.Points.Emplace(0.0f, FVector(WaterSplineDefaults.DefaultWidth, WaterSplineDefaults.DefaultDepth, 1.0f), FVector::ZeroVector, FVector::ZeroVector, CIM_CurveAuto);
-
-		SplineCurves.Position.Points.Emplace(1.0f, FVector(7000, -3000, 0), FVector::ZeroVector, FVector::ZeroVector, CIM_CurveAuto);
-		SplineCurves.Rotation.Points.Emplace(1.0f, FQuat::Identity, FQuat::Identity, FQuat::Identity, CIM_CurveAuto);
-		SplineCurves.Scale.Points.Emplace(1.0f, FVector(WaterSplineDefaults.DefaultWidth, WaterSplineDefaults.DefaultDepth, 1.0f), FVector::ZeroVector, FVector::ZeroVector, CIM_CurveAuto);
-
-		SplineCurves.Position.Points.Emplace(2.0f, FVector(6500, 6500, 0), FVector::ZeroVector, FVector::ZeroVector, CIM_CurveAuto);
-		SplineCurves.Rotation.Points.Emplace(2.0f, FQuat::Identity, FQuat::Identity, FQuat::Identity, CIM_CurveAuto);
-		SplineCurves.Scale.Points.Emplace(2.0f, FVector(WaterSplineDefaults.DefaultWidth, WaterSplineDefaults.DefaultDepth, 1.0f), FVector::ZeroVector, FVector::ZeroVector, CIM_CurveAuto);
+		ClearSplinePoints();
 
+		AddPoint(FSplinePoint(0.f, FVector(0, 0, 0), FVector::ZeroVector, FVector::ZeroVector, FQuat::Identity.Rotator(), FVector(WaterSplineDefaults.DefaultWidth, WaterSplineDefaults.DefaultDepth, 1.0f), ESplinePointType::Curve), false);
+		AddPoint(FSplinePoint(1.f, FVector(7000, -3000, 0), FVector::ZeroVector, FVector::ZeroVector, FQuat::Identity.Rotator(), FVector(WaterSplineDefaults.DefaultWidth, WaterSplineDefaults.DefaultDepth, 1.0f), ESplinePointType::Curve), false);
+		AddPoint(FSplinePoint(2.f, FVector(6500, 6500, 0), FVector::ZeroVector, FVector::ZeroVector, FQuat::Identity.Rotator(), FVector(WaterSplineDefaults.DefaultWidth, WaterSplineDefaults.DefaultDepth, 1.0f), ESplinePointType::Curve), false);
 	}
 }
 
@@ -209,15 +198,20 @@ bool UWaterSplineComponent::SynchronizeWaterProperties()
 	
 		for (int32 Point = 0; Point < GetNumberOfSplinePoints(); ++Point)
 		{
-			float Time = SplineCurves.Position.Points[Point].InVal;
+PRAGMA_DISABLE_DEPRECATION_WARNINGS
 			if (!SplineCurves.Scale.Points.IsValidIndex(Point))
 			{
-				SplineCurves.Scale.Points.Emplace(Time, FVector(WaterSplineDefaults.DefaultWidth, WaterSplineDefaults.DefaultDepth, 1.0f), FVector::ZeroVector, FVector::ZeroVector, CIM_CurveAuto);
+				float Param = GetInputKeyValueAtSplinePoint(Point);
+
+				SplineCurves.Scale.Points.Emplace(Param, FVector(WaterSplineDefaults.DefaultWidth, WaterSplineDefaults.DefaultDepth, 1.0f), FVector::ZeroVector, FVector::ZeroVector, CIM_CurveAuto);
+				// Required for the direct write above. Unfortunate we do in loop body.
+				SynchronizeSplines();
 			}
+PRAGMA_ENABLE_DEPRECATION_WARNINGS
 
 			bAnythingChanged |= Metadata->PropagateDefaultValue(Point, PreviousWaterSplineDefaults, WaterSplineDefaults);
 
-			FVector& Scale = SplineCurves.Scale.Points[Point].OutVal;
+			FVector Scale = GetScaleAtSplinePoint(Point);
 			float& DepthAtPoint = Metadata->Depth.Points[Point].OutVal;
 			float& WidthAtPoint = Metadata->RiverWidth.Points[Point].OutVal;
 
@@ -249,20 +243,7 @@ bool UWaterSplineComponent::SynchronizeWaterProperties()
 				Scale.Y = DepthAtPoint = FMath::Max(DepthAtPoint, KINDA_SMALL_NUMBER);
 			}
 
-			// #hack: temporarily clamp the tangents to a sensible range to prevent multiplicatively scaling until they hit infinity and cause a crash
-			auto ClampVec3 = [](const FVector& Vec3, double Min, double Max) {
-				FVector Result;
-				Result.X = FMath::Clamp(Vec3.X, Min, Max);
-				Result.Y = FMath::Clamp(Vec3.Y, Min, Max);
-				Result.Z = FMath::Clamp(Vec3.Z, Min, Max);
-				return Result;
-			};
-			FVector& ArriveTangent = SplineCurves.Position.Points[Point].ArriveTangent;
-			FVector& LeaveTangent = SplineCurves.Position.Points[Point].LeaveTangent;
-			
-			constexpr double MaxTangentValue = 1.e10L;
-			ArriveTangent = ClampVec3(ArriveTangent, -MaxTangentValue, MaxTangentValue);
-			LeaveTangent = ClampVec3(LeaveTangent, -MaxTangentValue, MaxTangentValue);
+			SetScaleAtSplinePoint(Point, Scale, false);
 		}
 	}
 
diff --git a/Engine/Plugins/Experimental/WaterAdvanced/Source/WaterAdvanced/Private/ShallowWaterRiverActor.cpp b/Engine/Plugins/Experimental/WaterAdvanced/Source/WaterAdvanced/Private/ShallowWaterRiverActor.cpp
--- a/Engine/Plugins/Experimental/WaterAdvanced/Source/WaterAdvanced/Private/ShallowWaterRiverActor.cpp
+++ b/Engine/Plugins/Experimental/WaterAdvanced/Source/WaterAdvanced/Private/ShallowWaterRiverActor.cpp
@@ -905,13 +905,13 @@ bool UShallowWaterRiverComponent::QueryWaterAtSplinePoint(TSoftObjectPtr<AWaterB
 
 			if (Metadata != nullptr)
 			{
-				OutPos = CurrSpline->SplineCurves.Position.Points[SplinePoint].OutVal;
+				OutPos = CurrSpline->GetLocationAtSplineInputKey(SplinePoint, ESplineCoordinateSpace::Local);
 				OutPos = CurrSpline->GetComponentTransform().TransformPosition(OutPos);
 
 				OutWidth = Metadata->RiverWidth.Points[SplinePoint].OutVal;
 				OutDepth = Metadata->Depth.Points[SplinePoint].OutVal;
 
-				OutTangent = CurrSpline->SplineCurves.Position.Points[SplinePoint].LeaveTangent;
+				OutTangent = CurrSpline->GetLeaveTangentAtSplinePoint(SplinePoint, ESplineCoordinateSpace::Local);
 
 				OutTangent = CurrSpline->GetComponentTransform().TransformVector(OutTangent);
 				OutTangent.Normalize();
diff --git a/Engine/Plugins/FX/Niagara/Source/Niagara/Private/NiagaraDataInterfaceSpline.cpp b/Engine/Plugins/FX/Niagara/Source/Niagara/Private/NiagaraDataInterfaceSpline.cpp
--- a/Engine/Plugins/FX/Niagara/Source/Niagara/Private/NiagaraDataInterfaceSpline.cpp
+++ b/Engine/Plugins/FX/Niagara/Source/Niagara/Private/NiagaraDataInterfaceSpline.cpp
@@ -660,9 +660,9 @@ bool UNiagaraDataInterfaceSpline::PerInstanceTick(void* PerInstanceData, FNiagar
 
 		bool bShouldBuildLUT = (bUseLUT || IsUsedWithGPUScript()) && InstData->SplineLUT.MaxIndex < 0;
 		
-		if (InstData->SplineCurvesVersion != SplineComponent->SplineCurves.Version)
+		if (InstData->SplineCurvesVersion != SplineComponent->GetVersion())
 		{
-			InstData->SplineCurves = SplineComponent->SplineCurves;
+			InstData->SplineCurves = SplineComponent->GetSplineCurves();
 			InstData->SplineCurvesVersion = InstData->SplineCurves.Version;
 			InstData->bSyncedGPUCopy = false;
 			InstData->SplineLUT.Reset();
diff --git a/Engine/Plugins/PCG/Source/PCG/Private/Data/PCGSplineStruct.cpp b/Engine/Plugins/PCG/Source/PCG/Private/Data/PCGSplineStruct.cpp
--- a/Engine/Plugins/PCG/Source/PCG/Private/Data/PCGSplineStruct.cpp
+++ b/Engine/Plugins/PCG/Source/PCG/Private/Data/PCGSplineStruct.cpp
@@ -81,7 +81,7 @@ namespace PCGSplineStruct
 void FPCGSplineStruct::Initialize(const USplineComponent* InSplineComponent)
 {
 	check(InSplineComponent);
-	SplineCurves = InSplineComponent->SplineCurves;
+	SplineCurves = InSplineComponent->GetSplineCurves();
 	Transform = InSplineComponent->GetComponentTransform();
 	DefaultUpVector = InSplineComponent->DefaultUpVector;
 	ReparamStepsPerSegment = InSplineComponent->ReparamStepsPerSegment;
@@ -125,7 +125,7 @@ void FPCGSplineStruct::ApplyTo(USplineComponent* InSplineComponent) const
 	InSplineComponent->DefaultUpVector = DefaultUpVector;
 	InSplineComponent->ReparamStepsPerSegment = ReparamStepsPerSegment;
 
-	InSplineComponent->SplineCurves = SplineCurves;
+	InSplineComponent->SetSpline(SplineCurves);
 	InSplineComponent->bStationaryEndpoints = false;
 	// TODO: metadata? might not be needed
 	InSplineComponent->SetClosedLoop(bClosedLoop);
diff --git a/Engine/Source/Editor/ComponentVisualizers/Private/SplineComponentVisualizer.cpp b/Engine/Source/Editor/ComponentVisualizers/Private/SplineComponentVisualizer.cpp
--- a/Engine/Source/Editor/ComponentVisualizers/Private/SplineComponentVisualizer.cpp
+++ b/Engine/Source/Editor/ComponentVisualizers/Private/SplineComponentVisualizer.cpp
@@ -99,6 +99,18 @@ namespace SplineComponentVisualizerLocals
 
 		return NearestVisible < TNumericLimits<double>::Max();
 	}
+
+	bool IsCurvePointType(EInterpCurveMode SplinePointType)
+	{
+		FInterpCurvePoint<float> Dummy;
+		Dummy.InterpMode = SplinePointType;
+		return Dummy.IsCurveKey();
+	}
+
+	bool IsCurvePointType(ESplinePointType::Type SplinePointType)
+	{
+		return IsCurvePointType(ConvertSplinePointTypeToInterpCurveMode(SplinePointType));
+	}
 }
 
 int32 USplineComponentVisualizerSelectionState::GetVerifiedLastKeyIndexSelected(const int32 InNumSplinePoints) const
@@ -330,7 +342,7 @@ FSplineComponentVisualizer::FSplineComponentVisualizer()
 
 	SplineComponentVisualizerActions = MakeShareable(new FUICommandList);
 
-	SplineCurvesProperty = FindFProperty<FProperty>(USplineComponent::StaticClass(), GET_MEMBER_NAME_CHECKED(USplineComponent, SplineCurves));
+	SplineCurvesProperty = FindFProperty<FProperty>(USplineComponent::StaticClass(), USplineComponent::GetSplinePropertyName());
 
 	SelectionState = NewObject<USplineComponentVisualizerSelectionState>(GetTransientPackage(), TEXT("SelectionState"), RF_Transactional);
 }
@@ -1142,7 +1154,8 @@ bool FSplineComponentVisualizer::GetWidgetLocation(const FEditorViewportClient*
 	{
 		check(SelectionState);
 
-		const FInterpCurveVector& Position = SplineComp->GetSplinePointsPosition();
+		const int32 NumPoints = SplineComp->GetNumberOfSplinePoints();
+		
 		const TSet<int32>& SelectedKeys = SelectionState->GetSelectedKeys();
 		int32 LastKeyIndexSelected = SelectionState->GetLastKeyIndexSelected();
 
@@ -1151,19 +1164,21 @@ bool FSplineComponentVisualizer::GetWidgetLocation(const FEditorViewportClient*
 		if (SelectedTangentHandle != INDEX_NONE)
 		{
 			// If tangent handle index is set, use that
-			check(SelectedTangentHandle < Position.Points.Num());
-			const auto& Point = Position.Points[SelectedTangentHandle];
+			check(SelectedTangentHandle < NumPoints);
+			const auto& Point = SplineComp->GetLocationAtSplinePoint(SelectedTangentHandle, ESplineCoordinateSpace::Local);
 
 			check(SelectedTangentHandleType != ESelectedTangentHandle::None);
 			const float TangentScale = GetDefault<ULevelEditorViewportSettings>()->SplineTangentScale;
 
 			if (SelectedTangentHandleType == ESelectedTangentHandle::Leave)
 			{
-				OutLocation = SplineComp->GetComponentTransform().TransformPosition(Point.OutVal + Point.LeaveTangent * TangentScale);
+				const FVector Tangent = SplineComp->GetLeaveTangentAtSplinePoint(SelectedTangentHandle, ESplineCoordinateSpace::Local);
+				OutLocation = SplineComp->GetComponentTransform().TransformPosition(Point + Tangent * TangentScale);
 			}
 			else if (SelectedTangentHandleType == ESelectedTangentHandle::Arrive)
 			{
-				OutLocation = SplineComp->GetComponentTransform().TransformPosition(Point.OutVal - Point.ArriveTangent * TangentScale);
+				const FVector Tangent = SplineComp->GetArriveTangentAtSplinePoint(SelectedTangentHandle, ESplineCoordinateSpace::Local);
+				OutLocation = SplineComp->GetComponentTransform().TransformPosition(Point - Tangent * TangentScale);
 			}
 
 			return true;
@@ -1172,11 +1187,10 @@ bool FSplineComponentVisualizer::GetWidgetLocation(const FEditorViewportClient*
 		{
 			// Otherwise use the last key index set
 			check(LastKeyIndexSelected >= 0);
-			if (LastKeyIndexSelected < Position.Points.Num())
+			if (LastKeyIndexSelected < NumPoints)
 			{
 				check(SelectedKeys.Contains(LastKeyIndexSelected));
-				const FInterpCurvePointVector& Point = Position.Points[LastKeyIndexSelected];
-				OutLocation = SplineComp->GetComponentTransform().TransformPosition(Point.OutVal);
+				OutLocation = SplineComp->GetLocationAtSplinePoint(LastKeyIndexSelected, ESplineCoordinateSpace::World);
 				if (!DuplicateDelayAccumulatedDrag.IsZero())
 				{
 					OutLocation += DuplicateDelayAccumulatedDrag;
@@ -1359,14 +1373,10 @@ bool FSplineComponentVisualizer::TransformSelectedTangent(EPropertyChangeType::T
 	USplineComponent* SplineComp = GetEditedSplineComponent();
 	if (SplineComp != nullptr)
 	{
-		FInterpCurveVector& SplinePosition = SplineComp->GetSplinePointsPosition();
-
-		const int32 NumPoints = SplinePosition.Points.Num();
-
 		check(SelectionState);
 		int32 SelectedTangentHandle;
 		ESelectedTangentHandle SelectedTangentHandleType;
-		SelectionState->GetVerifiedSelectedTangentHandle(NumPoints, SelectedTangentHandle, SelectedTangentHandleType);
+		SelectionState->GetVerifiedSelectedTangentHandle(SplineComp->GetNumberOfSplinePoints(), SelectedTangentHandle, SelectedTangentHandleType);
 
 		if (!InDeltaTranslate.IsZero())
 		{
@@ -1374,28 +1384,27 @@ bool FSplineComponentVisualizer::TransformSelectedTangent(EPropertyChangeType::T
 
 			const float TangentScale = GetDefault<ULevelEditorViewportSettings>()->SplineTangentScale;
 
-			FInterpCurvePoint<FVector>& EditedPoint = SplinePosition.Points[SelectedTangentHandle];
 			if (SplineComp->bAllowDiscontinuousSpline)
 			{
 				if (SelectedTangentHandleType == ESelectedTangentHandle::Leave)
 				{
-					EditedPoint.LeaveTangent += SplineComp->GetComponentTransform().InverseTransformVector(InDeltaTranslate) / TangentScale;
+					const FVector ArriveTangent = SplineComp->GetArriveTangentAtSplinePoint(SelectedTangentHandle, ESplineCoordinateSpace::Local);
+					const FVector LeaveTangent = SplineComp->GetLeaveTangentAtSplinePoint(SelectedTangentHandle, ESplineCoordinateSpace::Local) + SplineComp->GetComponentTransform().InverseTransformVector(InDeltaTranslate) / TangentScale;
+					SplineComp->SetTangentsAtSplinePoint(SelectedTangentHandle, ArriveTangent, LeaveTangent, ESplineCoordinateSpace::Local, false);
 				}
 				else
 				{
-					EditedPoint.ArriveTangent += SplineComp->GetComponentTransform().InverseTransformVector(-InDeltaTranslate) / TangentScale;
+					const FVector ArriveTangent = SplineComp->GetArriveTangentAtSplinePoint(SelectedTangentHandle, ESplineCoordinateSpace::Local) + SplineComp->GetComponentTransform().InverseTransformVector(-InDeltaTranslate) / TangentScale;
+					const FVector LeaveTangent = SplineComp->GetLeaveTangentAtSplinePoint(SelectedTangentHandle, ESplineCoordinateSpace::Local);
+					SplineComp->SetTangentsAtSplinePoint(SelectedTangentHandle, ArriveTangent, LeaveTangent, ESplineCoordinateSpace::Local, false);
 				}
 			}
 			else
 			{
 				const FVector Delta = (SelectedTangentHandleType == ESelectedTangentHandle::Leave) ? InDeltaTranslate : -InDeltaTranslate;
-				const FVector Tangent = EditedPoint.LeaveTangent + SplineComp->GetComponentTransform().InverseTransformVector(Delta) / TangentScale;
-
-				EditedPoint.LeaveTangent = Tangent;
-				EditedPoint.ArriveTangent = Tangent;
+				const FVector Tangent = SplineComp->GetLeaveTangentAtSplinePoint(SelectedTangentHandle, ESplineCoordinateSpace::Local) + SplineComp->GetComponentTransform().InverseTransformVector(Delta) / TangentScale;
+				SplineComp->SetTangentAtSplinePoint(SelectedTangentHandle, Tangent, ESplineCoordinateSpace::Local, false);
 			}
-
-			EditedPoint.InterpMode = CIM_CurveUser;
 		}
 
 		SplineComp->UpdateSpline();
@@ -1414,11 +1423,7 @@ bool FSplineComponentVisualizer::TransformSelectedKeys(EPropertyChangeType::Type
 	USplineComponent* SplineComp = GetEditedSplineComponent();
 	if (SplineComp != nullptr)
 	{
-		FInterpCurveVector& SplinePosition = SplineComp->GetSplinePointsPosition();
-		FInterpCurveQuat& SplineRotation = SplineComp->GetSplinePointsRotation();
-		FInterpCurveVector& SplineScale = SplineComp->GetSplinePointsScale();
-
-		const int32 NumPoints = SplinePosition.Points.Num();
+		const int32 NumPoints = SplineComp->GetNumberOfSplinePoints();
 
 		check(SelectionState);
 		const TSet<int32>& SelectedKeys = SelectionState->GetSelectedKeys();
@@ -1433,61 +1438,42 @@ bool FSplineComponentVisualizer::TransformSelectedKeys(EPropertyChangeType::Type
 			check(SelectedKeyIndex >= 0); 
 			check(SelectedKeyIndex < NumPoints);
 
-			FInterpCurvePoint<FVector>& EditedPoint = SplinePosition.Points[SelectedKeyIndex];
-			FInterpCurvePoint<FQuat>& EditedRotPoint = SplineRotation.Points[SelectedKeyIndex];
-			FInterpCurvePoint<FVector>& EditedScalePoint = SplineScale.Points[SelectedKeyIndex];
-
-
 			if (!InDeltaTranslate.IsZero())
 			{
-				// Find key position in world space
-				const FVector CurrentWorldPos = SplineComp->GetComponentTransform().TransformPosition(EditedPoint.OutVal);
-				// Move in world space
-				const FVector NewWorldPos = CurrentWorldPos + InDeltaTranslate;
-
-				// Convert back to local space
-				EditedPoint.OutVal = SplineComp->GetComponentTransform().InverseTransformPosition(NewWorldPos);
+				const FVector NewPosition = SplineComp->GetLocationAtSplinePoint(SelectedKeyIndex, ESplineCoordinateSpace::World) + InDeltaTranslate;
+				SplineComp->SetLocationAtSplinePoint(SelectedKeyIndex, NewPosition, ESplineCoordinateSpace::World, false);
 			}
 
 			if (!InDeltaRotate.IsZero())
 			{
-				// Set point tangent as user controlled
-				EditedPoint.InterpMode = CIM_CurveUser;
-
-				// Rotate tangent according to delta rotation
-				FVector NewTangent = SplineComp->GetComponentTransform().GetRotation().RotateVector(EditedPoint.LeaveTangent); // convert local-space tangent vector to world-space
-				NewTangent = InDeltaRotate.RotateVector(NewTangent); // apply world-space delta rotation to world-space tangent
-				NewTangent = SplineComp->GetComponentTransform().GetRotation().Inverse().RotateVector(NewTangent); // convert world-space tangent vector back into local-space
-				EditedPoint.LeaveTangent = NewTangent;
-				EditedPoint.ArriveTangent = NewTangent;
+				// note [nickolas.drake]: removed tangent setting code here because we set tangent in the SetRotation call below.
 
 				// Rotate spline rotation according to delta rotation
-				FQuat NewRot = SplineComp->GetComponentTransform().GetRotation() * EditedRotPoint.OutVal; // convert local-space rotation to world-space
+				FQuat NewRot = SplineComp->GetQuaternionAtSplinePoint(SelectedKeyIndex, ESplineCoordinateSpace::World);
 				NewRot = InDeltaRotate.Quaternion() * NewRot; // apply world-space rotation
-				NewRot = SplineComp->GetComponentTransform().GetRotation().Inverse() * NewRot; // convert world-space rotation to local-space
-				EditedRotPoint.OutVal = NewRot;
+				SplineComp->SetQuaternionAtSplinePoint(SelectedKeyIndex, NewRot, ESplineCoordinateSpace::World, false);
 			}
 
 			if (InDeltaScale.X != 0.0f)
 			{
-				// Set point tangent as user controlled
-				EditedPoint.InterpMode = CIM_CurveUser;
-
-				const FVector NewTangent = EditedPoint.LeaveTangent * (1.0f + InDeltaScale.X);
-				EditedPoint.LeaveTangent = NewTangent;
-				EditedPoint.ArriveTangent = NewTangent;
+				const FVector NewTangent = SplineComp->GetLeaveTangentAtSplinePoint(SelectedKeyIndex, ESplineCoordinateSpace::Local) + InDeltaScale.X;
+				SplineComp->SetTangentsAtSplinePoint(SelectedKeyIndex, NewTangent, NewTangent, ESplineCoordinateSpace::Local, false);
 			}
 
 			if (InDeltaScale.Y != 0.0f)
 			{
 				// Scale in Y adjusts the scale spline
-				EditedScalePoint.OutVal.Y *= (1.0f + InDeltaScale.Y);
+				FVector NewScale = SplineComp->GetScaleAtSplinePoint(SelectedKeyIndex);
+				NewScale.Y += InDeltaScale.Y;
+				SplineComp->SetScaleAtSplinePoint(SelectedKeyIndex, NewScale, false);
 			}
 
 			if (InDeltaScale.Z != 0.0f)
 			{
 				// Scale in Z adjusts the scale spline
-				EditedScalePoint.OutVal.Z *= (1.0f + InDeltaScale.Z);
+				FVector NewScale = SplineComp->GetScaleAtSplinePoint(SelectedKeyIndex);
+				NewScale.Z += InDeltaScale.Z;
+				SplineComp->SetScaleAtSplinePoint(SelectedKeyIndex, NewScale, false);
 			}
 		}
 
@@ -1649,8 +1635,7 @@ bool FSplineComponentVisualizer::HandleBoxSelect(const FBox& InBox, FEditorViewp
 		bool bAppendToSelection = InViewportClient->IsShiftPressed();
 
 		const TSet<int32>& SelectedKeys = SelectionState->GetSelectedKeys();
-		const FInterpCurveVector& SplineInfo = SplineComp->GetSplinePointsPosition();
-		int32 NumPoints = SplineInfo.Points.Num();
+		int32 NumPoints = SplineComp->GetNumberOfSplinePoints();
 
 		// Spline control point selection always uses transparent box selection.
 		for (int32 KeyIdx = 0; KeyIdx < NumPoints; KeyIdx++)
@@ -1694,8 +1679,7 @@ bool FSplineComponentVisualizer::HandleFrustumSelect(const FConvexVolume& InFrus
 		bool bAppendToSelection = InViewportClient->IsShiftPressed();
 
 		const TSet<int32>& SelectedKeys = SelectionState->GetSelectedKeys();
-		const FInterpCurveVector& SplineInfo = SplineComp->GetSplinePointsPosition();
-		int32 NumPoints = SplineInfo.Points.Num();
+		int32 NumPoints = SplineComp->GetNumberOfSplinePoints();
 
 		// Spline control point selection always uses transparent box selection.
 		for (int32 KeyIdx = 0; KeyIdx < NumPoints; KeyIdx++)
@@ -1779,9 +1763,7 @@ bool FSplineComponentVisualizer::HandleSnapTo(const bool bInAlign, const bool bI
 
 			SplineComp->Modify();
 
-			FInterpCurveVector& SplinePosition = SplineComp->GetSplinePointsPosition();
-			FInterpCurveQuat& SplineRotation = SplineComp->GetSplinePointsRotation();
-			int32 NumPoints = SplinePosition.Points.Num();
+			int32 NumPoints = SplineComp->GetNumberOfSplinePoints();
 
 			bool bMovedKey = false;
 
@@ -1793,41 +1775,34 @@ bool FSplineComponentVisualizer::HandleSnapTo(const bool bInAlign, const bool bI
 
 				FVector Direction = FVector(0.f, 0.f, -1.f);
 
-				FInterpCurvePoint<FVector>& EditedPoint = SplinePosition.Points[KeyIdx];
-				FInterpCurvePoint<FQuat>& EditedRotPoint = SplineRotation.Points[KeyIdx];
+				//FInterpCurvePoint<FVector>& EditedPoint = SplinePosition.Points[KeyIdx];
+				//FInterpCurvePoint<FQuat>& EditedRotPoint = SplineRotation.Points[KeyIdx];
 
 				FHitResult Hit(1.0f);
 				FCollisionQueryParams Params(SCENE_QUERY_STAT(MoveSplineKeyToTrace), true);
 
 				// Find key position in world space
-				const FVector CurrentWorldPos = SplineComp->GetComponentTransform().TransformPosition(EditedPoint.OutVal);
+				const FVector CurrentWorldPos = SplineComp->GetLocationAtSplinePoint(KeyIdx, ESplineCoordinateSpace::World);
 
 				if (SplineComp->GetWorld()->LineTraceSingleByChannel(Hit, CurrentWorldPos, CurrentWorldPos + Direction * WORLD_MAX, ECC_WorldStatic, Params))
 				{
-					// Convert back to local space
-					EditedPoint.OutVal = SplineComp->GetComponentTransform().InverseTransformPosition(Hit.Location);
+					SplineComp->SetLocationAtSplinePoint(KeyIdx, Hit.Location, ESplineCoordinateSpace::World, false);
 
 					if (bInAlign)
 					{		
-						// Set point tangent as user controlled
-						EditedPoint.InterpMode = CIM_CurveUser;
-
 						// Get delta rotation between up vector and hit normal
 						FVector WorldUpVector = SplineComp->GetUpVectorAtSplineInputKey((float)KeyIdx, ESplineCoordinateSpace::World);
 						FQuat DeltaRotate = FQuat::FindBetweenNormals(WorldUpVector, Hit.Normal);
 
 						// Rotate tangent according to delta rotation
-						FVector NewTangent = SplineComp->GetComponentTransform().GetRotation().RotateVector(EditedPoint.LeaveTangent); // convert local-space tangent vector to world-space
+						FVector NewTangent = SplineComp->GetLeaveTangentAtSplinePoint(KeyIdx, ESplineCoordinateSpace::World);
 						NewTangent = DeltaRotate.RotateVector(NewTangent); // apply world-space delta rotation to world-space tangent
-						NewTangent = SplineComp->GetComponentTransform().GetRotation().Inverse().RotateVector(NewTangent); // convert world-space tangent vector back into local-space
-						EditedPoint.LeaveTangent = NewTangent;
-						EditedPoint.ArriveTangent = NewTangent;
+						SplineComp->SetTangentAtSplinePoint(KeyIdx, NewTangent, ESplineCoordinateSpace::World, false);
 
 						// Rotate spline rotation according to delta rotation
-						FQuat NewRot = SplineComp->GetComponentTransform().GetRotation() * EditedRotPoint.OutVal; // convert local-space rotation to world-space
+						FQuat NewRot = SplineComp->GetRotationAtSplinePoint(KeyIdx, ESplineCoordinateSpace::World).Quaternion();
 						NewRot = DeltaRotate * NewRot; // apply world-space rotation
-						NewRot = SplineComp->GetComponentTransform().GetRotation().Inverse() * NewRot; // convert world-space rotation to local-space
-						EditedRotPoint.OutVal = NewRot;
+						SplineComp->SetRotationAtSplinePoint(KeyIdx, NewRot.Rotator(), ESplineCoordinateSpace::World, false);
 					}
 
 					bMovedKey = true;
@@ -1892,8 +1867,7 @@ void FSplineComponentVisualizer::OnSnapKeyToNearestSplinePoint(ESplineComponentS
 	check(SelectedKeys.Num() == 1);
 	check(SelectedKeys.Contains(LastKeyIndexSelected));
 
-	FInterpCurvePoint<FVector>& EditedPosition = SplineComp->GetSplinePointsPosition().Points[LastKeyIndexSelected];
-	const FVector WorldPos = SplineComp->GetComponentTransform().TransformPosition(EditedPosition.OutVal); // convert local-space position to world-space
+	const FVector WorldPos = SplineComp->GetLocationAtSplinePoint(LastKeyIndexSelected, ESplineCoordinateSpace::World);
 
 	double NearestDistanceSquared = 0.0f;
 	USplineComponent* NearestSplineComp = nullptr;
@@ -1917,8 +1891,7 @@ void FSplineComponentVisualizer::OnSnapKeyToNearestSplinePoint(ESplineComponentS
 
 	{
 		// Test non-adjacent points on current spline.
-		const FInterpCurveVector& SplineInfo = SplineComp->GetSplinePointsPosition();
-		const int32 NumPoints = SplineInfo.Points.Num();
+		const int32 NumPoints = SplineComp->GetNumberOfSplinePoints();
 
 		// Don't test against current or adjacent points
 		TSet<int32> IgnoreIndices;
@@ -1974,8 +1947,7 @@ void FSplineComponentVisualizer::OnSnapKeyToNearestSplinePoint(ESplineComponentS
 
 			if (TestComponentBoundingBox.IsInsideOrOn(WorldPos))
 			{
-				const FInterpCurveVector& SplineInfo = TestComponent->GetSplinePointsPosition();
-				const int32 NumPoints = SplineInfo.Points.Num();
+				const int32 NumPoints = SplineComp->GetNumberOfSplinePoints();
 				for (int32 KeyIdx = 0; KeyIdx < NumPoints; KeyIdx++)
 				{
 					UpdateNearestKey(TestComponent, KeyIdx);
@@ -1990,10 +1962,8 @@ void FSplineComponentVisualizer::OnSnapKeyToNearestSplinePoint(ESplineComponentS
 		return;
 	}
 
-	const FInterpCurvePoint<FVector>& NearestPosition = NearestSplineComp->GetSplinePointsPosition().Points[NearestKeyIndex];
-
 	// Copy position
-	const FVector NearestWorldPos = NearestSplineComp->GetComponentTransform().TransformPosition(NearestPosition.OutVal); // convert local-space position to world-space
+	const FVector NearestWorldPos = SplineComp->GetLocationAtSplinePoint(NearestKeyIndex, ESplineCoordinateSpace::World);
 	FVector NearestWorldUpVector(0.0f, 0.0f, 1.0f);
 	FVector NearestWorldTangent(0.0f, 1.0f, 0.0f);
 	FVector NearestWorldScale(1.0f, 1.0f, 1.0f);
@@ -2002,7 +1972,7 @@ void FSplineComponentVisualizer::OnSnapKeyToNearestSplinePoint(ESplineComponentS
 	if (InSnapMode == ESplineComponentSnapMode::AlignToTangent || InSnapMode == ESplineComponentSnapMode::AlignPerpendicularToTangent)
 	{
 		// Get tangent
-		NearestWorldTangent = NearestSplineComp->GetComponentTransform().GetRotation().RotateVector(NearestPosition.ArriveTangent); // convert local-space tangent vectors to world-space
+		NearestWorldTangent = SplineComp->GetArriveTangentAtSplinePoint(NearestKeyIndex, ESplineCoordinateSpace::World); // convert local-space tangent vectors to world-space
 
 		// Get up vector
 		NearestWorldUpVector = NearestSplineComp->GetUpVectorAtSplinePoint(NearestKeyIndex, ESplineCoordinateSpace::World);
@@ -2010,8 +1980,8 @@ void FSplineComponentVisualizer::OnSnapKeyToNearestSplinePoint(ESplineComponentS
 		// Get scale, only when aligning parallel
 		if (InSnapMode == ESplineComponentSnapMode::AlignToTangent)
 		{
-			const FInterpCurvePoint<FVector>& NearestScale = NearestSplineComp->GetSplinePointsScale().Points[NearestKeyIndex];
-			NearestWorldScale = SplineComp->GetComponentTransform().GetScale3D() * NearestScale.OutVal; // convert local-space rotation to world-space
+			const FVector NearestScale = SplineComp->GetScaleAtSplinePoint(NearestKeyIndex);
+			NearestWorldScale = SplineComp->GetComponentTransform().GetScale3D() * NearestScale; // convert local-space rotation to world-space
 		}
 
 		// Get metadata (only when aligning)
@@ -2065,12 +2035,8 @@ void FSplineComponentVisualizer::SnapKeyToTransform(const ESplineComponentSnapMo
 		Owner->Modify();
 	}
 
-	FInterpCurvePoint<FVector>& EditedPosition = SplineComp->GetSplinePointsPosition().Points[LastKeyIndexSelected];
-	FInterpCurvePoint<FQuat>& EditedRotation = SplineComp->GetSplinePointsRotation().Points[LastKeyIndexSelected];
-	FInterpCurvePoint<FVector>& EditedScale = SplineComp->GetSplinePointsScale().Points[LastKeyIndexSelected];
-
 	// Copy position
-	EditedPosition.OutVal = SplineComp->GetComponentTransform().InverseTransformPosition(InWorldPos); // convert world-space position to local-space
+	SplineComp->SetLocationAtSplinePoint(LastKeyIndexSelected, InWorldPos, ESplineCoordinateSpace::World, false);
 
 	if (InSnapMode == ESplineComponentSnapMode::AlignToTangent || InSnapMode == ESplineComponentSnapMode::AlignPerpendicularToTangent)
 	{
@@ -2088,7 +2054,7 @@ void FSplineComponentVisualizer::SnapKeyToTransform(const ESplineComponentSnapMo
 			NewTangent = DeltaRotate.RotateVector(NewTangent);
 		}
 
-		const FVector Tangent = SplineComp->GetComponentTransform().GetRotation().RotateVector(EditedPosition.ArriveTangent); // convert local-space tangent vectors to world-space
+		const FVector Tangent = SplineComp->GetArriveTangentAtSplinePoint(LastKeyIndexSelected, ESplineCoordinateSpace::World);
 
 		// Swap the tangents if they are not pointing in the same general direction
 		double CurrentAngle = FMath::Acos(FVector::DotProduct(Tangent, NewTangent) / Tangent.Size());
@@ -2102,21 +2068,21 @@ void FSplineComponentVisualizer::SnapKeyToTransform(const ESplineComponentSnapMo
 		}
 
 		// Update tangent
-		EditedPosition.ArriveTangent = NewTangent;
-		EditedPosition.LeaveTangent = NewTangent;
-		EditedPosition.InterpMode = CIM_CurveUser;
+		SplineComp->SetTangentAtSplinePoint(LastKeyIndexSelected, NewTangent, ESplineCoordinateSpace::Local, false);
 
 		// Copy rotation, it is only used to determine up vector so no need to adjust it 
 		const FQuat Rot = FQuat::FindBetweenNormals(FVector(0.0f, 0.0f, 1.0f), WorldUpVector);
-		EditedRotation.OutVal = SplineComp->GetComponentTransform().GetRotation().Inverse() * Rot; // convert world-space rotation to local-space
+		SplineComp->SetRotationAtSplinePoint(LastKeyIndexSelected, Rot.Rotator(), ESplineCoordinateSpace::World, false);
 
 		// Copy scale, only when aligning parallel
 		if (InSnapMode == ESplineComponentSnapMode::AlignToTangent)
 		{
 			const FVector SplineCompScale = SplineComp->GetComponentTransform().GetScale3D();
-			EditedScale.OutVal.X = FMath::IsNearlyZero(SplineCompScale.X) ? InScale.X : InScale.X / SplineCompScale.X; // convert world-space scale to local-space
-			EditedScale.OutVal.Y = FMath::IsNearlyZero(SplineCompScale.Y) ? InScale.Y : InScale.Y / SplineCompScale.Y;
-			EditedScale.OutVal.Z = FMath::IsNearlyZero(SplineCompScale.Z) ? InScale.Z : InScale.Z / SplineCompScale.Z;
+			FVector NewScale;
+			NewScale.X = FMath::IsNearlyZero(SplineCompScale.X) ? InScale.X : InScale.X / SplineCompScale.X;
+			NewScale.Y = FMath::IsNearlyZero(SplineCompScale.Y) ? InScale.Y : InScale.Y / SplineCompScale.Y;
+			NewScale.Z = FMath::IsNearlyZero(SplineCompScale.Z) ? InScale.Z : InScale.Z / SplineCompScale.Z;
+			SplineComp->SetScaleAtSplinePoint(LastKeyIndexSelected, NewScale, false);
 		}
 
 	}
@@ -2290,22 +2256,14 @@ void FSplineComponentVisualizer::SnapKeysToLastSelectedAxisPosition(const EAxis:
 		Owner->Modify();
 	}
 
-	FInterpCurveVector& SplinePositions = SplineComp->GetSplinePointsPosition();
-	FInterpCurveQuat& SplineRotations = SplineComp->GetSplinePointsRotation();
-
-	const FVector WorldPos = SplineComp->GetComponentTransform().TransformPosition(SplinePositions.Points[LastKeyIndexSelected].OutVal); 
-
-	int32 NumPoints = SplinePositions.Points.Num();
+	const FVector WorldPos = SplineComp->GetLocationAtSplinePoint(LastKeyIndexSelected, ESplineCoordinateSpace::World);
 
 	for (int32 KeyIdx : InSnapKeys)
 	{
 		if (KeyIdx >= 0 && KeyIdx < SplineComp->GetNumberOfSplinePoints())
 		{
-			FInterpCurvePoint<FVector>& EditedPosition = SplinePositions.Points[KeyIdx];
-			FInterpCurvePoint<FQuat>& EditedRotation = SplineRotations.Points[KeyIdx];
-
 			// Copy position
-			FVector NewWorldPos = SplineComp->GetComponentTransform().TransformPosition(EditedPosition.OutVal); // convert local-space position to world-space
+			FVector NewWorldPos = SplineComp->GetLocationAtSplinePoint(KeyIdx, ESplineCoordinateSpace::World);
 			if (InAxis == EAxis::X)
 			{
 				NewWorldPos.X = WorldPos.X;
@@ -2319,12 +2277,12 @@ void FSplineComponentVisualizer::SnapKeysToLastSelectedAxisPosition(const EAxis:
 				NewWorldPos.Z = WorldPos.Z;
 			}
 
-			EditedPosition.OutVal = SplineComp->GetComponentTransform().InverseTransformPosition(NewWorldPos); // convert world-space position to local-space
+			SplineComp->SetLocationAtSplinePoint(KeyIdx, NewWorldPos, ESplineCoordinateSpace::World);
 
 			// Set point to auto so its tangents will be auto-adjusted after snapping
 			if (SplineComp->bAdjustTangentsOnSnap)
 			{
-				EditedPosition.InterpMode = CIM_CurveAuto;
+				SplineComp->SetSplinePointType(KeyIdx, ConvertInterpCurveModeToSplinePointType(CIM_CurveAuto));
 			}
 		}
 	}
@@ -2410,34 +2368,12 @@ void FSplineComponentVisualizer::OnDuplicateKey()
 	SelectedKeysSorted.Sort([](int32 A, int32 B) { return A > B; });
 
 	// Insert duplicates into the list, highest index first, so that the lower indices remain the same
-	FInterpCurveVector& SplinePosition = SplineComp->GetSplinePointsPosition();
-	FInterpCurveQuat& SplineRotation = SplineComp->GetSplinePointsRotation();
-	FInterpCurveVector& SplineScale = SplineComp->GetSplinePointsScale();
-	USplineMetadata* SplineMetadata = SplineComp->GetSplinePointsMetadata();
-
 	for (int32 SelectedKeyIndex : SelectedKeysSorted)
 	{
 		check(SelectedKeyIndex >= 0);
 		check(SelectedKeyIndex < SplineComp->GetNumberOfSplinePoints());
 
-		// Insert duplicates into arrays.
-		// It's necessary to take a copy because copying existing array items by reference isn't allowed (the array may reallocate)
-		SplinePosition.Points.Insert(FInterpCurvePoint<FVector>(SplinePosition.Points[SelectedKeyIndex]), SelectedKeyIndex);
-		SplineRotation.Points.Insert(FInterpCurvePoint<FQuat>(SplineRotation.Points[SelectedKeyIndex]), SelectedKeyIndex);
-		SplineScale.Points.Insert(FInterpCurvePoint<FVector>(SplineScale.Points[SelectedKeyIndex]), SelectedKeyIndex);
-
-		if (SplineMetadata)
-		{
-			SplineMetadata->DuplicatePoint(SelectedKeyIndex);
-		}
-
-		// Adjust input keys of subsequent points
-		for (int Index = SelectedKeyIndex + 1; Index < SplinePosition.Points.Num(); Index++)
-		{
-			SplinePosition.Points[Index].InVal += 1.0f;
-			SplineRotation.Points[Index].InVal += 1.0f;
-			SplineScale.Points[Index].InVal += 1.0f;
-		}
+		SplineComp->AddPoint(SplineComp->GetSplinePointAt(SelectedKeyIndex, ESplineCoordinateSpace::Local), false);
 	}
 
 	SelectionState->Modify();
@@ -2528,11 +2464,11 @@ bool FSplineComponentVisualizer::DuplicateKeyForAltDrag(const FVector& InDrag)
 	static const double Angle60 = 1.0472;
 
 	// Insert duplicates into the list, highest index first, so that the lower indices remain the same
-	FInterpCurveVector& SplinePosition = SplineComp->GetSplinePointsPosition();
+	//FInterpCurveVector& SplinePosition = SplineComp->GetSplinePointsPosition();
 
 	// Find key position in world space
 	int32 CurrentIndex = LastKeyIndexSelected;
-	const FVector CurrentKeyWorldPos = SplineComp->GetComponentTransform().TransformPosition(SplinePosition.Points[CurrentIndex].OutVal);
+	const FVector CurrentKeyWorldPos = SplineComp->GetLocationAtSplinePoint(CurrentIndex, ESplineCoordinateSpace::World);
 
 	// Determine direction to insert new point				
 	bool bHasPrevKey = SplineComp->IsClosedLoop() || CurrentIndex > 0;
@@ -2541,7 +2477,7 @@ bool FSplineComponentVisualizer::DuplicateKeyForAltDrag(const FVector& InDrag)
 	{
 		// Wrap index around for closed-looped splines
 		int32 PrevKeyIndex = (CurrentIndex > 0 ? CurrentIndex - 1 : NumPoints - 1);
-		FVector PrevKeyWorldPos = SplineComp->GetComponentTransform().TransformPosition(SplinePosition.Points[PrevKeyIndex].OutVal);
+		FVector PrevKeyWorldPos = SplineComp->GetLocationAtSplinePoint(PrevKeyIndex, ESplineCoordinateSpace::World);
 		FVector SegmentDirection = PrevKeyWorldPos - CurrentKeyWorldPos;
 		if (!SegmentDirection.IsZero())
 		{
@@ -2559,7 +2495,7 @@ bool FSplineComponentVisualizer::DuplicateKeyForAltDrag(const FVector& InDrag)
 	{
 		// Wrap index around for closed-looped splines
 		int32 NextKeyIndex = (CurrentIndex + 1 < NumPoints ? CurrentIndex + 1 : 0);
-		FVector NextKeyWorldPos = SplineComp->GetComponentTransform().TransformPosition(SplinePosition.Points[NextKeyIndex].OutVal);
+		FVector NextKeyWorldPos = SplineComp->GetLocationAtSplinePoint(NextKeyIndex, ESplineCoordinateSpace::World);
 		FVector SegmentDirection = NextKeyWorldPos - CurrentKeyWorldPos;
 		if (!SegmentDirection.IsZero())
 		{
@@ -2632,11 +2568,10 @@ float FSplineComponentVisualizer::FindNearest(const FVector& InLocalPos, int32 I
 	check(InSegmentIndex >= 0);
 	check(InSegmentIndex < SplineComp->GetNumberOfSplineSegments());
 
-	FInterpCurveVector& SplinePosition = SplineComp->GetSplinePointsPosition();
-	float OutSquaredDistance = 0.0f;
-	float t = SplinePosition.FindNearestOnSegment(InLocalPos, InSegmentIndex, OutSquaredDistance);
-	OutSplinePos = SplinePosition.Eval(t, FVector::ZeroVector);
-	OutSplineTangent = SplinePosition.EvalDerivative(t, FVector::ZeroVector);
+	FVector WorldPos = SplineComp->GetComponentTransform().InverseTransformPosition(InLocalPos);
+	float t = SplineComp->FindInputKeyOnSegmentClosestToWorldLocation(WorldPos, InSegmentIndex);
+	OutSplinePos = SplineComp->GetLocationAtSplineInputKey(t, ESplineCoordinateSpace::Local);
+	OutSplineTangent = SplineComp->GetTangentAtSplineInputKey(t, ESplineCoordinateSpace::Local);
 
 	return t;
 }
@@ -2679,84 +2614,36 @@ void FSplineComponentVisualizer::SplitSegment(const FVector& InWorldPos, int32 I
 	{
 		SegmentEndIndex = 0;
 	}
-
-	FInterpCurveVector& SplinePosition = SplineComp->GetSplinePointsPosition();
-	FInterpCurveQuat& SplineRotation = SplineComp->GetSplinePointsRotation();
-	FInterpCurveVector& SplineScale = SplineComp->GetSplinePointsScale();
-	USplineMetadata* SplineMetadata = SplineComp->GetSplinePointsMetadata();
 	
 	// Set adjacent points to CurveAuto so their tangents adjust automatically as new point moves.
-	if (SplinePosition.Points[SegmentBeginIndex].InterpMode == CIM_CurveUser)
+	if (SplineComp->GetSplinePointType(SegmentBeginIndex) == ConvertInterpCurveModeToSplinePointType(CIM_CurveUser))
 	{
-		SplinePosition.Points[SegmentBeginIndex].InterpMode = CIM_CurveAuto;
+		SplineComp->SetSplinePointType(SegmentBeginIndex, ConvertInterpCurveModeToSplinePointType(CIM_CurveAuto), false);
 	}
-	if (SplinePosition.Points[SegmentEndIndex].InterpMode == CIM_CurveUser)
+	if (SplineComp->GetSplinePointType(SegmentEndIndex) == ConvertInterpCurveModeToSplinePointType(CIM_CurveUser))
 	{
-		SplinePosition.Points[SegmentEndIndex].InterpMode = CIM_CurveAuto;
+		SplineComp->SetSplinePointType(SegmentEndIndex, ConvertInterpCurveModeToSplinePointType(CIM_CurveAuto), false);
 	}
 
 	// Compute interpolated scale
 	FVector NewScale;
-	FInterpCurvePoint<FVector>& PrevScale = SplineScale.Points[SegmentBeginIndex];
-	FInterpCurvePoint<FVector>& NextScale = SplineScale.Points[SegmentEndIndex];
-	NewScale = FMath::LerpStable(PrevScale.OutVal, NextScale.OutVal, t);
+	FVector PrevScale = SplineComp->GetScaleAtSplinePoint(SegmentBeginIndex);
+	FVector NextScale = SplineComp->GetScaleAtSplinePoint(SegmentEndIndex);
+	NewScale = FMath::LerpStable(PrevScale, NextScale, t);
 
 	// Compute interpolated rot
 	FQuat NewRot;
-	FInterpCurvePoint<FQuat>& PrevRot = SplineRotation.Points[SegmentBeginIndex];
-	FInterpCurvePoint<FQuat>& NextRot = SplineRotation.Points[SegmentEndIndex];
-	NewRot = FMath::Lerp(PrevRot.OutVal, NextRot.OutVal, t);
+	FQuat PrevRot = SplineComp->GetRotationAtSplinePoint(SegmentBeginIndex, ESplineCoordinateSpace::Local).Quaternion();
+	FQuat NextRot = SplineComp->GetRotationAtSplinePoint(SegmentEndIndex, ESplineCoordinateSpace::Local).Quaternion();
+	NewRot = FMath::Lerp(PrevRot, NextRot, t);
 
 	// Determine which index to use when copying interp mode
 	int32 SourceIndex = bCopyFromSegmentBeginIndex ? SegmentBeginIndex : SegmentEndIndex;
 
-	FInterpCurvePoint<FVector> NewPoint(
-		(float)SegmentSplitIndex,
-		SplinePos,
-		FVector::ZeroVector,
-		FVector::ZeroVector,
-		SplinePosition.Points[SourceIndex].InterpMode);
-
-	FInterpCurvePoint<FQuat> NewRotPoint(
-		(float)SegmentSplitIndex,
-		NewRot,
-		FQuat::Identity,
-		FQuat::Identity,
-		CIM_CurveAuto);
-
-	FInterpCurvePoint<FVector> NewScalePoint(
-		(float)SegmentSplitIndex,
-		NewScale,
-		FVector::ZeroVector,
-		FVector::ZeroVector,
-		CIM_CurveAuto);
-
-	if (SegmentEndIndex == 0)
-	{
-		// Splitting last segment of a closed-looped spline
-		SplinePosition.Points.Emplace(NewPoint);
-		SplineRotation.Points.Emplace(NewRotPoint);
-		SplineScale.Points.Emplace(NewScalePoint);
-	}
-	else
-	{
-		SplinePosition.Points.Insert(NewPoint, SegmentEndIndex);
-		SplineRotation.Points.Insert(NewRotPoint, SegmentEndIndex);
-		SplineScale.Points.Insert(NewScalePoint, SegmentEndIndex);
-	}
-
-	if (SplineMetadata)
-	{
-		SplineMetadata->InsertPoint(SegmentEndIndex, t, SplineComp->IsClosedLoop());
-	}
-
-	// Adjust input keys of subsequent points
-	for (int Index = SegmentSplitIndex + 1; Index < SplineComp->GetNumberOfSplinePoints(); Index++)
-	{
-		SplinePosition.Points[Index].InVal += 1.0f;
-		SplineRotation.Points[Index].InVal += 1.0f;
-		SplineScale.Points[Index].InVal += 1.0f;
-	}
+	SplineComp->AddSplinePointAtIndex(LocalPos, SegmentSplitIndex, ESplineCoordinateSpace::Local, false);
+	SplineComp->SetQuaternionAtSplinePoint(SegmentSplitIndex, NewRot, ESplineCoordinateSpace::Local, false);
+	SplineComp->SetScaleAtSplinePoint(SegmentSplitIndex, NewScale, false);
+	SplineComp->SetSplinePointType(SegmentSplitIndex, SplineComp->GetSplinePointType(SourceIndex), false);
 
 	// Set selection to new key
 	ChangeSelectionState(SegmentSplitIndex, false);
@@ -2807,14 +2694,8 @@ void FSplineComponentVisualizer::UpdateSplitSegment(const FVector& InDrag)
 		SegmentEndIndex = 0;
 	}
 
-	FInterpCurveVector& SplinePosition = SplineComp->GetSplinePointsPosition();
-	FInterpCurveVector& SplineScale = SplineComp->GetSplinePointsScale();
-	FInterpCurveQuat& SplineRotation = SplineComp->GetSplinePointsRotation();
-	USplineMetadata* SplineMetadata = SplineComp->GetSplinePointsMetadata();
-
 	// Find key position in world space
-	FInterpCurvePoint<FVector>& EditedPoint = SplinePosition.Points[SegmentSplitIndex];
-	const FVector CurrentWorldPos = SplineComp->GetComponentTransform().TransformPosition(EditedPoint.OutVal);
+	const FVector CurrentWorldPos = SplineComp->GetLocationAtSplinePoint(SegmentSplitIndex, ESplineCoordinateSpace::World);
 
 	// Move in world space
 	const FVector NewWorldPos = CurrentWorldPos + InDrag;
@@ -2842,22 +2723,20 @@ void FSplineComponentVisualizer::UpdateSplitSegment(const FVector& InDrag)
 	DuplicateCacheSplitSegmentParam = t;
 
 	// Update location
-	EditedPoint.OutVal = LocalPos;
+	SplineComp->SetLocationAtSplinePoint(SegmentStartIndex, LocalPos, ESplineCoordinateSpace::Local, false);
 
 	// Update scale
-	FInterpCurvePoint<FVector>& EditedScale = SplineScale.Points[SegmentSplitIndex];
-	FInterpCurvePoint<FVector>& PrevScale = SplineScale.Points[SegmentStartIndex];
-	FInterpCurvePoint<FVector>& NextScale = SplineScale.Points[SegmentEndIndex];
-	EditedScale.OutVal = FMath::LerpStable(PrevScale.OutVal, NextScale.OutVal, t);
+	const FVector PrevScale = SplineComp->GetScaleAtSplinePoint(SegmentStartIndex);
+	const FVector NextScale = SplineComp->GetScaleAtSplinePoint(SegmentEndIndex);
+	SplineComp->SetScaleAtSplinePoint(SegmentSplitIndex, FMath::LerpStable(PrevScale, NextScale, t), false);
 
 	// Update rot
-	FInterpCurvePoint<FQuat>& EditedRot = SplineRotation.Points[SegmentSplitIndex];
-	FInterpCurvePoint<FQuat>& PrevRot = SplineRotation.Points[SegmentStartIndex];
-	FInterpCurvePoint<FQuat>& NextRot = SplineRotation.Points[SegmentEndIndex];
-	EditedRot.OutVal = FMath::Lerp(PrevRot.OutVal, NextRot.OutVal, t);
+	FQuat PrevRot = SplineComp->GetRotationAtSplinePoint(SegmentStartIndex, ESplineCoordinateSpace::Local).Quaternion();
+	FQuat NextRot = SplineComp->GetRotationAtSplinePoint(SegmentEndIndex, ESplineCoordinateSpace::Local).Quaternion();
+	SplineComp->SetRotationAtSplinePoint(SegmentSplitIndex, FMath::Lerp(PrevRot, NextRot, t).Rotator(), ESplineCoordinateSpace::Local, false);
 
 	// Update metadata
-	if (SplineMetadata)
+	if (USplineMetadata* SplineMetadata = SplineComp->GetSplinePointsMetadata())
 	{
 		SplineMetadata->UpdatePoint(SegmentSplitIndex, t, SplineComp->IsClosedLoop());
 	}
@@ -2885,76 +2764,30 @@ void FSplineComponentVisualizer::AddSegment(const FVector& InWorldPos, bool bApp
 	int32 KeyIdx = 0;
 	int32 NewKeyIdx = 0;
 
-	FInterpCurveVector& SplinePosition = SplineComp->GetSplinePointsPosition();
-
 	if (bAppend)
 	{
-		NewKeyIdx = SplinePosition.Points.Num();
+		NewKeyIdx = SplineComp->GetNumberOfSplinePoints();
 		KeyIdx = NewKeyIdx - 1;
 	}
-
-	FInterpCurveQuat& SplineRotation = SplineComp->GetSplinePointsRotation();
-	FInterpCurveVector& SplineScale = SplineComp->GetSplinePointsScale();
-	USplineMetadata* SplineMetadata = SplineComp->GetSplinePointsMetadata();
 	
 	// Set adjacent point to CurveAuto so its tangent adjusts automatically as new point moves.
-	if (SplinePosition.Points[KeyIdx].InterpMode == CIM_CurveUser)
+	if (ConvertSplinePointTypeToInterpCurveMode(SplineComp->GetSplinePointType(KeyIdx)) == CIM_CurveUser)
 	{
-		SplinePosition.Points[KeyIdx].InterpMode = CIM_CurveAuto;
+		SplineComp->SetSplinePointType(KeyIdx, ConvertInterpCurveModeToSplinePointType(CIM_CurveAuto), false);
 	}
 
 	// Compute local pos
-	FVector LocalPos = SplineComp->GetComponentTransform().InverseTransformPosition(InWorldPos);
-
-	FInterpCurvePoint<FVector> NewPoint(
-		(float)NewKeyIdx,
-		LocalPos,
-		FVector::ZeroVector,
-		FVector::ZeroVector,
-		SplinePosition.Points[KeyIdx].InterpMode);
-
-	FInterpCurvePoint<FQuat> NewRotPoint(
-		(float)NewKeyIdx,
-		SplineRotation.Points[KeyIdx].OutVal,
-		FQuat::Identity,
-		FQuat::Identity,
-		CIM_CurveAuto);
-
-	FInterpCurvePoint<FVector> NewScalePoint(
-		(float)NewKeyIdx,
-		SplineScale.Points[KeyIdx].OutVal,
-		FVector::ZeroVector,
-		FVector::ZeroVector,
-		CIM_CurveAuto);
-
-	if (KeyIdx == 0)
-	{
-		SplinePosition.Points.Insert(NewPoint, KeyIdx);
-		SplineRotation.Points.Insert(NewRotPoint, KeyIdx);
-		SplineScale.Points.Insert(NewScalePoint, KeyIdx);
-	}
-	else
-	{
-		SplinePosition.Points.Emplace(NewPoint);
-		SplineRotation.Points.Emplace(NewRotPoint);
-		SplineScale.Points.Emplace(NewScalePoint);
-	}
+	const FVector LocalPos = SplineComp->GetComponentTransform().InverseTransformPosition(InWorldPos);
 
-	// Adjust input keys of subsequent points
-	if (!bAppend)
-	{
-		for (int Index = 1; Index < SplinePosition.Points.Num(); Index++)
-		{
-			SplinePosition.Points[Index].InVal += 1.0f;
-			SplineRotation.Points[Index].InVal += 1.0f;
-			SplineScale.Points[Index].InVal += 1.0f;
-		}
-	}
+	// Must be saved before adding point so that KeyIdx remains valid while we use it to read from the existing data.
+	FQuat NewRot = SplineComp->GetQuaternionAtSplinePoint(KeyIdx, ESplineCoordinateSpace::Local);
+	FVector NewScale = SplineComp->GetScaleAtSplinePoint(KeyIdx);
+	ESplinePointType::Type NewType = SplineComp->GetSplinePointType(KeyIdx);
 
-	if (SplineMetadata)
-	{
-		SplineMetadata->DuplicatePoint(KeyIdx);
-	}
+	SplineComp->AddSplinePointAtIndex(LocalPos, NewKeyIdx, ESplineCoordinateSpace::Local, false);
+	SplineComp->SetQuaternionAtSplinePoint(NewKeyIdx, NewRot, ESplineCoordinateSpace::Local, false);
+	SplineComp->SetScaleAtSplinePoint(NewKeyIdx, NewScale, false);
+	SplineComp->SetSplinePointType(NewKeyIdx, NewType, false);
 
 	// Set selection to key
 	ChangeSelectionState(NewKeyIdx, false);
@@ -2978,7 +2811,7 @@ void FSplineComponentVisualizer::UpdateAddSegment(const FVector& InDrag)
 	check(SelectedKeys.Contains(LastKeyIndexSelected));
 	// Only work on keys at either end of a non-closed-looped spline 
 	check(!SplineComp->IsClosedLoop());
-	check(LastKeyIndexSelected == 0 || LastKeyIndexSelected == SplineComp->GetSplinePointsPosition().Points.Num() - 1);
+	check(LastKeyIndexSelected == 0 || LastKeyIndexSelected == SplineComp->GetNumberOfSplinePoints() - 1);
 
 	SplineComp->Modify();
 	if (AActor* Owner = SplineComp->GetOwner())
@@ -2987,10 +2820,9 @@ void FSplineComponentVisualizer::UpdateAddSegment(const FVector& InDrag)
 	}
 
 	// Move added point to new position
-	FInterpCurvePoint<FVector>& AddedPoint = SplineComp->GetSplinePointsPosition().Points[LastKeyIndexSelected];
-	const FVector CurrentWorldPos = SplineComp->GetComponentTransform().TransformPosition(AddedPoint.OutVal);
+	const FVector CurrentWorldPos = SplineComp->GetLocationAtSplinePoint(LastKeyIndexSelected, ESplineCoordinateSpace::World);
 	const FVector NewWorldPos = CurrentWorldPos + InDrag;
-	AddedPoint.OutVal = SplineComp->GetComponentTransform().InverseTransformPosition(NewWorldPos);
+	SplineComp->SetLocationAtSplinePoint(LastKeyIndexSelected, NewWorldPos, ESplineCoordinateSpace::World);
 
 	SplineComp->UpdateSpline();
 	SplineComp->bSplineHasBeenEdited = true;
@@ -3040,28 +2872,9 @@ void FSplineComponentVisualizer::OnDeleteKey()
 	SelectedKeysSorted.Sort([](int32 A, int32 B) { return A > B; });
 
 	// Delete selected keys from list, highest index first
-	FInterpCurveVector& SplinePosition = SplineComp->GetSplinePointsPosition();
-	FInterpCurveQuat& SplineRotation = SplineComp->GetSplinePointsRotation();
-	FInterpCurveVector& SplineScale = SplineComp->GetSplinePointsScale();
-	USplineMetadata* SplineMetadata = SplineComp->GetSplinePointsMetadata();
-		
 	for (int32 SelectedKeyIndex : SelectedKeysSorted)
 	{
-		if (SplineMetadata)
-		{
-			SplineMetadata->RemovePoint(SelectedKeyIndex);
-		}
-		
-		SplinePosition.Points.RemoveAt(SelectedKeyIndex);
-		SplineRotation.Points.RemoveAt(SelectedKeyIndex);
-		SplineScale.Points.RemoveAt(SelectedKeyIndex);
-
-		for (int Index = SelectedKeyIndex; Index < SplinePosition.Points.Num(); Index++)
-		{
-			SplinePosition.Points[Index].InVal -= 1.0f;
-			SplineRotation.Points[Index].InVal -= 1.0f;
-			SplineScale.Points[Index].InVal -= 1.0f;
-		}
+		SplineComp->RemoveSplinePoint(SelectedKeyIndex, false);
 	}
 
 	// Select first key
@@ -3070,7 +2883,6 @@ void FSplineComponentVisualizer::OnDeleteKey()
 	SelectionState->ClearSelectedSegmentIndex();
 	SelectionState->ClearSelectedTangentHandle();
 
-
 	SplineComp->UpdateSpline();
 	SplineComp->bSplineHasBeenEdited = true;
 
@@ -3093,7 +2905,7 @@ bool FSplineComponentVisualizer::CanDeleteKey() const
 	int32 LastKeyIndexSelected = SelectionState->GetLastKeyIndexSelected();
 	return (SplineComp != nullptr &&
 			SelectedKeys.Num() > 0 &&
-			SelectedKeys.Num() != SplineComp->SplineCurves.Position.Points.Num() &&
+			SelectedKeys.Num() != SplineComp->GetNumberOfSplinePoints() &&
 			LastKeyIndexSelected != INDEX_NONE);
 }
 
@@ -3142,10 +2954,9 @@ void FSplineComponentVisualizer::OnResetToAutomaticTangent(EInterpCurveMode Mode
 		const TSet<int32>& SelectedKeys = SelectionState->GetSelectedKeys();
 		for (int32 SelectedKeyIndex : SelectedKeys)
 		{
-			auto& Point = SplineComp->SplineCurves.Position.Points[SelectedKeyIndex];
-			if (Point.IsCurveKey())
+			if (SplineComponentVisualizerLocals::IsCurvePointType(SplineComp->GetSplinePointType(SelectedKeyIndex)))
 			{
-				Point.InterpMode = Mode;
+				SplineComp->SetSplinePointType(SelectedKeyIndex, ConvertInterpCurveModeToSplinePointType(Mode), false);
 			}
 		}
 
@@ -3172,8 +2983,8 @@ bool FSplineComponentVisualizer::CanResetToAutomaticTangent(EInterpCurveMode Mod
 		const TSet<int32>& SelectedKeys = SelectionState->GetSelectedKeys();
 		for (int32 SelectedKeyIndex : SelectedKeys)
 		{
-			const auto& Point = SplineComp->SplineCurves.Position.Points[SelectedKeyIndex];
-			if (Point.IsCurveKey() && Point.InterpMode != Mode)
+			ESplinePointType::Type CurrentMode = SplineComp->GetSplinePointType(SelectedKeyIndex);
+			if (SplineComponentVisualizerLocals::IsCurvePointType(CurrentMode) && ConvertSplinePointTypeToInterpCurveMode(CurrentMode) != Mode)
 			{
 				return true;
 			}
@@ -3206,7 +3017,7 @@ void FSplineComponentVisualizer::OnSetKeyType(EInterpCurveMode Mode)
 		{
 			check(SelectedKeyIndex >= 0);
 			check(SelectedKeyIndex < SplineComp->GetNumberOfSplinePoints());
-			SplineComp->SplineCurves.Position.Points[SelectedKeyIndex].InterpMode = Mode;
+			SplineComp->SetSplinePointType(SelectedKeyIndex, ConvertInterpCurveModeToSplinePointType(Mode), false);
 		}
 
 		SplineComp->UpdateSpline();
@@ -3237,8 +3048,8 @@ bool FSplineComponentVisualizer::IsKeyTypeSet(EInterpCurveMode Mode) const
 		{
 			check(SelectedKeyIndex >= 0);
 			check(SelectedKeyIndex < SplineComp->GetNumberOfSplinePoints());
-			const auto& SelectedPoint = SplineComp->SplineCurves.Position.Points[SelectedKeyIndex];
-			if ((Mode == CIM_CurveAuto && SelectedPoint.IsCurveKey()) || SelectedPoint.InterpMode == Mode)
+			const ESplinePointType::Type SelectedPointCurveType = SplineComp->GetSplinePointType(SelectedKeyIndex);
+			if ((Mode == CIM_CurveAuto && SplineComponentVisualizerLocals::IsCurvePointType(SelectedPointCurveType)) || SelectedPointCurveType == ConvertInterpCurveModeToSplinePointType(Mode))
 			{
 				return true;
 			}
@@ -3296,9 +3107,9 @@ void FSplineComponentVisualizer::OnSetDiscontinuousSpline()
 	// If not allowed discontinuous splines, set all ArriveTangents to match LeaveTangents
 	if (!SplineComp->bAllowDiscontinuousSpline)
 	{
-		for (int Index = 0; Index < SplineComp->SplineCurves.Position.Points.Num(); Index++)
+		for (int Index = 0; Index < SplineComp->GetNumberOfSplinePoints(); Index++)
 		{
-			SplineComp->SplineCurves.Position.Points[Index].ArriveTangent = SplineComp->SplineCurves.Position.Points[Index].LeaveTangent;
+			SplineComp->SetTangentAtSplinePoint(Index, SplineComp->GetLeaveTangentAtSplinePoint(Index, ESplineCoordinateSpace::Local), ESplineCoordinateSpace::Local, false);
 		}
 	}
 
@@ -3388,7 +3199,7 @@ bool FSplineComponentVisualizer::CanResetToDefault() const
 	USplineComponent* SplineComp = GetEditedSplineComponent();
 	if(SplineComp != nullptr)
     {
-        return SplineComp->SplineCurves != CastChecked<USplineComponent>(SplineComp->GetArchetype())->SplineCurves;
+        return SplineComp->CanResetToDefault();
     }
     else
     {
@@ -3544,8 +3355,7 @@ void FSplineComponentVisualizer::OnSelectAllSplinePoints()
 
 	if (USplineComponent* SplineComp = GetEditedSplineComponent())
 	{
-		const FInterpCurveVector& SplineInfo = SplineComp->GetSplinePointsPosition();
-		int32 NumPoints = SplineInfo.Points.Num();
+		int32 NumPoints = SplineComp->GetNumberOfSplinePoints();
 
 		check(SelectionState);
 		SelectionState->Modify();
@@ -3618,8 +3428,7 @@ void FSplineComponentVisualizer::GenerateContextMenuSections(FMenuBuilder& InMen
 			{
 				check(SelectedKeyIndex >= 0);
 				check(SelectedKeyIndex < SplineComp->GetNumberOfSplinePoints());
-				const auto& Point = SplineComp->SplineCurves.Position.Points[SelectedKeyIndex];
-				if (Point.IsCurveKey())
+				if (SplineComponentVisualizerLocals::IsCurvePointType(SplineComp->GetSplinePointType(SelectedKeyIndex)))
 				{
 					InMenuBuilder.AddSubMenu(
 						LOCTEXT("ResetToAutomaticTangent", "Reset to Automatic Tangent"),
diff --git a/Engine/Source/Editor/ComponentVisualizers/Private/SplineGeneratorPanel.cpp b/Engine/Source/Editor/ComponentVisualizers/Private/SplineGeneratorPanel.cpp
--- a/Engine/Source/Editor/ComponentVisualizers/Private/SplineGeneratorPanel.cpp
+++ b/Engine/Source/Editor/ComponentVisualizers/Private/SplineGeneratorPanel.cpp
@@ -63,7 +63,7 @@ void USplineGeneratorBase::Init(TWeakPtr<FSplineComponentVisualizer> InWeakSplin
 	StartKey = INDEX_NONE;
 	check(SplineVis->GetSelectedKeys().Num() == 1);
 	StartKey = *SplineVis->GetSelectedKeys().CreateConstIterator();
-	CachedSplineCurves = SelectedSplineComponent->SplineCurves;
+	CachedSplineCurves = SelectedSplineComponent->GetSplineCurves();
 }
 
 void USplineGeneratorBase::Reset()
@@ -84,14 +84,13 @@ void USplineGeneratorBase::Reset()
 			Owner->Modify();
 		}
 
-		SelectedSplineComponent->SplineCurves = CachedSplineCurves;
-
+		SelectedSplineComponent->SetSpline(CachedSplineCurves);
 		CachedSplineCurves = FSplineCurves();
 
 		SelectedSplineComponent->UpdateSpline();
 		SelectedSplineComponent->bSplineHasBeenEdited = true;
 
-		FProperty* SplineCurvesProperty = FindFProperty<FProperty>(USplineComponent::StaticClass(), GET_MEMBER_NAME_CHECKED(USplineComponent, SplineCurves));
+		FProperty* SplineCurvesProperty = FindFProperty<FProperty>(USplineComponent::StaticClass(), USplineComponent::GetSplinePropertyName());
 		FComponentVisualizer::NotifyPropertyModified(SelectedSplineComponent.Get(), SplineCurvesProperty);
 
 		GEditor->RedrawLevelEditingViewports(true);
@@ -121,7 +120,7 @@ void USplineGeneratorBase::PreviewCurve()
 		SelectedSplineComponent->UpdateSpline();
 		SelectedSplineComponent->bSplineHasBeenEdited = true;
 
-		FProperty* SplineCurvesProperty = FindFProperty<FProperty>(USplineComponent::StaticClass(), GET_MEMBER_NAME_CHECKED(USplineComponent, SplineCurves));
+		FProperty* SplineCurvesProperty = FindFProperty<FProperty>(USplineComponent::StaticClass(), USplineComponent::GetSplinePropertyName());
 		FComponentVisualizer::NotifyPropertyModified(SelectedSplineComponent.Get(), SplineCurvesProperty);
 
 		GEditor->RedrawLevelEditingViewports(true);
@@ -193,7 +192,7 @@ double CalcTangentMultiplier(const float InRadius, const float InRotInc)
 void UCircleSplineGenerator::BuildCurve()
 {
 	// Re-set the cached spline keys so we can add to them
-	SelectedSplineComponent->SplineCurves = CachedSplineCurves;
+	SelectedSplineComponent->SetSpline(CachedSplineCurves);
 
 	// naming is a little weird, will fix
 	const float BranchRightFlip = ((bKeepFirstKeyTangent && bBranchRight) ? 1.f : -1.f);
@@ -217,8 +216,6 @@ void UCircleSplineGenerator::BuildCurve()
 		SelectedSplineComponent->AddSplinePointAtIndex(CircleItr, GetAddIndex(Index), ESplineCoordinateSpace::Local);
 	}
 
-	FInterpCurveQuat& SplineRotation = SelectedSplineComponent->GetSplinePointsRotation();
-
 	float StartTangentFlip = BranchRightFlip * (bReverseDir ? 1.f : -1.f);
 	const FVector StartCircleTangent = (bKeepFirstKeyTangent ? StartQuat.GetForwardVector() : StartQuat.GetRightVector()) * StartTangentFlip;
 
@@ -228,8 +225,8 @@ void UCircleSplineGenerator::BuildCurve()
 		for (int32 Index = NumberOfPoints - 1; Index >= -1; Index--)
 		{
 			const FVector Tangent = StartCircleTangent.RotateAngleAxis(RotIncrement * float(Index + 1), StartQuat.GetUpVector()).GetSafeNormal();
+			SelectedSplineComponent->SetRotationAtSplinePoint(GetItrIndex(Index), (FQuat::MakeFromEuler(FVector(0.f, 0.f, RotIncrement * float(Index + 1))) * StartQuat).Rotator(), ESplineCoordinateSpace::Local, false);
 			SelectedSplineComponent->SetTangentAtSplinePoint(GetItrIndex(Index), Tangent * TangentMult, ESplineCoordinateSpace::Local);
-			SplineRotation.Points[GetItrIndex(Index)].OutVal = (FQuat::MakeFromEuler(FVector(0.f, 0.f, RotIncrement * float(Index + 1))) * StartQuat);
 		}
 	}
 }
@@ -237,7 +234,7 @@ void UCircleSplineGenerator::BuildCurve()
 void UArcSplineGenerator::BuildCurve()
 {
 	// Re-set the cached spline keys so we can add to them
-	SelectedSplineComponent->SplineCurves = CachedSplineCurves;
+	SelectedSplineComponent->SetSpline(CachedSplineCurves);
 
 	// Prevent the curve from going whack
 	Degrees = FMath::Clamp(Degrees, 0.f, 180.f * float(NumberOfPoints));
@@ -262,8 +259,6 @@ void UArcSplineGenerator::BuildCurve()
 		SelectedSplineComponent->AddSplinePointAtIndex(CircleItr, GetAddIndex(Index), ESplineCoordinateSpace::Local);
 	}
 
-	FInterpCurveQuat& SplineRotation = SelectedSplineComponent->GetSplinePointsRotation();
-
 	float StartTangentFlip = (bReverseDir ? -1.f : 1.f);
  	const FVector StartCircleTangent = (bKeepFirstKeyTangent ? StartQuat.GetForwardVector() : StartQuat.GetRightVector()) * StartTangentFlip;
 
@@ -275,8 +270,8 @@ void UArcSplineGenerator::BuildCurve()
 		for (int32 Index = 0; Index < NumberOfPoints; Index++)
 		{
 			const FVector Tangent = StartCircleTangent.RotateAngleAxis(RotIncrement * float(Index + 1), StartQuat.GetUpVector()).GetSafeNormal();
+			SelectedSplineComponent->SetRotationAtSplinePoint(GetItrIndex(Index), (FQuat::MakeFromEuler(FVector(0.f, 0.f, RotIncrement * float(Index + 1))) * StartQuat).Rotator(), ESplineCoordinateSpace::Local, false);
 			SelectedSplineComponent->SetTangentAtSplinePoint(GetItrIndex(Index), Tangent * TangentMult, ESplineCoordinateSpace::Local);
-			SplineRotation.Points[GetItrIndex(Index)].OutVal = (FQuat::MakeFromEuler(FVector(0.f, 0.f, RotIncrement * float(Index + 1))) * StartQuat);
 		}
 	}
 }
@@ -284,7 +279,7 @@ void UArcSplineGenerator::BuildCurve()
 void UEllipseSplineGenerator::BuildCurve()
 {
 	// Re-set the cached spline keys so we can add to them
-	SelectedSplineComponent->SplineCurves = CachedSplineCurves;
+	SelectedSplineComponent->SetSpline(CachedSplineCurves);
 
 	const float BranchRightFlip = ((bKeepFirstKeyTangent && bBranchRight) ? 1.f : -1.f);
 	const float ReverseDirFlip = (bReverseDir ^ !(ShapeAddMode & (EShapeAddMode::AppendBefore | EShapeAddMode::InsertBefore)) ? -1.f : 1.f) * BranchRightFlip * -1.f;
@@ -296,8 +291,6 @@ void UEllipseSplineGenerator::BuildCurve()
 	const FVector CenterDir = bKeepFirstKeyTangent ? StartQuat.GetRightVector() * BranchRightFlip : StartQuat.GetForwardVector();
 	const FVector CircleCenter = StartPoint + (bKeepFirstKeyTangent ? Length * .5f : Width * .5f) * CenterDir;
 
-	FInterpCurveQuat& SplineRotation = SelectedSplineComponent->GetSplinePointsRotation();
-
 	// Add points around ellipse
 	float CycleInc = (2 * PI / float(NumberOfPoints)) * ReverseDirFlip;
 	float CycleItr = CycleInc + (bKeepFirstKeyTangent ? PI * .5f : PI) * BranchRightFlip * -1.f;
@@ -315,15 +308,15 @@ void UEllipseSplineGenerator::BuildCurve()
 		const int32 AddIdx = GetAddIndex(Index);
 		SelectedSplineComponent->AddSplinePointAtIndex(RotatedEllipsePoint, AddIdx, ESplineCoordinateSpace::Local);
 		const FVector PrevTangent = SelectedSplineComponent->GetTangentAtSplinePoint(AddIdx, ESplineCoordinateSpace::Local);
+		SelectedSplineComponent->SetRotationAtSplinePoint(AddIdx, (FQuat::MakeFromEuler(FVector(0.f, 0.f, CycleItr)) * StartQuat).Rotator(), ESplineCoordinateSpace::Local, false);
 		SelectedSplineComponent->SetTangentAtSplinePoint(AddIdx, RotatedEllipseTangent * PrevTangent.Size(), ESplineCoordinateSpace::Local);
-		SplineRotation.Points[AddIdx].OutVal = (FQuat::MakeFromEuler(FVector(0.f, 0.f, CycleItr)) * StartQuat);
 	}
 }
 
 void USquareSplineGenerator::BuildCurve()
 {
 	// Re-set the cached spline keys so we can add to them
-	SelectedSplineComponent->SplineCurves = CachedSplineCurves;
+	SelectedSplineComponent->SetSpline(CachedSplineCurves);
 
 	const bool bPrepend = !(ShapeAddMode & (EShapeAddMode::AppendAfter | EShapeAddMode::InsertAfter));
 	const float LengthFlip = bPrepend ? -1.f : 1.f;
@@ -332,8 +325,6 @@ void USquareSplineGenerator::BuildCurve()
 	const FVector StartPoint = SelectedSplineComponent->GetLocationAtSplinePoint(GetAddIndex(-1), ESplineCoordinateSpace::Local);
 	const FQuat StartQuat = SelectedSplineComponent->GetQuaternionAtSplinePoint(GetAddIndex(-1), ESplineCoordinateSpace::Local);
 
-	FInterpCurveQuat& SplineRotation = SelectedSplineComponent->GetSplinePointsRotation();
-
 	const float RotDir = bBranchRight ? 90.f : -90.f;
 
 	FVector RectItr = StartPoint;
@@ -344,8 +335,8 @@ void USquareSplineGenerator::BuildCurve()
 		RectItr += DirItr * Length * LengthFlip;
 		const int32 AddIdx = GetAddIndex(Index);
 		SelectedSplineComponent->AddSplinePointAtIndex(RectItr, AddIdx, ESplineCoordinateSpace::Local);
-		SelectedSplineComponent->SplineCurves.Position.Points[bPrepend ? AddIdx : AddIdx - 1].InterpMode = EInterpCurveMode::CIM_Linear;
-		SplineRotation.Points[AddIdx].OutVal = FQuat::MakeFromEuler(FVector(0.f, 0.f, float(Index) * RotDir)) * StartQuat;
+		SelectedSplineComponent->SetSplinePointType(bPrepend ? AddIdx : AddIdx - 1, ConvertInterpCurveModeToSplinePointType(CIM_Linear));
+		SelectedSplineComponent->SetRotationAtSplinePoint(AddIdx, (FQuat::MakeFromEuler(FVector(0.f, 0.f, float(Index) * RotDir)) * StartQuat).Rotator(), ESplineCoordinateSpace::Local, false);
 		DirItr = DirItr.RotateAngleAxis(RotDir, StartQuat.GetUpVector());
 	}
 }
@@ -353,7 +344,7 @@ void USquareSplineGenerator::BuildCurve()
 void URectangleSplineGenerator::BuildCurve()
 {
 	// Re-set the cached spline keys so we can add to them
-	SelectedSplineComponent->SplineCurves = CachedSplineCurves;
+	SelectedSplineComponent->SetSpline(CachedSplineCurves);
 
 	const bool bPrepend = !(ShapeAddMode & (EShapeAddMode::AppendAfter | EShapeAddMode::InsertAfter));
 	const float LengthFlip = bPrepend ? -1.f : 1.f;
@@ -362,8 +353,6 @@ void URectangleSplineGenerator::BuildCurve()
 	const FVector StartPoint = SelectedSplineComponent->GetLocationAtSplinePoint(GetAddIndex(-1), ESplineCoordinateSpace::Local);
 	const FQuat StartQuat = SelectedSplineComponent->GetQuaternionAtSplinePoint(GetAddIndex(-1), ESplineCoordinateSpace::Local);
 
-	FInterpCurveQuat& SplineRotation = SelectedSplineComponent->GetSplinePointsRotation();
-
 	const float RotDir = bBranchRight ? 90.f : -90.f;
 
 	int32 Index = 0;
@@ -375,8 +364,8 @@ void URectangleSplineGenerator::BuildCurve()
 		const int32 AddIdx = GetAddIndex(Index);
 		RectItr += DirItr * Length * LengthFlip;
 		SelectedSplineComponent->AddSplinePointAtIndex(RectItr, AddIdx, ESplineCoordinateSpace::Local);
-		SelectedSplineComponent->SplineCurves.Position.Points[bPrepend ? AddIdx : AddIdx - 1].InterpMode = EInterpCurveMode::CIM_Linear;
-		SplineRotation.Points[AddIdx].OutVal = FQuat::MakeFromEuler(FVector(0.f, 0.f, float(Index) * RotDir)) * StartQuat;
+		SelectedSplineComponent->SetSplinePointType(bPrepend ? AddIdx : AddIdx - 1, ConvertInterpCurveModeToSplinePointType(CIM_Linear));
+		SelectedSplineComponent->SetRotationAtSplinePoint(AddIdx, (FQuat::MakeFromEuler(FVector(0.f, 0.f, float(Index) * RotDir)) * StartQuat).Rotator(), ESplineCoordinateSpace::Local, false);
 		Index++;
 	}
 
@@ -386,8 +375,8 @@ void URectangleSplineGenerator::BuildCurve()
 		DirItr = DirItr.RotateAngleAxis(RotDir, StartQuat.GetUpVector());
 		RectItr += DirItr * Width * LengthFlip;
 		SelectedSplineComponent->AddSplinePointAtIndex(RectItr, AddIdx, ESplineCoordinateSpace::Local);
-		SelectedSplineComponent->SplineCurves.Position.Points[bPrepend ? AddIdx : AddIdx - 1].InterpMode = EInterpCurveMode::CIM_Linear;
-		SplineRotation.Points[AddIdx].OutVal = FQuat::MakeFromEuler(FVector(0.f, 0.f, float(Index) * RotDir)) * StartQuat;
+		SelectedSplineComponent->SetSplinePointType(bPrepend ? AddIdx : AddIdx - 1, ConvertInterpCurveModeToSplinePointType(CIM_Linear));
+		SelectedSplineComponent->SetRotationAtSplinePoint(AddIdx, (FQuat::MakeFromEuler(FVector(0.f, 0.f, float(Index) * RotDir)) * StartQuat).Rotator(), ESplineCoordinateSpace::Local, false);
 		Index++;
 	}
 
@@ -397,8 +386,8 @@ void URectangleSplineGenerator::BuildCurve()
 		DirItr = DirItr.RotateAngleAxis(RotDir, StartQuat.GetUpVector());
 		RectItr += DirItr * Length * LengthFlip;
 		SelectedSplineComponent->AddSplinePointAtIndex(RectItr, AddIdx, ESplineCoordinateSpace::Local);
-		SelectedSplineComponent->SplineCurves.Position.Points[bPrepend ? AddIdx : AddIdx - 1].InterpMode = EInterpCurveMode::CIM_Linear;
-		SplineRotation.Points[AddIdx].OutVal = FQuat::MakeFromEuler(FVector(0.f, 0.f, float(Index) * RotDir)) * StartQuat;
+		SelectedSplineComponent->SetSplinePointType(bPrepend ? AddIdx : AddIdx - 1, ConvertInterpCurveModeToSplinePointType(CIM_Linear));
+		SelectedSplineComponent->SetRotationAtSplinePoint(AddIdx, (FQuat::MakeFromEuler(FVector(0.f, 0.f, float(Index) * RotDir)) * StartQuat).Rotator(), ESplineCoordinateSpace::Local, false);
 		Index++;
 	}
 
@@ -408,8 +397,8 @@ void URectangleSplineGenerator::BuildCurve()
 		DirItr = DirItr.RotateAngleAxis(RotDir, StartQuat.GetUpVector());
 		RectItr += DirItr * Width * LengthFlip;
 		SelectedSplineComponent->AddSplinePointAtIndex(RectItr, AddIdx, ESplineCoordinateSpace::Local);
-		SelectedSplineComponent->SplineCurves.Position.Points[bPrepend ? AddIdx : AddIdx - 1].InterpMode = EInterpCurveMode::CIM_Linear;
-		SplineRotation.Points[AddIdx].OutVal = FQuat::MakeFromEuler(FVector(0.f, 0.f, float(Index) * RotDir)) * StartQuat;
+		SelectedSplineComponent->SetSplinePointType(bPrepend ? AddIdx : AddIdx - 1, ConvertInterpCurveModeToSplinePointType(CIM_Linear));
+		SelectedSplineComponent->SetRotationAtSplinePoint(AddIdx, (FQuat::MakeFromEuler(FVector(0.f, 0.f, float(Index) * RotDir)) * StartQuat).Rotator(), ESplineCoordinateSpace::Local, false);
 		Index++;
 	}
 }
@@ -417,7 +406,7 @@ void URectangleSplineGenerator::BuildCurve()
 void ULineSplineGenerator::BuildCurve()
 {
 	// Re-set the cached spline keys so we can add to them
-	SelectedSplineComponent->SplineCurves = CachedSplineCurves;
+	SelectedSplineComponent->SetSpline(CachedSplineCurves);
 
 	const bool bPrepend = !(ShapeAddMode & (EShapeAddMode::AppendAfter | EShapeAddMode::InsertAfter));
 	const double LengthFlip = bPrepend ? -1.f : 1.f;
diff --git a/Engine/Source/Editor/DetailCustomizations/Private/SplineComponentDetails.cpp b/Engine/Source/Editor/DetailCustomizations/Private/SplineComponentDetails.cpp
--- a/Engine/Source/Editor/DetailCustomizations/Private/SplineComponentDetails.cpp
+++ b/Engine/Source/Editor/DetailCustomizations/Private/SplineComponentDetails.cpp
@@ -241,6 +241,7 @@ private:
 	bool IsOnePointSelected() const { return SelectedKeys.Num() == 1; }
 	bool ArePointsSelected() const { return (SelectedKeys.Num() > 0); };
 	bool AreNoPointsSelected() const { return (SelectedKeys.Num() == 0); };
+	bool CanSetInputKey() const;
 	TOptional<float> GetInputKey() const { return InputKey.Value; }
 	TOptional<FVector::FReal> GetPositionX() const { return Position.X; }
 	TOptional<FVector::FReal> GetPositionY() const { return Position.Y; }
@@ -340,7 +341,7 @@ FSplinePointDetails::FSplinePointDetails(USplineComponent* InOwningSplineCompone
 	SplineVisualizer = StaticCastSharedPtr<FSplineComponentVisualizer>(Visualizer);
 	check(SplineVisualizer.IsValid());
 
-	SplineCurvesProperty = FindFProperty<FProperty>(USplineComponent::StaticClass(), GET_MEMBER_NAME_CHECKED(USplineComponent, SplineCurves));
+	SplineCurvesProperty = FindFProperty<FProperty>(USplineComponent::StaticClass(), USplineComponent::GetSplinePropertyName());
 
 	const TArray<ESplinePointType::Type> EnabledSplinePointTypes = InOwningSplineComponent->GetEnabledSplinePointTypes();
 
@@ -536,7 +537,7 @@ void FSplinePointDetails::GenerateChildContent(IDetailChildrenBuilder& ChildrenB
 		.MaxDesiredWidth(125.0f)
 		[
 			SNew(SNumericEntryBox<float>)
-			.IsEnabled(TAttribute<bool>(this, &FSplinePointDetails::IsOnePointSelected))
+			.IsEnabled(TAttribute<bool>(this, &FSplinePointDetails::CanSetInputKey))
 			.Value(this, &FSplinePointDetails::GetInputKey)
 			.UndeterminedString(LOCTEXT("Multiple", "Multiple"))
 			.OnValueCommitted(this, &FSplinePointDetails::OnSetInputKey)
@@ -837,20 +838,15 @@ void FSplinePointDetails::UpdateValues()
 		bool bValidIndices = true;
 		for (int32 Index : SelectedKeys)
 		{
-			if (Index < 0 ||
-				Index >= SplineComp->GetSplinePointsPosition().Points.Num() ||
-				Index >= SplineComp->GetSplinePointsRotation().Points.Num() ||
-				Index >= SplineComp->GetSplinePointsScale().Points.Num())
+			if (Index < 0 || Index > SplineComp->GetNumberOfSplinePoints())
 			{
 				bValidIndices = false;
 				if (!bAlreadyWarnedInvalidIndex)
 				{
-					UE_LOG(LogSplineComponentDetails, Error, TEXT("Spline component details selected keys contains invalid index %d for spline %s with %d points, %d rotations, %d scales"),
+					UE_LOG(LogSplineComponentDetails, Error, TEXT("Spline component details selected keys contains invalid index %d for spline %s with %d points"),
 						Index,
 						*SplineComp->GetPathName(),
-						SplineComp->GetSplinePointsPosition().Points.Num(),
-						SplineComp->GetSplinePointsRotation().Points.Num(),
-						SplineComp->GetSplinePointsScale().Points.Num());
+						SplineComp->GetNumberOfSplinePoints());
 					bAlreadyWarnedInvalidIndex = true;
 				}
 				break;
@@ -861,33 +857,18 @@ void FSplinePointDetails::UpdateValues()
 		{
 			for (int32 Index : SelectedKeys)
 			{
-				const FTransform SplineToWorld = SplineComp->GetComponentToWorld();
-
-				if (bEditingLocationAbsolute)
-				{
-					const FVector AbsoluteLocation = SplineToWorld.TransformPosition(SplineComp->GetSplinePointsPosition().Points[Index].OutVal);
-					Position.Add(AbsoluteLocation);
-				}
-				else
-				{
-					Position.Add(SplineComp->GetSplinePointsPosition().Points[Index].OutVal);
-				}
-
-				if (bEditingRotationAbsolute)
-				{
-					const FQuat AbsoluteRotation = SplineToWorld.TransformRotation(SplineComp->GetSplinePointsRotation().Points[Index].OutVal);
-					Rotation.Add(AbsoluteRotation.Rotator());
-				}
-				else
-				{
-					Rotation.Add(SplineComp->GetSplinePointsRotation().Points[Index].OutVal.Rotator());
-				}
-
-				InputKey.Add(SplineComp->GetSplinePointsPosition().Points[Index].InVal);
-				Scale.Add(SplineComp->GetSplinePointsScale().Points[Index].OutVal);
-				ArriveTangent.Add(SplineComp->GetSplinePointsPosition().Points[Index].ArriveTangent);
-				LeaveTangent.Add(SplineComp->GetSplinePointsPosition().Points[Index].LeaveTangent);
-				PointType.Add(ConvertInterpCurveModeToSplinePointType(SplineComp->GetSplinePointsPosition().Points[Index].InterpMode));
+				// possibly could get this data in bulk via GetSplinePoint(Index), but doing a 1:1 swap for now.
+				
+				InputKey.Add(SplineComp->GetInputKeyValueAtSplinePoint(Index));
+				
+				Position.Add(SplineComp->GetLocationAtSplinePoint(Index, bEditingLocationAbsolute ? ESplineCoordinateSpace::World : ESplineCoordinateSpace::Local));
+				Rotation.Add(SplineComp->GetRotationAtSplinePoint(Index, bEditingRotationAbsolute ? ESplineCoordinateSpace::World : ESplineCoordinateSpace::Local));
+				Scale.Add(SplineComp->GetScaleAtSplinePoint(Index));
+				
+				ArriveTangent.Add(SplineComp->GetArriveTangentAtSplinePoint(Index, ESplineCoordinateSpace::Local));
+				LeaveTangent.Add(SplineComp->GetLeaveTangentAtSplinePoint(Index, ESplineCoordinateSpace::Local));
+				
+				PointType.Add(SplineComp->GetSplinePointType(Index));
 			}
 
 			if (SplineMetaDataDetails)
@@ -909,8 +890,27 @@ FName FSplinePointDetails::GetName() const
 	return Name;
 }
 
+bool FSplinePointDetails::CanSetInputKey() const
+{
+	if (!SplineComp)
+	{
+		return false;
+	}
+
+PRAGMA_DISABLE_DEPRECATION_WARNINGS
+	bool bUsingSplineCurves = SplineComp->GetSplinePropertyName() == GET_MEMBER_NAME_CHECKED(USplineComponent, SplineCurves);
+PRAGMA_ENABLE_DEPRECATION_WARNINGS
+
+	return IsOnePointSelected() && bUsingSplineCurves;
+}
+
 void FSplinePointDetails::OnSetInputKey(float NewValue, ETextCommit::Type CommitInfo)
 {
+PRAGMA_DISABLE_DEPRECATION_WARNINGS
+	// Note: This function strongly assumes that SplineCurves is the authoritative data structure backing the spline component.
+	// I have made this assumption valid by introducing CanSetInputKey() which verifies this by checking with the selected component.
+	// This assumption is necessary because there is no interface on the component which allows us to modify input keys, we must directly write to SplineCurves.
+	
 	if ((CommitInfo != ETextCommit::OnEnter && CommitInfo != ETextCommit::OnUserMovedFocus) || !SplineComp)
 	{
 		return;
@@ -986,6 +986,8 @@ void FSplinePointDetails::OnSetInputKey(float NewValue, ETextCommit::Type Commit
 	UpdateValues();
 
 	GEditor->RedrawLevelEditingViewports(true);
+
+PRAGMA_ENABLE_DEPRECATION_WARNINGS
 }
 
 void FSplinePointDetails::OnSetPosition(FVector::FReal NewValue, ETextCommit::Type CommitInfo, EAxis::Type Axis)
@@ -1000,31 +1002,21 @@ void FSplinePointDetails::OnSetPosition(FVector::FReal NewValue, ETextCommit::Ty
 		const FScopedTransaction Transaction(LOCTEXT("SetSplinePointPosition", "Set spline point position"), !bInSliderTransaction);
 		SplineComp->Modify();
 
+		const int32 NumPoints = SplineComp->GetNumberOfSplinePoints();
+		
 		for (int32 Index : SelectedKeys)
 		{
-			if (Index < 0 || Index >= SplineComp->GetSplinePointsPosition().Points.Num())
+			if (Index < 0 || Index >= NumPoints)
 			{
 				UE_LOG(LogSplineComponentDetails, Error, TEXT("Set spline point location: invalid index %d in selected points for spline component %s which contains %d spline points."),
-					Index, *SplineComp->GetPathName(), SplineComp->GetSplinePointsPosition().Points.Num());
+					Index, *SplineComp->GetPathName(), NumPoints);
 				continue;
 			}
 
-			if (bEditingLocationAbsolute)
-			{
-				const FTransform SplineToWorld = SplineComp->GetComponentToWorld();
-				const FVector RelativePos = SplineComp->GetSplinePointsPosition().Points[Index].OutVal;
-				FVector AbsolutePos = SplineToWorld.TransformPosition(RelativePos);
-				AbsolutePos.SetComponentForAxis(Axis, NewValue);
-				FVector PointPosition = SplineToWorld.InverseTransformPosition(AbsolutePos);
-
-				SplineComp->GetSplinePointsPosition().Points[Index].OutVal = PointPosition;
-			}
-			else
-			{
-				FVector PointPosition = SplineComp->GetSplinePointsPosition().Points[Index].OutVal;
-				PointPosition.SetComponentForAxis(Axis, NewValue);
-				SplineComp->GetSplinePointsPosition().Points[Index].OutVal = PointPosition;
-			}
+			const ESplineCoordinateSpace::Type SplineCoordinateSpace = bEditingLocationAbsolute ? ESplineCoordinateSpace::World : ESplineCoordinateSpace::Local;
+			FVector PointPosition = SplineComp->GetLocationAtSplinePoint(Index, SplineCoordinateSpace);
+			PointPosition.SetComponentForAxis(Axis, NewValue);
+			SplineComp->SetLocationAtSplinePoint(Index, PointPosition, SplineCoordinateSpace, false);
 		}
 	}
 
@@ -1055,19 +1047,24 @@ void FSplinePointDetails::OnSetArriveTangent(FVector::FReal NewValue, ETextCommi
 		const FScopedTransaction Transaction(LOCTEXT("SetSplinePointTangent", "Set spline point tangent"));
 		SplineComp->Modify();
 
+		const int32 NumPoints = SplineComp->GetNumberOfSplinePoints();
+		
 		for (int32 Index : SelectedKeys)
 		{
-			if (Index < 0 || Index >= SplineComp->GetSplinePointsPosition().Points.Num())
+			if (Index < 0 || Index >= NumPoints)
 			{
 				UE_LOG(LogSplineComponentDetails, Error, TEXT("Set spline point arrive tangent: invalid index %d in selected points for spline component %s which contains %d spline points."),
-					Index, *SplineComp->GetPathName(), SplineComp->GetSplinePointsPosition().Points.Num());
+					Index, *SplineComp->GetPathName(), NumPoints);
 				continue;
 			}
 
-			FVector PointTangent = SplineComp->GetSplinePointsPosition().Points[Index].ArriveTangent;
-			PointTangent.SetComponentForAxis(Axis, NewValue);
-			SplineComp->GetSplinePointsPosition().Points[Index].ArriveTangent = PointTangent;
-			SplineComp->GetSplinePointsPosition().Points[Index].InterpMode = CIM_CurveUser;
+			FVector LocalArriveTangent = SplineComp->GetArriveTangentAtSplinePoint(Index, ESplineCoordinateSpace::Local);
+			FVector LocalLeaveTangent = SplineComp->GetLeaveTangentAtSplinePoint(Index, ESplineCoordinateSpace::Local);
+
+			LocalArriveTangent.SetComponentForAxis(Axis, NewValue);
+			
+			SplineComp->SetTangentsAtSplinePoint(Index, LocalArriveTangent, LocalLeaveTangent, ESplineCoordinateSpace::Local, false);
+
 		}
 	}
 
@@ -1098,19 +1095,23 @@ void FSplinePointDetails::OnSetLeaveTangent(FVector::FReal NewValue, ETextCommit
 		const FScopedTransaction Transaction(LOCTEXT("SetSplinePointTangent", "Set spline point tangent"));
 		SplineComp->Modify();
 
+		const int32 NumPoints = SplineComp->GetNumberOfSplinePoints();
+		
 		for (int32 Index : SelectedKeys)
 		{
-			if (Index < 0 || Index >= SplineComp->GetSplinePointsPosition().Points.Num())
+			if (Index < 0 || Index >= NumPoints)
 			{
 				UE_LOG(LogSplineComponentDetails, Error, TEXT("Set spline point leave tangent: invalid index %d in selected points for spline component %s which contains %d spline points."),
-					Index, *SplineComp->GetPathName(), SplineComp->GetSplinePointsPosition().Points.Num());
+					Index, *SplineComp->GetPathName(), NumPoints);
 				continue;
 			}
 
-			FVector PointTangent = SplineComp->GetSplinePointsPosition().Points[Index].LeaveTangent;
-			PointTangent.SetComponentForAxis(Axis, NewValue);
-			SplineComp->GetSplinePointsPosition().Points[Index].LeaveTangent = PointTangent;
-			SplineComp->GetSplinePointsPosition().Points[Index].InterpMode = CIM_CurveUser;
+			FVector LocalArriveTangent = SplineComp->GetArriveTangentAtSplinePoint(Index, ESplineCoordinateSpace::Local);
+			FVector LocalLeaveTangent = SplineComp->GetLeaveTangentAtSplinePoint(Index, ESplineCoordinateSpace::Local);
+
+			LocalLeaveTangent.SetComponentForAxis(Axis, NewValue);
+			
+			SplineComp->SetTangentsAtSplinePoint(Index, LocalArriveTangent, LocalLeaveTangent, ESplineCoordinateSpace::Local, false);
 		}
 	}
 
@@ -1141,19 +1142,20 @@ void FSplinePointDetails::OnSetRotation(FRotator::FReal NewValue, ETextCommit::T
 	{
 		const FScopedTransaction Transaction(LOCTEXT("SetSplinePointRotation", "Set spline point rotation"));
 		SplineComp->Modify();
+		
+		const int32 NumPoints = SplineComp->GetNumberOfSplinePoints();
+		
 		FQuat SplineComponentRotation = SplineComp->GetComponentQuat();
 		for (int32 Index : SelectedKeys)
 		{
-			if (Index < 0 || Index >= SplineComp->GetSplinePointsRotation().Points.Num())
+			if (Index < 0 || Index >= NumPoints)
 			{
 				UE_LOG(LogSplineComponentDetails, Error, TEXT("Set spline point rotation: invalid index %d in selected points for spline component %s which contains %d spline points."),
-					Index, *SplineComp->GetPathName(), SplineComp->GetSplinePointsRotation().Points.Num());
+					Index, *SplineComp->GetPathName(), NumPoints);
 				continue;
 			}
 
-			FInterpCurvePoint<FVector>& EditedPoint = SplineComp->GetSplinePointsPosition().Points[Index];
-			FInterpCurvePoint<FQuat>& EditedRotPoint = SplineComp->GetSplinePointsRotation().Points[Index];
-			const FQuat CurrentRotationRelative = EditedRotPoint.OutVal;
+			const FQuat CurrentRotationRelative = SplineComp->GetRotationAtSplinePoint(Index, ESplineCoordinateSpace::Local).Quaternion();
 
 			if (bEditingRotationAbsolute)
 			{
@@ -1182,15 +1184,13 @@ void FSplinePointDetails::OnSetRotation(FRotator::FReal NewValue, ETextCommit::T
 				NewRotationRelative = NewRotationRotator.Quaternion();
 			}
 
-			SplineComp->GetSplinePointsRotation().Points[Index].OutVal = NewRotationRelative;
+			SplineComp->SetRotationAtSplinePoint(Index, NewRotationRelative.Rotator(), ESplineCoordinateSpace::Local);
 
 			FQuat DeltaRotate(NewRotationRelative * CurrentRotationRelative.Inverse());
 			// Rotate tangent according to delta rotation
-			FVector NewTangent = SplineComponentRotation.RotateVector(EditedPoint.LeaveTangent); // convert local-space tangent vector to world-space
+			FVector NewTangent = SplineComp->GetLeaveTangentAtSplinePoint(Index, ESplineCoordinateSpace::World);
 			NewTangent = DeltaRotate.RotateVector(NewTangent); // apply world-space delta rotation to world-space tangent
-			NewTangent = SplineComponentRotation.Inverse().RotateVector(NewTangent); // convert world-space tangent vector back into local-space
-			EditedPoint.LeaveTangent = NewTangent;
-			EditedPoint.ArriveTangent = NewTangent;
+			SplineComp->SetTangentAtSplinePoint(Index, NewTangent, ESplineCoordinateSpace::World, false);
 		}
 	}
 
@@ -1222,18 +1222,20 @@ void FSplinePointDetails::OnSetScale(FVector::FReal NewValue, ETextCommit::Type
 		const FScopedTransaction Transaction(LOCTEXT("SetSplinePointScale", "Set spline point scale"));
 		SplineComp->Modify();
 
+		const int32 NumPoints = SplineComp->GetNumberOfSplinePoints();
+		
 		for (int32 Index : SelectedKeys)
 		{
-			if (Index < 0 || Index >= SplineComp->GetSplinePointsScale().Points.Num())
+			if (Index < 0 || Index >= NumPoints)
 			{
 				UE_LOG(LogSplineComponentDetails, Error, TEXT("Set spline point scale: invalid index %d in selected points for spline component %s which contains %d spline points."),
-					Index, *SplineComp->GetPathName(), SplineComp->GetSplinePointsScale().Points.Num());
+					Index, *SplineComp->GetPathName(), NumPoints);
 				continue;
 			}
 
-			FVector PointScale = SplineComp->GetSplinePointsScale().Points[Index].OutVal;
+			FVector PointScale = SplineComp->GetScaleAtSplinePoint(Index);
 			PointScale.SetComponentForAxis(Axis, NewValue);
-			SplineComp->GetSplinePointsScale().Points[Index].OutVal = PointScale;
+			SplineComp->SetScaleAtSplinePoint(Index, PointScale, false);
 		}
 	}
 
@@ -1287,18 +1289,20 @@ void FSplinePointDetails::OnSplinePointTypeChanged(TSharedPtr<FString> NewValue,
 			const FScopedTransaction Transaction(LOCTEXT("SetSplinePointType", "Set spline point type"));
 			SplineComp->Modify();
 
+			const int32 NumPoints = SplineComp->GetNumberOfSplinePoints();
+			
 			for (int32 Index : SelectedKeys)
 			{
-				if (Index < 0 || Index >= SplineComp->GetSplinePointsPosition().Points.Num())
+				if (Index < 0 || Index >= NumPoints)
 				{
 					UE_LOG(LogSplineComponentDetails, Error, TEXT("Set spline point type: invalid index %d in selected points for spline component %s which contains %d spline points."),
-						Index, *SplineComp->GetPathName(), SplineComp->GetSplinePointsPosition().Points.Num());
+						Index, *SplineComp->GetPathName(), NumPoints);
 					continue;
 				}
 
-				if (SplineComp->GetSplinePointsPosition().Points[Index].InterpMode != Mode)
+				if (SplineComp->GetSplinePointType(Index) != ConvertInterpCurveModeToSplinePointType(Mode))
 				{
-					SplineComp->GetSplinePointsPosition().Points[Index].InterpMode = Mode;
+					SplineComp->SetSplinePointType(Index, ConvertInterpCurveModeToSplinePointType(Mode), false);
 					bWasModified = true;
 				}
 			}
@@ -1813,7 +1817,7 @@ TSharedRef<IDetailCustomization> FSplineComponentDetails::MakeInstance()
 void FSplineComponentDetails::CustomizeDetails(IDetailLayoutBuilder& DetailBuilder)
 {
 	// Hide the SplineCurves property
-	TSharedPtr<IPropertyHandle> SplineCurvesProperty = DetailBuilder.GetProperty(GET_MEMBER_NAME_CHECKED(USplineComponent, SplineCurves));
+	TSharedPtr<IPropertyHandle> SplineCurvesProperty = DetailBuilder.GetProperty(USplineComponent::GetSplinePropertyName());
 	SplineCurvesProperty->MarkHiddenByCustomization();
 
 
diff --git a/Engine/Source/Runtime/Engine/Classes/Components/SplineComponent.h b/Engine/Source/Runtime/Engine/Classes/Components/SplineComponent.h
--- a/Engine/Source/Runtime/Engine/Classes/Components/SplineComponent.h
+++ b/Engine/Source/Runtime/Engine/Classes/Components/SplineComponent.h
@@ -5,6 +5,7 @@
 #include "CoreMinimal.h"
 #include "UObject/ObjectMacros.h"
 #include "Components/PrimitiveComponent.h"
+#include "Curves/Spline.h"
 #include "Math/InterpCurve.h"
 #include "SplineComponent.generated.h"
 
@@ -60,6 +61,16 @@ struct FSplineCurves
 {
 	GENERATED_BODY()
 
+	static FSplineCurves FromSplineInterface(const FSpline& InProxy)
+	{
+		FSplineCurves NewCurve;
+		NewCurve.Position = InProxy.GetSplinePointsPosition();
+		NewCurve.Rotation = InProxy.GetSplinePointsRotation();
+		NewCurve.Scale = InProxy.GetSplinePointsScale();
+		NewCurve.ReparamTable = InProxy.GetSplineReparamTable();
+		return NewCurve;
+	}
+	
 	/** Spline built from position data. */
 	UPROPERTY()
 	FInterpCurveVector Position;
@@ -197,6 +208,18 @@ class USplineComponent : public UPrimitiveComponent
 {
 	GENERATED_UCLASS_BODY()
 
+protected:
+
+	// todo: replicate, but I believe we need to make the struct able to replicate in the first place.
+	UPROPERTY(EditAnywhere, Category=Points)
+	FSpline Spline;
+	
+public:
+
+	// Gives the name of the backing spline property to allow users
+	// the ability to correctly trigger property change events.
+	ENGINE_API static FName GetSplinePropertyName();
+	
 	UPROPERTY(EditAnywhere, Replicated, Category=Points)
 	FSplineCurves SplineCurves;
 
@@ -328,7 +351,6 @@ public:
 #endif
 	//~ End UPrimitiveComponent Interface.
 
-
 	//~ Begin USceneComponent Interface
 	ENGINE_API virtual FBoxSphereBounds CalcBounds(const FTransform& LocalToWorld) const override;
 	//~ End USceneComponent Interface
@@ -337,12 +359,19 @@ public:
 	static ENGINE_API void Draw(FPrimitiveDrawInterface* PDI, const FSceneView* View, const FInterpCurveVector& SplineInfo, const FMatrix& LocalToWorld, const FLinearColor& LineColor, uint8 DepthPriorityGroup);
 #endif
 
-	FInterpCurveVector& GetSplinePointsPosition() { return SplineCurves.Position; }
-	const FInterpCurveVector& GetSplinePointsPosition() const { return SplineCurves.Position; }
-	FInterpCurveQuat& GetSplinePointsRotation() { return SplineCurves.Rotation; }
-	const FInterpCurveQuat& GetSplinePointsRotation() const { return SplineCurves.Rotation; }
-	FInterpCurveVector& GetSplinePointsScale() { return SplineCurves.Scale; }
-	const FInterpCurveVector& GetSplinePointsScale() const { return SplineCurves.Scale; }
+	// Safe for caching spline data in the legacy format. Will auto convert data if necessary.
+	ENGINE_API FSplineCurves GetSplineCurves() const;
+	ENGINE_API void SetSpline(const FSpline& InSpline);
+	ENGINE_API void SetSpline(const FSplineCurves& InSplineCurves);
+	
+	/** Reset all spline data to match the current state of SplineCurves */
+	void SynchronizeSplines() { SetSpline(WarninglessSplineCurves()); }
+
+	ENGINE_API int32 GetVersion() const;
+	ENGINE_API const FInterpCurveVector& GetSplinePointsPosition() const;
+	ENGINE_API const FInterpCurveQuat& GetSplinePointsRotation() const;
+	ENGINE_API const FInterpCurveVector& GetSplinePointsScale() const;
+	ENGINE_API const FInterpCurveFloat& GetSplineReparamTable() const;
 
 	virtual USplineMetadata* GetSplinePointsMetadata() { return nullptr; }
 	virtual const USplineMetadata* GetSplinePointsMetadata() const { return nullptr; }
@@ -361,11 +390,11 @@ public:
 	/** Controls the visibility of the Spline point leave tangent editor in the details panel. */
 	virtual bool AllowsSplinePointLeaveTangentEditing() const { return true; }
 
-
 	ENGINE_API void ApplyComponentInstanceData(struct FSplineInstanceData* ComponentInstanceData, const bool bPostUCS);
 
 	/** Reset the spline to its default shape (a spline of two points) */
 	ENGINE_API void ResetToDefault();
+	ENGINE_API bool CanResetToDefault() const;
 
 	/** Update the spline tangents and SplineReparamTable */
 	UFUNCTION(BlueprintCallable, Category = Spline)
@@ -522,9 +551,10 @@ public:
 	UFUNCTION(BlueprintCallable, Category = Spline)
 	ENGINE_API void SetUpVectorAtSplinePoint(int32 PointIndex, const FVector& InUpVector, ESplineCoordinateSpace::Type CoordinateSpace, bool bUpdateSpline = true);
 
-	/** Set the rotation of an existing spline point */
+	/** Set the rotation of an existing spline point. This also aligns tangents with the forward vector of the rotation. */
 	UFUNCTION(BlueprintCallable, Category = Spline)
 	ENGINE_API void SetRotationAtSplinePoint(int32 PointIndex, const FRotator& InRotation, ESplineCoordinateSpace::Type CoordinateSpace, bool bUpdateSpline = true);
+	ENGINE_API void SetQuaternionAtSplinePoint(int32 PointIndex, const FQuat& InQuaternion, ESplineCoordinateSpace::Type CoordinateSpace, bool bUpdateSpline = true);
 
 	/** Set the scale at a given spline point */
 	UFUNCTION(BlueprintCallable, Category = Spline)
@@ -762,6 +792,10 @@ public:
 	UFUNCTION(BlueprintCallable, Category=Spline)
 	ENGINE_API float FindInputKeyClosestToWorldLocation(const FVector& WorldLocation) const;
 
+	/** Given a location, in world space, return the input key closest to that location on the specified segment. */
+	UFUNCTION(BlueprintCallable, Category=Spline)
+	ENGINE_API float FindInputKeyOnSegmentClosestToWorldLocation(const FVector& WorldLocation, int32 Index) const;
+	
 	/** Given a location, in world space, return the point on the curve that is closest to the location. */
 	UFUNCTION(BlueprintCallable, Category=Spline)
 	ENGINE_API FVector FindLocationClosestToWorldLocation(const FVector& WorldLocation, ESplineCoordinateSpace::Type CoordinateSpace) const;
@@ -809,7 +843,6 @@ public:
 	UFUNCTION(BlueprintCallable, Category = Spline)
 	ENGINE_API bool DivideSplineIntoPolylineRecursiveWithDistances(float StartDistanceAlongSpline, float EndDistanceAlongSpline, ESplineCoordinateSpace::Type CoordinateSpace, const float MaxSquareDistanceFromSpline, TArray<FVector>& OutPoints, TArray<double>& OutDistancesAlongSpline) const;
 
-
 	/** Given a threshold, returns a list of vertices along the spline segment that, treated as a list of segments (polyline), matches the spline shape. */
 	UFUNCTION(BlueprintCallable, Category = Spline)
 	ENGINE_API bool ConvertSplineSegmentToPolyLine(int32 SplinePointStartIndex, ESplineCoordinateSpace::Type CoordinateSpace, const float MaxSquareDistanceFromSpline, TArray<FVector>& OutPoints) const;
@@ -830,20 +863,59 @@ public:
 	UFUNCTION(BlueprintPure = false, Category = Spline)
 	ENGINE_API bool ConvertSplineToPolyline_InTimeRange(ESplineCoordinateSpace::Type CoordinateSpace, const float MaxSquareDistanceFromSpline, float StartTimeAlongSpline, float EndTimeAlongSpline, bool bUseConstantVelocity, TArray<FVector>& OutPoints, TArray<double>& OutDistancesAlongSpline, bool bAllowWrappingIfClosed = true) const;
 
+	UE_DEPRECATED(5.6, "Mutable access to SplinePointsPosition is deprecated, please use USplineComponent functions for indirect access.")
+	FInterpCurveVector& GetSplinePointsPosition() { return WarninglessSplineCurves().Position; }
+	
+	UE_DEPRECATED(5.6, "Mutable access to SplinePointsRotation is deprecated, please use USplineComponent functions for indirect access.")
+	FInterpCurveQuat& GetSplinePointsRotation() { return WarninglessSplineCurves().Rotation; }
+	
+	UE_DEPRECATED(5.6, "Mutable access to SplinePointsScale is deprecated, please use USplineComponent functions for indirect access.")
+	FInterpCurveVector& GetSplinePointsScale() { return WarninglessSplineCurves().Scale; }
+
+protected:
+
+	/** Returns the length of the specified spline segment up to the parametric value given */
+	ENGINE_API float GetSegmentLength(const int32 Index, const float Param = 1.0f) const;
+	
 private:
+	/**
+	 * This function takes on the roles of the previously existing functions GetPositionPointSafe, GetRotationPointSafe, and GetScalePointSafe.
+	 * Specifically, it makes sure that a provided index is valid by
+	 * 1) Specially handling Index == NumPoints if bClosedLoop is true in order to allow users to imagine the connecting point as an extra virtual point
+	 * 2) Preventing out of bound accesses by clamping the index. This clamping happens after 1).
+	 *
+	 * This function is only used where the 3 aforementioned functions were called. It should probably be used more widely and consistently in the future.
+	 */
+	int32 GetClampedIndex(int32 Index) const
+	{
+		// Treat Index == GetNumberOfSplinePoints() as Index == 0, then clamp to valid index range.
+		return FMath::Clamp(bClosedLoop && Index == GetNumberOfSplinePoints() ? 0 : Index, 0, GetNumberOfSplinePoints() - 1);
+	}
+	
 	/** The dummy value used for queries when there are no point in a spline */
 	static ENGINE_API const FInterpCurvePointVector DummyPointPosition;
 	static ENGINE_API const FInterpCurvePointQuat DummyPointRotation;
 	static ENGINE_API const FInterpCurvePointVector DummyPointScale;
 
-private:
+	// Remove when appending _DEPRECATED to SplineCurves.
+	FSplineCurves& WarninglessSplineCurves()
+	{
+PRAGMA_DISABLE_DEPRECATION_WARNINGS
+		return SplineCurves;
+PRAGMA_ENABLE_DEPRECATION_WARNINGS
+	}
+
+	// Remove when appending _DEPRECATED to SplineCurves.
+	const FSplineCurves& WarninglessSplineCurves() const
+	{
+PRAGMA_DISABLE_DEPRECATION_WARNINGS
+		return SplineCurves;
+PRAGMA_ENABLE_DEPRECATION_WARNINGS
+	}
 
 	/** Set the SplineCurves with the default shape i.e. a spline of two points (Used by default constructor) */
 	ENGINE_API void SetDefaultSpline();
 
-	/** Returns the length of the specified spline segment up to the parametric value given */
-	ENGINE_API float GetSegmentLength(const int32 Index, const float Param = 1.0f) const;
-
 	/** Returns the parametric value t which would result in a spline segment of the given length between S(0)...S(t) */
 	ENGINE_API float GetSegmentParamFromLength(const int32 Index, const float Length, const float SegmentLength) const;
 
@@ -852,54 +924,6 @@ private:
 	// Internal helper function called by ConvertSplineSegmentToPolyLine -- assumes the input is within a half-segment, so testing the distance to midpoint will be an accurate guide to subdivision
 	bool DivideSplineIntoPolylineRecursiveHelper(float StartDistanceAlongSpline, float EndDistanceAlongSpline, ESplineCoordinateSpace::Type CoordinateSpace, const float MaxSquareDistanceFromSpline, TArray<FVector>& OutPoints) const;
 
-	/** Returns a const reference to the specified position point, but gives back a dummy point if there are no points */
-	inline const FInterpCurvePointVector& GetPositionPointSafe(int32 PointIndex) const
-	{
-		const TArray<FInterpCurvePointVector>& Points = SplineCurves.Position.Points;
-		const int32 NumPoints = Points.Num();
-		if (NumPoints > 0)
-		{
-			const int32 ClampedIndex = (bClosedLoop && PointIndex >= NumPoints) ? 0 : FMath::Clamp(PointIndex, 0, NumPoints - 1);
-			return Points[ClampedIndex];
-		}
-		else
-		{
-			return DummyPointPosition;
-		}
-	}
-
-	/** Returns a const reference to the specified rotation point, but gives back a dummy point if there are no points */
-	inline const FInterpCurvePointQuat& GetRotationPointSafe(int32 PointIndex) const
-	{
-		const TArray<FInterpCurvePointQuat>& Points = SplineCurves.Rotation.Points;
-		const int32 NumPoints = Points.Num();
-		if (NumPoints > 0)
-		{
-			const int32 ClampedIndex = (bClosedLoop && PointIndex >= NumPoints) ? 0 : FMath::Clamp(PointIndex, 0, NumPoints - 1);
-			return Points[ClampedIndex];
-		}
-		else
-		{
-			return DummyPointRotation;
-		}
-	}
-
-	/** Returns a const reference to the specified scale point, but gives back a dummy point if there are no points */
-	inline const FInterpCurvePointVector& GetScalePointSafe(int32 PointIndex) const
-	{
-		const TArray<FInterpCurvePointVector>& Points = SplineCurves.Scale.Points;
-		const int32 NumPoints = Points.Num();
-		if (NumPoints > 0)
-		{
-			const int32 ClampedIndex = (bClosedLoop && PointIndex >= NumPoints) ? 0 : FMath::Clamp(PointIndex, 0, NumPoints - 1);
-			return Points[ClampedIndex];
-		}
-		else
-		{
-			return DummyPointScale;
-		}
-	}
-
 	// FSplineComponentVisualizer will access some private members when attempting to call NotifyPropertiesModified
 	friend class FSplineComponentVisualizer;
 };
@@ -948,7 +972,6 @@ public:
 ENGINE_API EInterpCurveMode ConvertSplinePointTypeToInterpCurveMode(ESplinePointType::Type SplinePointType);
 ENGINE_API ESplinePointType::Type ConvertInterpCurveModeToSplinePointType(EInterpCurveMode InterpCurveMode);
 
-
 // Deprecated method definitions
 inline void USplineComponent::AddSplineWorldPoint(const FVector& Position) { AddSplinePoint(Position, ESplineCoordinateSpace::World); }
 inline void USplineComponent::AddSplineLocalPoint(const FVector& Position) { AddSplinePoint(Position, ESplineCoordinateSpace::Local); }
diff --git a/Engine/Source/Runtime/Engine/Classes/Curves/Spline.h b/Engine/Source/Runtime/Engine/Classes/Curves/Spline.h
new file mode 100644
--- /dev/null
+++ b/Engine/Source/Runtime/Engine/Classes/Curves/Spline.h
@@ -0,0 +1,299 @@
+// Copyright Epic Games, Inc. All Rights Reserved.
+
+#pragma once
+
+#include "CoreMinimal.h"
+#include "Math/InterpCurve.h"
+
+#include "Spline.generated.h"
+
+struct FSplinePoint;
+struct FSplineCurves;
+
+/**
+ * A general purpose, reflected spline.
+ * The implementation can be configured at runtime.
+ */
+USTRUCT()
+struct FSpline
+{
+	GENERATED_BODY()
+	
+	ENGINE_API FSpline();
+	ENGINE_API FSpline(const FSpline& Other);
+	ENGINE_API FSpline& operator=(const FSpline& Other);
+	ENGINE_API ~FSpline();
+	ENGINE_API static FSpline FromSplineCurves(const FSplineCurves& InSpline);
+	
+	/* Control Point Index Interface */
+	
+	/** Adds point by parameter. The new point is inserted after all points with parameter values less than OR EQUAL to the new point's parameter.
+	 * Nick would like to change this behavior in the future. See UE-250236. */
+	ENGINE_API void AddPoint(const FSplinePoint& InPoint);
+	/** Adds point by index. Provided parameter is ignored. The new point is inserted before all points with parameter values greater than the new point's parameter. */
+	ENGINE_API void AddPointAtIndex(const FSplinePoint& InPoint, int32 Index);
+	ENGINE_API FSplinePoint GetPoint(const int32 Index) const;
+	ENGINE_API void RemovePoint(const int32 Index);
+	
+	ENGINE_API void SetLocation(int32 PointIndex, const FVector& InLocation);
+	ENGINE_API FVector GetLocation(const int32 Index) const;
+	
+	ENGINE_API void SetInTangent(const int32 Index, const FVector& InTangent);
+	ENGINE_API FVector GetInTangent(const int32 Index) const;
+	
+	ENGINE_API void SetOutTangent(const int32 Index, const FVector& OutTangent);
+	ENGINE_API FVector GetOutTangent(const int32 Index) const;
+
+	ENGINE_API void SetRotation(int32 PointIndex, const FQuat& InRotation);
+	ENGINE_API FQuat GetRotation(const int32 Index) const;
+	
+	ENGINE_API void SetScale(int32 PointIndex, const FVector& InScale);
+	ENGINE_API FVector GetScale(const int32 Index) const;
+	
+	ENGINE_API void SetSplinePointType(int32 PointIndex, EInterpCurveMode Type);
+	ENGINE_API EInterpCurveMode GetSplinePointType(int32 PointIndex) const;
+
+	ENGINE_API float GetSplinePointParameter(int32 PointIndex) const;
+	
+	/* Parameter Interface */
+
+	ENGINE_API FVector Evaluate(float Param) const;
+	ENGINE_API FVector EvaluateDerivative(float Param) const;
+	ENGINE_API FQuat EvaluateRotation(float Param) const;
+	ENGINE_API FVector EvaluateScale(float Param) const;
+	template<typename T>
+	T EvaluateAttribute(float Param, FName AttributeName) const
+	{
+		return T();
+	}
+
+	ENGINE_API float FindNearest(const FVector& InLocation) const;
+	ENGINE_API float FindNearestOnSegment(const FVector& InLocation, int32 SegmentIndex) const;
+	
+	/* Misc Interface */
+	
+	bool operator==(const FSpline& Other) const
+	{
+		if (CurrentImplementation == Other.CurrentImplementation)
+		{
+			if (LegacyData && Other.LegacyData)
+			{
+				return *LegacyData == *Other.LegacyData;
+			}
+
+			// if (NewData && Other.NewData)
+			// {
+			// 	return *NewData == *Other.NewData;
+			// }
+		}
+
+		return false;
+	}
+	
+	bool operator!=(const FSpline& Other) const
+	{
+		return !(*this == Other);
+	}
+	
+	friend ENGINE_API FArchive& operator<<(FArchive& Ar, FSpline& Spline)
+	{
+		Spline.Serialize(Ar);
+		return Ar;
+	}
+	bool Serialize(FArchive& Ar);
+
+	uint32 GetVersion() const { return Version; }
+	
+	/**
+	 * These functions will potentially be slow when not using the legacy spline.
+	 * This is because we will be forced to cache the underlying spline into
+	 * a data structure that allows us to return a const ref.
+	 */
+	ENGINE_API const FInterpCurveVector& GetSplinePointsPosition() const;
+	ENGINE_API const FInterpCurveQuat& GetSplinePointsRotation() const;
+	ENGINE_API const FInterpCurveVector& GetSplinePointsScale() const;
+	ENGINE_API const FInterpCurveFloat& GetSplineReparamTable() const;		// I may actually implement reparameterization with this interp curve for both spline backends.
+
+	/** Returns the length of the specified spline segment up to the parametric value given */
+	ENGINE_API float GetSegmentLength(const int32 Index, const float Param, bool bClosedLoop = false, const FVector& Scale3D = FVector(1.0f)) const;
+
+	/** Returns total length along this spline */
+	ENGINE_API float GetSplineLength() const;
+	
+	/** Returns the total number of control points on this spline. */
+	ENGINE_API int32 GetNumControlPoints() const;
+	
+	/** Reset the spline to an empty spline. */
+	ENGINE_API void Reset();
+	
+	/** Reset the rotation attribute channel to default values. */
+	ENGINE_API void ResetRotation();
+	
+	/** Reset the scale attribute channel to default values. */
+	ENGINE_API void ResetScale();
+	
+	/** Reset the spline to the default spline (2 points). */
+	ENGINE_API void ResetToDefault();
+
+	struct ENGINE_API FUpdateSplineParams
+	{
+		bool bClosedLoop = false;
+		bool bStationaryEndpoints = false;
+		int32 ReparamStepsPerSegment = 10;
+		bool bLoopPositionOverride = false;
+		float LoopPosition = 0.0f;
+		const FVector Scale3D = FVector(1.0f);
+	};
+	
+	/** Update the spline's internal data according to the passed-in params. */
+	ENGINE_API void UpdateSpline(const FUpdateSplineParams& InParams);
+
+private:
+	
+	class FLegacySpline
+	{
+		FInterpCurveVector PositionCurve;
+		FInterpCurveQuat RotationCurve;
+		FInterpCurveVector ScaleCurve;
+		FInterpCurveFloat ReparamTable;
+		
+	public:
+
+		ENGINE_API FLegacySpline();
+		static FLegacySpline FromSplineCurves(const FSplineCurves& InSpline);
+		
+		/* Control Point Index Interface */
+
+		void AddPoint(const FSplinePoint& InPoint);
+		void AddPointAtIndex(const FSplinePoint& InPoint, int32 Index);
+		FSplinePoint GetPoint(const int32 Index) const;
+		void RemovePoint(int32 Index);
+		
+		void SetLocation(int32 PointIndex, const FVector& InLocation);
+		FVector GetLocation(const int32 Index) const;
+		
+		void SetInTangent(const int32 Index, const FVector& InTangent);
+		FVector GetInTangent(const int32 Index) const;
+		
+		void SetOutTangent(const int32 Index, const FVector& OutTangent);
+		FVector GetOutTangent(const int32 Index) const;
+
+		void SetRotation(int32 PointIndex, const FQuat& InRotation);
+		FQuat GetRotation(const int32 Index) const;
+		
+		void SetScale(int32 PointIndex, const FVector& InScale);
+		FVector GetScale(const int32 Index) const;
+		
+		void SetSplinePointType(int32 PointIndex, EInterpCurveMode Type);
+		EInterpCurveMode GetSplinePointType(int32 PointIndex) const;
+
+		float GetSplinePointParameter(int32 PointIndex) const;
+		
+		/* Parameter Interface */
+
+		FVector Evaluate(float Param) const;
+		FVector EvaluateDerivative(float Param) const;
+		FQuat EvaluateRotation(float Param) const;
+		FVector EvaluateScale(float Param) const;
+		float FindNearest(const FVector& InLocation) const;
+		float FindNearestOnSegment(const FVector& InLocation, int32 SegmentIndex) const;
+		
+		/* Misc Interface */
+		
+		bool operator==(const FLegacySpline& Other) const
+		{
+			return PositionCurve == Other.PositionCurve && RotationCurve == Other.RotationCurve && ScaleCurve == Other.ScaleCurve;
+		}
+
+		bool operator!=(const FLegacySpline& Other) const
+		{
+			return !(*this == Other);
+		}
+		
+		friend ENGINE_API FArchive& operator<<(FArchive& Ar, FLegacySpline& Spline)
+		{
+			Spline.Serialize(Ar);
+			return Ar;
+		}
+		bool Serialize(FArchive& Ar);
+		
+		FInterpCurveVector& GetSplinePointsPosition() { return PositionCurve; }
+		const FInterpCurveVector& GetSplinePointsPosition() const { return PositionCurve; }
+		FInterpCurveQuat& GetSplinePointsRotation() { return RotationCurve; }
+		const FInterpCurveQuat& GetSplinePointsRotation() const { return RotationCurve; }
+		FInterpCurveVector& GetSplinePointsScale() { return ScaleCurve; }
+		const FInterpCurveVector& GetSplinePointsScale() const { return ScaleCurve; }
+		FInterpCurveFloat& GetSplineReparamTable() { return ReparamTable; }
+		const FInterpCurveFloat& GetSplineReparamTable() const { return ReparamTable; }
+
+		/** Returns the length of the specified spline segment up to the parametric value given */
+		float GetSegmentLength(const int32 Index, const float Param, bool bClosedLoop = false, const FVector& Scale3D = FVector(1.0f)) const;
+
+		/** Returns total length along this spline */
+		float GetSplineLength() const;
+		
+		/** Returns the total number of control points on this spline. */
+		int32 GetNumControlPoints() const;
+		
+		/** Reset the spline to an empty spline. */
+		void Reset();
+		
+		/** Reset the rotation attribute channel to default values. */
+		void ResetRotation();
+		
+		/** Reset the scale attribute channel to default values. */
+		void ResetScale();
+		
+		/** Reset the spline to the default spline (2 points). */
+		void ResetToDefault();
+		
+		/** Update the spline's internal data according to the passed-in params. */
+		void UpdateSpline(const FUpdateSplineParams& InParams);
+	};
+	
+	static inline const FLegacySpline DummySpline;
+	
+	// Used for upgrade logic in spline component.
+	// Not ideal, but allows us to automatically populate the proxy
+	// at serialize time when we might otherwise not be able to.
+	friend class USplineComponent;
+#if WITH_EDITOR
+	uint8 PreviousImplementation;
+#endif
+	uint8 CurrentImplementation;
+	
+	uint32 Version;
+	
+	struct {										// Invalid when CurrentImplementation is 0 or 3.
+		TSharedPtr<FLegacySpline> LegacyData;		// Valid when CurrentImplementation is 1.
+		//TSharedPtr<FHermiteSpline> NewData;		// Valid when CurrentImplementation is 2.
+	};
+	
+	bool IsEnabled() const { return CurrentImplementation != 0; }
+	bool IsLegacy() const { return CurrentImplementation == 1; }
+	bool IsNew() const { return CurrentImplementation == 2; }
+
+#if WITH_EDITOR
+	bool WasEnabled() const { return PreviousImplementation != 0; }
+	bool WasLegacy() const { return PreviousImplementation == 1; }
+	bool WasNew() const { return PreviousImplementation == 2; }
+
+	/** Called when the implementation is changed at editor time due to a console command. */
+	void OnSplineImplementationChanged();
+	FDelegateHandle OnSplineImplementationChangedHandle;
+#endif
+};
+
+template<>
+struct TStructOpsTypeTraits<FSpline> : public TStructOpsTypeTraitsBase2<FSpline>
+{
+	enum
+	{
+		WithSerializer				= true, // Enables the use of a custom Serialize method.
+		WithIdenticalViaEquality	= true, // Enables the use of a custom equality operator.
+
+		// todo: need to implement for copy/paste and duplicate.
+		//WithExportTextItem			= true, // Enables the use of a custom ExportTextItem method.
+		//WithImportTextItem			= true, // Enables the use of a custom ImportTextItem method.
+	};
+};
\ No newline at end of file
diff --git a/Engine/Source/Runtime/Engine/Private/Components/SplineComponent.cpp b/Engine/Source/Runtime/Engine/Private/Components/SplineComponent.cpp
--- a/Engine/Source/Runtime/Engine/Private/Components/SplineComponent.cpp
+++ b/Engine/Source/Runtime/Engine/Private/Components/SplineComponent.cpp
@@ -30,10 +30,35 @@ const FInterpCurvePointVector USplineComponent::DummyPointPosition(0.0f, FVector
 const FInterpCurvePointQuat USplineComponent::DummyPointRotation(0.0f, FQuat::Identity);
 const FInterpCurvePointVector USplineComponent::DummyPointScale(0.0f, FVector::OneVector);
 
+namespace UE::SplineComponent
+{
+
+bool GUseSplineCurves = true;
+FAutoConsoleVariableRef CVarEnableDaySequenceModifierDeviceCustomization(
+	TEXT("SplineComponent.UseSplineCurves"),
+	GUseSplineCurves,
+	TEXT("When true, SplineCurves is the 'source of truth'.")
+);
+	
+}
+
 USplineMetadata::USplineMetadata(const FObjectInitializer& ObjectInitializer)
 	: Super(ObjectInitializer)
 {
+}
 
+FName USplineComponent::GetSplinePropertyName()
+{
+	if (UE::SplineComponent::GUseSplineCurves)
+	{
+PRAGMA_DISABLE_DEPRECATION_WARNINGS
+		return GET_MEMBER_NAME_CHECKED(USplineComponent, SplineCurves);
+PRAGMA_ENABLE_DEPRECATION_WARNINGS
+	}
+	else
+	{
+		return GET_MEMBER_NAME_CHECKED(USplineComponent, Spline);
+	}
 }
 
 USplineComponent::USplineComponent(const FObjectInitializer& ObjectInitializer)
@@ -73,10 +98,10 @@ USplineComponent::USplineComponent(const FObjectInitializer& ObjectInitializer)
 
 #if WITH_EDITORONLY_DATA
 	// Set these deprecated values up so that old assets with default values load correctly (and are subsequently upgraded during Serialize)
-	SplineInfo_DEPRECATED = SplineCurves.Position;
-	SplineRotInfo_DEPRECATED = SplineCurves.Rotation;
-	SplineScaleInfo_DEPRECATED = SplineCurves.Scale;
-	SplineReparamTable_DEPRECATED = SplineCurves.ReparamTable;
+	SplineInfo_DEPRECATED = Spline.GetSplinePointsPosition();
+	SplineRotInfo_DEPRECATED = Spline.GetSplinePointsRotation();
+	SplineScaleInfo_DEPRECATED = Spline.GetSplinePointsScale();
+	SplineReparamTable_DEPRECATED = Spline.GetSplineReparamTable();
 #endif
 }
 
@@ -104,19 +129,26 @@ void USplineComponent::ResetToDefault()
 #endif
 }
 
+bool USplineComponent::CanResetToDefault() const
+{
+	return Spline != CastChecked<USplineComponent>(GetArchetype())->Spline;
+}
+
 void USplineComponent::SetDefaultSpline()
 {
-	SplineCurves.Position.Points.Reset(10);
-	SplineCurves.Rotation.Points.Reset(10);
-	SplineCurves.Scale.Points.Reset(10);
+	Spline.ResetToDefault();
 
-	SplineCurves.Position.Points.Emplace(0.0f, FVector(0, 0, 0), FVector::ZeroVector, FVector::ZeroVector, CIM_CurveAuto);
-	SplineCurves.Rotation.Points.Emplace(0.0f, FQuat::Identity, FQuat::Identity, FQuat::Identity, CIM_CurveAuto);
-	SplineCurves.Scale.Points.Emplace(0.0f, FVector(1.0f), FVector::ZeroVector, FVector::ZeroVector, CIM_CurveAuto);
+	WarninglessSplineCurves().Position.Points.Reset(10);
+    WarninglessSplineCurves().Rotation.Points.Reset(10);
+    WarninglessSplineCurves().Scale.Points.Reset(10);
 
-	SplineCurves.Position.Points.Emplace(1.0f, FVector(100, 0, 0), FVector::ZeroVector, FVector::ZeroVector, CIM_CurveAuto);
-	SplineCurves.Rotation.Points.Emplace(1.0f, FQuat::Identity, FQuat::Identity, FQuat::Identity, CIM_CurveAuto);
-	SplineCurves.Scale.Points.Emplace(1.0f, FVector(1.0f), FVector::ZeroVector, FVector::ZeroVector, CIM_CurveAuto);
+    WarninglessSplineCurves().Position.Points.Emplace(0.0f, FVector(0, 0, 0), FVector::ZeroVector, FVector::ZeroVector, CIM_CurveAuto);
+    WarninglessSplineCurves().Rotation.Points.Emplace(0.0f, FQuat::Identity, FQuat::Identity, FQuat::Identity, CIM_CurveAuto);
+    WarninglessSplineCurves().Scale.Points.Emplace(0.0f, FVector(1.0f), FVector::ZeroVector, FVector::ZeroVector, CIM_CurveAuto);
+
+    WarninglessSplineCurves().Position.Points.Emplace(1.0f, FVector(100, 0, 0), FVector::ZeroVector, FVector::ZeroVector, CIM_CurveAuto);
+    WarninglessSplineCurves().Rotation.Points.Emplace(1.0f, FQuat::Identity, FQuat::Identity, FQuat::Identity, CIM_CurveAuto);
+    WarninglessSplineCurves().Scale.Points.Emplace(1.0f, FVector(1.0f), FVector::ZeroVector, FVector::ZeroVector, CIM_CurveAuto);
 }
 
 void USplineComponent::GetLifetimeReplicatedProps(TArray<FLifetimeProperty>& OutLifetimeProps) const
@@ -156,7 +188,6 @@ ESplinePointType::Type ConvertInterpCurveModeToSplinePointType(EInterpCurveMode
 	}
 }
 
-
 void USplineComponent::Serialize(FArchive& Ar)
 {
 	Super::Serialize(Ar);
@@ -164,14 +195,24 @@ void USplineComponent::Serialize(FArchive& Ar)
 	Ar.UsingCustomVersion(FEditorObjectVersion::GUID);
 
 #if WITH_EDITORONLY_DATA
-	// Move points to new properties
-	if (Ar.CustomVer(FEditorObjectVersion::GUID) < FEditorObjectVersion::SplineComponentCurvesInStruct &&
-		Ar.IsLoading())
+	if (Ar.IsLoading())
 	{
-		SplineCurves.Position = SplineInfo_DEPRECATED;
-		SplineCurves.Rotation = SplineRotInfo_DEPRECATED;
-		SplineCurves.Scale = SplineScaleInfo_DEPRECATED;
-		SplineCurves.ReparamTable = SplineReparamTable_DEPRECATED;
+		// Move points into SplineCurves
+		if (Ar.CustomVer(FEditorObjectVersion::GUID) < FEditorObjectVersion::SplineComponentCurvesInStruct)
+		{
+			WarninglessSplineCurves().Position = SplineInfo_DEPRECATED;
+			WarninglessSplineCurves().Rotation = SplineRotInfo_DEPRECATED;
+			WarninglessSplineCurves().Scale = SplineScaleInfo_DEPRECATED;
+			WarninglessSplineCurves().ReparamTable = SplineReparamTable_DEPRECATED;
+		}
+
+#if WITH_EDITOR
+		// WasEnabled() is editor only, that data gets stripped from packaged builds.
+		if (!Spline.WasEnabled())
+		{
+			SynchronizeSplines();
+		}
+#endif
 	}
 #endif
 
@@ -179,25 +220,30 @@ void USplineComponent::Serialize(FArchive& Ar)
 	const FPackageFileVersion ArchiveUEVersion = Ar.UEVer();
 	if (ArchiveUEVersion < VER_UE4_INTERPCURVE_SUPPORTS_LOOPING)
 	{
-		int32 NumPoints = SplineCurves.Position.Points.Num();
+		int32 LegacyNumPoints = WarninglessSplineCurves().Position.Points.Num();
 
 		// The start point is no longer cloned as the endpoint when the spline is looped, so remove the extra endpoint if present
-		const bool bHasExtraEndpoint = bClosedLoop && (SplineCurves.Position.Points[0].OutVal == SplineCurves.Position.Points[NumPoints - 1].OutVal);
-
-		if (bHasExtraEndpoint)
+		if (bClosedLoop && ( GetLocationAtSplinePoint(0, ESplineCoordinateSpace::Local) == GetLocationAtSplinePoint(GetNumberOfSplinePoints() - 1, ESplineCoordinateSpace::Local)))
 		{
-			SplineCurves.Position.Points.RemoveAt(NumPoints - 1, EAllowShrinking::No);
-			NumPoints--;
+			Spline.RemovePoint(Spline.GetNumControlPoints() - 1);
+			WarninglessSplineCurves().Position.Points.RemoveAt(LegacyNumPoints - 1, EAllowShrinking::No);
+			LegacyNumPoints--;
 		}
 
-		// Fill the other two splines with some defaults
-		SplineCurves.Rotation.Points.Reset(NumPoints);
-		SplineCurves.Scale.Points.Reset(NumPoints);
-		for (int32 Count = 0; Count < NumPoints; Count++)
+		// Fill the other two splines with some defaults.
+		
+		WarninglessSplineCurves().Rotation.Points.Reset(LegacyNumPoints);
+		WarninglessSplineCurves().Scale.Points.Reset(LegacyNumPoints);
+		for (int32 Count = 0; Count < LegacyNumPoints; Count++)
 		{
-			SplineCurves.Rotation.Points.Emplace(0.0f, FQuat::Identity, FQuat::Identity, FQuat::Identity, CIM_CurveAuto);
-			SplineCurves.Scale.Points.Emplace(0.0f, FVector(1.0f), FVector::ZeroVector, FVector::ZeroVector, CIM_CurveAuto);
+			WarninglessSplineCurves().Rotation.Points.Emplace(0.0f, FQuat::Identity, FQuat::Identity, FQuat::Identity, CIM_CurveAuto);
+			WarninglessSplineCurves().Scale.Points.Emplace(0.0f, FVector(1.0f), FVector::ZeroVector, FVector::ZeroVector, CIM_CurveAuto);
 		}
+		
+		// Importantly, these 2 functions guarantee that 1 rotation and scale value exists for each position,
+		// but they are default values.
+		Spline.ResetRotation();
+		Spline.ResetScale();
 
 		UpdateSpline();
 	}
@@ -259,14 +305,18 @@ void FSplineCurves::UpdateSpline(bool bClosedLoop, bool bStationaryEndpoints, in
 		}
 		AccumulatedLength += GetSegmentLength(SegmentIndex, 1.0f, bClosedLoop, Scale3D);
 	}
-
 	ReparamTable.Points.Emplace(AccumulatedLength, static_cast<float>(NumSegments), 0.0f, 0.0f, CIM_Linear);
 	++Version;
 }
 
 void USplineComponent::UpdateSpline()
 {
-	SplineCurves.UpdateSpline(bClosedLoop, bStationaryEndpoints, ReparamStepsPerSegment, bLoopPositionOverride, LoopPosition, GetComponentTransform().GetScale3D());
+	FSpline::FUpdateSplineParams Params{bClosedLoop, bStationaryEndpoints, ReparamStepsPerSegment, bLoopPositionOverride, LoopPosition, GetComponentTransform().GetScale3D()};
+	WarninglessSplineCurves().UpdateSpline(Params.bClosedLoop, Params.bStationaryEndpoints, Params.ReparamStepsPerSegment, Params.bLoopPositionOverride, Params.LoopPosition, Params.Scale3D);
+	Spline.UpdateSpline(Params);
+
+	// todo: make Spline replicate
+	//MARK_PROPERTY_DIRTY_FROM_NAME(USplineComponent, Spline, this);
 	MARK_PROPERTY_DIRTY_FROM_NAME(USplineComponent, SplineCurves, this);
 	MARK_PROPERTY_DIRTY_FROM_NAME(USplineComponent, bClosedLoop, this);
 	MARK_PROPERTY_DIRTY_FROM_NAME(USplineComponent, bStationaryEndpoints, this);
@@ -290,7 +340,6 @@ void USplineComponent::SetOverrideConstructionScript(bool bInOverride)
 	bSplineHasBeenEdited = bInOverride;
 }
 
-
 float FSplineCurves::GetSegmentLength(const int32 Index, const float Param, bool bClosedLoop, const FVector& Scale3D) const
 {
 	const int32 NumPoints = Position.Points.Num();
@@ -361,7 +410,9 @@ float FSplineCurves::GetSegmentLength(const int32 Index, const float Param, bool
 
 float USplineComponent::GetSegmentLength(const int32 Index, const float Param) const
 {
-	return SplineCurves.GetSegmentLength(Index, Param, bClosedLoop, GetComponentTransform().GetScale3D());
+	return UE::SplineComponent::GUseSplineCurves
+		? WarninglessSplineCurves().GetSegmentLength(Index, Param, bClosedLoop, GetComponentTransform().GetScale3D())
+		: Spline.GetSegmentLength(Index, Param, bClosedLoop, GetComponentTransform().GetScale3D());
 }
 
 float USplineComponent::GetSegmentParamFromLength(const int32 Index, const float Length, const float SegmentLength) const
@@ -390,7 +441,10 @@ float USplineComponent::GetSegmentParamFromLength(const int32 Index, const float
 
 	// TODO: check if this works OK with delta InVal != 1.0f
 
-	const int32 NumPoints = SplineCurves.Position.Points.Num();
+	int32 NumPoints = UE::SplineComponent::GUseSplineCurves
+		? WarninglessSplineCurves().Position.Points.Num()
+		: Spline.GetNumControlPoints();
+	
 	const int32 LastPoint = NumPoints - 1;
 
 	check(Index >= 0 && ((bClosedLoop && Index < NumPoints) || (!bClosedLoop && Index < LastPoint)));
@@ -401,7 +455,10 @@ float USplineComponent::GetSegmentParamFromLength(const int32 Index, const float
 	// two iterations of Newton-Raphson is enough
 	for (int32 Iteration = 0; Iteration < 2; ++Iteration)
 	{
-		float TangentMagnitude = SplineCurves.Position.EvalDerivative(Index + Param, FVector::ZeroVector).Size();
+		float TangentMagnitude = UE::SplineComponent::GUseSplineCurves
+			? WarninglessSplineCurves().Position.EvalDerivative(Index + Param, FVector::ZeroVector).Size()
+			: Spline.EvaluateDerivative(Index + Param).Size();
+		
 		if (TangentMagnitude > 0.0f)
 		{
 			Param -= (GetSegmentLength(Index, Param) - Length) / TangentMagnitude;
@@ -412,10 +469,11 @@ float USplineComponent::GetSegmentParamFromLength(const int32 Index, const float
 	return Param;
 }
 
-
 FVector USplineComponent::GetLocationAtSplineInputKey(float InKey, ESplineCoordinateSpace::Type CoordinateSpace) const
 {
-	FVector Location = SplineCurves.Position.Eval(InKey, FVector::ZeroVector);
+	FVector Location = UE::SplineComponent::GUseSplineCurves
+		? WarninglessSplineCurves().Position.Eval(InKey, FVector::ZeroVector)
+		: Spline.Evaluate(InKey);
 
 	if (CoordinateSpace == ESplineCoordinateSpace::World)
 	{
@@ -425,10 +483,11 @@ FVector USplineComponent::GetLocationAtSplineInputKey(float InKey, ESplineCoordi
 	return Location;
 }
 
-
 FVector USplineComponent::GetTangentAtSplineInputKey(float InKey, ESplineCoordinateSpace::Type CoordinateSpace) const
 {
-	FVector Tangent = SplineCurves.Position.EvalDerivative(InKey, FVector::ZeroVector);
+	FVector Tangent = UE::SplineComponent::GUseSplineCurves
+		? WarninglessSplineCurves().Position.EvalDerivative(InKey, FVector::ZeroVector)
+		: Spline.EvaluateDerivative(InKey);
 
 	if (CoordinateSpace == ESplineCoordinateSpace::World)
 	{
@@ -438,10 +497,11 @@ FVector USplineComponent::GetTangentAtSplineInputKey(float InKey, ESplineCoordin
 	return Tangent;
 }
 
-
 FVector USplineComponent::GetDirectionAtSplineInputKey(float InKey, ESplineCoordinateSpace::Type CoordinateSpace) const
 {
-	FVector Direction = SplineCurves.Position.EvalDerivative(InKey, FVector::ZeroVector).GetSafeNormal();
+	FVector Direction = UE::SplineComponent::GUseSplineCurves
+		? WarninglessSplineCurves().Position.EvalDerivative(InKey, FVector::ZeroVector).GetSafeNormal()
+		: Spline.EvaluateDerivative(InKey).GetSafeNormal();
 
 	if (CoordinateSpace == ESplineCoordinateSpace::World)
 	{
@@ -452,19 +512,22 @@ FVector USplineComponent::GetDirectionAtSplineInputKey(float InKey, ESplineCoord
 	return Direction;
 }
 
-
 FRotator USplineComponent::GetRotationAtSplineInputKey(float InKey, ESplineCoordinateSpace::Type CoordinateSpace) const
 {
 	return GetQuaternionAtSplineInputKey(InKey, CoordinateSpace).Rotator();
 }
 
-
 FQuat USplineComponent::GetQuaternionAtSplineInputKey(float InKey, ESplineCoordinateSpace::Type CoordinateSpace) const
 {
-	FQuat Quat = SplineCurves.Rotation.Eval(InKey, FQuat::Identity);
+	FQuat Quat = UE::SplineComponent::GUseSplineCurves
+		? WarninglessSplineCurves().Rotation.Eval(InKey, FQuat::Identity)
+		: Spline.EvaluateRotation(InKey);
 	Quat.Normalize();
 
-	const FVector Direction = SplineCurves.Position.EvalDerivative(InKey, FVector::ZeroVector).GetSafeNormal();
+	FVector Direction = UE::SplineComponent::GUseSplineCurves
+		? WarninglessSplineCurves().Position.EvalDerivative(InKey, FVector::ZeroVector).GetSafeNormal()
+		: Spline.EvaluateDerivative(InKey).GetSafeNormal();
+	
 	const FVector UpVector = Quat.RotateVector(DefaultUpVector);
 
 	FQuat Rot = (FRotationMatrix::MakeFromXZ(Direction, UpVector)).ToQuat();
@@ -477,7 +540,6 @@ FQuat USplineComponent::GetQuaternionAtSplineInputKey(float InKey, ESplineCoordi
 	return Rot;
 }
 
-
 FVector USplineComponent::GetUpVectorAtSplineInputKey(float InKey, ESplineCoordinateSpace::Type CoordinateSpace) const
 {
 	const FQuat Quat = GetQuaternionAtSplineInputKey(InKey, ESplineCoordinateSpace::Local);
@@ -491,7 +553,6 @@ FVector USplineComponent::GetUpVectorAtSplineInputKey(float InKey, ESplineCoordi
 	return UpVector;
 }
 
-
 FVector USplineComponent::GetRightVectorAtSplineInputKey(float InKey, ESplineCoordinateSpace::Type CoordinateSpace) const
 {
 	const FQuat Quat = GetQuaternionAtSplineInputKey(InKey, ESplineCoordinateSpace::Local);
@@ -505,7 +566,6 @@ FVector USplineComponent::GetRightVectorAtSplineInputKey(float InKey, ESplineCoo
 	return RightVector;
 }
 
-
 FTransform USplineComponent::GetTransformAtSplineInputKey(float InKey, ESplineCoordinateSpace::Type CoordinateSpace, bool bUseScale) const
 {
 	const FVector Location(GetLocationAtSplineInputKey(InKey, ESplineCoordinateSpace::Local));
@@ -522,7 +582,6 @@ FTransform USplineComponent::GetTransformAtSplineInputKey(float InKey, ESplineCo
 	return Transform;
 }
 
-
 float USplineComponent::GetRollAtSplineInputKey(float InKey, ESplineCoordinateSpace::Type CoordinateSpace) const
 {
 	return GetRotationAtSplineInputKey(InKey, CoordinateSpace).Roll;
@@ -531,14 +590,17 @@ float USplineComponent::GetRollAtSplineInputKey(float InKey, ESplineCoordinateSp
 
 FVector USplineComponent::GetScaleAtSplineInputKey(float InKey) const
 {
-	const FVector Scale = SplineCurves.Scale.Eval(InKey, FVector(1.0f));
-	return Scale;
+	return UE::SplineComponent::GUseSplineCurves
+		? WarninglessSplineCurves().Scale.Eval(InKey, FVector(1.0f))
+		: Spline.EvaluateScale(InKey);
 }
 
-
 float USplineComponent::GetDistanceAlongSplineAtSplineInputKey(float InKey) const
 {
-	const int32 NumPoints = SplineCurves.Position.Points.Num();
+	const int32 NumPoints = UE::SplineComponent::GUseSplineCurves
+		? WarninglessSplineCurves().Position.Points.Num()
+		: Spline.GetNumControlPoints();
+	
 	const int32 NumSegments = bClosedLoop ? NumPoints : NumPoints - 1;
 
 	if ((InKey >= 0) && (InKey < NumSegments))
@@ -546,12 +608,16 @@ float USplineComponent::GetDistanceAlongSplineAtSplineInputKey(float InKey) cons
 		const int32 PointIndex = FMath::FloorToInt(InKey);
 		const float Fraction = InKey - PointIndex;
 		const int32 ReparamPointIndex = PointIndex * ReparamStepsPerSegment;
-		const float Distance = SplineCurves.ReparamTable.Points[ReparamPointIndex].InVal;
+		
+		const float Distance = UE::SplineComponent::GUseSplineCurves
+			? WarninglessSplineCurves().ReparamTable.Points[ReparamPointIndex].InVal
+			: Spline.GetSplineReparamTable().Points[ReparamPointIndex].InVal;
+		
 		return Distance + GetSegmentLength(PointIndex, Fraction);
 	}
 	else if (InKey >= NumSegments)
 	{
-		return SplineCurves.GetSplineLength();
+		return GetSplineLength();
 	}
 
 	return 0.0f;
@@ -561,7 +627,11 @@ float USplineComponent::GetDistanceAlongSplineAtLocation(const FVector& InLocati
 {
 	const FVector LocalLocation = (CoordinateSpace == ESplineCoordinateSpace::World) ? GetComponentTransform().InverseTransformPosition(InLocation) : InLocation;
 	float Dummy;
-	float Key = SplineCurves.Position.FindNearest(LocalLocation, Dummy);
+	
+	const float Key = UE::SplineComponent::GUseSplineCurves
+		? WarninglessSplineCurves().Position.FindNearest(LocalLocation, Dummy)
+		: Spline.FindNearest(LocalLocation);
+	
 	return GetDistanceAlongSplineAtSplineInputKey(Key);
 }
 
@@ -600,7 +670,6 @@ void USplineComponent::SetClosedLoop(bool bInClosedLoop, bool bUpdateSpline)
 	}
 }
 
-
 void USplineComponent::SetClosedLoopAtPosition(bool bInClosedLoop, float Key, bool bUpdateSpline)
 {
 	bClosedLoop = bInClosedLoop;
@@ -613,13 +682,11 @@ void USplineComponent::SetClosedLoopAtPosition(bool bInClosedLoop, float Key, bo
 	}
 }
 
-
 bool USplineComponent::IsClosedLoop() const
 {
 	return bClosedLoop;
 }
 
-
 void USplineComponent::SetUnselectedSplineSegmentColor(const FLinearColor& Color)
 {
 #if WITH_EDITORONLY_DATA
@@ -642,22 +709,21 @@ void USplineComponent::SetTangentColor(const FLinearColor& Color)
 #endif
 }
 
-
 void USplineComponent::SetDrawDebug(bool bShow)
 {
 	bDrawDebug = bShow;
 	MarkRenderStateDirty();
 }
 
-
 void USplineComponent::ClearSplinePoints(bool bUpdateSpline)
 {
-	SplineCurves.Position.Points.Reset();
-	SplineCurves.Rotation.Points.Reset();
-	SplineCurves.Scale.Points.Reset();
+	Spline.Reset();
 
-	USplineMetadata* Metadata = GetSplinePointsMetadata();
-	if (Metadata)
+	WarninglessSplineCurves().Position.Points.Reset();
+	WarninglessSplineCurves().Rotation.Points.Reset();
+	WarninglessSplineCurves().Scale.Points.Reset();
+	
+	if (USplineMetadata* Metadata = GetSplinePointsMetadata())
 	{
 		Metadata->Reset(0);
 	}
@@ -668,7 +734,6 @@ void USplineComponent::ClearSplinePoints(bool bUpdateSpline)
 	}
 }
 
-
 static int32 UpperBound(const TArray<FInterpCurvePoint<FVector>>& SplinePoints, float Value)
 {
 	int32 Count = SplinePoints.Num();
@@ -691,45 +756,46 @@ static int32 UpperBound(const TArray<FInterpCurvePoint<FVector>>& SplinePoints,
 	return First;
 }
 
-
 void USplineComponent::AddPoint(const FSplinePoint& InSplinePoint, bool bUpdateSpline)
 {
-	const int32 Index = UpperBound(SplineCurves.Position.Points, InSplinePoint.InputKey);
-
-	SplineCurves.Position.Points.Insert(FInterpCurvePoint<FVector>(
-		InSplinePoint.InputKey,
-		InSplinePoint.Position,
-		InSplinePoint.ArriveTangent,
-		InSplinePoint.LeaveTangent,
-		ConvertSplinePointTypeToInterpCurveMode(InSplinePoint.Type)
-		),
-		Index);
-
-	SplineCurves.Rotation.Points.Insert(FInterpCurvePoint<FQuat>(
-		InSplinePoint.InputKey,
-		InSplinePoint.Rotation.Quaternion(),
-		FQuat::Identity,
-		FQuat::Identity,
-		CIM_CurveAuto
-		),
-		Index);
-
-	SplineCurves.Scale.Points.Insert(FInterpCurvePoint<FVector>(
-		InSplinePoint.InputKey,
-		InSplinePoint.Scale,
-		FVector::ZeroVector,
-		FVector::ZeroVector,
-		CIM_CurveAuto
-		),
-		Index);
+	Spline.AddPoint(InSplinePoint);
 
-	USplineMetadata* Metadata = GetSplinePointsMetadata();
-	if (Metadata)
+	{
+		int32 Index = UpperBound(WarninglessSplineCurves().Position.Points, InSplinePoint.InputKey);
+		
+        WarninglessSplineCurves().Position.Points.Insert(FInterpCurvePoint<FVector>(
+        	InSplinePoint.InputKey,
+        	InSplinePoint.Position,
+        	InSplinePoint.ArriveTangent,
+        	InSplinePoint.LeaveTangent,
+        	ConvertSplinePointTypeToInterpCurveMode(InSplinePoint.Type)
+        ), Index);
+    
+        WarninglessSplineCurves().Rotation.Points.Insert(FInterpCurvePoint<FQuat>(
+        	InSplinePoint.InputKey,
+        	InSplinePoint.Rotation.Quaternion(),
+        	FQuat::Identity,
+        	FQuat::Identity,
+        	CIM_CurveAuto
+        ), Index);
+    
+        WarninglessSplineCurves().Scale.Points.Insert(FInterpCurvePoint<FVector>(
+        	InSplinePoint.InputKey,
+        	InSplinePoint.Scale,
+        	FVector::ZeroVector,
+        	FVector::ZeroVector,
+        	CIM_CurveAuto
+        ), Index);
+	}
+	
+	if (USplineMetadata* Metadata = GetSplinePointsMetadata())
 	{
 		Metadata->AddPoint(InSplinePoint.InputKey);
 	}
 
-	if (bLoopPositionOverride && LoopPosition <= SplineCurves.Position.Points.Last().InVal)
+	const float LastPointKey = GetInputKeyValueAtSplinePoint(GetNumberOfSplinePoints() - 1);
+	
+	if (bLoopPositionOverride && LoopPosition <= LastPointKey)
 	{
 		bLoopPositionOverride = false;
 	}
@@ -740,14 +806,13 @@ void USplineComponent::AddPoint(const FSplinePoint& InSplinePoint, bool bUpdateS
 	}
 }
 
-
 void USplineComponent::AddPoints(const TArray<FSplinePoint>& InSplinePoints, bool bUpdateSpline)
 {
-	const int32 NumPoints = SplineCurves.Position.Points.Num() + InSplinePoints.Num();
+	const int32 NumPoints = WarninglessSplineCurves().Position.Points.Num() + InSplinePoints.Num();
 	// Position, Rotation, and Scale will all grow together.
-	SplineCurves.Position.Points.Reserve(NumPoints);
-	SplineCurves.Rotation.Points.Reserve(NumPoints);
-	SplineCurves.Scale.Points.Reserve(NumPoints);
+	WarninglessSplineCurves().Position.Points.Reserve(NumPoints);
+	WarninglessSplineCurves().Rotation.Points.Reserve(NumPoints);
+	WarninglessSplineCurves().Scale.Points.Reserve(NumPoints);
 
 	for (const auto& SplinePoint : InSplinePoints)
 	{
@@ -768,18 +833,30 @@ void USplineComponent::AddSplinePoint(const FVector& Position, ESplineCoordinate
 
 	// Add the spline point at the end of the array, adding 1.0 to the current last input key.
 	// This continues the former behavior in which spline points had to be separated by an interval of 1.0.
-	const float InKey = SplineCurves.Position.Points.Num() ? SplineCurves.Position.Points.Last().InVal + 1.0f : 0.0f;
-
-	SplineCurves.Position.Points.Emplace(InKey, TransformedPosition, FVector::ZeroVector, FVector::ZeroVector, CIM_CurveAuto);
-	SplineCurves.Rotation.Points.Emplace(InKey, FQuat::Identity, FQuat::Identity, FQuat::Identity, CIM_CurveAuto);
-	SplineCurves.Scale.Points.Emplace(InKey, FVector(1.0f), FVector::ZeroVector, FVector::ZeroVector, CIM_CurveAuto);
-	USplineMetadata* Metadata = GetSplinePointsMetadata();
-	if (Metadata)
+	const float InKey = GetNumberOfSplinePoints() > 0
+		? GetInputKeyValueAtSplinePoint(GetNumberOfSplinePoints() - 1) + 1.0f
+		: 0.0f;
+
+	FSplinePoint NewPoint;
+	NewPoint.InputKey = InKey;
+	NewPoint.Position = TransformedPosition;
+	NewPoint.ArriveTangent = FVector::ZeroVector;
+	NewPoint.LeaveTangent = FVector::ZeroVector;
+	NewPoint.Rotation = FQuat::Identity.Rotator();
+	NewPoint.Scale = FVector(1.f);
+	NewPoint.Type = ConvertInterpCurveModeToSplinePointType(CIM_CurveAuto);
+
+	Spline.AddPoint(NewPoint);
+
+	WarninglessSplineCurves().Position.Points.Emplace(InKey, TransformedPosition, FVector::ZeroVector, FVector::ZeroVector, CIM_CurveAuto);
+	WarninglessSplineCurves().Rotation.Points.Emplace(InKey, FQuat::Identity, FQuat::Identity, FQuat::Identity, CIM_CurveAuto);
+	WarninglessSplineCurves().Scale.Points.Emplace(InKey, FVector(1.0f), FVector::ZeroVector, FVector::ZeroVector, CIM_CurveAuto);
+	
+	if (USplineMetadata* Metadata = GetSplinePointsMetadata())
 	{
 		Metadata->AddPoint(InKey);
 	}
 
-
 	if (bLoopPositionOverride)
 	{
 		LoopPosition += 1.0f;
@@ -793,20 +870,30 @@ void USplineComponent::AddSplinePoint(const FVector& Position, ESplineCoordinate
 
 void USplineComponent::AddSplinePointAtIndex(const FVector& Position, int32 Index, ESplineCoordinateSpace::Type CoordinateSpace, bool bUpdateSpline)
 {
+	int32 NumPoints = GetNumberOfSplinePoints();
+	
 	const FVector TransformedPosition = (CoordinateSpace == ESplineCoordinateSpace::World) ?
 		GetComponentTransform().InverseTransformPosition(Position) : Position;
 
-	int32 NumPoints = SplineCurves.Position.Points.Num();
-
 	if (Index >= 0 && Index <= NumPoints)
 	{
-		const float InKey = (Index == 0) ? 0.0f : SplineCurves.Position.Points[Index - 1].InVal + 1.0f;
+		const float InKey = (Index == 0) ? 0.0f : GetInputKeyValueAtSplinePoint(Index - 1) + 1.0f;
 
-		SplineCurves.Position.Points.Insert(FInterpCurvePoint<FVector>(InKey, TransformedPosition, FVector::ZeroVector, FVector::ZeroVector, CIM_CurveAuto), Index);
-		SplineCurves.Rotation.Points.Insert(FInterpCurvePoint<FQuat>(InKey, FQuat::Identity, FQuat::Identity, FQuat::Identity, CIM_CurveAuto), Index);
-		SplineCurves.Scale.Points.Insert(FInterpCurvePoint<FVector>(InKey, FVector(1.0f), FVector::ZeroVector, FVector::ZeroVector, CIM_CurveAuto), Index);
-		USplineMetadata* Metadata = GetSplinePointsMetadata();
-		if (Metadata)
+		FSplinePoint NewPoint;
+		NewPoint.InputKey = InKey;
+		NewPoint.Position = TransformedPosition;
+		NewPoint.ArriveTangent = FVector::ZeroVector;
+		NewPoint.LeaveTangent = FVector::ZeroVector;
+		NewPoint.Rotation = FQuat::Identity.Rotator();
+		NewPoint.Scale = FVector(1.f);
+
+		Spline.AddPointAtIndex(NewPoint, Index);
+
+		WarninglessSplineCurves().Position.Points.Insert(FInterpCurvePoint<FVector>(InKey, TransformedPosition, FVector::ZeroVector, FVector::ZeroVector, CIM_CurveAuto), Index);
+		WarninglessSplineCurves().Rotation.Points.Insert(FInterpCurvePoint<FQuat>(InKey, FQuat::Identity, FQuat::Identity, FQuat::Identity, CIM_CurveAuto), Index);
+		WarninglessSplineCurves().Scale.Points.Insert(FInterpCurvePoint<FVector>(InKey, FVector(1.0f), FVector::ZeroVector, FVector::ZeroVector, CIM_CurveAuto), Index);
+		
+		if (USplineMetadata* Metadata = GetSplinePointsMetadata())
 		{
 			Metadata->InsertPoint(InKey, Index, bClosedLoop);
 		}
@@ -816,9 +903,9 @@ void USplineComponent::AddSplinePointAtIndex(const FVector& Position, int32 Inde
 		// Adjust subsequent points' input keys to make room for the value just added
 		for (int I = Index + 1; I < NumPoints; ++I)
 		{
-			SplineCurves.Position.Points[I].InVal += 1.0f;
-			SplineCurves.Rotation.Points[I].InVal += 1.0f;
-			SplineCurves.Scale.Points[I].InVal += 1.0f;
+			WarninglessSplineCurves().Position.Points[I].InVal += 1.0f;
+			WarninglessSplineCurves().Rotation.Points[I].InVal += 1.0f;
+			WarninglessSplineCurves().Scale.Points[I].InVal += 1.0f;
 		}
 
 		if (bLoopPositionOverride)
@@ -835,15 +922,17 @@ void USplineComponent::AddSplinePointAtIndex(const FVector& Position, int32 Inde
 
 void USplineComponent::RemoveSplinePoint(int32 Index, bool bUpdateSpline)
 {
-	int32 NumPoints = SplineCurves.Position.Points.Num();
+	int32 NumPoints = GetNumberOfSplinePoints();
 
 	if (Index >= 0 && Index < NumPoints)
 	{
-		SplineCurves.Position.Points.RemoveAt(Index, EAllowShrinking::No);
-		SplineCurves.Rotation.Points.RemoveAt(Index, EAllowShrinking::No);
-		SplineCurves.Scale.Points.RemoveAt(Index, EAllowShrinking::No);
-		USplineMetadata* Metadata = GetSplinePointsMetadata();
-		if (Metadata)
+		Spline.RemovePoint(Index);
+
+		WarninglessSplineCurves().Position.Points.RemoveAt(Index, EAllowShrinking::No);
+		WarninglessSplineCurves().Rotation.Points.RemoveAt(Index, EAllowShrinking::No);
+		WarninglessSplineCurves().Scale.Points.RemoveAt(Index, EAllowShrinking::No);
+		
+		if (USplineMetadata* Metadata = GetSplinePointsMetadata())
 		{
 			Metadata->RemovePoint(Index);
 		}
@@ -853,9 +942,9 @@ void USplineComponent::RemoveSplinePoint(int32 Index, bool bUpdateSpline)
 		// Adjust all following spline point input keys to close the gap left by the removed point
 		while (Index < NumPoints)
 		{
-			SplineCurves.Position.Points[Index].InVal -= 1.0f;
-			SplineCurves.Rotation.Points[Index].InVal -= 1.0f;
-			SplineCurves.Scale.Points[Index].InVal -= 1.0f;
+			WarninglessSplineCurves().Position.Points[Index].InVal -= 1.0f;
+			WarninglessSplineCurves().Rotation.Points[Index].InVal -= 1.0f;
+			WarninglessSplineCurves().Scale.Points[Index].InVal -= 1.0f;
 			Index++;
 		}
 
@@ -874,9 +963,12 @@ void USplineComponent::RemoveSplinePoint(int32 Index, bool bUpdateSpline)
 void USplineComponent::SetSplinePoints(const TArray<FVector>& Points, ESplineCoordinateSpace::Type CoordinateSpace, bool bUpdateSpline)
 {
 	const int32 NumPoints = Points.Num();
-	SplineCurves.Position.Points.Reset(NumPoints);
-	SplineCurves.Rotation.Points.Reset(NumPoints);
-	SplineCurves.Scale.Points.Reset(NumPoints);
+
+	WarninglessSplineCurves().Position.Points.Reset(NumPoints);
+	WarninglessSplineCurves().Rotation.Points.Reset(NumPoints);
+	WarninglessSplineCurves().Scale.Points.Reset(NumPoints);
+	
+	Spline.Reset();
 
 	USplineMetadata* Metadata = GetSplinePointsMetadata();
 	if (Metadata)
@@ -890,9 +982,19 @@ void USplineComponent::SetSplinePoints(const TArray<FVector>& Points, ESplineCoo
 		const FVector TransformedPoint = (CoordinateSpace == ESplineCoordinateSpace::World) ?
 			GetComponentTransform().InverseTransformPosition(Point) : Point;
 
-		SplineCurves.Position.Points.Emplace(InputKey, TransformedPoint, FVector::ZeroVector, FVector::ZeroVector, CIM_CurveAuto);
-		SplineCurves.Rotation.Points.Emplace(InputKey, FQuat::Identity, FQuat::Identity, FQuat::Identity, CIM_CurveAuto);
-		SplineCurves.Scale.Points.Emplace(InputKey, FVector(1.0f), FVector::ZeroVector, FVector::ZeroVector, CIM_CurveAuto);
+		FSplinePoint NewPoint;
+		NewPoint.InputKey = InputKey;
+		NewPoint.Position = TransformedPoint;
+		NewPoint.ArriveTangent = FVector::ZeroVector;
+		NewPoint.LeaveTangent = FVector::ZeroVector;
+		NewPoint.Rotation = FQuat::Identity.Rotator();
+		NewPoint.Scale = FVector(1.f);
+
+		Spline.AddPoint(NewPoint);
+
+		WarninglessSplineCurves().Position.Points.Emplace(InputKey, TransformedPoint, FVector::ZeroVector, FVector::ZeroVector, CIM_CurveAuto);
+		WarninglessSplineCurves().Rotation.Points.Emplace(InputKey, FQuat::Identity, FQuat::Identity, FQuat::Identity, CIM_CurveAuto);
+		WarninglessSplineCurves().Scale.Points.Emplace(InputKey, FVector(1.0f), FVector::ZeroVector, FVector::ZeroVector, CIM_CurveAuto);
 
 		if (Metadata)
 		{
@@ -910,17 +1012,17 @@ void USplineComponent::SetSplinePoints(const TArray<FVector>& Points, ESplineCoo
 	}
 }
 
-
 void USplineComponent::SetLocationAtSplinePoint(int32 PointIndex, const FVector& InLocation, ESplineCoordinateSpace::Type CoordinateSpace, bool bUpdateSpline)
 {
-	const int32 NumPoints = SplineCurves.Position.Points.Num();
+	const int32 NumPoints = GetNumberOfSplinePoints();
 
 	if ((PointIndex >= 0) && (PointIndex < NumPoints))
 	{
 		const FVector TransformedLocation = (CoordinateSpace == ESplineCoordinateSpace::World) ?
 			GetComponentTransform().InverseTransformPosition(InLocation) : InLocation;
 
-		SplineCurves.Position.Points[PointIndex].OutVal = TransformedLocation;
+		Spline.SetLocation(PointIndex, TransformedLocation);
+		WarninglessSplineCurves().Position.Points[PointIndex].OutVal = TransformedLocation;
 
 		if (bUpdateSpline)
 		{
@@ -929,31 +1031,14 @@ void USplineComponent::SetLocationAtSplinePoint(int32 PointIndex, const FVector&
 	}
 }
 
-
 void USplineComponent::SetTangentAtSplinePoint(int32 PointIndex, const FVector& InTangent, ESplineCoordinateSpace::Type CoordinateSpace, bool bUpdateSpline)
 {
-	const int32 NumPoints = SplineCurves.Position.Points.Num();
-
-	if ((PointIndex >= 0) && (PointIndex < NumPoints))
-	{
-		const FVector TransformedTangent = (CoordinateSpace == ESplineCoordinateSpace::World) ?
-			GetComponentTransform().InverseTransformVector(InTangent) : InTangent;
-
-		SplineCurves.Position.Points[PointIndex].LeaveTangent = TransformedTangent;
-		SplineCurves.Position.Points[PointIndex].ArriveTangent = TransformedTangent;
-		SplineCurves.Position.Points[PointIndex].InterpMode = CIM_CurveUser;
-
-		if (bUpdateSpline)
-		{
-			UpdateSpline();
-		}
-	}
+	SetTangentsAtSplinePoint(PointIndex, InTangent, InTangent, CoordinateSpace, bUpdateSpline);
 }
 
-
 void USplineComponent::SetTangentsAtSplinePoint(int32 PointIndex, const FVector& InArriveTangent, const FVector& InLeaveTangent, ESplineCoordinateSpace::Type CoordinateSpace, bool bUpdateSpline)
 {
-	const int32 NumPoints = SplineCurves.Position.Points.Num();
+	const int32 NumPoints = GetNumberOfSplinePoints();
 
 	if ((PointIndex >= 0) && (PointIndex < NumPoints))
 	{
@@ -962,9 +1047,12 @@ void USplineComponent::SetTangentsAtSplinePoint(int32 PointIndex, const FVector&
 		const FVector TransformedLeaveTangent = (CoordinateSpace == ESplineCoordinateSpace::World) ?
 			GetComponentTransform().InverseTransformVector(InLeaveTangent) : InLeaveTangent;
 
-		SplineCurves.Position.Points[PointIndex].ArriveTangent = TransformedArriveTangent;
-		SplineCurves.Position.Points[PointIndex].LeaveTangent = TransformedLeaveTangent;
-		SplineCurves.Position.Points[PointIndex].InterpMode = CIM_CurveUser;
+		Spline.SetInTangent(PointIndex, TransformedArriveTangent);
+		Spline.SetOutTangent(PointIndex, TransformedLeaveTangent);
+
+		WarninglessSplineCurves().Position.Points[PointIndex].ArriveTangent = TransformedArriveTangent;
+		WarninglessSplineCurves().Position.Points[PointIndex].LeaveTangent = TransformedLeaveTangent;
+		WarninglessSplineCurves().Position.Points[PointIndex].InterpMode = CIM_CurveUser;
 
 		if (bUpdateSpline)
 		{
@@ -973,10 +1061,9 @@ void USplineComponent::SetTangentsAtSplinePoint(int32 PointIndex, const FVector&
 	}
 }
 
-
 void USplineComponent::SetUpVectorAtSplinePoint(int32 PointIndex, const FVector& InUpVector, ESplineCoordinateSpace::Type CoordinateSpace, bool bUpdateSpline)
 {
-	const int32 NumPoints = SplineCurves.Position.Points.Num();
+	const int32 NumPoints = GetNumberOfSplinePoints();
 
 	if ((PointIndex >= 0) && (PointIndex < NumPoints))
 	{
@@ -984,7 +1071,8 @@ void USplineComponent::SetUpVectorAtSplinePoint(int32 PointIndex, const FVector&
 			GetComponentTransform().InverseTransformVector(InUpVector.GetSafeNormal()) : InUpVector.GetSafeNormal();
 
 		FQuat Quat = FQuat::FindBetween(DefaultUpVector, TransformedUpVector);
-		SplineCurves.Rotation.Points[PointIndex].OutVal = Quat;
+		Spline.SetRotation(PointIndex, Quat);
+		WarninglessSplineCurves().Rotation.Points[PointIndex].OutVal = Quat;
 
 		if (bUpdateSpline)
 		{
@@ -995,16 +1083,28 @@ void USplineComponent::SetUpVectorAtSplinePoint(int32 PointIndex, const FVector&
 
 void USplineComponent::SetRotationAtSplinePoint(int32 PointIndex, const FRotator& InRotation, ESplineCoordinateSpace::Type CoordinateSpace, bool bUpdateSpline /*= true*/)
 {
-	if (SplineCurves.Rotation.Points.IsValidIndex(PointIndex))
+	SetQuaternionAtSplinePoint(PointIndex, InRotation.Quaternion(), CoordinateSpace, bUpdateSpline);
+}
+
+void USplineComponent::SetQuaternionAtSplinePoint(int32 PointIndex, const FQuat& InQuaternion, ESplineCoordinateSpace::Type CoordinateSpace, bool bUpdateSpline)
+{
+	const int32 NumPoints = GetNumberOfSplinePoints();
+	
+	if ((PointIndex >= 0) && (PointIndex < NumPoints))
 	{
+		// InQuaternion coerced into local space. All subsequent operations in local space.
 		const FQuat Quat = (CoordinateSpace == ESplineCoordinateSpace::World) ?
-			GetComponentTransform().InverseTransformRotation(InRotation.Quaternion()) : InRotation.Quaternion();
+			GetComponentTransform().InverseTransformRotation(InQuaternion) : InQuaternion;
 
+		// Align up vector with rotation
 		FVector UpVector = Quat.GetUpVector();
-		SetUpVectorAtSplinePoint(PointIndex, UpVector, CoordinateSpace, false);
+		SetUpVectorAtSplinePoint(PointIndex, UpVector, ESplineCoordinateSpace::Local, false);
 
-		FVector Direction = Quat.GetForwardVector();
-		SetTangentAtSplinePoint(PointIndex, Direction, CoordinateSpace, false);
+		// Align tangents with rotation, preserving magnitude.
+		const float ArriveTangentMag = GetArriveTangentAtSplinePoint(PointIndex, ESplineCoordinateSpace::Local).Length();
+		const float LeaveTangentMag = GetLeaveTangentAtSplinePoint(PointIndex, ESplineCoordinateSpace::Local).Length();
+		const FVector NewTangent = Quat.GetForwardVector().GetSafeNormal();
+		SetTangentsAtSplinePoint(PointIndex, ArriveTangentMag * NewTangent, LeaveTangentMag * NewTangent, ESplineCoordinateSpace::Local, false);
 
 		if (bUpdateSpline)
 		{
@@ -1015,9 +1115,12 @@ void USplineComponent::SetRotationAtSplinePoint(int32 PointIndex, const FRotator
 
 void USplineComponent::SetScaleAtSplinePoint(int32 PointIndex, const FVector& InScaleVector, bool bUpdateSpline /*= true*/)
 {
-	if (SplineCurves.Rotation.Points.IsValidIndex(PointIndex))
+	const int32 NumPoints = GetNumberOfSplinePoints();
+	
+	if ((PointIndex >= 0) && (PointIndex < NumPoints))
 	{
-		SplineCurves.Scale.Points[PointIndex].OutVal = InScaleVector;
+		WarninglessSplineCurves().Scale.Points[PointIndex].OutVal = InScaleVector;
+		Spline.SetScale(PointIndex, InScaleVector);
 
 		if (bUpdateSpline)
 		{
@@ -1028,20 +1131,29 @@ void USplineComponent::SetScaleAtSplinePoint(int32 PointIndex, const FVector& In
 
 ESplinePointType::Type USplineComponent::GetSplinePointType(int32 PointIndex) const
 {
-	if ((PointIndex >= 0) && (PointIndex < SplineCurves.Position.Points.Num()))
+	const int32 NumPoints = GetNumberOfSplinePoints();
+	
+	if (PointIndex >= 0 && PointIndex < NumPoints)
 	{
-		return ConvertInterpCurveModeToSplinePointType(SplineCurves.Position.Points[PointIndex].InterpMode);
+		const EInterpCurveMode Mode = UE::SplineComponent::GUseSplineCurves
+			? WarninglessSplineCurves().Position.Points[PointIndex].InterpMode.GetValue()
+			: Spline.GetSplinePointType(PointIndex);
+		
+		return ConvertInterpCurveModeToSplinePointType(Mode);
 	}
 
 	return ESplinePointType::Constant;
 }
 
-
 void USplineComponent::SetSplinePointType(int32 PointIndex, ESplinePointType::Type Type, bool bUpdateSpline)
 {
-	if ((PointIndex >= 0) && (PointIndex < SplineCurves.Position.Points.Num()))
+	const int32 NumPoints = GetNumberOfSplinePoints();
+	
+	if (PointIndex >= 0 && PointIndex < NumPoints)
 	{
-		SplineCurves.Position.Points[PointIndex].InterpMode = ConvertSplinePointTypeToInterpCurveMode(Type);
+		Spline.SetSplinePointType(PointIndex, ConvertSplinePointTypeToInterpCurveMode(Type));
+		WarninglessSplineCurves().Position.Points[PointIndex].InterpMode = ConvertSplinePointTypeToInterpCurveMode(Type);
+		
 		if (bUpdateSpline)
 		{
 			UpdateSpline();
@@ -1049,154 +1161,146 @@ void USplineComponent::SetSplinePointType(int32 PointIndex, ESplinePointType::Ty
 	}
 }
 
-
 int32 USplineComponent::GetNumberOfSplinePoints() const
 {
-	// No longer returns an imaginary extra endpoint if closed loop is set
-	const int32 NumPoints = SplineCurves.Position.Points.Num();
-	return NumPoints;
+	return UE::SplineComponent::GUseSplineCurves
+		? WarninglessSplineCurves().Position.Points.Num()
+		: Spline.GetNumControlPoints();
 }
 
 int32 USplineComponent::GetNumberOfSplineSegments() const
 {
-	const int32 NumPoints = SplineCurves.Position.Points.Num();
+	const int32 NumPoints = GetNumberOfSplinePoints();
 	return (bClosedLoop ? NumPoints : FMath::Max(0, NumPoints - 1));
 }
 
 float USplineComponent::GetInputKeyValueAtSplinePoint(int32 PointIndex) const
 {
-	const FInterpCurvePointVector& SplinePoint = GetPositionPointSafe(PointIndex);
-	return SplinePoint.InVal;
+	PointIndex = GetClampedIndex(PointIndex);
+	
+	return UE::SplineComponent::GUseSplineCurves
+		? WarninglessSplineCurves().Position.Points[PointIndex].InVal
+		: Spline.GetSplinePointParameter(PointIndex);
 }
 
-
 FSplinePoint USplineComponent::GetSplinePointAt(int32 PointIndex, ESplineCoordinateSpace::Type CoordinateSpace) const
 {
-	const FInterpCurvePointVector& SplinePoint = GetPositionPointSafe(PointIndex);
+	PointIndex = GetClampedIndex(PointIndex);
 
-	const FInterpCurvePointQuat& RotationPoint = GetRotationPointSafe(PointIndex);
-	const FRotator& Rotation = GetRotationAtSplineInputKey(RotationPoint.InVal, CoordinateSpace);
-	const ESplinePointType::Type SplinePointType = GetSplinePointType(PointIndex);
-
-	const FVector Scale = GetScaleAtSplinePoint(PointIndex);
-
-	return FSplinePoint(SplinePoint.InVal, 
-		SplinePoint.OutVal,
-		SplinePoint.ArriveTangent,
-		SplinePoint.LeaveTangent,
-		Rotation,
-		Scale,
-		SplinePointType);
+	return FSplinePoint(GetInputKeyValueAtSplinePoint(PointIndex),
+		GetLocationAtSplinePoint(PointIndex, CoordinateSpace),
+		GetArriveTangentAtSplinePoint(PointIndex, CoordinateSpace),
+		GetLeaveTangentAtSplinePoint(PointIndex, CoordinateSpace),
+		GetRotationAtSplinePoint(PointIndex, CoordinateSpace),
+		GetScaleAtSplinePoint(PointIndex),
+		GetSplinePointType(PointIndex));
 }
 
 FVector USplineComponent::GetLocationAtSplinePoint(int32 PointIndex, ESplineCoordinateSpace::Type CoordinateSpace) const
 {
-	const FInterpCurvePointVector& Point = GetPositionPointSafe(PointIndex);
-	const FVector& Location = Point.OutVal;
+	PointIndex = GetClampedIndex(PointIndex);
+	
+	const FVector Location = UE::SplineComponent::GUseSplineCurves
+		? WarninglessSplineCurves().Position.Points[PointIndex].OutVal
+		: Spline.GetLocation(PointIndex);
+	
 	return (CoordinateSpace == ESplineCoordinateSpace::World) ? GetComponentTransform().TransformPosition(Location) : Location;
 }
 
-
 FVector USplineComponent::GetDirectionAtSplinePoint(int32 PointIndex, ESplineCoordinateSpace::Type CoordinateSpace) const
 {
-	const FInterpCurvePointVector& Point = GetPositionPointSafe(PointIndex);
-	const FVector Direction = Point.LeaveTangent.GetSafeNormal();
-	return (CoordinateSpace == ESplineCoordinateSpace::World) ? GetComponentTransform().TransformVector(Direction) : Direction;
+	return GetTangentAtSplinePoint(PointIndex, CoordinateSpace).GetSafeNormal();
 }
 
-
 FVector USplineComponent::GetTangentAtSplinePoint(int32 PointIndex, ESplineCoordinateSpace::Type CoordinateSpace) const
 {
-	const FInterpCurvePointVector& Point = GetPositionPointSafe(PointIndex);
-	const FVector& Direction = Point.LeaveTangent;
-	return (CoordinateSpace == ESplineCoordinateSpace::World) ? GetComponentTransform().TransformVector(Direction) : Direction;
+	return GetLeaveTangentAtSplinePoint(PointIndex, CoordinateSpace);
 }
 
-
 FVector USplineComponent::GetArriveTangentAtSplinePoint(int32 PointIndex, ESplineCoordinateSpace::Type CoordinateSpace) const
 {
-	const FInterpCurvePointVector& Point = GetPositionPointSafe(PointIndex);
-	const FVector& Direction = Point.ArriveTangent;
-	return (CoordinateSpace == ESplineCoordinateSpace::World) ? GetComponentTransform().TransformVector(Direction) : Direction;
-}
+	PointIndex = GetClampedIndex(PointIndex);
+
+	const FVector Tangent = UE::SplineComponent::GUseSplineCurves
+		? WarninglessSplineCurves().Position.Points[PointIndex].ArriveTangent
+		: Spline.GetInTangent(PointIndex);
 
+	return (CoordinateSpace == ESplineCoordinateSpace::World) ? GetComponentTransform().TransformVector(Tangent) : Tangent;
+}
 
 FVector USplineComponent::GetLeaveTangentAtSplinePoint(int32 PointIndex, ESplineCoordinateSpace::Type CoordinateSpace) const
 {
-	const FInterpCurvePointVector& Point = GetPositionPointSafe(PointIndex);
-	const FVector& Direction = Point.LeaveTangent;
-	return (CoordinateSpace == ESplineCoordinateSpace::World) ? GetComponentTransform().TransformVector(Direction) : Direction;
-}
+	PointIndex = GetClampedIndex(PointIndex);
+
+	const FVector Tangent = UE::SplineComponent::GUseSplineCurves
+		? WarninglessSplineCurves().Position.Points[PointIndex].LeaveTangent
+		: Spline.GetOutTangent(PointIndex);
 
+	return (CoordinateSpace == ESplineCoordinateSpace::World) ? GetComponentTransform().TransformVector(Tangent) : Tangent;
+}
 
 FQuat USplineComponent::GetQuaternionAtSplinePoint(int32 PointIndex, ESplineCoordinateSpace::Type CoordinateSpace) const
 {
-	const FInterpCurvePointQuat& Point = GetRotationPointSafe(PointIndex);
-	return GetQuaternionAtSplineInputKey(Point.InVal, CoordinateSpace);
+	return GetQuaternionAtSplineInputKey(GetInputKeyValueAtSplinePoint(PointIndex), CoordinateSpace);
 }
 
-
 FRotator USplineComponent::GetRotationAtSplinePoint(int32 PointIndex, ESplineCoordinateSpace::Type CoordinateSpace) const
 {
-	const FInterpCurvePointQuat& Point = GetRotationPointSafe(PointIndex);
-	return GetRotationAtSplineInputKey(Point.InVal, CoordinateSpace);
+	return GetRotationAtSplineInputKey(GetInputKeyValueAtSplinePoint(PointIndex), CoordinateSpace);
 }
 
-
 FVector USplineComponent::GetUpVectorAtSplinePoint(int32 PointIndex, ESplineCoordinateSpace::Type CoordinateSpace) const
 {
-	const FInterpCurvePointQuat& Point = GetRotationPointSafe(PointIndex);
-	return GetUpVectorAtSplineInputKey(Point.InVal, CoordinateSpace);
+	return GetUpVectorAtSplineInputKey(GetInputKeyValueAtSplinePoint(PointIndex), CoordinateSpace);
 }
 
-
 FVector USplineComponent::GetRightVectorAtSplinePoint(int32 PointIndex, ESplineCoordinateSpace::Type CoordinateSpace) const
 {
-	const FInterpCurvePointQuat& Point = GetRotationPointSafe(PointIndex);
-	return GetRightVectorAtSplineInputKey(Point.InVal, CoordinateSpace);
+	return GetRightVectorAtSplineInputKey(GetInputKeyValueAtSplinePoint(PointIndex), CoordinateSpace);
 }
 
-
 float USplineComponent::GetRollAtSplinePoint(int32 PointIndex, ESplineCoordinateSpace::Type CoordinateSpace) const
 {
-	const FInterpCurvePointQuat& Point = GetRotationPointSafe(PointIndex);
-	return GetRollAtSplineInputKey(Point.InVal, CoordinateSpace);
+	return GetRollAtSplineInputKey(GetInputKeyValueAtSplinePoint(PointIndex), CoordinateSpace);
 }
 
-
 FVector USplineComponent::GetScaleAtSplinePoint(int32 PointIndex) const
 {
-	const FInterpCurvePointVector& Point = GetScalePointSafe(PointIndex);
-	return Point.OutVal;
-}
+	PointIndex = GetClampedIndex(PointIndex);
 
+	return UE::SplineComponent::GUseSplineCurves
+		? WarninglessSplineCurves().Scale.Points[PointIndex].OutVal
+		: Spline.GetScale(PointIndex);
+}
 
 FTransform USplineComponent::GetTransformAtSplinePoint(int32 PointIndex, ESplineCoordinateSpace::Type CoordinateSpace, bool bUseScale) const
 {
-	const FInterpCurvePointQuat& Point = GetRotationPointSafe(PointIndex);
-	return GetTransformAtSplineInputKey(Point.InVal, CoordinateSpace, bUseScale);
+	return GetTransformAtSplineInputKey(GetInputKeyValueAtSplinePoint(PointIndex), CoordinateSpace, bUseScale);
 }
 
-
 void USplineComponent::GetLocationAndTangentAtSplinePoint(int32 PointIndex, FVector& Location, FVector& Tangent, ESplineCoordinateSpace::Type CoordinateSpace) const
 {
-	const FInterpCurvePointQuat& Point = GetRotationPointSafe(PointIndex);
-	const float InputKey = Point.InVal;
+	const float InputKey = GetInputKeyValueAtSplinePoint(PointIndex);
 	Location = GetLocationAtSplineInputKey(InputKey, CoordinateSpace);
 	Tangent = GetTangentAtSplineInputKey(InputKey, CoordinateSpace);
 }
 
-
 float USplineComponent::GetDistanceAlongSplineAtSplinePoint(int32 PointIndex) const
 {
-	const int32 NumPoints = SplineCurves.Position.Points.Num();
+	const int32 NumPoints = GetNumberOfSplinePoints();
 	const int32 NumSegments = bClosedLoop ? NumPoints : NumPoints - 1;
 
+	const int32 NumReparamPoints = UE::SplineComponent::GUseSplineCurves
+		? WarninglessSplineCurves().ReparamTable.Points.Num()
+		: Spline.GetSplineReparamTable().Points.Num();
+	
 	// Ensure that if the reparam table is not prepared yet we don't attempt to access it. This can happen
 	// early in the construction of the spline component object.
-	if ((PointIndex >= 0) && (PointIndex < NumSegments + 1) && ((PointIndex * ReparamStepsPerSegment) < SplineCurves.ReparamTable.Points.Num()))
+	if ((PointIndex >= 0) && (PointIndex < NumSegments + 1) && ((PointIndex * ReparamStepsPerSegment) < NumReparamPoints))
 	{
-		return SplineCurves.ReparamTable.Points[PointIndex * ReparamStepsPerSegment].InVal;
+		return UE::SplineComponent::GUseSplineCurves
+			? WarninglessSplineCurves().ReparamTable.Points[PointIndex * ReparamStepsPerSegment].InVal
+			: Spline.GetSplineReparamTable().Points[PointIndex * ReparamStepsPerSegment].InVal;
 	}
 
 	return 0.0f;
@@ -1217,10 +1321,11 @@ float FSplineCurves::GetSplineLength() const
 
 float USplineComponent::GetSplineLength() const
 {
-	return SplineCurves.GetSplineLength();
+	return UE::SplineComponent::GUseSplineCurves
+		? WarninglessSplineCurves().GetSplineLength()
+		: Spline.GetSplineLength();
 }
 
-
 void USplineComponent::SetDefaultUpVector(const FVector& UpVector, ESplineCoordinateSpace::Type CoordinateSpace)
 {
 	if (CoordinateSpace == ESplineCoordinateSpace::World)
@@ -1235,7 +1340,6 @@ void USplineComponent::SetDefaultUpVector(const FVector& UpVector, ESplineCoordi
 	UpdateSpline();
 }
 
-
 FVector USplineComponent::GetDefaultUpVector(ESplineCoordinateSpace::Type CoordinateSpace) const
 {
 	if (CoordinateSpace == ESplineCoordinateSpace::World)
@@ -1248,7 +1352,6 @@ FVector USplineComponent::GetDefaultUpVector(ESplineCoordinateSpace::Type Coordi
 	}
 }
 
-
 float USplineComponent::GetInputKeyAtDistanceAlongSpline(float Distance) const
 {
 	return GetTimeAtDistanceAlongSpline(Distance);
@@ -1256,19 +1359,21 @@ float USplineComponent::GetInputKeyAtDistanceAlongSpline(float Distance) const
 
 float USplineComponent::GetInputKeyValueAtDistanceAlongSpline(float Distance) const
 {
-	const int32 NumPoints = SplineCurves.Position.Points.Num();
+	const int32 NumPoints = GetNumberOfSplinePoints();
 
 	if (NumPoints < 2)
 	{
 		return 0.0f;
 	}
 	
-	return SplineCurves.ReparamTable.Eval(Distance, 0.0f);
+	return UE::SplineComponent::GUseSplineCurves
+		? WarninglessSplineCurves().ReparamTable.Eval(Distance, 0.0f)
+		: Spline.GetSplineReparamTable().Eval(Distance, 0.0f);
 }
 
 float USplineComponent::GetTimeAtDistanceAlongSpline(float Distance) const
 {
-	const int32 NumPoints = SplineCurves.Position.Points.Num();
+	const int32 NumPoints = GetNumberOfSplinePoints();
 
 	if (NumPoints < 2)
 	{
@@ -1276,79 +1381,92 @@ float USplineComponent::GetTimeAtDistanceAlongSpline(float Distance) const
 	}
 
 	const float TimeMultiplier = Duration / (bClosedLoop ? NumPoints : (NumPoints - 1.0f));
-	return SplineCurves.ReparamTable.Eval(Distance, 0.0f) * TimeMultiplier;
+	return (UE::SplineComponent::GUseSplineCurves
+		? WarninglessSplineCurves().ReparamTable.Eval(Distance, 0.0f)
+		: Spline.GetSplineReparamTable().Eval(Distance, 0.0f)) * TimeMultiplier;
 }
 
 FVector USplineComponent::GetLocationAtDistanceAlongSpline(float Distance, ESplineCoordinateSpace::Type CoordinateSpace) const
 {
-	const float Param = SplineCurves.ReparamTable.Eval(Distance, 0.0f);
+	const float Param = UE::SplineComponent::GUseSplineCurves
+		? WarninglessSplineCurves().ReparamTable.Eval(Distance, 0.0f)
+		: Spline.GetSplineReparamTable().Eval(Distance, 0.0f);
 	return GetLocationAtSplineInputKey(Param, CoordinateSpace);
 }
 
-
 FVector USplineComponent::GetTangentAtDistanceAlongSpline(float Distance, ESplineCoordinateSpace::Type CoordinateSpace) const
 {
-	const float Param = SplineCurves.ReparamTable.Eval(Distance, 0.0f);
+	const float Param = UE::SplineComponent::GUseSplineCurves
+		? WarninglessSplineCurves().ReparamTable.Eval(Distance, 0.0f)
+		: Spline.GetSplineReparamTable().Eval(Distance, 0.0f);
 	return GetTangentAtSplineInputKey(Param, CoordinateSpace);
 }
 
 
 FVector USplineComponent::GetDirectionAtDistanceAlongSpline(float Distance, ESplineCoordinateSpace::Type CoordinateSpace) const
 {
-	const float Param = SplineCurves.ReparamTable.Eval(Distance, 0.0f);
+	const float Param = UE::SplineComponent::GUseSplineCurves
+		? WarninglessSplineCurves().ReparamTable.Eval(Distance, 0.0f)
+		: Spline.GetSplineReparamTable().Eval(Distance, 0.0f);
 	return GetDirectionAtSplineInputKey(Param, CoordinateSpace);
 }
 
-
 FQuat USplineComponent::GetQuaternionAtDistanceAlongSpline(float Distance, ESplineCoordinateSpace::Type CoordinateSpace) const
 {
-	const float Param = SplineCurves.ReparamTable.Eval(Distance, 0.0f);
+	const float Param = UE::SplineComponent::GUseSplineCurves
+		? WarninglessSplineCurves().ReparamTable.Eval(Distance, 0.0f)
+		: Spline.GetSplineReparamTable().Eval(Distance, 0.0f);
 	return GetQuaternionAtSplineInputKey(Param, CoordinateSpace);
 }
 
-
 FRotator USplineComponent::GetRotationAtDistanceAlongSpline(float Distance, ESplineCoordinateSpace::Type CoordinateSpace) const
 {
-	const float Param = SplineCurves.ReparamTable.Eval(Distance, 0.0f);
+	const float Param = UE::SplineComponent::GUseSplineCurves
+		? WarninglessSplineCurves().ReparamTable.Eval(Distance, 0.0f)
+		: Spline.GetSplineReparamTable().Eval(Distance, 0.0f);
 	return GetRotationAtSplineInputKey(Param, CoordinateSpace);
 }
 
-
 FVector USplineComponent::GetUpVectorAtDistanceAlongSpline(float Distance, ESplineCoordinateSpace::Type CoordinateSpace) const
 {
-	const float Param = SplineCurves.ReparamTable.Eval(Distance, 0.0f);
+	const float Param = UE::SplineComponent::GUseSplineCurves
+		? WarninglessSplineCurves().ReparamTable.Eval(Distance, 0.0f)
+		: Spline.GetSplineReparamTable().Eval(Distance, 0.0f);
 	return GetUpVectorAtSplineInputKey(Param, CoordinateSpace);
 }
 
-
 FVector USplineComponent::GetRightVectorAtDistanceAlongSpline(float Distance, ESplineCoordinateSpace::Type CoordinateSpace) const
 {
-	const float Param = SplineCurves.ReparamTable.Eval(Distance, 0.0f);
+	const float Param = UE::SplineComponent::GUseSplineCurves
+		? WarninglessSplineCurves().ReparamTable.Eval(Distance, 0.0f)
+		: Spline.GetSplineReparamTable().Eval(Distance, 0.0f);
 	return GetRightVectorAtSplineInputKey(Param, CoordinateSpace);
 }
 
-
 float USplineComponent::GetRollAtDistanceAlongSpline(float Distance, ESplineCoordinateSpace::Type CoordinateSpace) const
 {
-	const float Param = SplineCurves.ReparamTable.Eval(Distance, 0.0f);
+	const float Param = UE::SplineComponent::GUseSplineCurves
+		? WarninglessSplineCurves().ReparamTable.Eval(Distance, 0.0f)
+		: Spline.GetSplineReparamTable().Eval(Distance, 0.0f);
 	return GetRollAtSplineInputKey(Param, CoordinateSpace);
 }
 
-
 FVector USplineComponent::GetScaleAtDistanceAlongSpline(float Distance) const
 {
-	const float Param = SplineCurves.ReparamTable.Eval(Distance, 0.0f);
+	const float Param = UE::SplineComponent::GUseSplineCurves
+		? WarninglessSplineCurves().ReparamTable.Eval(Distance, 0.0f)
+		: Spline.GetSplineReparamTable().Eval(Distance, 0.0f);
 	return GetScaleAtSplineInputKey(Param);
 }
 
-
 FTransform USplineComponent::GetTransformAtDistanceAlongSpline(float Distance, ESplineCoordinateSpace::Type CoordinateSpace, bool bUseScale) const
 {
-	const float Param = SplineCurves.ReparamTable.Eval(Distance, 0.0f);
+	const float Param = UE::SplineComponent::GUseSplineCurves
+		? WarninglessSplineCurves().ReparamTable.Eval(Distance, 0.0f)
+		: Spline.GetSplineReparamTable().Eval(Distance, 0.0f);
 	return GetTransformAtSplineInputKey(Param, CoordinateSpace, bUseScale);
 }
 
-
 FVector USplineComponent::GetLocationAtTime(float Time, ESplineCoordinateSpace::Type CoordinateSpace, bool bUseConstantVelocity) const
 {
 	if (Duration == 0.0f)
@@ -1362,14 +1480,13 @@ FVector USplineComponent::GetLocationAtTime(float Time, ESplineCoordinateSpace::
 	}
 	else
 	{
-		const int32 NumPoints = SplineCurves.Position.Points.Num();
+		const int32 NumPoints = GetNumberOfSplinePoints();
 		const int32 NumSegments = bClosedLoop ? NumPoints : NumPoints - 1;
 		const float TimeMultiplier = NumSegments / Duration;
 		return GetLocationAtSplineInputKey(Time * TimeMultiplier, CoordinateSpace);
 	}
 }
 
-
 FVector USplineComponent::GetDirectionAtTime(float Time, ESplineCoordinateSpace::Type CoordinateSpace, bool bUseConstantVelocity) const
 {
 	if (Duration == 0.0f)
@@ -1383,14 +1500,13 @@ FVector USplineComponent::GetDirectionAtTime(float Time, ESplineCoordinateSpace:
 	}
 	else
 	{
-		const int32 NumPoints = SplineCurves.Position.Points.Num();
+		const int32 NumPoints = GetNumberOfSplinePoints();
 		const int32 NumSegments = bClosedLoop ? NumPoints : NumPoints - 1;
 		const float TimeMultiplier = NumSegments / Duration;
 		return GetDirectionAtSplineInputKey(Time * TimeMultiplier, CoordinateSpace);
 	}
 }
 
-
 FVector USplineComponent::GetTangentAtTime(float Time, ESplineCoordinateSpace::Type CoordinateSpace, bool bUseConstantVelocity) const
 {
 	if (Duration == 0.0f)
@@ -1404,14 +1520,13 @@ FVector USplineComponent::GetTangentAtTime(float Time, ESplineCoordinateSpace::T
 	}
 	else
 	{
-		const int32 NumPoints = SplineCurves.Position.Points.Num();
+		const int32 NumPoints = GetNumberOfSplinePoints();
 		const int32 NumSegments = bClosedLoop ? NumPoints : NumPoints - 1;
 		const float TimeMultiplier = NumSegments / Duration;
 		return GetTangentAtSplineInputKey(Time * TimeMultiplier, CoordinateSpace);
 	}
 }
 
-
 FRotator USplineComponent::GetRotationAtTime(float Time, ESplineCoordinateSpace::Type CoordinateSpace, bool bUseConstantVelocity) const
 {
 	if (Duration == 0.0f)
@@ -1425,14 +1540,13 @@ FRotator USplineComponent::GetRotationAtTime(float Time, ESplineCoordinateSpace:
 	}
 	else
 	{
-		const int32 NumPoints = SplineCurves.Position.Points.Num();
+		const int32 NumPoints = GetNumberOfSplinePoints();
 		const int32 NumSegments = bClosedLoop ? NumPoints : NumPoints - 1;
 		const float TimeMultiplier = NumSegments / Duration;
 		return GetRotationAtSplineInputKey(Time * TimeMultiplier, CoordinateSpace);
 	}
 }
 
-
 FQuat USplineComponent::GetQuaternionAtTime(float Time, ESplineCoordinateSpace::Type CoordinateSpace, bool bUseConstantVelocity) const
 {
 	if (Duration == 0.0f)
@@ -1446,14 +1560,13 @@ FQuat USplineComponent::GetQuaternionAtTime(float Time, ESplineCoordinateSpace::
 	}
 	else
 	{
-		const int32 NumPoints = SplineCurves.Position.Points.Num();
+		const int32 NumPoints = GetNumberOfSplinePoints();
 		const int32 NumSegments = bClosedLoop ? NumPoints : NumPoints - 1;
 		const float TimeMultiplier = NumSegments / Duration;
 		return GetQuaternionAtSplineInputKey(Time * TimeMultiplier, CoordinateSpace);
 	}
 }
 
-
 FVector USplineComponent::GetUpVectorAtTime(float Time, ESplineCoordinateSpace::Type CoordinateSpace, bool bUseConstantVelocity) const
 {
 	if (Duration == 0.0f)
@@ -1467,14 +1580,13 @@ FVector USplineComponent::GetUpVectorAtTime(float Time, ESplineCoordinateSpace::
 	}
 	else
 	{
-		const int32 NumPoints = SplineCurves.Position.Points.Num();
+		const int32 NumPoints = GetNumberOfSplinePoints();
 		const int32 NumSegments = bClosedLoop ? NumPoints : NumPoints - 1;
 		const float TimeMultiplier = NumSegments / Duration;
 		return GetUpVectorAtSplineInputKey(Time * TimeMultiplier, CoordinateSpace);
 	}
 }
 
-
 FVector USplineComponent::GetRightVectorAtTime(float Time, ESplineCoordinateSpace::Type CoordinateSpace, bool bUseConstantVelocity) const
 {
 	if (Duration == 0.0f)
@@ -1488,14 +1600,13 @@ FVector USplineComponent::GetRightVectorAtTime(float Time, ESplineCoordinateSpac
 	}
 	else
 	{
-		const int32 NumPoints = SplineCurves.Position.Points.Num();
+		const int32 NumPoints = GetNumberOfSplinePoints();
 		const int32 NumSegments = bClosedLoop ? NumPoints : NumPoints - 1;
 		const float TimeMultiplier = NumSegments / Duration;
 		return GetRightVectorAtSplineInputKey(Time * TimeMultiplier, CoordinateSpace);
 	}
 }
 
-
 float USplineComponent::GetRollAtTime(float Time, ESplineCoordinateSpace::Type CoordinateSpace, bool bUseConstantVelocity) const
 {
 	if (Duration == 0.0f)
@@ -1509,14 +1620,13 @@ float USplineComponent::GetRollAtTime(float Time, ESplineCoordinateSpace::Type C
 	}
 	else
 	{
-		const int32 NumPoints = SplineCurves.Position.Points.Num();
+		const int32 NumPoints = GetNumberOfSplinePoints();
 		const int32 NumSegments = bClosedLoop ? NumPoints : NumPoints - 1;
 		const float TimeMultiplier = NumSegments / Duration;
 		return GetRollAtSplineInputKey(Time * TimeMultiplier, CoordinateSpace);
 	}
 }
 
-
 FTransform USplineComponent::GetTransformAtTime(float Time, ESplineCoordinateSpace::Type CoordinateSpace, bool bUseConstantVelocity, bool bUseScale) const
 {
 	if (Duration == 0.0f)
@@ -1530,14 +1640,13 @@ FTransform USplineComponent::GetTransformAtTime(float Time, ESplineCoordinateSpa
 	}
 	else
 	{
-		const int32 NumPoints = SplineCurves.Position.Points.Num();
+		const int32 NumPoints = GetNumberOfSplinePoints();
 		const int32 NumSegments = bClosedLoop ? NumPoints : NumPoints - 1;
 		const float TimeMultiplier = NumSegments / Duration;
 		return GetTransformAtSplineInputKey(Time * TimeMultiplier, CoordinateSpace, bUseScale);
 	}
 }
 
-
 FVector USplineComponent::GetScaleAtTime(float Time, bool bUseConstantVelocity) const
 {
 	if (Duration == 0.0f)
@@ -1551,21 +1660,30 @@ FVector USplineComponent::GetScaleAtTime(float Time, bool bUseConstantVelocity)
 	}
 	else
 	{
-		const int32 NumPoints = SplineCurves.Position.Points.Num();
+		const int32 NumPoints = GetNumberOfSplinePoints();
 		const int32 NumSegments = bClosedLoop ? NumPoints : NumPoints - 1;
 		const float TimeMultiplier = NumSegments / Duration;
 		return GetScaleAtSplineInputKey(Time * TimeMultiplier);
 	}
 }
 
-
 float USplineComponent::FindInputKeyClosestToWorldLocation(const FVector& WorldLocation) const
 {
 	const FVector LocalLocation = GetComponentTransform().InverseTransformPosition(WorldLocation);
 	float Dummy;
-	return SplineCurves.Position.FindNearest(LocalLocation, Dummy);
+	return UE::SplineComponent::GUseSplineCurves
+		? WarninglessSplineCurves().Position.FindNearest(LocalLocation, Dummy)
+		: Spline.FindNearest(LocalLocation);
 }
 
+float USplineComponent::FindInputKeyOnSegmentClosestToWorldLocation(const FVector& WorldLocation, int32 Index) const
+{
+	const FVector LocalLocation = GetComponentTransform().InverseTransformPosition(WorldLocation);
+	float Dummy;
+	return UE::SplineComponent::GUseSplineCurves
+		? WarninglessSplineCurves().Position.FindNearestOnSegment(LocalLocation, Index, Dummy)
+		: Spline.FindNearestOnSegment(LocalLocation, Index);
+}
 
 FVector USplineComponent::FindLocationClosestToWorldLocation(const FVector& WorldLocation, ESplineCoordinateSpace::Type CoordinateSpace) const
 {
@@ -1573,63 +1691,54 @@ FVector USplineComponent::FindLocationClosestToWorldLocation(const FVector& Worl
 	return GetLocationAtSplineInputKey(Param, CoordinateSpace);
 }
 
-
 FVector USplineComponent::FindDirectionClosestToWorldLocation(const FVector& WorldLocation, ESplineCoordinateSpace::Type CoordinateSpace) const
 {
 	const float Param = FindInputKeyClosestToWorldLocation(WorldLocation);
 	return GetDirectionAtSplineInputKey(Param, CoordinateSpace);
 }
 
-
 FVector USplineComponent::FindTangentClosestToWorldLocation(const FVector& WorldLocation, ESplineCoordinateSpace::Type CoordinateSpace) const
 {
 	const float Param = FindInputKeyClosestToWorldLocation(WorldLocation);
 	return GetTangentAtSplineInputKey(Param, CoordinateSpace);
 }
 
-
 FQuat USplineComponent::FindQuaternionClosestToWorldLocation(const FVector& WorldLocation, ESplineCoordinateSpace::Type CoordinateSpace) const
 {
 	const float Param = FindInputKeyClosestToWorldLocation(WorldLocation);
 	return GetQuaternionAtSplineInputKey(Param, CoordinateSpace);
 }
 
-
 FRotator USplineComponent::FindRotationClosestToWorldLocation(const FVector& WorldLocation, ESplineCoordinateSpace::Type CoordinateSpace) const
 {
 	const float Param = FindInputKeyClosestToWorldLocation(WorldLocation);
 	return GetRotationAtSplineInputKey(Param, CoordinateSpace);
 }
 
-
 FVector USplineComponent::FindUpVectorClosestToWorldLocation(const FVector& WorldLocation, ESplineCoordinateSpace::Type CoordinateSpace) const
 {
 	const float Param = FindInputKeyClosestToWorldLocation(WorldLocation);
 	return GetUpVectorAtSplineInputKey(Param, CoordinateSpace);
 }
 
-
 FVector USplineComponent::FindRightVectorClosestToWorldLocation(const FVector& WorldLocation, ESplineCoordinateSpace::Type CoordinateSpace) const
 {
 	const float Param = FindInputKeyClosestToWorldLocation(WorldLocation);
 	return GetRightVectorAtSplineInputKey(Param, CoordinateSpace);
 }
 
-
 float USplineComponent::FindRollClosestToWorldLocation(const FVector& WorldLocation, ESplineCoordinateSpace::Type CoordinateSpace) const
 {
 	const float Param = FindInputKeyClosestToWorldLocation(WorldLocation);
 	return GetRollAtSplineInputKey(Param, CoordinateSpace);
 }
 
-
 FVector USplineComponent::FindScaleClosestToWorldLocation(const FVector& WorldLocation) const
 {
 	const float Param = FindInputKeyClosestToWorldLocation(WorldLocation);
 	return GetScaleAtSplineInputKey(Param);
 }
 
-
 FTransform USplineComponent::FindTransformClosestToWorldLocation(const FVector& WorldLocation, ESplineCoordinateSpace::Type CoordinateSpace, bool bUseScale) const
 {
 	const float Param = FindInputKeyClosestToWorldLocation(WorldLocation);
@@ -1771,7 +1880,7 @@ bool USplineComponent::ConvertSplineToPolyLineWithDistances(ESplineCoordinateSpa
 
 bool USplineComponent::ConvertSplineToPolyline_InDistanceRange(ESplineCoordinateSpace::Type CoordinateSpace, const float InMaxSquareDistanceFromSpline, float RangeStart, float RangeEnd, TArray<FVector>& OutPoints, TArray<double>& OutDistancesAlongSpline, bool bAllowWrappingIfClosed) const
 {
-	const int32 NumPoints = SplineCurves.Position.Points.Num();
+	const int32 NumPoints = GetNumberOfSplinePoints();
 	if (NumPoints == 0)
 	{
 		return false;
@@ -1859,8 +1968,16 @@ bool USplineComponent::ConvertSplineToPolyline_InDistanceRange(ESplineCoordinate
 		return bHasAdded;
 	} // end of the wrap-around case, after this values will be in the normal range
 	
-	int32 SegmentStart = SplineCurves.ReparamTable.GetPointIndexForInputValue(RangeStart) / ReparamStepsPerSegment;
-	int32 SegmentEnd = FMath::Min(NumSegments, 1 + SplineCurves.ReparamTable.GetPointIndexForInputValue(RangeEnd) / ReparamStepsPerSegment);
+	int32 StartIndex = UE::SplineComponent::GUseSplineCurves
+		? WarninglessSplineCurves().ReparamTable.GetPointIndexForInputValue(RangeStart)
+		: Spline.GetSplineReparamTable().GetPointIndexForInputValue(RangeStart);
+	
+	int32 EndIndex = UE::SplineComponent::GUseSplineCurves
+		? WarninglessSplineCurves().ReparamTable.GetPointIndexForInputValue(RangeEnd)
+		: Spline.GetSplineReparamTable().GetPointIndexForInputValue(RangeEnd);
+	
+	int32 SegmentStart = StartIndex / ReparamStepsPerSegment;
+	int32 SegmentEnd = FMath::Min(NumSegments, 1 + EndIndex / ReparamStepsPerSegment);
 
 	TArray<FVector> NewPoints;
 	TArray<double> NewDistances;
@@ -1913,7 +2030,7 @@ bool USplineComponent::ConvertSplineToPolyline_InDistanceRange(ESplineCoordinate
 
 bool USplineComponent::ConvertSplineToPolyline_InTimeRange(ESplineCoordinateSpace::Type CoordinateSpace, const float MaxSquareDistanceFromSpline, float StartTimeAlongSpline, float EndTimeAlongSpline, bool bUseConstantVelocity, TArray<FVector>& OutPoints, TArray<double>& OutDistancesAlongSpline, bool bAllowWrappingIfClosed) const
 {
-	if (SplineCurves.Position.Points.Num() == 0)
+	if (GetNumberOfSplinePoints() == 0)
 	{
 		return false;
 	}
@@ -1928,7 +2045,7 @@ bool USplineComponent::ConvertSplineToPolyline_InTimeRange(ESplineCoordinateSpac
 		}
 		else
 		{
-			const int32 NumPoints = SplineCurves.Position.Points.Num();
+			const int32 NumPoints = GetNumberOfSplinePoints();
 			const int32 NumSegments = bClosedLoop ? NumPoints : NumPoints - 1;
 			// Note: 'InputKey' values correspond to the spline in parameter space, in the range of 0 to NumSegments
 			float InputKey = TimeFrac * NumSegments;
@@ -1957,7 +2074,6 @@ bool USplineComponent::ConvertSplineToPolyline_InTimeRange(ESplineCoordinateSpac
 	return ConvertSplineToPolyline_InDistanceRange(CoordinateSpace, MaxSquareDistanceFromSpline, TimeToDistance(StartTimeAlongSpline), TimeToDistance(EndTimeAlongSpline), OutPoints, OutDistancesAlongSpline, bAllowWrappingIfClosed);
 }
 
-
 template<class T>
 T GetPropertyValueAtSplinePoint(const USplineMetadata* Metadata, int32 Index, FName PropertyName)
 {
@@ -2009,7 +2125,7 @@ FPrimitiveSceneProxy* USplineComponent::CreateSceneProxy()
 		FSplineSceneProxy(const USplineComponent* InComponent)
 			: FPrimitiveSceneProxy(InComponent)
 			, bDrawDebug(InComponent->bDrawDebug)
-			, SplineInfo(InComponent->SplineCurves.Position)
+			, SplineInfo(UE::SplineComponent::GUseSplineCurves ? InComponent->WarninglessSplineCurves().Position : InComponent->Spline.GetSplinePointsPosition())
 #if WITH_EDITORONLY_DATA
 			, LineColor(InComponent->EditorUnselectedSplineSegmentColor)
 #else
@@ -2150,17 +2266,18 @@ void USplineComponent::Draw(FPrimitiveDrawInterface* PDI, const FSceneView* View
 
 FBoxSphereBounds USplineComponent::CalcBounds(const FTransform& LocalToWorld) const
 {
+	const auto& InterpCurve = UE::SplineComponent::GUseSplineCurves ? WarninglessSplineCurves().Position : Spline.GetSplinePointsPosition();
 
 #if SPLINE_FAST_BOUNDS_CALCULATION
 	FBox BoundingBox(0);
-	for (const auto& InterpPoint : SplineCurves.Position.Points)
+	for (const auto& InterpPoint : InterpCurve.Points)
 	{
 		BoundingBox += InterpPoint.OutVal;
 	}
 
 	return FBoxSphereBounds(BoundingBox.TransformBy(LocalToWorld));
 #else
-	const int32 NumPoints = SplineCurves.Position.Points.Num();
+	const int32 NumPoints = GetNumberOfSplinePoints();
 	const int32 NumSegments = bClosedLoop ? NumPoints : NumPoints - 1;
 
 	FVector Min(WORLD_MAX);
@@ -2171,11 +2288,11 @@ FBoxSphereBounds USplineComponent::CalcBounds(const FTransform& LocalToWorld) co
 		{
 			const bool bLoopSegment = (Index == NumPoints - 1);
 			const int32 NextIndex = bLoopSegment ? 0 : Index + 1;
-			const FInterpCurvePoint<FVector>& ThisInterpPoint = SplineCurves.Position.Points[Index];
-			FInterpCurvePoint<FVector> NextInterpPoint = SplineCurves.Position.Points[NextIndex];
+			const FInterpCurvePoint<FVector>& ThisInterpPoint = InterpCurve.Points[Index];
+			FInterpCurvePoint<FVector> NextInterpPoint = InterpCurve.Points[NextIndex];
 			if (bLoopSegment)
 			{
-				NextInterpPoint.InVal = ThisInterpPoint.InVal + SplineCurves.Position.LoopKeyOffset;
+				NextInterpPoint.InVal = ThisInterpPoint.InVal + InterpCurve.LoopKeyOffset;
 			}
 
 			CurveVectorFindIntervalBounds(ThisInterpPoint, NextInterpPoint, Min, Max);
@@ -2183,7 +2300,7 @@ FBoxSphereBounds USplineComponent::CalcBounds(const FTransform& LocalToWorld) co
 	}
 	else if (NumPoints == 1)
 	{
-		Min = Max = SplineCurves.Position.Points[0].OutVal;
+		Min = Max = InterpCurve.Points[0].OutVal;
 	}
 	else
 	{
@@ -2198,7 +2315,7 @@ FBoxSphereBounds USplineComponent::CalcBounds(const FTransform& LocalToWorld) co
 bool USplineComponent::GetIgnoreBoundsForEditorFocus() const
 {
 	// Cannot compute proper bounds when there's no point so don't participate to editor focus if that's the case : 
-	return Super::GetIgnoreBoundsForEditorFocus() || SplineCurves.Position.Points.Num() == 0;
+	return Super::GetIgnoreBoundsForEditorFocus() || GetNumberOfSplinePoints() == 0;
 }
 
 TStructOnScope<FActorComponentInstanceData> USplineComponent::GetComponentInstanceData() const
@@ -2208,7 +2325,7 @@ TStructOnScope<FActorComponentInstanceData> USplineComponent::GetComponentInstan
 	
 	if (bSplineHasBeenEdited)
 	{
-		SplineInstanceData->SplineCurves = SplineCurves;
+		SplineInstanceData->SplineCurves = GetSplineCurves();
 		SplineInstanceData->bClosedLoop = bClosedLoop;
 	}
 	SplineInstanceData->bSplineHasBeenEdited = bSplineHasBeenEdited;
@@ -2216,6 +2333,59 @@ TStructOnScope<FActorComponentInstanceData> USplineComponent::GetComponentInstan
 	return InstanceData;
 }
 
+FSplineCurves USplineComponent::GetSplineCurves() const
+{
+	return UE::SplineComponent::GUseSplineCurves
+		? WarninglessSplineCurves()
+		: FSplineCurves::FromSplineInterface(Spline);
+}
+
+void USplineComponent::SetSpline(const FSpline& InSpline)
+{
+	WarninglessSplineCurves() = FSplineCurves::FromSplineInterface(InSpline);
+	Spline = InSpline;
+}
+
+void USplineComponent::SetSpline(const FSplineCurves& InSplineCurves)
+{
+	WarninglessSplineCurves() = InSplineCurves;
+	Spline = FSpline::FromSplineCurves(InSplineCurves);
+}
+
+int32 USplineComponent::GetVersion() const
+{
+	return UE::SplineComponent::GUseSplineCurves
+		? WarninglessSplineCurves().Version
+		: Spline.GetVersion();
+}
+
+const FInterpCurveVector& USplineComponent::GetSplinePointsPosition() const
+{
+	return UE::SplineComponent::GUseSplineCurves
+		? WarninglessSplineCurves().Position
+		: Spline.GetSplinePointsPosition();
+}
+
+const FInterpCurveQuat& USplineComponent::GetSplinePointsRotation() const
+{
+	return UE::SplineComponent::GUseSplineCurves
+		? WarninglessSplineCurves().Rotation
+		: Spline.GetSplinePointsRotation();
+}
+
+const FInterpCurveVector& USplineComponent::GetSplinePointsScale() const
+{
+	return UE::SplineComponent::GUseSplineCurves
+		? WarninglessSplineCurves().Scale
+		: Spline.GetSplinePointsScale();
+}
+
+const FInterpCurveFloat& USplineComponent::GetSplineReparamTable() const
+{
+	return UE::SplineComponent::GUseSplineCurves
+		? WarninglessSplineCurves().ReparamTable
+		: Spline.GetSplineReparamTable();
+}
 
 TArray<ESplinePointType::Type> USplineComponent::GetEnabledSplinePointTypes() const
 {
@@ -2243,23 +2413,23 @@ void USplineComponent::ApplyComponentInstanceData(FSplineInstanceData* SplineIns
 		}
 		else
 		{
-			bModifiedByConstructionScript = (SplineInstanceData->SplineCurvesPreUCS != SplineCurves);
+			bModifiedByConstructionScript = SplineInstanceData->SplineCurvesPreUCS != GetSplineCurves();
 
 			// If we are restoring the saved state, unmark the SplineCurves property as 'modified'.
 			// We don't want to consider that these changes have been made through the UCS.
 			TArray<FProperty*> Properties;
-			Properties.Emplace(FindFProperty<FProperty>(USplineComponent::StaticClass(), GET_MEMBER_NAME_CHECKED(USplineComponent, SplineCurves)));
+			Properties.Emplace(FindFProperty<FProperty>(StaticClass(), GetSplinePropertyName()));
 			RemoveUCSModifiedProperties(Properties);
 		}
 	}
 	else
 	{
-		SplineInstanceData->SplineCurvesPreUCS = SplineCurves;
+		SplineInstanceData->SplineCurvesPreUCS = GetSplineCurves();
 	}
 
 	if (SplineInstanceData->bSplineHasBeenEdited)
 	{
-		SplineCurves = SplineInstanceData->SplineCurves;
+		SetSpline(SplineInstanceData->SplineCurves);
 		bClosedLoop = SplineInstanceData->bClosedLoop;
 		bModifiedByConstructionScript = false;
 	}
@@ -2269,7 +2439,6 @@ void USplineComponent::ApplyComponentInstanceData(FSplineInstanceData* SplineIns
 	UpdateSpline();
 }
 
-
 #if WITH_EDITOR
 void USplineComponent::PostEditChangeChainProperty(FPropertyChangedChainEvent& PropertyChangedEvent)
 {
@@ -2311,5 +2480,4 @@ void FSplinePositionLinearApproximation::Build(const FSplineCurves& InCurves, TA
 	}
 
 	OutPoints.Emplace(InCurves.Position.Points.Last().OutVal, InCurves.ReparamTable.Points.Last().OutVal);
-}
-
+}
\ No newline at end of file
diff --git a/Engine/Source/Runtime/Engine/Private/Curves/Spline.cpp b/Engine/Source/Runtime/Engine/Private/Curves/Spline.cpp
new file mode 100644
--- /dev/null
+++ b/Engine/Source/Runtime/Engine/Private/Curves/Spline.cpp
@@ -0,0 +1,1276 @@
+// Copyright Epic Games, Inc. All Rights Reserved.
+
+#include "Curves/Spline.h"
+
+#include "Components/SplineComponent.h"	// for FSplineCurves, ideally removed
+#include "UObject/FortniteMainBranchObjectVersion.h"
+
+#include UE_INLINE_GENERATED_CPP_BY_NAME(Spline)
+
+namespace UE::Spline
+{
+
+int32 GImplementation = 1;
+#if WITH_EDITOR
+static FSimpleMulticastDelegate& GetOnSplineImplementationChanged()
+{
+	static FSimpleMulticastDelegate OnSplineImplementationChanged;
+	return OnSplineImplementationChanged;
+};
+static void SplineImplementationSink(IConsoleVariable*)
+{
+	static int32 GPreviousImplementation = -1;
+	if (GPreviousImplementation != GImplementation)
+	{
+		GImplementation = FMath::Clamp(GImplementation, 0, 2);
+		GetOnSplineImplementationChanged().Broadcast();
+		GPreviousImplementation = GImplementation;
+	}
+}
+FAutoConsoleVariableRef CVarSplineImplementation(
+	TEXT("Spline.Implementation"),
+	GImplementation,
+	TEXT("0) Not Implemented - 1) Legacy Implementation - 2) New Implementation"),
+	FConsoleVariableDelegate::CreateStatic(&SplineImplementationSink)
+);
+#endif
+	
+}
+
+FSpline::FSpline()
+: CurrentImplementation(UE::Spline::GImplementation)
+, Version(0xffffffff)
+{
+#if WITH_EDITOR
+	PreviousImplementation = UE::Spline::GImplementation;
+#endif
+	
+	switch (CurrentImplementation)
+	{
+	default:	// Intentionally doing nothing
+		break;
+	case 1:
+		LegacyData = MakeShared<FLegacySpline>();
+		break;
+	case 2:
+		// todo: implement for new spline
+		break;
+	}
+
+#if WITH_EDITOR
+	OnSplineImplementationChangedHandle = UE::Spline::GetOnSplineImplementationChanged().AddRaw(this, &FSpline::OnSplineImplementationChanged);
+#endif
+}
+
+FSpline::FSpline(const FSpline& Other)
+	: FSpline()
+{
+	*this = Other;
+}
+
+FSpline& FSpline::operator=(const FSpline& Other)
+{
+	CurrentImplementation = Other.CurrentImplementation;
+
+	switch (CurrentImplementation)
+	{
+	default:
+		break;
+	case 1:
+		if (Other.LegacyData)
+			{ LegacyData = MakeShared<FLegacySpline>(*Other.LegacyData); }
+		else
+			{ LegacyData = MakeShared<FLegacySpline>(); }
+		break;
+	case 2:
+		// if (Other.NewData)
+		// 	{ NewData = MakeShared<FHermiteSpline>(*Other.NewData); }
+		// else
+		// 	{ NewData = MakeShared<FHermiteSpline>(); }
+		break;
+	}
+	
+	return *this;
+}
+
+FSpline::~FSpline()
+{
+#if WITH_EDITOR
+	if (OnSplineImplementationChangedHandle.IsValid())
+	{
+		UE::Spline::GetOnSplineImplementationChanged().Remove(OnSplineImplementationChangedHandle);
+		OnSplineImplementationChangedHandle.Reset();
+	}
+#endif
+}
+
+FSpline FSpline::FromSplineCurves(const FSplineCurves& InSpline)
+{
+	FSpline NewSpline;
+	
+#if WITH_EDITOR
+	NewSpline.PreviousImplementation = UE::Spline::GImplementation;
+#endif
+	NewSpline.CurrentImplementation = UE::Spline::GImplementation;
+
+	switch (NewSpline.CurrentImplementation)
+	{
+	default:
+		// Intentionally doing nothing
+		break;
+	case 1:
+		NewSpline.LegacyData = MakeShared<FLegacySpline>(FLegacySpline::FromSplineCurves(InSpline));
+		break;
+	case 2:
+		// todo: implement for new spline
+		break;
+	}
+	
+	return NewSpline;
+}
+
+void FSpline::AddPoint(const FSplinePoint& InPoint)
+{
+	switch (CurrentImplementation)
+	{
+	default:
+		// Intentionally doing nothing
+		break;
+	case 1:
+		LegacyData->AddPoint(InPoint);
+		break;
+	case 2:
+		// todo: implement for new spline
+		break;
+	}
+}
+
+void FSpline::AddPointAtIndex(const FSplinePoint& InPoint, int32 Index)
+{
+	switch (CurrentImplementation)
+	{
+	default:
+		// Intentionally doing nothing
+		break;
+	case 1:
+		LegacyData->AddPointAtIndex(InPoint, Index);
+		break;
+	case 2:
+		// todo: implement for new spline
+		break;
+	}
+}
+
+FSplinePoint FSpline::GetPoint(const int32 Index) const
+{
+	switch (CurrentImplementation)
+	{
+	default:
+		return FSplinePoint();
+	case 1:
+		return LegacyData->GetPoint(Index);
+	case 2:
+		// todo: implement for new spline
+		return FSplinePoint();
+	}
+}
+
+void FSpline::RemovePoint(const int32 Index)
+{
+	switch (CurrentImplementation)
+	{
+	default:
+		// Intentionally doing nothing
+		break;
+	case 1:
+		LegacyData->RemovePoint(Index);
+		break;
+	case 2:
+		// todo: implement for new spline
+		break;
+	}
+}
+
+void FSpline::SetLocation(int32 PointIndex, const FVector& InLocation)
+{
+	switch (CurrentImplementation)
+	{
+	default:
+		// Intentionally doing nothing
+		break;
+	case 1:
+		LegacyData->SetLocation(PointIndex, InLocation);
+		break;
+	case 2:
+		// todo: implement for new spline
+		break;
+	}
+}
+
+FVector FSpline::GetLocation(const int32 Index) const
+{
+	switch (CurrentImplementation)
+	{
+	default:
+		// Intentionally doing nothing
+		return FVector();
+	case 1:
+		return LegacyData->GetLocation(Index);
+	case 2:
+		// todo: implement for new spline
+		return FVector();
+	}
+}
+
+void FSpline::SetInTangent(const int32 Index, const FVector& InTangent)
+{
+	switch (CurrentImplementation)
+	{
+	default:
+		// Intentionally doing nothing
+		break;
+	case 1:
+		LegacyData->SetInTangent(Index, InTangent);
+		break;
+	case 2:
+		// todo: implement for new spline
+		break;
+	}
+}
+
+FVector FSpline::GetInTangent(const int32 Index) const
+{
+	switch (CurrentImplementation)
+	{
+	default:
+		// Intentionally doing nothing
+		return FVector();
+	case 1:
+		return LegacyData->GetInTangent(Index);
+	case 2:
+		// todo: implement for new spline
+		return FVector();
+	}
+}
+
+void FSpline::SetOutTangent(const int32 Index, const FVector& OutTangent)
+{
+	switch (CurrentImplementation)
+	{
+	default:
+		// Intentionally doing nothing
+		break;
+	case 1:
+		LegacyData->SetOutTangent(Index, OutTangent);
+		break;
+	case 2:
+		// todo: implement for new spline
+		break;
+	}
+}
+
+FVector FSpline::GetOutTangent(const int32 Index) const
+{
+	switch (CurrentImplementation)
+	{
+	default:
+		// Intentionally doing nothing
+		return FVector();
+	case 1:
+		return LegacyData->GetOutTangent(Index);
+	case 2:
+		// todo: implement for new spline
+		return FVector();
+	}
+}
+
+void FSpline::SetRotation(int32 PointIndex, const FQuat& InRotation)
+{
+	switch (CurrentImplementation)
+	{
+	default:
+		// Intentionally doing nothing
+		break;
+	case 1:
+		LegacyData->SetRotation(PointIndex, InRotation);
+		break;
+	case 2:
+		// todo: implement for new spline
+		break;
+	}
+}
+
+FQuat FSpline::GetRotation(const int32 Index) const
+{
+	switch (CurrentImplementation)
+	{
+	default:
+		// Intentionally doing nothing
+		return FQuat();
+	case 1:
+		return LegacyData->GetRotation(Index);
+	case 2:
+		// todo: implement for new spline
+		return FQuat();
+	}
+}
+
+void FSpline::SetScale(int32 PointIndex, const FVector& InScale)
+{
+	switch (CurrentImplementation)
+	{
+	default:
+		// Intentionally doing nothing
+		break;
+	case 1:
+		LegacyData->SetScale(PointIndex, InScale);
+		break;
+	case 2:
+		// todo: implement for new spline
+		break;
+	}
+}
+
+FVector FSpline::GetScale(const int32 Index) const
+{
+	switch (CurrentImplementation)
+	{
+	default:
+		// Intentionally doing nothing
+		return FVector(1.0);
+	case 1:
+		return LegacyData->GetScale(Index);
+	case 2:
+		// todo: implement for new spline
+		return FVector(1.0);
+	}
+}
+
+void FSpline::SetSplinePointType(int32 PointIndex, EInterpCurveMode Type)
+{
+	switch (CurrentImplementation)
+	{
+	default:
+		// Intentionally doing nothing
+		break;
+	case 1:
+		LegacyData->SetSplinePointType(PointIndex, Type);
+		break;
+	case 2:
+		// todo: implement for new spline
+		break;
+	}
+}
+
+EInterpCurveMode FSpline::GetSplinePointType(int32 PointIndex) const
+{
+	switch (CurrentImplementation)
+	{
+	default:
+		// Intentionally doing nothing
+		return CIM_Unknown;
+	case 1:
+		return LegacyData->GetSplinePointType(PointIndex);
+	case 2:
+		// todo: implement for new spline
+		return CIM_Unknown;
+	}
+}
+
+float FSpline::GetSplinePointParameter(int32 PointIndex) const
+{
+	switch (CurrentImplementation)
+	{
+	default:
+		// Intentionally doing nothing
+		return 0.f;
+	case 1:
+		return LegacyData->GetSplinePointParameter(PointIndex);
+	case 2:
+		// todo: implement for new spline
+		return 0.f;
+	}
+}
+
+FVector FSpline::Evaluate(float Param) const
+{
+	switch (CurrentImplementation)
+	{
+	default:
+		// Intentionally doing nothing
+		return FVector();
+	case 1:
+		return LegacyData->Evaluate(Param);
+	case 2:
+		// todo: implement for new spline
+		return FVector();
+	}
+}
+
+FVector FSpline::EvaluateDerivative(float Param) const
+{
+	switch (CurrentImplementation)
+	{
+	default:
+		// Intentionally doing nothing
+		return FVector();
+	case 1:
+		return LegacyData->EvaluateDerivative(Param);
+	case 2:
+		// todo: implement for new spline
+		return FVector();
+	}
+}
+
+FQuat FSpline::EvaluateRotation(float Param) const
+{
+	switch (CurrentImplementation)
+	{
+	default:
+		// Intentionally doing nothing
+		return FQuat();
+	case 1:
+		return LegacyData->EvaluateRotation(Param);
+	case 2:
+		// todo: implement for new spline
+		return FQuat();
+	}
+}
+
+FVector FSpline::EvaluateScale(float Param) const
+{
+	switch (CurrentImplementation)
+	{
+	default:
+		// Intentionally doing nothing
+		return FVector();
+	case 1:
+		return LegacyData->EvaluateScale(Param);
+	case 2:
+		// todo: implement for new spline
+		return FVector();
+	}	
+}
+
+float FSpline::FindNearest(const FVector& InLocation) const
+{
+	switch (CurrentImplementation)
+	{
+	default:
+		// Intentionally doing nothing
+		return 0.f;
+	case 1:
+		return LegacyData->FindNearest(InLocation);
+	case 2:
+		// todo: implement for new spline
+		return 0.f;
+	}
+}
+
+float FSpline::FindNearestOnSegment(const FVector& InLocation, int32 SegmentIndex) const
+{
+	switch (CurrentImplementation)
+	{
+	default:
+		// Intentionally doing nothing
+		return 0.f;
+	case 1:
+		return LegacyData->FindNearestOnSegment(InLocation, SegmentIndex);
+	case 2:
+		// todo: implement for new spline
+		return 0.f;
+	}
+}
+
+bool FSpline::Serialize(FArchive& Ar)
+{
+	Ar.UsingCustomVersion(FFortniteMainBranchObjectVersion::GUID);
+	
+	// Data format:
+
+	// Byte 1 - The data format, determine by CurrentImplementation at the time of last save
+	// Remaining N Bytes - Spline data (or empty, if byte 0 is 0). The format is determined by byte 1.
+	
+	if (Ar.IsLoading())
+	{
+		// using intermediate because I am scared of serializing a 2 bit field.
+		int8 PreviousImpl;
+		Ar << PreviousImpl;
+#if WITH_EDITOR
+		PreviousImplementation = PreviousImpl;
+#endif
+
+		auto WasEnabled = [PreviousImpl]()
+			{ return PreviousImpl != 0; };
+
+		auto WasLegacy = [PreviousImpl]()
+        	{ return PreviousImpl == 1; };
+		
+		if (WasEnabled())
+		{
+			if (WasLegacy())
+			{
+				if (IsLegacy())
+				{
+					// read in the data directly to the legacy variant
+					LegacyData = MakeShared<FLegacySpline>();
+					Ar << *LegacyData;
+				}
+				else if (IsNew())
+				{
+					// read in the data to an intermediate legacy format structure
+					FLegacySpline IntermediateData;
+					Ar << IntermediateData;
+					// convert to new format
+				}
+				else if (!IsEnabled())
+				{
+					// read in the data but do nothing
+					FLegacySpline IntermediateData;
+					Ar << IntermediateData;
+				}
+			}
+			else
+			{
+				if (IsLegacy())
+				{
+					// read in the data to a new format structure
+					// convert to legacy format variant
+				}
+				else if (IsNew())
+				{
+					// read in the data directly to the new variant
+				}
+				else if (!IsEnabled())
+				{
+					// read in the data but do nothing
+				}
+				
+				// will implement in followup change.
+			}
+		}
+	}
+	else
+	{
+		uint8 CurrentImpl = CurrentImplementation;
+		Ar << CurrentImpl;
+		
+		if (IsLegacy())
+		{
+			Ar << *LegacyData;
+		}
+		else if (IsNew())
+		{
+			// todo: serialize new variant
+		}
+	}
+
+	return true;
+}
+
+const FInterpCurveVector& FSpline::GetSplinePointsPosition() const
+{
+	switch (CurrentImplementation)
+	{
+	default:
+		// Intentionally doing nothing
+		return DummySpline.GetSplinePointsPosition();
+	case 1:
+		return LegacyData->GetSplinePointsPosition();
+	case 2:
+		// todo: implement for new spline
+		return DummySpline.GetSplinePointsPosition();
+	}
+}
+
+const FInterpCurveQuat& FSpline::GetSplinePointsRotation() const
+{
+	switch (CurrentImplementation)
+	{
+	default:
+		// Intentionally doing nothing
+		return DummySpline.GetSplinePointsRotation();
+	case 1:
+		return LegacyData->GetSplinePointsRotation();
+	case 2:
+		// todo: implement for new spline
+		return DummySpline.GetSplinePointsRotation();
+	}
+}
+
+const FInterpCurveVector& FSpline::GetSplinePointsScale() const
+{
+	switch (CurrentImplementation)
+	{
+	default:
+		// Intentionally doing nothing
+		return DummySpline.GetSplinePointsScale();
+	case 1:
+		return LegacyData->GetSplinePointsScale();
+	case 2:
+		// todo: implement for new spline
+		return DummySpline.GetSplinePointsScale();
+	}
+}
+
+const FInterpCurveFloat& FSpline::GetSplineReparamTable() const
+{
+	switch (CurrentImplementation)
+	{
+	default:
+		// Intentionally doing nothing
+		return DummySpline.GetSplineReparamTable();
+	case 1:
+		return LegacyData->GetSplineReparamTable();
+	case 2:
+		// todo: implement for new spline
+		return DummySpline.GetSplineReparamTable();
+	}
+}
+
+float FSpline::GetSegmentLength(const int32 Index, const float Param, bool bClosedLoop, const FVector& Scale3D) const
+{
+	switch (CurrentImplementation)
+	{
+	default:
+		// Intentionally doing nothing
+		return 0.f;
+	case 1:
+		return LegacyData->GetSegmentLength(Index, Param, bClosedLoop, Scale3D);
+	case 2:
+		// todo: implement for new spline
+		return 0.f;
+	}
+}
+
+float FSpline::GetSplineLength() const
+{
+	switch (CurrentImplementation)
+	{
+	default:
+		// Intentionally doing nothing
+		return 0.f;
+	case 1:
+		return LegacyData->GetSplineLength();
+	case 2:
+		// todo: implement for new spline
+		return 0.f;
+	}
+}
+
+int32 FSpline::GetNumControlPoints() const
+{
+	switch (CurrentImplementation)
+	{
+	default:
+		// Intentionally doing nothing
+		return 0.f;
+	case 1:
+		return LegacyData->GetNumControlPoints();
+	case 2:
+		// todo: implement for new spline
+		return 0.f;
+	}
+}
+
+void FSpline::Reset()
+{
+	switch (CurrentImplementation)
+	{
+	default:
+		// Intentionally doing nothing
+		break;
+	case 1:
+		return LegacyData->Reset();
+	case 2:
+		// todo: implement for new spline
+		break;
+	}
+}
+	
+void FSpline::ResetRotation()
+{
+	switch (CurrentImplementation)
+	{
+	default:
+		// Intentionally doing nothing
+		break;
+	case 1:
+		return LegacyData->ResetRotation();
+	case 2:
+		// todo: implement for new spline
+		break;
+	}
+}
+	
+void FSpline::ResetScale()
+{
+	switch (CurrentImplementation)
+	{
+	default:
+		// Intentionally doing nothing
+		break;
+	case 1:
+		return LegacyData->ResetRotation();
+	case 2:
+		// todo: implement for new spline
+		break;
+	}
+}
+	
+void FSpline::ResetToDefault()
+{
+	switch (CurrentImplementation)
+	{
+	default:
+		// Intentionally doing nothing
+		break;
+	case 1:
+		return LegacyData->ResetToDefault();
+	case 2:
+		// todo: implement for new spline
+		break;
+	}
+}
+
+void FSpline::UpdateSpline(const FUpdateSplineParams& InParams)
+{
+	switch (CurrentImplementation)
+    {
+    default:
+    	// Intentionally doing nothing
+    	break;
+    case 1:
+    	return LegacyData->UpdateSpline(InParams);
+    case 2:
+    	// todo: implement for new spline
+    	break;
+    }
+}
+
+#if WITH_EDITOR
+void FSpline::OnSplineImplementationChanged()
+{
+	using namespace UE::Spline;
+
+	// This implements a state machine I drew out.
+	// I am pretty confident it covers all cases because it handles all edges of a
+	// directed digraph (of 3 nodes, the 3 possible states, with self-edges) where edges represent state transitions.
+
+	// Do you see the nice triangle?
+	
+	if (GImplementation == CurrentImplementation)			// 0->0, 1->1, 2->2
+	{
+	}
+	else if (GImplementation == 0)							// 1->0, 2->0
+	{
+		LegacyData.Reset();
+		// NewData.Reset();
+	}
+	else if (CurrentImplementation == 0)					// 0->1, 0->2
+	{
+		// Here we just need to alloc but not copy
+		// GImplementation guaranteed to be 1 or 2
+		if (GImplementation == 1)
+		{
+			LegacyData = MakeShared<FLegacySpline>();
+		}
+		else if (GImplementation == 2)
+		{
+			// NewData = MakeShared<FHermiteSpline>();
+		}
+	}
+	else if (CurrentImplementation == 1)					// 1->2
+	{
+		TSharedPtr<FLegacySpline> CurrentData = LegacyData;
+		LegacyData.Reset();
+		// NewData = MakeShared<FHermiteSpline>(*CurrentData);
+	}
+	else if (CurrentImplementation == 2)					// 2->1
+	{
+		// Cannot implement properly yet
+		// TSharedPtr<FLegacySpline> CurrentData = NewData;
+		// NewData.Reset();
+		// LegacyData = MakeShared<FLegacySpline>(*CurrentData);
+		
+		// This is temporary to make sure LegacyData is always valid when we
+		// expect it to be, but this should be removed when above code is active.
+		LegacyData = MakeShared<FLegacySpline>();
+	}
+
+	CurrentImplementation = GImplementation;
+}
+#endif
+
+
+
+
+FSpline::FLegacySpline::FLegacySpline()
+{
+	ResetToDefault();
+}
+
+FSpline::FLegacySpline FSpline::FLegacySpline::FromSplineCurves(const FSplineCurves& InSpline)
+{
+	FLegacySpline NewSpline;
+	NewSpline.PositionCurve = InSpline.Position;
+	NewSpline.RotationCurve = InSpline.Rotation;
+	NewSpline.ScaleCurve = InSpline.Scale;
+	NewSpline.ReparamTable = InSpline.ReparamTable;
+	return NewSpline;
+}
+
+void FSpline::FLegacySpline::AddPoint(const FSplinePoint& InPoint)
+{
+	auto UpperBound = [this](float Value)
+	{
+		int32 Count = PositionCurve.Points.Num();
+		int32 First = 0;
+
+		while (Count > 0)
+		{
+			const int32 Middle = Count / 2;
+			if (Value >= PositionCurve.Points[First + Middle].InVal)
+			{
+				First += Middle + 1;
+				Count -= Middle + 1;
+			}
+			else
+			{
+				Count = Middle;
+			}
+		}
+
+		return First;
+	};
+	
+	int32 Index = UpperBound(InPoint.InputKey);
+	
+	PositionCurve.Points.Insert(FInterpCurvePoint<FVector>(
+		InPoint.InputKey,
+		InPoint.Position,
+		InPoint.ArriveTangent,
+		InPoint.LeaveTangent,
+		ConvertSplinePointTypeToInterpCurveMode(InPoint.Type)
+	), Index);
+
+	RotationCurve.Points.Insert(FInterpCurvePoint<FQuat>(
+		InPoint.InputKey,
+		InPoint.Rotation.Quaternion(),
+		FQuat::Identity,
+		FQuat::Identity,
+		CIM_CurveAuto
+	), Index);
+
+	ScaleCurve.Points.Insert(FInterpCurvePoint<FVector>(
+		InPoint.InputKey,
+		InPoint.Scale,
+		FVector::ZeroVector,
+		FVector::ZeroVector,
+		CIM_CurveAuto
+	), Index);
+}
+
+void FSpline::FLegacySpline::AddPointAtIndex(const FSplinePoint& InPoint, int32 Index)
+{
+	const float InKey = (Index == 0) ? 0.0f : GetSplinePointParameter(Index - 1) + 1.0f;
+	
+	PositionCurve.Points.Insert(FInterpCurvePoint<FVector>(
+		InKey,
+		InPoint.Position,
+		InPoint.ArriveTangent,
+		InPoint.LeaveTangent,
+		ConvertSplinePointTypeToInterpCurveMode(InPoint.Type)
+	), Index);
+
+	RotationCurve.Points.Insert(FInterpCurvePoint<FQuat>(
+		InKey,
+		InPoint.Rotation.Quaternion(),
+		FQuat::Identity,
+		FQuat::Identity,
+		CIM_CurveAuto
+	), Index);
+
+	ScaleCurve.Points.Insert(FInterpCurvePoint<FVector>(
+		InKey,
+		InPoint.Scale,
+		FVector::ZeroVector,
+		FVector::ZeroVector,
+		CIM_CurveAuto
+	), Index);
+
+	// Increment all point input values after the inserted element
+	for (int32 Idx = Index + 1; Idx < PositionCurve.Points.Num(); Idx++)
+	{
+		PositionCurve.Points[Idx].InVal += 1.f;
+		RotationCurve.Points[Idx].InVal += 1.f;
+		ScaleCurve.Points[Idx].InVal += 1.f;
+	}
+}
+
+FSplinePoint FSpline::FLegacySpline::GetPoint(const int32 Index) const
+{
+	FSplinePoint Point;
+
+	if (!PositionCurve.Points.IsValidIndex(Index))
+	{
+		return Point;
+	}
+	
+	Point.InputKey = PositionCurve.Points[Index].InVal;
+	Point.Position = PositionCurve.Points[Index].OutVal;
+	Point.ArriveTangent = PositionCurve.Points[Index].ArriveTangent;
+	Point.LeaveTangent = PositionCurve.Points[Index].LeaveTangent;
+	Point.Rotation = RotationCurve.Points[Index].OutVal.Rotator();
+	Point.Scale = ScaleCurve.Points[Index].OutVal;
+	Point.Type = ConvertInterpCurveModeToSplinePointType(PositionCurve.Points[Index].InterpMode);
+	
+	return Point;
+}
+
+void FSpline::FLegacySpline::RemovePoint(int32 Index)
+{
+	if (!PositionCurve.Points.IsValidIndex(Index))
+	{
+		return;
+	}
+	
+	PositionCurve.Points.RemoveAt(Index);
+	RotationCurve.Points.RemoveAt(Index);
+	ScaleCurve.Points.RemoveAt(Index);
+
+	while (Index < GetNumControlPoints())
+	{
+		PositionCurve.Points[Index].InVal -= 1.0f;
+		RotationCurve.Points[Index].InVal -= 1.0f;
+		ScaleCurve.Points[Index].InVal -= 1.0f;
+		Index++;
+	}
+}
+
+void FSpline::FLegacySpline::SetLocation(int32 PointIndex, const FVector& InLocation)
+{
+	if (!PositionCurve.Points.IsValidIndex(PointIndex))
+	{
+		return;
+	}
+	
+	PositionCurve.Points[PointIndex].OutVal = InLocation;
+}
+
+FVector FSpline::FLegacySpline::GetLocation(const int32 PointIndex) const
+{
+	return PositionCurve.Points.IsValidIndex(PointIndex) ? PositionCurve.Points[PointIndex].OutVal : FVector();
+}
+
+void FSpline::FLegacySpline::SetInTangent(const int32 Index, const FVector& InTangent)
+{
+	if (!PositionCurve.Points.IsValidIndex(Index))
+	{
+		return;
+	}
+	
+	PositionCurve.Points[Index].ArriveTangent = InTangent;
+	PositionCurve.Points[Index].InterpMode = CIM_CurveUser;
+}
+
+FVector FSpline::FLegacySpline::GetInTangent(const int32 Index) const
+{
+	return PositionCurve.Points.IsValidIndex(Index) ? PositionCurve.Points[Index].ArriveTangent : FVector();
+}
+
+void FSpline::FLegacySpline::SetOutTangent(const int32 Index, const FVector& OutTangent)
+{
+	if (!PositionCurve.Points.IsValidIndex(Index))
+	{
+		return;
+	}
+	
+	PositionCurve.Points[Index].LeaveTangent = OutTangent;
+	PositionCurve.Points[Index].InterpMode = CIM_CurveUser;
+}
+
+FVector FSpline::FLegacySpline::GetOutTangent(const int32 Index) const
+{
+	return PositionCurve.Points.IsValidIndex(Index) ? PositionCurve.Points[Index].LeaveTangent : FVector();
+}
+
+void FSpline::FLegacySpline::SetRotation(int32 PointIndex, const FQuat& InRotation)
+{
+	if (!RotationCurve.Points.IsValidIndex(PointIndex))
+	{
+		return;
+	}
+	
+	RotationCurve.Points[PointIndex].OutVal = InRotation;
+}
+
+FQuat FSpline::FLegacySpline::GetRotation(const int32 PointIndex) const
+{
+	return RotationCurve.Points.IsValidIndex(PointIndex) ? RotationCurve.Points[PointIndex].OutVal : FQuat();
+}
+
+void FSpline::FLegacySpline::SetScale(int32 PointIndex, const FVector& InScale)
+{
+	if (!ScaleCurve.Points.IsValidIndex(PointIndex))
+	{
+		return;
+	}
+	
+	ScaleCurve.Points[PointIndex].OutVal = InScale;
+}
+
+FVector FSpline::FLegacySpline::GetScale(const int32 PointIndex) const
+{
+	return ScaleCurve.Points.IsValidIndex(PointIndex) ? ScaleCurve.Points[PointIndex].OutVal : FVector();
+}
+
+void FSpline::FLegacySpline::SetSplinePointType(int32 PointIndex, EInterpCurveMode Type)
+{
+	if (!PositionCurve.Points.IsValidIndex(PointIndex))
+	{
+		return;
+	}
+	
+	PositionCurve.Points[PointIndex].InterpMode = Type;
+}
+
+EInterpCurveMode FSpline::FLegacySpline::GetSplinePointType(const int32 PointIndex) const
+{
+	return PositionCurve.Points.IsValidIndex(PointIndex) ? PositionCurve.Points[PointIndex].InterpMode.GetValue() : CIM_Unknown;
+}
+
+float FSpline::FLegacySpline::GetSplinePointParameter(int32 PointIndex) const
+{
+	return PositionCurve.Points.IsValidIndex(PointIndex) ? PositionCurve.Points[PointIndex].InVal : 0.f;
+}
+
+FVector FSpline::FLegacySpline::Evaluate(float Param) const
+{
+	return PositionCurve.Eval(Param);
+}
+
+FVector FSpline::FLegacySpline::EvaluateDerivative(float Param) const
+{
+	return PositionCurve.EvalDerivative(Param);
+}
+
+FQuat FSpline::FLegacySpline::EvaluateRotation(float Param) const
+{
+	return RotationCurve.Eval(Param);
+}
+
+FVector FSpline::FLegacySpline::EvaluateScale(float Param) const
+{
+	return ScaleCurve.Eval(Param);
+}
+
+float FSpline::FLegacySpline::FindNearest(const FVector& InLocation) const
+{
+	float Dummy;
+	return PositionCurve.FindNearest(InLocation, Dummy);
+}
+
+float FSpline::FLegacySpline::FindNearestOnSegment(const FVector& InLocation, int32 SegmentIndex) const
+{
+	if (!PositionCurve.Points.IsValidIndex(SegmentIndex))
+	{
+		return 0.f;
+	}
+	
+	float Dummy;
+	return PositionCurve.FindNearestOnSegment(InLocation, SegmentIndex, Dummy);
+}
+
+bool FSpline::FLegacySpline::Serialize(FArchive& Ar)
+{
+	Ar.UsingCustomVersion(FFortniteMainBranchObjectVersion::GUID);
+	Ar << PositionCurve;
+	Ar << RotationCurve;
+	Ar << ScaleCurve;
+	Ar << ReparamTable;
+
+	return true;
+}
+
+float FSpline::FLegacySpline::GetSegmentLength(const int32 Index, const float Param, bool bClosedLoop, const FVector& Scale3D) const
+{
+	const int32 NumPoints = PositionCurve.Points.Num();
+	const int32 LastPoint = NumPoints - 1;
+
+	check(Index >= 0 && ((bClosedLoop && Index < NumPoints) || (!bClosedLoop && Index < LastPoint)));
+	check(Param >= 0.0f && Param <= 1.0f);
+
+	// Evaluate the length of a Hermite spline segment.
+	// This calculates the integral of |dP/dt| dt, where P(t) is the spline equation with components (x(t), y(t), z(t)).
+	// This isn't solvable analytically, so we use a numerical method (Legendre-Gauss quadrature) which performs very well
+	// with functions of this type, even with very few samples.  In this case, just 5 samples is sufficient to yield a
+	// reasonable result.
+
+	struct FLegendreGaussCoefficient
+	{
+		float Abscissa;
+		float Weight;
+	};
+
+	static const FLegendreGaussCoefficient LegendreGaussCoefficients[] =
+	{
+		{ 0.0f, 0.5688889f },
+		{ -0.5384693f, 0.47862867f },
+		{ 0.5384693f, 0.47862867f },
+		{ -0.90617985f, 0.23692688f },
+		{ 0.90617985f, 0.23692688f }
+	};
+
+	const auto& StartPoint = PositionCurve.Points[Index];
+	const auto& EndPoint = PositionCurve.Points[Index == LastPoint ? 0 : Index + 1];
+
+	const auto& P0 = StartPoint.OutVal;
+	const auto& T0 = StartPoint.LeaveTangent;
+	const auto& P1 = EndPoint.OutVal;
+	const auto& T1 = EndPoint.ArriveTangent;
+
+	// Special cases for linear or constant segments
+	if (StartPoint.InterpMode == CIM_Linear)
+	{
+		return ((P1 - P0) * Scale3D).Size() * Param;
+	}
+	else if (StartPoint.InterpMode == CIM_Constant)
+	{
+		// Special case: constant interpolation acts like distance = 0 for all p in [0, 1[ but for p == 1, the distance returned is the linear distance between start and end
+		return Param == 1.f ? ((P1 - P0) * Scale3D).Size() : 0.0f;
+	}
+
+	// Cache the coefficients to be fed into the function to calculate the spline derivative at each sample point as they are constant.
+	const FVector Coeff1 = ((P0 - P1) * 2.0f + T0 + T1) * 3.0f;
+	const FVector Coeff2 = (P1 - P0) * 6.0f - T0 * 4.0f - T1 * 2.0f;
+	const FVector Coeff3 = T0;
+
+	const float HalfParam = Param * 0.5f;
+
+	float Length = 0.0f;
+	for (const auto& LegendreGaussCoefficient : LegendreGaussCoefficients)
+	{
+		// Calculate derivative at each Legendre-Gauss sample, and perform a weighted sum
+		const float Alpha = HalfParam * (1.0f + LegendreGaussCoefficient.Abscissa);
+		const FVector Derivative = ((Coeff1 * Alpha + Coeff2) * Alpha + Coeff3) * Scale3D;
+		Length += Derivative.Size() * LegendreGaussCoefficient.Weight;
+	}
+	Length *= HalfParam;
+
+	return Length;
+}
+
+float FSpline::FLegacySpline::GetSplineLength() const
+{
+	const int32 NumPoints = ReparamTable.Points.Num();
+
+	// This is given by the input of the last entry in the remap table
+	if (NumPoints > 0)
+	{
+		return ReparamTable.Points.Last().InVal;
+	}
+
+	return 0.0f;
+}
+
+int32 FSpline::FLegacySpline::GetNumControlPoints() const
+{
+	return PositionCurve.Points.Num();
+}
+
+void FSpline::FLegacySpline::Reset()
+{
+	PositionCurve.Points.Reset();
+	RotationCurve.Points.Reset();
+	ScaleCurve.Points.Reset();
+}
+
+void FSpline::FLegacySpline::ResetRotation()
+{
+	RotationCurve.Points.Reset(PositionCurve.Points.Num());
+
+	for (int32 Count = RotationCurve.Points.Num(); Count < PositionCurve.Points.Num(); Count++)
+	{
+		RotationCurve.Points.Emplace(static_cast<float>(Count), FQuat::Identity, FQuat::Identity, FQuat::Identity, CIM_CurveAuto);
+	}
+}
+
+void FSpline::FLegacySpline::ResetScale()
+{
+	ScaleCurve.Points.Reset(PositionCurve.Points.Num());
+
+	for (int32 Count = ScaleCurve.Points.Num(); Count < PositionCurve.Points.Num(); Count++)
+	{
+		ScaleCurve.Points.Emplace(static_cast<float>(Count), FVector(1.0f), FVector::ZeroVector, FVector::ZeroVector, CIM_CurveAuto);
+	}
+}
+
+void FSpline::FLegacySpline::ResetToDefault()
+{
+	PositionCurve.Points.Reset(10);
+	RotationCurve.Points.Reset(10);
+	ScaleCurve.Points.Reset(10);
+
+	PositionCurve.Points.Emplace(0.0f, FVector(0, 0, 0), FVector::ZeroVector, FVector::ZeroVector, CIM_CurveAuto);
+	RotationCurve.Points.Emplace(0.0f, FQuat::Identity, FQuat::Identity, FQuat::Identity, CIM_CurveAuto);
+	ScaleCurve.Points.Emplace(0.0f, FVector(1.0f), FVector::ZeroVector, FVector::ZeroVector, CIM_CurveAuto);
+
+	PositionCurve.Points.Emplace(1.0f, FVector(100, 0, 0), FVector::ZeroVector, FVector::ZeroVector, CIM_CurveAuto);
+	RotationCurve.Points.Emplace(1.0f, FQuat::Identity, FQuat::Identity, FQuat::Identity, CIM_CurveAuto);
+	ScaleCurve.Points.Emplace(1.0f, FVector(1.0f), FVector::ZeroVector, FVector::ZeroVector, CIM_CurveAuto);
+}
+
+void FSpline::FLegacySpline::UpdateSpline(const FUpdateSplineParams& InParams)
+{
+	const int32 NumPoints = PositionCurve.Points.Num();
+	check(RotationCurve.Points.Num() == NumPoints && ScaleCurve.Points.Num() == NumPoints);
+
+#if DO_CHECK
+	// Ensure input keys are strictly ascending
+	for (int32 Index = 1; Index < NumPoints; Index++)
+	{
+		ensureAlways(PositionCurve.Points[Index - 1].InVal < PositionCurve.Points[Index].InVal);
+	}
+#endif
+
+	// Ensure splines' looping status matches with that of the spline component
+	if (InParams.bClosedLoop)
+	{
+		const float LastKey = PositionCurve.Points.Num() > 0 ? PositionCurve.Points.Last().InVal : 0.0f;
+		const float LoopKey = InParams.bLoopPositionOverride ? InParams.LoopPosition : LastKey + 1.0f;
+		PositionCurve.SetLoopKey(LoopKey);
+		RotationCurve.SetLoopKey(LoopKey);
+		ScaleCurve.SetLoopKey(LoopKey);
+	}
+	else
+	{
+		PositionCurve.ClearLoopKey();
+		RotationCurve.ClearLoopKey();
+		ScaleCurve.ClearLoopKey();
+	}
+
+	// Automatically set the tangents on any CurveAuto keys
+	PositionCurve.AutoSetTangents(0.0f, InParams.bStationaryEndpoints);
+	RotationCurve.AutoSetTangents(0.0f, InParams.bStationaryEndpoints);
+	ScaleCurve.AutoSetTangents(0.0f, InParams.bStationaryEndpoints);
+
+	// Now initialize the spline reparam table
+	const int32 NumSegments = InParams.bClosedLoop ? NumPoints : FMath::Max(0, NumPoints - 1);
+
+	// Start by clearing it
+	ReparamTable.Points.Reset(NumSegments * InParams.ReparamStepsPerSegment + 1);
+	float AccumulatedLength = 0.0f;
+	for (int32 SegmentIndex = 0; SegmentIndex < NumSegments; ++SegmentIndex)
+	{
+		for (int32 Step = 0; Step < InParams.ReparamStepsPerSegment; ++Step)
+		{
+			const float Param = static_cast<float>(Step) / InParams.ReparamStepsPerSegment;
+			const float SegmentLength = (Step == 0) ? 0.0f : GetSegmentLength(SegmentIndex, Param, InParams.bClosedLoop, InParams.Scale3D);
+
+			ReparamTable.Points.Emplace(SegmentLength + AccumulatedLength, SegmentIndex + Param, 0.0f, 0.0f, CIM_Linear);
+		}
+		AccumulatedLength += GetSegmentLength(SegmentIndex, 1.0f, InParams.bClosedLoop, InParams.Scale3D);
+	}
+	ReparamTable.Points.Emplace(AccumulatedLength, static_cast<float>(NumSegments), 0.0f, 0.0f, CIM_Linear);
+}
\ No newline at end of file
diff --git a/Engine/Source/Runtime/Landscape/Private/LandscapeBlueprintSupport.cpp b/Engine/Source/Runtime/Landscape/Private/LandscapeBlueprintSupport.cpp
--- a/Engine/Source/Runtime/Landscape/Private/LandscapeBlueprintSupport.cpp
+++ b/Engine/Source/Runtime/Landscape/Private/LandscapeBlueprintSupport.cpp
@@ -34,7 +34,7 @@ void ALandscapeProxy::EditorApplySpline(USplineComponent* InSplineComponent, flo
 
 			TArray<FLandscapeSplineInterpPoint> Points;
 			LandscapeSplineRaster::FPointifyFalloffs Falloffs(StartSideFalloff, EndSideFalloff);
-			LandscapeSplineRaster::Pointify(InSplineComponent->SplineCurves.Position, Points, NumSubdivisions, 0.0f, 0.0f, StartWidth, EndWidth, StartWidth, EndWidth, Falloffs, StartRoll, EndRoll);
+			LandscapeSplineRaster::Pointify(StaticCastPtr<const USplineComponent, USplineComponent>(InSplineComponent)->GetSplinePointsPosition(), Points, NumSubdivisions, 0.0f, 0.0f, StartWidth, EndWidth, StartWidth, EndWidth, Falloffs, StartRoll, EndRoll);
 
 			FTransform SplineToWorld = InSplineComponent->GetComponentTransform();
 			LandscapeSplineRaster::RasterizeSegmentPoints(GetLandscapeInfo(), MoveTemp(Points), SplineToWorld, bRaiseHeights, bLowerHeights, PaintLayer);
diff --git a/Engine/Source/Runtime/Landscape/Private/LandscapeSplines.cpp b/Engine/Source/Runtime/Landscape/Private/LandscapeSplines.cpp
--- a/Engine/Source/Runtime/Landscape/Private/LandscapeSplines.cpp
+++ b/Engine/Source/Runtime/Landscape/Private/LandscapeSplines.cpp
@@ -924,7 +924,7 @@ void ULandscapeSplinesComponent::CopyToSplineComponent(USplineComponent* SplineC
 				// Populate first point's arrive tangent
 				if (SplineComponent->GetNumberOfSplinePoints() > 0)
 				{
-					const FVector LeaveTangent = SplineComponent->SplineCurves.Position.Points[0].LeaveTangent;
+					const FVector LeaveTangent = SplineComponent->GetLeaveTangentAtSplinePoint(0, ESplineCoordinateSpace::Local);
 					SplineComponent->SetTangentsAtSplinePoint(0, ArriveTangent, LeaveTangent, ESplineCoordinateSpace::World, false);
 				}
 			}
diff --git a/Engine/Source/Runtime/NavigationSystem/Private/SplineNavModifierComponent.cpp b/Engine/Source/Runtime/NavigationSystem/Private/SplineNavModifierComponent.cpp
--- a/Engine/Source/Runtime/NavigationSystem/Private/SplineNavModifierComponent.cpp
+++ b/Engine/Source/Runtime/NavigationSystem/Private/SplineNavModifierComponent.cpp
@@ -119,7 +119,7 @@ USplineNavModifierComponent::USplineNavModifierComponent(const FObjectInitialize
 	// If a spline is already attached, store its update-checking data
 	if (const USplineComponent* Spline = Cast<USplineComponent>(AttachedSpline.GetComponent(GetOwner())))
 	{
-		SplineVersion = Spline->SplineCurves.Version;
+		SplineVersion = Spline->GetVersion();
 		SplineTransform = Spline->GetComponentTransform();
 	}
 #endif // WITH_EDITORONLY_DATA
@@ -151,7 +151,7 @@ void USplineNavModifierComponent::TickComponent(const float DeltaTime, const ELe
 		{
 			bool bRequiresNavigationUpdate = false;
 	
-			const uint32 NextVersion = Spline->SplineCurves.Version;
+			const uint32 NextVersion = Spline->GetVersion();
 			if (SplineVersion != NextVersion)
 			{
 				SplineVersion = NextVersion;
@@ -174,7 +174,7 @@ void USplineNavModifierComponent::TickComponent(const float DeltaTime, const ELe
 		else
 		{
 			// The spline just became valid; store its data and use it to update nav data
-			SplineVersion = Spline->SplineCurves.Version;
+			SplineVersion = Spline->GetVersion();
 			SplineTransform = Spline->GetComponentTransform();
 
 			UpdateNavigationWithComponentData();
