Fixed another edge case of fixup of duplicated actor folders

diff --git a/Engine/Source/Runtime/Engine/Private/Level.cpp b/Engine/Source/Runtime/Engine/Private/Level.cpp
--- a/Engine/Source/Runtime/Engine/Private/Level.cpp
+++ b/Engine/Source/Runtime/Engine/Private/Level.cpp
@@ -3306,86 +3306,55 @@ void ULevel::FixupActorFolders()
 			});
 
 			TSet<UActorFolder*> FoldersToDelete;
-			auto HasParentToDelete = [&FoldersToDelete](UActorFolder* InFolder)
-			{
-				UActorFolder* Parent = InFolder->GetParent();
-				while (Parent)
-				{
-					if (FoldersToDelete.Contains(Parent))
-					{
-						return true;
-					}
-					Parent = Parent->GetParent();
-				}
-				return false;
-			};
+			TSet<UActorFolder*> FoldersToFixup;
+			TMap<UActorFolder*, UActorFolder*> DuplicateFolders;
+			TMap<FName, UActorFolder*> KeepDuplicateFolders;
 
-			TSet<UActorFolder*> FoldersToKeep;
-			auto FindFolderToKeep = [&FoldersToKeep, HasParentToDelete](TArray<UActorFolder*>& Folders)
+			// Fill DuplicateFolders & KeepDuplicateFolders
+			for (FName& DuplicatePath : SortedDuplicatePaths)
 			{
-				UActorFolder* FolderToKeep = nullptr;
- 				// Check if folder is already marked to keep by one of its children
-				for (UActorFolder* Folder : Folders)
+				TArray<UActorFolder*>& Folders = PathToFolders.FindChecked(DuplicatePath);
+				for (int i = 0; i < Folders.Num(); ++i)
 				{
-					if (FoldersToKeep.Contains(Folder))
+					if (i == 0)
 					{
-						FolderToKeep = Folder;
-						break;
+						KeepDuplicateFolders.Add(Folders[0]->GetPath(), Folders[0]);
 					}
-				}
-				if (!FolderToKeep)
-				{
-					for (UActorFolder* Folder : Folders)
+					else
 					{
-						if (!HasParentToDelete(Folder))
-						{
-							FolderToKeep = Folder;
-							break;
-						}
+						UActorFolder* FolderToKeep = Folders[0];
+						UActorFolder* FolderToDelete = Folders[i];
+						check(!FoldersToDelete.Contains(FolderToDelete));
+						FoldersToDelete.Add(FolderToDelete);
+
+						// Cleanup FolderLabelToActorFolders from duplicate folder as we don't want GetActorFolder(Path) to resolve to a duplicate ActorFolder
+						FActorFolderSet* FoundSet = FolderLabelToActorFolders.Find(FolderToDelete->GetLabel());
+						// We should have at least 2 elements in the set (the one we keep and one or multiple duplicates)
+						check(FoundSet&& FoundSet->GetActorFolders().Num() > 1 && FoundSet->GetActorFolders().Contains(FolderToDelete));
+						verify(FoundSet->Remove(FolderToDelete));
+						check(!FoundSet->IsEmpty());
+
+						// Keep a map that will be used to retrieve the folder we keep for each duplicate/deleted folder
+						DuplicateFolders.Add(FolderToDelete, FolderToKeep);
 					}
+					FoldersToFixup.Add(Folders[i]);
 				}
-				return FolderToKeep;
-			};
+			}
 
-			TMap<UActorFolder*, UActorFolder*> DuplicateFolders;
-			for (FName& DuplicatePath : SortedDuplicatePaths)
+			// Fixup Parent Folders so that they reference a folder that we keep
+			for (UActorFolder* Folder : FoldersToFixup)
 			{
-				// Choose to keep the folder that has no parent to delete
-				TArray<UActorFolder*>& Folders = PathToFolders.FindChecked(DuplicatePath);
-				UActorFolder* FolderToKeep = FindFolderToKeep(Folders);
-				check(FolderToKeep);
-
-				// Mark folder to keep
-				check(!FoldersToDelete.Contains(FolderToKeep));
-				FoldersToKeep.Add(FolderToKeep);
-
-				// Also mark to keep the whole parent chain
-				UActorFolder* Parent = FolderToKeep->GetParent();
+				UActorFolder* Parent = Folder->GetParent();
 				while (Parent)
 				{
-					check(!FoldersToDelete.Contains(Parent));
-					FoldersToKeep.Add(Parent);
-					Parent = Parent->GetParent();
-				}
-
-				// All other folders are considered duplicated and will be marked as deleted
-				for (UActorFolder* Folder : Folders)
-				{
-					if (Folder != FolderToKeep)
+					if (UActorFolder** KeptParent = KeepDuplicateFolders.Find(Parent->GetPath()))
 					{
-						check(!FoldersToKeep.Contains(Folder));
-						FoldersToDelete.Add(Folder);
-
-						// Cleanup FolderLabelToActorFolders from duplicate folder as we don't want GetActorFolder(Path) to resolve to a duplicate ActorFolder
-						FActorFolderSet* FoundSet = FolderLabelToActorFolders.Find(Folder->GetLabel());
-						// We should have at least 2 elements in the set (the one we keep and one or multiple duplicates)
-						check(FoundSet && FoundSet->GetActorFolders().Num() > 1 && FoundSet->GetActorFolders().Contains(Folder));
-						verify(FoundSet->Remove(Folder));
-						check(!FoundSet->IsEmpty());
-
-						// Keep a map that will be used to retrieve the folder we keep for each duplicate/deleted folder
-						DuplicateFolders.Add(Folder, FolderToKeep);
+						if (Parent != *KeptParent)
+						{
+							Folder->SetParent(*KeptParent);
+						}
 					}
+					Parent = Parent = Parent->GetParent();
 				}
 			}
 
