Add IAD support to Asset Manager.
- Fix Primary Asset Data current state being updated when an async load failed.  This would prevent retrying which may succeed if after a download.
- Add ability to create an "error" streamable handle which can help propogate errors to combined handles without trying to load anything.
- Add function to gather cache pins from child handles

diff --git a/Engine/Source/Runtime/Engine/Classes/Engine/AssetManager.h b/Engine/Source/Runtime/Engine/Classes/Engine/AssetManager.h
--- a/Engine/Source/Runtime/Engine/Classes/Engine/AssetManager.h
+++ b/Engine/Source/Runtime/Engine/Classes/Engine/AssetManager.h
@@ -32,6 +32,22 @@ DECLARE_DELEGATE_OneParam(FAssetManagerAcquireResourceDelegate, bool /* bSuccess
 /** Delegate called when new asset search root is registered due to runtime asset mounting, path will not have a trailing slash */
 DECLARE_MULTICAST_DELEGATE_OneParam(FOnAddedAssetSearchRoot, const FString&);
 
+using FAssetManagerDownloadParams = FStreamableDownloadParams;
+
+struct FAssetManagerLoadParams
+{
+	/** [optional] Delegate to call when load finishes. Will be called on the next tick if asset is already loaded, or many seconds later */
+	FStreamableDelegateWithHandle OnComplete;
+	/** [optional] Delegate to call when the load is canceled. Will be called on the next tick or many seconds later */
+	FStreamableDelegateWithHandle OnCancel;
+	/** [optional] delegate that is called periodically as request updates */
+	FStreamableUpdateDelegate OnUpdate;
+	/** [optional] Priority to pass to the streaming system, higher priority will be loaded first */
+	TAsyncLoadPriority Priority = UE::StreamableManager::Private::DefaultAsyncLoadPriority;
+	/** [optional] If present, will attempt to download necessary packages with IoStoreOnDemand */
+	TOptional<FAssetManagerDownloadParams> DownloadParams;
+};
+
 /** 
  * A singleton UObject that is responsible for loading and unloading PrimaryAssets, and maintaining game-specific asset references
  * Games should override this class and change the class reference
@@ -227,21 +243,38 @@ public:
 	 * Loads a list of Primary Assets. This will start an async load of those assets, calling callback on completion.
 	 * These assets will stay in memory until explicitly unloaded.
 	 * You can wait on the returned streamable request or poll as needed.
-	 * If there is no work to do, returned handle will be null and delegate will get called before function returns.
+	 * If there is no work to do, returned handle will be null and delegate will get called.
+	 * Prefer the overloads with Param structs for new code.
 	 *
 	 * @param AssetsToLoad		List of primary assets to load
 	 * @param LoadBundles		List of bundles to load for those assets
-	 * @param DelegateToCall	Delegate that will be called on completion, may be called before function returns if assets are already loaded
+	 * @param DelegateToCall	Delegate that will be called on completion
 	 * @param Priority			Async loading priority for this request
 	 * @return					Streamable Handle that can be used to poll or wait. You do not need to keep this handle to stop the assets from being unloaded
 	 */
-	ENGINE_API virtual TSharedPtr<FStreamableHandle> LoadPrimaryAssets(const TArray<FPrimaryAssetId>& AssetsToLoad, const TArray<FName>& LoadBundles = TArray<FName>(), FStreamableDelegate DelegateToCall = FStreamableDelegate(), TAsyncLoadPriority Priority = FStreamableManager::DefaultAsyncLoadPriority);
+	virtual TSharedPtr<FStreamableHandle> LoadPrimaryAssets(const TArray<FPrimaryAssetId>& AssetsToLoad, const TArray<FName>& LoadBundles = TArray<FName>(), FStreamableDelegate DelegateToCall = FStreamableDelegate(), TAsyncLoadPriority Priority = FStreamableManager::DefaultAsyncLoadPriority);
+	/**
+	 * Loads a list of Primary Assets. This will start an async load of those assets, calling callback on completion.
+	 * These assets will stay in memory until explicitly unloaded.
+	 * You can wait on the returned streamable request or poll as needed.
+	 * If there is no work to do, returned handle will be null and complete delegate will get called.
+	 *
+	 * @param AssetsToLoad		List of primary assets to load
+	 * @param LoadBundles		List of bundles to load for those assets
+	 * @param LoadParams		Contains callbacks and async load priority
+	 * @return					Streamable Handle that can be used to poll or wait. You do not need to keep this handle to stop the assets from being unloaded
+	 */
+	ENGINE_API virtual TSharedPtr<FStreamableHandle> LoadPrimaryAssets(const TArray<FPrimaryAssetId>& AssetsToLoad, const TArray<FName>& LoadBundles, FAssetManagerLoadParams&& LoadParams);
 
-	/** Single asset wrapper */
-	ENGINE_API virtual TSharedPtr<FStreamableHandle> LoadPrimaryAsset(const FPrimaryAssetId& AssetToLoad, const TArray<FName>& LoadBundles = TArray<FName>(), FStreamableDelegate DelegateToCall = FStreamableDelegate(), TAsyncLoadPriority Priority = FStreamableManager::DefaultAsyncLoadPriority);
+	/** Single asset wrapper, see LoadPrimaryAssets */
+	virtual TSharedPtr<FStreamableHandle> LoadPrimaryAsset(const FPrimaryAssetId& AssetToLoad, const TArray<FName>& LoadBundles = TArray<FName>(), FStreamableDelegate DelegateToCall = FStreamableDelegate(), TAsyncLoadPriority Priority = FStreamableManager::DefaultAsyncLoadPriority);
+	/** Single asset wrapper, see LoadPrimaryAssets */
+	ENGINE_API virtual TSharedPtr<FStreamableHandle> LoadPrimaryAsset(const FPrimaryAssetId& AssetToLoad, const TArray<FName>& LoadBundles, FAssetManagerLoadParams&& LoadParams);
 
 	/** Loads all assets of a given type, useful for cooking */
-	ENGINE_API virtual TSharedPtr<FStreamableHandle> LoadPrimaryAssetsWithType(FPrimaryAssetType PrimaryAssetType, const TArray<FName>& LoadBundles = TArray<FName>(), FStreamableDelegate DelegateToCall = FStreamableDelegate(), TAsyncLoadPriority Priority = FStreamableManager::DefaultAsyncLoadPriority);
+	virtual TSharedPtr<FStreamableHandle> LoadPrimaryAssetsWithType(FPrimaryAssetType PrimaryAssetType, const TArray<FName>& LoadBundles = TArray<FName>(), FStreamableDelegate DelegateToCall = FStreamableDelegate(), TAsyncLoadPriority Priority = FStreamableManager::DefaultAsyncLoadPriority);
+	/** Loads all assets of a given type, useful for cooking */
+	ENGINE_API virtual TSharedPtr<FStreamableHandle> LoadPrimaryAssetsWithType(FPrimaryAssetType PrimaryAssetType, const TArray<FName>& LoadBundles, FAssetManagerLoadParams&& LoadParams);
 
 	/** 
 	 * Unloads a list of Primary Assets that were previously Loaded.
@@ -261,30 +294,56 @@ public:
 	/** 
 	 * Changes the bundle state of a set of loaded primary assets.
 	 * You can wait on the returned streamable request or poll as needed.
-	 * If there is no work to do, returned handle will be null and delegate will get called before function returns.
+	 * If there is no work to do, returned handle will be null and delegate will get called.
+	 * Prefer the overloads with Param structs for new code.
 	 *
 	 * @param AssetsToChange	List of primary assets to change state of
 	 * @param AddBundles		List of bundles to add
 	 * @param RemoveBundles		Explicit list of bundles to remove
 	 * @param RemoveAllBundles	If true, remove all existing bundles even if not in remove list
-	 * @param DelegateToCall	Delegate that will be called on completion, may be called before function returns if assets are already loaded
+	 * @param DelegateToCall	Delegate that will be called on completion
 	 * @param Priority			Async loading priority for this request
 	 * @return					Streamable Handle that can be used to poll or wait. You do not need to keep this handle to stop the assets from being unloaded
 	 */
-	ENGINE_API virtual TSharedPtr<FStreamableHandle> ChangeBundleStateForPrimaryAssets(const TArray<FPrimaryAssetId>& AssetsToChange, const TArray<FName>& AddBundles, const TArray<FName>& RemoveBundles, bool bRemoveAllBundles = false, FStreamableDelegate DelegateToCall = FStreamableDelegate(), TAsyncLoadPriority Priority = FStreamableManager::DefaultAsyncLoadPriority);
+	virtual TSharedPtr<FStreamableHandle> ChangeBundleStateForPrimaryAssets(const TArray<FPrimaryAssetId>& AssetsToChange, const TArray<FName>& AddBundles, const TArray<FName>& RemoveBundles, bool bRemoveAllBundles = false, FStreamableDelegate DelegateToCall = FStreamableDelegate(), TAsyncLoadPriority Priority = FStreamableManager::DefaultAsyncLoadPriority);
+	/**
+	 * Changes the bundle state of a set of loaded primary assets.
+	 * You can wait on the returned streamable request or poll as needed.
+	 * If there is no work to do, returned handle will be null and delegate will get called.
+	 *
+	 * @param AssetsToChange	List of primary assets to change state of
+	 * @param AddBundles		List of bundles to add
+	 * @param RemoveBundles		Explicit list of bundles to remove
+	 * @param RemoveAllBundles	If true, remove all existing bundles even if not in remove list
+	 * @param LoadParams		Contains callbacks and async load priority
+	 * @return					Streamable Handle that can be used to poll or wait. You do not need to keep this handle to stop the assets from being unloaded
+	 */
+	ENGINE_API virtual TSharedPtr<FStreamableHandle> ChangeBundleStateForPrimaryAssets(const TArray<FPrimaryAssetId>& AssetsToChange, const TArray<FName>& AddBundles, const TArray<FName>& RemoveBundles, bool bRemoveAllBundles, FAssetManagerLoadParams&& LoadParams);
 
 	/** 
 	 * Changes the bundle state of all loaded primary assets. Only assets matching OldBundles will be modified
 	 * You can wait on the returned streamable request or poll as needed.
-	 * If there is no work to do, returned handle will be null and delegate will get called before function returns.
+	 * If there is no work to do, returned handle will be null and delegate will get called.
+	 * Prefer the overloads with Param structs for new code.
 	 *
 	 * @param NewBundles		New bundle state for the assets that are changed
 	 * @param OldBundles		Old bundle state, it will remove these bundles and replace with NewBundles
-	 * @param DelegateToCall	Delegate that will be called on completion, may be called before function returns if assets are already loaded
+	 * @param DelegateToCall	Delegate that will be called on completion
 	 * @param Priority			Async loading priority for this request
 	 * @return					Streamable Handle that can be used to poll or wait. You do not need to keep this handle to stop the assets from being unloaded
 	 */
-	ENGINE_API virtual TSharedPtr<FStreamableHandle> ChangeBundleStateForMatchingPrimaryAssets(const TArray<FName>& NewBundles, const TArray<FName>& OldBundles, FStreamableDelegate DelegateToCall = FStreamableDelegate(), TAsyncLoadPriority Priority = FStreamableManager::DefaultAsyncLoadPriority);
+	virtual TSharedPtr<FStreamableHandle> ChangeBundleStateForMatchingPrimaryAssets(const TArray<FName>& NewBundles, const TArray<FName>& OldBundles, FStreamableDelegate DelegateToCall = FStreamableDelegate(), TAsyncLoadPriority Priority = FStreamableManager::DefaultAsyncLoadPriority);
+	/**
+	 * Changes the bundle state of all loaded primary assets. Only assets matching OldBundles will be modified
+	 * You can wait on the returned streamable request or poll as needed.
+	 * If there is no work to do, returned handle will be null and delegate will get called.
+	 *
+	 * @param NewBundles		New bundle state for the assets that are changed
+	 * @param OldBundles		Old bundle state, it will remove these bundles and replace with NewBundles
+	 * @param LoadParams		Contains callbacks and async load priority
+	 * @return					Streamable Handle that can be used to poll or wait. You do not need to keep this handle to stop the assets from being unloaded
+	 */
+	ENGINE_API virtual TSharedPtr<FStreamableHandle> ChangeBundleStateForMatchingPrimaryAssets(const TArray<FName>& NewBundles, const TArray<FName>& OldBundles, FAssetManagerLoadParams&& LoadParams);
 
 	/** 
 	 * Returns the loading handle associated with the primary asset, it can then be checked for progress or waited on
@@ -338,22 +397,35 @@ public:
 	/**
 	 * Preloads data for a set of assets in a specific bundle state, and returns a handle you must keep active.
 	 * These assets are not officially Loaded, so Unload/ChangeBundleState will not affect them and if you release the handle without otherwise loading the assets they will be freed.
+	 * Prefer the overloads with Param structs for new code.
 	 *
 	 * @param AssetsToLoad		List of primary assets to load
 	 * @param LoadBundles		List of bundles to load for those assets
 	 * @param bLoadRecursive	If true, this will call RecursivelyExpandBundleData and recurse into sub bundles of other primary assets loaded by a bundle reference
-	 * @param DelegateToCall	Delegate that will be called on completion, may be called before function returns if assets are already loaded
+	 * @param DelegateToCall	Delegate that will be called on completion
 	 * @param Priority			Async loading priority for this request
 	 * @return					Streamable Handle that must be stored to keep the preloaded assets from being freed
 	 */
-	ENGINE_API virtual TSharedPtr<FStreamableHandle> PreloadPrimaryAssets(const TArray<FPrimaryAssetId>& AssetsToLoad, const TArray<FName>& LoadBundles, bool bLoadRecursive, FStreamableDelegate DelegateToCall = FStreamableDelegate(), TAsyncLoadPriority Priority = FStreamableManager::DefaultAsyncLoadPriority);
+	virtual TSharedPtr<FStreamableHandle> PreloadPrimaryAssets(const TArray<FPrimaryAssetId>& AssetsToLoad, const TArray<FName>& LoadBundles, bool bLoadRecursive, FStreamableDelegate DelegateToCall = FStreamableDelegate(), TAsyncLoadPriority Priority = FStreamableManager::DefaultAsyncLoadPriority);
+	/**
+	 * Preloads data for a set of assets in a specific bundle state, and returns a handle you must keep active.
+	 * These assets are not officially Loaded, so Unload/ChangeBundleState will not affect them and if you release the handle without otherwise loading the assets they will be freed.
+	 *
+	 * @param AssetsToLoad		List of primary assets to load
+	 * @param LoadBundles		List of bundles to load for those assets
+	 * @param bLoadRecursive	If true, this will call RecursivelyExpandBundleData and recurse into sub bundles of other primary assets loaded by a bundle reference
+	 * @param LoadParams		Contains callbacks and async load priority
+	 * @return					Streamable Handle that must be stored to keep the preloaded assets from being freed
+	 */
+	ENGINE_API virtual TSharedPtr<FStreamableHandle> PreloadPrimaryAssets(const TArray<FPrimaryAssetId>& AssetsToLoad, const TArray<FName>& LoadBundles, bool bLoadRecursive, FAssetManagerLoadParams&& LoadParams);
 
 	/**
 	 * Load non primary assets with the primary streamable manager.
 	 * This will not auto release the handle, release it if needed.
+	 * Prefer the overloads with Param structs for new code.
 	 * 
 	 * @param AssetList			List of non primary assets to load
-	 * @param DelegateToCall	[optional] Delegate that will be called on completion, may be called before function returns if assets are already loaded
+	 * @param DelegateToCall	[optional] Delegate that will be called on completion
 	 * @param Priority			[optional] Async loading priority for this request
 	 * @param DebguName			[optional] Name of this handle, either FString or anything that can construct FString, will be reported in debug tools, will report Source Location if DebugName is not specified explicitly
 	 * @return					Streamable Handle that must be stored to keep the preloaded assets from being freed
@@ -364,7 +436,6 @@ public:
 		FStreamableDelegate DelegateToCall = FStreamableDelegate(),
 		TAsyncLoadPriority Priority = FStreamableManager::DefaultAsyncLoadPriority,
 		DebugNameType&& DebugNameOrLocation = UE::FSourceLocation::Current());
-
 	/** rvalue reference overload for Asset List */
 	template< typename DebugNameType = UE::FSourceLocation >
 	TSharedPtr<FStreamableHandle> LoadAssetList(
@@ -372,6 +443,26 @@ public:
 		FStreamableDelegate DelegateToCall = FStreamableDelegate(),
 		TAsyncLoadPriority Priority = FStreamableManager::DefaultAsyncLoadPriority,
 		DebugNameType&& DebugNameOrLocation = UE::FSourceLocation::Current());
+	/**
+	 * Load non primary assets with the primary streamable manager.
+	 * This will not auto release the handle, release it if needed.
+	 * 
+	 * @param AssetList			List of non primary assets to load
+	 * @param LoadParams		Contains callbacks and async load priority
+	 * @param DebguName			[optional] Name of this handle, either FString or anything that can construct FString, will be reported in debug tools, will report Source Location if DebugName is not specified explicitly
+	 * @return					Streamable Handle that must be stored to keep the preloaded assets from being freed
+	 */
+	template< typename DebugNameType = UE::FSourceLocation >
+	TSharedPtr<FStreamableHandle> LoadAssetList(
+		const TArray<FSoftObjectPath>& AssetList,
+		FAssetManagerLoadParams&& LoadParams,
+		DebugNameType&& DebugNameOrLocation = UE::FSourceLocation::Current());
+	/** rvalue reference overload for Asset List */
+	template< typename DebugNameType = UE::FSourceLocation >
+	TSharedPtr<FStreamableHandle> LoadAssetList(
+		TArray<FSoftObjectPath>&& AssetList,
+		FAssetManagerLoadParams&& LoadParams,
+		DebugNameType&& DebugNameOrLocation = UE::FSourceLocation::Current());
 
 	/** Returns a single AssetBundleInfo, matching Scope and Name */
 	ENGINE_API virtual FAssetBundleEntry GetAssetBundleEntry(const FPrimaryAssetId& BundleScope, FName BundleName) const;
@@ -727,17 +818,20 @@ protected:
 	ENGINE_API FPrimaryAssetData* GetNameData(const FPrimaryAssetId& PrimaryAssetId, bool bCheckRedirector = true);
 	ENGINE_API const FPrimaryAssetData* GetNameData(const FPrimaryAssetId& PrimaryAssetId, bool bCheckRedirector = true) const;
 
-	/** Rebuilds the ObjectReferenceList, needed after global object state has changed */
-	UE_DEPRECATED(5.3, "Function was split up into different callsites. Use OnObjectReferenceListInvalidated() if you need a hook in the same place.")
-	ENGINE_API virtual void RebuildObjectReferenceList();
-
 	ENGINE_API virtual void OnObjectReferenceListInvalidated();
 
 	ENGINE_API void CallPreGarbageCollect();
 	ENGINE_API virtual void PreGarbageCollect();
 
 	/** Called when an internal load handle finishes, handles setting to pending state */
-	ENGINE_API virtual void OnAssetStateChangeCompleted(FPrimaryAssetId PrimaryAssetId, TSharedPtr<FStreamableHandle> BoundHandle, FStreamableDelegate WrappedDelegate);
+	UE_DEPRECATED(5.6, "Use NotifyOnAssetStateChangeCompleted instead")
+	virtual void OnAssetStateChangeCompleted(FPrimaryAssetId PrimaryAssetId, TSharedPtr<FStreamableHandle> BoundHandle, FStreamableDelegate WrappedDelegate) { checkf(false, TEXT("Deprecated")); }
+
+	/** Called when an internal load handle finishes */
+	virtual void NotifyOnAssetStateChangeCompleted(const FPrimaryAssetId& PrimaryAssetId) {}
+
+	/** Called when an internal load handle finishes, handles setting to pending state */
+	ENGINE_API void OnAssetStateChangeCompleted(const TSharedPtr<FStreamableHandle>& BoundHandle, const FPrimaryAssetId& PrimaryAssetId, FStreamableDelegateWithHandle&& WrappedDelegate);
 
 	/** Scans all asset types specified in DefaultGame */
 	ENGINE_API virtual void ScanPrimaryAssetTypesFromConfig();
@@ -767,7 +861,14 @@ protected:
 	ENGINE_API virtual void OnChunkDownloaded(uint32 ChunkId, bool bSuccess);
 
 	/** Called to load asset list */
-	ENGINE_API virtual TSharedPtr<FStreamableHandle> LoadAssetListInternal(TArray<FSoftObjectPath>&& AssetList, FStreamableDelegate&& DelegateToCall, TAsyncLoadPriority Priority, FString&& DebugName);
+	UE_DEPRECATED(5.6, "Function was split up into two overloads. Use those if you need a hook in the same place.")
+	virtual TSharedPtr<FStreamableHandle> LoadAssetListInternal(TArray<FSoftObjectPath>&& AssetList, FStreamableDelegateWithHandle&& DelegateToCall, TAsyncLoadPriority Priority, FString&& DebugName)
+	{
+		checkf(false, TEXT("Deprecated"));
+		return nullptr;
+	}
+	/** Called to load asset list */
+	ENGINE_API virtual TSharedPtr<FStreamableHandle> LoadAssetListInternal(TArray<FSoftObjectPath>&& AssetList, FAssetManagerLoadParams&& LoadParams, FString&& DebugName);
 
 #if WITH_EDITOR
 	/** Function used during creating Management references to decide when to recurse and set references */
@@ -981,43 +1082,113 @@ private:
 	friend struct FCompiledAssetManagerSearchRules;
 };
 
+inline TSharedPtr<FStreamableHandle> UAssetManager::LoadPrimaryAssets(const TArray<FPrimaryAssetId>& AssetsToLoad, const TArray<FName>& LoadBundles, FStreamableDelegate DelegateToCall, TAsyncLoadPriority Priority)
+{
+	FAssetManagerLoadParams LoadParams;
+	LoadParams.OnComplete = UE::StreamableManager::Private::WrapDelegate(DelegateToCall);
+	LoadParams.Priority = Priority;
+
+	return LoadPrimaryAssets(AssetsToLoad, LoadBundles, MoveTemp(LoadParams));
+}
+
+inline TSharedPtr<FStreamableHandle> UAssetManager::LoadPrimaryAsset(const FPrimaryAssetId& AssetToLoad, const TArray<FName>& LoadBundles, FStreamableDelegate DelegateToCall, TAsyncLoadPriority Priority)
+{
+	FAssetManagerLoadParams LoadParams;
+	LoadParams.OnComplete = UE::StreamableManager::Private::WrapDelegate(DelegateToCall);
+	LoadParams.Priority = Priority;
+
+	return LoadPrimaryAsset(AssetToLoad, LoadBundles, MoveTemp(LoadParams));
+}
+
+inline TSharedPtr<FStreamableHandle> UAssetManager::LoadPrimaryAssetsWithType(FPrimaryAssetType PrimaryAssetType, const TArray<FName>& LoadBundles, FStreamableDelegate DelegateToCall, TAsyncLoadPriority Priority)
+{
+	FAssetManagerLoadParams LoadParams;
+	LoadParams.OnComplete = UE::StreamableManager::Private::WrapDelegate(DelegateToCall);
+	LoadParams.Priority = Priority;
+
+	return LoadPrimaryAssetsWithType(PrimaryAssetType, LoadBundles, MoveTemp(LoadParams));
+}
+
+inline TSharedPtr<FStreamableHandle> UAssetManager::ChangeBundleStateForPrimaryAssets(const TArray<FPrimaryAssetId>& AssetsToChange, const TArray<FName>& AddBundles, const TArray<FName>& RemoveBundles, bool bRemoveAllBundles, FStreamableDelegate DelegateToCall, TAsyncLoadPriority Priority)
+{
+	FAssetManagerLoadParams LoadParams;
+	LoadParams.OnComplete = UE::StreamableManager::Private::WrapDelegate(MoveTemp(DelegateToCall));
+	LoadParams.Priority = Priority;
+
+	return ChangeBundleStateForPrimaryAssets(AssetsToChange, AddBundles, RemoveBundles, bRemoveAllBundles, MoveTemp(LoadParams));
+}
+
+inline TSharedPtr<FStreamableHandle> UAssetManager::ChangeBundleStateForMatchingPrimaryAssets(const TArray<FName>& NewBundles, const TArray<FName>& OldBundles, FStreamableDelegate DelegateToCall, TAsyncLoadPriority Priority)
+{
+	FAssetManagerLoadParams LoadParams;
+	LoadParams.OnComplete = UE::StreamableManager::Private::WrapDelegate(DelegateToCall);
+	LoadParams.Priority = Priority;
+
+	return ChangeBundleStateForMatchingPrimaryAssets(NewBundles, OldBundles, MoveTemp(LoadParams));
+}
+
+inline TSharedPtr<FStreamableHandle> UAssetManager::PreloadPrimaryAssets(const TArray<FPrimaryAssetId>& AssetsToLoad, const TArray<FName>& LoadBundles, bool bLoadRecursive, FStreamableDelegate DelegateToCall, TAsyncLoadPriority Priority)
+{
+	FAssetManagerLoadParams LoadParams;
+	LoadParams.OnComplete = UE::StreamableManager::Private::WrapDelegate(MoveTemp(DelegateToCall));
+	LoadParams.Priority = Priority;
+
+	return PreloadPrimaryAssets(AssetsToLoad, LoadBundles, bLoadRecursive, MoveTemp(LoadParams));
+}
+
+template< typename DebugNameType >
+TSharedPtr<FStreamableHandle> UAssetManager::LoadAssetList(
+	const TArray<FSoftObjectPath>& AssetList,
+	FStreamableDelegate DelegateToCall,
+	TAsyncLoadPriority Priority,
+	DebugNameType&& DebugNameOrLocation)
+{
+	FAssetManagerLoadParams LoadParams;
+	LoadParams.OnComplete = UE::StreamableManager::Private::WrapDelegate(MoveTemp(DelegateToCall));
+	LoadParams.Priority = Priority;
+
+	return LoadAssetList(CopyTemp(AssetList), MoveTemp(LoadParams), Forward<DebugNameType>(DebugNameOrLocation));
+}
+
 template< typename DebugNameType >
 TSharedPtr<FStreamableHandle> UAssetManager::LoadAssetList(
 	TArray<FSoftObjectPath>&& AssetList,
 	FStreamableDelegate DelegateToCall,
 	TAsyncLoadPriority Priority,
 	DebugNameType&& DebugNameOrLocation)
+{
+	FAssetManagerLoadParams LoadParams;
+	LoadParams.OnComplete = UE::StreamableManager::Private::WrapDelegate(MoveTemp(DelegateToCall));
+	LoadParams.Priority = Priority;
+
+	return LoadAssetList(MoveTemp(AssetList), MoveTemp(LoadParams), Forward<DebugNameType>(DebugNameOrLocation));
+}
+
+template< typename DebugNameType >
+TSharedPtr<FStreamableHandle> UAssetManager::LoadAssetList(
+	const TArray<FSoftObjectPath>& AssetList,
+	FAssetManagerLoadParams&& LoadParams,
+	DebugNameType&& DebugNameOrLocation)
+{
+	return LoadAssetList(CopyTemp(AssetList), MoveTemp(LoadParams), Forward<DebugNameType>(DebugNameOrLocation));
+}
+
+template< typename DebugNameType >
+TSharedPtr<FStreamableHandle> UAssetManager::LoadAssetList(
+	TArray<FSoftObjectPath>&& AssetList,
+	FAssetManagerLoadParams&& LoadParams,
+	DebugNameType&& DebugNameOrLocation)
 {
 	if constexpr (std::is_same_v<std::decay_t<DebugNameType>, UE::FSourceLocation>)
 	{
 		return LoadAssetListInternal(
-			MoveTemp(AssetList),
-			MoveTemp(DelegateToCall),
-			Priority,
+			MoveTemp(AssetList), MoveTemp(LoadParams),
 			UE::SourceLocation::ToFileAndLineString(DebugNameOrLocation));
 	}
 	else
 	{
 		return LoadAssetListInternal(
-			MoveTemp(AssetList),
-			MoveTemp(DelegateToCall),
-			Priority,
+			MoveTemp(AssetList), MoveTemp(LoadParams),
 			FString{ Forward<DebugNameType>(DebugNameOrLocation) });
 	}
 }
-
-
-template< typename DebugNameType >
-TSharedPtr<FStreamableHandle> UAssetManager::LoadAssetList(
-	const TArray<FSoftObjectPath>& AssetList,
-	FStreamableDelegate DelegateToCall,
-	TAsyncLoadPriority Priority,
-	DebugNameType&& DebugNameOrLocation)
-{
-	// explicit copy
-	return LoadAssetList(
-		TArray<FSoftObjectPath>{AssetList},
-		MoveTemp(DelegateToCall),
-		Priority,
-		Forward<DebugNameType>(DebugNameOrLocation));
-}
diff --git a/Engine/Source/Runtime/Engine/Classes/Engine/StreamableManager.h b/Engine/Source/Runtime/Engine/Classes/Engine/StreamableManager.h
--- a/Engine/Source/Runtime/Engine/Classes/Engine/StreamableManager.h
+++ b/Engine/Source/Runtime/Engine/Classes/Engine/StreamableManager.h
@@ -33,7 +33,7 @@ enum class EStreamableErrorType : uint8
 	PackageLoadError,
 	/** Low-level package loading was canceled. This is different than canceling a FStreamableHandle which just detaches from the low level request. */
 	PackageLoadCanceled,
-	/** Error downloading assets, not implemented yet */
+	/** Error downloading assets */
 	DownloadError,
 };
 
@@ -128,6 +128,7 @@ struct FStreamableDownloadCachePin
 	ENGINE_API void Release();
 	ENGINE_API bool IsValid() const;
 	operator bool() const { return IsValid(); }
+	ENGINE_API bool operator==(FStreamableDownloadCachePin& Other) const;
 private:
 	UE::IoStore::FOnDemandContentHandle ContentHandle;
 	friend struct UE::StreamableManager::Private::FDownloadContext;
@@ -255,12 +256,15 @@ struct FStreamableHandle : public TSharedFromThis<FStreamableHandle>
 		return Error;
 	}
 
-	/** Get download cache pin */
+	/** Get download cache pin for this handle */
 	FStreamableDownloadCachePin GetDownloadCachePin() const
 	{
 		return DownloadCachePin;
 	}
 
+	/** Get download cache pins for this handle and all childern */
+	TArray<FStreamableDownloadCachePin> GetDownloadCachePins() const;
+
 	/** Returns true if we've done all the loading we can now, ie all handles are either completed or stalled */
 	ENGINE_API bool HasLoadCompletedOrStalled() const;
 
@@ -539,6 +543,7 @@ public:
 	static ENGINE_API const FString HandleDebugName_Preloading;
 	static ENGINE_API const FString HandleDebugName_AssetList;
 	static ENGINE_API const FString HandleDebugName_CombinedHandle;
+	static ENGINE_API const FString HandleDebugName_Error;
 
 private:
 	friend struct FStreamableManager;
@@ -816,6 +821,20 @@ struct FStreamableManager : public FGCObject
 		const FString& DebugName = FStreamableHandle::HandleDebugName_CombinedHandle,
 		EStreamableManagerCombinedHandleOptions Options = EStreamableManagerCombinedHandleOptions::None);
 
+	/**
+	 * Creates a canceled handle with an error set. Does not do any loading. Does not manage the handle.
+	 * Useful in situations where a handle needs to be returned with an error without starting an actual load occuring.
+	 *
+	 * @param Error					The error to set on the handle
+	 * @param Params				Load Parameters
+	 * @param DownloadParams		Additional Parameters to control downloading
+	 */
+	ENGINE_API TSharedPtr<FStreamableHandle> CreateErrorHandle(
+		FStreamableError&& Error,
+		FStreamableAsyncLoadParams&& Params,
+		TOptional<FStreamableDownloadParams> DownloadParams = {},
+		const FString& DebugName = FStreamableHandle::HandleDebugName_CombinedHandle);
+
 	/** 
 	 * Gets list of handles that are directly referencing this asset, returns true if any found.
 	 * Combined Handles will not be returned by this function.
@@ -1030,7 +1049,7 @@ TSharedPtr<FStreamableHandle> FStreamableManager::RequestAsyncLoad(
 {
 	if (!DownloadParams.CachePin)
 	{
-		DownloadParams.CachePin = FStreamableDownloadCachePin::Create(Forward<DebugNameType>(DebugNameOrLocation));
+		DownloadParams.CachePin = FStreamableDownloadCachePin::Create(DebugNameOrLocation);
 	}
 
 	if constexpr (std::is_same_v<std::decay_t<DebugNameType>, UE::FSourceLocation>)
diff --git a/Engine/Source/Runtime/Engine/Private/AssetManager.cpp b/Engine/Source/Runtime/Engine/Private/AssetManager.cpp
--- a/Engine/Source/Runtime/Engine/Private/AssetManager.cpp
+++ b/Engine/Source/Runtime/Engine/Private/AssetManager.cpp
@@ -364,6 +364,80 @@ GInitialBulkScan;
 static FDelayedAutoRegisterHelper GInitialBulkScanStopper(EDelayedRegisterRunPhase::StartOfEnginePreInit,
 	[] { FCoreDelegates::OnPostEngineInit.AddLambda([] { GInitialBulkScan.StopOnce(); }); } );
 
+namespace UE::AssetManager::Private
+{
+	static constexpr int MaxDebugWarningLen = 1024;
+	static constexpr int MaxEnsureMessageLen = 256;
+
+	static TArray<FSoftObjectPath> GetPreloadPrimaryAssetsLoadList(
+		const UAssetManager* Self, const TArray<FPrimaryAssetId>& AssetsToLoad, const TArray<FName>& LoadBundles, bool bLoadRecursive,
+		FStringBuilderBase& DebugValid, FStringBuilderBase& DebugInvalid)
+	{
+		TArray<FSoftObjectPath> PathsToLoad;
+		bool ExeededMaxLength = false;
+
+		for (const FPrimaryAssetId& PrimaryAssetId : AssetsToLoad)
+		{
+			if (Self->GetPrimaryAssetLoadList(PathsToLoad, PrimaryAssetId, LoadBundles, bLoadRecursive, false))
+			{
+				if (DebugValid.Len() < MaxDebugWarningLen)
+				{
+					DebugValid << (DebugValid.Len() > 0 ? TEXT(", ") : TEXT("")) << PrimaryAssetId.ToString();
+				}
+				else
+				{
+					ExeededMaxLength = true;
+				}
+			}
+			else
+			{
+				if (DebugInvalid.Len() < MaxDebugWarningLen)
+				{
+					DebugInvalid << (DebugInvalid.Len() > 0 ? TEXT(", ") : TEXT("")) << PrimaryAssetId.ToString();
+				}
+				else
+				{
+					ExeededMaxLength = true;
+				}
+			}
+		}
+
+		if (DebugInvalid.Len() > 0)
+		{
+			if (DebugValid.Len() > 0)
+			{
+				DebugInvalid << TEXT(" (Valid: ") << DebugValid << TEXT(")");
+			}
+
+			DebugInvalid << (ExeededMaxLength ? TEXT("...") : TEXT(""));
+
+			UE_LOG(LogAssetManager, Warning, TEXT("Requested preload of some Primary Assets failed: %s"), *DebugInvalid);
+
+			// Trim for shorter ensure message
+			int TrimEndLen = FMath::Max(0, DebugInvalid.Len() - MaxEnsureMessageLen);
+			if (TrimEndLen > 0)
+			{
+				DebugInvalid.RemoveSuffix(TrimEndLen);
+				DebugInvalid << TEXT("...");
+			}
+		}
+
+		return PathsToLoad;
+	}
+
+	void GetDownloadCachePinDebugName(const TArray<FPrimaryAssetId>& AssetsToLoad, FStringBuilderBase& Builder)
+	{
+		for (const FPrimaryAssetId& PrimaryAssetId : AssetsToLoad)
+		{
+			if (Builder.Len() >= MaxDebugWarningLen)
+			{
+				break;
+			}
+
+			Builder << (Builder.Len() > 0 ? TEXT(", ") : TEXT("")) << PrimaryAssetId.ToString();
+		}
+	}
+}
 
 UAssetManager::UAssetManager()
 {
@@ -1739,11 +1813,22 @@ void UAssetManager::GetPrimaryAssetTypeInfoList(TArray<FPrimaryAssetTypeInfo>& A
 	}
 }
 
-TSharedPtr<FStreamableHandle> UAssetManager::ChangeBundleStateForPrimaryAssets(const TArray<FPrimaryAssetId>& AssetsToChange, const TArray<FName>& AddBundles, const TArray<FName>& RemoveBundles, bool bRemoveAllBundles, FStreamableDelegate DelegateToCall, TAsyncLoadPriority Priority)
+TSharedPtr<FStreamableHandle> UAssetManager::ChangeBundleStateForPrimaryAssets(
+	const TArray<FPrimaryAssetId>& AssetsToChange,
+	const TArray<FName>& AddBundles,
+	const TArray<FName>& RemoveBundles,
+	bool bRemoveAllBundles,
+	FAssetManagerLoadParams&& LoadParams)
 {
+	if (LoadParams.DownloadParams && !LoadParams.DownloadParams->CachePin)
+	{
+		TStringBuilder<256> DebugName;
+		UE::AssetManager::Private::GetDownloadCachePinDebugName(AssetsToChange, DebugName);
+		LoadParams.DownloadParams->CachePin = FStreamableDownloadCachePin::Create(DebugName);
+	}
+
 	TArray<TSharedPtr<FStreamableHandle> > NewHandles, ExistingHandles;
 	TArray<FPrimaryAssetId> NewAssets;
-	TSharedPtr<FStreamableHandle> ReturnHandle;
 	int32 MoviePlayerNumAssets = 0;
 	for (const FPrimaryAssetId& PrimaryAssetId : AssetsToChange)
 	{
@@ -1762,9 +1847,15 @@ TSharedPtr<FStreamableHandle> UAssetManager::ChangeBundleStateForPrimaryAssets(c
 		{
 			FPlatformMisc::PumpEssentialAppMessages();
 
+			// If a request has ever been made with a cache pin, require one for any subsequent requests.
+			// This can help prevent user error if downloads aren't consistently requested and prevent losing cache pins
+			// when a handle is destroyed.
+			const bool bDownloadCachePinRequired = 
+				(NameData->PendingState.Handle && NameData->PendingState.Handle->GetDownloadCachePin().IsValid()) ||
+				(NameData->CurrentState.Handle && NameData->CurrentState.Handle->GetDownloadCachePin().IsValid());
+
 			// Iterate list of changes, compute new bundle set
-			bool bLoadIfNeeded = false;
-			
+
 			// Use pending state if valid
 			TArray<FName, TInlineAllocator<32>> NewBundleState;
 
@@ -1814,14 +1905,14 @@ TSharedPtr<FStreamableHandle> UAssetManager::ChangeBundleStateForPrimaryAssets(c
 				// Dynamic types can have no base asset path
 				PathsToLoad.Add(AssetPath);
 			}
-			
+
 			for (const FName& BundleName : NewBundleState)
 			{
 				FAssetBundleEntry Entry = GetAssetBundleEntry(PrimaryAssetId, BundleName);
 
 				if (Entry.IsValid())
 				{
-					for (const FTopLevelAssetPath & Path : Entry.AssetPaths)
+					for (const FTopLevelAssetPath& Path : Entry.AssetPaths)
 					{
 						PathsToLoad.AddUnique(FSoftObjectPath(Path));
 					}
@@ -1861,19 +1952,48 @@ TSharedPtr<FStreamableHandle> UAssetManager::ChangeBundleStateForPrimaryAssets(c
 				DebugName << TEXT(")");
 			}
 
-			NewHandle = LoadAssetList(MoveTemp(PathsToLoad), FStreamableDelegate(), Priority, FString::ConstructFromPtrSize(DebugName.ToString(), DebugName.Len()));
+			// Callbacks will be bound to combined handle below.
+			FAssetManagerLoadParams LoadParamsWithoutCallbacks;
+			LoadParamsWithoutCallbacks.Priority = LoadParams.Priority;
+			if (LoadParams.DownloadParams)
+			{
+				LoadParamsWithoutCallbacks.DownloadParams = CopyTemp(*LoadParams.DownloadParams);
+				NewHandle = LoadAssetList(MoveTemp(PathsToLoad), MoveTemp(LoadParamsWithoutCallbacks), FString(DebugName));
+			}
+			else if (!bDownloadCachePinRequired)
+			{
+				NewHandle = LoadAssetList(MoveTemp(PathsToLoad), MoveTemp(LoadParamsWithoutCallbacks), FString(DebugName));
+			}
+			else
+			{
+				// The user previously use download params for requests with this primary asset, so we expect them going forward.
+				// See comment above where bDownloadCachePinRequired is delcared.
+				FStreamableAsyncLoadParams StreamableLoadParams;
+				StreamableLoadParams.TargetsToStream = MoveTemp(PathsToLoad);
+				NewHandle = StreamableManager.CreateErrorHandle(
+					FStreamableError(EStreamableErrorType::DownloadError, FString::Printf(TEXT("Expected DownloadParams for %s"), *PrimaryAssetId.ToString())),
+					MoveTemp(StreamableLoadParams), {}, FString(DebugName));
+			}
 
 			if (!NewHandle.IsValid())
 			{
-				// LoadAssetList already throws an error, no need to do it here as well
+				// LoadAssetList already logs an error, no need to do it here as well
 				continue;
 			}
 
 			if (NewHandle->HasLoadCompleted())
 			{
-				// Copy right into active
-				NameData->CurrentState.BundleNames = NewBundleState;
-				NameData->CurrentState.Handle = NewHandle;
+				// If an error occured, don't update the current state with a bad state
+				if (NewHandle->HasError())
+				{
+					UE_LOG(LogAssetManager, Warning, TEXT("ChangeBundleStateForPrimaryAssetsInternal: Error (%s) loading asset! (%s)"), *LexToString(*NewHandle->GetError()), *PrimaryAssetId.ToString());
+				}
+				else
+				{
+					// Copy right into active
+					NameData->CurrentState.BundleNames = NewBundleState;
+					NameData->CurrentState.Handle = NewHandle;
+				}
 			}
 			else
 			{
@@ -1881,7 +2001,12 @@ TSharedPtr<FStreamableHandle> UAssetManager::ChangeBundleStateForPrimaryAssets(c
 				NameData->PendingState.BundleNames = NewBundleState;
 				NameData->PendingState.Handle = NewHandle;
 
-				NewHandle->BindCompleteDelegate(FStreamableDelegate::CreateUObject(this, &UAssetManager::OnAssetStateChangeCompleted, PrimaryAssetId, NewHandle, FStreamableDelegate()));
+				NewHandle->BindCompleteDelegate(FStreamableDelegateWithHandle::CreateWeakLambda(this,
+					[this, PrimaryAssetId](const TSharedPtr<FStreamableHandle>& Handle)
+					{
+						OnAssetStateChangeCompleted(Handle, PrimaryAssetId, FStreamableDelegateWithHandle());
+					}
+				));
 			}
 
 			NewHandles.Add(NewHandle);
@@ -1893,23 +2018,13 @@ TSharedPtr<FStreamableHandle> UAssetManager::ChangeBundleStateForPrimaryAssets(c
 		}
 	}
 
+	TSharedPtr<FStreamableHandle> ReturnHandle;
 	if (NewHandles.Num() > 1 || ExistingHandles.Num() > 0)
 	{
 		// If multiple handles or we have an old handle, need to make wrapper handle
 		NewHandles.Append(ExistingHandles);
 
 		ReturnHandle = StreamableManager.CreateCombinedHandle(NewHandles, FString::Printf(TEXT("%s CreateCombinedHandle"), *GetName()));
-
-		// Call delegate or bind to meta handle
-		if (ReturnHandle->HasLoadCompleted())
-		{
-			FStreamableHandle::ExecuteDelegate(MoveTemp(DelegateToCall));
-		}
-		else
-		{
-			// Call external callback when completed
-			ReturnHandle->BindCompleteDelegate(MoveTemp(DelegateToCall));
-		}
 	}
 	else if (NewHandles.Num() == 1)
 	{
@@ -1917,37 +2032,54 @@ TSharedPtr<FStreamableHandle> UAssetManager::ChangeBundleStateForPrimaryAssets(c
 		ensure(NewAssets.Num() == 1);
 
 		// If only one handle, return it and add callback
+		// Call internal callback and external callback when it finishes
+		LoadParams.OnComplete = FStreamableDelegateWithHandle::CreateWeakLambda(this,
+			[this, PrimaryAssetId = NewAssets[0], DelegateToCall = MoveTemp(LoadParams.OnComplete)](const TSharedPtr<FStreamableHandle>& Handle) mutable
+			{
+				OnAssetStateChangeCompleted(Handle, PrimaryAssetId, MoveTemp(DelegateToCall));
+			});
+	}
+
+	if (ReturnHandle)
+	{
+		// Call delegate or bind to handle
 		if (ReturnHandle->HasLoadCompleted())
 		{
-			FStreamableHandle::ExecuteDelegate(MoveTemp(DelegateToCall));
+			FStreamableHandle::ExecuteDelegate(MoveTemp(LoadParams.OnComplete), ReturnHandle);
+		}
+		else if (ReturnHandle->WasCanceled())
+		{
+			FStreamableHandle::ExecuteDelegate(MoveTemp(LoadParams.OnCancel), ReturnHandle);
 		}
 		else
 		{
-			// Call internal callback and external callback when it finishes
-			ReturnHandle->BindCompleteDelegate(FStreamableDelegate::CreateUObject(this, &UAssetManager::OnAssetStateChangeCompleted, NewAssets[0], ReturnHandle, MoveTemp(DelegateToCall)));
+			// Call external callback when completed
+			ReturnHandle->BindCompleteDelegate(MoveTemp(LoadParams.OnComplete));
+			ReturnHandle->BindCancelDelegate(MoveTemp(LoadParams.OnCancel));
+			ReturnHandle->BindUpdateDelegate(MoveTemp(LoadParams.OnUpdate));
 		}
 	}
 	else
 	{
 		// Call completion callback, nothing to do
-		FStreamableHandle::ExecuteDelegate(MoveTemp(DelegateToCall));
+		FStreamableHandle::ExecuteDelegate(MoveTemp(LoadParams.OnComplete));
 	}
 
 	return ReturnHandle;
 }
 
-TSharedPtr<FStreamableHandle> UAssetManager::ChangeBundleStateForMatchingPrimaryAssets(const TArray<FName>& NewBundles, const TArray<FName>& OldBundles, FStreamableDelegate DelegateToCall, TAsyncLoadPriority Priority)
+TSharedPtr<FStreamableHandle> UAssetManager::ChangeBundleStateForMatchingPrimaryAssets(const TArray<FName>& NewBundles, const TArray<FName>& OldBundles, FAssetManagerLoadParams&& LoadParams)
 {
 	TArray<FPrimaryAssetId> AssetsToChange;
 
 	if (GetPrimaryAssetsWithBundleState(AssetsToChange, TArray<FPrimaryAssetType>(), OldBundles))
 	{
 		// This will call delegate when done
-		return ChangeBundleStateForPrimaryAssets(AssetsToChange, NewBundles, OldBundles, false, MoveTemp(DelegateToCall), Priority);
+		return ChangeBundleStateForPrimaryAssets(AssetsToChange, NewBundles, OldBundles, false, MoveTemp(LoadParams));
 	}
 
 	// Nothing to transition, call delegate now
-	DelegateToCall.ExecuteIfBound();
+	FStreamableHandle::ExecuteDelegate(MoveTemp(LoadParams.OnComplete));
 	return nullptr;
 }
 
@@ -2018,81 +2150,42 @@ bool UAssetManager::GetPrimaryAssetLoadList(TArray<FSoftObjectPath>& OutAssetLoa
 	return NameData != nullptr;
 }
 
-TSharedPtr<FStreamableHandle> UAssetManager::PreloadPrimaryAssets(const TArray<FPrimaryAssetId>& AssetsToLoad, const TArray<FName>& LoadBundles, bool bLoadRecursive, FStreamableDelegate DelegateToCall, TAsyncLoadPriority Priority)
+TSharedPtr<FStreamableHandle> UAssetManager::PreloadPrimaryAssets(const TArray<FPrimaryAssetId>& AssetsToLoad, const TArray<FName>& LoadBundles, bool bLoadRecursive, FAssetManagerLoadParams&& LoadParams)
 {
-	TArray<FSoftObjectPath> PathsToLoad;
 	TStringBuilder<256> DebugValid;
 	TStringBuilder<256> DebugInvalid;
-	TSharedPtr<FStreamableHandle> ReturnHandle;
-	const int MaxDebugWarningLen = 1024;
-	const int MaxEnsureMessageLen = 256;
-	bool ExeededMaxLength = false;
-
-	for (const FPrimaryAssetId& PrimaryAssetId : AssetsToLoad)
-	{
-		if (GetPrimaryAssetLoadList(PathsToLoad, PrimaryAssetId, LoadBundles, bLoadRecursive, false))
-		{
-			if (DebugValid.Len() < MaxDebugWarningLen)
-			{
-				DebugValid << (DebugValid.Len() > 0 ? TEXT(", ") : TEXT("")) << PrimaryAssetId.ToString();
-			}
-			else
-			{
-				ExeededMaxLength = true;
-			}
-		}
-		else
-		{
-			if (DebugInvalid.Len() < MaxDebugWarningLen)
-			{
-				DebugInvalid << (DebugInvalid.Len() > 0? TEXT(", ") : TEXT("")) << PrimaryAssetId.ToString();
-			}
-			else
-			{
-				ExeededMaxLength = true;
-			}
-		}
-	}
-
-	ReturnHandle = LoadAssetList(MoveTemp(PathsToLoad), MoveTemp(DelegateToCall), Priority, FString(*DebugValid));
-
-	if (DebugInvalid.Len() > 0)
-	{
-		if (DebugValid.Len() > 0)
-		{
-			DebugInvalid << TEXT(" (Valid: ") << DebugValid << TEXT(")");
-		}
-
-		DebugInvalid << (ExeededMaxLength ? TEXT("...") : TEXT(""));
 
-		UE_LOG(LogAssetManager, Warning, TEXT("Requested preload of some Primary Assets failed: %s"), *DebugInvalid);
+	TArray<FSoftObjectPath> PathsToLoad = UE::AssetManager::Private::GetPreloadPrimaryAssetsLoadList(this, AssetsToLoad, LoadBundles, bLoadRecursive, DebugValid, DebugInvalid);
 
-		// Trim for shorter ensure message
-		int TrimEndLen = FMath::Max(0, DebugInvalid.Len() - MaxEnsureMessageLen);
-		if (TrimEndLen > 0)
-		{
-			DebugInvalid.RemoveSuffix(TrimEndLen);
-			DebugInvalid << TEXT("...");
-		}
-	}
+	TSharedPtr<FStreamableHandle> ReturnHandle = LoadAssetList(MoveTemp(PathsToLoad), MoveTemp(LoadParams), FString(*DebugValid));
 
 	ensureMsgf(ReturnHandle.IsValid(), TEXT("Requested preload of Primary Assets failed: %s"), *DebugInvalid);
 
 	return ReturnHandle;
 }
 
-void UAssetManager::OnAssetStateChangeCompleted(FPrimaryAssetId PrimaryAssetId, TSharedPtr<FStreamableHandle> BoundHandle, FStreamableDelegate WrappedDelegate)
+void UAssetManager::OnAssetStateChangeCompleted(const TSharedPtr<FStreamableHandle>& BoundHandle, const FPrimaryAssetId& PrimaryAssetId, FStreamableDelegateWithHandle&& WrappedDelegate)
 {
 	FPrimaryAssetData* NameData = GetNameData(PrimaryAssetId);
 
 	if (NameData)
 	{
+		NotifyOnAssetStateChangeCompleted(PrimaryAssetId);
+
 		if (NameData->PendingState.Handle == BoundHandle)
 		{
-			NameData->CurrentState.Handle = NameData->PendingState.Handle;
-			NameData->CurrentState.BundleNames = NameData->PendingState.BundleNames;
+			// If an error occured, don't update the current state with a bad state
+			if (BoundHandle->HasError())
+			{
+				UE_LOG(LogAssetManager, Warning, TEXT("OnAssetStateChangeCompleted: Received error (%s) loading asset! (%s)"), *LexToString(*BoundHandle->GetError()), *PrimaryAssetId.ToString());
+			}
+			else
+			{
+				NameData->CurrentState.Handle = NameData->PendingState.Handle;
+				NameData->CurrentState.BundleNames = NameData->PendingState.BundleNames;
+			}
 
-			WrappedDelegate.ExecuteIfBound();
+			WrappedDelegate.ExecuteIfBound(BoundHandle);
 
 			// Clear old state, but don't cancel handle as we just copied it into current
 			NameData->PendingState.Reset(false);
@@ -2108,21 +2201,21 @@ void UAssetManager::OnAssetStateChangeCompleted(FPrimaryAssetId PrimaryAssetId,
 	}
 }
 
-TSharedPtr<FStreamableHandle> UAssetManager::LoadPrimaryAssets(const TArray<FPrimaryAssetId>& AssetsToLoad, const TArray<FName>& LoadBundles, FStreamableDelegate DelegateToCall, TAsyncLoadPriority Priority)
+TSharedPtr<FStreamableHandle> UAssetManager::LoadPrimaryAssets(const TArray<FPrimaryAssetId>& AssetsToLoad, const TArray<FName>& LoadBundles, FAssetManagerLoadParams&& LoadParams)
 {
-	return ChangeBundleStateForPrimaryAssets(AssetsToLoad, LoadBundles, TArray<FName>(), true, MoveTemp(DelegateToCall), Priority);
+	return ChangeBundleStateForPrimaryAssets(AssetsToLoad, LoadBundles, TArray<FName>(), true, MoveTemp(LoadParams));
 }
 
-TSharedPtr<FStreamableHandle> UAssetManager::LoadPrimaryAsset(const FPrimaryAssetId& AssetToLoad, const TArray<FName>& LoadBundles, FStreamableDelegate DelegateToCall, TAsyncLoadPriority Priority)
+TSharedPtr<FStreamableHandle> UAssetManager::LoadPrimaryAsset(const FPrimaryAssetId& AssetToLoad, const TArray<FName>& LoadBundles, FAssetManagerLoadParams&& LoadParams)
 {
-	return LoadPrimaryAssets(TArray<FPrimaryAssetId>{AssetToLoad}, LoadBundles, MoveTemp(DelegateToCall), Priority);
+	return LoadPrimaryAssets(TArray<FPrimaryAssetId>{AssetToLoad}, LoadBundles, MoveTemp(LoadParams));
 }
 
-TSharedPtr<FStreamableHandle> UAssetManager::LoadPrimaryAssetsWithType(FPrimaryAssetType PrimaryAssetType, const TArray<FName>& LoadBundles, FStreamableDelegate DelegateToCall, TAsyncLoadPriority Priority)
+TSharedPtr<FStreamableHandle> UAssetManager::LoadPrimaryAssetsWithType(FPrimaryAssetType PrimaryAssetType, const TArray<FName>& LoadBundles, FAssetManagerLoadParams&& LoadParams)
 {
 	TArray<FPrimaryAssetId> Assets;
 	GetPrimaryAssetIdList(PrimaryAssetType, Assets);
-	return LoadPrimaryAssets(Assets, LoadBundles, MoveTemp(DelegateToCall), Priority);
+	return LoadPrimaryAssets(Assets, LoadBundles, MoveTemp(LoadParams));
 }
 
 TSharedPtr<FStreamableHandle> UAssetManager::GetPrimaryAssetHandle(const FPrimaryAssetId& PrimaryAssetId, bool bForceCurrent, TArray<FName>* Bundles) const
@@ -2266,7 +2359,7 @@ int32 UAssetManager::UnloadPrimaryAssetsWithType(FPrimaryAssetType PrimaryAssetT
 	return UnloadPrimaryAssets(Assets);
 }
 
-TSharedPtr<FStreamableHandle> UAssetManager::LoadAssetListInternal(TArray<FSoftObjectPath>&& AssetList, FStreamableDelegate&& DelegateToCall, TAsyncLoadPriority Priority, FString&& DebugName)
+TSharedPtr<FStreamableHandle> UAssetManager::LoadAssetListInternal(TArray<FSoftObjectPath>&& AssetList, FAssetManagerLoadParams&& LoadParams, FString&& DebugName)
 {
 	TSharedPtr<FStreamableHandle> NewHandle;
 	TArray<int32> MissingChunks, ErrorChunks;
@@ -2284,14 +2377,31 @@ TSharedPtr<FStreamableHandle> UAssetManager::LoadAssetListInternal(TArray<FSoftO
 	}
 
 	// SynchronousLoad doesn't make sense if chunks are missing
-	if (bShouldUseSynchronousLoad && MissingChunks.Num() == 0)
+	// SynchronousLoad doesn't make sense if we are going to download packages	
+	if (!LoadParams.DownloadParams && bShouldUseSynchronousLoad && MissingChunks.Num() == 0)
 	{
 		NewHandle = StreamableManager.RequestSyncLoad(MoveTemp(AssetList), false, MoveTemp(DebugName));
-		FStreamableHandle::ExecuteDelegate(MoveTemp(DelegateToCall));
+		FStreamableHandle::ExecuteDelegate(MoveTemp(LoadParams.OnComplete));
 	}
 	else
 	{
-		NewHandle = StreamableManager.RequestAsyncLoad(MoveTemp(AssetList), MoveTemp(DelegateToCall), Priority, false, MissingChunks.Num() > 0, MoveTemp(DebugName));
+		FStreamableAsyncLoadParams Params;
+		Params.TargetsToStream = MoveTemp(AssetList);
+		Params.OnComplete = MoveTemp(LoadParams.OnComplete);
+		Params.OnCancel = MoveTemp(LoadParams.OnCancel);
+		Params.OnUpdate = MoveTemp(LoadParams.OnUpdate);
+		Params.Priority = LoadParams.Priority;
+		Params.bManageActiveHandle = false;
+		Params.bStartStalled = MissingChunks.Num() > 0;
+
+		if (LoadParams.DownloadParams)
+		{
+			NewHandle = StreamableManager.RequestAsyncLoad(MoveTemp(Params), MoveTemp(*LoadParams.DownloadParams), MoveTemp(DebugName));
+		}
+		else
+		{
+			NewHandle = StreamableManager.RequestAsyncLoad(MoveTemp(Params), MoveTemp(DebugName));
+		}
 
 		if (MissingChunks.Num() > 0 && NewHandle.IsValid())
 		{
@@ -2466,13 +2576,15 @@ void UAssetManager::AcquireResourcesForAssetList(const TArray<FSoftObjectPath>&
 	if (ErrorChunks.Num() > 0)
 	{
 		// At least one chunk doesn't exist, fail
-		FStreamableDelegate TempDelegate = FStreamableDelegate::CreateLambda([CompleteDelegate = MoveTemp(CompleteDelegate), MissingChunks]() { CompleteDelegate.ExecuteIfBound(false, MissingChunks); });
+		FStreamableDelegateWithHandle TempDelegate = FStreamableDelegateWithHandle::CreateLambda(
+			[CompleteDelegate = MoveTemp(CompleteDelegate), MissingChunks](const TSharedPtr<struct FStreamableHandle>&) { CompleteDelegate.ExecuteIfBound(false, MissingChunks); });
 		FStreamableHandle::ExecuteDelegate(MoveTemp(TempDelegate));
 	}
 	else if (MissingChunks.Num() == 0)
 	{
 		// All here, schedule the callback
-		FStreamableDelegate TempDelegate = FStreamableDelegate::CreateLambda([CompleteDelegate = MoveTemp(CompleteDelegate)]() { CompleteDelegate.ExecuteIfBound(true, TArray<int32>()); });
+		FStreamableDelegateWithHandle TempDelegate = FStreamableDelegateWithHandle::CreateLambda(
+			[CompleteDelegate = MoveTemp(CompleteDelegate)](const TSharedPtr<struct FStreamableHandle>&) { CompleteDelegate.ExecuteIfBound(true, TArray<int32>()); });
 		FStreamableHandle::ExecuteDelegate(MoveTemp(TempDelegate));
 	}
 	else
@@ -2737,7 +2849,6 @@ const FPrimaryAssetData* UAssetManager::GetNameData(const FPrimaryAssetId& Prima
 {
 	const TSharedRef<FPrimaryAssetTypeData>* FoundType = AssetTypeMap.Find(PrimaryAssetId.PrimaryAssetType);
 
-	// Try redirected name
 	if (FoundType)
 	{
 		const FPrimaryAssetData* FoundName = (*FoundType)->GetAssets().Find(PrimaryAssetId.PrimaryAssetName);
@@ -2748,6 +2859,7 @@ const FPrimaryAssetData* UAssetManager::GetNameData(const FPrimaryAssetId& Prima
 		}
 	}
 
+	// Try redirected name
 	if (bCheckRedirector)
 	{
 		FPrimaryAssetId RedirectedId = GetRedirectedPrimaryAssetId(PrimaryAssetId);
@@ -2762,10 +2874,6 @@ const FPrimaryAssetData* UAssetManager::GetNameData(const FPrimaryAssetId& Prima
 	return nullptr;
 }
 
-void UAssetManager::RebuildObjectReferenceList()
-{
-}
-
 void UAssetManager::OnObjectReferenceListInvalidated()
 {
 	bIsManagementDatabaseCurrent = false;
diff --git a/Engine/Source/Runtime/Engine/Private/StreamableManager.cpp b/Engine/Source/Runtime/Engine/Private/StreamableManager.cpp
--- a/Engine/Source/Runtime/Engine/Private/StreamableManager.cpp
+++ b/Engine/Source/Runtime/Engine/Private/StreamableManager.cpp
@@ -72,10 +72,15 @@ FStreamableDownloadCachePin::~FStreamableDownloadCachePin() = default;
 FString FStreamableDownloadCachePin::GetDebugName() const { return LexToString(ContentHandle); }
 void FStreamableDownloadCachePin::Release() { ContentHandle.Reset(); }
 bool FStreamableDownloadCachePin::IsValid() const { return ContentHandle.IsValid(); }
+bool FStreamableDownloadCachePin::operator==(FStreamableDownloadCachePin& Other) const
+{
+	return ContentHandle == Other.ContentHandle;
+}
 
 const FString FStreamableHandle::HandleDebugName_Preloading = FString(TEXT("Preloading"));
 const FString FStreamableHandle::HandleDebugName_AssetList = FString(TEXT("LoadAssetList"));
 const FString FStreamableHandle::HandleDebugName_CombinedHandle = FString(TEXT("CreateCombinedHandle"));
+const FString FStreamableHandle::HandleDebugName_Error = FString(TEXT("CreateErrorHandle"));
 
 /** Helper class that defers streamable manager delegates until the next frame */
 class FStreamableDelegateDelayHelper : public FTickableGameObject
@@ -672,6 +677,30 @@ EAsyncPackageState::Type FStreamableHandle::WaitUntilComplete(float Timeout, boo
 	return State;
 }
 
+TArray<FStreamableDownloadCachePin> FStreamableHandle::GetDownloadCachePins() const
+{
+	TArray<FStreamableDownloadCachePin> Pins;
+	
+	TArray<const FStreamableHandle*, TInlineAllocator<32>> Stack;
+	Stack.Push(this);
+
+	do
+	{
+		const FStreamableHandle* Current = Stack.Pop(EAllowShrinking::No);
+		if (Current->DownloadCachePin.IsValid())
+		{
+			Pins.AddUnique(Current->DownloadCachePin);
+		}
+
+		for (const TSharedPtr<FStreamableHandle>& ChildHandle : Current->ChildHandles)
+		{
+			Stack.Push(ChildHandle.Get());
+		}
+	} while (!Stack.IsEmpty());
+
+	return Pins;
+}
+
 bool FStreamableHandle::HasLoadCompletedOrStalled() const
 {
 	// Check this handle
@@ -2342,6 +2371,30 @@ TSharedPtr<FStreamableHandle> FStreamableManager::CreateCombinedHandle(const TCo
 	return NewRequest;
 }
 
+TSharedPtr<FStreamableHandle> FStreamableManager::CreateErrorHandle(FStreamableError&& Error, FStreamableAsyncLoadParams&& Params, TOptional<FStreamableDownloadParams> DownloadParams, const FString& DebugName)
+{
+	TSharedRef<FStreamableHandle> NewRequest = MakeShared<FStreamableHandle>(FStreamableHandle::FPrivateToken());
+	NewRequest->OwningManager = this;
+	NewRequest->RequestedAssets = MoveTemp(Params.TargetsToStream);
+	NewRequest->DebugName = DebugName;
+	NewRequest->Priority = Params.Priority;
+
+	if (DownloadParams)
+	{
+		if (!DownloadParams->CachePin)
+		{
+			DownloadParams->CachePin = FStreamableDownloadCachePin::Create(DebugName);
+		}
+
+		NewRequest->DownloadCachePin = DownloadParams->CachePin;
+	}
+
+	NewRequest->CancelHandle();
+	NewRequest->SetError(MoveTemp(Error));
+
+	return NewRequest;
+}
+
 bool FStreamableManager::GetActiveHandles(const FSoftObjectPath& Target, TArray<TSharedRef<FStreamableHandle>>& HandleList, bool bOnlyManagedHandles) const
 {
 	check(IsInGameThread());
diff --git a/Engine/Source/Runtime/Experimental/IoStore/OnDemand/Public/IO/IoStoreOnDemand.h b/Engine/Source/Runtime/Experimental/IoStore/OnDemand/Public/IO/IoStoreOnDemand.h
--- a/Engine/Source/Runtime/Experimental/IoStore/OnDemand/Public/IO/IoStoreOnDemand.h
+++ b/Engine/Source/Runtime/Experimental/IoStore/OnDemand/Public/IO/IoStoreOnDemand.h
@@ -108,6 +108,11 @@ public:
 	/** Returns a string representing the content handle. */
 	UE_API friend FString LexToString(const FOnDemandContentHandle& Handle);
 
+	bool operator==(FOnDemandContentHandle& Other) const
+	{
+		return Handle == Other.Handle;
+	}
+
 private:
 	friend class FOnDemandIoStore;
 	FSharedInternalContentHandle Handle;
