- Fix for missing dll export

#rnx

diff --git a/Engine/Source/Runtime/Engine/Private/GameViewportClient.cpp b/Engine/Source/Runtime/Engine/Private/GameViewportClient.cpp
--- a/Engine/Source/Runtime/Engine/Private/GameViewportClient.cpp
+++ b/Engine/Source/Runtime/Engine/Private/GameViewportClient.cpp
@@ -3978,7 +3978,7 @@ bool UGameViewportClient::HandleKillParticlesCommand( const TCHAR* Cmd, FOutputD
 	// Don't kill in the Editor to avoid potential content clobbering.
 	if( !GIsEditor )
 	{
-		extern bool GIsAllowingParticles;
+		extern ENGINE_API bool GIsAllowingParticles;
 		// Deactivate system and kill existing particles.
 		for( TObjectIterator<UParticleSystemComponent> It; It; ++It )
 		{
