Fixed CIS - Static Analysis following CL 40207329

diff --git a/Engine/Source/Runtime/Engine/Private/Level.cpp b/Engine/Source/Runtime/Engine/Private/Level.cpp
--- a/Engine/Source/Runtime/Engine/Private/Level.cpp
+++ b/Engine/Source/Runtime/Engine/Private/Level.cpp
@@ -3354,7 +3354,7 @@ void ULevel::FixupActorFolders()
 							Folder->SetParent(*KeptParent);
 						}
 					}
-					Parent = Parent = Parent->GetParent();
+					Parent = Parent->GetParent();
 				}
 			}
 
