Restored tangent buffers to GPU Skin Cache that were mistakenly culled out.

diff --git a/Engine/Source/Runtime/Engine/Private/GPUSkinCache.cpp b/Engine/Source/Runtime/Engine/Private/GPUSkinCache.cpp
--- a/Engine/Source/Runtime/Engine/Private/GPUSkinCache.cpp
+++ b/Engine/Source/Runtime/Engine/Private/GPUSkinCache.cpp
@@ -272,17 +272,15 @@ enum class EGPUSkinCacheDispatchFlags : uint8
 	None                 = 0,
 	Position             = 1 << 0,
 	PositionPrevious     = 1 << 1,
-	Tangents             = 1 << 2
+	RecomputeTangents    = 1 << 2
 };
 ENUM_CLASS_FLAGS(EGPUSkinCacheDispatchFlags);
 
 enum class EGPUSkinCacheBufferBits : uint8
 {
 	None                 = 0,
-	Tangents             = 1 << 0,
 	IntermediateTangents = 1 << 1,
 	PositionPrevious     = 1 << 2,
-	AllTangents          = Tangents | IntermediateTangents
 };
 ENUM_CLASS_FLAGS(EGPUSkinCacheBufferBits);
 
@@ -409,15 +407,13 @@ public:
 
 		void UpdateVertexFactoryDeclaration(FRHICommandListBase& RHICmdList)
 		{
-			FRHIShaderResourceView* RHITangentSRV = TangentBuffer ? TangentBuffer->Buffer.SRV : TangentBufferSRV;
-
 			FGPUSkinPassthroughVertexFactory::FAddVertexAttributeDesc Desc;
 			Desc.FrameNumber = SourceVertexFactory->GetShaderData().UpdatedFrameNumber;
 			Desc.StreamBuffers[FGPUSkinPassthroughVertexFactory::EVertexAttribute::VertexPosition] = GetPositionRWBuffer()->Buffer.Buffer;
-			Desc.StreamBuffers[FGPUSkinPassthroughVertexFactory::EVertexAttribute::VertexTangent] = RHITangentSRV->GetBuffer();
+			Desc.StreamBuffers[FGPUSkinPassthroughVertexFactory::EVertexAttribute::VertexTangent] = TangentBuffer->Buffer.Buffer;
 			Desc.SRVs[FGPUSkinPassthroughVertexFactory::EShaderResource::Position] = GetPositionRWBuffer()->Buffer.SRV;
 			Desc.SRVs[FGPUSkinPassthroughVertexFactory::EShaderResource::PreviousPosition] = GetPreviousPositionRWBuffer()->Buffer.SRV;
-			Desc.SRVs[FGPUSkinPassthroughVertexFactory::EShaderResource::Tangent] = RHITangentSRV;
+			Desc.SRVs[FGPUSkinPassthroughVertexFactory::EShaderResource::Tangent] = TangentBuffer->Buffer.SRV;
 			TargetVertexFactory->SetVertexAttributes(RHICmdList, SourceVertexFactory, Desc);
 		}
 	};
@@ -1577,6 +1573,9 @@ void FGPUSkinCache::DispatchPassExecute(FRHICommandList& RHICmdList, FDispatchSe
 
 			uint32 VertexCountAlign64 = FMath::DivideAndRoundUp(DispatchData.NumVertices, (uint32)64);
 
+			FSkinCacheRWBuffer* TangentBuffer = DispatchData.GetActiveTangentRWBuffer();
+			check(TangentBuffer);
+
 			if (EnumHasAnyFlags(DispatchData.DispatchFlags, EGPUSkinCacheDispatchFlags::PositionPrevious))
 			{
 				const FVertexBufferAndSRV& PrevBoneBuffer = ShaderData.GetBoneBufferForReading(true);
@@ -1588,7 +1587,7 @@ void FGPUSkinCache::DispatchPassExecute(FRHICommandList& RHICmdList, FDispatchSe
 					Entry,
 					DispatchData,
 					DispatchData.GetPreviousPositionRWBuffer()->Buffer.UAV,
-					DispatchData.GetActiveTangentRWBuffer() ? DispatchData.GetActiveTangentRWBuffer()->Buffer.UAV.GetReference() : nullptr
+					TangentBuffer->Buffer.UAV
 				);
 
 				INC_DWORD_STAT_BY(STAT_GPUSkinCache_TotalNumVertices, VertexCountAlign64 * 64);
@@ -1608,7 +1607,7 @@ void FGPUSkinCache::DispatchPassExecute(FRHICommandList& RHICmdList, FDispatchSe
 					Entry,
 					DispatchData,
 					DispatchData.GetPositionRWBuffer()->Buffer.UAV,
-					DispatchData.GetActiveTangentRWBuffer() ? DispatchData.GetActiveTangentRWBuffer()->Buffer.UAV.GetReference() : nullptr
+					TangentBuffer->Buffer.UAV
 				);
 
 				INC_DWORD_STAT_BY(STAT_GPUSkinCache_TotalNumVertices, VertexCountAlign64 * 64);
@@ -1623,10 +1622,7 @@ void FGPUSkinCache::DispatchPassExecute(FRHICommandList& RHICmdList, FDispatchSe
 				LastShaderIndex = -1;
 			}
 
-			if (EnumHasAnyFlags(DispatchData.DispatchFlags, EGPUSkinCacheDispatchFlags::Tangents))
-			{
-				BuffersToTransitionToRead.Add(DispatchData.GetTangentRWBuffer());
-			}
+			BuffersToTransitionToRead.Add(DispatchData.GetTangentRWBuffer());
 
 			check(DispatchData.PreviousPositionBuffer != DispatchData.PositionBuffer);
 		}
@@ -1646,7 +1642,7 @@ void FGPUSkinCache::DispatchPassExecute(FRHICommandList& RHICmdList, FDispatchSe
 	{
 		FGPUSkinCacheEntry::FSectionDispatchData& DispatchData = DispatchItem.SkinCacheEntry->DispatchData[DispatchItem.Section];
 
-		if (EnumHasAnyFlags(DispatchData.DispatchFlags, EGPUSkinCacheDispatchFlags::Tangents))
+		if (EnumHasAnyFlags(DispatchData.DispatchFlags, EGPUSkinCacheDispatchFlags::RecomputeTangents))
 		{
 			BuffersToSRVForRecomputeTangents.Add(DispatchData.GetPositionRWBuffer());
 			BuffersToSRVForRecomputeTangents.Add(DispatchData.GetActiveTangentRWBuffer());
@@ -1673,7 +1669,7 @@ void FGPUSkinCache::DispatchPassExecute(FRHICommandList& RHICmdList, FDispatchSe
 			for (FDispatchEntry& DispatchItem : BatchDispatches)
 			{
 				FGPUSkinCacheEntry::FSectionDispatchData& DispatchData = DispatchItem.SkinCacheEntry->DispatchData[DispatchItem.Section];
-				if (EnumHasAnyFlags(DispatchData.DispatchFlags, EGPUSkinCacheDispatchFlags::Tangents))
+				if (EnumHasAnyFlags(DispatchData.DispatchFlags, EGPUSkinCacheDispatchFlags::RecomputeTangents))
 				{
 					DispatchUpdateSkinTangents(RHICmdList, DispatchItem.SkinCacheEntry, DispatchItem.Section, StagingBuffer, true);
 
@@ -1692,7 +1688,7 @@ void FGPUSkinCache::DispatchPassExecute(FRHICommandList& RHICmdList, FDispatchSe
 				for (FDispatchEntry& DispatchItem : BatchDispatches)
 				{
 					FGPUSkinCacheEntry::FSectionDispatchData& DispatchData = DispatchItem.SkinCacheEntry->DispatchData[DispatchItem.Section];
-					if (EnumHasAnyFlags(DispatchData.DispatchFlags, EGPUSkinCacheDispatchFlags::Tangents))
+					if (EnumHasAnyFlags(DispatchData.DispatchFlags, EGPUSkinCacheDispatchFlags::RecomputeTangents))
 					{
 						TangentBuffers.Add(DispatchData.GetTangentRWBuffer());
 						BuffersToTransitionToRead.Add(DispatchData.GetTangentRWBuffer());
@@ -1707,7 +1703,7 @@ void FGPUSkinCache::DispatchPassExecute(FRHICommandList& RHICmdList, FDispatchSe
 				for (FDispatchEntry& DispatchItem : BatchDispatches)
 				{
 					FGPUSkinCacheEntry::FSectionDispatchData& DispatchData = DispatchItem.SkinCacheEntry->DispatchData[DispatchItem.Section];
-					if (EnumHasAnyFlags(DispatchData.DispatchFlags, EGPUSkinCacheDispatchFlags::Tangents))
+					if (EnumHasAnyFlags(DispatchData.DispatchFlags, EGPUSkinCacheDispatchFlags::RecomputeTangents))
 					{
 						DispatchUpdateSkinTangents(RHICmdList, DispatchItem.SkinCacheEntry, DispatchItem.Section, StagingBuffer, false);
 					}
@@ -1782,12 +1778,11 @@ bool FGPUSkinCache::ProcessEntry(
 						IntermediateAccumulatedTangentBufferOffset += RenderSection.GetNumVertices();
 					}
 
-					if (RenderSection.RecomputeTangentsVertexMaskChannel < ESkinVertexColorChannel::None)
+					if (RenderSection.RecomputeTangentsVertexMaskChannel < ESkinVertexColorChannel::None && LodData.MultiSizeIndexContainer.GetIndexBuffer())
 					{
 						BufferBits |= EGPUSkinCacheBufferBits::IntermediateTangents;
 					}
 
-					BufferBits |= EGPUSkinCacheBufferBits::Tangents;
 					IntermediateAccumulatedTangentBufferSize += RenderSection.GetNumVertices();
 				}
 			}
@@ -1796,6 +1791,12 @@ bool FGPUSkinCache::ProcessEntry(
 		BufferBits |= EGPUSkinCacheBufferBits::PositionPrevious;
 	}
 
+	// Don't allocate the accumulated tangents buffer if none of the sections are using it.
+	if (!EnumHasAnyFlags(BufferBits, EGPUSkinCacheBufferBits::IntermediateTangents))
+	{
+		IntermediateAccumulatedTangentBufferSize = 0;
+	}
+
 	if (InOutEntry)
 	{
 		// If the LOD changed, the entry has to be invalidated
@@ -2113,17 +2114,18 @@ void FGPUSkinCache::PrepareUpdateSkinning(FGPUSkinCacheEntry* Entry, int32 Secti
 	{
 		DispatchData.PositionBuffer = &DispatchData.PositionTracker.Allocation->GetPositionBuffer();
 		DispatchData.DispatchFlags |= EGPUSkinCacheDispatchFlags::Position;
+		Context.BuffersToTransition.Add(DispatchData.PositionBuffer);
 	}
 
-	if (auto* TangentBuffer = DispatchData.PositionTracker.GetTangentBuffer())
-	{
-		DispatchData.TangentBuffer = TangentBuffer;
-		DispatchData.DispatchFlags |= EGPUSkinCacheDispatchFlags::Tangents;
-	}
-
+	DispatchData.TangentBuffer = DispatchData.PositionTracker.GetTangentBuffer();
 	DispatchData.IntermediateTangentBuffer = DispatchData.PositionTracker.GetIntermediateTangentBuffer();
 	DispatchData.IntermediateAccumulatedTangentBuffer = DispatchData.PositionTracker.GetIntermediateAccumulatedTangentBuffer();
 
+	if (DispatchData.IntermediateTangentBuffer)
+	{
+		DispatchData.DispatchFlags |= EGPUSkinCacheDispatchFlags::RecomputeTangents;
+	}
+
 	int32 NumBuffers = 0;
 	NumBuffers += DispatchData.PositionBuffer ? 1 : 0;
 	NumBuffers += DispatchData.PreviousPositionBuffer ? 1 : 0;
@@ -2586,26 +2588,24 @@ FGPUSkinCache::FRWBuffersAllocation::FRWBuffersAllocation(FRHICommandList& RHICm
 		PositionBuffers[Index].Buffer.Buffer->SetOwnerName(OwnerName);
 		PositionBuffers[Index].AccessState = ERHIAccess::Unknown;
 	}
+	
+	// OpenGL ES does not support writing to RGBA16_SNORM images, instead pack data into SINT in the shader
+	const EPixelFormat TangentsFormat = IsOpenGLPlatform(GMaxRHIShaderPlatform) ? PF_R16G16B16A16_SINT : PF_R16G16B16A16_SNORM;
 
-	if (EnumHasAnyFlags(InBufferBits, EGPUSkinCacheBufferBits::Tangents))
-	{
-		// OpenGL ES does not support writing to RGBA16_SNORM images, instead pack data into SINT in the shader
-		const EPixelFormat TangentsFormat = IsOpenGLPlatform(GMaxRHIShaderPlatform) ? PF_R16G16B16A16_SINT : PF_R16G16B16A16_SNORM;
-
-		Tangents.Buffer.ClassName = TangentsName;
-		Tangents.Buffer.OwnerName = OwnerName;
-		Tangents.Buffer.Initialize(RHICmdList, TEXT("SkinCacheTangents"), TangentBufferBytesPerElement, NumVertices * 2, TangentsFormat, BUF_Static);
-		Tangents.Buffer.Buffer->SetOwnerName(OwnerName);
-		Tangents.AccessState = ERHIAccess::Unknown;
+	// Tangents are skinned inside the main skinning compute shader and are always allocated, even if the recompute tangents pass doesn't run.
+	Tangents.Buffer.ClassName = TangentsName;
+	Tangents.Buffer.OwnerName = OwnerName;
+	Tangents.Buffer.Initialize(RHICmdList, TEXT("SkinCacheTangents"), TangentBufferBytesPerElement, NumVertices * 2, TangentsFormat, BUF_Static);
+	Tangents.Buffer.Buffer->SetOwnerName(OwnerName);
+	Tangents.AccessState = ERHIAccess::Unknown;
 
-		if (EnumHasAnyFlags(InBufferBits, EGPUSkinCacheBufferBits::IntermediateTangents))
-		{
-			IntermediateTangents.Buffer.ClassName = IntermediateTangentsName;
-			IntermediateTangents.Buffer.OwnerName = OwnerName;
-			IntermediateTangents.Buffer.Initialize(RHICmdList, TEXT("SkinCacheIntermediateTangents"), TangentBufferBytesPerElement, NumVertices * 2, TangentsFormat, BUF_Static);
-			IntermediateTangents.Buffer.Buffer->SetOwnerName(OwnerName);
-			IntermediateTangents.AccessState = ERHIAccess::Unknown;
-		}
+	if (EnumHasAnyFlags(InBufferBits, EGPUSkinCacheBufferBits::IntermediateTangents))
+	{
+		IntermediateTangents.Buffer.ClassName = IntermediateTangentsName;
+		IntermediateTangents.Buffer.OwnerName = OwnerName;
+		IntermediateTangents.Buffer.Initialize(RHICmdList, TEXT("SkinCacheIntermediateTangents"), TangentBufferBytesPerElement, NumVertices * 2, TangentsFormat, BUF_Static);
+		IntermediateTangents.Buffer.Buffer->SetOwnerName(OwnerName);
+		IntermediateTangents.AccessState = ERHIAccess::Unknown;
 	}
 
 	if (IntermediateAccumulatedTangentsSize > 0)
@@ -2628,26 +2628,19 @@ FGPUSkinCache::FRWBuffersAllocation::~FRWBuffersAllocation()
 		PositionBuffers[Index].Release();
 	}
 
-	if (EnumHasAnyFlags(BufferBits, EGPUSkinCacheBufferBits::Tangents))
-	{
-		Tangents.Release();
-		IntermediateTangents.Release();
-	}
-
-	if (IntermediateAccumulatedTangentsSize > 0)
-	{
-		IntermediateAccumulatedTangents.Release();
-	}
+	Tangents.Release();
+	IntermediateTangents.Release();
+	IntermediateAccumulatedTangents.Release();
 }
 
 FGPUSkinCache::FSkinCacheRWBuffer* FGPUSkinCache::FRWBuffersAllocation::GetTangentBuffer()
 {
-	return EnumHasAllFlags(BufferBits, EGPUSkinCacheBufferBits::Tangents) ? &Tangents : nullptr;
+	return &Tangents;
 }
 
 FGPUSkinCache::FSkinCacheRWBuffer* FGPUSkinCache::FRWBuffersAllocation::GetIntermediateTangentBuffer()
 {
-	return EnumHasAllFlags(BufferBits, EGPUSkinCacheBufferBits::AllTangents) ? &IntermediateTangents : nullptr;
+	return EnumHasAllFlags(BufferBits, EGPUSkinCacheBufferBits::IntermediateTangents) ? &IntermediateTangents : nullptr;
 }
 
 bool FGPUSkinCache::FRWBuffersAllocation::HasPreviousBuffer() const
@@ -2657,12 +2650,11 @@ bool FGPUSkinCache::FRWBuffersAllocation::HasPreviousBuffer() const
 
 uint64 FGPUSkinCache::FRWBuffersAllocation::CalculateRequiredMemory(EGPUSkinCacheBufferBits BufferBits, uint32 NumVertices, uint32 IntermediateAccumulatedTangentsSize)
 {
-	const bool bTangents = EnumHasAnyFlags(BufferBits, EGPUSkinCacheBufferBits::Tangents);
 	const bool bIntermediateTangents = EnumHasAnyFlags(BufferBits, EGPUSkinCacheBufferBits::IntermediateTangents);
 	const bool bPositionPrevious = EnumHasAnyFlags(BufferBits, EGPUSkinCacheBufferBits::PositionPrevious);
 
 	uint64 PositionBufferSize = PosBufferBytesPerElement * NumVertices * 3 * (bPositionPrevious ? NUM_BUFFERS : 1);
-	uint64 TangentBufferSize =  bTangents ? TangentBufferBytesPerElement * NumVertices * 2 : 0;
+	uint64 TangentBufferSize =  TangentBufferBytesPerElement * NumVertices * 2;
 	uint64 IntermediateTangentBufferSize = bIntermediateTangents ? TangentBufferBytesPerElement * NumVertices * 2 : 0;
 	uint64 AccumulatedTangentBufferSize = IntermediateAccumulatedTangentsSize * FGPUSkinCache::IntermediateAccumBufferNumInts * sizeof(int32);
 	return TangentBufferSize + IntermediateTangentBufferSize + PositionBufferSize + AccumulatedTangentBufferSize;
