[World Partition Actor Descriptors] Better logs when removing a non-registered actor descriptors.

diff --git a/Engine/Source/Runtime/Engine/Public/WorldPartition/ActorDescList.h b/Engine/Source/Runtime/Engine/Public/WorldPartition/ActorDescList.h
--- a/Engine/Source/Runtime/Engine/Public/WorldPartition/ActorDescList.h
+++ b/Engine/Source/Runtime/Engine/Public/WorldPartition/ActorDescList.h
@@ -226,7 +226,11 @@ void TActorDescList<DescType>::RemoveActorDescriptor(DescType* ActorDesc)
 {
 	check(!bIsProxy);
 	check(ActorDesc);
-	verify(ActorsByGuid.Remove(ActorDesc->GetGuid()));
+	
+	if (!ensure((ActorsByGuid.Remove(ActorDesc->GetGuid()))))
+	{
+		UE_LOG(LogWorldPartition, Error, TEXT("Removing unknown actor descriptor: %s"), *ActorDesc->ToString(FWorldPartitionActorDesc::EToStringMode::Full));
+	}
 }
 
 template<class DescType>
