Iris replication support for online beacons
-Add NetDriver API and Iris support to "pre-register" replicated objects - beacons use this instead of calling into the GuidCache directly.
-Iris tracks these pre-registered objects in the ObjectReferenceCache and serializes a flag in the dynamic actor creation header.
-Misc. other changes to support a beacon replication system and a game replication system simultaneously in the same world.
-Bump engine network version.
-Existing beacon unit tests in OnlineBeaconUnitTest.cpp can now use Iris and pass when Iris is enabled.
-Added more unit tests.

diff --git a/Engine/Plugins/Online/OnlineSubsystemUtils/Source/OnlineSubsystemUtils/OnlineSubsystemUtils.Build.cs b/Engine/Plugins/Online/OnlineSubsystemUtils/Source/OnlineSubsystemUtils/OnlineSubsystemUtils.Build.cs
--- a/Engine/Plugins/Online/OnlineSubsystemUtils/Source/OnlineSubsystemUtils/OnlineSubsystemUtils.Build.cs
+++ b/Engine/Plugins/Online/OnlineSubsystemUtils/Source/OnlineSubsystemUtils/OnlineSubsystemUtils.Build.cs
@@ -66,5 +66,7 @@ public class OnlineSubsystemUtils : ModuleRules
 		}
 
 		PublicDependencyModuleNames.Add("OnlineSubsystem");
+
+		SetupIrisSupport(Target);
 	}
 }
diff --git a/Engine/Plugins/Online/OnlineSubsystemUtils/Source/OnlineSubsystemUtils/Private/OnlineBeacon.cpp b/Engine/Plugins/Online/OnlineSubsystemUtils/Source/OnlineSubsystemUtils/Private/OnlineBeacon.cpp
--- a/Engine/Plugins/Online/OnlineSubsystemUtils/Source/OnlineSubsystemUtils/Private/OnlineBeacon.cpp
+++ b/Engine/Plugins/Online/OnlineSubsystemUtils/Source/OnlineSubsystemUtils/Private/OnlineBeacon.cpp
@@ -4,6 +4,9 @@
 #include "Engine/Channel.h"
 #include "Engine/NetConnection.h"
 #include "Engine/Engine.h"
+#if UE_WITH_IRIS
+#include "Iris/ReplicationSystem/ReplicationSystem.h"
+#endif // UE_WITH_IRIS
 
 #include UE_INLINE_GENERATED_CPP_BY_NAME(OnlineBeacon)
 
@@ -94,6 +97,16 @@ void AOnlineBeacon::OnActorChannelOpen(FInBunch& Bunch, UNetConnection* Connecti
 {
 	Connection->OwningActor = this;
 	Super::OnActorChannelOpen(Bunch, Connection);
+
+#if UE_WITH_IRIS
+	{
+		// Enable replication
+		if (UReplicationSystem* ReplicationSystem = Connection->Driver->GetReplicationSystem())
+		{
+			ReplicationSystem->SetReplicationEnabledForConnection(Connection->GetConnectionHandle().GetParentConnectionId(), true);
+		}
+	}
+#endif // UE_WITH_IRIS
 }
 
 bool AOnlineBeacon::IsRelevancyOwnerFor(const AActor* ReplicatedActor, const AActor* ActorOwner, const AActor* ConnectionActor) const
diff --git a/Engine/Plugins/Online/OnlineSubsystemUtils/Source/OnlineSubsystemUtils/Private/OnlineBeaconClient.cpp b/Engine/Plugins/Online/OnlineSubsystemUtils/Source/OnlineSubsystemUtils/Private/OnlineBeaconClient.cpp
--- a/Engine/Plugins/Online/OnlineSubsystemUtils/Source/OnlineSubsystemUtils/Private/OnlineBeaconClient.cpp
+++ b/Engine/Plugins/Online/OnlineSubsystemUtils/Source/OnlineSubsystemUtils/Private/OnlineBeaconClient.cpp
@@ -479,13 +479,13 @@ void AOnlineBeaconClient::NotifyControlMessage(UNetConnection* Connection, uint8
 			}
 		case NMT_BeaconAssignGUID:
 			{
-				FNetworkGUID NetGUID;
+				UE::Net::FNetIDVariant NetID;
 
-				if (FNetControlMessage<NMT_BeaconAssignGUID>::Receive(Bunch, NetGUID))
+				if (FNetControlMessage<NMT_BeaconAssignGUID>::Receive(Bunch, NetID))
 				{
-					if (NetGUID.IsValid())
+					if (NetID.IsValid())
 					{
-						Connection->Driver->GetNetGuidCache()->RegisterNetGUID_Client(NetGUID, this);
+						Connection->Driver->AssignNetIDToObject(NetID, this);
 
 						FString BeaconType = GetBeaconType();
 						FNetControlMessage<NMT_BeaconNetGUIDAck>::Send(Connection, BeaconType);
diff --git a/Engine/Plugins/Online/OnlineSubsystemUtils/Source/OnlineSubsystemUtils/Private/OnlineBeaconHost.cpp b/Engine/Plugins/Online/OnlineSubsystemUtils/Source/OnlineSubsystemUtils/Private/OnlineBeaconHost.cpp
--- a/Engine/Plugins/Online/OnlineSubsystemUtils/Source/OnlineSubsystemUtils/Private/OnlineBeaconHost.cpp
+++ b/Engine/Plugins/Online/OnlineSubsystemUtils/Source/OnlineSubsystemUtils/Private/OnlineBeaconHost.cpp
@@ -450,7 +450,7 @@ bool AOnlineBeaconHost::HandleControlMessage(UNetConnection* Connection, uint8 M
 			Connection->SetClientLoginState(EClientLoginState::ReceivedJoin);
 			NewClientActor->SetConnectionState(EBeaconConnectionState::Pending);
 
-			FNetworkGUID NetGUID = Connection->Driver->GetNetGuidCache()->AssignNewNetGUID_Server(NewClientActor);
+			UE::Net::FNetIDVariant NetID = Connection->Driver->CreateNetIDForObject(NewClientActor);
 			NewClientActor->SetNetConnection(Connection);
 			Connection->PlayerId = UniqueId;
 			NewClientActor->SetRole(ROLE_Authority);
@@ -459,7 +459,7 @@ bool AOnlineBeaconHost::HandleControlMessage(UNetConnection* Connection, uint8 M
 			NewClientActor->SetNetDriverName(NetDriverName);
 			ClientActors.Add(NewClientActor);
 
-			FNetControlMessage<NMT_BeaconAssignGUID>::Send(Connection, NetGUID);
+			FNetControlMessage<NMT_BeaconAssignGUID>::Send(Connection, NetID);
 			Connection->FlushNet(true);
 		}
 		break;
diff --git a/Engine/Plugins/Online/OnlineSubsystemUtils/Source/OnlineSubsystemUtils/Private/Tests/OnlineBeaconUnitTestUtils.cpp b/Engine/Plugins/Online/OnlineSubsystemUtils/Source/OnlineSubsystemUtils/Private/Tests/OnlineBeaconUnitTestUtils.cpp
--- a/Engine/Plugins/Online/OnlineSubsystemUtils/Source/OnlineSubsystemUtils/Private/Tests/OnlineBeaconUnitTestUtils.cpp
+++ b/Engine/Plugins/Online/OnlineSubsystemUtils/Source/OnlineSubsystemUtils/Private/Tests/OnlineBeaconUnitTestUtils.cpp
@@ -10,6 +10,7 @@
 #include "Engine/NetConnection.h"
 #include "Engine/World.h"
 #include "GameFramework/OnlineReplStructs.h"
+#include "GameFramework/WorldSettings.h"
 #include "IpNetDriver.h"
 #include "Misc/ConfigCacheIni.h"
 #include "Online/CoreOnline.h"
@@ -150,6 +151,10 @@ public:
 		// Initialize world.
 		World = FAutomationEditorCommonUtils::CreateNewMap();
 		World->InitializeActorsForPlay(FURL());
+		if (World->GetWorldSettings())
+		{
+			World->GetWorldSettings()->NotifyBeginPlay();
+		}
 
 		// Bind delegates for handling network encryption.
 		BindNetEncryptionDelegates();
@@ -178,13 +183,28 @@ public:
 			TestBeaconNetdriverDefinition.DriverClassName = FName(UOnlineBeaconUnitTestNetDriver::StaticClass()->GetPathName());
 			GEngine->NetDriverDefinitions.Add(TestBeaconNetdriverDefinition);
 		}
+
+		// Install Iris netdriver config for the test driver.
+		const bool bIrisConfigExists = Algo::FindBy(GEngine->IrisNetDriverConfigs, BeaconUnitTest::NetDriverDefinitionName, [](const FIrisNetDriverConfig& IrisConfig)->const FName& { return IrisConfig.NetDriverDefinition; }) != nullptr;
+		if (ensure(!bIrisConfigExists))
+		{
+			FIrisNetDriverConfig TestBeaconNetDriverIrisConfig;
+			TestBeaconNetDriverIrisConfig.NetDriverDefinition = BeaconUnitTest::NetDriverDefinitionName;
+			TestBeaconNetDriverIrisConfig.bCanUseIris = true;
+			GEngine->IrisNetDriverConfigs.Add(TestBeaconNetDriverIrisConfig);
+		}
 	}
 
 	virtual ~FTestPrerequisitesImpl()
 	{
+		World->EndPlay(EEndPlayReason::EndPlayInEditor);
+
 		// Remove unit test net driver definition.
 		GEngine->NetDriverDefinitions.SetNum(Algo::RemoveIf(GEngine->NetDriverDefinitions, [&](const FNetDriverDefinition& Definition) -> bool { return Definition.DefName == BeaconUnitTest::NetDriverDefinitionName; }));
 
+		// Remove unit test net driver Iris config.
+		GEngine->IrisNetDriverConfigs.SetNum(Algo::RemoveIf(GEngine->IrisNetDriverConfigs, [&](const FIrisNetDriverConfig& IrisConfig) -> bool { return IrisConfig.NetDriverDefinition == BeaconUnitTest::NetDriverDefinitionName; }));
+
 		FWorldDelegates::OnWorldTickEnd.Remove(NetworkFlushDelegateHandle);
 		UnbindNetEncryptionDelegates();
 		World = nullptr;
diff --git a/Engine/Source/Runtime/Core/Public/Misc/EngineNetworkCustomVersion.h b/Engine/Source/Runtime/Core/Public/Misc/EngineNetworkCustomVersion.h
--- a/Engine/Source/Runtime/Core/Public/Misc/EngineNetworkCustomVersion.h
+++ b/Engine/Source/Runtime/Core/Public/Misc/EngineNetworkCustomVersion.h
@@ -91,6 +91,7 @@ struct FEngineNetworkCustomVersion
 		CustomExports = 36,						// Bump version to support CustomExports (such as NetTokens used for supporting "Dynamic Replication" of GameplayTags -- see FGameplayTag::NetSerialize_Packed)		// -----<new versions can be added above this line>-------------------------------------------------
 		MontagePlayCountSerialization = 37,		// Bump version to support serialization changes for dynamic montage play count 
 		RemoteObjectReferences = 38,			// Bump version to support serialization changes for remote object references 
+		BeaconNetIDVariant = 39,				// Bump version for online beacon's BeaconAssignGUID control message switching to use FNetIDVariant
 
 		// -----<new versions can be added above this line>-------------------------------------------------
 		VersionPlusOne,
diff --git a/Engine/Source/Runtime/Engine/Classes/Engine/NetDriver.h b/Engine/Source/Runtime/Engine/Classes/Engine/NetDriver.h
--- a/Engine/Source/Runtime/Engine/Classes/Engine/NetDriver.h
+++ b/Engine/Source/Runtime/Engine/Classes/Engine/NetDriver.h
@@ -357,6 +357,7 @@ namespace UE::Net
 namespace UE::Net
 {
 	class FNetTokenStore;
+	class FNetIDVariant;
 
 	class FScopedIgnoreStaticActorDestruction
 	{
@@ -2088,6 +2089,29 @@ public:
 	 */
 	bool SendDestructionInfoForLevelUnloadIfDormant(AActor* ThisActor, UNetConnection* Connection);
 
+	/**
+	 * Explicitly generates and assigns a new network ID (FNetRefHandle if using Iris, FNetworkGUID if using generic replication)
+	 * to Object and returns it. This can be used in combination with AssignNetIDToObject to associate actor instances
+	 * that are spawned manually for replication, as opposed to the default behavior of letting the engine spawn
+	 * the client's instance of a replicated actor automatically.
+	 * This function would typically be called on a server, and the caller is responsible for communicating
+	 * the returned ID to the remote/client side and calling AssignNetIDToObject to associate the ID with
+	 * the client instance of the object.
+	 * Only actors or Iris root objects are supported.
+	 * @see UNetDriver::AssignNetIDToObject
+	 */
+	ENGINE_API UE::Net::FNetIDVariant CreateNetIDForObject(UObject* Object);
+
+	/**
+	 * Explicitly assigns NetID to Object. Can be called on clients that have received a NetID generated on
+	 * the server by a call to CreateNetIDForObject. These functions can be used to associate actor instances
+	 * that are spawned manually for replication, as opposed to the default behavior of letting the engine spawn
+	 * the client's instance of a replicated actor automatically.
+	 * Only actors or Iris root objects are supported.
+	 * @see UNetDriver::CreateNetIDForObject
+	 */
+	ENGINE_API void AssignNetIDToObject(UE::Net::FNetIDVariant NetID, const UObject* Object);
+
 protected:
 
 	void SetIsInTick(bool bIsInTick) { bInTick = bIsInTick; }
diff --git a/Engine/Source/Runtime/Engine/Private/Net/Iris/ReplicationSystem/EngineReplicationBridge.cpp b/Engine/Source/Runtime/Engine/Private/Net/Iris/ReplicationSystem/EngineReplicationBridge.cpp
--- a/Engine/Source/Runtime/Engine/Private/Net/Iris/ReplicationSystem/EngineReplicationBridge.cpp
+++ b/Engine/Source/Runtime/Engine/Private/Net/Iris/ReplicationSystem/EngineReplicationBridge.cpp
@@ -500,7 +500,7 @@ void UEngineReplicationBridge::StopReplicatingActor(AActor* Actor, EEndPlayReaso
 	}
 
 	// If we are shutting down or the actor is a nettemporary we do not need to validate that we are not detaching remote instances by accident.
-	const bool bIsShuttingDown = (EndPlayReason == EEndPlayReason::EndPlayInEditor) || (EndPlayReason == EEndPlayReason::Quit);
+	const bool bIsShuttingDown = (EndPlayReason == EEndPlayReason::EndPlayInEditor) || (EndPlayReason == EEndPlayReason::Quit) || NetDriver->GetPendingDestruction();
 	// Streaming out a level can happen prior to all actors in it being destroyed/ending replication through replication so we don't want to validate they're detached by accident.
 	const bool bIsStreamingOutLevel = (EndPlayReason == EEndPlayReason::RemovedFromWorld);
 	if (bIsShuttingDown || bIsStreamingOutLevel || Actor->bNetTemporary)
diff --git a/Engine/Source/Runtime/Engine/Private/Net/Iris/ReplicationSystem/NetActorFactory.cpp b/Engine/Source/Runtime/Engine/Private/Net/Iris/ReplicationSystem/NetActorFactory.cpp
--- a/Engine/Source/Runtime/Engine/Private/Net/Iris/ReplicationSystem/NetActorFactory.cpp
+++ b/Engine/Source/Runtime/Engine/Private/Net/Iris/ReplicationSystem/NetActorFactory.cpp
@@ -223,6 +223,7 @@ TUniquePtr<UE::Net::FNetObjectCreationHeader> UNetActorFactory::CreateAndFillHea
 		// Fill in Header
 		Header->ArchetypeReference = ActorBridge->GetOrCreateObjectReference(Archetype);
 		Header->bUsePersistentLevel = (UE::Net::Private::SerializeNewActorOverrideLevel == 0) || (ActorBridge->GetNetDriver()->GetWorld()->PersistentLevel == ActorLevel);
+		Header->bIsPreRegistered = ActorBridge->IsNetRefHandlePreRegistered(Handle);
 
 		if (!Header->bUsePersistentLevel)
 		{
@@ -365,9 +366,29 @@ UNetObjectFactory::FInstantiateResult UNetActorFactory::InstantiateReplicatedObj
 		return Result;
 	}
 
-	// For dynamic actors, spawn a new actor using the provided information
+	// For dynamic actors, spawn a new actor using the provided information, or find a pre-registered actor
 	const FDynamicActorNetCreationHeader* DynamicHeader = static_cast<const FDynamicActorNetCreationHeader*>(Header);
 
+	// If pre-registered, look for existing instance.
+	if (DynamicHeader->bIsPreRegistered)
+	{
+		UObject* FoundObject = Bridge->GetPreRegisteredObject(Context.Handle);
+		if (!FoundObject)
+		{
+			UE_LOG(LogIris, Error, TEXT("UNetActorFactory::InstantiateReplicatedObjectFromHeader Unable to find pre-registered actor: %s"), *Context.Handle.ToString());
+			return FInstantiateResult();
+		}
+
+		FInstantiateResult Result{ .Instance = FoundObject };
+
+		if (NetDriver->ShouldClientDestroyActor(Cast<AActor>(FoundObject)))
+		{
+			Result.Flags |= EReplicationBridgeCreateNetRefHandleResultFlags::AllowDestroyInstanceFromRemote;
+		}
+
+		return Result;
+	}
+
 	// Find archetype
 	AActor* Archetype = Cast<AActor>(Bridge->ResolveObjectReference(DynamicHeader->ArchetypeReference, Context.ResolveContext));
 	if (!Archetype)
@@ -622,6 +643,8 @@ bool FDynamicActorNetCreationHeader::Serialize(const FCreationHeaderContext& Con
 	static constexpr float RotationEpsilon = 0.001f;
 	WriteRotator(Writer, SpawnInfo.Rotation, DefaultSpawnInfo.Rotation, RotationEpsilon);
 
+	Writer->WriteBool(bIsPreRegistered);
+
 	if (Writer->WriteBool(CustomCreationDataBitCount > 0))
 	{
 		Writer->WriteBits(CustomCreationDataBitCount - 1U, 16U);
@@ -653,6 +676,8 @@ bool FDynamicActorNetCreationHeader::Deserialize(const FCreationHeaderContext& C
 
 	ReadRotator(Reader, SpawnInfo.Rotation, DefaultSpawnInfo.Rotation);
 	
+	bIsPreRegistered = Reader->ReadBool();
+
 	if (Reader->ReadBool())
 	{
 		CustomCreationDataBitCount = 1U + Reader->ReadBits(16U);
diff --git a/Engine/Source/Runtime/Engine/Private/Net/Iris/ReplicationSystem/ReplicationSystemUtil.cpp b/Engine/Source/Runtime/Engine/Private/Net/Iris/ReplicationSystem/ReplicationSystemUtil.cpp
--- a/Engine/Source/Runtime/Engine/Private/Net/Iris/ReplicationSystem/ReplicationSystemUtil.cpp
+++ b/Engine/Source/Runtime/Engine/Private/Net/Iris/ReplicationSystem/ReplicationSystemUtil.cpp
@@ -438,7 +438,7 @@ void FReplicationSystemUtil::SetActorComponentNetCondition(const UActorComponent
 	}
 }
 
-void FReplicationSystemUtil::BeginReplicationForActorsInWorld(UWorld* World)
+void FReplicationSystemUtil::BeginReplicationForActorsInWorldForNetDriver(UWorld* World, UNetDriver* NetDriver)
 {
 	// We only do this if the world already is initialized
 	// Normally we rely on Begin/EndPlay to control if an actor is replicated or not.
@@ -447,7 +447,7 @@ void FReplicationSystemUtil::BeginReplicationForActorsInWorld(UWorld* World)
 		for (FActorIterator Iter(World); Iter; ++Iter)
 		{
 			AActor* Actor = *Iter;
-			if (IsValid(Actor) && Actor->HasActorBegunPlay() && ULevel::IsNetActor(Actor))
+			if (IsValid(Actor) && Actor->HasActorBegunPlay() && NetDriver->ShouldReplicateActor(Actor) && ULevel::IsNetActor(Actor))
 			{
 				Actor->BeginReplication();
 			}
diff --git a/Engine/Source/Runtime/Engine/Private/NetDriver.cpp b/Engine/Source/Runtime/Engine/Private/NetDriver.cpp
--- a/Engine/Source/Runtime/Engine/Private/NetDriver.cpp
+++ b/Engine/Source/Runtime/Engine/Private/NetDriver.cpp
@@ -25,6 +25,7 @@
 #include "EngineUtils.h"
 #include "Net/NetworkProfiler.h"
 #include "Engine/PackageMapClient.h"
+#include "Net/NetIDVariant.h"
 #include "Net/RepLayout.h"
 #include "Net/DataReplication.h"
 #include "Engine/ControlChannel.h"
@@ -58,6 +59,7 @@
 #include "Iris/Core/IrisProfiler.h"
 #include "Iris/Core/IrisMemoryTracker.h"
 #include "Net/Experimental/Iris/DataStreamChannel.h"
+#include "Iris/ReplicationSystem/ObjectReplicationBridge.h"
 #include "Iris/ReplicationSystem/ReplicationSystem.h"
 #include "Iris/ReplicationSystem/ReplicationView.h"
 #include "Iris/ReplicationSystem/Filtering/NetObjectFilter.h"
@@ -6904,7 +6906,7 @@ void UNetDriver::SetWorld(class UWorld* InWorld)
 			}
 			else
 			{
-				UE::Net::FReplicationSystemUtil::BeginReplicationForActorsInWorld(InWorld);
+				UE::Net::FReplicationSystemUtil::BeginReplicationForActorsInWorldForNetDriver(InWorld, this);
 			}
 		}
 		else
@@ -7187,7 +7189,7 @@ void UNetDriver::SetReplicationSystem(UReplicationSystem* InReplicationSystem)
 		// When we run using Iris, we use ReplicationSystemId as our unique identifier
 		NetTraceId = ReplicationSystem->GetId();
 
-		UE::Net::FReplicationSystemUtil::BeginReplicationForActorsInWorld(GetWorld());
+		UE::Net::FReplicationSystemUtil::BeginReplicationForActorsInWorldForNetDriver(GetWorld(), this);
 	}
 }
 
@@ -8662,6 +8664,89 @@ void UNetDriver::UpdateNetworkStats()
 	}
 }
 
+UE::Net::FNetIDVariant UNetDriver::CreateNetIDForObject(UObject* Object)
+{
+	using namespace UE::Net;
+
+#if UE_WITH_IRIS
+	if (bIsUsingIris)
+	{
+		if (ReplicationSystem)
+		{
+			UObjectReplicationBridge* ObjectBridge = ReplicationSystem->GetReplicationBridgeAs<UObjectReplicationBridge>();
+			if (ObjectBridge)
+			{
+				FNetRefHandle AllocatedRefHandle = ObjectBridge->PreRegisterNewObjectReferenceHandle(Object);
+				return FNetIDVariant(AllocatedRefHandle);
+			}
+			else
+			{
+				UE_LOG(LogNet, Warning, TEXT("UNetDriver::CreateNetIDForObject: failed to assign a NetRefHandle for object %s, no ObjectBridge. NetDriver %s."), *GetFullNameSafe(Object), *GetFullName());
+			}
+		}
+		else
+		{
+			UE_LOG(LogNet, Warning, TEXT("UNetDriver::CreateNetIDForObject: failed to assign a NetRefHandle for object %s, no ReplicationSystem. NetDriver %s."), *GetFullNameSafe(Object), *GetFullName());
+			ensureMsgf(false, TEXT("ReplicationSystem is null"));
+		}
+
+		return FNetIDVariant(FNetRefHandle());
+	}
+	else
+#endif
+	{
+		FNetworkGUID GUID = GuidCache->AssignNewNetGUID_Server(Object);
+		return FNetIDVariant(GUID);
+	}
+}
+
+void UNetDriver::AssignNetIDToObject(UE::Net::FNetIDVariant NetID, const UObject* Object)
+{
+#if UE_WITH_IRIS
+	if (bIsUsingIris)
+	{
+		if (NetID.GetVariant().IsType<UE::Net::FNetRefHandle>())
+		{
+			if (ReplicationSystem)
+			{
+				UObjectReplicationBridge* ObjectBridge = ReplicationSystem->GetReplicationBridgeAs<UObjectReplicationBridge>();
+				if (ObjectBridge)
+				{
+					ObjectBridge->PreRegisterObjectWithReferenceHandle(Object, NetID.GetVariant().Get<UE::Net::FNetRefHandle>());
+				}
+				else
+				{
+					UE_LOG(LogNet, Warning, TEXT("UNetDriver::AssignNetIDToObject: failed to register NetID %s for object %s with net driver %s, no ObjectBridge."), *NetID.ToString(), *GetFullNameSafe(Object), *GetFullName());
+				}
+			}
+			else
+			{
+				UE_LOG(LogNet, Warning, TEXT("UNetDriver::AssignNetIDToObject: failed to register NetID %s for object %s with net driver %s, no ReplicationSystem."), *NetID.ToString(), *GetFullNameSafe(Object), *GetFullName());
+				ensureMsgf(false, TEXT("NReplicationSystem is null"));
+			}
+		}
+		else
+		{
+			UE_LOG(LogNet, Warning, TEXT("UNetDriver::AssignNetIDToObject: NetID %s for object %s is not an FNetRefHandle but NetDriver %s is using Iris."), *NetID.ToString(), *GetFullNameSafe(Object), *GetFullName());
+			ensureMsgf(false, TEXT("NetID is not an FNetRefHandle"));
+		}
+	}
+	else
+#endif
+	{
+		if (NetID.GetVariant().IsType<FNetworkGUID>())
+		{
+			GuidCache->RegisterNetGUID_Client(NetID.GetVariant().Get<FNetworkGUID>(), Object);
+		}
+		else
+		{
+			const FString ErrorMessage = FString::Printf(TEXT("UNetDriver::AssignNetIDToObject: NetID %s for object %s is not a NetGUID but NetDriver %s is using generic replication."), *NetID.ToString(), *GetFullNameSafe(Object), *GetFullName());
+			UE_LOG(LogNet, Warning, TEXT("%s"), *ErrorMessage);
+			ensureMsgf(false, TEXT("%s"), *ErrorMessage);
+		}
+	}
+}
+
 #if NET_DEBUG_RELEVANT_ACTORS
 FAutoConsoleCommandWithWorld	DumpRelevantActorsCommand(
 	TEXT("net.DumpRelevantActors"), 
diff --git a/Engine/Source/Runtime/Engine/Public/Net/DataChannel.h b/Engine/Source/Runtime/Engine/Public/Net/DataChannel.h
--- a/Engine/Source/Runtime/Engine/Public/Net/DataChannel.h
+++ b/Engine/Source/Runtime/Engine/Public/Net/DataChannel.h
@@ -12,6 +12,7 @@
 #include "Misc/NetworkVersion.h"
 #include "GameFramework/OnlineReplStructs.h"
 #include "Net/Core/Connection/NetEnums.h"
+#include "Net/NetIDVariant.h"
 #include "Net/DataBunch.h"
 #include "Engine/ChildConnection.h"
 
@@ -211,7 +212,7 @@ DEFINE_CONTROL_CHANNEL_MESSAGE(NetPing, 24, ENetPingControlMessage /* MessageTyp
 
 DEFINE_CONTROL_CHANNEL_MESSAGE(BeaconWelcome, 25); // server tells client they're ok to attempt to join (client sends netspeed/beacontype)
 DEFINE_CONTROL_CHANNEL_MESSAGE(BeaconJoin, 26, FString, FUniqueNetIdRepl);  // server tries to create beacon type requested by client, sends NetGUID for actor sync
-DEFINE_CONTROL_CHANNEL_MESSAGE(BeaconAssignGUID, 27, FNetworkGUID); // client assigns NetGUID from server to beacon actor, sends NetGUIDAck
+DEFINE_CONTROL_CHANNEL_MESSAGE(BeaconAssignGUID, 27, UE::Net::FNetIDVariant); // client assigns NetGUID from server to beacon actor, sends NetGUIDAck
 DEFINE_CONTROL_CHANNEL_MESSAGE(BeaconNetGUIDAck, 28, FString); // server received NetGUIDAck from client, connection established successfully
 
 DEFINE_CONTROL_CHANNEL_MESSAGE(IrisProtocolMismatch, 29, uint64); // client has a different protocol hash from the server
diff --git a/Engine/Source/Runtime/Engine/Public/Net/Iris/ReplicationSystem/NetActorFactory.h b/Engine/Source/Runtime/Engine/Public/Net/Iris/ReplicationSystem/NetActorFactory.h
--- a/Engine/Source/Runtime/Engine/Public/Net/Iris/ReplicationSystem/NetActorFactory.h
+++ b/Engine/Source/Runtime/Engine/Public/Net/Iris/ReplicationSystem/NetActorFactory.h
@@ -89,6 +89,7 @@ public:
 	FNetObjectReference LevelReference; // Only when bUsePersistentLevel is false
 	
 	bool bUsePersistentLevel = false;
+	bool bIsPreRegistered = false;
 
 	bool Serialize(const FCreationHeaderContext& Context, UE::Net::Private::EActorNetSpawnInfoFlags SpawnFlags, const FActorNetSpawnInfo& DefaultSpawnInfo) const;
 	bool Deserialize(const FCreationHeaderContext& Context, const FActorNetSpawnInfo& DefaultSpawnInfo);
diff --git a/Engine/Source/Runtime/Engine/Public/Net/Iris/ReplicationSystem/ReplicationSystemUtil.h b/Engine/Source/Runtime/Engine/Public/Net/Iris/ReplicationSystem/ReplicationSystemUtil.h
--- a/Engine/Source/Runtime/Engine/Public/Net/Iris/ReplicationSystem/ReplicationSystemUtil.h
+++ b/Engine/Source/Runtime/Engine/Public/Net/Iris/ReplicationSystem/ReplicationSystemUtil.h
@@ -119,8 +119,8 @@ struct FReplicationSystemUtil
 	/** Remove dependent actor from parent. The dependent actor will function as a standard standalone replicated actor. */
 	ENGINE_API static void RemoveDependentActor(const AActor* Parent, AActor* Child);
 
-	/** Begin replication for all networked actors in the world. */
-	ENGINE_API static void BeginReplicationForActorsInWorld(UWorld* World);
+	/** Begin replication for all networked actors that belong to a specific driver in the world. */
+	ENGINE_API static void BeginReplicationForActorsInWorldForNetDriver(UWorld* World, UNetDriver* NetDriver);
 
 	/** Notify the ReplicationSystem of a dormancy change. */
 	ENGINE_API static void NotifyActorDormancyChange(UReplicationSystem* ReplicationSystem, AActor* Actor, ENetDormancy OldDormancyState);
diff --git a/Engine/Source/Runtime/Experimental/Iris/Core/Private/Iris/ReplicationSystem/ObjectReferenceCache.cpp b/Engine/Source/Runtime/Experimental/Iris/Core/Private/Iris/ReplicationSystem/ObjectReferenceCache.cpp
--- a/Engine/Source/Runtime/Experimental/Iris/Core/Private/Iris/ReplicationSystem/ObjectReferenceCache.cpp
+++ b/Engine/Source/Runtime/Experimental/Iris/Core/Private/Iris/ReplicationSystem/ObjectReferenceCache.cpp
@@ -225,6 +225,29 @@ FNetRefHandle FObjectReferenceCache::CreateObjectReferenceHandle(const UObject*
 	return Ref.GetRefHandle();
 }
 
+FNetRefHandle FObjectReferenceCache::PreRegisterObjectReferenceHandle(const UObject* Object)
+{
+	if (IsAuthority())
+	{
+		FNetObjectReference Ref;
+		CreateObjectReferenceInternal(Object, Ref);
+		FCachedNetObjectReference* CachedRef = ReferenceHandleToCachedReference.Find(Ref.GetRefHandle());
+		if (CachedRef)
+		{
+			CachedRef->bIsPreRegistered = true;
+		}
+
+		return Ref.GetRefHandle();
+	}
+	else
+	{
+		UE_LOG_REFERENCECACHE(Error, TEXT("FObjectReferenceCache::PreRegisterObjectReferenceHandle attempted to create new handle for object %s but this is not the reference authority (it's a client)."), *GetFullNameSafe(Object));
+		ensureMsgf(false, TEXT("FObjectReferenceCache::PreRegisterObjectReferenceHandle attempted to create new handle without having reference authority."));
+	}
+
+	return FNetRefHandle::GetInvalid();
+}
+
 bool FObjectReferenceCache::CreateObjectReferenceInternal(const UObject* Object, FNetObjectReference& OutReference)
 {
 	if (!IsValid(Object))
@@ -522,6 +545,27 @@ void FObjectReferenceCache::AddRemoteReference(FNetRefHandle RefHandle, const UO
 	}
 }
 
+void FObjectReferenceCache::AddPreRegisteredReference(FNetRefHandle RefHandle, const UObject* Object)
+{
+	if (!IsAuthority())
+	{
+		// Promote to full handle
+		FNetRefHandle CompleteRefHandle = FNetRefHandleManager::MakeNetRefHandle(RefHandle.GetId(), ReplicationSystem->GetId());
+
+		AddRemoteReference(CompleteRefHandle, Object);
+
+		if (FCachedNetObjectReference* CacheObjectPtr = ReferenceHandleToCachedReference.Find(CompleteRefHandle))
+		{
+			CacheObjectPtr->bIsPreRegistered = true;
+		}
+	}
+	else
+	{
+		UE_LOG_REFERENCECACHE(Error, TEXT("FObjectReferenceCache::AddPreRegisteredReference attempted to associate handle %s with object %s on the reference authority (server)."), *RefHandle.ToString(), *GetFullNameSafe(Object));
+		ensureMsgf(false, TEXT("FObjectReferenceCache::AddPreRegisteredReference attempted to associate a handle with an object on the reference authority."));
+	}
+}
+
 void FObjectReferenceCache::RemoveReference(FNetRefHandle RefHandle, const UObject* Object)
 {
 	// Cleanup cached data for dynamic instance
@@ -639,12 +683,27 @@ UObject* FObjectReferenceCache::GetObjectFromReferenceHandle(FNetRefHandle RefHa
 	return ResolvedObject;
 }
 
+UObject* FObjectReferenceCache::GetPreRegisteredObjectFromReferenceHandle(FNetRefHandle RefHandle)
+{
+	FCachedNetObjectReference* CachedObject = ReferenceHandleToCachedReference.Find(RefHandle);
+
+	UObject* ResolvedObject = (CachedObject && CachedObject->bIsPreRegistered) ? CachedObject->Object.Get() : nullptr;
+
+	return ResolvedObject;
+}
+
 bool FObjectReferenceCache::IsNetRefHandleBroken(FNetRefHandle RefHandle, bool bMustBeRegistered) const
 {
 	const FCachedNetObjectReference* CacheObjectPtr = ReferenceHandleToCachedReference.Find(RefHandle);	
 	return CacheObjectPtr ? CacheObjectPtr->bIsBroken : bMustBeRegistered;
 }
 
+bool FObjectReferenceCache::IsNetRefHandlePreRegistered(FNetRefHandle RefHandle) const
+{
+	const FCachedNetObjectReference* CacheObjectPtr = ReferenceHandleToCachedReference.Find(RefHandle);
+	return CacheObjectPtr ? CacheObjectPtr->bIsPreRegistered : false;
+}
+
 bool FObjectReferenceCache::IsNetRefHandlePending(FNetRefHandle NetRefHandle, const FPendingBatches& PendingBatches) const
 {
 	// Check Outer chain
diff --git a/Engine/Source/Runtime/Experimental/Iris/Core/Private/Iris/ReplicationSystem/ObjectReferenceCache.h b/Engine/Source/Runtime/Experimental/Iris/Core/Private/Iris/ReplicationSystem/ObjectReferenceCache.h
--- a/Engine/Source/Runtime/Experimental/Iris/Core/Private/Iris/ReplicationSystem/ObjectReferenceCache.h
+++ b/Engine/Source/Runtime/Experimental/Iris/Core/Private/Iris/ReplicationSystem/ObjectReferenceCache.h
@@ -48,12 +48,18 @@ public:
 	// Create and assign a new NetHandle to the object
 	FNetRefHandle CreateObjectReferenceHandle(const UObject* Object);
 
+	// Create and assign a new pre-registered NetHandle to the object
+	FNetRefHandle PreRegisterObjectReferenceHandle(const UObject* Object);
+
 	// Get existing handle for object
 	FNetRefHandle GetObjectReferenceHandleFromObject(const UObject* Object, EGetRefHandleFlags GetRefHandleFlags = EGetRefHandleFlags::None) const;
 
 	// Get object from handle, only if the object is in the cache.
 	UObject* GetObjectFromReferenceHandle(FNetRefHandle RefHandle);
 
+	// Get object from handle, only if the object is in the cache and only if it was registered via PreregisterObjectReferenceHandle or AddPreregisteredReference.
+	UObject* GetPreRegisteredObjectFromReferenceHandle(FNetRefHandle RefHandle);
+
 	// Try to resolve the object reference and try to load it if the object cannot be found
 	UObject* ResolveObjectReferenceHandle(FNetRefHandle RefHandle, const FNetObjectResolveContext& ResolveContext);
 
@@ -66,6 +72,9 @@ public:
 	// Returns true of the provided NetRefHandle or one of its outers is pending async loading.
 	bool IsNetRefHandlePending(FNetRefHandle NetRefHandle, const FPendingBatches& PendingBatches) const;
 
+	// Returns true of the provided NetRefHandle was assigned via PreRegisterObjectReferenceHandle or AddPreregisteredReference.
+	bool IsNetRefHandlePreRegistered(FNetRefHandle NetRefHandle) const;
+
 	// Find replicated outer
 	FNetObjectReference GetReplicatedOuter(const FNetObjectReference& Reference) const;
 
@@ -78,6 +87,9 @@ public:
 	// Add reference for dynamically spawned object
 	void AddRemoteReference(FNetRefHandle RefHandle, const UObject* Object);
 
+	// Add reference for pre-registered object (client)
+	void AddPreRegisteredReference(FNetRefHandle RefHandle, const UObject* Object);
+
 	// Remove references to dynamic objects
 	void RemoveReference(FNetRefHandle RefHandle, const UObject* Object);
 
@@ -156,11 +168,12 @@ private:
 		FNetRefHandle OuterNetRefHandle;
 
 		// Flags
-		uint8 bNoLoad : 1;				// Don't load this, only do a find
-		uint8 bIgnoreWhenMissing : 1;
-		uint8 bIsPackage : 1;
-		uint8 bIsBroken : 1;
-		uint8 bIsPending : 1;
+		uint8 bNoLoad : 1 = false;				// Don't load this, only do a find
+		uint8 bIgnoreWhenMissing : 1 = false;
+		uint8 bIsPackage : 1 = false;
+		uint8 bIsBroken : 1 = false;
+		uint8 bIsPending : 1 = false;
+		uint8 bIsPreRegistered : 1 = false;		// True if assigned via PreregisterObjectReferenceHandle or AddPreregisteredReference
 	};
 
 	struct FQueuedBatchObjectReference
diff --git a/Engine/Source/Runtime/Experimental/Iris/Core/Private/Iris/ReplicationSystem/ObjectReplicationBridge.cpp b/Engine/Source/Runtime/Experimental/Iris/Core/Private/Iris/ReplicationSystem/ObjectReplicationBridge.cpp
--- a/Engine/Source/Runtime/Experimental/Iris/Core/Private/Iris/ReplicationSystem/ObjectReplicationBridge.cpp
+++ b/Engine/Source/Runtime/Experimental/Iris/Core/Private/Iris/ReplicationSystem/ObjectReplicationBridge.cpp
@@ -282,6 +282,11 @@ UObject* UObjectReplicationBridge::GetReplicatedObject(FNetRefHandle Handle) con
 	return IsReplicatedHandle(Handle) ? GetObjectFromReferenceHandle(Handle) : nullptr;
 };
 
+UObject* UObjectReplicationBridge::GetPreRegisteredObject(FNetRefHandle Handle) const
+{
+	return GetObjectReferenceCache()->GetPreRegisteredObjectFromReferenceHandle(Handle);
+};
+
 UE::Net::FNetRefHandle UObjectReplicationBridge::GetReplicatedRefHandle(const UObject* Object, EGetRefHandleFlags GetRefHandleFlags) const
 {
 	const FNetRefHandle Handle = GetObjectReferenceCache()->GetObjectReferenceHandleFromObject(Object, GetRefHandleFlags);
@@ -304,6 +309,21 @@ UE::Net::FNetRefHandle UObjectReplicationBridge::GetReplicatedRefHandle(FNetHand
 	return ObjectData.RefHandle;
 }
 
+bool UObjectReplicationBridge::IsNetRefHandlePreRegistered(FNetRefHandle Handle)
+{
+	return ObjectReferenceCache->IsNetRefHandlePreRegistered(Handle);
+}
+
+UE::Net::FNetRefHandle UObjectReplicationBridge::PreRegisterNewObjectReferenceHandle(UObject* Instance)
+{
+	return ObjectReferenceCache->PreRegisterObjectReferenceHandle(Instance);
+}
+
+void UObjectReplicationBridge::PreRegisterObjectWithReferenceHandle(const UObject* Instance, FNetRefHandle Handle)
+{
+	ObjectReferenceCache->AddPreRegisteredReference(Handle, Instance);
+}
+
 UE::Net::FNetRefHandle UObjectReplicationBridge::StartReplicatingNetObject(UObject* Instance, UE::Net::EReplicationFragmentTraits Traits, UE::Net::FNetObjectFactoryId NetFactoryId)
 {
 	using namespace UE::Net;
@@ -553,6 +573,15 @@ UE::Net::FNetRefHandle UObjectReplicationBridge::StartReplicatingSubObject(UObje
 	checkf(!Instance->HasAnyFlags(RF_ArchetypeObject | RF_ClassDefaultObject), TEXT("Iris cannot replicate subobject %s owned by %s because it's an %s"), 
 		*GetNameSafe(Instance), *GetNameSafe(LocalNetRefHandleManager.GetReplicatedObjectInstance(LocalNetRefHandleManager.GetInternalIndex(Params.RootObjectHandle))), Instance->HasAnyFlags(RF_ArchetypeObject)?TEXT("Archetype"):TEXT("DefaultObject"));
 
+	// Pre-registering subobjects is not supported
+	FNetRefHandle PotentiallyPreRegisteredHandle = GetObjectReferenceCache()->GetObjectReferenceHandleFromObject(Instance);
+	if (PotentiallyPreRegisteredHandle.IsValid() && IsNetRefHandlePreRegistered(PotentiallyPreRegisteredHandle))
+	{
+		UE_LOG_OBJECTREPLICATIONBRIDGE(Error, TEXT("UObjectReplicationBridge::StartReplicatingSubObject attempted to replicate pre-registered object %s."), *GetFullNameSafe(Instance));
+		ensureMsgf(false, TEXT("Attempted to replicate a pre-registered subobject"));
+		return FNetRefHandle::GetInvalid();
+	}
+
 	FNetRefHandle SubObjectRefHandle = GetReplicatedRefHandle(Instance);
 	if (SubObjectRefHandle.IsValid())
 	{
@@ -560,7 +589,7 @@ UE::Net::FNetRefHandle UObjectReplicationBridge::StartReplicatingSubObject(UObje
 		check(Params.RootObjectHandle == LocalNetRefHandleManager.GetRootObjectOfSubObject(SubObjectRefHandle));
 		return SubObjectRefHandle;
 	}
-	
+
 	const EReplicationFragmentTraits Traits = EReplicationFragmentTraits::CanReplicate;
 	SubObjectRefHandle = StartReplicatingNetObject(Instance, Traits, NetFactoryId);
 
diff --git a/Engine/Source/Runtime/Experimental/Iris/Core/Public/Iris/ReplicationSystem/ObjectReplicationBridge.h b/Engine/Source/Runtime/Experimental/Iris/Core/Public/Iris/ReplicationSystem/ObjectReplicationBridge.h
--- a/Engine/Source/Runtime/Experimental/Iris/Core/Public/Iris/ReplicationSystem/ObjectReplicationBridge.h
+++ b/Engine/Source/Runtime/Experimental/Iris/Core/Public/Iris/ReplicationSystem/ObjectReplicationBridge.h
@@ -167,6 +167,32 @@ public:
 	/** Return the name of the default spatial filter config. @see UObjectReplicationBridgeConfig */
 	FName GetDefaultSpatialFilterName() const { return DefaultSpatialFilterName; }
 
+	/** Return true only if Handle was registered via PreRegisterNewObjectReferenceHandle or PreregisterObjectWithReferenceHandle */
+	IRISCORE_API bool IsNetRefHandlePreRegistered(FNetRefHandle Handle);
+
+	/** Get the Object from a handle, even if the object was pre-registered and hasn't replicated yet. */
+	IRISCORE_API UObject* GetPreRegisteredObject(FNetRefHandle Handle) const;
+
+	/**
+	 * Creates a NetRefHandle explicitly for object Instance and flags it as pre-registered in the ObjectReferenceCache.
+	 * If supported by an object factory, the pre-registered flag is meant to be serialized in a root object creation header
+	 * and when a client receives a dynamic object header with this flag, it will look for an instance already registered by
+	 * PreRegisterObjectWithReferenceHandle instead of creating a new instance of the object.
+	 * Meant for use by the NetDriver's CreateNetIDForObject function.
+	 * @see UNetDriver::CreateNetIDForObject, UObjectReplicationBridge::PreRegisterObjectWithReferenceHandle
+	 */
+	IRISCORE_API UE::Net::FNetRefHandle PreRegisterNewObjectReferenceHandle(UObject* Instance);
+
+	/**
+	 * Directly assigns Handle to object Instance and flags it as pre-registered in the ObjectReferenceCache.
+	 * If supported by the object factory, clients receiving an object creation header with the pre-registered flag
+	 * will look for existing instances registered by this function instead of creating a new instance and
+	 * use the found instance for replication.
+	 * Meant for use by the NetDriver's AssignNetIDToObject function.
+	 * @see UNetDriver::AssignNetIDToObject, UObjectReplicationBridge::PreRegisterNewObjectReferenceHandle
+	 */
+	IRISCORE_API void PreRegisterObjectWithReferenceHandle(const UObject* Instance, FNetRefHandle Handle);
+
 public:
 
 	struct FRootObjectReplicationParams
