[Materials] Replace redundant override Compile() functions with data-driven content.

#rnx

diff --git a/Engine/Config/BaseMaterialExpressions.ini b/Engine/Config/BaseMaterialExpressions.ini
--- a/Engine/Config/BaseMaterialExpressions.ini
+++ b/Engine/Config/BaseMaterialExpressions.ini
@@ -969,7 +969,7 @@ bShaderInputData=true
 +ExternalCodeDeclarations=(ReturnType=Float, Name="DecalLifetimeOpacity", Definition="DecalLifetimeOpacity()", Derivative=Zero, ShaderFrequency=Pixel, Domains=(MD_DeferredDecal))
 +ExternalCodeDeclarations=(ReturnType=Float, Name="DeltaTime", Definition="View.DeltaTime", Derivative=Zero, bIsInlined=true)
 +ExternalCodeDeclarations=(ReturnType=Float, Name="EyeAdaptation", Definition="EyeAdaptationLookup()", Derivative=Zero, ShaderFrequency=Pixel, bIsInlined=true)
-+ExternalCodeDeclarations=(ReturnType=Float4, Name="FontSignedDistanceData", Definition="GetFontSignedDistanceData(Parameters)", ShaderFrequency=Pixel)
++ExternalCodeDeclarations=(ReturnType=Float4, Name="FontSignedDistanceData", Definition="GetFontSignedDistanceData(Parameters)", ShaderFrequency=Pixel, Domains=(MD_UI))
 +ExternalCodeDeclarations=(ReturnType=Float, Name="IsOrthographic", Definition="IsOrthoProjectionFloat()", Derivative=Zero, bIsInlined=true)
 +ExternalCodeDeclarations=(ReturnType=Float3, Name="LightVector", Definition="Parameters.LightVector", ShaderFrequency=Pixel|Compute, bIsInlined=true)
 +ExternalCodeDeclarations=(ReturnType=Float, Name="MeshPaintTextureCoordinateIndex", Definition="GetMeshPaintTextureCoordinateIndex(GetPrimitiveData(Parameters))", bIsInlined=true)
diff --git a/Engine/Source/Runtime/Engine/Private/Materials/HLSLMaterialTranslator.cpp b/Engine/Source/Runtime/Engine/Private/Materials/HLSLMaterialTranslator.cpp
--- a/Engine/Source/Runtime/Engine/Private/Materials/HLSLMaterialTranslator.cpp
+++ b/Engine/Source/Runtime/Engine/Private/Materials/HLSLMaterialTranslator.cpp
@@ -3856,18 +3856,22 @@ int32 FHLSLMaterialTranslator::AddExternalCodeChunk(const FMaterialExternalCodeD
 	}
 
 	// Validate this external code can be used for the current material domain. Empty list implies no restriction on material domains.
-	if (!InExternalCode.Domains.IsEmpty() && InExternalCode.Domains.Find(Material->GetMaterialDomain()) == INDEX_NONE)
+	if (!InExternalCode.Domains.IsEmpty())
 	{
-		TStringBuilder<1024> SupportedMaterialDomainsString;
-		for (EMaterialDomain Domain : InExternalCode.Domains)
+		checkf(Material != nullptr, TEXT("Missing material while emitting external HLSL code '%s'"), *InExternalCode.Name.ToString());
+		if (InExternalCode.Domains.Find(Material->GetMaterialDomain()) == INDEX_NONE)
 		{
-			if (SupportedMaterialDomainsString.Len() > 0)
+			TStringBuilder<1024> SupportedMaterialDomainsString;
+			for (EMaterialDomain Domain : InExternalCode.Domains)
 			{
-				SupportedMaterialDomainsString.Append(TEXT(", "));
+				if (SupportedMaterialDomainsString.Len() > 0)
+				{
+					SupportedMaterialDomainsString.Append(TEXT(", "));
+				}
+				SupportedMaterialDomainsString.Append(*MaterialDomainString(Domain));
 			}
-			SupportedMaterialDomainsString.Append(*MaterialDomainString(Domain));
+			return Errorf(TEXT("External code '%s' (asset: %s) is only available in the following material domains: %s"), *InExternalCode.Name.ToString(), *Material->GetAssetPath().ToString(), SupportedMaterialDomainsString.ToString());
 		}
-		return Errorf(TEXT("External code '%s' (asset: %s) is only available in the following material domains: %s"), *InExternalCode.Name.ToString(), *Material->GetAssetPath().ToString(), SupportedMaterialDomainsString.ToString());
 	}
 
 	// Emit external code definition as new code chunk
diff --git a/Engine/Source/Runtime/Engine/Private/Materials/MaterialExpressions.cpp b/Engine/Source/Runtime/Engine/Private/Materials/MaterialExpressions.cpp
--- a/Engine/Source/Runtime/Engine/Private/Materials/MaterialExpressions.cpp
+++ b/Engine/Source/Runtime/Engine/Private/Materials/MaterialExpressions.cpp
@@ -9852,15 +9852,6 @@ UMaterialExpressionFontSignedDistance::UMaterialExpressionFontSignedDistance(con
 }
 
 #if WITH_EDITOR
-int32 UMaterialExpressionFontSignedDistance::Compile(class FMaterialCompiler* Compiler, int32 OutputIndex)
-{
-	if(!Material || Material->MaterialDomain != MD_UI)
-	{
-		return Compiler->Errorf(TEXT("Font Signed Distance node is only compatible with UI materials"));
-	}
-	return Super::Compile(Compiler, OutputIndex);
-}
-
 void UMaterialExpressionFontSignedDistance::GetCaption(TArray<FString>& OutCaptions) const
 {
 	OutCaptions.Add(TEXT("Font Signed Distance"));
@@ -21486,16 +21477,6 @@ UMaterialExpressionPreSkinnedPosition::UMaterialExpressionPreSkinnedPosition(con
 }
 
 #if WITH_EDITOR
-int32 UMaterialExpressionPreSkinnedPosition::Compile(class FMaterialCompiler* Compiler, int32 OutputIndex)
-{
-	if (Compiler->GetCurrentShaderFrequency() != SF_Vertex)
-	{
-		return Compiler->Errorf(TEXT("Pre-skinned position is only available in the vertex shader, pass through custom interpolators if needed."));
-	}
-
-	return Super::Compile(Compiler, OutputIndex);
-}
-
 void UMaterialExpressionPreSkinnedPosition::GetCaption(TArray<FString>& OutCaptions) const
 {
 	OutCaptions.Add(TEXT("Pre-Skinned Local Position (deprecated)"));
diff --git a/Engine/Source/Runtime/Engine/Public/Materials/MaterialExpressionFontSignedDistance.h b/Engine/Source/Runtime/Engine/Public/Materials/MaterialExpressionFontSignedDistance.h
--- a/Engine/Source/Runtime/Engine/Public/Materials/MaterialExpressionFontSignedDistance.h
+++ b/Engine/Source/Runtime/Engine/Public/Materials/MaterialExpressionFontSignedDistance.h
@@ -13,7 +13,6 @@ class UMaterialExpressionFontSignedDistance : public UMaterialExpressionExternal
 
 	//~ Begin UMaterialExpression Interface
 #if WITH_EDITOR
-	virtual int32 Compile(class FMaterialCompiler* Compiler, int32 OutputIndex) override;
 	virtual void GetCaption(TArray<FString>& OutCaptions) const override;
 #endif
 	//~ End UMaterialExpression Interface
diff --git a/Engine/Source/Runtime/Engine/Public/Materials/MaterialExpressionPreSkinnedPosition.h b/Engine/Source/Runtime/Engine/Public/Materials/MaterialExpressionPreSkinnedPosition.h
--- a/Engine/Source/Runtime/Engine/Public/Materials/MaterialExpressionPreSkinnedPosition.h
+++ b/Engine/Source/Runtime/Engine/Public/Materials/MaterialExpressionPreSkinnedPosition.h
@@ -14,7 +14,6 @@ class UMaterialExpressionPreSkinnedPosition : public UMaterialExpressionExternal
 
 	//~ Begin UMaterialExpression Interface
 #if WITH_EDITOR
-	virtual int32 Compile(class FMaterialCompiler* Compiler, int32 OutputIndex) override;
 	virtual void GetCaption(TArray<FString>& OutCaptions) const override;
 	virtual void GetExpressionToolTip(TArray<FString>& OutToolTip) override;
 	virtual FText GetKeywords() const override {return FText::FromString(TEXT("position"));}
