Disable distance field for IsSky material used on skydome where this is not required.

#tests pc editor dx12
#rnx

diff --git a/Engine/Source/Runtime/Engine/Private/StaticMeshSceneProxy.cpp b/Engine/Source/Runtime/Engine/Private/StaticMeshSceneProxy.cpp
--- a/Engine/Source/Runtime/Engine/Private/StaticMeshSceneProxy.cpp
+++ b/Engine/Source/Runtime/Engine/Private/StaticMeshSceneProxy.cpp
@@ -321,7 +321,7 @@ FStaticMeshSceneProxy::FStaticMeshSceneProxy(const FStaticMeshSceneProxyDesc& In
 	DistanceFieldData = RenderData->LODResources[0].DistanceFieldData;
 	CardRepresentationData = RenderData->LODResources[0].CardRepresentationData;
 
-	bSupportsDistanceFieldRepresentation = MaterialRelevance.bOpaque && !MaterialRelevance.bUsesSingleLayerWaterMaterial && DistanceFieldData && DistanceFieldData->IsValid();
+	bSupportsDistanceFieldRepresentation = MaterialRelevance.bOpaque && !MaterialRelevance.bUsesSkyMaterial && !MaterialRelevance.bUsesSingleLayerWaterMaterial && DistanceFieldData && DistanceFieldData->IsValid();
 	bCastsDynamicIndirectShadow = InProxyDesc.bCastDynamicShadow && InProxyDesc.CastShadow && InProxyDesc.bCastDistanceFieldIndirectShadow && InProxyDesc.Mobility != EComponentMobility::Static && !InProxyDesc.bIsFirstPerson;
 	DynamicIndirectShadowMinVisibility = FMath::Clamp(InProxyDesc.DistanceFieldIndirectShadowMinVisibility, 0.0f, 1.0f);
 	DistanceFieldSelfShadowBias = FMath::Max(InProxyDesc.bOverrideDistanceFieldSelfShadowBias ? InProxyDesc.DistanceFieldSelfShadowBias : InProxyDesc.GetStaticMesh()->DistanceFieldSelfShadowBias, 0.0f);
