Skeletal Mesh Cloth - Fixed overzealous skeletal mesh component cloth ticking when montages are set to tick.

[RN] Skeletal Mesh Cloth - Fixed overzealous skeletal mesh compoonent cloth ticking when montages are set to tick.
#tests Tested with and without CVAR against legacy behavior.

diff --git a/Engine/Source/Runtime/Engine/Private/SkeletalMeshComponentPhysics.cpp b/Engine/Source/Runtime/Engine/Private/SkeletalMeshComponentPhysics.cpp
--- a/Engine/Source/Runtime/Engine/Private/SkeletalMeshComponentPhysics.cpp
+++ b/Engine/Source/Runtime/Engine/Private/SkeletalMeshComponentPhysics.cpp
@@ -57,6 +57,7 @@ TAutoConsoleVariable<int32> CVarEnableClothPhysics(TEXT("p.ClothPhysics"), 1, TE
 
 static TAutoConsoleVariable<int32> CVarEnableClothPhysicsUseTaskThread(TEXT("p.ClothPhysics.UseTaskThread"), 1, TEXT("If 1, run cloth on the task thread. If 0, run on game thread."));
 static TAutoConsoleVariable<int32> CVarClothPhysicsTickWaitForParallelClothTask(TEXT("p.ClothPhysics.WaitForParallelClothTask"), 0, TEXT("If 1, always wait for cloth task completion in the Cloth Tick function. If 0, wait at end-of-frame updates instead if allowed by component settings"));
+static TAutoConsoleVariable<int32> CVarClothPhysicsTickWithMontages(TEXT("p.ClothPhysics.TickWithMontages"), 0, TEXT("Set to 1 to keep legacy behavior and tick clothing with montage related VisibilityBasedAnimTickOption."));
 
 static TAutoConsoleVariable<int32> CVarEnableKinematicDeferralPrePhysicsCondition(TEXT("p.EnableKinematicDeferralPrePhysicsCondition"), 1, TEXT("If is 1, and deferral would've been disallowed due to EUpdateTransformFlags, allow if in PrePhysics tick. If 0, condition is unchanged."));
 
@@ -3561,7 +3562,12 @@ void USkeletalMeshComponent::TickClothing(float DeltaTime, FTickFunction& ThisTi
 	}
 
 	// Use the component update flag to gate simulation to respect the always tick options
-	bool bShouldTick = ((VisibilityBasedAnimTickOption < EVisibilityBasedAnimTickOption::OnlyTickPoseWhenRendered) || bRecentlyRendered);
+	static const bool bClothPhysicsTickWithMontages = (CVarClothPhysicsTickWithMontages.GetValueOnAnyThread() != 0);
+	const bool bShouldTick = bClothPhysicsTickWithMontages ? 
+		((VisibilityBasedAnimTickOption < EVisibilityBasedAnimTickOption::OnlyTickPoseWhenRendered) || bRecentlyRendered) :  // Legacy ticking behavior that ticked with montages
+		(VisibilityBasedAnimTickOption == EVisibilityBasedAnimTickOption::AlwaysTickPoseAndRefreshBones ||
+		VisibilityBasedAnimTickOption == EVisibilityBasedAnimTickOption::AlwaysTickPose ||
+		bRecentlyRendered);
 
 	if (bShouldTick)
 	{
