Streamable Manager
- DownloadParams is now an optional field on LoadParams instead of being passed seperately.

diff --git a/Engine/Source/Runtime/Engine/Classes/Engine/StreamableManager.h b/Engine/Source/Runtime/Engine/Classes/Engine/StreamableManager.h
--- a/Engine/Source/Runtime/Engine/Classes/Engine/StreamableManager.h
+++ b/Engine/Source/Runtime/Engine/Classes/Engine/StreamableManager.h
@@ -685,6 +685,17 @@ enum class EStreamableDownloadPriority : uint8
 	High
 };
 
+/** Additional parameters for downloading packages */
+struct FStreamableDownloadParams
+{
+	/** Handle to keep packages in the local cache. If no cache pin is provided, a new one will be created. */
+	FStreamableDownloadCachePin CachePin;
+	/** Download priority */
+	EStreamableDownloadPriority Priority = EStreamableDownloadPriority::Normal;
+	/** If set, any soft references will also be installed */
+	bool bInstallSoftReferences = false;
+};
+
 /** Parameters for an async load */
 struct FStreamableAsyncLoadParams
 {
@@ -702,17 +713,8 @@ struct FStreamableAsyncLoadParams
 	bool bManageActiveHandle = false;
 	/** [optional] If true, the handle will start in a stalled state and will not attempt to actually async load until StartStalledHandle is called on it */
 	bool bStartStalled = false;
-};
-
-/** Additional parameters for downloading packages */
-struct FStreamableDownloadParams
-{
-	/** Handle to keep packages in the local cache. If no cache pin is provided, a new one will be created. */
-	FStreamableDownloadCachePin CachePin;
-	/** Download priority */
-	EStreamableDownloadPriority Priority = EStreamableDownloadPriority::Normal;
-	/** If set, any soft references will also be installed */
-	bool bInstallSoftReferences = false;
+	/** [optional] If present, will attempt to download necessary packages with IoStoreOnDemand */
+	TOptional<FStreamableDownloadParams> DownloadParams;
 };
 
 /** A native class for managing streaming assets in and keeping them in memory. AssetManager is the global singleton version of this with blueprint access */
@@ -753,20 +755,6 @@ struct FStreamableManager : public FGCObject
 		FStreamableAsyncLoadParams&& Params,
 		DebugNameType&& DebugNameOrLocation = UE::FSourceLocation::Current());
 
-	/** 
-	 * This is the primary streamable operation. Requests streaming of one or more target objects. When complete, a delegate function is called. Returns a Streamable Handle.
-	 * This overload will download packages if needed.
-	 *
-	 * @param Params				Load Parameters
-	 * @param DownloadParams		Additional Parameters to control downloading
-	 * @param DebugName				[optional] Name of this handle, either FString or anything that can construct FString, will be reported in debug tools, will report Source Location if DebugName is not specified explicitly
-	 */ 
-	template< typename DebugNameType = UE::FSourceLocation >
-	TSharedPtr<FStreamableHandle> RequestAsyncLoad(
-		FStreamableAsyncLoadParams&& LoadParams,
-		FStreamableDownloadParams&& DownloadParams,
-		DebugNameType&& DebugNameOrLocation = UE::FSourceLocation::Current());
-
 	/** 
 	 * Synchronously load a set of assets, and return a handle. This can be very slow and may stall the game thread for several seconds.
 	 * 
@@ -832,7 +820,6 @@ struct FStreamableManager : public FGCObject
 	ENGINE_API TSharedPtr<FStreamableHandle> CreateErrorHandle(
 		FStreamableError&& Error,
 		FStreamableAsyncLoadParams&& Params,
-		TOptional<FStreamableDownloadParams> DownloadParams = {},
 		const FString& DebugName = FStreamableHandle::HandleDebugName_CombinedHandle);
 
 	/** 
@@ -892,7 +879,7 @@ private:
 	void AsyncLoadCallback(FSoftObjectPath Request, UPackage* Package, EAsyncLoadingResult::Type Result);
 
 	ENGINE_API static bool ShouldStripDebugName();
-	ENGINE_API TSharedPtr<FStreamableHandle> RequestAsyncLoadInternal(FStreamableAsyncLoadParams&& Params, FStreamableDownloadParams&& DownloadParams, FString&& DebugName);
+	ENGINE_API TSharedPtr<FStreamableHandle> RequestAsyncLoadInternal(FStreamableAsyncLoadParams&& Params, FString&& DebugName);
 	ENGINE_API TSharedPtr<FStreamableHandle> RequestSyncLoadInternal(TArray<FSoftObjectPath>&& TargetsToStream, bool bManageActiveHandle, FString&& DebugName);
 
 	/** Map of paths to streamable objects, this will be the post-redirector name */
@@ -1008,14 +995,12 @@ TSharedPtr<FStreamableHandle> FStreamableManager::RequestAsyncLoad(
 	{
 		return RequestAsyncLoadInternal(
 			MoveTemp(Params),
-			FStreamableDownloadParams(),
 			ShouldStripDebugName() ? FString() : UE::SourceLocation::ToFileAndLineString(DebugNameOrLocation));
 	}
 	else
 	{
 		return RequestAsyncLoadInternal(
 			MoveTemp(Params),
-			FStreamableDownloadParams(),
 			ShouldStripDebugName() ? FString() : FString{ Forward<DebugNameType>(DebugNameOrLocation) });
 	}
 }
@@ -1029,41 +1014,12 @@ TSharedPtr<FStreamableHandle> FStreamableManager::RequestAsyncLoad(
 	{
 		return RequestAsyncLoadInternal(
 			MoveTemp(Params),
-			FStreamableDownloadParams(),
 			ShouldStripDebugName() ? FString() : UE::SourceLocation::ToFileAndLineString(DebugNameOrLocation));
 	}
 	else
 	{
 		return RequestAsyncLoadInternal(
 			MoveTemp(Params),
-			FStreamableDownloadParams(),
-			ShouldStripDebugName() ? FString() : FString{ Forward<DebugNameType>(DebugNameOrLocation) });
-	}
-}
-
-template< typename DebugNameType >
-TSharedPtr<FStreamableHandle> FStreamableManager::RequestAsyncLoad(
-	FStreamableAsyncLoadParams&& LoadParams,
-	FStreamableDownloadParams&& DownloadParams,
-	DebugNameType&& DebugNameOrLocation)
-{
-	if (!DownloadParams.CachePin)
-	{
-		DownloadParams.CachePin = FStreamableDownloadCachePin::Create(DebugNameOrLocation);
-	}
-
-	if constexpr (std::is_same_v<std::decay_t<DebugNameType>, UE::FSourceLocation>)
-	{
-		return RequestAsyncLoadInternal(
-			MoveTemp(LoadParams),
-			MoveTemp(DownloadParams),
-			ShouldStripDebugName() ? FString() : UE::SourceLocation::ToFileAndLineString(DebugNameOrLocation));
-	}
-	else
-	{
-		return RequestAsyncLoadInternal(
-			MoveTemp(LoadParams),
-			MoveTemp(DownloadParams),
 			ShouldStripDebugName() ? FString() : FString{ Forward<DebugNameType>(DebugNameOrLocation) });
 	}
 }
diff --git a/Engine/Source/Runtime/Engine/Private/AssetManager.cpp b/Engine/Source/Runtime/Engine/Private/AssetManager.cpp
--- a/Engine/Source/Runtime/Engine/Private/AssetManager.cpp
+++ b/Engine/Source/Runtime/Engine/Private/AssetManager.cpp
@@ -1972,7 +1972,7 @@ TSharedPtr<FStreamableHandle> UAssetManager::ChangeBundleStateForPrimaryAssets(
 				StreamableLoadParams.TargetsToStream = MoveTemp(PathsToLoad);
 				NewHandle = StreamableManager.CreateErrorHandle(
 					FStreamableError(EStreamableErrorType::DownloadError, FString::Printf(TEXT("Expected DownloadParams for %s"), *PrimaryAssetId.ToString())),
-					MoveTemp(StreamableLoadParams), {}, FString(DebugName));
+					MoveTemp(StreamableLoadParams), FString(DebugName));
 			}
 
 			if (!NewHandle.IsValid())
@@ -2393,15 +2393,9 @@ TSharedPtr<FStreamableHandle> UAssetManager::LoadAssetListInternal(TArray<FSoftO
 		Params.Priority = LoadParams.Priority;
 		Params.bManageActiveHandle = false;
 		Params.bStartStalled = MissingChunks.Num() > 0;
+		Params.DownloadParams = MoveTemp(LoadParams.DownloadParams);
 
-		if (LoadParams.DownloadParams)
-		{
-			NewHandle = StreamableManager.RequestAsyncLoad(MoveTemp(Params), MoveTemp(*LoadParams.DownloadParams), MoveTemp(DebugName));
-		}
-		else
-		{
-			NewHandle = StreamableManager.RequestAsyncLoad(MoveTemp(Params), MoveTemp(DebugName));
-		}
+		NewHandle = StreamableManager.RequestAsyncLoad(MoveTemp(Params), MoveTemp(DebugName));
 
 		if (MissingChunks.Num() > 0 && NewHandle.IsValid())
 		{
diff --git a/Engine/Source/Runtime/Engine/Private/StreamableManager.cpp b/Engine/Source/Runtime/Engine/Private/StreamableManager.cpp
--- a/Engine/Source/Runtime/Engine/Private/StreamableManager.cpp
+++ b/Engine/Source/Runtime/Engine/Private/StreamableManager.cpp
@@ -1728,7 +1728,7 @@ bool FStreamableManager::ShouldStripDebugName()
 	return bShouldStripDebugName;
 }
 
-TSharedPtr<FStreamableHandle> FStreamableManager::RequestAsyncLoadInternal(FStreamableAsyncLoadParams&& Params, FStreamableDownloadParams&& DownloadParams, FString&& DebugName)
+TSharedPtr<FStreamableHandle> FStreamableManager::RequestAsyncLoadInternal(FStreamableAsyncLoadParams&& Params, FString&& DebugName)
 {
 	LLM_SCOPE(ELLMTag::StreamingManager);
 
@@ -1834,12 +1834,17 @@ TSharedPtr<FStreamableHandle> FStreamableManager::RequestAsyncLoadInternal(FStre
 		ManagedActiveHandles.Add(NewRequest);
 	}
 
-	if (DownloadParams.CachePin)
+	if (Params.DownloadParams)
 	{
-		NewRequest->DownloadCachePin = DownloadParams.CachePin;
+		NewRequest->DownloadCachePin = Params.DownloadParams->CachePin;
+		if (!NewRequest->DownloadCachePin)
+		{
+			NewRequest->DownloadCachePin = FStreamableDownloadCachePin::Create(NewRequest->DebugName);
+		}
+
 		NewRequest->DownloadContext = MakeUnique<UE::StreamableManager::Private::FDownloadContext>();
-		NewRequest->DownloadContext->Priority = DownloadParams.Priority;
-		NewRequest->DownloadContext->bInstallSoftReferences = DownloadParams.bInstallSoftReferences;
+		NewRequest->DownloadContext->Priority = Params.DownloadParams->Priority;
+		NewRequest->DownloadContext->bInstallSoftReferences = Params.DownloadParams->bInstallSoftReferences;
 	}
 
 	if (Params.bStartStalled)
@@ -1867,7 +1872,7 @@ TSharedPtr<FStreamableHandle> FStreamableManager::RequestSyncLoadInternal(TArray
 	AsyncParams.bManageActiveHandle = bManageActiveHandle;
 
 	// Do an async load and wait to complete. In some cases this will do a sync load due to safety issues
-	TSharedPtr<FStreamableHandle> Request = RequestAsyncLoadInternal(MoveTemp(AsyncParams), FStreamableDownloadParams(), MoveTemp(DebugName));
+	TSharedPtr<FStreamableHandle> Request = RequestAsyncLoadInternal(MoveTemp(AsyncParams), MoveTemp(DebugName));
 
 	bForceSynchronousLoads = false;
 
@@ -2371,7 +2376,7 @@ TSharedPtr<FStreamableHandle> FStreamableManager::CreateCombinedHandle(const TCo
 	return NewRequest;
 }
 
-TSharedPtr<FStreamableHandle> FStreamableManager::CreateErrorHandle(FStreamableError&& Error, FStreamableAsyncLoadParams&& Params, TOptional<FStreamableDownloadParams> DownloadParams, const FString& DebugName)
+TSharedPtr<FStreamableHandle> FStreamableManager::CreateErrorHandle(FStreamableError&& Error, FStreamableAsyncLoadParams&& Params, const FString& DebugName)
 {
 	TSharedRef<FStreamableHandle> NewRequest = MakeShared<FStreamableHandle>(FStreamableHandle::FPrivateToken());
 	NewRequest->OwningManager = this;
@@ -2379,14 +2384,14 @@ TSharedPtr<FStreamableHandle> FStreamableManager::CreateErrorHandle(FStreamableE
 	NewRequest->DebugName = DebugName;
 	NewRequest->Priority = Params.Priority;
 
-	if (DownloadParams)
+	if (Params.DownloadParams)
 	{
-		if (!DownloadParams->CachePin)
+		if (!Params.DownloadParams->CachePin)
 		{
-			DownloadParams->CachePin = FStreamableDownloadCachePin::Create(DebugName);
+			Params.DownloadParams->CachePin = FStreamableDownloadCachePin::Create(DebugName);
 		}
 
-		NewRequest->DownloadCachePin = DownloadParams->CachePin;
+		NewRequest->DownloadCachePin = Params.DownloadParams->CachePin;
 	}
 
 	NewRequest->CancelHandle();
