Comments for Asset Manager download params to make clear the feature is experimental.

diff --git a/Engine/Source/Runtime/Engine/Classes/Engine/AssetManager.h b/Engine/Source/Runtime/Engine/Classes/Engine/AssetManager.h
--- a/Engine/Source/Runtime/Engine/Classes/Engine/AssetManager.h
+++ b/Engine/Source/Runtime/Engine/Classes/Engine/AssetManager.h
@@ -32,6 +32,7 @@ DECLARE_DELEGATE_OneParam(FAssetManagerAcquireResourceDelegate, bool /* bSuccess
 /** Delegate called when new asset search root is registered due to runtime asset mounting, path will not have a trailing slash */
 DECLARE_MULTICAST_DELEGATE_OneParam(FOnAddedAssetSearchRoot, const FString&);
 
+/** EXPERIMENTAL -  Additional parameters for downloading packages */
 using FAssetManagerDownloadParams = FStreamableDownloadParams;
 
 struct FAssetManagerLoadParams
@@ -44,7 +45,7 @@ struct FAssetManagerLoadParams
 	FStreamableUpdateDelegate OnUpdate;
 	/** [optional] Priority to pass to the streaming system, higher priority will be loaded first */
 	TAsyncLoadPriority Priority = UE::StreamableManager::Private::DefaultAsyncLoadPriority;
-	/** [optional] If present, will attempt to download necessary packages with IoStoreOnDemand */
+	/** [optional] EXPERIMENTAL - If present, will attempt to download necessary packages with IoStoreOnDemand */
 	TOptional<FAssetManagerDownloadParams> DownloadParams;
 };
 
diff --git a/Engine/Source/Runtime/Engine/Classes/Engine/StreamableManager.h b/Engine/Source/Runtime/Engine/Classes/Engine/StreamableManager.h
--- a/Engine/Source/Runtime/Engine/Classes/Engine/StreamableManager.h
+++ b/Engine/Source/Runtime/Engine/Classes/Engine/StreamableManager.h
@@ -685,7 +685,7 @@ enum class EStreamableDownloadPriority : uint8
 	High
 };
 
-/** Additional parameters for downloading packages */
+/** EXPERIMENTAL -  Additional parameters for downloading packages */
 struct FStreamableDownloadParams
 {
 	/** Handle to keep packages in the local cache. If no cache pin is provided, a new one will be created. */
@@ -713,7 +713,7 @@ struct FStreamableAsyncLoadParams
 	bool bManageActiveHandle = false;
 	/** [optional] If true, the handle will start in a stalled state and will not attempt to actually async load until StartStalledHandle is called on it */
 	bool bStartStalled = false;
-	/** [optional] If present, will attempt to download necessary packages with IoStoreOnDemand */
+	/** [optional] EXPERIMENTAL - If present, will attempt to download necessary packages with IoStoreOnDemand */
 	TOptional<FStreamableDownloadParams> DownloadParams;
 };
 
