Prevents UMaterialExpressionPinBase from being deleted by the user. Before the user was able to accidentally delete a pin base in a material expression composite, permanently breaking the expression.

Also, turned off logging "editor only data" renamed warning when a material is duplicated, as the renaming will always occur and is normal behavior.

diff --git a/Engine/Source/Runtime/Engine/Private/Materials/MaterialExpressions.cpp b/Engine/Source/Runtime/Engine/Private/Materials/MaterialExpressions.cpp
--- a/Engine/Source/Runtime/Engine/Private/Materials/MaterialExpressions.cpp
+++ b/Engine/Source/Runtime/Engine/Private/Materials/MaterialExpressions.cpp
@@ -1981,6 +1981,11 @@ bool UMaterialExpression::Modify( bool bAlwaysMarkDirty/*=true*/ )
 	return Super::Modify(bAlwaysMarkDirty);
 }
 
+bool UMaterialExpressionPinBase::CanUserDeleteExpression() const
+{
+	return false;
+}
+
 bool UMaterialExpression::MatchesSearchQuery( const TCHAR* SearchQuery )
 {
 	if (FCString::Stristr(SearchQuery, TEXT("NAME=")) != nullptr)
