Build fix in MaterialExpressions.cpp.

#rnx

diff --git a/Engine/Source/Runtime/Engine/Private/Materials/MaterialExpressions.cpp b/Engine/Source/Runtime/Engine/Private/Materials/MaterialExpressions.cpp
--- a/Engine/Source/Runtime/Engine/Private/Materials/MaterialExpressions.cpp
+++ b/Engine/Source/Runtime/Engine/Private/Materials/MaterialExpressions.cpp
@@ -1981,11 +1981,6 @@ bool UMaterialExpression::Modify( bool bAlwaysMarkDirty/*=true*/ )
 	return Super::Modify(bAlwaysMarkDirty);
 }
 
-bool UMaterialExpressionPinBase::CanUserDeleteExpression() const
-{
-	return false;
-}
-
 bool UMaterialExpression::MatchesSearchQuery( const TCHAR* SearchQuery )
 {
 	if (FCString::Stristr(SearchQuery, TEXT("NAME=")) != nullptr)
