Reworked GPU skin transitions to skip iterating all dispatch entries when no tangents are present. All transitions are pre-computed up front.

diff --git a/Engine/Source/Runtime/Engine/Private/GPUSkinCache.cpp b/Engine/Source/Runtime/Engine/Private/GPUSkinCache.cpp
--- a/Engine/Source/Runtime/Engine/Private/GPUSkinCache.cpp
+++ b/Engine/Source/Runtime/Engine/Private/GPUSkinCache.cpp
@@ -175,14 +175,6 @@ FAutoConsoleVariableRef CVarRecomputeTangentsParallelDispatch(
 	ECVF_RenderThreadSafe
 );
 
-static int32 GSkinCacheMaxDispatchesPerCmdList = 0;
-FAutoConsoleVariableRef CVarGPUSkinCacheMaxDispatchesPerCmdList(
-	TEXT("r.SkinCache.MaxDispatchesPerCmdList"),
-	GSkinCacheMaxDispatchesPerCmdList,
-	TEXT("Maximum number of compute shader dispatches which are batched together into a single command list to fix potential TDRs."),
-	ECVF_RenderThreadSafe
-);
-
 static int32 GSkinCachePrintMemorySummary = 0;
 FAutoConsoleVariableRef CVarGPUSkinCachePrintMemorySummary(
 	TEXT("r.SkinCache.PrintMemorySummary"),
@@ -1079,7 +1071,50 @@ class FRecomputeTangentsPerVertexPassCS : public FBaseRecomputeTangentsPerVertex
 IMPLEMENT_SHADER_TYPE(template<>, FRecomputeTangentsPerVertexPassCS<0>, TEXT("/Engine/Private/RecomputeTangentsPerVertexPass.usf"), TEXT("MainCS"), SF_Compute);
 IMPLEMENT_SHADER_TYPE(template<>, FRecomputeTangentsPerVertexPassCS<1>, TEXT("/Engine/Private/RecomputeTangentsPerVertexPass.usf"), TEXT("MainCS"), SF_Compute);
 
-void FGPUSkinCache::DispatchUpdateSkinTangents(FRHICommandList& RHICmdList, FGPUSkinCacheEntry* Entry, int32 SectionIndex, FSkinCacheRWBuffer*& StagingBuffer, bool bTrianglePass)
+void FGPUSkinCache::DispatchUpdateSkinTangentsVertexPass(FRHICommandList& RHICmdList, FGPUSkinCacheEntry* Entry, int32 SectionIndex, FSkinCacheRWBuffer*& StagingBuffer)
+{
+	FGPUSkinCacheEntry::FSectionDispatchData& DispatchData = Entry->DispatchData[SectionIndex];
+
+	SCOPED_DRAW_EVENTF(RHICmdList, SkinTangents_PerVertexPass, TEXT("%sTangentsVertex Mesh=%s, LOD=%d, Chunk=%d, InputStreamStart=%d, OutputStreamStart=%d, Vert=%d")
+		, RHI_BREADCRUMB_FORCE_STRING_LITERAL(Entry->Mode == EGPUSkinCacheEntryMode::RayTracing ? TEXT("[RT]") : TEXT(""))
+		, GetSkeletalMeshObjectDebugName(Entry->GPUSkin)
+		, Entry->LOD
+		, SectionIndex
+		, DispatchData.InputStreamStart
+		, DispatchData.OutputStreamStart
+		, DispatchData.NumVertices
+	);
+
+	//#todo-gpuskin Feature level?
+	auto* GlobalShaderMap = GetGlobalShaderMap(GetFeatureLevel());
+	TShaderMapRef<FRecomputeTangentsPerVertexPassCS<0>> ComputeShader0(GlobalShaderMap);
+	TShaderMapRef<FRecomputeTangentsPerVertexPassCS<1>> ComputeShader1(GlobalShaderMap);
+	TShaderRef<FBaseRecomputeTangentsPerVertexShader> ComputeShader;
+	if (DispatchData.Section->RecomputeTangentsVertexMaskChannel < ESkinVertexColorChannel::None)
+		ComputeShader = ComputeShader1;
+	else
+		ComputeShader = ComputeShader0;
+
+	uint32 VertexCount = DispatchData.NumVertices;
+	uint32 ThreadGroupCountValue = FMath::DivideAndRoundUp(VertexCount, ComputeShader->ThreadGroupSizeX);
+
+	if (!GRecomputeTangentsParallelDispatch)
+	{
+		// When triangle & vertex passes are interleaved, resource transition is needed in between.
+		RHICmdList.Transition({
+			DispatchData.GetTangentRWBuffer()->UpdateAccessState(ERHIAccess::UAVCompute),
+			StagingBuffer->UpdateAccessState(ERHIAccess::UAVCompute)
+		}, ERHITransitionCreateFlags::NoFence);
+	}
+
+	SetComputePipelineState(RHICmdList, ComputeShader.GetComputeShader());
+
+	SetShaderParametersLegacyCS(RHICmdList, ComputeShader, Entry, DispatchData, GRecomputeTangentsParallelDispatch ? DispatchData.GetIntermediateAccumulatedTangentBuffer()->Buffer : StagingBuffer->Buffer);
+	DispatchComputeShader(RHICmdList, ComputeShader.GetShader(), ThreadGroupCountValue, 1, 1);
+	UnsetShaderParametersLegacyCS(RHICmdList, ComputeShader);
+}
+
+void FGPUSkinCache::DispatchUpdateSkinTangentsTrianglePass(FRHICommandList& RHICmdList, FGPUSkinCacheEntry* Entry, int32 SectionIndex, FSkinCacheRWBuffer*& StagingBuffer)
 {
 	FGPUSkinCacheEntry::FSectionDispatchData& DispatchData = Entry->DispatchData[SectionIndex];
 
@@ -1087,150 +1122,101 @@ void FGPUSkinCache::DispatchUpdateSkinTangents(FRHICommandList& RHICmdList, FGPU
 	const int32 LODIndex = Entry->LOD;
 	FSkeletalMeshLODRenderData& LodData = SkelMeshRenderData.LODRenderData[LODIndex];
 
-	if (bTrianglePass)
+	if (!GRecomputeTangentsParallelDispatch)
 	{
-		if (!GRecomputeTangentsParallelDispatch)
-		{	
-			if (StagingBuffers.Num() != GNumTangentIntermediateBuffers)
+		if (StagingBuffers.Num() != GNumTangentIntermediateBuffers)
+		{
+			// Release extra buffers if shrinking
+			for (int32 Index = GNumTangentIntermediateBuffers; Index < StagingBuffers.Num(); ++Index)
 			{
-				// Release extra buffers if shrinking
-				for (int32 Index = GNumTangentIntermediateBuffers; Index < StagingBuffers.Num(); ++Index)
-				{
-					StagingBuffers[Index].Release();
-				}
-				StagingBuffers.SetNum(GNumTangentIntermediateBuffers, EAllowShrinking::No);
+				StagingBuffers[Index].Release();
 			}
+			StagingBuffers.SetNum(GNumTangentIntermediateBuffers, EAllowShrinking::No);
+		}
 
-			// no need to clear the staging buffer because we create it cleared and clear it after each usage in the per vertex pass
-			uint32 NumIntsPerBuffer = DispatchData.NumVertices * FGPUSkinCache::IntermediateAccumBufferNumInts;
-			CurrentStagingBufferIndex = (CurrentStagingBufferIndex + 1) % StagingBuffers.Num();
-			StagingBuffer = &StagingBuffers[CurrentStagingBufferIndex];
-			if (StagingBuffer->Buffer.NumBytes < NumIntsPerBuffer * sizeof(uint32))
-			{
-				StagingBuffer->Release();
-				StagingBuffer->Buffer.Initialize(RHICmdList, TEXT("SkinTangentIntermediate"), sizeof(int32), NumIntsPerBuffer, PF_R32_SINT, BUF_UnorderedAccess);
-				RHICmdList.BindDebugLabelName(StagingBuffer->Buffer.UAV, TEXT("SkinTangentIntermediate"));
+		// no need to clear the staging buffer because we create it cleared and clear it after each usage in the per vertex pass
+		uint32 NumIntsPerBuffer = DispatchData.NumVertices * FGPUSkinCache::IntermediateAccumBufferNumInts;
+		CurrentStagingBufferIndex = (CurrentStagingBufferIndex + 1) % StagingBuffers.Num();
+		StagingBuffer = &StagingBuffers[CurrentStagingBufferIndex];
+		if (StagingBuffer->Buffer.NumBytes < NumIntsPerBuffer * sizeof(uint32))
+		{
+			StagingBuffer->Release();
+			StagingBuffer->Buffer.Initialize(RHICmdList, TEXT("SkinTangentIntermediate"), sizeof(int32), NumIntsPerBuffer, PF_R32_SINT, BUF_UnorderedAccess);
+			RHICmdList.BindDebugLabelName(StagingBuffer->Buffer.UAV, TEXT("SkinTangentIntermediate"));
 
-				const uint32 MemSize = NumIntsPerBuffer * sizeof(uint32);
-				SET_MEMORY_STAT(STAT_GPUSkinCache_TangentsIntermediateMemUsed, MemSize);
+			const uint32 MemSize = NumIntsPerBuffer * sizeof(uint32);
+			SET_MEMORY_STAT(STAT_GPUSkinCache_TangentsIntermediateMemUsed, MemSize);
 
-				// The UAV must be zero-filled. We leave it zeroed after each round (see RecomputeTangentsPerVertexPass.usf), so this is only needed on when the buffer is first created.
-				RHICmdList.ClearUAVUint(StagingBuffer->Buffer.UAV, FUintVector4(0, 0, 0, 0));
-			}
+			// The UAV must be zero-filled. We leave it zeroed after each round (see RecomputeTangentsPerVertexPass.usf), so this is only needed on when the buffer is first created.
+			RHICmdList.ClearUAVUint(StagingBuffer->Buffer.UAV, FUintVector4(0, 0, 0, 0));
 		}
+	}
 
-		{
-			auto* GlobalShaderMap = GetGlobalShaderMap(GetFeatureLevel());
-			TShaderMapRef<FRecomputeTangentsPerTrianglePassCS<0>> ComputeShader00(GlobalShaderMap);
-			TShaderMapRef<FRecomputeTangentsPerTrianglePassCS<1>> ComputeShader01(GlobalShaderMap);
-			TShaderMapRef<FRecomputeTangentsPerTrianglePassCS<2>> ComputeShader10(GlobalShaderMap);
-			TShaderMapRef<FRecomputeTangentsPerTrianglePassCS<3>> ComputeShader11(GlobalShaderMap);
+	auto* GlobalShaderMap = GetGlobalShaderMap(GetFeatureLevel());
+	TShaderMapRef<FRecomputeTangentsPerTrianglePassCS<0>> ComputeShader00(GlobalShaderMap);
+	TShaderMapRef<FRecomputeTangentsPerTrianglePassCS<1>> ComputeShader01(GlobalShaderMap);
+	TShaderMapRef<FRecomputeTangentsPerTrianglePassCS<2>> ComputeShader10(GlobalShaderMap);
+	TShaderMapRef<FRecomputeTangentsPerTrianglePassCS<3>> ComputeShader11(GlobalShaderMap);
 
-			bool bFullPrecisionUV = LodData.StaticVertexBuffers.StaticMeshVertexBuffer.GetUseFullPrecisionUVs();
+	bool bFullPrecisionUV = LodData.StaticVertexBuffers.StaticMeshVertexBuffer.GetUseFullPrecisionUVs();
 
-			TShaderRef<FBaseRecomputeTangentsPerTriangleShader> Shader;
+	TShaderRef<FBaseRecomputeTangentsPerTriangleShader> Shader;
 
-			if (bFullPrecisionUV)
-			{
-				if (GAllowDupedVertsForRecomputeTangents) Shader = ComputeShader01;
-				else Shader = ComputeShader11;
-			}
-			else
-			{
-				if (GAllowDupedVertsForRecomputeTangents) Shader = ComputeShader00;
-				else Shader = ComputeShader10;
-			}
+	if (bFullPrecisionUV)
+	{
+		if (GAllowDupedVertsForRecomputeTangents) Shader = ComputeShader01;
+		else Shader = ComputeShader11;
+	}
+	else
+	{
+		if (GAllowDupedVertsForRecomputeTangents) Shader = ComputeShader00;
+		else Shader = ComputeShader10;
+	}
 
-			check(Shader.IsValid());
+	check(Shader.IsValid());
 
-			uint32 NumTriangles = DispatchData.NumTriangles;
-			uint32 ThreadGroupCountValue = FMath::DivideAndRoundUp(NumTriangles, FBaseRecomputeTangentsPerTriangleShader::ThreadGroupSizeX);
+	uint32 NumTriangles = DispatchData.NumTriangles;
+	uint32 ThreadGroupCountValue = FMath::DivideAndRoundUp(NumTriangles, FBaseRecomputeTangentsPerTriangleShader::ThreadGroupSizeX);
 
-			SCOPED_DRAW_EVENTF(RHICmdList, SkinTangents_PerTrianglePass, TEXT("%sTangentsTri  Mesh=%s, LOD=%d, Chunk=%d, IndexStart=%d Tri=%d BoneInfluenceType=%d UVPrecision=%d")
-				, RHI_BREADCRUMB_FORCE_STRING_LITERAL(Entry->Mode == EGPUSkinCacheEntryMode::RayTracing ? TEXT("[RT]") : TEXT(""))
-				, GetSkeletalMeshObjectDebugName(Entry->GPUSkin)
-				, LODIndex
-				, SectionIndex
-				, DispatchData.IndexBufferOffsetValue
-				, DispatchData.NumTriangles
-				, Entry->BoneInfluenceType
-				, bFullPrecisionUV
-			);
+	SCOPED_DRAW_EVENTF(RHICmdList, SkinTangents_PerTrianglePass, TEXT("%sTangentsTri  Mesh=%s, LOD=%d, Chunk=%d, IndexStart=%d Tri=%d BoneInfluenceType=%d UVPrecision=%d")
+		, RHI_BREADCRUMB_FORCE_STRING_LITERAL(Entry->Mode == EGPUSkinCacheEntryMode::RayTracing ? TEXT("[RT]") : TEXT(""))
+		, GetSkeletalMeshObjectDebugName(Entry->GPUSkin)
+		, LODIndex
+		, SectionIndex
+		, DispatchData.IndexBufferOffsetValue
+		, DispatchData.NumTriangles
+		, Entry->BoneInfluenceType
+		, bFullPrecisionUV
+	);
 
-			if (!GAllowDupedVertsForRecomputeTangents)
-			{
+	if (!GAllowDupedVertsForRecomputeTangents)
+	{
 #if WITH_EDITOR
-				check(LodData.RenderSections[SectionIndex].DuplicatedVerticesBuffer.DupVertData.Num() && LodData.RenderSections[SectionIndex].DuplicatedVerticesBuffer.DupVertIndexData.Num());
+		check(LodData.RenderSections[SectionIndex].DuplicatedVerticesBuffer.DupVertData.Num() && LodData.RenderSections[SectionIndex].DuplicatedVerticesBuffer.DupVertIndexData.Num());
 #endif
-				DispatchData.DuplicatedIndices = LodData.RenderSections[SectionIndex].DuplicatedVerticesBuffer.DuplicatedVerticesIndexBuffer.VertexBufferSRV;
-				DispatchData.DuplicatedIndicesIndices = LodData.RenderSections[SectionIndex].DuplicatedVerticesBuffer.LengthAndIndexDuplicatedVerticesIndexBuffer.VertexBufferSRV;
-			}
-
-			if (!GRecomputeTangentsParallelDispatch)
-			{
-				// When triangle & vertex passes are interleaved, resource transition is needed in between.
-				RHICmdList.Transition({
-					DispatchData.GetActiveTangentRWBuffer()->UpdateAccessState(ERHIAccess::SRVCompute),
-					StagingBuffer->UpdateAccessState(ERHIAccess::UAVCompute)
-				}, ERHITransitionCreateFlags::NoFence);
-			}
-
-			INC_DWORD_STAT_BY(STAT_GPUSkinCache_NumTrianglesForRecomputeTangents, NumTriangles);
-
-			const FRWBuffer& ShaderStagingBuffer = GRecomputeTangentsParallelDispatch ? DispatchData.GetIntermediateAccumulatedTangentBuffer()->Buffer : StagingBuffer->Buffer;
-
-			FRHIComputeShader* ShaderRHI = Shader.GetComputeShader();
-			SetComputePipelineState(RHICmdList, Shader.GetComputeShader());
-
-			SetShaderParametersLegacyCS(RHICmdList, Shader, Entry, DispatchData, ShaderStagingBuffer);
-			DispatchComputeShader(RHICmdList, Shader.GetShader(), ThreadGroupCountValue, 1, 1);
-			UnsetShaderParametersLegacyCS(RHICmdList, Shader);
-
-			IncrementDispatchCounter(RHICmdList);
-		}
+		DispatchData.DuplicatedIndices = LodData.RenderSections[SectionIndex].DuplicatedVerticesBuffer.DuplicatedVerticesIndexBuffer.VertexBufferSRV;
+		DispatchData.DuplicatedIndicesIndices = LodData.RenderSections[SectionIndex].DuplicatedVerticesBuffer.LengthAndIndexDuplicatedVerticesIndexBuffer.VertexBufferSRV;
 	}
-	else
-	{
-		SCOPED_DRAW_EVENTF(RHICmdList, SkinTangents_PerVertexPass, TEXT("%sTangentsVertex Mesh=%s, LOD=%d, Chunk=%d, InputStreamStart=%d, OutputStreamStart=%d, Vert=%d")
-			, RHI_BREADCRUMB_FORCE_STRING_LITERAL(Entry->Mode == EGPUSkinCacheEntryMode::RayTracing ? TEXT("[RT]") : TEXT(""))
-			, GetSkeletalMeshObjectDebugName(Entry->GPUSkin)
-			, LODIndex
-			, SectionIndex
-			, DispatchData.InputStreamStart
-			, DispatchData.OutputStreamStart
-			, DispatchData.NumVertices
-		);
-
-		//#todo-gpuskin Feature level?
-		auto* GlobalShaderMap = GetGlobalShaderMap(GetFeatureLevel());
-		TShaderMapRef<FRecomputeTangentsPerVertexPassCS<0>> ComputeShader0(GlobalShaderMap);
-		TShaderMapRef<FRecomputeTangentsPerVertexPassCS<1>> ComputeShader1(GlobalShaderMap);
-		TShaderRef<FBaseRecomputeTangentsPerVertexShader> ComputeShader;
-		if (DispatchData.Section->RecomputeTangentsVertexMaskChannel < ESkinVertexColorChannel::None)
-			ComputeShader = ComputeShader1;
-		else
-			ComputeShader = ComputeShader0;
 
-		uint32 VertexCount = DispatchData.NumVertices;
-		uint32 ThreadGroupCountValue = FMath::DivideAndRoundUp(VertexCount, ComputeShader->ThreadGroupSizeX);
+	if (!GRecomputeTangentsParallelDispatch)
+	{
+		// When triangle & vertex passes are interleaved, resource transition is needed in between.
+		RHICmdList.Transition({
+			DispatchData.GetActiveTangentRWBuffer()->UpdateAccessState(ERHIAccess::SRVCompute),
+			StagingBuffer->UpdateAccessState(ERHIAccess::UAVCompute)
+		}, ERHITransitionCreateFlags::NoFence);
+	}
 
-		if (!GRecomputeTangentsParallelDispatch)
-		{
-			// When triangle & vertex passes are interleaved, resource transition is needed in between.
-			RHICmdList.Transition({
-				DispatchData.GetTangentRWBuffer()->UpdateAccessState(ERHIAccess::UAVCompute),
-				StagingBuffer->UpdateAccessState(ERHIAccess::UAVCompute)
-			}, ERHITransitionCreateFlags::NoFence);
-		}
+	INC_DWORD_STAT_BY(STAT_GPUSkinCache_NumTrianglesForRecomputeTangents, NumTriangles);
 
-		SetComputePipelineState(RHICmdList, ComputeShader.GetComputeShader());
+	const FRWBuffer& ShaderStagingBuffer = GRecomputeTangentsParallelDispatch ? DispatchData.GetIntermediateAccumulatedTangentBuffer()->Buffer : StagingBuffer->Buffer;
 
-		SetShaderParametersLegacyCS(RHICmdList, ComputeShader, Entry, DispatchData, GRecomputeTangentsParallelDispatch ? DispatchData.GetIntermediateAccumulatedTangentBuffer()->Buffer : StagingBuffer->Buffer);
-		DispatchComputeShader(RHICmdList, ComputeShader.GetShader(), ThreadGroupCountValue, 1, 1);
-		UnsetShaderParametersLegacyCS(RHICmdList, ComputeShader);
+	FRHIComputeShader* ShaderRHI = Shader.GetComputeShader();
+	SetComputePipelineState(RHICmdList, Shader.GetComputeShader());
 
-		IncrementDispatchCounter(RHICmdList);
-	}
+	SetShaderParametersLegacyCS(RHICmdList, Shader, Entry, DispatchData, ShaderStagingBuffer);
+	DispatchComputeShader(RHICmdList, Shader.GetShader(), ThreadGroupCountValue, 1, 1);
+	UnsetShaderParametersLegacyCS(RHICmdList, Shader);
 }
 
 FGPUSkinCache::FRWBuffersAllocation* FGPUSkinCache::TryAllocBuffer(FRHICommandList& RHICmdList, EGPUSkinCacheBufferBits BufferBits, uint32 NumVertices, uint32 InIntermediateAccumulatedTangentsSize, const FName& OwnerName)
@@ -1279,7 +1265,7 @@ void FGPUSkinCache::TransitionAllToReadable(FRHICommandList& RHICmdList, TConstA
 				UAVs.Add(Buffer->UpdateAccessState(ReadState));
 			}
 		}
-		RHICmdList.Transition(UAVs, ERHIPipeline::Graphics, IsGPUSkinCacheRayTracingSupported() ? ERHIPipeline::All : ERHIPipeline::Graphics);
+		RHICmdList.Transition(UAVs, ERHIPipeline::Graphics, IsGPUSkinCacheRayTracingSupported() && !FRDGBuilder::IsImmediateMode() ? ERHIPipeline::All : ERHIPipeline::Graphics);
 	}
 }
 
@@ -1401,17 +1387,120 @@ void FGPUSkinCache::DispatchSetupForeground(FRHICommandList& RHICmdList, FDispat
 {
 	TRACE_CPUPROFILER_EVENT_SCOPE(GPUSkinCache_SetupDispatch);
 
-	Context.BuffersToTransition.Reserve(BatchDispatches.Num() * 2);
+	const int32 NumDispatches = BatchDispatches.Num();
+
+	Context.TangentDispatches.Reserve(NumDispatches);
+	Context.Transitions.FinalRead.Reserve(NumDispatches * NUM_BUFFERS);
+	Context.Transitions.Skinning.Write.Reserve(NumDispatches * NUM_BUFFERS);
+	Context.Transitions.RecomputeTangents.Read.Reserve(NumDispatches * 2);
 
-	for (FDispatchEntry& DispatchItem : BatchDispatches)
+	if (GRecomputeTangentsParallelDispatch)
 	{
-		PrepareUpdateSkinning(DispatchItem.SkinCacheEntry, DispatchItem.Section, Context);
+		Context.Transitions.RecomputeTangents.Write.Reserve(NumDispatches);
 	}
 
+	struct
+	{
+		int32 NumRayTracingDispatches = 0;
+		int32 NumRayTracingBuffers = 0;
+		int32 NumBuffers = 0;
+
+	} Stats;
+
+	for (const FDispatchEntry& DispatchItem : BatchDispatches)
+	{
+		FGPUSkinCacheEntry* Entry = DispatchItem.SkinCacheEntry;
+		Entry->bQueuedForDispatch = false;
+
+		FGPUSkinCacheEntry::FSectionDispatchData& DispatchData = Entry->DispatchData[DispatchItem.Section];
+		DispatchData.DispatchFlags = EGPUSkinCacheDispatchFlags::None;
+		DispatchData.PreviousPositionBuffer = nullptr;
+		DispatchData.RevisionNumber = 0;
+
+		if (DispatchData.PositionTracker.Allocation->HasPreviousBuffer())
+		{
+			const FGPUBaseSkinVertexFactory::FShaderDataType& ShaderData = DispatchData.SourceVertexFactory->GetShaderData();
+
+			const FVertexBufferAndSRV& BoneBuffer = ShaderData.GetBoneBufferForReading(false);
+			const uint32 CurrentRevision = ShaderData.GetRevisionNumber(false);
+			DispatchData.PositionBuffer = DispatchData.PositionTracker.Find(BoneBuffer, CurrentRevision);
+
+			const FVertexBufferAndSRV& PreviousBoneBuffer = ShaderData.GetBoneBufferForReading(true);
+			const uint32 PreviousRevision = ShaderData.GetRevisionNumber(true);
+			DispatchData.PreviousPositionBuffer = DispatchData.PositionTracker.Find(PreviousBoneBuffer, PreviousRevision);
+
+			// Allocate buffers if not found, excluding buffers already in use.  Or make the current buffer distinct if it happens to equal previous.
+			if (!DispatchData.PositionBuffer || DispatchData.PositionBuffer == DispatchData.PreviousPositionBuffer)
+			{
+				DispatchData.PositionBuffer = DispatchData.PositionTracker.AllocateUnused(BoneBuffer, CurrentRevision, DispatchData.PreviousPositionBuffer);
+				DispatchData.DispatchFlags |= EGPUSkinCacheDispatchFlags::Position;
+
+				Context.Transitions.Skinning.Write.Emplace(DispatchData.PositionBuffer);
+				Context.Transitions.FinalRead.Emplace(DispatchData.PositionBuffer);
+			}
+
+			if (!DispatchData.PreviousPositionBuffer)
+			{
+				DispatchData.PreviousPositionBuffer = DispatchData.PositionTracker.AllocateUnused(PreviousBoneBuffer, PreviousRevision, DispatchData.PositionBuffer);
+				DispatchData.DispatchFlags |= EGPUSkinCacheDispatchFlags::PositionPrevious;
+
+				Context.Transitions.Skinning.Write.Emplace(DispatchData.PreviousPositionBuffer);
+				Context.Transitions.FinalRead.Emplace(DispatchData.PreviousPositionBuffer);
+			}
+		}
+		else
+		{
+			DispatchData.PositionBuffer = &DispatchData.PositionTracker.Allocation->GetPositionBuffer();
+			DispatchData.DispatchFlags |= EGPUSkinCacheDispatchFlags::Position;
+
+			Context.Transitions.Skinning.Write.Emplace(DispatchData.PositionBuffer);
+			Context.Transitions.FinalRead.Emplace(DispatchData.PositionBuffer);
+		}
+
+		check(DispatchData.PreviousPositionBuffer != DispatchData.PositionBuffer);
+
+		DispatchData.IntermediateTangentBuffer = DispatchData.PositionTracker.GetIntermediateTangentBuffer();
+		DispatchData.IntermediateAccumulatedTangentBuffer = DispatchData.PositionTracker.GetIntermediateAccumulatedTangentBuffer();
+		DispatchData.TangentBuffer = DispatchData.PositionTracker.GetTangentBuffer();
+
+		if (EnumHasAnyFlags(DispatchData.DispatchFlags, EGPUSkinCacheDispatchFlags::Position))
+		{
+			if (DispatchData.IntermediateTangentBuffer)
+			{
+				DispatchData.DispatchFlags |= EGPUSkinCacheDispatchFlags::RecomputeTangents;
+
+				Context.Transitions.Skinning.Write.Emplace(DispatchData.IntermediateTangentBuffer);
+				Context.Transitions.RecomputeTangents.Read.Emplace(DispatchData.IntermediateTangentBuffer);
+				Context.Transitions.RecomputeTangents.Read.Emplace(DispatchData.PositionBuffer);
+
+				if (GRecomputeTangentsParallelDispatch)
+				{
+					Context.Transitions.RecomputeTangents.Write.Add(DispatchData.GetIntermediateAccumulatedTangentBuffer());
+				}
+
+				Context.TangentDispatches.Emplace(DispatchItem);
+			}
+
+			Context.Transitions.Skinning.Write.Emplace(DispatchData.TangentBuffer);
+			Context.Transitions.FinalRead.Emplace(DispatchData.TangentBuffer);
+		}
+
+		int32 NumBuffers = 0;
+		NumBuffers += DispatchData.PositionBuffer ? 1 : 0;
+		NumBuffers += DispatchData.PreviousPositionBuffer ? 1 : 0;
+		NumBuffers += DispatchData.TangentBuffer ? 1 : 0;
+		NumBuffers += DispatchData.IntermediateTangentBuffer ? 1 : 0;
+		NumBuffers += DispatchData.IntermediateAccumulatedTangentBuffer ? 1 : 0;
+
+		Stats.NumBuffers += NumBuffers;
+		Stats.NumRayTracingBuffers    += Entry->Mode == EGPUSkinCacheEntryMode::RayTracing ? NumBuffers : 0;
+		Stats.NumRayTracingDispatches += Entry->Mode == EGPUSkinCacheEntryMode::RayTracing ? 1 : 0;
+	}
+
+	SET_DWORD_STAT(STAT_GPUSkinCache_NumBuffers, Stats.NumBuffers);
+	SET_DWORD_STAT(STAT_GPUSkinCache_NumRayTracingBuffers, Stats.NumRayTracingBuffers);
+	SET_DWORD_STAT(STAT_GPUSkinCache_NumRayTracingDispatches, Stats.NumRayTracingDispatches);
 	SET_DWORD_STAT(STAT_GPUSkinCache_NumDispatches, BatchDispatches.Num());
-	SET_DWORD_STAT(STAT_GPUSkinCache_NumRayTracingDispatches, Context.Stats.NumRayTracingDispatches);
-	SET_DWORD_STAT(STAT_GPUSkinCache_NumBuffers, Context.Stats.NumBuffers);
-	SET_DWORD_STAT(STAT_GPUSkinCache_NumRayTracingBuffers, Context.Stats.NumRayTracingBuffers);
 
 #if RHI_RAYTRACING
 	if (IsGPUSkinCacheRayTracingSupported())
@@ -1482,7 +1571,7 @@ void FGPUSkinCache::DispatchSetupBackground(FRHICommandList& RHICmdList, FDispat
 		});
 	}
 
-	TransitionBufferUAVs(RHICmdList, Context.BuffersToTransition, Context.SkinningBuffersToOverlap);
+	TransitionBufferUAVs(RHICmdList, Context.Transitions.Skinning.Write, Context.Transitions.Skinning.Overlap);
 
 	{
 		TRACE_CPUPROFILER_EVENT_SCOPE(UpdateVertexFactoryDeclarations);
@@ -1509,8 +1598,7 @@ void FGPUSkinCache::DispatchPassExecute(FRHICommandList& RHICmdList, FDispatchSe
 	RenderCaptureInterface::FScopedCapture RenderCapture(bCapture, &RHICmdList);
 	GNumDispatchesToCapture -= bCapture ? 1 : 0;
 
-	TArray<FSkinCacheRWBuffer*, FConcurrentLinearArrayAllocator> BuffersToTransitionToRead;
-	RHICmdList.BeginUAVOverlap(Context.SkinningBuffersToOverlap);
+	RHICmdList.BeginUAVOverlap(Context.Transitions.Skinning.Overlap);
 
 	{
 		TRACE_CPUPROFILER_EVENT_SCOPE(GPUSkinCache_UpdateSkinningBatches);
@@ -1527,9 +1615,8 @@ void FGPUSkinCache::DispatchPassExecute(FRHICommandList& RHICmdList, FDispatchSe
 #undef SKIN_CACHE_SHADER
 
 		int32 SortedCount = Context.SortedDispatches.Num();
-		for (int32 i = 0; i < SortedCount; ++i)
+		for (const FSortedDispatchEntry& SortedEntry : Context.SortedDispatches)
 		{
-			const FSortedDispatchEntry& SortedEntry = Context.SortedDispatches[i];
 			if (SortedEntry.ShaderIndex != LastShaderIndex)
 			{
 				TRACE_CPUPROFILER_EVENT_SCOPE(ChangeShader);
@@ -1592,8 +1679,6 @@ void FGPUSkinCache::DispatchPassExecute(FRHICommandList& RHICmdList, FDispatchSe
 
 				INC_DWORD_STAT_BY(STAT_GPUSkinCache_TotalNumVertices, VertexCountAlign64 * 64);
 				RHICmdList.DispatchComputeShader(VertexCountAlign64, 1, 1);
-
-				BuffersToTransitionToRead.Add(DispatchData.GetPreviousPositionRWBuffer());
 			}
 
 			if (EnumHasAnyFlags(DispatchData.DispatchFlags, EGPUSkinCacheDispatchFlags::Position))
@@ -1612,18 +1697,8 @@ void FGPUSkinCache::DispatchPassExecute(FRHICommandList& RHICmdList, FDispatchSe
 
 				INC_DWORD_STAT_BY(STAT_GPUSkinCache_TotalNumVertices, VertexCountAlign64 * 64);
 				RHICmdList.DispatchComputeShader(VertexCountAlign64, 1, 1);
-
-				BuffersToTransitionToRead.Add(DispatchData.GetPositionRWBuffer());
 			}
 
-			if (IncrementDispatchCounter(RHICmdList))
-			{
-				// The command list was submitted so rebind the shader on the next loop.
-				LastShaderIndex = -1;
-			}
-
-			BuffersToTransitionToRead.Add(DispatchData.GetTangentRWBuffer());
-
 			check(DispatchData.PreviousPositionBuffer != DispatchData.PositionBuffer);
 		}
 
@@ -1633,80 +1708,49 @@ void FGPUSkinCache::DispatchPassExecute(FRHICommandList& RHICmdList, FDispatchSe
 		}
 	}
 
-	RHICmdList.EndUAVOverlap(Context.SkinningBuffersToOverlap);
-
-	TArray<FSkinCacheRWBuffer*> BuffersToSRVForRecomputeTangents;
-	TArray<FSkinCacheRWBuffer*> IntermediateAccumulatedTangentBuffers;
+	RHICmdList.EndUAVOverlap(Context.Transitions.Skinning.Overlap);
 
-	for (FDispatchEntry& DispatchItem : BatchDispatches)
-	{
-		FGPUSkinCacheEntry::FSectionDispatchData& DispatchData = DispatchItem.SkinCacheEntry->DispatchData[DispatchItem.Section];
-
-		if (EnumHasAnyFlags(DispatchData.DispatchFlags, EGPUSkinCacheDispatchFlags::RecomputeTangents))
-		{
-			BuffersToSRVForRecomputeTangents.Add(DispatchData.GetPositionRWBuffer());
-			BuffersToSRVForRecomputeTangents.Add(DispatchData.GetActiveTangentRWBuffer());
-			BuffersToTransitionToRead.Add(DispatchData.GetPositionRWBuffer());
-
-			if (GRecomputeTangentsParallelDispatch)
-			{
-				IntermediateAccumulatedTangentBuffers.Add(DispatchData.GetIntermediateAccumulatedTangentBuffer());
-			}
-		}
-	}
-
-	if (!BuffersToSRVForRecomputeTangents.IsEmpty())
+	if (!Context.Transitions.RecomputeTangents.Read.IsEmpty())
 	{
 		TArray<FRHIUnorderedAccessView*> IntermediateAccumulatedTangentBuffersToOverlap;
-		TransitionBuffers(RHICmdList, BuffersToSRVForRecomputeTangents, ERHIAccess::SRVCompute);
-		TransitionBufferUAVs(RHICmdList, IntermediateAccumulatedTangentBuffers, IntermediateAccumulatedTangentBuffersToOverlap);
+		TransitionBuffers(RHICmdList, Context.Transitions.RecomputeTangents.Read, ERHIAccess::SRVCompute);
+		TransitionBufferUAVs(RHICmdList, Context.Transitions.RecomputeTangents.Write, IntermediateAccumulatedTangentBuffersToOverlap);
 		RHICmdList.BeginUAVOverlap(IntermediateAccumulatedTangentBuffersToOverlap);
 
 		{
 			SCOPED_DRAW_EVENT(RHICmdList, GPUSkinCache_RecomputeTangentsBatches);
 			FSkinCacheRWBuffer* StagingBuffer = nullptr;
+			TArray<FSkinCacheRWBuffer*> TangentBuffers;
 
-			for (FDispatchEntry& DispatchItem : BatchDispatches)
+			if (GRecomputeTangentsParallelDispatch)
 			{
-				FGPUSkinCacheEntry::FSectionDispatchData& DispatchData = DispatchItem.SkinCacheEntry->DispatchData[DispatchItem.Section];
-				if (EnumHasAnyFlags(DispatchData.DispatchFlags, EGPUSkinCacheDispatchFlags::RecomputeTangents))
-				{
-					DispatchUpdateSkinTangents(RHICmdList, DispatchItem.SkinCacheEntry, DispatchItem.Section, StagingBuffer, true);
-
-					if (!GRecomputeTangentsParallelDispatch)
-					{
-						// When parallel dispatching is off, triangle pass and vertex pass are dispatched interleaved.
-						DispatchUpdateSkinTangents(RHICmdList, DispatchItem.SkinCacheEntry, DispatchItem.Section, StagingBuffer, false);
-					}
-				}
+				TangentBuffers.Reserve(Context.TangentDispatches.Num());
 			}
 
-			if (GRecomputeTangentsParallelDispatch)
+			for (const FDispatchEntry& DispatchItem : Context.TangentDispatches)
 			{
-				TArray<FSkinCacheRWBuffer*> TangentBuffers;
+				DispatchUpdateSkinTangentsTrianglePass(RHICmdList, DispatchItem.SkinCacheEntry, DispatchItem.Section, StagingBuffer);
 
-				for (FDispatchEntry& DispatchItem : BatchDispatches)
+				if (GRecomputeTangentsParallelDispatch)
 				{
-					FGPUSkinCacheEntry::FSectionDispatchData& DispatchData = DispatchItem.SkinCacheEntry->DispatchData[DispatchItem.Section];
-					if (EnumHasAnyFlags(DispatchData.DispatchFlags, EGPUSkinCacheDispatchFlags::RecomputeTangents))
-					{
-						TangentBuffers.Add(DispatchData.GetTangentRWBuffer());
-						BuffersToTransitionToRead.Add(DispatchData.GetTangentRWBuffer());
-					}
+					TangentBuffers.Add(DispatchItem.SkinCacheEntry->DispatchData[DispatchItem.Section].GetTangentRWBuffer());
 				}
+				else
+				{
+					DispatchUpdateSkinTangentsVertexPass(RHICmdList, DispatchItem.SkinCacheEntry, DispatchItem.Section, StagingBuffer);
+				}
+			}
 
+			if (GRecomputeTangentsParallelDispatch)
+			{
 				TArray<FRHIUnorderedAccessView*> TangentBuffersToOverlap;
-				TransitionBuffers(RHICmdList, IntermediateAccumulatedTangentBuffers, ERHIAccess::UAVCompute);
+				TransitionBuffers(RHICmdList, Context.Transitions.RecomputeTangents.Write, ERHIAccess::UAVCompute);
 				TransitionBufferUAVs(RHICmdList, TangentBuffers, TangentBuffersToOverlap);
 				RHICmdList.BeginUAVOverlap(TangentBuffersToOverlap);
 
-				for (FDispatchEntry& DispatchItem : BatchDispatches)
+				for (const FDispatchEntry& DispatchItem : Context.TangentDispatches)
 				{
-					FGPUSkinCacheEntry::FSectionDispatchData& DispatchData = DispatchItem.SkinCacheEntry->DispatchData[DispatchItem.Section];
-					if (EnumHasAnyFlags(DispatchData.DispatchFlags, EGPUSkinCacheDispatchFlags::RecomputeTangents))
-					{
-						DispatchUpdateSkinTangents(RHICmdList, DispatchItem.SkinCacheEntry, DispatchItem.Section, StagingBuffer, false);
-					}
+					DispatchUpdateSkinTangentsVertexPass(RHICmdList, DispatchItem.SkinCacheEntry, DispatchItem.Section, StagingBuffer);
 				}
 
 				RHICmdList.EndUAVOverlap(TangentBuffersToOverlap);
@@ -1716,11 +1760,7 @@ void FGPUSkinCache::DispatchPassExecute(FRHICommandList& RHICmdList, FDispatchSe
 		RHICmdList.EndUAVOverlap(IntermediateAccumulatedTangentBuffersToOverlap);
 	}
 
-	{
-		TRACE_CPUPROFILER_EVENT_SCOPE(TransitionAllToReadable);
-		TransitionAllToReadable(RHICmdList, BuffersToTransitionToRead);
-	}
-
+	TransitionAllToReadable(RHICmdList, Context.Transitions.FinalRead);
 	BatchDispatches.Reset();
 }
 
@@ -2079,76 +2119,6 @@ void FGPUSkinCache::Release(FGPUSkinCacheEntry*& SkinCacheEntry)
 	}
 }
 
-void FGPUSkinCache::PrepareUpdateSkinning(FGPUSkinCacheEntry* Entry, int32 Section, FDispatchSetupContext& Context)
-{
-	FGPUSkinCacheEntry::FSectionDispatchData& DispatchData = Entry->DispatchData[Section];
-	FGPUBaseSkinVertexFactory::FShaderDataType& ShaderData = DispatchData.SourceVertexFactory->GetShaderData();
-
-	const FVertexBufferAndSRV& BoneBuffer = ShaderData.GetBoneBufferForReading(false);
-	const FVertexBufferAndSRV& PrevBoneBuffer = ShaderData.GetBoneBufferForReading(true);
-
-	uint32 CurrentRevision = ShaderData.GetRevisionNumber(false);
-	uint32 PreviousRevision = ShaderData.GetRevisionNumber(true);
-
-	DispatchData.DispatchFlags = EGPUSkinCacheDispatchFlags::None;
-
-	if (DispatchData.PositionTracker.Allocation->HasPreviousBuffer())
-	{
-		// Find existing buffers
-		DispatchData.PositionBuffer = DispatchData.PositionTracker.Find(BoneBuffer, CurrentRevision);
-		DispatchData.PreviousPositionBuffer = DispatchData.PositionTracker.Find(PrevBoneBuffer, PreviousRevision);
-
-		// Allocate buffers if not found, excluding buffers already in use.  Or make the current buffer distinct if it happens to equal previous.
-		if (!DispatchData.PositionBuffer || DispatchData.PositionBuffer == DispatchData.PreviousPositionBuffer)
-		{
-			DispatchData.PositionBuffer = DispatchData.PositionTracker.AllocateUnused(BoneBuffer, CurrentRevision, DispatchData.PreviousPositionBuffer, &Context.BuffersToTransition);
-			DispatchData.DispatchFlags |= EGPUSkinCacheDispatchFlags::Position;
-		}
-		if (!DispatchData.PreviousPositionBuffer)
-		{
-			DispatchData.PreviousPositionBuffer = DispatchData.PositionTracker.AllocateUnused(PrevBoneBuffer, PreviousRevision, DispatchData.PositionBuffer, &Context.BuffersToTransition);
-			DispatchData.DispatchFlags |= EGPUSkinCacheDispatchFlags::PositionPrevious;
-		}
-	}
-	else
-	{
-		DispatchData.PositionBuffer = &DispatchData.PositionTracker.Allocation->GetPositionBuffer();
-		DispatchData.DispatchFlags |= EGPUSkinCacheDispatchFlags::Position;
-		Context.BuffersToTransition.Add(DispatchData.PositionBuffer);
-	}
-
-	DispatchData.TangentBuffer = DispatchData.PositionTracker.GetTangentBuffer();
-	DispatchData.IntermediateTangentBuffer = DispatchData.PositionTracker.GetIntermediateTangentBuffer();
-	DispatchData.IntermediateAccumulatedTangentBuffer = DispatchData.PositionTracker.GetIntermediateAccumulatedTangentBuffer();
-
-	if (DispatchData.IntermediateTangentBuffer)
-	{
-		DispatchData.DispatchFlags |= EGPUSkinCacheDispatchFlags::RecomputeTangents;
-	}
-
-	int32 NumBuffers = 0;
-	NumBuffers += DispatchData.PositionBuffer ? 1 : 0;
-	NumBuffers += DispatchData.PreviousPositionBuffer ? 1 : 0;
-	NumBuffers += DispatchData.TangentBuffer ? 1 : 0;
-	NumBuffers += DispatchData.IntermediateTangentBuffer ? 1 : 0;
-	NumBuffers += DispatchData.IntermediateAccumulatedTangentBuffer ? 1 : 0;
-
-	Context.Stats.NumBuffers += NumBuffers;
-	Context.Stats.NumRayTracingBuffers += Entry->Mode == EGPUSkinCacheEntryMode::RayTracing ? NumBuffers : 0;
-	Context.Stats.NumRayTracingDispatches += Entry->Mode == EGPUSkinCacheEntryMode::RayTracing ? 1 : 0;
-
-	if (DispatchData.DispatchFlags != EGPUSkinCacheDispatchFlags::None && DispatchData.GetActiveTangentRWBuffer())
-	{
-		Context.BuffersToTransition.Add(DispatchData.GetActiveTangentRWBuffer());
-	}
-
-	check(DispatchData.PreviousPositionBuffer != DispatchData.PositionBuffer);
-
-	// Clear the flag that this is queued for dispatch.
-	Entry->bQueuedForDispatch = false;
-	DispatchData.RevisionNumber = 0;
-}
-
 void FGPUSkinCache::FRWBuffersAllocation::RemoveAllFromTransitionArray(TSet<FSkinCacheRWBuffer*>& InBuffersToTransition)
 {
 	for (uint32 i = 0; i < NUM_BUFFERS; i++)
@@ -2360,22 +2330,6 @@ void FGPUSkinCache::CVarSinkFunction()
 
 FAutoConsoleVariableSink FGPUSkinCache::CVarSink(FConsoleCommandDelegate::CreateStatic(&CVarSinkFunction));
 
-bool FGPUSkinCache::IncrementDispatchCounter(FRHICommandList& RHICmdList)
-{
-	if (GSkinCacheMaxDispatchesPerCmdList > 0)
-	{
-		DispatchCounter++;
-		if (DispatchCounter >= GSkinCacheMaxDispatchesPerCmdList)
-		{
-			//UE_LOG(LogSkinCache, Log, TEXT("SubmitCommandsHint issued after %d dispatches"), DispatchCounter);
-			RHICmdList.SubmitCommandsHint();
-			DispatchCounter = 0;
-			return true;
-		}
-	}
-	return false;
-}
-
 uint64 FGPUSkinCache::GetExtraRequiredMemoryAndReset()
 {
 	if (GSkinCachePrintMemorySummary == 2 || (GSkinCachePrintMemorySummary == 1 && ExtraRequiredMemory > 0))
diff --git a/Engine/Source/Runtime/Engine/Public/GPUSkinCache.h b/Engine/Source/Runtime/Engine/Public/GPUSkinCache.h
--- a/Engine/Source/Runtime/Engine/Public/GPUSkinCache.h
+++ b/Engine/Source/Runtime/Engine/Public/GPUSkinCache.h
@@ -104,8 +104,6 @@
 class FGPUSkinCache
 {
 public:
-	struct FRWBufferTracker;
-
 	enum ESkinCacheInitSettings
 	{
 		// max 256 bones as we use a byte to index
@@ -118,18 +116,6 @@
 		IntermediateAccumBufferNumInts = 8,
 	};
 
-	struct FDispatchEntry
-	{
-		FGPUSkinCacheEntry* SkinCacheEntry = nullptr;
-		uint32 Section = 0;
-	};
-
-	struct FSortedDispatchEntry
-	{
-		int32 ShaderIndex;
-		int32 BatchIndex;
-	};
-
 	FGPUSkinCache() = delete;
 	ENGINE_API FGPUSkinCache(ERHIFeatureLevel::Type InFeatureLevel, bool bInRequiresMemoryLimit, UWorld* InWorld);
 	ENGINE_API ~FGPUSkinCache();
@@ -170,6 +156,46 @@
 
 	static bool IsGPUSkinCacheRayTracingSupported();
 
+	FGPUSkinCacheEntry const* GetSkinCacheEntry(uint32 ComponentId) const;
+	static FRWBuffer* GetPositionBuffer(FRDGBuilder& GraphBuilder, FGPUSkinCacheEntry const* Entry, uint32 SectionIndex);
+	static FRWBuffer* GetPreviousPositionBuffer(FRDGBuilder& GraphBuilder, FGPUSkinCacheEntry const* Entry, uint32 SectionIndex);
+	static FRWBuffer* GetTangentBuffer(FRDGBuilder& GraphBuilder, FGPUSkinCacheEntry const* Entry, uint32 SectionIndex);
+
+	UE_DEPRECATED(5.6, "GetUpdatedFrame is no longer used")
+	static uint32 GetUpdatedFrame(FGPUSkinCacheEntry const*, uint32) { return 0; }
+
+#if RHI_RAYTRACING
+	void ProcessRayTracingGeometryToUpdate(FRHICommandList& RHICmdList, FGPUSkinCacheEntry* SkinCacheEntry);
+#endif // RHI_RAYTRACING
+
+	ENGINE_API void Dispatch(FRHICommandList& RHICmdList);
+
+	ENGINE_API UE::Tasks::FTask BeginDispatch(FRDGBuilder& GraphBuilder, const UE::Tasks::FTask& PrerequisitesTask);
+	ENGINE_API void EndDispatch(FRDGBuilder& GraphBuilder);
+
+	inline ERHIFeatureLevel::Type GetFeatureLevel() const { return FeatureLevel; }
+
+	inline bool HasWork() const { return !BatchDispatches.IsEmpty(); }
+
+	UE_DEPRECATED(5.6, "Use Dispatch instead.")
+	void DoDispatch(FRHICommandList& RHICmdList) { Dispatch(RHICmdList); }
+
+private:
+	friend FGPUSkinCacheEntry;
+	struct FRWBufferTracker;
+
+	struct FDispatchEntry
+	{
+		FGPUSkinCacheEntry* SkinCacheEntry = nullptr;
+		uint32 Section = 0;
+	};
+
+	struct FSortedDispatchEntry
+	{
+		int32 ShaderIndex;
+		int32 BatchIndex;
+	};
+
 	enum
 	{
 		NUM_BUFFERS = 2,
@@ -205,15 +231,13 @@
 
 	struct FRWBuffersAllocation
 	{
-		friend struct FRWBufferTracker;
+		friend FRWBufferTracker;
 
 		static uint64 CalculateRequiredMemory(EGPUSkinCacheBufferBits BufferBits, uint32 NumVertices, uint32 IntermediateAccumulatedTangentsSize);
 
 		FRWBuffersAllocation(FRHICommandList& RHICmdList, EGPUSkinCacheBufferBits BufferBits, uint32 NumVertices, uint32 IntermediateAccumulatedTangentsSize, const FName& OwnerName);
 		~FRWBuffersAllocation();
 
-		void AllocateTangents();
-
 		uint64 GetNumBytes() const
 		{
 			return CalculateRequiredMemory(BufferBits, NumVertices, IntermediateAccumulatedTangentsSize);
@@ -305,19 +329,14 @@
 			return Allocation ? Allocation->GetIntermediateAccumulatedTangentBuffer() : nullptr;
 		}
 
-		// Allocates an element that's not the "Used" element passed in (or if Used is NULL, allocates any element).  Optionally adds element to "AllocatedItems" if non-null.
-		FSkinCacheRWBuffer* AllocateUnused(const FVertexBufferAndSRV& BoneBuffer, uint32 Revision, const FSkinCacheRWBuffer* Used, TArray<FSkinCacheRWBuffer*>* AllocatedItems)
+		// Allocates an element that's not the "Used" element passed in (or if Used is NULL, allocates any element).
+		FSkinCacheRWBuffer* AllocateUnused(const FVertexBufferAndSRV& BoneBuffer, uint32 Revision, const FSkinCacheRWBuffer* Used)
 		{
 			int32 UnusedIndex = Used == &Allocation->PositionBuffers[0] ? 1 : 0;
 			Revisions[UnusedIndex] = Revision;
 			BoneBuffers[UnusedIndex] = &BoneBuffer;
 
-			FSkinCacheRWBuffer* Result = &Allocation->PositionBuffers[UnusedIndex];
-			if (AllocatedItems)
-			{
-				AllocatedItems->Add(Result);
-			}
-			return Result;
+			return &Allocation->PositionBuffers[UnusedIndex];
 		}
 
 		// On recreate of the render state where the GPU skin cache entry is preserved, the bone buffer will have been reallocated,
@@ -340,46 +359,32 @@
 		const FVertexBufferAndSRV* BoneBuffers[NUM_BUFFERS];
 	};
 
-	FGPUSkinCacheEntry const* GetSkinCacheEntry(uint32 ComponentId) const;
-	static FRWBuffer* GetPositionBuffer(FRDGBuilder& GraphBuilder, FGPUSkinCacheEntry const* Entry, uint32 SectionIndex);
-	static FRWBuffer* GetPreviousPositionBuffer(FRDGBuilder& GraphBuilder, FGPUSkinCacheEntry const* Entry, uint32 SectionIndex);
-	static FRWBuffer* GetTangentBuffer(FRDGBuilder& GraphBuilder, FGPUSkinCacheEntry const* Entry, uint32 SectionIndex);
-
-	UE_DEPRECATED(5.6, "GetUpdatedFrame is no longer used")
-	static uint32 GetUpdatedFrame(FGPUSkinCacheEntry const*, uint32) { return 0; }
-
-#if RHI_RAYTRACING
-	void ProcessRayTracingGeometryToUpdate(FRHICommandList& RHICmdList, FGPUSkinCacheEntry* SkinCacheEntry);
-#endif // RHI_RAYTRACING
-
-	ENGINE_API void Dispatch(FRHICommandList& RHICmdList);
-
-	ENGINE_API UE::Tasks::FTask BeginDispatch(FRDGBuilder& GraphBuilder, const UE::Tasks::FTask& PrerequisitesTask);
-	ENGINE_API void EndDispatch(FRDGBuilder& GraphBuilder);
-
-	inline ERHIFeatureLevel::Type GetFeatureLevel() const { return FeatureLevel; }
-
-	inline bool HasWork() const { return !BatchDispatches.IsEmpty(); }
-
-	UE_DEPRECATED(5.6, "Use Dispatch instead.")
-	void DoDispatch(FRHICommandList& RHICmdList) { Dispatch(RHICmdList); }
-
-protected:
 	struct FDispatchSetupContext
 	{
-		TArray<FSkinCacheRWBuffer*> BuffersToTransition;
+		TArray<FDispatchEntry> TangentDispatches;
 		TArray<FSortedDispatchEntry> SortedDispatches;
-		TArray<FRHIUnorderedAccessView*> SkinningBuffersToOverlap;
 
 		FRHICommandList* RHICmdListDispatch = nullptr;
 
 		struct
 		{
-			int32 NumRayTracingDispatches = 0;
-			int32 NumRayTracingBuffers = 0;
-			int32 NumBuffers = 0;
+			TArray<FSkinCacheRWBuffer*> FinalRead;
 
-		} Stats;
+			struct
+			{
+				TArray<FSkinCacheRWBuffer*> Write;
+				TArray<FRHIUnorderedAccessView*> Overlap;
+
+			} Skinning;
+
+			struct
+			{
+				TArray<FSkinCacheRWBuffer*> Write;
+				TArray<FSkinCacheRWBuffer*> Read;
+			
+			} RecomputeTangents;
+
+		} Transitions;
 	};
 
 	void DispatchSetupForeground(FRHICommandList& RHICmdList, FDispatchSetupContext& SetupContext);
@@ -396,13 +401,13 @@
 	TArray<FDispatchEntry> BatchDispatches;
 
 	FRWBuffersAllocation* TryAllocBuffer(FRHICommandList& RHICmdList, EGPUSkinCacheBufferBits BufferBits, uint32 NumVertices, uint32 InIntermediateAccumulatedTangentsSize, const FName& OwnerName);
-	void DispatchUpdateSkinTangents(FRHICommandList& RHICmdList, FGPUSkinCacheEntry* Entry, int32 SectionIndex, FSkinCacheRWBuffer*& StagingBuffer, bool bTrianglePass);
+	void DispatchUpdateSkinTangentsVertexPass(FRHICommandList& RHICmdList, FGPUSkinCacheEntry* Entry, int32 SectionIndex, FSkinCacheRWBuffer*& StagingBuffer);
+	void DispatchUpdateSkinTangentsTrianglePass(FRHICommandList& RHICmdList, FGPUSkinCacheEntry* Entry, int32 SectionIndex, FSkinCacheRWBuffer*& StagingBuffer);
 
-	void PrepareUpdateSkinning(FGPUSkinCacheEntry* Entry, int32 Section, FDispatchSetupContext& Context);
-
 	void Cleanup();
 	static void ReleaseSkinCacheEntry(FGPUSkinCacheEntry* SkinCacheEntry);
 	void InvalidateAllEntries();
+
 	uint64 UsedMemoryInBytes;
 	uint64 ExtraRequiredMemory;
 	int32 FlushCounter;
@@ -421,9 +426,6 @@
 	static void CVarSinkFunction();
 	static FAutoConsoleVariableSink CVarSink;
 
-	bool IncrementDispatchCounter(FRHICommandList& RHICmdList);
-	int32 DispatchCounter = 0;
-
 	uint32 GetNextTransitionFence() { return ++TransitionFence; }
 	uint32 TransitionFence = 0;
 
