Moved UMaterialExpressionPinBase::CanUserDeleteExpression() in the same location in code as Main.

#rnx

diff --git a/Engine/Source/Runtime/Engine/Private/Materials/MaterialExpressions.cpp b/Engine/Source/Runtime/Engine/Private/Materials/MaterialExpressions.cpp
--- a/Engine/Source/Runtime/Engine/Private/Materials/MaterialExpressions.cpp
+++ b/Engine/Source/Runtime/Engine/Private/Materials/MaterialExpressions.cpp
@@ -11742,6 +11742,11 @@ bool UMaterialExpressionPinBase::CanUserDeleteExpression() const
 }
 
 
+{
+	return false;
+}
+
+
 int32 UMaterialExpressionFresnel::Compile(class FMaterialCompiler* Compiler, int32 OutputIndex)
 {
 	// pow(1 - max(0,Normal dot Camera),Exponent) * (1 - BaseReflectFraction) + BaseReflectFraction
