Fixed non-deterministic cook due to FName numbers being unstable for instanced static mesh MID duplicates. The shared/redundant MID map could get wiped by a GC between cooking levels, resulting in recreation of MIDs at inconsistent intervals. Here we'll regenerate the numbers from scratch once we have the proper outer that they'll be duplicated under.

diff --git a/Engine/Source/Runtime/Engine/Private/ISMPartition/ISMComponentDescriptor.cpp b/Engine/Source/Runtime/Engine/Private/ISMPartition/ISMComponentDescriptor.cpp
--- a/Engine/Source/Runtime/Engine/Private/ISMPartition/ISMComponentDescriptor.cpp
+++ b/Engine/Source/Runtime/Engine/Private/ISMPartition/ISMComponentDescriptor.cpp
@@ -206,6 +206,14 @@ void FISMComponentDescriptorBase::PostLoadFixup(UObject* Loader)
 	BodyInstance.FixupData(Loader);
 }
 
+UMaterialInterface* FISMComponentDescriptorBase::DuplicateMaterialInterface(const UMaterialInterface& InMaterialInterface, UObject& NewOuter)
+{
+	// Regenerate a new unique number since the template MIDs don't have stable numbers during cook
+	const FName BaseName = FName(InMaterialInterface.GetFName(), 0);
+	const FName NewName = MakeUniqueObjectName(&NewOuter, UMaterialInterface::StaticClass(), BaseName);
+	return DuplicateObject<UMaterialInterface>(&InMaterialInterface, &NewOuter, NewName);
+}
+
 bool FISMComponentDescriptorBase::operator!=(const FISMComponentDescriptorBase& Other) const
 {
 	return !(*this == Other);
@@ -441,14 +449,14 @@ void FISMComponentDescriptor::InitComponent(UInstancedStaticMeshComponent* ISMCo
 			{
 				// As override materials are normally outered to their owner component, we need to duplicate them here to make sure we don't create
 				// references to actors in other levels (for packed level instances or HLOD actors).
-				MaterialInterface = DuplicateObject<UMaterialInterface>(MaterialInterface, ISMComponent);
+				MaterialInterface = DuplicateMaterialInterface(*MaterialInterface, *ISMComponent);
 
 				// If the MID we just duplicated has a nanite override that's also not an asset, duplicate that too
 				UMaterialInstanceDynamic* OverrideMID = Cast<UMaterialInstanceDynamic>(MaterialInterface);
 				UMaterialInterface* NaniteOverride = OverrideMID ? OverrideMID->GetNaniteOverride() : nullptr;
 				if (NaniteOverride && !NaniteOverride->IsAsset())
 				{
-					OverrideMID->SetNaniteOverride(DuplicateObject<UMaterialInterface>(NaniteOverride, ISMComponent));
+					OverrideMID->SetNaniteOverride(DuplicateMaterialInterface(*NaniteOverride, *ISMComponent));
 				}
 			}
 		}
@@ -485,14 +493,14 @@ void FSoftISMComponentDescriptor::InitComponent(UInstancedStaticMeshComponent* I
 			{
 				// As override materials are normally outered to their owner component, we need to duplicate them here to make sure we don't create
 				// references to actors in other levels (for packed level instances or HLOD actors).
-				MaterialInterface = DuplicateObject<UMaterialInterface>(MaterialInterface, ISMComponent);
+				MaterialInterface = DuplicateMaterialInterface(*MaterialInterface, *ISMComponent);
 
 				// If the MID we just duplicated has a nanite override that's also not an asset, duplicate that too
 				UMaterialInstanceDynamic* OverrideMID = Cast<UMaterialInstanceDynamic>(MaterialInterface);
 				UMaterialInterface* NaniteOverride = OverrideMID ? OverrideMID->GetNaniteOverride() : nullptr; 
 				if (NaniteOverride && !NaniteOverride->IsAsset())
 				{
-					OverrideMID->SetNaniteOverride(DuplicateObject<UMaterialInterface>(NaniteOverride, ISMComponent));
+					OverrideMID->SetNaniteOverride(DuplicateMaterialInterface(*NaniteOverride, *ISMComponent));
 				}
 			}
 		}
diff --git a/Engine/Source/Runtime/Engine/Public/ISMPartition/ISMComponentDescriptor.h b/Engine/Source/Runtime/Engine/Public/ISMPartition/ISMComponentDescriptor.h
--- a/Engine/Source/Runtime/Engine/Public/ISMPartition/ISMComponentDescriptor.h
+++ b/Engine/Source/Runtime/Engine/Public/ISMPartition/ISMComponentDescriptor.h
@@ -35,6 +35,8 @@ struct FISMComponentDescriptorBase
 	ENGINE_API bool operator!=(const FISMComponentDescriptorBase& Other) const;
 	ENGINE_API bool operator==(const FISMComponentDescriptorBase& Other) const;
 
+	static UMaterialInterface* DuplicateMaterialInterface(const UMaterialInterface& InMaterialInterface, UObject& NewOuter);
+
 	friend inline uint32 GetTypeHash(const FISMComponentDescriptorBase& Key)
 	{
 		return Key.GetTypeHash();
