PR #12912: Fix missing externals when inheriting RenderTarget

diff --git a/Engine/Source/Runtime/Engine/Classes/Engine/CanvasRenderTarget2D.h b/Engine/Source/Runtime/Engine/Classes/Engine/CanvasRenderTarget2D.h
--- a/Engine/Source/Runtime/Engine/Classes/Engine/CanvasRenderTarget2D.h
+++ b/Engine/Source/Runtime/Engine/Classes/Engine/CanvasRenderTarget2D.h
@@ -99,8 +99,8 @@ public:
 	ENGINE_API void SetSampleCount(ETextureRenderTargetSampleCount InSampleCount);
 
 #if WITH_EDITOR
-	virtual bool CanEditChange(const FProperty* InProperty) const override;
-	virtual void PostEditChangeProperty(FPropertyChangedEvent& PropertyChangedEvent) override;
+	ENGINE_API virtual bool CanEditChange(const FProperty* InProperty) const override;
+	ENGINE_API virtual void PostEditChangeProperty(FPropertyChangedEvent& PropertyChangedEvent) override;
 #endif // WITH_EDITOR
 
 protected:
diff --git a/Engine/Source/Runtime/Engine/Classes/Engine/TextureRenderTarget2D.h b/Engine/Source/Runtime/Engine/Classes/Engine/TextureRenderTarget2D.h
--- a/Engine/Source/Runtime/Engine/Classes/Engine/TextureRenderTarget2D.h
+++ b/Engine/Source/Runtime/Engine/Classes/Engine/TextureRenderTarget2D.h
@@ -203,12 +203,12 @@ class UTextureRenderTarget2D : public UTextureRenderTarget
 	ENGINE_API ETextureSourceFormat GetTextureFormatForConversionToTexture2D() const;
 
 	//~ Begin UTextureRenderTarget Interface
-	virtual bool CanConvertToTexture(ETextureSourceFormat& OutTextureSourceFormat, EPixelFormat& OutPixelFormat, FText* OutErrorMessage) const override;
-	virtual TSubclassOf<UTexture> GetTextureUClass() const override;
-	virtual EPixelFormat GetFormat() const override;
-	virtual bool IsSRGB() const override;
-	virtual float GetDisplayGamma() const override;
-	virtual ETextureClass GetRenderTargetTextureClass() const override { return ETextureClass::TwoD; }
+	ENGINE_API virtual bool CanConvertToTexture(ETextureSourceFormat& OutTextureSourceFormat, EPixelFormat& OutPixelFormat, FText* OutErrorMessage) const override;
+	ENGINE_API virtual TSubclassOf<UTexture> GetTextureUClass() const override;
+	ENGINE_API virtual EPixelFormat GetFormat() const override;
+	ENGINE_API virtual bool IsSRGB() const override;
+	ENGINE_API virtual float GetDisplayGamma() const override;
+	ENGINE_API virtual ETextureClass GetRenderTargetTextureClass() const override { return ETextureClass::TwoD; }
 	//~ End UTextureRenderTarget Interface
 
 	/**
