[Input]

- Add timestamps to input events when they are sent through the Message Handler. Pass these timestamps along the input stack through to the ViewportClient via the FInputKeyEventArgs type.
- Add some data to FInputKeyEventArgs to allow it to support everything it needs to be used for InputAxis as well as InputKey.
- Deprecate the old InputAxis function which took in many paramaters in favor of one which uses the FInputKeyEventArgs type to make it more maintainable and behave the same as InputKey.

Next steps: Pass this input timestamp along to the PlayerController and UPlayerInput objects so that you can access it via gameplay code.

#tests PIE in editor, run cooked lyra build on win64 and remote kits

diff --git a/Engine/Plugins/Experimental/DynamicMaterial/Source/DynamicMaterialEditor/Private/UI/Widgets/Editor/DMMaterialPreviewViewportClient.cpp b/Engine/Plugins/Experimental/DynamicMaterial/Source/DynamicMaterialEditor/Private/UI/Widgets/Editor/DMMaterialPreviewViewportClient.cpp
--- a/Engine/Plugins/Experimental/DynamicMaterial/Source/DynamicMaterialEditor/Private/UI/Widgets/Editor/DMMaterialPreviewViewportClient.cpp
+++ b/Engine/Plugins/Experimental/DynamicMaterial/Source/DynamicMaterialEditor/Private/UI/Widgets/Editor/DMMaterialPreviewViewportClient.cpp
@@ -93,13 +93,13 @@ bool FDMMaterialPreviewViewportClient::InputKey(const FInputKeyEventArgs& InEven
 		|| AdvancedPreviewScene->HandleInputKey(InEventArgs);
 }
 
-bool FDMMaterialPreviewViewportClient::InputAxis(FViewport* InViewport, FInputDeviceId InDeviceId, FKey InKey, float InDelta, float InDeltaTime, int32 InNumSamples, bool bGamepad)
+bool FDMMaterialPreviewViewportClient::InputAxis(const FInputKeyEventArgs& InEventArgs)
 {
 	bool bResult = true;
 
 	if (!bDisableInput)
 	{
-		bResult = AdvancedPreviewScene->HandleViewportInput(InViewport, InDeviceId, InKey, InDelta, InDeltaTime, InNumSamples, bGamepad);
+		bResult = AdvancedPreviewScene->HandleViewportInput(InEventArgs.Viewport, InEventArgs.InputDevice, InEventArgs.Key, InEventArgs.AmountDepressed, InEventArgs.DeltaTime, InEventArgs.NumSamples, InEventArgs.IsGamepad());
 
 		if (bResult)
 		{
@@ -107,7 +107,7 @@ bool FDMMaterialPreviewViewportClient::InputAxis(FViewport* InViewport, FInputDe
 		}
 		else
 		{
-			bResult = FEditorViewportClient::InputAxis(InViewport, InDeviceId, InKey, InDelta, InDeltaTime, InNumSamples, bGamepad);
+			bResult = FEditorViewportClient::InputAxis(InEventArgs);
 		}
 	}
 
diff --git a/Engine/Plugins/Experimental/DynamicMaterial/Source/DynamicMaterialEditor/Private/UI/Widgets/Editor/DMMaterialPreviewViewportClient.h b/Engine/Plugins/Experimental/DynamicMaterial/Source/DynamicMaterialEditor/Private/UI/Widgets/Editor/DMMaterialPreviewViewportClient.h
--- a/Engine/Plugins/Experimental/DynamicMaterial/Source/DynamicMaterialEditor/Private/UI/Widgets/Editor/DMMaterialPreviewViewportClient.h
+++ b/Engine/Plugins/Experimental/DynamicMaterial/Source/DynamicMaterialEditor/Private/UI/Widgets/Editor/DMMaterialPreviewViewportClient.h
@@ -27,8 +27,7 @@ public:
 
 	//~ Begin FEditorViewportClient
 	virtual bool InputKey(const FInputKeyEventArgs& InEventArgs) override;
-	virtual bool InputAxis(FViewport* InViewport, FInputDeviceId InDeviceId, FKey InKey, float InDelta, float InDeltaTime,
-		int32 InNumSamples /* =1 */, bool bGamepad /* =false */) override;
+	virtual bool InputAxis(const FInputKeyEventArgs& InEventArgs) override;
 	virtual FLinearColor GetBackgroundColor() const override;
 	virtual void Tick(float InDeltaSeconds) override;
 	virtual bool ShouldOrbitCamera() const override;
diff --git a/Engine/Plugins/FX/Cascade/Source/Cascade/Private/CascadePreviewViewportClient.cpp b/Engine/Plugins/FX/Cascade/Source/Cascade/Private/CascadePreviewViewportClient.cpp
--- a/Engine/Plugins/FX/Cascade/Source/Cascade/Private/CascadePreviewViewportClient.cpp
+++ b/Engine/Plugins/FX/Cascade/Source/Cascade/Private/CascadePreviewViewportClient.cpp
@@ -659,20 +659,22 @@ static auto CVarCascadeDragSpeed = TAutoConsoleVariable<float>(TEXT("CascadeDrag
 static TAutoConsoleVariable<float> CVarCascadeRotateSpeed(TEXT("CascadeRotateSpeed"),0.005f,TEXT("Cascade drag speed."));
 static TAutoConsoleVariable<float> CVarCascadeScaleSpeed(TEXT("CascadeScaleSpeed"),1.0f,TEXT("Cascade scale speed."));
 
-bool FCascadeEdPreviewViewportClient::InputAxis(FViewport* InViewport, FInputDeviceId DeviceId, FKey Key, float Delta, float DeltaTime, int32 NumSamples, bool bGamepad)
+bool FCascadeEdPreviewViewportClient::InputAxis(const FInputKeyEventArgs& Args)
 {
 	bool bHandled = false;
 
+	const float Delta = Args.AmountDepressed;
+
 	if (bManipulatingVectorField)
 	{
 		UParticleModuleVectorFieldLocal* VectorFieldModule = Cast<UParticleModuleVectorFieldLocal>(CascadePtr.Pin()->GetSelectedModule());
 		if (VectorFieldModule)
 		{
-			const float MoveX = ((Key == EKeys::MouseX) ? Delta : 0.0f) * DragX;
-			const float MoveY = ((Key == EKeys::MouseY) ? Delta : 0.0f) * DragY;
+			const float MoveX = ((Args.Key == EKeys::MouseX) ? Delta : 0.0f) * DragX;
+			const float MoveY = ((Args.Key == EKeys::MouseY) ? Delta : 0.0f) * DragY;
 			const float MoveAmount = MoveX + MoveY;
 
-			VectorFieldModule->PreEditChange(NULL);
+			VectorFieldModule->PreEditChange(nullptr);
 			if (WidgetMM == WMM_Translate)
 			{
 				VectorFieldModule->RelativeTranslation += (LocalManipulateDir * MoveAmount * CVarCascadeDragSpeed.GetValueOnGameThread());
@@ -695,12 +697,12 @@ bool FCascadeEdPreviewViewportClient::InputAxis(FViewport* InViewport, FInputDev
 	}
 	else
 	{
-		bHandled = FEditorViewportClient::InputAxis(InViewport,DeviceId,Key,Delta,DeltaTime,NumSamples,bGamepad);
+		bHandled = FEditorViewportClient::InputAxis(Args);
 	}
 
 	if (!IsRealtime() && !FMath::IsNearlyZero(Delta))
 	{
-		InViewport->Invalidate();
+		Args.Viewport->Invalidate();
 	}
 
 	return bHandled;
diff --git a/Engine/Plugins/FX/Cascade/Source/Cascade/Private/CascadePreviewViewportClient.h b/Engine/Plugins/FX/Cascade/Source/Cascade/Private/CascadePreviewViewportClient.h
--- a/Engine/Plugins/FX/Cascade/Source/Cascade/Private/CascadePreviewViewportClient.h
+++ b/Engine/Plugins/FX/Cascade/Source/Cascade/Private/CascadePreviewViewportClient.h
@@ -34,13 +34,13 @@ class FCascadeEdPreviewViewportClient : public FEditorViewportClient
 public:
 	/** Constructor */
 	FCascadeEdPreviewViewportClient(TWeakPtr<FCascade> InCascade, const TSharedRef<SCascadePreviewViewport>& InCascadeViewport);
-	~FCascadeEdPreviewViewportClient();
+	virtual ~FCascadeEdPreviewViewportClient() override;
 
 	/** FEditorViewportClient interface */
 	virtual void Draw(FViewport* Viewport, FCanvas* Canvas) override;
 	virtual void Draw(const FSceneView* View, FPrimitiveDrawInterface* PDI) override;
 	virtual bool InputKey(const FInputKeyEventArgs& InEventArgs) override;
-	virtual bool InputAxis(FViewport* Viewport, FInputDeviceId DeviceId, FKey Key, float Delta, float DeltaTime, int32 NumSamples = 1, bool bGamepad = false) override;
+	virtual bool InputAxis(const FInputKeyEventArgs& InEventArgs) override;
 	virtual FSceneInterface* GetScene() const override;
 	virtual FLinearColor GetBackgroundColor() const override;
 	virtual bool ShouldOrbitCamera() const override;
diff --git a/Engine/Plugins/FX/Niagara/Source/NiagaraEditor/Private/Widgets/SNiagaraBakerViewport.cpp b/Engine/Plugins/FX/Niagara/Source/NiagaraEditor/Private/Widgets/SNiagaraBakerViewport.cpp
--- a/Engine/Plugins/FX/Niagara/Source/NiagaraEditor/Private/Widgets/SNiagaraBakerViewport.cpp
+++ b/Engine/Plugins/FX/Niagara/Source/NiagaraEditor/Private/Widgets/SNiagaraBakerViewport.cpp
@@ -101,8 +101,12 @@ public:
 		return true;
 	}
 
-	virtual bool InputAxis(FViewport* InViewport, FInputDeviceId InDeviceID, FKey Key, float Delta, float InDeltaTime, int32 InNumSamples, bool InbGamepad) override
+	virtual bool InputAxis(const FInputKeyEventArgs& Args) override
 	{
+		const FKey& Key = Args.Key;
+		FViewport* InViewport = Args.Viewport;
+		const float Delta = Args.AmountDepressed;
+		
 		// Viewport movement only enabled when preview is enabled otherwise there is no feedback
 		auto ViewModel = WeakViewModel.Pin();
 		UNiagaraBakerSettings* BakerSettings = ViewModel ? ViewModel->GetBakerSettings() : nullptr;
diff --git a/Engine/Plugins/Runtime/CommonUI/Source/CommonUI/Private/CommonGameViewportClient.cpp b/Engine/Plugins/Runtime/CommonUI/Source/CommonUI/Private/CommonGameViewportClient.cpp
--- a/Engine/Plugins/Runtime/CommonUI/Source/CommonUI/Private/CommonGameViewportClient.cpp
+++ b/Engine/Plugins/Runtime/CommonUI/Source/CommonUI/Private/CommonGameViewportClient.cpp
@@ -64,20 +64,20 @@ bool UCommonGameViewportClient::InputKey(const FInputKeyEventArgs& InEventArgs)
 	return Super::InputKey(EventArgs);
 }
 
-bool UCommonGameViewportClient::InputAxis(FViewport* InViewport, FInputDeviceId InputDevice, FKey Key, float Delta, float DeltaTime, int32 NumSamples, bool bGamepad)
+bool UCommonGameViewportClient::InputAxis(const FInputKeyEventArgs& Args)
 {
 	FReply RerouteResult = FReply::Unhandled();
 
-	if (!OnRerouteAxis().ExecuteIfBound(InputDevice, Key, Delta, RerouteResult))
+	if (!OnRerouteAxis().ExecuteIfBound(Args.InputDevice, Args.Key, Args.AmountDepressed, RerouteResult))
 	{
-		HandleRerouteAxis(InputDevice, Key, Delta, RerouteResult);
+		HandleRerouteAxis(Args.InputDevice, Args.Key, Args.AmountDepressed, RerouteResult);
 	}
 
 	if (RerouteResult.IsEventHandled())
 	{
 		return true;
 	}
-	return Super::InputAxis(InViewport, InputDevice, Key, Delta, DeltaTime, NumSamples, bGamepad);
+	return Super::InputAxis(Args);
 }
 
 bool UCommonGameViewportClient::InputTouch(FViewport* InViewport, int32 ControllerId, uint32 Handle, ETouchType::Type Type, const FVector2D& TouchLocation, float Force, FDateTime DeviceTimestamp, uint32 TouchpadIndex)
diff --git a/Engine/Plugins/Runtime/CommonUI/Source/CommonUI/Public/CommonGameViewportClient.h b/Engine/Plugins/Runtime/CommonUI/Source/CommonUI/Public/CommonGameViewportClient.h
--- a/Engine/Plugins/Runtime/CommonUI/Source/CommonUI/Public/CommonGameViewportClient.h
+++ b/Engine/Plugins/Runtime/CommonUI/Source/CommonUI/Public/CommonGameViewportClient.h
@@ -22,11 +22,11 @@ class COMMONUI_API UCommonGameViewportClient : public UGameViewportClient
 public:
 
 	UCommonGameViewportClient(FVTableHelper& Helper);
-	virtual ~UCommonGameViewportClient();
+	virtual ~UCommonGameViewportClient() override;
 
 	// UGameViewportClient interface begin
 	virtual bool InputKey(const FInputKeyEventArgs& EventArgs) override;
-	virtual bool InputAxis(FViewport* InViewport, FInputDeviceId InputDevice, FKey Key, float Delta, float DeltaTime, int32 NumSamples, bool bGamepad) override;
+	virtual bool InputAxis(const FInputKeyEventArgs& Args) override;
 	virtual bool InputTouch(FViewport* InViewport, int32 ControllerId, uint32 Handle, ETouchType::Type Type, const FVector2D& TouchLocation, float Force, FDateTime DeviceTimestamp, uint32 TouchpadIndex) override;
 	virtual void MouseMove(FViewport* InViewport, int32 X, int32 Y) override;
 	virtual void CapturedMouseMove(FViewport* InViewport, int32 X, int32 Y) override;
diff --git a/Engine/Plugins/TextureGraph/Source/TextureGraphEditor/Private/STG_EditorViewport.cpp b/Engine/Plugins/TextureGraph/Source/TextureGraphEditor/Private/STG_EditorViewport.cpp
--- a/Engine/Plugins/TextureGraph/Source/TextureGraphEditor/Private/STG_EditorViewport.cpp
+++ b/Engine/Plugins/TextureGraph/Source/TextureGraphEditor/Private/STG_EditorViewport.cpp
@@ -50,7 +50,7 @@ public:
 
 	// FEditorViewportClient interface
 	virtual bool InputKey(const FInputKeyEventArgs& EventArgs) override;
-	virtual bool InputAxis(FViewport* Viewport, FInputDeviceId DeviceID, FKey Key, float Delta, float DeltaTime, int32 NumSamples = 1, bool bGamepad = false) override;
+	virtual bool InputAxis(const FInputKeyEventArgs& EventArgs) override;
 
 	virtual FLinearColor GetBackgroundColor() const override;
 
@@ -144,20 +144,20 @@ bool FTG_EditorViewportClient::InputKey(const FInputKeyEventArgs& EventArgs)
 	return bHandled;
 }
 
-bool FTG_EditorViewportClient::InputAxis(FViewport* InViewport, FInputDeviceId DeviceID, FKey Key, float Delta, float DeltaTime, int32 NumSamples/* =1 */, bool bGamepad/* =false */)
+bool FTG_EditorViewportClient::InputAxis(const FInputKeyEventArgs& EventArgs)
 {
 	bool bResult = true;
 
 	if (!bDisableInput)
 	{
-		bResult = AdvancedPreviewScene->HandleViewportInput(InViewport, DeviceID, Key, Delta, DeltaTime, NumSamples, bGamepad);
+		bResult = AdvancedPreviewScene->HandleViewportInput(EventArgs.Viewport, EventArgs.InputDevice, EventArgs.Key, EventArgs.AmountDepressed, EventArgs.DeltaTime, EventArgs.NumSamples, EventArgs.IsGamepad());
 		if (bResult)
 		{
 			Invalidate();
 		}
 		else
 		{
-			bResult = FEditorViewportClient::InputAxis(InViewport, DeviceID, Key, Delta, DeltaTime, NumSamples, bGamepad);
+			bResult = FEditorViewportClient::InputAxis(EventArgs);
 		}
 	}
 
diff --git a/Engine/Plugins/VirtualProduction/CameraCalibration/Source/CameraCalibrationEditor/Private/AssetEditor/Viewport/SimulcamEditorViewportClient.cpp b/Engine/Plugins/VirtualProduction/CameraCalibration/Source/CameraCalibrationEditor/Private/AssetEditor/Viewport/SimulcamEditorViewportClient.cpp
--- a/Engine/Plugins/VirtualProduction/CameraCalibration/Source/CameraCalibrationEditor/Private/AssetEditor/Viewport/SimulcamEditorViewportClient.cpp
+++ b/Engine/Plugins/VirtualProduction/CameraCalibration/Source/CameraCalibrationEditor/Private/AssetEditor/Viewport/SimulcamEditorViewportClient.cpp
@@ -310,8 +310,12 @@ bool FSimulcamEditorViewportClient::InputChar(FViewport* Viewport, int32 Control
 	return false;
 }
 
-bool FSimulcamEditorViewportClient::InputAxis(FViewport* Viewport, FInputDeviceId DeviceID, FKey Key, float Delta, float DeltaTime, int32 NumSamples, bool bGamepad)
+bool FSimulcamEditorViewportClient::InputAxis(const FInputKeyEventArgs& InEventArgs)
 {
+	const FKey& Key = InEventArgs.Key;
+	const float Delta = InEventArgs.AmountDepressed;
+	FViewport* Viewport = InEventArgs.Viewport;
+	
 	if (Key == EKeys::MouseX || Key == EKeys::MouseY)
 	{
 		if (ShouldUseMousePanning(Viewport))
diff --git a/Engine/Plugins/VirtualProduction/CameraCalibration/Source/CameraCalibrationEditor/Private/AssetEditor/Viewport/SimulcamEditorViewportClient.h b/Engine/Plugins/VirtualProduction/CameraCalibration/Source/CameraCalibrationEditor/Private/AssetEditor/Viewport/SimulcamEditorViewportClient.h
--- a/Engine/Plugins/VirtualProduction/CameraCalibration/Source/CameraCalibrationEditor/Private/AssetEditor/Viewport/SimulcamEditorViewportClient.h
+++ b/Engine/Plugins/VirtualProduction/CameraCalibration/Source/CameraCalibrationEditor/Private/AssetEditor/Viewport/SimulcamEditorViewportClient.h
@@ -25,7 +25,7 @@ public:
 	virtual void Draw(FViewport* Viewport, FCanvas* Canvas) override;
 	virtual bool InputKey(const FInputKeyEventArgs& InEventArgs) override;
 	virtual bool InputChar(FViewport* Viewport, int32 ControllerId, TCHAR Character) override;
-	virtual bool InputAxis(FViewport* Viewport, FInputDeviceId DeviceId, FKey Key, float Delta, float DeltaTime, int32 NumSamples=1, bool bGamepad=false) override;
+	virtual bool InputAxis(const FInputKeyEventArgs& InEventArgs) override;
 	virtual UWorld* GetWorld() const override { return nullptr; }
 	virtual EMouseCursor::Type GetCursor(FViewport* Viewport, int32 X, int32 Y) override;
 	virtual void MouseMove(FViewport* Viewport, int32 X, int32 Y) override;
diff --git a/Engine/Source/Editor/DistCurveEditor/Private/CurveEditorViewportClient.cpp b/Engine/Source/Editor/DistCurveEditor/Private/CurveEditorViewportClient.cpp
--- a/Engine/Source/Editor/DistCurveEditor/Private/CurveEditorViewportClient.cpp
+++ b/Engine/Source/Editor/DistCurveEditor/Private/CurveEditorViewportClient.cpp
@@ -805,13 +805,13 @@ void FCurveEditorViewportClient::CapturedMouseMove(FViewport* Viewport, int32 X,
 	MouseMove(Viewport, X, Y);
 }
 
-bool FCurveEditorViewportClient::InputAxis(FViewport* Viewport, FInputDeviceId InputDevice, FKey Key, float Delta, float DeltaTime, int32 NumSamples , bool bGamepad)
+bool FCurveEditorViewportClient::InputAxis(const FInputKeyEventArgs& Args)
 {
-	if (Key == EKeys::MouseX || Key == EKeys::MouseY)
+	if (Args.Key == EKeys::MouseX || Args.Key == EKeys::MouseY)
 	{
-		int32 X = Viewport->GetMouseX();
-		int32 Y = Viewport->GetMouseY();
-		MouseMove(Viewport, X, Y);
+		int32 X = Args.Viewport->GetMouseX();
+		int32 Y = Args.Viewport->GetMouseY();
+		MouseMove(Args.Viewport, X, Y);
 		return true;
 	}
 	return false;
diff --git a/Engine/Source/Editor/DistCurveEditor/Private/CurveEditorViewportClient.h b/Engine/Source/Editor/DistCurveEditor/Private/CurveEditorViewportClient.h
--- a/Engine/Source/Editor/DistCurveEditor/Private/CurveEditorViewportClient.h
+++ b/Engine/Source/Editor/DistCurveEditor/Private/CurveEditorViewportClient.h
@@ -38,7 +38,7 @@ public:
 	virtual bool InputKey(const FInputKeyEventArgs& EventArgs) override;
 	virtual void MouseMove(FViewport* Viewport, int32 X, int32 Y) override;
 	virtual void CapturedMouseMove(FViewport* Viewport, int32 X, int32 Y) override;
-	virtual bool InputAxis(FViewport* Viewport, FInputDeviceId InputDevice, FKey Key, float Delta, float DeltaTime, int32 NumSamples = 1, bool bGamepad = false) override;
+	virtual bool InputAxis(const FInputKeyEventArgs& Args) override;
 
 	/** Exec handler */
 	void Exec(const TCHAR* Cmd);
diff --git a/Engine/Source/Editor/MaterialEditor/Private/SMaterialEditorViewport.cpp b/Engine/Source/Editor/MaterialEditor/Private/SMaterialEditorViewport.cpp
--- a/Engine/Source/Editor/MaterialEditor/Private/SMaterialEditorViewport.cpp
+++ b/Engine/Source/Editor/MaterialEditor/Private/SMaterialEditorViewport.cpp
@@ -62,7 +62,7 @@ public:
 
 	// FEditorViewportClient interface
 	virtual bool InputKey(const FInputKeyEventArgs& EventArgs) override;
-	virtual bool InputAxis(FViewport* InViewport, FInputDeviceId DeviceId, FKey Key, float Delta, float DeltaTime, int32 NumSamples/* =1 */, bool bGamepad/* =false */) override;
+	virtual bool InputAxis(const FInputKeyEventArgs& Args) override;
 	virtual FLinearColor GetBackgroundColor() const override;
 	virtual void Tick(float DeltaSeconds) override;
 	virtual void Draw(FViewport* Viewport,FCanvas* Canvas) override;
@@ -154,20 +154,20 @@ bool FMaterialEditorViewportClient::InputKey(const FInputKeyEventArgs& EventArgs
 	return bHandled;
 }
 
-bool FMaterialEditorViewportClient::InputAxis(FViewport* InViewport, FInputDeviceId DeviceId, FKey Key, float Delta, float DeltaTime, int32 NumSamples/* =1 */, bool bGamepad/* =false */)
+bool FMaterialEditorViewportClient::InputAxis(const FInputKeyEventArgs& Args)
 {
 	bool bResult = true;
 
 	if (!bDisableInput)
 	{
-		bResult = AdvancedPreviewScene->HandleViewportInput(InViewport, DeviceId, Key, Delta, DeltaTime, NumSamples, bGamepad);
+		bResult = AdvancedPreviewScene->HandleViewportInput(Args.Viewport, Args.InputDevice, Args.Key, Args.AmountDepressed, Args.DeltaTime, Args.NumSamples, Args.IsGamepad());
 		if (bResult)
 		{
 			Invalidate();
 		}
 		else
 		{
-			bResult = FEditorViewportClient::InputAxis(InViewport, DeviceId, Key, Delta, DeltaTime, NumSamples, bGamepad);
+			bResult = FEditorViewportClient::InputAxis(Args);
 		}
 	}
 
diff --git a/Engine/Source/Editor/Persona/Private/AnimationEditorViewportClient.cpp b/Engine/Source/Editor/Persona/Private/AnimationEditorViewportClient.cpp
--- a/Engine/Source/Editor/Persona/Private/AnimationEditorViewportClient.cpp
+++ b/Engine/Source/Editor/Persona/Private/AnimationEditorViewportClient.cpp
@@ -1576,21 +1576,21 @@ bool FAnimationViewportClient::InputKey(const FInputKeyEventArgs& EventArgs)
 		: FEditorViewportClient::InputKey(EventArgs);
 }
 
-bool FAnimationViewportClient::InputAxis(FViewport* InViewport, FInputDeviceId DeviceId, FKey Key, float Delta, float DeltaTime, int32 NumSamples /*= 1*/, bool bGamepad /*= false*/)
+bool FAnimationViewportClient::InputAxis(const FInputKeyEventArgs& EventArgs)
 {
 	bool bResult = true;
 
 	if (!bDisableInput)
 	{
 		FAdvancedPreviewScene* AdvancedScene = (FAdvancedPreviewScene*)PreviewScene;
-		bResult = AdvancedScene->HandleViewportInput(InViewport, DeviceId, Key, Delta, DeltaTime, NumSamples, bGamepad);
+		bResult = AdvancedScene->HandleViewportInput(EventArgs.Viewport, EventArgs.InputDevice, EventArgs.Key, EventArgs.AmountDepressed, EventArgs.DeltaTime, EventArgs.NumSamples, EventArgs.IsGamepad());
 		if (bResult)
 		{
 			Invalidate();
 		}
 		else
 		{
-			bResult = FEditorViewportClient::InputAxis(InViewport, DeviceId, Key, Delta, DeltaTime, NumSamples, bGamepad);
+			bResult = FEditorViewportClient::InputAxis(EventArgs);
 		}
 	}
 
diff --git a/Engine/Source/Editor/Persona/Public/AnimationEditorViewportClient.h b/Engine/Source/Editor/Persona/Public/AnimationEditorViewportClient.h
--- a/Engine/Source/Editor/Persona/Public/AnimationEditorViewportClient.h
+++ b/Engine/Source/Editor/Persona/Public/AnimationEditorViewportClient.h
@@ -52,7 +52,7 @@ public:
 	virtual void Draw(const FSceneView* View,FPrimitiveDrawInterface* PDI) override;
 	virtual void DrawCanvas( FViewport& InViewport, FSceneView& View, FCanvas& Canvas ) override;
 	virtual bool InputKey(const FInputKeyEventArgs& EventArgs) override;
-	virtual bool InputAxis(FViewport* InViewport, FInputDeviceId DeviceId, FKey Key, float Delta, float DeltaTime, int32 NumSamples = 1, bool bGamepad = false) override;
+	virtual bool InputAxis(const FInputKeyEventArgs& EventArgs) override;
 
 	// Sets what bones are drawn by DrawMeshBones and ShowBoneNames
 	virtual void UpdateBonesToDraw();
diff --git a/Engine/Source/Editor/StaticMeshEditor/Private/StaticMeshEditorViewportClient.cpp b/Engine/Source/Editor/StaticMeshEditor/Private/StaticMeshEditorViewportClient.cpp
--- a/Engine/Source/Editor/StaticMeshEditor/Private/StaticMeshEditorViewportClient.cpp
+++ b/Engine/Source/Editor/StaticMeshEditor/Private/StaticMeshEditorViewportClient.cpp
@@ -1034,20 +1034,20 @@ bool FStaticMeshEditorViewportClient::InputKey(const FInputKeyEventArgs& EventAr
 	return bHandled;
 }
 
-bool FStaticMeshEditorViewportClient::InputAxis(FViewport* InViewport, FInputDeviceId DeviceId, FKey Key, float Delta, float DeltaTime, int32 NumSamples, bool bGamepad)
+bool FStaticMeshEditorViewportClient::InputAxis(const FInputKeyEventArgs& Args)
 {
 	bool bResult = true;
 	
 	if (!bDisableInput)
 	{
-		bResult = AdvancedPreviewScene->HandleViewportInput(InViewport, DeviceId, Key, Delta, DeltaTime, NumSamples, bGamepad);
+		bResult = AdvancedPreviewScene->HandleViewportInput(Args.Viewport, Args.InputDevice, Args.Key, Args.AmountDepressed, Args.DeltaTime, Args.NumSamples, Args.IsGamepad());
 		if (bResult)
 		{
 			Invalidate();
 		}
 		else
 		{
-			bResult = FEditorViewportClient::InputAxis(InViewport, DeviceId, Key, Delta, DeltaTime, NumSamples, bGamepad);
+			bResult = FEditorViewportClient::InputAxis(Args);
 		}
 	}
 
diff --git a/Engine/Source/Editor/StaticMeshEditor/Private/StaticMeshEditorViewportClient.h b/Engine/Source/Editor/StaticMeshEditor/Private/StaticMeshEditorViewportClient.h
--- a/Engine/Source/Editor/StaticMeshEditor/Private/StaticMeshEditorViewportClient.h
+++ b/Engine/Source/Editor/StaticMeshEditor/Private/StaticMeshEditorViewportClient.h
@@ -27,7 +27,7 @@ public:
 	// FEditorViewportClient interface
 	virtual void MouseMove(FViewport* Viewport,int32 x, int32 y) override;
 	virtual bool InputKey(const FInputKeyEventArgs& EventArgs) override;
-	virtual bool InputAxis(FViewport* Viewport, FInputDeviceId DeviceId, FKey Key, float Delta, float DeltaTime, int32 NumSamples=1, bool bGamepad=false) override;
+	virtual bool InputAxis(const FInputKeyEventArgs& Args) override;
 	virtual void ProcessClick(class FSceneView& View, class HHitProxy* HitProxy, FKey Key, EInputEvent Event, uint32 HitX, uint32 HitY) override;
 	virtual void Tick(float DeltaSeconds) override;
 	virtual void Draw(const FSceneView* View,FPrimitiveDrawInterface* PDI) override;
diff --git a/Engine/Source/Editor/TextureEditor/Private/Models/TextureEditorViewportClient.cpp b/Engine/Source/Editor/TextureEditor/Private/Models/TextureEditorViewportClient.cpp
--- a/Engine/Source/Editor/TextureEditor/Private/Models/TextureEditorViewportClient.cpp
+++ b/Engine/Source/Editor/TextureEditor/Private/Models/TextureEditorViewportClient.cpp
@@ -560,8 +560,12 @@ bool FTextureEditorViewportClient::InputKey(const FInputKeyEventArgs& InEventArg
 	return false;
 }
 
-bool FTextureEditorViewportClient::InputAxis(FViewport* Viewport, FInputDeviceId DeviceId, FKey Key, float Delta, float DeltaTime, int32 NumSamples, bool bGamepad)
+bool FTextureEditorViewportClient::InputAxis(const FInputKeyEventArgs& InEventArgs)
 {
+	const FKey& Key = InEventArgs.Key;
+	FViewport* Viewport = InEventArgs.Viewport;
+	const float Delta = InEventArgs.AmountDepressed;
+	
 	if (Key == EKeys::MouseX || Key == EKeys::MouseY)
 	{
 		TSharedPtr<ITextureEditorToolkit> TextureEditorPinned = TextureEditorPtr.Pin();
diff --git a/Engine/Source/Editor/TextureEditor/Private/Models/TextureEditorViewportClient.h b/Engine/Source/Editor/TextureEditor/Private/Models/TextureEditorViewportClient.h
--- a/Engine/Source/Editor/TextureEditor/Private/Models/TextureEditorViewportClient.h
+++ b/Engine/Source/Editor/TextureEditor/Private/Models/TextureEditorViewportClient.h
@@ -20,12 +20,12 @@ class FTextureEditorViewportClient
 public:
 	/** Constructor */
 	FTextureEditorViewportClient(TWeakPtr<ITextureEditorToolkit> InTextureEditor, TWeakPtr<STextureEditorViewport> InTextureEditorViewport);
-	~FTextureEditorViewportClient();
+	virtual ~FTextureEditorViewportClient() override;
 
 	/** FViewportClient interface */
 	virtual void Draw(FViewport* Viewport, FCanvas* Canvas) override;
 	virtual bool InputKey(const FInputKeyEventArgs& InEventArgs) override;
-	virtual bool InputAxis(FViewport* Viewport, FInputDeviceId DeviceId, FKey Key, float Delta, float DeltaTime, int32 NumSamples = 1, bool bGamepad = false) override;
+	virtual bool InputAxis(const FInputKeyEventArgs& InEventArgs) override;
 	virtual bool InputGesture(FViewport* Viewport, EGestureEvent GestureType, const FVector2D& GestureDelta, bool bIsDirectionInvertedFromDevice) override;
 	virtual UWorld* GetWorld() const override { return nullptr; }
 	virtual EMouseCursor::Type GetCursor(FViewport* Viewport, int32 X, int32 Y) override;
diff --git a/Engine/Source/Editor/UnrealEd/Private/EditorViewportClient.cpp b/Engine/Source/Editor/UnrealEd/Private/EditorViewportClient.cpp
--- a/Engine/Source/Editor/UnrealEd/Private/EditorViewportClient.cpp
+++ b/Engine/Source/Editor/UnrealEd/Private/EditorViewportClient.cpp
@@ -4784,9 +4784,16 @@ FVector FEditorViewportClient::TranslateDelta( FKey InKey, float InDelta, bool I
 	return vec;
 }
 
-bool FEditorViewportClient::InputAxis(FViewport* InViewport, FInputDeviceId DeviceID, FKey Key, float Delta, float DeltaTime, int32 NumSamples, bool bGamepad)
-{
-	return Internal_InputAxis(InViewport, DeviceID, Key, Delta, DeltaTime, NumSamples, bGamepad);
+bool FEditorViewportClient::InputAxis(const FInputKeyEventArgs& Args)
+{
+	return Internal_InputAxis(
+		Args.Viewport,
+		Args.InputDevice,
+		Args.Key,
+		Args.AmountDepressed,
+		Args.DeltaTime,
+		Args.NumSamples,
+		Args.IsGamepad());
 }
 
 bool FEditorViewportClient::Internal_InputAxis(FViewport* InViewport, FInputDeviceId DeviceID, FKey Key, float Delta, float DeltaTime, int32 NumSamples, bool bGamepad)
diff --git a/Engine/Source/Editor/UnrealEd/Private/LevelEditorViewport.cpp b/Engine/Source/Editor/UnrealEd/Private/LevelEditorViewport.cpp
--- a/Engine/Source/Editor/UnrealEd/Private/LevelEditorViewport.cpp
+++ b/Engine/Source/Editor/UnrealEd/Private/LevelEditorViewport.cpp
@@ -4010,7 +4010,7 @@ private:
 	FLevelEditorViewportClient* PrevCurrentLevelEditingViewportClient;
 };
 
-bool FLevelEditorViewportClient::InputAxis(FViewport* InViewport, FInputDeviceId DeviceId, FKey Key, float Delta, float DeltaTime, int32 NumSamples, bool bGamepad)
+bool FLevelEditorViewportClient::InputAxis(const FInputKeyEventArgs& Args)
 {
 	if (bDisableInput)
 	{
@@ -4019,14 +4019,14 @@ bool FLevelEditorViewportClient::InputAxis(FViewport* InViewport, FInputDeviceId
 
 	// Forward input axis events to the Editor World extension collection and give extensions an opportunity to consume input.
 	IPlatformInputDeviceMapper& DeviceMapper = IPlatformInputDeviceMapper::Get();
-	FPlatformUserId UserId = DeviceMapper.GetUserForInputDevice(DeviceId);
+	FPlatformUserId UserId = DeviceMapper.GetUserForInputDevice(Args.InputDevice);
 
 	int32 ControllerId;
-	if (DeviceMapper.RemapUserAndDeviceToControllerId(UserId, ControllerId, DeviceId))
+	if (DeviceMapper.RemapUserAndDeviceToControllerId(UserId, ControllerId, Args.InputDevice))
 	{
 		UEditorWorldExtensionCollection& ExtensionCollection = *GEditor->GetEditorWorldExtensionsManager()->GetEditorWorldExtensions(GetWorld());
 		
-		if (ExtensionCollection.InputAxis(this, InViewport, ControllerId, Key, Delta, DeltaTime))
+		if (ExtensionCollection.InputAxis(this, Args.Viewport, ControllerId, Args.Key, Args.AmountDepressed, Args.DeltaTime))
 		{
 			return true;
 		}
@@ -4036,7 +4036,7 @@ bool FLevelEditorViewportClient::InputAxis(FViewport* InViewport, FInputDeviceId
 
 	FScopedSetCurrentViewportClient( this );
 
-	return FEditorViewportClient::InputAxis(InViewport, DeviceId, Key, Delta, DeltaTime, NumSamples, bGamepad);
+	return FEditorViewportClient::InputAxis(Args);
 }
 
 
diff --git a/Engine/Source/Editor/UnrealEd/Public/EditorViewportClient.h b/Engine/Source/Editor/UnrealEd/Public/EditorViewportClient.h
--- a/Engine/Source/Editor/UnrealEd/Public/EditorViewportClient.h
+++ b/Engine/Source/Editor/UnrealEd/Public/EditorViewportClient.h
@@ -599,7 +599,7 @@ public:
 	UNREALED_API virtual void RedrawRequested(FViewport* Viewport) override;
 	UNREALED_API virtual void RequestInvalidateHitProxy(FViewport* Viewport) override;
 	UNREALED_API virtual bool InputKey(const FInputKeyEventArgs& EventArgs) override;
-	UNREALED_API virtual bool InputAxis(FViewport* Viewport, FInputDeviceId DeviceID, FKey Key, float Delta, float DeltaTime, int32 NumSamples=1, bool bGamepad=false) override;
+	UNREALED_API virtual bool InputAxis(const FInputKeyEventArgs& Args) override;
 	UNREALED_API virtual bool InputGesture(FViewport* Viewport, EGestureEvent GestureType, const FVector2D& GestureDelta, bool bIsDirectionInvertedFromDevice) override;
 	UNREALED_API virtual void ReceivedFocus(FViewport* Viewport) override;
 	UNREALED_API virtual void MouseEnter(FViewport* Viewport,int32 x, int32 y) override;
diff --git a/Engine/Source/Editor/UnrealEd/Public/LevelEditorViewport.h b/Engine/Source/Editor/UnrealEd/Public/LevelEditorViewport.h
--- a/Engine/Source/Editor/UnrealEd/Public/LevelEditorViewport.h
+++ b/Engine/Source/Editor/UnrealEd/Public/LevelEditorViewport.h
@@ -200,7 +200,7 @@ public:
 	// FEditorViewportClient interface
 	UNREALED_API virtual void DrawCanvas( FViewport& InViewport, FSceneView& View, FCanvas& Canvas ) override;
 	UNREALED_API virtual bool InputKey(const FInputKeyEventArgs& InEventArgs) override;
-	UNREALED_API virtual bool InputAxis(FViewport* Viewport, FInputDeviceId DeviceId, FKey Key, float Delta, float DeltaTime, int32 NumSamples=1, bool bGamepad=false) override;
+	UNREALED_API virtual bool InputAxis(const FInputKeyEventArgs& Args) override;
 	UNREALED_API virtual EMouseCursor::Type GetCursor(FViewport* Viewport,int32 X,int32 Y) override;
 	UNREALED_API virtual void CapturedMouseMove(FViewport* InViewport, int32 InMouseX, int32 InMouseY) override;
 	UNREALED_API virtual void MouseMove(FViewport* InViewport, int32 x, int32 y) override;
diff --git a/Engine/Source/Runtime/Engine/Classes/Engine/GameViewportClient.h b/Engine/Source/Runtime/Engine/Classes/Engine/GameViewportClient.h
--- a/Engine/Source/Runtime/Engine/Classes/Engine/GameViewportClient.h
+++ b/Engine/Source/Runtime/Engine/Classes/Engine/GameViewportClient.h
@@ -170,9 +170,7 @@ public:
 	//~ Begin FViewportClient Interface.
 	virtual void RedrawRequested(FViewport* InViewport) override {}
 	ENGINE_API virtual bool InputKey(const FInputKeyEventArgs& EventArgs) override;
-	
-	ENGINE_API virtual bool InputAxis(FViewport* Viewport, FInputDeviceId InputDevice, FKey Key, float Delta, float DeltaTime, int32 NumSamples = 1, bool bGamepad = false) override;
-	
+	ENGINE_API virtual bool InputAxis(const FInputKeyEventArgs& EventArgs) override;
 	ENGINE_API virtual bool InputChar(FViewport* Viewport,int32 ControllerId, TCHAR Character) override;
 	ENGINE_API virtual bool InputTouch(FViewport* Viewport, int32 ControllerId, uint32 Handle, ETouchType::Type Type, const FVector2D& TouchLocation, float Force, FDateTime DeviceTimestamp, uint32 TouchpadIndex) override;
 	ENGINE_API virtual bool InputMotion(FViewport* Viewport, int32 ControllerId, const FVector& Tilt, const FVector& RotationRate, const FVector& Gravity, const FVector& Acceleration) override;
diff --git a/Engine/Source/Runtime/Engine/Private/GameViewportClient.cpp b/Engine/Source/Runtime/Engine/Private/GameViewportClient.cpp
--- a/Engine/Source/Runtime/Engine/Private/GameViewportClient.cpp
+++ b/Engine/Source/Runtime/Engine/Private/GameViewportClient.cpp
@@ -749,7 +749,8 @@ bool UGameViewportClient::InputKey(const FInputKeyEventArgs& InEventArgs)
 		ULocalPlayer* const TargetPlayer = GEngine->GetLocalPlayerFromInputDevice(this, EventArgs.InputDevice);
 		if (TargetPlayer && TargetPlayer->PlayerController)
 		{
-			bResult = TargetPlayer->PlayerController->InputKey(FInputKeyParams(EventArgs.Key, EventArgs.Event, static_cast<double>(EventArgs.AmountDepressed), EventArgs.IsGamepad(), EventArgs.InputDevice));
+			// TODO: Deprecate FInputKeyParams and use FInputKeyParams instead.
+			bResult = TargetPlayer->PlayerController->InputKey(FInputKeyParams(EventArgs.Key, EventArgs.Event, EventArgs.AmountDepressed, EventArgs.IsGamepad(), EventArgs.InputDevice));
 		}
 
 		// A gameviewport is always considered to have responded to a mouse buttons to avoid throttling
@@ -763,17 +764,21 @@ bool UGameViewportClient::InputKey(const FInputKeyEventArgs& InEventArgs)
 	return bResult;
 }
 
-bool UGameViewportClient::InputAxis(FViewport* InViewport, FInputDeviceId InputDevice, FKey Key, float Delta, float DeltaTime, int32 NumSamples, bool bGamepad)
+bool UGameViewportClient::InputAxis(const FInputKeyEventArgs& InArgs)
 {
 	if (IgnoreInput())
 	{
 		return false;
 	}
 
-	// Handle mapping controller id and key if needed
-	FInputKeyEventArgs EventArgs(InViewport, InputDevice, Key, IE_Axis);
+	// Handle mapping controller id and key if needed, which may need to mutate the args
+	FInputKeyEventArgs MutableEventArgs = InArgs;
 
-	RemapControllerInput(EventArgs);
+	FViewport* InViewport = MutableEventArgs.Viewport;
+	const bool bGamepad = MutableEventArgs.IsGamepad();
+	float Delta = MutableEventArgs.AmountDepressed;
+
+	RemapControllerInput(MutableEventArgs);
 
 #if WITH_EDITOR
 	if (bGamepad && InViewport)
@@ -791,7 +796,7 @@ bool UGameViewportClient::InputAxis(FViewport* InViewport, FInputDeviceId InputD
 			{
 				if (UGameViewportClient* NextViewport = GEngine->GetNextPIEViewport(this))
 				{
-					const bool bResult = NextViewport->InputAxis(InViewport, InputDevice, Key, Delta, DeltaTime, NumSamples, bGamepad);
+					const bool bResult = NextViewport->InputAxis(MutableEventArgs);
 					return false;
 				}
 			}
@@ -799,7 +804,7 @@ bool UGameViewportClient::InputAxis(FViewport* InViewport, FInputDeviceId InputD
 	}
 #endif
 	
-	OnInputAxisEvent.Broadcast(InViewport, EventArgs.ControllerId, EventArgs.Key, Delta, DeltaTime, NumSamples, EventArgs.IsGamepad());
+	OnInputAxisEvent.Broadcast(InViewport, MutableEventArgs.ControllerId, MutableEventArgs.Key, Delta, MutableEventArgs.DeltaTime, MutableEventArgs.NumSamples, MutableEventArgs.IsGamepad());
 	
 	bool bResult = false;
 
@@ -813,21 +818,21 @@ bool UGameViewportClient::InputAxis(FViewport* InViewport, FInputDeviceId InputD
 			// route to subsystems that care
 			if (ViewportConsole != nullptr)
 			{
-				bResult = ViewportConsole->InputAxis(EventArgs.InputDevice, EventArgs.Key, Delta, DeltaTime, NumSamples, EventArgs.IsGamepad());
+				bResult = ViewportConsole->InputAxis(MutableEventArgs.InputDevice, MutableEventArgs.Key, Delta, MutableEventArgs.DeltaTime, MutableEventArgs.NumSamples, MutableEventArgs.IsGamepad());
 			}
 		
 			// Try the override callback, this may modify event args
 			if (!bResult && OnOverrideInputAxisEvent.IsBound())
 			{
-				bResult = OnOverrideInputAxisEvent.Execute(EventArgs, Delta, DeltaTime, NumSamples);
+				bResult = OnOverrideInputAxisEvent.Execute(MutableEventArgs, Delta, MutableEventArgs.DeltaTime, MutableEventArgs.NumSamples);
 			}
 
 			if (!bResult)
 			{
-				ULocalPlayer* const TargetPlayer = GEngine->GetLocalPlayerFromInputDevice(this, EventArgs.InputDevice);
+				ULocalPlayer* const TargetPlayer = GEngine->GetLocalPlayerFromInputDevice(this, MutableEventArgs.InputDevice);
 				if (TargetPlayer && TargetPlayer->PlayerController)
 				{
-					bResult = TargetPlayer->PlayerController->InputKey(FInputKeyParams(EventArgs.Key, (double)Delta, DeltaTime, NumSamples, EventArgs.IsGamepad(), EventArgs.InputDevice));
+					bResult = TargetPlayer->PlayerController->InputKey(FInputKeyParams(MutableEventArgs.Key, (double)Delta, MutableEventArgs.DeltaTime, MutableEventArgs.NumSamples, MutableEventArgs.IsGamepad(), MutableEventArgs.InputDevice));
 				}
 			}
 
diff --git a/Engine/Source/Runtime/Engine/Private/Slate/SceneViewport.cpp b/Engine/Source/Runtime/Engine/Private/Slate/SceneViewport.cpp
--- a/Engine/Source/Runtime/Engine/Private/Slate/SceneViewport.cpp
+++ b/Engine/Source/Runtime/Engine/Private/Slate/SceneViewport.cpp
@@ -251,36 +251,38 @@ void FSceneViewport::UpdateModifierKeys( const FPointerEvent& InMouseEvent )
 	KeyStateMap.Add(EKeys::RightCommand, InMouseEvent.IsRightCommandDown());
 }
 
-void FSceneViewport::ApplyModifierKeys(const FModifierKeysState& InKeysState)
+void FSceneViewport::ApplyModifierKeys(const FModifierKeysState& InKeysState, const uint64 Timestamp)
 {
 	if (ViewportClient && GetSizeXY() != FIntPoint::ZeroValue)
 	{
+		const FInputDeviceId DefaultInputDevice = IPlatformInputDeviceMapper::Get().GetDefaultInputDevice();
+		
 		// Switch to the viewport clients world before processing input
 		FScopedConditionalWorldSwitcher WorldSwitcher(ViewportClient);
 
 		if (InKeysState.IsLeftAltDown())
 		{
-			ViewportClient->InputKey(FInputKeyEventArgs(this, 0, EKeys::LeftAlt, IE_Pressed));
+			ViewportClient->InputKey(FInputKeyEventArgs(this, DefaultInputDevice, EKeys::LeftAlt, IE_Pressed, Timestamp));
 		}
 		if (InKeysState.IsRightAltDown())
 		{
-			ViewportClient->InputKey(FInputKeyEventArgs(this, 0, EKeys::RightAlt, IE_Pressed));
+			ViewportClient->InputKey(FInputKeyEventArgs(this, DefaultInputDevice, EKeys::RightAlt, IE_Pressed, Timestamp));
 		}
 		if (InKeysState.IsLeftControlDown())
 		{
-			ViewportClient->InputKey(FInputKeyEventArgs(this, 0, EKeys::LeftControl, IE_Pressed));
+			ViewportClient->InputKey(FInputKeyEventArgs(this, DefaultInputDevice, EKeys::LeftControl, IE_Pressed, Timestamp));
 		}
 		if (InKeysState.IsRightControlDown())
 		{
-			ViewportClient->InputKey(FInputKeyEventArgs(this, 0, EKeys::RightControl, IE_Pressed));
+			ViewportClient->InputKey(FInputKeyEventArgs(this, DefaultInputDevice, EKeys::RightControl, IE_Pressed, Timestamp));
 		}
 		if (InKeysState.IsLeftShiftDown())
 		{
-			ViewportClient->InputKey(FInputKeyEventArgs(this, 0, EKeys::LeftShift, IE_Pressed));
+			ViewportClient->InputKey(FInputKeyEventArgs(this, DefaultInputDevice, EKeys::LeftShift, IE_Pressed, Timestamp));
 		}
 		if (InKeysState.IsRightShiftDown())
 		{
-			ViewportClient->InputKey(FInputKeyEventArgs(this, 0, EKeys::RightShift, IE_Pressed));
+			ViewportClient->InputKey(FInputKeyEventArgs(this, DefaultInputDevice, EKeys::RightShift, IE_Pressed, Timestamp));
 		}
 	}
 }
@@ -303,8 +305,14 @@ void FSceneViewport::ProcessAccumulatedPointerInput()
 	{
 		const float DeltaTime = FApp::GetDeltaTime();
 		FInputDeviceId DefaultInputDevice = IPlatformInputDeviceMapper::Get().GetDefaultInputDevice();
-		ViewportClient->InputAxis( this, DefaultInputDevice, EKeys::MouseX, MouseDelta.X, DeltaTime, NumMouseSamplesX );
-		ViewportClient->InputAxis( this, DefaultInputDevice, EKeys::MouseY, MouseDelta.Y, DeltaTime, NumMouseSamplesY );
+		// TODO: Populate this with a mouse timestamp
+		const uint64 MoseRecentMouseTimestamp = 0u;
+		
+		FInputKeyEventArgs MouseXArgs (this, DefaultInputDevice, EKeys::MouseX, MouseDelta.X, DeltaTime, NumMouseSamplesX, MoseRecentMouseTimestamp);
+		ViewportClient->InputAxis(MouseXArgs);
+
+		FInputKeyEventArgs MouseYArgs (this, DefaultInputDevice, EKeys::MouseY, MouseDelta.Y, DeltaTime, NumMouseSamplesY, MoseRecentMouseTimestamp);
+		ViewportClient->InputAxis(MouseYArgs);
 	}
 
 	if ( bCursorHiddenDueToCapture )
@@ -538,7 +546,7 @@ FReply FSceneViewport::OnMouseButtonDown( const FGeometry& InGeometry, const FPo
 		if (!HasFocus())
 		{
 			FModifierKeysState KeysState = FSlateApplication::Get().GetModifierKeys();
-			ApplyModifierKeys(KeysState);
+			ApplyModifierKeys(KeysState, InMouseEvent.GetEventTimestamp());
 		}
 
 		const bool bTemporaryCapture = 
@@ -553,7 +561,7 @@ FReply FSceneViewport::OnMouseButtonDown( const FGeometry& InGeometry, const FPo
 		// Process the mouse event
 		if (bTemporaryCapture || bProcessInputPrimary)
 		{
-			if (!ViewportClient->InputKey(FInputKeyEventArgs(this, InMouseEvent.GetUserIndex(), InMouseEvent.GetEffectingButton(), IE_Pressed, 1.0f, InMouseEvent.IsTouchEvent())))
+			if (!ViewportClient->InputKey(FInputKeyEventArgs(this, InMouseEvent.GetInputDeviceId(), InMouseEvent.GetEffectingButton(), IE_Pressed, 1.0f, InMouseEvent.IsTouchEvent(), InMouseEvent.GetEventTimestamp())))
 			{
 				CurrentReplyState = FReply::Unhandled();
 			}
@@ -662,7 +670,7 @@ FReply FSceneViewport::OnMouseButtonUp( const FGeometry& InGeometry, const FPoin
 	
 	if( ViewportClient && GetSizeXY() != FIntPoint::ZeroValue )
 	{
-		if (!ViewportClient->InputKey(FInputKeyEventArgs(this, InMouseEvent.GetUserIndex(), InMouseEvent.GetEffectingButton(), IE_Released, 1.0f, InMouseEvent.IsTouchEvent())))
+		if (!ViewportClient->InputKey(FInputKeyEventArgs(this, InMouseEvent.GetInputDeviceId(), InMouseEvent.GetEffectingButton(), IE_Released, 1.0f, InMouseEvent.IsTouchEvent(), InMouseEvent.GetEventTimestamp())))
 		{
 			CurrentReplyState = FReply::Unhandled(); 
 		}
@@ -797,9 +805,9 @@ FReply FSceneViewport::OnMouseWheel( const FGeometry& InGeometry, const FPointer
 		FKey const ViewportClientKey = InMouseEvent.GetWheelDelta() < 0 ? EKeys::MouseScrollDown : EKeys::MouseScrollUp;
 
 		// Pressed and released should be sent
-		ViewportClient->InputKey(FInputKeyEventArgs(this, InMouseEvent.GetUserIndex(), ViewportClientKey, IE_Pressed, 1.0f, InMouseEvent.IsTouchEvent()));
-		ViewportClient->InputKey(FInputKeyEventArgs(this, InMouseEvent.GetUserIndex(), ViewportClientKey, IE_Released, 1.0f, InMouseEvent.IsTouchEvent()));
-		ViewportClient->InputAxis(this, InMouseEvent.GetInputDeviceId(), EKeys::MouseWheelAxis, InMouseEvent.GetWheelDelta(), FApp::GetDeltaTime());
+		ViewportClient->InputKey(FInputKeyEventArgs(this, InMouseEvent.GetInputDeviceId(), ViewportClientKey, IE_Pressed, 1.0f, InMouseEvent.IsTouchEvent(), InMouseEvent.GetEventTimestamp()));
+		ViewportClient->InputKey(FInputKeyEventArgs(this, InMouseEvent.GetInputDeviceId(), ViewportClientKey, IE_Released, 1.0f, InMouseEvent.IsTouchEvent(), InMouseEvent.GetEventTimestamp()));
+		ViewportClient->InputAxis(FInputKeyEventArgs(this, InMouseEvent.GetInputDeviceId(), EKeys::MouseWheelAxis, InMouseEvent.GetWheelDelta(), FApp::GetDeltaTime(), 1, InMouseEvent.GetEventTimestamp()));
 	}
 	return CurrentReplyState;
 }
@@ -824,7 +832,7 @@ FReply FSceneViewport::OnMouseButtonDoubleClick( const FGeometry& InGeometry, co
 		// Switch to the viewport clients world before processing input
 		FScopedConditionalWorldSwitcher WorldSwitcher( ViewportClient );
 
-		if (!ViewportClient->InputKey(FInputKeyEventArgs(this, InMouseEvent.GetUserIndex(), InMouseEvent.GetEffectingButton(), IE_DoubleClick, 1.0f, InMouseEvent.IsTouchEvent())))
+		if (!ViewportClient->InputKey(FInputKeyEventArgs(this, InMouseEvent.GetInputDeviceId(), InMouseEvent.GetEffectingButton(), IE_DoubleClick, 1.0f, InMouseEvent.IsTouchEvent(), InMouseEvent.GetEventTimestamp())))
 		{
 			CurrentReplyState = FReply::Unhandled(); 
 		}
@@ -1077,7 +1085,7 @@ FReply FSceneViewport::OnKeyDown( const FGeometry& InGeometry, const FKeyEvent&
 			// Switch to the viewport clients world before processing input
 			FScopedConditionalWorldSwitcher WorldSwitcher(ViewportClient);
 
-			if (!ViewportClient->InputKey(FInputKeyEventArgs(this, InKeyEvent.GetInputDeviceId(), Key, InKeyEvent.IsRepeat() ? IE_Repeat : IE_Pressed, 1.0f, false)))
+			if (!ViewportClient->InputKey(FInputKeyEventArgs(this, InKeyEvent.GetInputDeviceId(), Key, InKeyEvent.IsRepeat() ? IE_Repeat : IE_Pressed, 1.0f, false, InKeyEvent.GetEventTimestamp())))
 			{
 				CurrentReplyState = FReply::Unhandled();
 			}
@@ -1105,7 +1113,7 @@ FReply FSceneViewport::OnKeyUp( const FGeometry& InGeometry, const FKeyEvent& In
 			// Switch to the viewport clients world before processing input
 			FScopedConditionalWorldSwitcher WorldSwitcher(ViewportClient);
 
-			if (!ViewportClient->InputKey(FInputKeyEventArgs(this, InKeyEvent.GetInputDeviceId(), Key, IE_Released, 1.0f, false)))
+			if (!ViewportClient->InputKey(FInputKeyEventArgs(this, InKeyEvent.GetInputDeviceId(), Key, IE_Released, 1.0f, false, InKeyEvent.GetEventTimestamp())))
 			{
 				CurrentReplyState = FReply::Unhandled();
 			}
@@ -1134,7 +1142,16 @@ FReply FSceneViewport::OnAnalogValueChanged(const FGeometry& MyGeometry, const F
 			// Switch to the viewport clients world before processing input
 			FScopedConditionalWorldSwitcher WorldSwitcher(ViewportClient);
 
-			if (!ViewportClient->InputAxis(this, InAnalogInputEvent.GetInputDeviceId(), Key, Key == EKeys::Gamepad_RightY ? -InAnalogInputEvent.GetAnalogValue() : InAnalogInputEvent.GetAnalogValue(), FApp::GetDeltaTime(), 1, Key.IsGamepadKey()))
+			FInputKeyEventArgs Args(
+				this,
+				InAnalogInputEvent.GetInputDeviceId(),
+				Key,
+				Key == EKeys::Gamepad_RightY ? -InAnalogInputEvent.GetAnalogValue() : InAnalogInputEvent.GetAnalogValue(),
+				FApp::GetDeltaTime(),
+				1,
+				InAnalogInputEvent.GetEventTimestamp());
+
+			if (!ViewportClient->InputAxis(Args))
 			{
 				CurrentReplyState = FReply::Unhandled();
 			}
diff --git a/Engine/Source/Runtime/Engine/Private/UnrealClient.cpp b/Engine/Source/Runtime/Engine/Private/UnrealClient.cpp
--- a/Engine/Source/Runtime/Engine/Private/UnrealClient.cpp
+++ b/Engine/Source/Runtime/Engine/Private/UnrealClient.cpp
@@ -2539,6 +2539,26 @@ void FViewportClient::RequestInvalidateHitProxy(FViewport* Viewport)
 	Viewport->InvalidateHitProxy();
 }
 
+PRAGMA_DISABLE_DEPRECATION_WARNINGS
+bool FViewportClient::InputAxis(
+		FViewport* Viewport,
+		FInputDeviceId InputDevice,
+		FKey Key,
+		float Delta,
+		float DeltaTime,
+		int32 NumSamples,
+		bool bGamepad)
+{
+	// Note: Set the timestamp to 0u here to represent an "invalid" timestamp, since this older input axis
+	// function did not account for timestamps. If we initalized it to the current cycle count here,
+	// then deprecated code paths would always be considered "newer" then anything calling the newer API, where the timestamp
+	// is calculated when the actual slate input event is constructed. 
+	FInputKeyEventArgs Args(Viewport, InputDevice, Key, Delta, DeltaTime, NumSamples, 0u);
+
+	return InputAxis(Args);
+}
+PRAGMA_ENABLE_DEPRECATION_WARNINGS
+
 TOptional<TSharedRef<SWidget>> FViewportClient::MapCursor(FViewport* Viewport, const FCursorReply& CursorReply)
 {
 	return TOptional<TSharedRef<SWidget>>();
diff --git a/Engine/Source/Runtime/Engine/Public/InputKeyEventArgs.h b/Engine/Source/Runtime/Engine/Public/InputKeyEventArgs.h
--- a/Engine/Source/Runtime/Engine/Public/InputKeyEventArgs.h
+++ b/Engine/Source/Runtime/Engine/Public/InputKeyEventArgs.h
@@ -9,32 +9,147 @@
 
 class FViewport;
 
+/**
+ * Input key event arguments. This is data that is required for the viewport client to
+ * process input via its InputKey/InputAxis functions.
+ *
+ * It represents a single "input event" which has happened which we want game or editor code
+ * to be able to process. This may accumulate several "raw" input events from the message
+ * handler into one with multiple sample readings, or it could be just one single key press event.
+ *
+ * This data can represent gamepad, keyboard, mouse, or touch data. 
+ */
 struct FInputKeyEventArgs
 {
 public:
-	FInputKeyEventArgs(FViewport* InViewport, int32 InControllerId, FKey InKey, EInputEvent InEvent)
-		: Viewport(InViewport)
-		, ControllerId(InControllerId)
-		, InputDevice(FInputDeviceId::CreateFromInternalId(InControllerId))
-		, Key(InKey)
-		, Event(InEvent)
-		, AmountDepressed(1.0f)
-		, bIsTouchEvent(false)
+
+	/**
+	 * Construct a FInputKeyEventArgs based on data acquired from the Slate input events.
+	 * These event args are used to translate from Slate input event arguments into a
+	 * standardized form for the ViewportClient and then the rest of the gameplay framework.
+	 * 
+	 * @param InViewport		The viewport which the axis movement is from.
+	 * @param InInputDevice		The input device that triggered this axis movement
+	 * @param InKey				The name of the axis which moved.
+	 * @param InEvent			The key event that these args represent.
+	 * @param InAmountDepressed	For axis events, the amount that the key is depressed
+	 * @param bInIsTouchEvent	True if this event represents a touch input
+	 * @param InEventTimestamp	The timestamp of when this input event was originally polled in the form of CPU cycles.
+	 */
+	FInputKeyEventArgs(
+		FViewport* InViewport,
+		const FInputDeviceId InInputDevice,
+		const FKey& InKey,
+		const EInputEvent InEvent,
+		const float InAmountDepressed,
+		const bool bInIsTouchEvent,
+		const uint64 InEventTimestamp)
+	: Viewport(InViewport)
+	, InputDevice(InInputDevice)
+	, Key(InKey)
+	, Event(InEvent)
+	, AmountDepressed(InAmountDepressed)
+	, bIsTouchEvent(bInIsTouchEvent)
+	, EventTimestamp(InEventTimestamp)
 	{
+		// Populate the legacy ControllerId based on the newer FPlatformUserId. This is used in some cases in PIE
+		// to pass input events along to the next viewport client and "fake" that the input came from that viewports
+		// primary player.
+		const FPlatformUserId UserID = IPlatformInputDeviceMapper::Get().GetUserForInputDevice(InInputDevice);
+		ControllerId = FPlatformMisc::GetUserIndexForPlatformUser(UserID);
 	}
 
+	UE_DEPRECATED(5.6, "Use the newer constructors which takes a timestamp instead.")
+	FInputKeyEventArgs(
+	FViewport* InViewport,
+	FInputDeviceId InInputDevice,
+	FKey InKey,
+	EInputEvent InEvent,
+	float InAmountDepressed,
+	bool bInIsTouchEvent)
+		: FInputKeyEventArgs(InViewport,
+		InInputDevice,
+		InKey,
+		InEvent,
+		InAmountDepressed,
+		bInIsTouchEvent,
+		/*EventTimestamp*/0u)
+	{}
+
+	/**
+	 * Constrcut a FInputKeyEventArgs based on data acquired from the Slate input events.
+	 * These event args are used to translate from Slate input event arguments into a
+	 * standardized form for the ViewportClient and then the rest of the gameplay framework.
+	 * 
+	 * @param InViewport		The viewport which the axis movement is from.
+	 * @param InInputDevice		The input device that triggered this axis movement
+	 * @param InKey				The name of the axis which moved.
+	 * @param InEvent			The key event that these args represent.
+	 * @param InEventTimestamp	The timestamp of when this input event was originally polled in the form of CPU cycles.
+	 */
+	FInputKeyEventArgs(
+		FViewport* InViewport,
+		const FInputDeviceId InInputDevice,
+		const FKey& InKey,
+		const EInputEvent InEvent,
+		const uint64 InEventTimestamp)
+		: FInputKeyEventArgs(
+		 InViewport,
+		 InInputDevice, 
+		 InKey, 
+		 InEvent,
+		 /*AmountDepressed*/1.0f,
+		 /*bIsTouch*/ false,
+		 InEventTimestamp)
+	{
+	}
+
+	UE_DEPRECATED(5.6, "Use the newer constructors which takes a timestamp instead.")
 	FInputKeyEventArgs(FViewport* InViewport, FInputDeviceId InInputDevice, FKey InKey, EInputEvent InEvent)
-		: Viewport(InViewport)
-		, InputDevice(InInputDevice)
-		, Key(InKey)
-		, Event(InEvent)
-		, AmountDepressed(1.0f)
-		, bIsTouchEvent(false)
+	: FInputKeyEventArgs(
+	 InViewport,
+	 InInputDevice, 
+	 InKey, 
+	 InEvent,
+	 /*AmountDepressed*/1.0f,
+	 /*bIsTouch*/ false,
+	 0u)
 	{
-		FPlatformUserId UserID = IPlatformInputDeviceMapper::Get().GetUserForInputDevice(InInputDevice);
-		ControllerId = FPlatformMisc::GetUserIndexForPlatformUser(UserID);
 	}
 
+	/**
+	 * Construct a FInputKeyEventArgs based on the old params of the ViewportClient InputAxis function. 
+	 * 
+	 * @param InViewport		The viewport which the axis movement is from.
+	 * @param InInputDevice		The input device that triggered this axis movement
+	 * @param InKey				The name of the axis which moved.
+	 * @param InDelta			The axis movement delta.
+	 * @param InDeltaTime		The time since the last axis update.
+	 * @param InNumSamples		The number of device samples that contributed to this Delta, useful for things like smoothing
+	 * @param InEventTimestamp	The timestamp of when this input event was originally polled in the form of CPU cycles.
+	 */
+	FInputKeyEventArgs(
+		FViewport* InViewport,
+		const FInputDeviceId InInputDevice,
+		const FKey& InKey,
+		const float InDelta,
+		const float InDeltaTime,
+		const int32 InNumSamples,
+		const uint64 InEventTimestamp)
+		: FInputKeyEventArgs(
+			InViewport,
+			InInputDevice,
+			InKey,
+			IE_Axis,
+			/*AmountDepressed*/InDelta,
+			/*bIsTouchEvent*/ false,
+			InEventTimestamp)
+	{
+		DeltaTime = InDeltaTime;
+		NumSamples = InNumSamples;
+	}
+	
+	UE_DEPRECATED(5.6, "Use the newer constructors which take in a FInputDeviceId and a timestamp instead.")
 	FInputKeyEventArgs(FViewport* InViewport, int32 InControllerId, FKey InKey, EInputEvent InEvent, float InAmountDepressed, bool bInIsTouchEvent)
 		: Viewport(InViewport)
 		, ControllerId(InControllerId)
@@ -45,35 +160,93 @@ public:
 		, bIsTouchEvent(bInIsTouchEvent)
 	{
 	}
-	
-	FInputKeyEventArgs(FViewport* InViewport, FInputDeviceId InInputDevice, FKey InKey, EInputEvent InEvent, float InAmountDepressed, bool bInIsTouchEvent)
-		: Viewport(InViewport)
-		, InputDevice(InInputDevice)
-		, Key(InKey)
-		, Event(InEvent)
-		, AmountDepressed(InAmountDepressed)
-		, bIsTouchEvent(bInIsTouchEvent)
+
+	UE_DEPRECATED(5.6, "Use the newer constructors which take in a FInputDeviceId and a timestamp instead.")
+	FInputKeyEventArgs(FViewport* InViewport, int32 InControllerId, FKey InKey, EInputEvent InEvent)
+		: FInputKeyEventArgs(
+			InViewport,
+			FInputDeviceId::CreateFromInternalId(InControllerId),
+			InKey,
+			InEvent,
+			/*EventTimestamp*/ 0u)
 	{
-		FPlatformUserId UserID = IPlatformInputDeviceMapper::Get().GetUserForInputDevice(InInputDevice);
-		ControllerId = FPlatformMisc::GetUserIndexForPlatformUser(UserID);
+		ControllerId = InControllerId;
 	}
-
+	
+	/**
+	 * @return True if this input event is for a Gamepay key, or should be treated as such.
+	 * 
+	 * This may be true for non-gamepad FKeys if the input event is simulated input.
+	 */
 	bool IsGamepad() const { return Key.IsGamepadKey(); }
 
+	/**
+	 * @return The platform user which this input event originated from based on its input device.
+	 */
+	FPlatformUserId GetPlatformUser() const
+	{
+		return IPlatformInputDeviceMapper::Get().GetUserForInputDevice(InputDevice);
+	}
+
 public:
+	
+	/**
+	 * The viewport from which this key event originated from.
+	 */
+	FViewport* Viewport = nullptr;
 
-	// The viewport which the key event is from.
-	FViewport* Viewport;
 	// The controller which the key event is from.
 	int32 ControllerId;
-	// The input device which this key event is from
-	FInputDeviceId InputDevice;
-	// The type of event which occurred.
-	FKey Key;
-	// The type of event which occurred.
-	EInputEvent Event;
-	// For analog keys, the depression percent.
+	
+	/**
+	 * The input device which this event originated from.
+	 */
+	FInputDeviceId InputDevice = INPUTDEVICEID_NONE;
+
+	/**
+	 * The Key that this input event is for
+	 */
+	FKey Key = {};
+	
+	/**
+	 * The type of event which occurred.
+	 */
+	EInputEvent Event = IE_MAX;
+	
+	/**
+	 * The value that this input event represents.
+	 * 
+	 * For analog keys, the depression percent.
+	 */
 	float AmountDepressed;
-	// input came from a touch surface.This may be a faked mouse button from touch.
-	bool bIsTouchEvent;
+
+	/** The time between the previous frame and the current one */
+	float DeltaTime = 1.0f / 60.f;
+
+	/**
+	 * For analog key events. The number of samples of analog input
+	 * which are contained in this events InputValue.
+	 */
+	int32 NumSamples = 1;
+	
+	/**
+	 * True if this input event originated from a touch surface.
+	 * 
+	 * Note: This may be set to true for simulated touch inputs from things like a mouse button.
+	 */
+	bool bIsTouchEvent = false;
+	
+	/**
+	 * The timestamp of when this input event was originally polled.
+	 * 
+	 * This data should be set to be as representative as possible of the time of when the
+	 * input event originated, most of the time from the raw Slate input event (FInputEvent)
+	 * on the message handler.
+	 *
+	 * This timestamp is in terms of CPU Cycles, gathered from the platform's high resolution clock
+	 * via the FPlatformTime::Cycles64() function.
+	 *
+	 * A timestamp value of 0u means that this event was not initalized with a valid timestamp
+	 */
+	uint64 EventTimestamp = 0u;
 };
diff --git a/Engine/Source/Runtime/Engine/Public/Slate/SceneViewport.h b/Engine/Source/Runtime/Engine/Public/Slate/SceneViewport.h
--- a/Engine/Source/Runtime/Engine/Public/Slate/SceneViewport.h
+++ b/Engine/Source/Runtime/Engine/Public/Slate/SceneViewport.h
@@ -335,7 +335,7 @@ private:
 	 *
 	 * @param InKeysState	The key state containing the states of the modifier keys
 	 */
-	ENGINE_API void ApplyModifierKeys( const FModifierKeysState& InKeysState );
+	ENGINE_API void ApplyModifierKeys(const FModifierKeysState& InKeysState, const uint64 Timestamp);
 
 	/** Utility function to create an FReply that properly gets Focus and capture based on the settings*/
 	ENGINE_API FReply AcquireFocusAndCapture(FIntPoint MousePosition, EFocusCause FocusCause = EFocusCause::SetDirectly);
diff --git a/Engine/Source/Runtime/Engine/Public/ViewportClient.h b/Engine/Source/Runtime/Engine/Public/ViewportClient.h
--- a/Engine/Source/Runtime/Engine/Public/ViewportClient.h
+++ b/Engine/Source/Runtime/Engine/Public/ViewportClient.h
@@ -39,6 +39,7 @@ public:
 	/**
 	 * Check a key event received by the viewport.
 	 * If the viewport client uses the event, it should return true to consume it.
+	 * 
 	 * @param	EventArgs - The Input event args.
 	 * @return	True to consume the key event, false to pass it on.
 	 */
@@ -56,14 +57,33 @@ public:
 	 * @param	bGamepad - input came from gamepad (ie xbox controller)
 	 * @return	True to consume the axis movement, false to pass it on.
 	 */
-	virtual bool InputAxis(FViewport* Viewport, FInputDeviceId InputDevice, FKey Key, float Delta, float DeltaTime, int32 NumSamples = 1, bool bGamepad = false) { return false; }
+	UE_DEPRECATED(5.6, "Use the version of InputAxis which takes in a FInputKeyEventArgs instead.")
+	ENGINE_API virtual bool InputAxis(
+	 	FViewport* Viewport,
+	 	FInputDeviceId InputDevice,
+	 	FKey Key,
+	 	float Delta,
+	 	float DeltaTime,
+	 	int32 NumSamples = 1,
+	 	bool bGamepad = false) final;
 
+	/**
+	 * Check an axis movement received by the viewport. If this viewport client uses this input
+	 * and wants to consume it (stop it from being passed along to further viewports) then it should
+	 * return true.
+	 * 
+	 * @param Args The input event arguments for this axis movement.
+	 * @return True to consume the axis movement, false to pass it on.
+	 */
+	virtual bool InputAxis(const FInputKeyEventArgs& Args) { return false; }
+	
 	/**
 	 * Check a character input received by the viewport.
 	 * If the viewport client uses the character, it should return true to consume it.
 	 * @param	Viewport - The viewport which the axis movement is from.
 	 * @param	ControllerId - The controller which the axis movement is from.
 	 * @param	Character - The character.
+	 * // TODO_BH: Timestamp
 	 * @return	True to consume the character, false to pass it on.
 	 */
 	virtual bool InputChar(FViewport* Viewport,int32 ControllerId,TCHAR Character) { return false; }
@@ -79,6 +99,7 @@ public:
 	 * @param	Force - How hard the touch is
 	 * @param	DeviceTimestamp - Timestamp of the event
 	 * @param	TouchpadIndex - For devices with multiple touchpads, this is the index of which one
+	 * // TODO_BH: Timestamp
 	 * @return	True to consume the key event, false to pass it on.
 	 */
 	virtual bool InputTouch(FViewport* Viewport, int32 ControllerId, uint32 Handle, ETouchType::Type Type, const FVector2D& TouchLocation, float Force, FDateTime DeviceTimestamp, uint32 TouchpadIndex) { return false; }
@@ -89,6 +110,7 @@ public:
 	 * @param	Viewport - The viewport which the event is from.
 	 * @param	GestureType - @todo desc
 	 * @param	GestureDelta - @todo desc
+	 * // TODO_BH: Timestamp
 	 * @return	True to consume the gesture event, false to pass it on.
 	 */
 	virtual bool InputGesture(FViewport* Viewport, EGestureEvent GestureType, const FVector2D& GestureDelta, bool bIsDirectionInvertedFromDevice) { return false; }
@@ -102,12 +124,14 @@ public:
 	 * @param RotationRate	How fast the tilt is changing
 	 * @param Gravity		Describes the current gravity of the device
 	 * @param Acceleration  Describes the acceleration of the device
+	 * // TODO_BH: Timestamp
 	 * @return	True to consume the motion event, false to pass it on.
 	 */
 	virtual bool InputMotion(FViewport* Viewport, int32 ControllerId, const FVector& Tilt, const FVector& RotationRate, const FVector& Gravity, const FVector& Acceleration) { return false; }
 
 	virtual void SetIsSimulateInEditorViewport(bool bInIsSimulateInEditorViewport) { };
 
+	// TODO: Deprecate? Nothing uses this.
 	virtual bool WantsPollingMouseMovement(void) const { return true; }
 
 	virtual void MouseEnter( FViewport* Viewport,int32 x, int32 y ) {}
@@ -325,7 +349,7 @@ public:
 		, bShouldUpdateDPIScale(true)
 	{}
 
-	ENGINE_API virtual ~FCommonViewportClient();
+	ENGINE_API virtual ~FCommonViewportClient() override;
 
 	/** Tells this viewport to update editor dpi scale when needed */
 	ENGINE_API void RequestUpdateDPIScale();
diff --git a/Engine/Source/Runtime/SlateCore/Public/Input/Events.h b/Engine/Source/Runtime/SlateCore/Public/Input/Events.h
--- a/Engine/Source/Runtime/SlateCore/Public/Input/Events.h
+++ b/Engine/Source/Runtime/SlateCore/Public/Input/Events.h
@@ -166,6 +166,7 @@ public:
 		, UserIndex(0)
 		, InputDeviceId(INPUTDEVICEID_NONE)
 		, EventPath(nullptr)
+		, EventTimestamp(FPlatformTime::Cycles64())	/* Initalize the timestamp of event input event to the current CPU cycle */
 	{ }
 
 	/**
@@ -180,6 +181,7 @@ public:
 		, UserIndex(InUserIndex)
 		, InputDeviceId(FInputDeviceId::CreateFromInternalId(InUserIndex))
 		, EventPath(nullptr)
+		, EventTimestamp(FPlatformTime::Cycles64())	/* Initalize the timestamp of event input event to the current CPU cycle */
 	{ }
 
 	FInputEvent(const FModifierKeysState& InModifierKeys, const FInputDeviceId InDeviceId, const bool bInIsRepeat)
@@ -187,6 +189,7 @@ public:
 		, bIsRepeat(bInIsRepeat)
 		, InputDeviceId(InDeviceId)
 		, EventPath(nullptr)
+		, EventTimestamp(FPlatformTime::Cycles64())	/* Initalize the timestamp of event input event to the current CPU cycle */
 	{
 		// Set the User Index to the PlatformUser ID by default for backwards compatibility
 		UserIndex = GetPlatformUserId().GetInternalId();
@@ -347,6 +350,22 @@ public:
 		return IPlatformInputDeviceMapper::Get().GetUserForInputDevice(InputDeviceId);
 	}
 
+	/**
+	 * @return The CPU cycle timestamp of when this input event was created.  
+	 */
+	uint64 GetEventTimestamp() const
+	{
+		return EventTimestamp;
+	}
+
+	/**
+	 * @return The amount of time in milliseconds which have passes since this input event was constructed
+	 */
+	double GetMillisecondsSinceEvent() const
+	{
+		return FPlatformTime::ToMilliseconds64(FPlatformTime::Cycles64() - EventTimestamp);
+	}
+
 	/** The event path provides additional context for handling */
 	SLATECORE_API FGeometry FindGeometry(const TSharedRef<SWidget>& WidgetToFind) const;
 
@@ -387,6 +406,11 @@ protected:
 
 	// Events are sent along paths. See (GetEventPath).
 	const FWidgetPath* EventPath;
+
+	// The timestamp of when this input event was recorded. This is the
+	// cycle at which this event was recorded, populated by FPlatformTime::Cycles64()
+	// at the time of construction.
+	uint64 EventTimestamp;
 };
 
 template<>
