PR #12917: Fix FindIndexByName to use FString for FindCustomShowFlagByName lookup

diff --git a/Engine/Source/Runtime/Engine/Private/ShowFlags.cpp b/Engine/Source/Runtime/Engine/Private/ShowFlags.cpp
--- a/Engine/Source/Runtime/Engine/Private/ShowFlags.cpp
+++ b/Engine/Source/Runtime/Engine/Private/ShowFlags.cpp
@@ -215,7 +215,7 @@ int32 FEngineShowFlags::FindIndexByName(const TCHAR* Name, const TCHAR* CommaSep
 
 		#include "ShowFlagsValues.inl" // IWYU pragma: keep
 
-		ECustomShowFlag CustomFlagIndex = FindCustomShowFlagByName(Name);
+		ECustomShowFlag CustomFlagIndex = FindCustomShowFlagByName(Search);
 		if (CustomFlagIndex != ECustomShowFlag::None)
 		{
 			return SF_FirstCustom + (uint32)CustomFlagIndex;
