Fixing issue where a player's Character could become unable to jump after spamming jump input under poor network conditions, until they run off a ledge or otherwise change from walking mode.

diff --git a/Engine/Source/Runtime/Engine/Private/Components/CharacterMovementComponent.cpp b/Engine/Source/Runtime/Engine/Private/Components/CharacterMovementComponent.cpp
--- a/Engine/Source/Runtime/Engine/Private/Components/CharacterMovementComponent.cpp
+++ b/Engine/Source/Runtime/Engine/Private/Components/CharacterMovementComponent.cpp
@@ -376,6 +376,13 @@ namespace CharacterMovementCVars
 		TEXT("Whether to apply character impact physics force as impulse when the object hit is not moving."),
 		ECVF_Default);
 
+	static bool bResetJumpStateDuringResimIfNotFalling = true;
+	FAutoConsoleVariableRef CVarResetJumpStateDuringResimIfNotFalling(
+		TEXT("p.ResetJumpStateDuringResimIfNotFalling"),
+		bResetJumpStateDuringResimIfNotFalling,
+		TEXT("If true, we'll reset jump state during a correction & resimulation if not falling (e.g. walking). This fixes a bug that locks jump input, which could occur under poor net conditions."),
+		ECVF_Default);
+
 #if CSV_PROFILER_STATS
 	bool bClientRecordMovePackedRpcStatsToCsv = false;
 	FAutoConsoleVariableRef CVarClientRecordMovePackedRpcStatsToCsv(
@@ -12971,13 +12978,21 @@ void FSavedMove_Character::PrepMoveFor(ACharacter* Character)
 	}
 
 	Character->GetCharacterMovement()->bForceMaxAccel = bForceMaxAccel;
-	Character->bWasJumping = bWasJumping;
-	Character->JumpKeyHoldTime = JumpKeyHoldTime;
-	Character->JumpForceTimeRemaining = JumpForceTimeRemaining;
-	Character->JumpMaxCount = JumpMaxCount;
-	Character->JumpCurrentCount = JumpCurrentCount;
-	Character->JumpCurrentCountPreJump = JumpCurrentCount;
 
+	if (CharacterMovementCVars::bResetJumpStateDuringResimIfNotFalling && !Character->GetCharacterMovement()->IsFalling())
+	{
+		Character->ResetJumpState();
+	}
+	else
+	{
+		Character->bWasJumping = bWasJumping;
+		Character->JumpKeyHoldTime = JumpKeyHoldTime;
+		Character->JumpForceTimeRemaining = JumpForceTimeRemaining;
+		Character->JumpMaxCount = JumpMaxCount;
+		Character->JumpCurrentCount = JumpCurrentCount;
+		Character->JumpCurrentCountPreJump = JumpCurrentCount;
+	}
+	
 	StartPackedMovementMode = Character->GetCharacterMovement()->PackNetworkMovementMode();
 }
 
