Skeletal Mesh: Fixes for creating and adding sockets via BP
- When calling AddSocket and either SocketName or BoneName are None, they get filled in with default values ("Socket" and root bone, respectively).
- RenameSocket no longer checks whether the old name was None, so it's possible to rename a freshly created socket before adding it.
- Fixed a crash in persona when encountering broken socket that had no valid outer object.

#rnx

diff --git a/Engine/Source/Editor/Persona/Private/AssetSearchBoxUtilPersona.cpp b/Engine/Source/Editor/Persona/Private/AssetSearchBoxUtilPersona.cpp
--- a/Engine/Source/Editor/Persona/Private/AssetSearchBoxUtilPersona.cpp
+++ b/Engine/Source/Editor/Persona/Private/AssetSearchBoxUtilPersona.cpp
@@ -52,19 +52,21 @@ void SAssetSearchBoxForBones::Construct( const FArguments& InArgs, const class U
 	{
 		Skeleton = Cast<const USkeleton>(Outer);
 	}
-	if( Skeleton && InArgs._IncludeSocketsForSuggestions.Get() )
+	if( Skeleton )
 	{
-		for( auto SocketIt=Skeleton->Sockets.CreateConstIterator(); SocketIt; ++SocketIt )
+		if (InArgs._IncludeSocketsForSuggestions.Get())
 		{
-			PossibleSuggestions.Add(FAssetSearchBoxSuggestion::MakeSimpleSuggestion((*SocketIt)->SocketName.ToString()));
+			for( auto SocketIt=Skeleton->Sockets.CreateConstIterator(); SocketIt; ++SocketIt )
+			{
+				PossibleSuggestions.Add(FAssetSearchBoxSuggestion::MakeSimpleSuggestion((*SocketIt)->SocketName.ToString()));
+			}
+		}
+		
+		const TArray<FMeshBoneInfo>& Bones = Skeleton->GetReferenceSkeleton().GetRefBoneInfo();
+		for( auto BoneIt=Bones.CreateConstIterator(); BoneIt; ++BoneIt )
+		{
+			PossibleSuggestions.Add(FAssetSearchBoxSuggestion::MakeSimpleSuggestion(BoneIt->Name.ToString()));
 		}
-	}
-	check(Skeleton);
-
-	const TArray<FMeshBoneInfo>& Bones = Skeleton->GetReferenceSkeleton().GetRefBoneInfo();
-	for( auto BoneIt=Bones.CreateConstIterator(); BoneIt; ++BoneIt )
-	{
-		PossibleSuggestions.Add(FAssetSearchBoxSuggestion::MakeSimpleSuggestion(BoneIt->Name.ToString()));
 	}
 
 	// create the asset search box
diff --git a/Engine/Source/Editor/SkeletalMeshEditor/Private/SkeletalMeshEditorSubsystem.cpp b/Engine/Source/Editor/SkeletalMeshEditor/Private/SkeletalMeshEditorSubsystem.cpp
--- a/Engine/Source/Editor/SkeletalMeshEditor/Private/SkeletalMeshEditorSubsystem.cpp
+++ b/Engine/Source/Editor/SkeletalMeshEditor/Private/SkeletalMeshEditorSubsystem.cpp
@@ -527,12 +527,6 @@ bool USkeletalMeshEditorSubsystem::RenameSocket(USkeletalMesh* SkeletalMesh, FNa
 		return false;
 	}
 
-	if (OldName == NAME_None)
-	{
-		UE_LOG(LogSkeletalMeshEditorSubsystem, Error, TEXT("RenameSocket: The OldName is None."));
-		return false;
-	}
-
 	if (NewName == NAME_None)
 	{
 		UE_LOG(LogSkeletalMeshEditorSubsystem, Error, TEXT("RenameSocket: The NewName is None."));
diff --git a/Engine/Source/Runtime/Engine/Private/SkeletalMesh.cpp b/Engine/Source/Runtime/Engine/Private/SkeletalMesh.cpp
--- a/Engine/Source/Runtime/Engine/Private/SkeletalMesh.cpp
+++ b/Engine/Source/Runtime/Engine/Private/SkeletalMesh.cpp
@@ -4580,77 +4580,84 @@ USkeletalMeshSocket* USkeletalMesh::FindSocket(FName InSocketName) const
 #if WITH_EDITOR
 void USkeletalMesh::AddSocket(USkeletalMeshSocket* InSocket, bool bAddToSkeleton /*= false*/)
 {
-	if (InSocket)
+	if (!InSocket)
 	{
-		if (InSocket->GetOuter() == this)
+		return;
+	}
+	
+	// The socket needs to be owned already by this skeletal mesh.
+	if (InSocket->GetOuter() != this)
+	{
+		UE_LOG(LogSkeletalMesh, Error, TEXT("Failed to add socket as the socket its outer should be %s but is %s."), *GetFullName(), *InSocket->GetOuter()->GetFullName());
+		return;
+	}
+
+	// If the socket was freshly created, which by default doesn't have a name, assign a default name to it now.
+	FName WantedSocketName;
+	if (InSocket->SocketName.IsNone())
+	{
+		static const FName BaseSocketName("Socket");
+		
+		int32 TestNumber = 0;
+		do
 		{
-			const FString SocketNameString = InSocket->SocketName.ToString();
-			const FString TrimmedNameString = SocketNameString.TrimStartAndEnd();
-			const FName TrimmedName = FName(*TrimmedNameString);
-			if (!Sockets.ContainsByPredicate([TrimmedName](const TObjectPtr<USkeletalMeshSocket>& Socket)
-			{
-				return Socket->SocketName == TrimmedName;
-			}))
-			{
-				const FReferenceSkeleton& ReferenceSkeleton = GetRefSkeleton();
-				if (ReferenceSkeleton.FindBoneIndex(InSocket->BoneName) != INDEX_NONE)
-				{				
-					if (!TrimmedNameString.IsEmpty())
-					{
-						const FScopedTransaction Transaction(LOCTEXT("AddSocket", "Add Socket"));
-						Modify();
-
-						// Update socket name in place
-						InSocket->SocketName = TrimmedName;					
-						Sockets.Add(InSocket);
-		    
-						if (bAddToSkeleton)
-						{
-							USkeleton* CurrentSkeleton = GetSkeleton();				
-							if (!CurrentSkeleton->Sockets.ContainsByPredicate([Name=InSocket->SocketName](const TObjectPtr<USkeletalMeshSocket>& Socket)
-							{
-								return Socket->SocketName == Name;
-							}))
-							{
-								const FReferenceSkeleton& SkeletonReferenceSkeleton = CurrentSkeleton->GetReferenceSkeleton();
-								if (SkeletonReferenceSkeleton.FindBoneIndex(InSocket->BoneName) != INDEX_NONE)
-								{
-									CurrentSkeleton->Modify();
-			    
-									USkeletalMeshSocket* NewSocket = DuplicateObject<USkeletalMeshSocket>(InSocket, CurrentSkeleton);
-									check(NewSocket);
-									CurrentSkeleton->Sockets.Add(NewSocket);
-								}
-								else
-								{
-									// Bone does not exists on Skeleton
-									UE_LOG(LogSkeletalMesh, Warning, TEXT("Failed to add socket to Skeleton (%s) as the provided bone name %s does not exist."), *CurrentSkeleton->GetFullName(), *InSocket->BoneName.ToString());
-								}
-							}
-						}
-					}
-					else
-					{
-						// Invalid socket name
-						UE_LOG(LogSkeletalMesh, Warning, TEXT("Failed to add socket as the provided socket name %s is invalid."), *TrimmedNameString);
-					}
-				}
-				else
-				{
-					// Bone does not exists
-					UE_LOG(LogSkeletalMesh, Warning, TEXT("Failed to add socket as the provided bone name %s does not exist."), *InSocket->BoneName.ToString());
-				}
-			}
-			else
-			{
-				// Socket already exists
-				UE_LOG(LogSkeletalMesh, Warning, TEXT("Failed to add socket as a socket with name %s already exist."), *InSocket->BoneName.ToString());
-			}
+			WantedSocketName = FName(BaseSocketName, TestNumber++);
+		} while (FindSocket(WantedSocketName) != nullptr);
+	}
+	else
+	{
+		// Make sure it's unique across all known sockets.
+		const FString SocketNameString = InSocket->SocketName.ToString();
+		const FString TrimmedNameString = SocketNameString.TrimStartAndEnd();
+		WantedSocketName = FName(*TrimmedNameString);
+		if (Sockets.ContainsByPredicate([WantedSocketName](const TObjectPtr<USkeletalMeshSocket>& Socket)
+		{
+			return Socket->SocketName == WantedSocketName;
+		}))
+		{
+			// Socket already exists
+			UE_LOG(LogSkeletalMesh, Error, TEXT("Failed to add socket as a socket with name %s already exist."), *InSocket->BoneName.ToString());
+			return;
 		}
-		else
+	}
+
+	// Check if the bone exists. If set to None, as would happen with a default created socket, then assign the socket to the root bone.
+	// If wanting to add to the skeleton, then the additional restriction of the bone name needing to exist on the skeleton is also required.
+	const FReferenceSkeleton& ReferenceSkeleton = bAddToSkeleton ? GetSkeleton()->GetReferenceSkeleton() : GetRefSkeleton();
+	FName WantedBoneName = InSocket->BoneName;
+	if (InSocket->BoneName.IsNone())
+	{
+		WantedBoneName = ReferenceSkeleton.GetBoneName(0);
+	}
+	else if (ReferenceSkeleton.FindBoneIndex(InSocket->BoneName) == INDEX_NONE)
+	{
+		UE_LOG(LogSkeletalMesh, Error, TEXT("Failed to add socket as the provided bone name %s does not exist."), *InSocket->BoneName.ToString());
+		return;
+	}
+
+	const FScopedTransaction Transaction(LOCTEXT("AddSocket", "Add Socket"));
+	
+	Modify();
+	InSocket->Modify();
+
+	InSocket->SocketName = WantedSocketName;
+	InSocket->BoneName = WantedBoneName;
+	
+	Sockets.Add(InSocket);
+
+	if (bAddToSkeleton)
+	{
+		USkeleton* CurrentSkeleton = GetSkeleton();				
+		if (!CurrentSkeleton->Sockets.ContainsByPredicate([Name=InSocket->SocketName](const TObjectPtr<USkeletalMeshSocket>& Socket)
+		{
+			return Socket->SocketName == Name;
+		}))
 		{
-			// Socket outer is not correct
-			UE_LOG(LogSkeletalMesh, Warning, TEXT("Failed to add socket as the socket its outer should be %s but is %s."), *GetFullName(), *InSocket->GetOuter()->GetFullName());
+			CurrentSkeleton->Modify();
+
+			USkeletalMeshSocket* NewSocket = DuplicateObject<USkeletalMeshSocket>(InSocket, CurrentSkeleton);
+			check(NewSocket);
+			CurrentSkeleton->Sockets.Add(NewSocket);
 		}
 	}
 }
