Fast-path for editor access of current target platform compressed data

diff --git a/Engine/Source/Runtime/Engine/Classes/Animation/AnimSequence.h b/Engine/Source/Runtime/Engine/Classes/Animation/AnimSequence.h
--- a/Engine/Source/Runtime/Engine/Classes/Animation/AnimSequence.h
+++ b/Engine/Source/Runtime/Engine/Classes/Animation/AnimSequence.h
@@ -806,6 +806,7 @@ protected:
 	ENGINE_API bool PollCacheDerivedData(const FIoHash& KeyHash) const;
 	ENGINE_API void EndCacheDerivedData(const FIoHash& KeyHash);
 
+	mutable FCompressedAnimSequence* CurrentPlatformData = nullptr; 
 	TMap<FIoHash, TUniquePtr<FCompressedAnimSequence>> DataByPlatformKeyHash;
 	TMap<FIoHash, TPimplPtr<UE::Anim::FAnimationSequenceAsyncCacheTask>> CacheTasksByKeyHash;
 	mutable FRWLock CompressedDataLock;
diff --git a/Engine/Source/Runtime/Engine/Private/Animation/AnimSequence.cpp b/Engine/Source/Runtime/Engine/Private/Animation/AnimSequence.cpp
--- a/Engine/Source/Runtime/Engine/Private/Animation/AnimSequence.cpp
+++ b/Engine/Source/Runtime/Engine/Private/Animation/AnimSequence.cpp
@@ -1944,6 +1944,12 @@ const FCompressedAnimSequence& UAnimSequence::GetCompressedData_Internal() const
 	if(!GetOutermost()->HasAnyPackageFlags(PKG_Cooked))
 	{
 		static thread_local FCompressedAnimSequence DefaultCompressedAnimSequence;
+
+		if (CurrentPlatformData)
+		{
+			return *CurrentPlatformData;
+		}
+		
 		const ITargetPlatform* RunningPlatform = GetTargetPlatformManagerRef().GetRunningTargetPlatform();
 
 		if(DataByPlatformKeyHash.Num())
@@ -1994,6 +2000,11 @@ const FCompressedAnimSequence& UAnimSequence::GetPlatformCompressedData(const FA
 		}
 		else
 		{
+			if (CurrentPlatformData)
+			{
+				return *CurrentPlatformData;
+			}
+			
 			const ITargetPlatform* RunningPlatform = GetTargetPlatformManagerRef().GetRunningTargetPlatform();
 			Hash = GetDerivedDataKeyHash(RunningPlatform);
 		}
@@ -2041,6 +2052,11 @@ FCompressedAnimSequence& UAnimSequence::GetPlatformCompressedData(const FAnimExt
 		}
 		else
 		{
+			if (CurrentPlatformData)
+			{
+				return *CurrentPlatformData;
+			}
+			
 			const ITargetPlatform* RunningPlatform = GetTargetPlatformManagerRef().GetRunningTargetPlatform();
 			Hash = GetDerivedDataKeyHash(RunningPlatform);
 		}
@@ -5054,6 +5070,12 @@ FIoHash UAnimSequence::BeginCacheDerivedData(const ITargetPlatform* TargetPlatfo
 		TargetData = DataByPlatformKeyHash.Emplace(KeyHash, MakeUnique<FCompressedAnimSequence>()).Get();
 	}
 
+	const ITargetPlatform* RunningPlatform = GetTargetPlatformManagerRef().GetRunningTargetPlatform();	
+	if (TargetPlatform == RunningPlatform)
+	{
+		CurrentPlatformData = nullptr;
+	}
+
 	// Reset the target compressed data, to mark it invalid
 	check(TargetData);
 	TargetData->Reset();
@@ -5196,6 +5218,8 @@ void UAnimSequence::ClearAllCompressionData()
 				It.RemoveCurrent();
 			}
 		}
+		
+		CurrentPlatformData = nullptr;
 	}
 }
 
@@ -5209,6 +5233,12 @@ void UAnimSequence::ClearCompressionData(const FIoHash& InKeyHash)
 		UE_COMPRESSED_DATA_WRITE_SCOPE(this);
 		check(DataByPlatformKeyHash.Contains(InKeyHash));
 		DataByPlatformKeyHash.Remove(InKeyHash);
+		
+		const ITargetPlatform* RunningPlatform = GetTargetPlatformManagerRef().GetRunningTargetPlatform();
+		if (GetDerivedDataKeyHash(RunningPlatform) == InKeyHash)
+		{
+			CurrentPlatformData = nullptr;
+		}
 	}
 }
 
@@ -5351,17 +5381,20 @@ void UAnimSequence::FinishAsyncTasks()
 			{
 				if (TaskData->IsValid(this, true))
 				{
-#if WITH_EDITOR
 					//This is only safe during sync anim compression
 					if (GetSkeleton())
 					{
 						SetSkeletonVirtualBoneGuid(GetSkeleton()->GetVirtualBoneGuid());
 					}
-#endif
+
 					if (bCompressedRunningPlatform)
 					{
 						FAssetRegistryModule& AssetRegistryModule = FModuleManager::LoadModuleChecked<FAssetRegistryModule>(TEXT("AssetRegistry"));
 						AssetRegistryModule.Get().AssetTagsFinalized(*this);
+
+						// Cache compressed data ptr for current platform to by-pass TMap lookup requiring data-hash retrieval/generation and locking
+						ensure(CurrentPlatformData == nullptr);
+						CurrentPlatformData = TaskData;
 					}
 				
 					check(TaskData->IsValid(this, true));
