Null-ptr check for static analysis warning

diff --git a/Engine/Source/Runtime/Engine/Private/Animation/AnimSequence.cpp b/Engine/Source/Runtime/Engine/Private/Animation/AnimSequence.cpp
--- a/Engine/Source/Runtime/Engine/Private/Animation/AnimSequence.cpp
+++ b/Engine/Source/Runtime/Engine/Private/Animation/AnimSequence.cpp
@@ -5366,44 +5366,47 @@ void UAnimSequence::FinishAsyncTasks()
 			It->Value->Wait();
 
 			FCompressedAnimSequence* TaskData = It->Value->GetTargetCompressedData();
-			const bool bCompressedRunningPlatform = It->Key == CurrentPlatformHash;
-			auto ResetData = [this, TaskData, It]()
-			{
-				TaskData->Reset();
-				DataByPlatformKeyHash.Remove(It->Key);
-			};
-
-			if (It->Value.Get()->WasCancelled())
-			{
-				ResetData();
-			}
-			else
-			{
-				if (TaskData->IsValid(this, true))
+			if(TaskData != nullptr)
+			{				
+				const bool bCompressedRunningPlatform = It->Key == CurrentPlatformHash;
+				auto ResetData = [this, TaskData, It]()
+				{
+					TaskData->Reset();
+					DataByPlatformKeyHash.Remove(It->Key);
+				};
+	
+				if (It->Value.Get()->WasCancelled())
+				{
+					ResetData();
+				}
+				else
 				{
-					//This is only safe during sync anim compression
-					if (GetSkeleton())
+					if (TaskData->IsValid(this, true))
 					{
-						SetSkeletonVirtualBoneGuid(GetSkeleton()->GetVirtualBoneGuid());
+						//This is only safe during sync anim compression
+						if (GetSkeleton())
+						{
+							SetSkeletonVirtualBoneGuid(GetSkeleton()->GetVirtualBoneGuid());
+						}
+	
+						if (bCompressedRunningPlatform)
+						{
+							FAssetRegistryModule& AssetRegistryModule = FModuleManager::LoadModuleChecked<FAssetRegistryModule>(TEXT("AssetRegistry"));
+							AssetRegistryModule.Get().AssetTagsFinalized(*this);
+	
+							// Cache compressed data ptr for current platform to by-pass TMap lookup requiring data-hash retrieval/generation and locking
+							ensure(CurrentPlatformData == nullptr);
+							CurrentPlatformData = TaskData;
+						}
+					
+						check(TaskData->IsValid(this, true));
 					}
-
-					if (bCompressedRunningPlatform)
+					else
 					{
-						FAssetRegistryModule& AssetRegistryModule = FModuleManager::LoadModuleChecked<FAssetRegistryModule>(TEXT("AssetRegistry"));
-						AssetRegistryModule.Get().AssetTagsFinalized(*this);
-
-						// Cache compressed data ptr for current platform to by-pass TMap lookup requiring data-hash retrieval/generation and locking
-						ensure(CurrentPlatformData == nullptr);
-						CurrentPlatformData = TaskData;
+						// Failed to compress
+						UE_LOG(LogAnimationCompression, Display, TEXT("Failed to finish async Animation Compression task for %s, as the generated data is not valid."), *GetName());
+						ResetData();
 					}
-				
-					check(TaskData->IsValid(this, true));
-				}
-				else
-				{
-					// Failed to compress
-					UE_LOG(LogAnimationCompression, Display, TEXT("Failed to finish async Animation Compression task for %s, as the generated data is not valid."), *GetName());
-					ResetData();
 				}
 			}
 
