Use .interpolate rather than .multiply during attribute blending, this addresses non-zero default attributes values e.g. FTransform/FQuat

diff --git a/Engine/Source/Runtime/Engine/Classes/Animation/AttributeBlendOperator.h b/Engine/Source/Runtime/Engine/Classes/Animation/AttributeBlendOperator.h
--- a/Engine/Source/Runtime/Engine/Classes/Animation/AttributeBlendOperator.h
+++ b/Engine/Source/Runtime/Engine/Classes/Animation/AttributeBlendOperator.h
@@ -92,7 +92,7 @@ namespace UE
 						}
 						else
 						{
-							OutAttribute = Attribute.Multiply(AttributeWeight);
+							OutAttribute.Interpolate(Attribute, AttributeWeight);
 						}
 
 						if constexpr (UE::Anim::TAttributeTypeTraits<AttributeType>::RequiresNormalization)
@@ -196,7 +196,7 @@ namespace UE
 							const AttributeType& Attribute = It.GetValue();
 							if (ExistingIndex == INDEX_NONE)
 							{
-								OutAttribute = Attribute.Multiply(AttributeWeight);
+								OutAttribute.Interpolate(Attribute, AttributeWeight);
 							}
 							else
 							{
diff --git a/Engine/Source/Runtime/Engine/Classes/Animation/BuiltInAttributeTypes.h b/Engine/Source/Runtime/Engine/Classes/Animation/BuiltInAttributeTypes.h
--- a/Engine/Source/Runtime/Engine/Classes/Animation/BuiltInAttributeTypes.h
+++ b/Engine/Source/Runtime/Engine/Classes/Animation/BuiltInAttributeTypes.h
@@ -206,7 +206,6 @@ struct FQuaternionAnimationAttribute
 		{
 			if (AdditiveType == AAT_None)
 			{
-				
 				const FQuat WeightedRotation = Attribute.Value * Weight;
 				
 				// From VectorAccumulateQuaternionShortestPath
@@ -251,7 +250,6 @@ struct FQuaternionAnimationAttribute
 	void Interpolate(const FQuaternionAnimationAttribute& Attribute, float Alpha)
 	{
 		Value = FQuat::FastLerp(Value, Attribute.Value, Alpha).GetNormalized();
-		Value.Normalize();
 	}
 };
 
