StaticMesh cook: early out before call to FPrimitiveInstanceDataManager::BeginCacheForCookedPlatformData, to avoid unnecessary computation of FInstanceUpdateComponentDesc when no features are enabled that require BeginCacheForCookedPlatformData to be called.  Computing FInstanceUpdateComponentDesc forces an early flush of the static mesh compilation queue where GetUsedMaterials() is called, catastrophically reducing parallelism.  On CitySample, this optimization saves 11% in overall cook time (5 minutes) following a StaticMesh version bump.

#rnx

diff --git a/Engine/Source/Runtime/Engine/Private/InstancedStaticMesh.cpp b/Engine/Source/Runtime/Engine/Private/InstancedStaticMesh.cpp
--- a/Engine/Source/Runtime/Engine/Private/InstancedStaticMesh.cpp
+++ b/Engine/Source/Runtime/Engine/Private/InstancedStaticMesh.cpp
@@ -5602,13 +5602,32 @@ void UInstancedStaticMeshComponent::BeginCacheForCookedPlatformData( const ITarg
 {
 	Super::BeginCacheForCookedPlatformData(TargetPlatform);
 
+	// Early out to avoid call to BuildComponentInstanceData if we don't need to call BeginCacheForCookedPlatformData
+	bool bNeedPrimitiveInstanceDataManagerCookedData = PrimitiveInstanceDataManager.ShouldWriteCookedData(TargetPlatform, PerInstanceSMData.Num());
+
 	FInstanceUpdateComponentDesc ComponentData;
 	if (GetStaticMesh() != nullptr)
 	{
 		UpdateComponentToWorld();
-		BuildComponentInstanceData(GMaxRHIFeatureLevel, ComponentData);
+
+		if (bNeedPrimitiveInstanceDataManagerCookedData)
+		{
+			// Calling this function can be expensive, because of its call to GetUsedMaterials, which immediately forces completion of in-flight
+			// static mesh compilation, reducing parallelism.  That's why we try to avoid the call where it's not necessary.
+			//
+			// TODO:  In the future, we could make it so the following function is itself more efficient.  The component data is only affected by
+			// materials where displacement is used (rare), and we could skip the call to GetUsedMaterials if none of the original material references
+			// in the static mesh have displacement.  This isn't completely trivial, as the static mesh currently protects access to the material
+			// references, also forcing completion of in-flight compilation when accessed.  You'd need some way to bypass that protection, perhaps
+			// caching the relevant information (whether displacement may be present) on load.
+			BuildComponentInstanceData(GMaxRHIFeatureLevel, ComponentData);
+		}
+	}
+
+	if (bNeedPrimitiveInstanceDataManagerCookedData)
+	{
+		PrimitiveInstanceDataManager.BeginCacheForCookedPlatformData(TargetPlatform, MoveTemp(ComponentData), MakeStridedView(PerInstanceSMData, &FInstancedStaticMeshInstanceData::Transform));
 	}
-	PrimitiveInstanceDataManager.BeginCacheForCookedPlatformData(TargetPlatform, MoveTemp(ComponentData), MakeStridedView(PerInstanceSMData, &FInstancedStaticMeshInstanceData::Transform));
 }
 	
 bool UInstancedStaticMeshComponent::IsCachedCookedPlatformDataLoaded( const ITargetPlatform* TargetPlatform )
diff --git a/Engine/Source/Runtime/Engine/Private/InstancedStaticMesh/ISMInstanceDataManager.cpp b/Engine/Source/Runtime/Engine/Private/InstancedStaticMesh/ISMInstanceDataManager.cpp
--- a/Engine/Source/Runtime/Engine/Private/InstancedStaticMesh/ISMInstanceDataManager.cpp
+++ b/Engine/Source/Runtime/Engine/Private/InstancedStaticMesh/ISMInstanceDataManager.cpp
@@ -987,6 +987,13 @@ bool FPrimitiveInstanceDataManager::ShouldUsePrecomputed()
 
 #if WITH_EDITOR
 
+// Calling this early can save performance if no features are enabled that require BeginCacheForCookedPlatformData to be called.  Specifically, that function
+// does no work if this function returns false.  If additional features require BeginCacheForCookedPlatformData to be called, or logic is added outside the
+// "bShouldBuild" condition in that function, this would need to be revisited.
+//
+// The big advantage of this early out is that unnecessarily computing the FInstanceUpdateComponentDesc argument to BeginCacheForCookedPlatformData can
+// have catastrophic effects on cook performance, because it forces an early flush of the static mesh compilation queue where GetUsedMaterials() is called,
+// reducing parallelism.  On a test project, this optimization saves 11% in overall cook time following a StaticMesh version bump.
 bool FPrimitiveInstanceDataManager::ShouldWriteCookedData(const ITargetPlatform* TargetPlatform, int32 NumInstancesToBuildFor)
 {
 	EComponentMobility::Type Mobility = PrimitiveComponent.IsValid() ? PrimitiveComponent->Mobility.GetValue() : EComponentMobility::Type::Movable;
diff --git a/Engine/Source/Runtime/Engine/Public/InstancedStaticMesh/ISMInstanceDataManager.h b/Engine/Source/Runtime/Engine/Public/InstancedStaticMesh/ISMInstanceDataManager.h
--- a/Engine/Source/Runtime/Engine/Public/InstancedStaticMesh/ISMInstanceDataManager.h
+++ b/Engine/Source/Runtime/Engine/Public/InstancedStaticMesh/ISMInstanceDataManager.h
@@ -167,6 +167,9 @@ public:
 	void OnRegister(int32 InNumInstances);
 
 #if WITH_EDITOR
+	// Early out test to save performance if no features are enabled that require BeginCacheForCookedPlatformData to be called
+	bool ShouldWriteCookedData(const ITargetPlatform* TargetPlatform, int32 NumInstancesToBuildFor);
+
 	void BeginCacheForCookedPlatformData(const ITargetPlatform* TargetPlatform, FInstanceUpdateComponentDesc &&ComponentData, TStridedView<FMatrix> InstanceTransforms);
 
 	void WriteCookedRenderData(FArchive& Ar, FInstanceUpdateComponentDesc &&ComponentData, TStridedView<FMatrix> InstanceTransforms);
@@ -192,9 +195,6 @@ private:
 	void InitChangeSet(const union FChangeDesc &ChangeDesc, const FInstanceUpdateComponentDesc &ComponentData, FISMInstanceUpdateChangeSet &ChangeSet);
 
 #if WITH_EDITOR
-
-	bool ShouldWriteCookedData(const ITargetPlatform* TargetPlatform, int32 NumInstancesToBuildFor);
-
 	/**
 	 * Build precomputed data from the input.
 	 */
