Cherrypicking fix to prevent try to compresspreviously cooked AnimationSequence, as its source data will be non-existent (cooked-editor runtime)

diff --git a/Engine/Source/Runtime/Engine/Private/Animation/AnimSequence.cpp b/Engine/Source/Runtime/Engine/Private/Animation/AnimSequence.cpp
--- a/Engine/Source/Runtime/Engine/Private/Animation/AnimSequence.cpp
+++ b/Engine/Source/Runtime/Engine/Private/Animation/AnimSequence.cpp
@@ -5016,7 +5016,7 @@ void UAnimSequence::ValidateCompressionSettings()
 bool UAnimSequence::CanBeCompressed() const
 {
 	const UPackage* Package = GetPackage();
-	return Package && (Package->GetHasBeenEndLoaded() || Package->GetLoadedPath().IsEmpty());
+	return Package && (Package->GetHasBeenEndLoaded() || Package->GetLoadedPath().IsEmpty()) && !Package->HasAnyPackageFlags(PKG_Cooked);
 }
 
 FIoHash UAnimSequence::BeginCacheDerivedData(const ITargetPlatform* TargetPlatform)
