StaticMesh cook:  distance field and card representation task optimizations, saving around 10 minutes off CitySample cook time on a static mesh version bump.
* Distance field and card representation DDC key is always based on the platform independent base render data's DerivedDataKey.  Internally, cooking was already always using the platform independent data, so this doesn't change behavior, but it formalizes it and allows all platforms to use the same DDC data, rather than cooking the same data separately, saving significant performance when cooking for multiple platforms.  The option to pass in a derived data key was deprecated in the API.
* Card representation task work dependent on distance fields is now directly run in the same worker task when cooking, rather than being added to the queue to be kicked off later from the game thread, reducing task overhead and latency.  Because the build APIs are public, the inner logic was encapsulated to allow it to be called directly, while preserving the original public APIs.
* Card representation build reuses mesh vertex data from the distance field task, rather than generating it again.

#rnx

diff --git a/Engine/Source/Runtime/Core/Private/UObject/FortniteMainBranchObjectVersion.cpp b/Engine/Source/Runtime/Core/Private/UObject/FortniteMainBranchObjectVersion.cpp
--- a/Engine/Source/Runtime/Core/Private/UObject/FortniteMainBranchObjectVersion.cpp
+++ b/Engine/Source/Runtime/Core/Private/UObject/FortniteMainBranchObjectVersion.cpp
@@ -13,7 +13,7 @@ TMap<FGuid, FGuid> FFortniteMainBranchObjectVersion::GetSystemGuids()
 	SystemGuids.Add(DevGuids.NIAGARASHADERMAP_DERIVEDDATA_VER, FGuid("45D5978F33D940EEA6602B78559900DC"));
 	SystemGuids.Add(DevGuids.Niagara_LatestScriptCompileVersion, FGuid("670A2EBDEE1A4794BB77BBFF331EA7EF"));
 	SystemGuids.Add(DevGuids.SkeletalMeshDerivedDataVersion, FGuid("5E93192FCF3647EEB7F9D4FFF1750050"));
-	SystemGuids.Add(DevGuids.STATICMESH_DERIVEDDATA_VER, FGuid("656E338C78C34641A5D3BFAE303997AF"));
+	SystemGuids.Add(DevGuids.STATICMESH_DERIVEDDATA_VER, FGuid("8FE50FFED98A4DDCA3410CBA50849667"));
 	SystemGuids.Add(DevGuids.MaterialTranslationDDCVersion, FGuid("0195246791767C5080E1D4C3ED151B18"));
 	return SystemGuids;
 }
diff --git a/Engine/Source/Runtime/Engine/Private/DistanceFieldAtlas.cpp b/Engine/Source/Runtime/Engine/Private/DistanceFieldAtlas.cpp
--- a/Engine/Source/Runtime/Engine/Private/DistanceFieldAtlas.cpp
+++ b/Engine/Source/Runtime/Engine/Private/DistanceFieldAtlas.cpp
@@ -127,6 +127,12 @@ FDistanceFieldAsyncQueue* GDistanceFieldAsyncQueue = NULL;
 
 #if WITH_EDITOR
 
+FAsyncCardRepresentationTask* BeginCacheMeshCardRepresentationInternal(
+	UStaticMesh* StaticMeshAsset,
+	FStaticMeshRenderData& RenderData,
+	const FString& DistanceFieldKey,
+	FSourceMeshDataForDerivedDataTask&& OptionalSourceMeshData);
+
 // DDC key for distance field data, must be changed when modifying the generation code or data format
 #define DISTANCEFIELD_DERIVEDDATA_VER TEXT("9137A53B-09F6-4C2F-86E1-29FE6834D74B")
 
@@ -195,6 +201,11 @@ void BuildSignedDistanceFieldBuildSectionData(UStaticMesh* Mesh, uint32 LODIndex
 }
 
 void FDistanceFieldVolumeData::CacheDerivedData(const FString& InStaticMeshDerivedDataKey, const ITargetPlatform* TargetPlatform, UStaticMesh* Mesh, FStaticMeshRenderData& RenderData, UStaticMesh* GenerateSource, float DistanceFieldResolutionScale, bool bGenerateDistanceFieldAsIfTwoSided)
+{
+	CacheDerivedData(Mesh, GenerateSource, DistanceFieldResolutionScale, bGenerateDistanceFieldAsIfTwoSided);
+}
+
+void FDistanceFieldVolumeData::CacheDerivedData(UStaticMesh* Mesh, UStaticMesh* GenerateSource, float DistanceFieldResolutionScale, bool bGenerateDistanceFieldAsIfTwoSided)
 {
 	TRACE_CPUPROFILER_EVENT_SCOPE(FDistanceFieldVolumeData::CacheDerivedData);
 	COOK_STAT(auto Timer = DistanceFieldCookStats::UsageStats.TimeSyncWork());
@@ -204,7 +215,7 @@ void FDistanceFieldVolumeData::CacheDerivedData(const FString& InStaticMeshDeriv
 	const uint32 LODIndex = 0;
 	BuildSignedDistanceFieldBuildSectionData(Mesh, LODIndex, BuildSectionData);
 
-	FString DistanceFieldKey = BuildDistanceFieldDerivedDataKey(InStaticMeshDerivedDataKey);
+	FString DistanceFieldKey = BuildDistanceFieldDerivedDataKey(Mesh->GetRenderData()->DerivedDataKey);
 
 	for (int32 SectionIndex = 0; SectionIndex < BuildSectionData.Num(); SectionIndex++)
 	{
@@ -222,17 +233,24 @@ void FDistanceFieldVolumeData::CacheDerivedData(const FString& InStaticMeshDeriv
 		FMemoryReader Ar(DerivedData, /*bIsPersistent=*/ true);
 		Serialize(Ar, Mesh);
 
-		BeginCacheMeshCardRepresentation(TargetPlatform, Mesh, RenderData, DistanceFieldKey);
+		BeginCacheMeshCardRepresentation(Mesh, *Mesh->GetRenderData(), DistanceFieldKey);
 	}
 	else if (GDistanceFieldAsyncQueue)
 	{
 		check(Mesh && GenerateSource);
 
+		// Mark that this resource is async building
+		bAsyncBuilding = true;
+
 		// We don't actually build the resource until later, so only track the cycles used here.
 		COOK_STAT(Timer.TrackCyclesOnly());
 		FAsyncDistanceFieldTask* NewTask = new FAsyncDistanceFieldTask;
 		NewTask->DDCKey = DistanceFieldKey;
-		NewTask->TargetPlatform = TargetPlatform;
+
+		// Always use GetRunningTargetPlatform() here (the cooker's platform, typically Windows Editor), as the data is intended to be platform independent.
+		// Besides that, TargetPlatform is only used to fetch an IMeshBuilderModule interface, and in practice the BuildMeshVertexPositions function on
+		// that interface doesn't vary per platform (all platforms internally use FStaticMeshBuilder::BuildMeshVertexPositions).
+		NewTask->TargetPlatform = GetTargetPlatformManagerRef().GetRunningTargetPlatform();
 		NewTask->StaticMesh = Mesh;
 		NewTask->GenerateSource = GenerateSource;
 		NewTask->DistanceFieldResolutionScale = DistanceFieldResolutionScale;
@@ -245,7 +263,7 @@ void FDistanceFieldVolumeData::CacheDerivedData(const FString& InStaticMeshDeriv
 		// Nanite overrides source static mesh with a coarse representation. Need to load original data before we build the mesh SDF.
 		if (Mesh->IsNaniteEnabled())
 		{
-			IMeshBuilderModule& MeshBuilderModule = IMeshBuilderModule::GetForPlatform(TargetPlatform);
+			IMeshBuilderModule& MeshBuilderModule = IMeshBuilderModule::GetForPlatform(NewTask->TargetPlatform);
 			if (!MeshBuilderModule.BuildMeshVertexPositions(Mesh, NewTask->SourceMeshData.TriangleIndices, NewTask->SourceMeshData.VertexPositions, NewTask->SourceMeshData.Sections))
 			{
 				UE_LOG(LogStaticMesh, Error, TEXT("Failed to build static mesh. See previous line(s) for details."));
@@ -751,6 +769,23 @@ void FDistanceFieldAsyncQueue::Build(FAsyncDistanceFieldTask* Task, FQueuedThrea
 			Task->DistanceFieldResolutionScale,
 			Task->bGenerateDistanceFieldAsIfTwoSided,
 			*Task->GeneratedVolumeData);
+
+		// If cooking, build the dependent mesh card representation data immediately in the existing worker task
+		if (IsRunningCookCommandlet())
+		{
+			FAsyncCardRepresentationTask* CardTask = BeginCacheMeshCardRepresentationInternal(
+				Task->StaticMesh,
+				*Task->StaticMesh->GetRenderData(),
+				Task->DDCKey,
+				MoveTemp(Task->SourceMeshData));
+
+			if (CardTask)
+			{
+				GCardRepresentationAsyncQueue->RunImmediateTask(CardTask);
+			}
+
+			Task->bGeneratedCardRepresentation = true;
+		}
 	}
 
 	{
@@ -862,12 +897,14 @@ void FDistanceFieldAsyncQueue::ProcessAsyncTasks(bool bLimitExecutionTime)
 				COOK_STAT(Timer.AddMiss(DerivedData.Num()));
 			}
 
-			BeginCacheMeshCardRepresentation(
-				Task->TargetPlatform,
-				Task->StaticMesh,
-				Task->StaticMesh->GetPlatformStaticMeshRenderData(Task->StaticMesh, Task->TargetPlatform),
-				Task->DDCKey,
-				MoveTemp(Task->SourceMeshData));
+			if (!Task->bGeneratedCardRepresentation)
+			{
+				BeginCacheMeshCardRepresentation(
+					Task->StaticMesh,
+					*Task->StaticMesh->GetRenderData(),
+					Task->DDCKey,
+					MoveTemp(Task->SourceMeshData));
+			}
 		}
 
 		delete Task;
diff --git a/Engine/Source/Runtime/Engine/Private/MeshCardRepresentation.cpp b/Engine/Source/Runtime/Engine/Private/MeshCardRepresentation.cpp
--- a/Engine/Source/Runtime/Engine/Private/MeshCardRepresentation.cpp
+++ b/Engine/Source/Runtime/Engine/Private/MeshCardRepresentation.cpp
@@ -72,6 +72,15 @@ static TAutoConsoleVariable<int32> CVarCardRepresentationDebugSurfelDirection(
 	TEXT("Generate cards for only surfels pointing in a specific direction."),
 	ECVF_Default);
 
+FAsyncCardRepresentationTask* CreateCardRepresentationTask(
+	const FString& InDDCKey,
+	UStaticMesh* Mesh,
+	UStaticMesh* GenerateSource,
+	int32 MaxLumenMeshCards,
+	bool bGenerateDistanceFieldAsIfTwoSided,
+	FSourceMeshDataForDerivedDataTask&& OptionalSourceMeshData,
+	FCardRepresentationData& OutRenderData);
+
 float MeshCardRepresentation::GetMinDensity()
 {
 	return FMath::Clamp(CVarCardRepresentationMinDensity.GetValueOnAnyThread(), 0.0f, 1.0f);
@@ -132,8 +141,7 @@ FString BuildCardRepresentationDerivedDataKey(const FString& InMeshKey, int32 Ma
 
 extern void BuildSignedDistanceFieldBuildSectionData(UStaticMesh* Mesh, uint32 LODIndex, TArray<FSignedDistanceFieldBuildSectionData>& OutData);
 
-void BeginCacheMeshCardRepresentation(
-	const ITargetPlatform* TargetPlatform,
+FAsyncCardRepresentationTask* BeginCacheMeshCardRepresentationInternal(
 	UStaticMesh* StaticMeshAsset,
 	FStaticMeshRenderData& RenderData,
 	const FString& DistanceFieldKey,
@@ -141,10 +149,10 @@ void BeginCacheMeshCardRepresentation(
 {
 	static const auto CVarCards = IConsoleManager::Get().FindTConsoleVariableDataInt(TEXT("r.MeshCardRepresentation"));
 
+	FAsyncCardRepresentationTask* NewTask = nullptr;
+
 	if (CVarCards->GetValueOnAnyThread() != 0)
 	{
-		const FMeshBuildSettings& BuildSettings = StaticMeshAsset->GetSourceModel(0).BuildSettings;
-		FString Key = BuildCardRepresentationDerivedDataKey(DistanceFieldKey, BuildSettings.MaxLumenMeshCards);
 		if (RenderData.LODResources.IsValidIndex(0))
 		{
 			if (!RenderData.LODResources[0].CardRepresentationData)
@@ -152,28 +160,47 @@ void BeginCacheMeshCardRepresentation(
 				RenderData.LODResources[0].CardRepresentationData = new FCardRepresentationData();
 			}
 
-			UStaticMesh* MeshToGenerateFrom = StaticMeshAsset;
-
-			RenderData.LODResources[0].CardRepresentationData->CacheDerivedData(
-				Key,
-				TargetPlatform,
-				StaticMeshAsset,
-				MeshToGenerateFrom,
-				BuildSettings.MaxLumenMeshCards,
-				BuildSettings.bGenerateDistanceFieldAsIfTwoSided,
-				MoveTemp(OptionalSourceMeshData));
+			// Only generate distance fields and card representations for the base render data, not platform render data.
+			if (&RenderData == StaticMeshAsset->GetRenderData())
+			{
+				const FMeshBuildSettings& BuildSettings = StaticMeshAsset->GetSourceModel(0).BuildSettings;
+				FString Key = BuildCardRepresentationDerivedDataKey(DistanceFieldKey, BuildSettings.MaxLumenMeshCards);
+
+				UStaticMesh* MeshToGenerateFrom = StaticMeshAsset;
+
+				NewTask = CreateCardRepresentationTask(
+					Key,
+					StaticMeshAsset,
+					MeshToGenerateFrom,
+					BuildSettings.MaxLumenMeshCards,
+					BuildSettings.bGenerateDistanceFieldAsIfTwoSided,
+					MoveTemp(OptionalSourceMeshData),
+					*RenderData.LODResources[0].CardRepresentationData);
+			}
+			else
+			{
+				// Copy the base render data to platform render data.  If it hasn't been generated or loaded from DDC yet, it will be empty,
+				// and the async queue will later copy the actual data when complete.
+				FCardRepresentationData* BaseCardRepresentationData = StaticMeshAsset->GetRenderData()->LODResources[0].CardRepresentationData;
+				if (BaseCardRepresentationData)
+				{
+					*RenderData.LODResources[0].CardRepresentationData = *BaseCardRepresentationData;
+				}
+			}
 		}
 	}
+
+	return NewTask;
 }
 
-void FCardRepresentationData::CacheDerivedData(
+FAsyncCardRepresentationTask* CreateCardRepresentationTask(
 	const FString& InDDCKey,
-	const ITargetPlatform* TargetPlatform,
 	UStaticMesh* Mesh,
 	UStaticMesh* GenerateSource,
 	int32 MaxLumenMeshCards,
 	bool bGenerateDistanceFieldAsIfTwoSided,
-	FSourceMeshDataForDerivedDataTask&& OptionalSourceMeshData)
+	FSourceMeshDataForDerivedDataTask&& OptionalSourceMeshData,
+	FCardRepresentationData& OutRenderData)
 {
 	TRACE_CPUPROFILER_EVENT_SCOPE(FCardRepresentationData::CacheDerivedData);
 
@@ -185,7 +212,7 @@ void FCardRepresentationData::CacheDerivedData(
 	{
 		COOK_STAT(Timer.AddHit(DerivedData.Num()));
 		FMemoryReader Ar(DerivedData, /*bIsPersistent=*/ true);
-		Ar << *this;
+		Ar << OutRenderData;
 
 		if (Ar.IsError())
 		{
@@ -216,13 +243,48 @@ void FCardRepresentationData::CacheDerivedData(
 		}
 		else if (Mesh->IsNaniteEnabled())
 		{
-			IMeshBuilderModule& MeshBuilderModule = IMeshBuilderModule::GetForPlatform(TargetPlatform);
+			// Always use GetRunningTargetPlatform() here (the cooker's platform, typically Windows Editor), as the data is intended to be platform independent.
+			// Besides that, TargetPlatform is only used to fetch an IMeshBuilderModule interface, and in practice the BuildMeshVertexPositions function on
+			// that interface doesn't vary per platform (all platforms internally use FStaticMeshBuilder::BuildMeshVertexPositions).
+			IMeshBuilderModule& MeshBuilderModule = IMeshBuilderModule::GetForPlatform(GetTargetPlatformManagerRef().GetRunningTargetPlatform());
 			if (!MeshBuilderModule.BuildMeshVertexPositions(Mesh, NewTask->SourceMeshData.TriangleIndices, NewTask->SourceMeshData.VertexPositions, NewTask->SourceMeshData.Sections))
 			{
 				UE_LOG(LogStaticMesh, Error, TEXT("Failed to build static mesh. See previous line(s) for details."));
 			}
 		}
 
+		return NewTask;
+	}
+
+	return nullptr;
+}
+
+void BeginCacheMeshCardRepresentation(
+	UStaticMesh* StaticMeshAsset,
+	FStaticMeshRenderData& RenderData,
+	const FString& DistanceFieldKey,
+	FSourceMeshDataForDerivedDataTask&& OptionalSourceMeshData)
+{
+	FAsyncCardRepresentationTask* NewTask = BeginCacheMeshCardRepresentationInternal(StaticMeshAsset, RenderData, DistanceFieldKey, MoveTemp(OptionalSourceMeshData));
+
+	if (NewTask)
+	{
+		GCardRepresentationAsyncQueue->AddTask(NewTask);
+	}
+}
+
+void FCardRepresentationData::CacheDerivedData(
+	const FString& InDDCKey,
+	UStaticMesh* Mesh,
+	UStaticMesh* GenerateSource,
+	int32 MaxLumenMeshCards,
+	bool bGenerateDistanceFieldAsIfTwoSided,
+	FSourceMeshDataForDerivedDataTask&& OptionalSourceMeshData)
+{
+	FAsyncCardRepresentationTask* NewTask = CreateCardRepresentationTask(InDDCKey, Mesh, GenerateSource, MaxLumenMeshCards, bGenerateDistanceFieldAsIfTwoSided, MoveTemp(OptionalSourceMeshData), *this);
+
+	if (NewTask)
+	{
 		GCardRepresentationAsyncQueue->AddTask(NewTask);
 	}
 }
@@ -401,6 +463,24 @@ void FCardRepresentationAsyncQueue::StartBackgroundTask(FAsyncCardRepresentation
 	Task->AsyncTask->StartBackgroundTask(ThreadPool.Get(), EQueuedWorkPriority::Lowest, EQueuedWorkFlags::DoNotRunInsideBusyWait, RequiredMemory, TEXT("Card"));
 }
 
+void FCardRepresentationAsyncQueue::RunImmediateTask(FAsyncCardRepresentationTask* Task)
+{
+#if WITH_EDITOR
+	if (!MeshUtilities)
+	{
+		FPlatformAtomics::InterlockedCompareExchangePointer((void**)&MeshUtilities, &FModuleManager::Get().LoadModuleChecked<IMeshUtilities>(TEXT("MeshUtilities")), nullptr);
+	}
+
+	// Thread pool argument isn't actually used by the function
+	Build(Task, *ThreadPool.Get());
+
+	// Array protection when called from multiple threads
+	FScopeLock Lock(&CriticalSection);
+	ReferencedTasks.Add(Task);
+	CompletedTasks.Add(Task);
+#endif
+}
+
 void FCardRepresentationAsyncQueue::ProcessPendingTasks()
 {
 	TRACE_CPUPROFILER_EVENT_SCOPE(FCardRepresentationAsyncQueue::ProcessPendingTasks);
diff --git a/Engine/Source/Runtime/Engine/Private/StaticMesh.cpp b/Engine/Source/Runtime/Engine/Private/StaticMesh.cpp
--- a/Engine/Source/Runtime/Engine/Private/StaticMesh.cpp
+++ b/Engine/Source/Runtime/Engine/Private/StaticMesh.cpp
@@ -4006,16 +4006,35 @@ void FStaticMeshRenderData::Cache(const ITargetPlatform* TargetPlatform, UStatic
 				LODResources[0].DistanceFieldData->AssetName = Owner->GetFName();
 			}
 
-			const FMeshBuildSettings& BuildSettings = Owner->GetSourceModel(0).BuildSettings;
-			UStaticMesh* MeshToGenerateFrom = BuildSettings.DistanceFieldReplacementMesh ? ToRawPtr(BuildSettings.DistanceFieldReplacementMesh) : Owner;
-
-			if (BuildSettings.DistanceFieldReplacementMesh)
+			// Only generate distance fields and card representations for the base render data, not platform render data.
+			if (this == Owner->GetRenderData())
 			{
-				// Make sure dependency is postloaded
-				BuildSettings.DistanceFieldReplacementMesh->ConditionalPostLoad();
+				const FMeshBuildSettings& BuildSettings = Owner->GetSourceModel(0).BuildSettings;
+				UStaticMesh* MeshToGenerateFrom = BuildSettings.DistanceFieldReplacementMesh ? ToRawPtr(BuildSettings.DistanceFieldReplacementMesh) : Owner;
+
+				if (BuildSettings.DistanceFieldReplacementMesh)
+				{
+					// Make sure dependency is postloaded
+					BuildSettings.DistanceFieldReplacementMesh->ConditionalPostLoad();
+				}
+
+				LODResources[0].DistanceFieldData->CacheDerivedData(Owner, MeshToGenerateFrom, BuildSettings.DistanceFieldResolutionScale, BuildSettings.bGenerateDistanceFieldAsIfTwoSided);
 			}
+			else
+			{
+				// If base data is already built (implying it was loaded from DDC), copy that to platform data, otherwise
+				// it will be copied later when the async queue finishes generating it.
+				FDistanceFieldVolumeData* BaseDistanceFieldData = Owner->GetRenderData()->LODResources[0].DistanceFieldData;
+				if (BaseDistanceFieldData && BaseDistanceFieldData->bAsyncBuilding == false)
+				{
+					*LODResources[0].DistanceFieldData = *BaseDistanceFieldData;
+				}
 
-			LODResources[0].DistanceFieldData->CacheDerivedData(DerivedDataKey, TargetPlatform, Owner, *this, MeshToGenerateFrom, BuildSettings.DistanceFieldResolutionScale, BuildSettings.bGenerateDistanceFieldAsIfTwoSided);
+				// For platform data, this either copies the base render data, or allocates a placeholder structure for data to be copied into
+				// later when tasks generated by CacheDerivedData called above for the base render data completes.  For base render data (the
+				// "if" case above), this call is deferred until after distance fields are built.
+				BeginCacheMeshCardRepresentation(Owner, *this, FString());
+			}
 		}
 		else
 		{
@@ -4024,7 +4043,7 @@ void FStaticMeshRenderData::Cache(const ITargetPlatform* TargetPlatform, UStatic
 	}
 	else
 	{
-		BeginCacheMeshCardRepresentation(TargetPlatform, Owner, *this, DerivedDataKey);
+		BeginCacheMeshCardRepresentation(Owner, *this, Owner->GetRenderData()->DerivedDataKey);
 	}
 }
 #endif // #if WITH_EDITOR
diff --git a/Engine/Source/Runtime/Engine/Public/DistanceFieldAtlas.h b/Engine/Source/Runtime/Engine/Public/DistanceFieldAtlas.h
--- a/Engine/Source/Runtime/Engine/Public/DistanceFieldAtlas.h
+++ b/Engine/Source/Runtime/Engine/Public/DistanceFieldAtlas.h
@@ -290,8 +290,11 @@ public:
 
 #if WITH_EDITORONLY_DATA
 
+	UE_DEPRECATED(5.6, "Distance fields are always generated from the UStaticMesh's base render data, and platform independent.  Simplified version of function below removes now unused arguments.")
 	ENGINE_API void CacheDerivedData(const FString& InStaticMeshDerivedDataKey, const ITargetPlatform* TargetPlatform, UStaticMesh* Mesh, class FStaticMeshRenderData& RenderData, UStaticMesh* GenerateSource, float DistanceFieldResolutionScale, bool bGenerateDistanceFieldAsIfTwoSided);
 
+	ENGINE_API void CacheDerivedData(UStaticMesh* Mesh, UStaticMesh* GenerateSource, float DistanceFieldResolutionScale, bool bGenerateDistanceFieldAsIfTwoSided);
+
 #endif
 
 	ENGINE_API void Serialize(FArchive& Ar, UObject* Owner);
@@ -337,6 +340,7 @@ public:
 	FString DDCKey;
 	FDistanceFieldVolumeData* GeneratedVolumeData;
 	TUniquePtr<FAsyncTask<FAsyncDistanceFieldTaskWorker>> AsyncTask = nullptr;
+	bool bGeneratedCardRepresentation = false;
 };
 
 /** Class that manages asynchronous building of mesh distance fields. */
diff --git a/Engine/Source/Runtime/Engine/Public/MeshCardBuild.h b/Engine/Source/Runtime/Engine/Public/MeshCardBuild.h
--- a/Engine/Source/Runtime/Engine/Public/MeshCardBuild.h
+++ b/Engine/Source/Runtime/Engine/Public/MeshCardBuild.h
@@ -169,7 +169,6 @@ public:
 
 	void CacheDerivedData(
 		const FString& InDDCKey,
-		const ITargetPlatform* TargetPlatform,
 		UStaticMesh* Mesh,
 		UStaticMesh* GenerateSource,
 		int32 MaxLumenMeshCards,
@@ -186,7 +185,7 @@ public:
 		bool bGenerateDistanceFieldAsIfTwoSided,
 		FSourceMeshDataForDerivedDataTask* OptionalSourceMeshData)
 	{
-		CacheDerivedData(InDDCKey, TargetPlatform, Mesh, GenerateSource, MaxLumenMeshCards, bGenerateDistanceFieldAsIfTwoSided, OptionalSourceMeshData ? *OptionalSourceMeshData : FSourceMeshDataForDerivedDataTask());
+		CacheDerivedData(InDDCKey, Mesh, GenerateSource, MaxLumenMeshCards, bGenerateDistanceFieldAsIfTwoSided, OptionalSourceMeshData ? *OptionalSourceMeshData : FSourceMeshDataForDerivedDataTask());
 	}
 
 #endif
@@ -283,6 +282,7 @@ public:
 
 private:
 	friend FAsyncCardRepresentationTaskWorker;
+	friend class FDistanceFieldAsyncQueue;
 
 	ENGINE_API FName GetAssetTypeName() const override;
 	ENGINE_API FTextFormat GetAssetNameFormat() const override;
@@ -303,6 +303,9 @@ private:
 	/** Task will be sent to a background worker. */
 	void StartBackgroundTask(FAsyncCardRepresentationTask* Task);
 
+	/** Task is being run immediately from an existing background worker (typically FAsyncDistanceFieldTask). */
+	void RunImmediateTask(FAsyncCardRepresentationTask* Task);
+
 	/** Cancel or finish any work for any task matching the predicate. */
 	void CancelAndDeleteTaskByPredicate(TFunctionRef<bool(FAsyncCardRepresentationTask*)> ShouldCancelPredicate);
 
@@ -350,7 +353,6 @@ extern ENGINE_API void AppendMeshCardConstantsToClassSchema(FAppendToClassSchema
 extern ENGINE_API FString BuildCardRepresentationDerivedDataKey(const FString& InMeshKey);
 
 extern ENGINE_API void BeginCacheMeshCardRepresentation(
-	const ITargetPlatform* TargetPlatform,
 	UStaticMesh* StaticMeshAsset,
 	class FStaticMeshRenderData& RenderData,
 	const FString& DistanceFieldKey,
@@ -364,7 +366,7 @@ inline void BeginCacheMeshCardRepresentation(
 	const FString& DistanceFieldKey,
 	FSourceMeshDataForDerivedDataTask* OptionalSourceMeshData)
 {
-	BeginCacheMeshCardRepresentation(TargetPlatform, StaticMeshAsset, RenderData, DistanceFieldKey, OptionalSourceMeshData ? *OptionalSourceMeshData : FSourceMeshDataForDerivedDataTask());
+	BeginCacheMeshCardRepresentation(StaticMeshAsset, RenderData, DistanceFieldKey, OptionalSourceMeshData ? *OptionalSourceMeshData : FSourceMeshDataForDerivedDataTask());
 }
 #endif
 
