Remove deprecated Input functions from 5.1 on UConsole

#okforgithhub public

diff --git a/Engine/Source/Runtime/Engine/Classes/Engine/Console.h b/Engine/Source/Runtime/Engine/Classes/Engine/Console.h
--- a/Engine/Source/Runtime/Engine/Classes/Engine/Console.h
+++ b/Engine/Source/Runtime/Engine/Classes/Engine/Console.h
@@ -214,18 +214,6 @@ class UConsole
 	
 	/** Update the auto complete indices from the typed string */
 	ENGINE_API virtual void UpdateCompleteIndices();
-	
-	/**
-	 * Process a character input event (typing) routed through unrealscript from another object. This method is assigned as the value for the
-	 * OnReceivedNativeInputKey delegate so that native input events are routed to this unrealscript function.
-	 *
-	 * @param	ControllerId	the controller that generated this character input event
-	 * @param	Unicode			the character that was typed
-	 *
-	 * @return	True to consume the character, false to pass it on.
-	 */
-	UE_DEPRECATED(5.1, "This version of InputChar_Typing is deprecated. Please use the version that takes a DeviceId instead.")
-	ENGINE_API virtual bool InputChar_Typing( int32 ControllerId, const FString& Unicode );
 
 	/**
 	 * Process a character input event (typing) routed through unrealscript from another object. This method is assigned as the value for the
@@ -248,31 +236,11 @@ class UConsole
 
 	/** Perform actions on transition from Typing state */
 	ENGINE_API virtual void EndState_Typing( FName NextStateName );
-	
-	/**
-	 * This state is used when the console is open.
-	 */
-	UE_DEPRECATED(5.1, "This version of InputChar_Open has been deprecated. Please use the version that takes a DeviceId instead.")
-	ENGINE_API virtual bool InputChar_Open( int32 ControllerId, const FString& Unicode );
 
 	/**
 	 * This state is used when the console is open.
 	 */
 	ENGINE_API virtual bool InputChar_Open( FInputDeviceId DeviceId, const FString& Unicode );
-	
-	/**
-	 * Process an input key event routed through unrealscript from another object. This method is assigned as the value for the
-	 * OnReceivedNativeInputKey delegate so that native input events are routed to this unrealscript function.
-	 *
-	 * @param	ControllerId	the controller that generated this input key event
-	 * @param	Key				the name of the key which an event occured for (KEY_Up, KEY_Down, etc.)
-	 * @param	EventType		the type of event which occured (pressed, released, etc.)
-	 * @param	AmountDepressed	for analog keys, the depression percent.
-	 *
-	 * @return	true to consume the key event, false to pass it on.
-	 */
-	UE_DEPRECATED(5.1, "This version of InputKey_Open has been deprecated. Please use the version that takes a DeviceId instead.")
-	ENGINE_API virtual bool InputKey_Open( int32 ControllerId, FKey Key, EInputEvent Event, float AmountDepressed = 1.f, bool bGamepad = false );
 
 	/**
 	 * Process an input key event routed through unrealscript from another object. This method is assigned as the value for the
@@ -297,21 +265,13 @@ class UConsole
 
 	/** Perform actions on transition from Open state */
 	ENGINE_API virtual void EndState_Open(FName NextStateName);
-
-	UE_DEPRECATED(5.1, "This version of InputChar has been deprecated. Please use the version that takes an FInputDeviceId instead")
-	ENGINE_API virtual bool InputChar(int32 ControllerId, const FString& Unicode);
+	
 	ENGINE_API virtual bool InputChar(FInputDeviceId DeviceId, const FString& Unicode);
 	
-	UE_DEPRECATED(5.1, "This version of InputKey has been deprecated. Please use the version that takes an FInputDeviceId instead")
-	ENGINE_API virtual bool InputKey(int32 ControllerId, FKey Key, EInputEvent Event, float AmountDepressed=1.f, bool bGamepad=false);
 	ENGINE_API virtual bool InputKey(FInputDeviceId DeviceId, FKey Key, EInputEvent Event, float AmountDepressed=1.f, bool bGamepad=false);
-
-	UE_DEPRECATED(5.1, "This version of InputAxis has been deprecated. Please use the version that takes an FInputDeviceId instead")
-	virtual bool InputAxis(int32 ControllerId, FKey Key, float Delta, float DeltaTime, int32 NumSamples=1, bool bGamepad=false) { return false; };
+	
 	virtual bool InputAxis(FInputDeviceId DevideId, FKey Key, float Delta, float DeltaTime, int32 NumSamples=1, bool bGamepad=false) { return false; };
-
-	UE_DEPRECATED(5.1, "This version of InputTouch has been deprecated. Please use the version that takes an FInputDeviceId instead")
-	virtual bool InputTouch(int32 ControllerId, uint32 Handle, ETouchType::Type Type, const FVector2D& TouchLocation, float Force, FDateTime DeviceTimestamp, uint32 TouchpadIndex) { return false; }
+	
 	virtual bool InputTouch(FInputDeviceId DevideId, uint32 Handle, ETouchType::Type Type, const FVector2D& TouchLocation, float Force, FDateTime DeviceTimestamp, uint32 TouchpadIndex) { return false; }
 
 	ENGINE_API virtual void MouseMove(FViewport* Viewport, int32 X, int32 Y);
diff --git a/Engine/Source/Runtime/Engine/Private/UserInterface/Console.cpp b/Engine/Source/Runtime/Engine/Private/UserInterface/Console.cpp
--- a/Engine/Source/Runtime/Engine/Private/UserInterface/Console.cpp
+++ b/Engine/Source/Runtime/Engine/Private/UserInterface/Console.cpp
@@ -810,14 +810,6 @@ void UConsole::AppendInputText(const FString& Text)
 	UpdatePrecompletedInputLine();
 }
 
-PRAGMA_DISABLE_DEPRECATION_WARNINGS
-bool UConsole::InputChar_Typing(int32 ControllerId, const FString& Unicode)
-{
-	FInputDeviceId DeviceId = FInputDeviceId::CreateFromInternalId(ControllerId);
-	return InputChar_Typing(DeviceId, Unicode);
-}
-PRAGMA_ENABLE_DEPRECATION_WARNINGS
-
 bool UConsole::InputChar_Typing(FInputDeviceId DeviceId, const FString& Unicode)
 {
 	if (bCaptureKeyInput)
@@ -1244,14 +1236,6 @@ void UConsole::EndState_Typing(FName NextStateName)
 	bAutoCompleteLocked = false;
 }
 
-PRAGMA_DISABLE_DEPRECATION_WARNINGS
-bool UConsole::InputChar_Open(int32 ControllerId, const FString& Unicode)
-{
-	FInputDeviceId DeviceId = FInputDeviceId::CreateFromInternalId(ControllerId);
-	return InputChar_Open(DeviceId, Unicode);
-}
-PRAGMA_ENABLE_DEPRECATION_WARNINGS
-
 bool UConsole::InputChar_Open(FInputDeviceId DeviceId, const FString& Unicode)
 {
 	if (bCaptureKeyInput)
@@ -1264,14 +1248,6 @@ bool UConsole::InputChar_Open(FInputDeviceId DeviceId, const FString& Unicode)
 	return true;
 }
 
-PRAGMA_DISABLE_DEPRECATION_WARNINGS
-bool UConsole::InputKey_Open(int32 ControllerId, FKey Key, EInputEvent Event, float AmountDepressed, bool bGamepad)
-{
-	FInputDeviceId DeviceId = FInputDeviceId::CreateFromInternalId(ControllerId);
-	return InputKey_Open(DeviceId, Key, Event, AmountDepressed, bGamepad);
-}
-PRAGMA_ENABLE_DEPRECATION_WARNINGS
-
 bool UConsole::InputKey_Open(FInputDeviceId DeviceId, FKey Key, EInputEvent Event, float AmountDepressed, bool bGamepad)
 {
 	if (Key == EKeys::LeftMouseButton)
@@ -1769,14 +1745,6 @@ void UConsole::EndState_Open(FName NextStateName)
 	}
 }
 
-PRAGMA_DISABLE_DEPRECATION_WARNINGS
-bool UConsole::InputChar(int32 ControllerId, const FString& Unicode)
-{
-	FInputDeviceId DeviceId = FInputDeviceId::CreateFromInternalId(ControllerId);
-	return InputChar(DeviceId, Unicode);
-}
-PRAGMA_ENABLE_DEPRECATION_WARNINGS
-
 bool UConsole::InputChar(FInputDeviceId DeviceId, const FString& Unicode)
 {
 	if (ConsoleState == NAME_Typing)
@@ -1790,14 +1758,6 @@ bool UConsole::InputChar(FInputDeviceId DeviceId, const FString& Unicode)
 	return bCaptureKeyInput;
 }
 
-PRAGMA_DISABLE_DEPRECATION_WARNINGS
-bool UConsole::InputKey(int32 ControllerId, FKey Key, EInputEvent Event, float AmountDepressed, bool bGamepad)
-{
-	FInputDeviceId DeviceId = FInputDeviceId::CreateFromInternalId(ControllerId);
-	return InputKey(DeviceId, Key, Event, AmountDepressed, bGamepad);
-}
-PRAGMA_ENABLE_DEPRECATION_WARNINGS
-
 bool UConsole::InputKey(FInputDeviceId DeviceId, FKey Key, EInputEvent Event, float AmountDepressed, bool bGamepad)
 {
 	bool bWasConsumed = InputKey_InputLine(DeviceId, Key, Event, AmountDepressed, bGamepad);
