Backout:
Fast-path for editor access of current target platform compressed data

diff --git a/Engine/Source/Runtime/Engine/Classes/Animation/AnimSequence.h b/Engine/Source/Runtime/Engine/Classes/Animation/AnimSequence.h
--- a/Engine/Source/Runtime/Engine/Classes/Animation/AnimSequence.h
+++ b/Engine/Source/Runtime/Engine/Classes/Animation/AnimSequence.h
@@ -806,7 +806,6 @@ protected:
 	ENGINE_API bool PollCacheDerivedData(const FIoHash& KeyHash) const;
 	ENGINE_API void EndCacheDerivedData(const FIoHash& KeyHash);
 
-	mutable FCompressedAnimSequence* CurrentPlatformData = nullptr; 
 	TMap<FIoHash, TUniquePtr<FCompressedAnimSequence>> DataByPlatformKeyHash;
 	TMap<FIoHash, TPimplPtr<UE::Anim::FAnimationSequenceAsyncCacheTask>> CacheTasksByKeyHash;
 	mutable FRWLock CompressedDataLock;
diff --git a/Engine/Source/Runtime/Engine/Private/Animation/AnimSequence.cpp b/Engine/Source/Runtime/Engine/Private/Animation/AnimSequence.cpp
--- a/Engine/Source/Runtime/Engine/Private/Animation/AnimSequence.cpp
+++ b/Engine/Source/Runtime/Engine/Private/Animation/AnimSequence.cpp
@@ -1944,12 +1944,6 @@ const FCompressedAnimSequence& UAnimSequence::GetCompressedData_Internal() const
 	if(!GetOutermost()->HasAnyPackageFlags(PKG_Cooked))
 	{
 		static thread_local FCompressedAnimSequence DefaultCompressedAnimSequence;
-
-		if (CurrentPlatformData)
-		{
-			return *CurrentPlatformData;
-		}
-		
 		const ITargetPlatform* RunningPlatform = GetTargetPlatformManagerRef().GetRunningTargetPlatform();
 
 		if(DataByPlatformKeyHash.Num())
@@ -2000,11 +1994,6 @@ const FCompressedAnimSequence& UAnimSequence::GetPlatformCompressedData(const FA
 		}
 		else
 		{
-			if (CurrentPlatformData)
-			{
-				return *CurrentPlatformData;
-			}
-			
 			const ITargetPlatform* RunningPlatform = GetTargetPlatformManagerRef().GetRunningTargetPlatform();
 			Hash = GetDerivedDataKeyHash(RunningPlatform);
 		}
@@ -2052,11 +2041,6 @@ FCompressedAnimSequence& UAnimSequence::GetPlatformCompressedData(const FAnimExt
 		}
 		else
 		{
-			if (CurrentPlatformData)
-			{
-				return *CurrentPlatformData;
-			}
-			
 			const ITargetPlatform* RunningPlatform = GetTargetPlatformManagerRef().GetRunningTargetPlatform();
 			Hash = GetDerivedDataKeyHash(RunningPlatform);
 		}
@@ -5070,12 +5054,6 @@ FIoHash UAnimSequence::BeginCacheDerivedData(const ITargetPlatform* TargetPlatfo
 		TargetData = DataByPlatformKeyHash.Emplace(KeyHash, MakeUnique<FCompressedAnimSequence>()).Get();
 	}
 
-	const ITargetPlatform* RunningPlatform = GetTargetPlatformManagerRef().GetRunningTargetPlatform();	
-	if (TargetPlatform == RunningPlatform)
-	{
-		CurrentPlatformData = nullptr;
-	}
-
 	// Reset the target compressed data, to mark it invalid
 	check(TargetData);
 	TargetData->Reset();
@@ -5218,8 +5196,6 @@ void UAnimSequence::ClearAllCompressionData()
 				It.RemoveCurrent();
 			}
 		}
-		
-		CurrentPlatformData = nullptr;
 	}
 }
 
@@ -5233,12 +5209,6 @@ void UAnimSequence::ClearCompressionData(const FIoHash& InKeyHash)
 		UE_COMPRESSED_DATA_WRITE_SCOPE(this);
 		check(DataByPlatformKeyHash.Contains(InKeyHash));
 		DataByPlatformKeyHash.Remove(InKeyHash);
-		
-		const ITargetPlatform* RunningPlatform = GetTargetPlatformManagerRef().GetRunningTargetPlatform();
-		if (GetDerivedDataKeyHash(RunningPlatform) == InKeyHash)
-		{
-			CurrentPlatformData = nullptr;
-		}
 	}
 }
 
@@ -5366,47 +5336,41 @@ void UAnimSequence::FinishAsyncTasks()
 			It->Value->Wait();
 
 			FCompressedAnimSequence* TaskData = It->Value->GetTargetCompressedData();
-			if(TaskData != nullptr)
-			{				
-				const bool bCompressedRunningPlatform = It->Key == CurrentPlatformHash;
-				auto ResetData = [this, TaskData, It]()
-				{
-					TaskData->Reset();
-					DataByPlatformKeyHash.Remove(It->Key);
-				};
-	
-				if (It->Value.Get()->WasCancelled())
-				{
-					ResetData();
-				}
-				else
+			const bool bCompressedRunningPlatform = It->Key == CurrentPlatformHash;
+			auto ResetData = [this, TaskData, It]()
+			{
+				TaskData->Reset();
+				DataByPlatformKeyHash.Remove(It->Key);
+			};
+
+			if (It->Value.Get()->WasCancelled())
+			{
+				ResetData();
+			}
+			else
+			{
+				if (TaskData->IsValid(this, true))
 				{
-					if (TaskData->IsValid(this, true))
+#if WITH_EDITOR
+					//This is only safe during sync anim compression
+					if (GetSkeleton())
 					{
-						//This is only safe during sync anim compression
-						if (GetSkeleton())
-						{
-							SetSkeletonVirtualBoneGuid(GetSkeleton()->GetVirtualBoneGuid());
-						}
-	
-						if (bCompressedRunningPlatform)
-						{
-							FAssetRegistryModule& AssetRegistryModule = FModuleManager::LoadModuleChecked<FAssetRegistryModule>(TEXT("AssetRegistry"));
-							AssetRegistryModule.Get().AssetTagsFinalized(*this);
-	
-							// Cache compressed data ptr for current platform to by-pass TMap lookup requiring data-hash retrieval/generation and locking
-							ensure(CurrentPlatformData == nullptr);
-							CurrentPlatformData = TaskData;
-						}
-					
-						check(TaskData->IsValid(this, true));
+						SetSkeletonVirtualBoneGuid(GetSkeleton()->GetVirtualBoneGuid());
 					}
-					else
+#endif
+					if (bCompressedRunningPlatform)
 					{
-						// Failed to compress
-						UE_LOG(LogAnimationCompression, Display, TEXT("Failed to finish async Animation Compression task for %s, as the generated data is not valid."), *GetName());
-						ResetData();
+						FAssetRegistryModule& AssetRegistryModule = FModuleManager::LoadModuleChecked<FAssetRegistryModule>(TEXT("AssetRegistry"));
+						AssetRegistryModule.Get().AssetTagsFinalized(*this);
 					}
+				
+					check(TaskData->IsValid(this, true));
+				}
+				else
+				{
+					// Failed to compress
+					UE_LOG(LogAnimationCompression, Display, TEXT("Failed to finish async Animation Compression task for %s, as the generated data is not valid."), *GetName());
+					ResetData();
 				}
 			}
 
