Export CreateQueryFilterData from PhysicsInterfaceUtils.h by adding ENGINE_API to its declaration to fix unresolved external.

diff --git a/Engine/Source/Runtime/Engine/Public/Physics/PhysicsInterfaceUtils.h b/Engine/Source/Runtime/Engine/Public/Physics/PhysicsInterfaceUtils.h
--- a/Engine/Source/Runtime/Engine/Public/Physics/PhysicsInterfaceUtils.h
+++ b/Engine/Source/Runtime/Engine/Public/Physics/PhysicsInterfaceUtils.h
@@ -16,7 +16,7 @@ struct FConstraintInstance;
 // FILTER DATA
 
 /** Utility for creating a filter data object for performing a query (trace) against the scene */
-FCollisionFilterData CreateQueryFilterData(const uint8 MyChannel, const bool bTraceComplex, const FCollisionResponseContainer& InCollisionResponseContainer, const struct FCollisionQueryParams& QueryParam, const struct FCollisionObjectQueryParams & ObjectParam, const bool bMultitrace);
+ENGINE_API FCollisionFilterData CreateQueryFilterData(const uint8 MyChannel, const bool bTraceComplex, const FCollisionResponseContainer& InCollisionResponseContainer, const struct FCollisionQueryParams& QueryParam, const struct FCollisionObjectQueryParams & ObjectParam, const bool bMultitrace);
 
 struct FConstraintBrokenDelegateData
 {
