Add ENGINE_API to LowLevelRaycast, LowLevelSweep, and LowLevelOverlap so that they're usable accross module boundaries.

diff --git a/Engine/Source/Runtime/Engine/Public/Physics/Experimental/ChaosInterfaceWrapper.h b/Engine/Source/Runtime/Engine/Public/Physics/Experimental/ChaosInterfaceWrapper.h
--- a/Engine/Source/Runtime/Engine/Public/Physics/Experimental/ChaosInterfaceWrapper.h
+++ b/Engine/Source/Runtime/Engine/Public/Physics/Experimental/ChaosInterfaceWrapper.h
@@ -44,10 +44,10 @@ UPrimitiveComponent* GetPrimitiveComponentFromUserData(const Chaos::FGeometryPar
 }
 
 template<typename TContainer, typename THit>
-void LowLevelRaycast(const TContainer& Container, const FVector& Start, const FVector& Dir, float DeltaMag, FPhysicsHitCallback<THit>& HitBuffer, EHitFlags OutputFlags, FQueryFlags QueryFlags, const FCollisionFilterData& Filter, const ChaosInterface::FQueryFilterData& QueryFilterData, ICollisionQueryFilterCallbackBase* QueryCallback, const ChaosInterface::FQueryDebugParams& DebugParams = {});
+ENGINE_API void LowLevelRaycast(const TContainer& Container, const FVector& Start, const FVector& Dir, float DeltaMag, FPhysicsHitCallback<THit>& HitBuffer, EHitFlags OutputFlags, FQueryFlags QueryFlags, const FCollisionFilterData& Filter, const ChaosInterface::FQueryFilterData& QueryFilterData, ICollisionQueryFilterCallbackBase* QueryCallback, const ChaosInterface::FQueryDebugParams& DebugParams = {});
 
 template<typename TContainer, typename THit>
-void LowLevelSweep(const TContainer& Container, const FPhysicsGeometry& Geom, const FTransform& StartTM, const FVector& Dir, float DeltaMag, FPhysicsHitCallback<THit>& HitBuffer, EHitFlags OutputFlags, FQueryFlags QueryFlags, const FCollisionFilterData& Filter, const ChaosInterface::FQueryFilterData& QueryFilterData, ICollisionQueryFilterCallbackBase* QueryCallback, const ChaosInterface::FQueryDebugParams& DebugParams = {});
+ENGINE_API void LowLevelSweep(const TContainer& Container, const FPhysicsGeometry& Geom, const FTransform& StartTM, const FVector& Dir, float DeltaMag, FPhysicsHitCallback<THit>& HitBuffer, EHitFlags OutputFlags, FQueryFlags QueryFlags, const FCollisionFilterData& Filter, const ChaosInterface::FQueryFilterData& QueryFilterData, ICollisionQueryFilterCallbackBase* QueryCallback, const ChaosInterface::FQueryDebugParams& DebugParams = {});
 
 template<typename TContainer, typename THit>
-void LowLevelOverlap(const TContainer& Container, const FPhysicsGeometry& Geom, const FTransform& GeomPose, FPhysicsHitCallback<THit>& HitBuffer, FQueryFlags QueryFlags, const FCollisionFilterData& Filter, const ChaosInterface::FQueryFilterData& QueryFilterData, ICollisionQueryFilterCallbackBase* QueryCallback, const ChaosInterface::FQueryDebugParams& DebugParams = {});
\ No newline at end of file
+ENGINE_API void LowLevelOverlap(const TContainer& Container, const FPhysicsGeometry& Geom, const FTransform& GeomPose, FPhysicsHitCallback<THit>& HitBuffer, FQueryFlags QueryFlags, const FCollisionFilterData& Filter, const ChaosInterface::FQueryFilterData& QueryFilterData, ICollisionQueryFilterCallbackBase* QueryCallback, const ChaosInterface::FQueryDebugParams& DebugParams = {});
\ No newline at end of file
