Defer morph target updates to happen in the uploads pass instead of the setup pass.

diff --git a/Engine/Source/Runtime/Engine/Private/SkeletalRenderGPUSkin.cpp b/Engine/Source/Runtime/Engine/Private/SkeletalRenderGPUSkin.cpp
--- a/Engine/Source/Runtime/Engine/Private/SkeletalRenderGPUSkin.cpp
+++ b/Engine/Source/Runtime/Engine/Private/SkeletalRenderGPUSkin.cpp
@@ -322,7 +322,7 @@ private:
 	};
 
 	void ProcessUpdatedDynamicData(TArrayView<FCommand> Commands, FRHICommandList& RHICmdList);
-	void UploadBoneData(TConstArrayView<FCommand> Commands, FRHICommandList& RHICmdList);
+	void UpdateBufferData(TConstArrayView<FCommand> Commands, FRHICommandList& RHICmdList);
 
 	TArray<FDynamicSkelMeshObjectDataGPUSkin*, FConcurrentLinearArrayAllocator> DynamicDatasToFree;
 	TArray<FSkeletalMeshObjectGPUSkin*, FConcurrentLinearArrayAllocator> MeshDeformer;
@@ -405,11 +405,11 @@ void FSkeletalMeshUpdatePacketGPUSkin::ProcessStage_Upload(FRHICommandList& RHIC
 {
 	for (FSkeletalMeshObjectGPUSkin* MeshObject : MeshDeformer)
 	{
-		MeshObject->UpdateBoneData(RHICmdList, EGPUSkinCacheEntryMode::Raster);
+		MeshObject->UpdateBufferData(RHICmdList, EGPUSkinCacheEntryMode::Raster);
 	}
 
-	UploadBoneData(Inline, RHICmdList);
-	UploadBoneData(SkinCache, RHICmdList);
+	UpdateBufferData(Inline, RHICmdList);
+	UpdateBufferData(SkinCache, RHICmdList);
 
 	for (FDynamicSkelMeshObjectDataGPUSkin* DynamicData : DynamicDatasToFree)
 	{
@@ -446,15 +446,15 @@ void FSkeletalMeshUpdatePacketGPUSkin::ProcessUpdatedDynamicData(TArrayView<FCom
 	}
 }
 
-void FSkeletalMeshUpdatePacketGPUSkin::UploadBoneData(TConstArrayView<FCommand> Commands, FRHICommandList& RHICmdList)
+void FSkeletalMeshUpdatePacketGPUSkin::UpdateBufferData(TConstArrayView<FCommand> Commands, FRHICommandList& RHICmdList)
 {
 	for (const FCommand& Command : Commands)
 	{
 		FSkeletalMeshObjectGPUSkin& MeshObject = *Command.MeshObject;
-		MeshObject.UpdateBoneData(RHICmdList, EGPUSkinCacheEntryMode::Raster);
+		MeshObject.UpdateBufferData(RHICmdList, EGPUSkinCacheEntryMode::Raster);
 		if (Command.bUpdateRayTracingMode)
 		{
-			MeshObject.UpdateBoneData(RHICmdList, EGPUSkinCacheEntryMode::RayTracing);
+			MeshObject.UpdateBufferData(RHICmdList, EGPUSkinCacheEntryMode::RayTracing);
 		}
 	}
 }
@@ -750,7 +750,7 @@ void FSkeletalMeshObjectGPUSkin::UpdateDynamicData_RenderThread(FRHICommandList&
 	DynamicData = InDynamicData;
 
 	ProcessUpdatedDynamicData(RHICmdList, GPUSkinCache, EGPUSkinCacheEntryMode::Raster);
-	UpdateBoneData(RHICmdList, EGPUSkinCacheEntryMode::Raster);
+	UpdateBufferData(RHICmdList, EGPUSkinCacheEntryMode::Raster);
 
 #if RHI_RAYTRACING
 	const bool bSkinCacheSupported = GPUSkinCache && GEnableGPUSkinCache;
@@ -759,7 +759,7 @@ void FSkeletalMeshObjectGPUSkin::UpdateDynamicData_RenderThread(FRHICommandList&
 	if (bSkinCacheForRayTracingSupported && IsRayTracingSkinCacheUpdateNeeded())
 	{
 		ProcessUpdatedDynamicData(RHICmdList, GPUSkinCache, EGPUSkinCacheEntryMode::RayTracing);
-		UpdateBoneData(RHICmdList, EGPUSkinCacheEntryMode::RayTracing);
+		UpdateBufferData(RHICmdList, EGPUSkinCacheEntryMode::RayTracing);
 	}
 	else
 	{
@@ -806,7 +806,15 @@ int32 FDynamicSkelMeshObjectDataGPUSkin::GetLODIndex(EGPUSkinCacheEntryMode Mode
 #endif
 }
 
-void FSkeletalMeshObjectGPUSkin::UpdateBoneData(FRHICommandList& RHICmdList, EGPUSkinCacheEntryMode Mode)
+
+bool FSkeletalMeshObjectGPUSkin::IsSkinCacheEnabled(EGPUSkinCacheEntryMode Mode) const
+{
+	return GEnableGPUSkinCache &&
+		// Force skin cache enabled for ray tracing if the inline skinning technique was requested.
+		(DynamicData->GPUSkinTechnique == ESkeletalMeshGPUSkinTechnique::GPUSkinCache || (DynamicData->GPUSkinTechnique == ESkeletalMeshGPUSkinTechnique::Inline && Mode == EGPUSkinCacheEntryMode::RayTracing));
+}
+
+void FSkeletalMeshObjectGPUSkin::UpdateBufferData(FRHICommandList& RHICmdList, EGPUSkinCacheEntryMode Mode)
 {
 	const int32 LODIndex = DynamicData->GetLODIndex(Mode);
 	FSkeletalMeshObjectLOD& LOD = LODs[LODIndex];
@@ -820,6 +828,29 @@ void FSkeletalMeshObjectGPUSkin::UpdateBoneData(FRHICommandList& RHICmdList, EGP
 	if (Mode == EGPUSkinCacheEntryMode::Raster)
 	{
 		PreviousReferenceToLocalMatrices = DynamicData->GetPreviousReferenceToLocal(Mode);
+
+		if (MorphVertexBuffer)
+		{
+			if (MorphVertexBuffer->bNeedsInitialClear)
+			{
+				if (MorphVertexBuffer->GetUAV())
+				{
+					RHICmdList.Transition(FRHITransitionInfo(MorphVertexBuffer->GetUAV(), ERHIAccess::Unknown, ERHIAccess::UAVCompute));
+					RHICmdList.ClearUAVUint(MorphVertexBuffer->GetUAV(), FUintVector4(0, 0, 0, 0));
+					RHICmdList.Transition(FRHITransitionInfo(MorphVertexBuffer->GetUAV(), ERHIAccess::UAVCompute, ERHIAccess::SRVMask));
+				}
+				MorphVertexBuffer->bNeedsInitialClear = false;
+			}
+	
+			// only update if the morph data changed and there are weighted morph targets
+			if (bMorphNeedsUpdate)
+			{
+				const bool bSkinCacheEnabled = IsSkinCacheEnabled(EGPUSkinCacheEntryMode::Raster);
+				const FSkeletalMeshLODRenderData& LODData = SkeletalMeshRenderData->LODRenderData[LODIndex];
+				UpdateMorphVertexBuffer(RHICmdList, EGPUSkinCacheEntryMode::Raster, LOD, LODData, bSkinCacheEnabled);
+				bMorphNeedsUpdate = false;
+			}
+		}
 	}
 
 	for (int32 SectionIdx = 0; SectionIdx < Sections.Num(); SectionIdx++)
@@ -867,36 +898,36 @@ void FSkeletalMeshObjectGPUSkin::ProcessUpdatedDynamicData(FRHICommandList& RHIC
 	bool bHasPreviousReferenceToLocal = false;
 	int32 PreviousRevisionNumber = CurrentRevisionNumber;
 
-	// RayTracing does not need the previous buffer at all, so don't allocate it.
 	if (Mode == EGPUSkinCacheEntryMode::Raster)
 	{
+		// RayTracing does not need the previous buffer at all, so don't allocate it.
 		bHasPreviousReferenceToLocal = !DynamicData->GetPreviousReferenceToLocal(Mode).IsEmpty();
 		PreviousRevisionNumber = bHasPreviousReferenceToLocal ? DynamicData->PreviousRevisionNumber : INDEX_NONE;
-	}
-
-	FMorphVertexBuffer* MorphVertexBuffer = nullptr;
-
-	if (LODData.GetNumVertices() > 0 && (bHasWeightedActiveMorphs || bHasExternalMorphs))
-	{
-		bMorphNeedsUpdate |= GForceUpdateMorphTargets != 0;
 
-		if (bMorphNeedsUpdate)
+		// RayTracing-specific LOD's can't have a separate morph target buffer at the moment because there is only one morph vertex buffer across the entire mesh.
+		if (LODData.GetNumVertices() > 0 && (bHasWeightedActiveMorphs || bHasExternalMorphs))
 		{
-			LOD.MorphVertexBufferPool->SetCurrentRevisionNumber(CurrentRevisionNumber);
-		}
+			bMorphNeedsUpdate |= GForceUpdateMorphTargets != 0;
+			MorphVertexBuffer = &LOD.MorphVertexBufferPool->GetMorphVertexBufferForWriting();
 
-		MorphVertexBuffer = &LOD.MorphVertexBufferPool->GetMorphVertexBufferForWriting();
+			// Force an update if this is the first use of the buffer
+			if (!MorphVertexBuffer->bHasBeenUpdated)
+			{
+				bMorphNeedsUpdate = true;
+			}
 
-		// Force an update if this is the first use of the buffer
-		if (!MorphVertexBuffer->bHasBeenUpdated)
+			if (bMorphNeedsUpdate)
+			{
+				LOD.MorphVertexBufferPool->SetCurrentRevisionNumber(CurrentRevisionNumber);
+				LOD.MorphVertexBufferPool->SetUpdatedFrameNumber(BoneTransformFrameNumber);
+			}
+		}
+		else
 		{
-			bMorphNeedsUpdate = true;
+			MorphVertexBuffer = nullptr;
+			bMorphNeedsUpdate = false;
 		}
 	}
-	else
-	{
-		bMorphNeedsUpdate = false;
-	}
 
 	const bool bGPUSkinCacheEnabled = GPUSkinCache && GEnableGPUSkinCache &&
 		// Force skin cache enabled for ray tracing if the inline skinning technique was requested.
@@ -927,29 +958,6 @@ void FSkeletalMeshObjectGPUSkin::ProcessUpdatedDynamicData(FRHICommandList& RHIC
 #endif
 	}
 
-	if (Mode == EGPUSkinCacheEntryMode::Raster && MorphVertexBuffer)
-	{
-		if (MorphVertexBuffer->bNeedsInitialClear)
-		{
-			QUICK_SCOPE_CYCLE_COUNTER(STAT_FSkeletalMeshObjectGPUSkin_ProcessUpdatedDynamicData_ClearMorphBuffer);
-			if (MorphVertexBuffer->GetUAV())
-			{
-				RHICmdList.Transition(FRHITransitionInfo(MorphVertexBuffer->GetUAV(), ERHIAccess::Unknown, ERHIAccess::UAVCompute));
-				RHICmdList.ClearUAVUint(MorphVertexBuffer->GetUAV(), FUintVector4(0, 0, 0, 0));
-				RHICmdList.Transition(FRHITransitionInfo(MorphVertexBuffer->GetUAV(), ERHIAccess::UAVCompute, ERHIAccess::SRVMask));
-			}
-			MorphVertexBuffer->bNeedsInitialClear = false;
-		}
-
-		// only update if the morph data changed and there are weighted morph targets
-		if (bMorphNeedsUpdate)
-		{
-			bMorphNeedsUpdate = false;
-			UpdateMorphVertexBuffer(RHICmdList, Mode, LOD, LODData, bGPUSkinCacheEnabled, *MorphVertexBuffer);
-			LOD.MorphVertexBufferPool->SetUpdatedFrameNumber(BoneTransformFrameNumber);
-		}
-	}
-
 	bool bAllowAddToSkinCache = bGPUSkinCacheEnabled;
 	bool bAllSectionsUsingSkinCache = true;
 
@@ -1203,26 +1211,6 @@ void FSkeletalMeshObjectGPUSkin::UpdateRayTracingGeometry(FRHICommandListBase& R
 
 #endif // RHI_RAYTRACING
 
-
-int32 FSkeletalMeshObjectGPUSkin::CalcNumActiveGPUMorphSets(FMorphVertexBuffer& MorphVertexBuffer, const FExternalMorphSets& ExternalMorphSets) const
-{
-	if (!UseGPUMorphTargets(FeatureLevel) || !IsValidRef(MorphVertexBuffer.VertexBufferRHI))
-	{
-		return 0;
-	}
-
-	// Count all active external morph sets.
-	int32 NumMorphSets = 1; // Start at one, as we have our standard morph targets as well.
-	for (const auto& MorphSet : ExternalMorphSets)
-	{
-		if (IsExternalMorphSetActive(MorphSet.Key, *MorphSet.Value))
-		{
-			NumMorphSets++;
-		}
-	}
-	return NumMorphSets;
-}
-
 bool FSkeletalMeshObjectGPUSkin::IsExternalMorphSetActive(int32 MorphSetID, const FExternalMorphSet& MorphSet) const
 {
 	const FMorphTargetVertexInfoBuffers& CompressedBuffers = MorphSet.MorphBuffers;
@@ -1310,14 +1298,28 @@ static void CalculateMorphDeltaBoundsIncludingExternalMorphs(
 	);
 }
 
-void FSkeletalMeshObjectGPUSkin::UpdateMorphVertexBuffer(FRHICommandList& RHICmdList, EGPUSkinCacheEntryMode Mode, FSkeletalMeshObjectLOD& LOD, const FSkeletalMeshLODRenderData& LODData, 
-															bool bGPUSkinCacheEnabled, FMorphVertexBuffer& MorphVertexBuffer)
+void FSkeletalMeshObjectGPUSkin::UpdateMorphVertexBuffer(
+	FRHICommandList& RHICmdList,
+	EGPUSkinCacheEntryMode Mode,
+	FSkeletalMeshObjectLOD& LOD,
+	const FSkeletalMeshLODRenderData& LODData,
+	bool bGPUSkinCacheEnabled)
 {
 	QUICK_SCOPE_CYCLE_COUNTER(STAT_FSkeletalMeshObjectGPUSkin_ProcessUpdatedDynamicData_UpdateMorphBuffer);
-	
-	if (UseGPUMorphTargets(FeatureLevel) && IsValidRef(MorphVertexBuffer.VertexBufferRHI))
+	check(MorphVertexBuffer && MorphVertexBuffer->VertexBufferRHI);
+
+	if (UseGPUMorphTargets(FeatureLevel))
 	{
-		const int32 NumMorphSets = CalcNumActiveGPUMorphSets(MorphVertexBuffer, DynamicData->ExternalMorphSets);
+		// Count all active external morph sets.
+		int32 NumMorphSets = 1; // Start at one, as we have our standard morph targets as well.
+		for (const auto& MorphSet : DynamicData->ExternalMorphSets)
+		{
+			if (IsExternalMorphSetActive(MorphSet.Key, *MorphSet.Value))
+			{
+				NumMorphSets++;
+			}
+		}
+
 		int32 MorphSetIndex = 0;
 
 		// Calculate the delta bounds.
@@ -1343,7 +1345,7 @@ void FSkeletalMeshObjectGPUSkin::UpdateMorphVertexBuffer(FRHICommandList& RHICmd
 			DynamicData->SectionIdsUseByActiveMorphTargets,
 			GetDebugName(), 
 			Mode, 
-			MorphVertexBuffer, 
+			*MorphVertexBuffer, 
 			true,	// Only clear the morph vertex buffer at the first morph set.
 			(MorphSetIndex == NumMorphSets-1),
 			MorphScale,
@@ -1366,7 +1368,7 @@ void FSkeletalMeshObjectGPUSkin::UpdateMorphVertexBuffer(FRHICommandList& RHICmd
 					DynamicData->SectionIdsUseByActiveMorphTargets,
 					GetDebugName(),
 					Mode,
-					MorphVertexBuffer,
+					*MorphVertexBuffer,
 					false,	// Don't clear the vertex buffer as we already did with the standard morph targets above.
 					(MorphSetIndex == NumMorphSets - 1),
 					MorphScale,
@@ -1382,7 +1384,7 @@ void FSkeletalMeshObjectGPUSkin::UpdateMorphVertexBuffer(FRHICommandList& RHICmd
 	else
 	{
 		// update the morph data for the lod (before SkinCache)
-		LOD.UpdateMorphVertexBufferCPU(RHICmdList, DynamicData->ActiveMorphTargets, DynamicData->MorphTargetWeights, DynamicData->SectionIdsUseByActiveMorphTargets, bGPUSkinCacheEnabled, MorphVertexBuffer);
+		LOD.UpdateMorphVertexBufferCPU(RHICmdList, DynamicData->ActiveMorphTargets, DynamicData->MorphTargetWeights, DynamicData->SectionIdsUseByActiveMorphTargets, bGPUSkinCacheEnabled, *MorphVertexBuffer);
 	}
 }
 
diff --git a/Engine/Source/Runtime/Engine/Private/SkeletalRenderGPUSkin.h b/Engine/Source/Runtime/Engine/Private/SkeletalRenderGPUSkin.h
--- a/Engine/Source/Runtime/Engine/Private/SkeletalRenderGPUSkin.h
+++ b/Engine/Source/Runtime/Engine/Private/SkeletalRenderGPUSkin.h
@@ -460,12 +460,6 @@ public:
 	}
 	//~ End FSkeletalMeshObject Interface
 
-	/**
-	 * Calculate how many GPU compressed morph target sets are active.
-	 * This includes regular morph targets as well as external morph targets.
-	 */
-	int32 CalcNumActiveGPUMorphSets(FMorphVertexBuffer& MorphVertexBuffer, const FExternalMorphSets& ExternalMorphSets) const;
-
 	/** Check if a given morph set is active or not. If so, we will process it. */
 	bool IsExternalMorphSetActive(int32 MorphSetID, const FExternalMorphSet& MorphSet) const;
 
@@ -677,10 +671,10 @@ private:
 
 	void InitMorphResources();
 
-	void UpdateBoneData(FRHICommandList& RHICmdList, EGPUSkinCacheEntryMode Mode);
+	void UpdateBufferData(FRHICommandList& RHICmdList, EGPUSkinCacheEntryMode Mode);
 	void ProcessUpdatedDynamicData(FRHICommandList& RHICmdList, FGPUSkinCache* GPUSkinCache, EGPUSkinCacheEntryMode Mode);
 
-	void UpdateMorphVertexBuffer(FRHICommandList& RHICmdList, EGPUSkinCacheEntryMode Mode, FSkeletalMeshObjectLOD& LOD, const FSkeletalMeshLODRenderData& LODData, bool bGPUSkinCacheEnabled, FMorphVertexBuffer& MorphVertexBuffer);
+	void UpdateMorphVertexBuffer(FRHICommandList& RHICmdList, EGPUSkinCacheEntryMode Mode, FSkeletalMeshObjectLOD& LOD, const FSkeletalMeshLODRenderData& LODData, bool bGPUSkinCacheEnabled);
 
 	bool IsRayTracingSkinCacheUpdateNeeded() const
 	{
@@ -694,8 +688,11 @@ private:
 #endif
 	}
 
+	bool IsSkinCacheEnabled(EGPUSkinCacheEntryMode Mode) const;
+
 	FSkeletalMeshUpdateHandle UpdateHandle;
 
+	FMorphVertexBuffer* MorphVertexBuffer = nullptr;
 	bool bMorphResourcesInitialized = false;
 	bool bMorphNeedsUpdate = false;
 
