Increase the number of safe frames to 4 for skeletal mesh bone buffers, to fix flickering issue in CitySample with vsync off.

diff --git a/Engine/Source/Runtime/Engine/Public/GPUSkinVertexFactory.h b/Engine/Source/Runtime/Engine/Public/GPUSkinVertexFactory.h
--- a/Engine/Source/Runtime/Engine/Public/GPUSkinVertexFactory.h
+++ b/Engine/Source/Runtime/Engine/Public/GPUSkinVertexFactory.h
@@ -31,7 +31,7 @@ public:
 	typedef uint32 CreationArguments;
 	enum
 	{
-		NumSafeFrames = 3, /** Number of frames to leaves buffers before reclaiming/reusing */
+		NumSafeFrames = 4, /** Number of frames to leaves buffers before reclaiming/reusing */
 		NumPoolBucketSizes = 18, /** Number of pool buckets */
 		NumToDrainPerFrame = 10, /** Max. number of resources to cull in a single frame */
 		CullAfterFramesNum = 30 /** Resources are culled if unused for more frames than this */
