[shaders] add accessor function for GCreateShadersOnLoad and remove dllexport, this only appears to be used inside the Engine module and is not exposed in any header

diff --git a/Engine/Source/Runtime/Engine/Private/Materials/MaterialShader.cpp b/Engine/Source/Runtime/Engine/Private/Materials/MaterialShader.cpp
--- a/Engine/Source/Runtime/Engine/Private/Materials/MaterialShader.cpp
+++ b/Engine/Source/Runtime/Engine/Private/Materials/MaterialShader.cpp
@@ -3419,8 +3419,7 @@ void FMaterialShaderMap::Register(EShaderPlatform InShaderPlatform)
 		}
 	} MaterialShaderMapInnerLazyInitializer;
 
-	extern int32 GCreateShadersOnLoad;
-	if (GCreateShadersOnLoad && GMaxRHIShaderPlatform == InShaderPlatform)
+	if (CreateShadersOnLoad() && GMaxRHIShaderPlatform == InShaderPlatform)
 	{
 		FShaderMapResource* ShaderResource = GetResource();
 		if (ShaderResource)
diff --git a/Engine/Source/Runtime/Engine/Private/ShaderCompiler/ShaderCompiler.cpp b/Engine/Source/Runtime/Engine/Private/ShaderCompiler/ShaderCompiler.cpp
--- a/Engine/Source/Runtime/Engine/Private/ShaderCompiler/ShaderCompiler.cpp
+++ b/Engine/Source/Runtime/Engine/Private/ShaderCompiler/ShaderCompiler.cpp
@@ -506,12 +506,16 @@ static TAutoConsoleVariable<int32> CVarCullBeforeFetch(
 	TEXT(" 1: Enable"),
 	ECVF_ReadOnly);
 
-ENGINE_API int32 GCreateShadersOnLoad = 0;
+int32 GCreateShadersOnLoad = 0;
 static FAutoConsoleVariableRef CVarCreateShadersOnLoad(
 	TEXT("r.CreateShadersOnLoad"),
 	GCreateShadersOnLoad,
-	TEXT("Whether to create shaders on load, which can reduce hitching, but use more memory.  Otherwise they will be created as needed.")
-);
+	TEXT("Whether to create shaders on load, which can reduce hitching, but use more memory.  Otherwise they will be created as needed."));
+
+bool CreateShadersOnLoad()
+{
+	return GCreateShadersOnLoad != 0;
+}
 
 static TAutoConsoleVariable<int32> CVarShadersValidation(
 	TEXT("r.Shaders.Validation"),
@@ -534,6 +538,7 @@ namespace ShaderCompiler
 	{
 		return GDumpShaderDebugInfo != FShaderCompilingManager::EDumpShaderDebugInfo::Always;
 	}
+
 } // namespace ShaderCompiler
 
 #if ENABLE_COOK_STATS
@@ -5329,7 +5334,7 @@ void CompileGlobalShaderMap(EShaderPlatform Platform, const ITargetPlatform* Tar
 		// If any shaders weren't loaded, compile them now.
 		VerifyGlobalShaders(Platform, TargetPlatform, bLoadedFromCacheFile);
 
-		if (GCreateShadersOnLoad && Platform == GMaxRHIShaderPlatform)
+		if (CreateShadersOnLoad() && Platform == GMaxRHIShaderPlatform)
 		{
 			GGlobalShaderMap[Platform]->BeginCreateAllShaders();
 		}
diff --git a/Engine/Source/Runtime/Engine/Public/ShaderCompiler.h b/Engine/Source/Runtime/Engine/Public/ShaderCompiler.h
--- a/Engine/Source/Runtime/Engine/Public/ShaderCompiler.h
+++ b/Engine/Source/Runtime/Engine/Public/ShaderCompiler.h
@@ -51,6 +51,7 @@ DECLARE_LOG_CATEGORY_EXTERN(LogShaderCompilers, Log, All);
 #define DEBUG_INFINITESHADERCOMPILE 0
 
 bool AreShaderErrorsFatal();
+bool CreateShadersOnLoad();
 
 extern ENGINE_API bool IsShaderJobCacheDDCEnabled();
 extern ENGINE_API bool IsMaterialMapDDCEnabled();
