SceneGraph
* Simplified replication_component DisableReplication test
* FTestWorlds::IsServerObjectReplicated now returns true if the object is part of the replication world on the server.

#rnx

diff --git a/Engine/Source/Runtime/Engine/Private/Tests/NetTestHelpers.cpp b/Engine/Source/Runtime/Engine/Private/Tests/NetTestHelpers.cpp
--- a/Engine/Source/Runtime/Engine/Private/Tests/NetTestHelpers.cpp
+++ b/Engine/Source/Runtime/Engine/Private/Tests/NetTestHelpers.cpp
@@ -531,7 +531,7 @@ bool FTestWorlds::DoesReplicatedObjectExistOnClient(UObject* ServerObject, uint3
 	return IsValid(FindReplicatedObjectOnClient(ServerObject, ClientIndex));
 }
 
-bool FTestWorlds::IsServerObjectReplicated(const UObject* ServerObject, uint32 ClientIndex) const
+bool FTestWorlds::IsServerObjectReplicated(const UObject* ServerObject) const
 {
 	using namespace UE::Net;
 
@@ -542,17 +542,12 @@ bool FTestWorlds::IsServerObjectReplicated(const UObject* ServerObject, uint32 C
 	}
 
 	const FTestWorldInstance::FContext ServerContext = Server.GetTestContext();
-	const FTestWorldInstance::FContext ClientContext = Clients[ClientIndex].GetTestContext();
 	
 	if (ServerContext.NetDriver->IsUsingIrisReplication())
 	{
 #if UE_WITH_IRIS
 		const FNetRefHandle NetHandle = ServerContext.IrisRepBridge->GetReplicatedRefHandle(ServerObject);
-		if (!NetHandle.IsValid())
-		{
-			return false;
-		}
-		return ClientContext.IrisRepBridge->GetReplicatedObject(NetHandle) != nullptr;
+		return NetHandle.IsValid();
 #else
 		return false;
 #endif
@@ -560,11 +555,7 @@ bool FTestWorlds::IsServerObjectReplicated(const UObject* ServerObject, uint32 C
 	else
 	{
 		const FNetworkGUID NetGUID = ServerContext.NetDriver->GetNetGuidCache()->GetNetGUID(ServerObject);
-		if (!NetGUID.IsValid())
-		{
-			return false;
-		}
-		return ClientContext.NetDriver->GetNetGuidCache()->GetObjectFromNetGUID(NetGUID, false) != nullptr;
+		return NetGUID.IsValid();
 	}
 }
 
diff --git a/Engine/Source/Runtime/Engine/Public/Tests/NetTestHelpers.h b/Engine/Source/Runtime/Engine/Public/Tests/NetTestHelpers.h
--- a/Engine/Source/Runtime/Engine/Public/Tests/NetTestHelpers.h
+++ b/Engine/Source/Runtime/Engine/Public/Tests/NetTestHelpers.h
@@ -223,10 +223,10 @@ struct FTestWorlds
 	ENGINE_API bool DoesReplicatedObjectExistOnClient(UObject* ServerObject, uint32 ClientIndex) const;
 
 	/**
-	 * Returns if the object is replicated on clients or not.
+	 * Returns if the object is registered for replication on the server.
 	 * @param ServerObject The server object you want to see if it is replicated
 	 */
-	ENGINE_API bool IsServerObjectReplicated(const UObject* ServerObject, uint32 ClientIndex) const;
+	ENGINE_API bool IsServerObjectReplicated(const UObject* ServerObject) const;
 
 public:
 
