[DSTM] Fix for ensure that might occur when migrating a physics state

In addition, I updated the code to apply the physics state if we already have a valid state instead of re-creating it (which is the case that more ofter occurs during borrowing)

diff --git a/Engine/Source/Runtime/Engine/Classes/Components/PrimitiveComponent.h b/Engine/Source/Runtime/Engine/Classes/Components/PrimitiveComponent.h
--- a/Engine/Source/Runtime/Engine/Classes/Components/PrimitiveComponent.h
+++ b/Engine/Source/Runtime/Engine/Classes/Components/PrimitiveComponent.h
@@ -3158,7 +3158,7 @@ protected:
 #if UE_WITH_REMOTE_OBJECT_HANDLE
 	ENGINE_API virtual void MigratePhysicsState(FArchive& Ar) override;
 	
-	ENGINE_API virtual void ApplyMigratedPhysicsState() override;
+	ENGINE_API virtual void ApplyMigratedPhysicsState(Chaos::FSerializedDataBufferPtr MigratedStateData = nullptr) override;
 #endif
 
 protected:
diff --git a/Engine/Source/Runtime/Engine/Classes/Components/SkeletalMeshComponent.h b/Engine/Source/Runtime/Engine/Classes/Components/SkeletalMeshComponent.h
--- a/Engine/Source/Runtime/Engine/Classes/Components/SkeletalMeshComponent.h
+++ b/Engine/Source/Runtime/Engine/Classes/Components/SkeletalMeshComponent.h
@@ -2016,7 +2016,7 @@ protected:
 
 #if UE_WITH_REMOTE_OBJECT_HANDLE
 	ENGINE_API virtual void MigratePhysicsState(FArchive& Ar) override;
-	ENGINE_API virtual void ApplyMigratedPhysicsState() override;
+	ENGINE_API virtual void ApplyMigratedPhysicsState(Chaos::FSerializedDataBufferPtr MigratedStateData = nullptr) override;
 #endif
 
 	//~ End UPrimitiveComponent Interface.
diff --git a/Engine/Source/Runtime/Engine/Classes/Interfaces/IPhysicsComponent.h b/Engine/Source/Runtime/Engine/Classes/Interfaces/IPhysicsComponent.h
--- a/Engine/Source/Runtime/Engine/Classes/Interfaces/IPhysicsComponent.h
+++ b/Engine/Source/Runtime/Engine/Classes/Interfaces/IPhysicsComponent.h
@@ -4,6 +4,7 @@
 
 #include "Chaos/ParticleHandleFwd.h"
 #include "Chaos/PhysicsObject.h"
+#include "Chaos/Serialization/SolverSerializer.h"
 #include "Containers/Array.h"
 #include "UObject/ObjectMacros.h"
 #include "UObject/Interface.h"
@@ -31,7 +32,7 @@ public:
 
 protected:
 	virtual void MigratePhysicsState(FArchive& Ar) = 0;
-	virtual void ApplyMigratedPhysicsState() = 0;
+	virtual void ApplyMigratedPhysicsState(Chaos::FSerializedDataBufferPtr MigratedStateData = nullptr) = 0;
 public:
 
 #endif
diff --git a/Engine/Source/Runtime/Engine/Private/Components/PrimitiveComponent.cpp b/Engine/Source/Runtime/Engine/Private/Components/PrimitiveComponent.cpp
--- a/Engine/Source/Runtime/Engine/Private/Components/PrimitiveComponent.cpp
+++ b/Engine/Source/Runtime/Engine/Private/Components/PrimitiveComponent.cpp
@@ -1284,16 +1284,14 @@ void UPrimitiveComponent::MigratePhysicsState(FArchive& Ar)
 			return;
 		}
 
-		if (FPhysScene* PhysicsScene = FPhysicsInterface::GetCurrentSceneForComponent(this))
+		if (HasValidPhysicsState())
+		{
+			ApplyMigratedPhysicsState(MoveTemp(MigratedStateDataBuffer));
+		}
+		else if (FPhysScene* PhysicsScene = FPhysicsInterface::GetCurrentSceneForComponent(this))
 		{
 			PhysicsScene->PushPendingSerializedStateForObject(this, MoveTemp(MigratedStateDataBuffer));
-
-			OnComponentPhysicsStateChanged.AddDynamic(this, &UPrimitiveComponent::HandlePostMigrationPhysicsStateChanged);
-
-			if (HasValidPhysicsState())
-			{
-				RecreatePhysicsState();
-			}
+			OnComponentPhysicsStateChanged.AddUniqueDynamic(this, &UPrimitiveComponent::HandlePostMigrationPhysicsStateChanged);
 		}
 	}
 	else if (Ar.IsSaving())
@@ -1323,9 +1321,18 @@ void UPrimitiveComponent::MigratePhysicsState(FArchive& Ar)
 	}
 }
 
-void UPrimitiveComponent::ApplyMigratedPhysicsState()
+void UPrimitiveComponent::ApplyMigratedPhysicsState(Chaos::FSerializedDataBufferPtr MigratedStateData)
 {
-	if (ensure(HasValidPhysicsState()))
+	if (!ensure(HasValidPhysicsState()))
+	{
+		return;
+	}
+
+	if (MigratedStateData)
+	{
+		BodyInstance.ApplySerializedState(MoveTemp(MigratedStateData));
+	}
+	else
 	{
 		FPhysScene* PhysicsScene = FPhysicsInterface::GetCurrentSceneForComponent(this);
 		if (Chaos::FSerializedDataBufferPtr RawMigratedData = PhysicsScene ? PhysicsScene->PopPendingSerializedStateForObject(this) : nullptr)
diff --git a/Engine/Source/Runtime/Engine/Private/Components/SkeletalMeshComponent.cpp b/Engine/Source/Runtime/Engine/Private/Components/SkeletalMeshComponent.cpp
--- a/Engine/Source/Runtime/Engine/Private/Components/SkeletalMeshComponent.cpp
+++ b/Engine/Source/Runtime/Engine/Private/Components/SkeletalMeshComponent.cpp
@@ -3877,7 +3877,11 @@ void USkeletalMeshComponent::MigratePhysicsState(FArchive& Ar)
 			return;
 		}
 
-		if (FPhysScene* PhysicsScene = FPhysicsInterface::GetCurrentSceneForComponent(this))
+		if (HasValidPhysicsState())
+		{
+			ApplyMigratedPhysicsState(MoveTemp(MigratedStateData));
+		}
+		else if (FPhysScene* PhysicsScene = FPhysicsInterface::GetCurrentSceneForComponent(this))
 		{
 			PhysicsScene->PushPendingSerializedStateForObject(this, MoveTemp(MigratedStateData));
 
@@ -3888,11 +3892,6 @@ void USkeletalMeshComponent::MigratePhysicsState(FArchive& Ar)
 					SKComponent->ApplyMigratedPhysicsState();
 				}
 			}));
-
-			if (HasValidPhysicsState())
-			{
-				RecreatePhysicsState();
-			}
 		}
 	}
 	else if (Ar.IsSaving())
@@ -3940,11 +3939,20 @@ void USkeletalMeshComponent::MigratePhysicsState(FArchive& Ar)
 	}
 }
 
-void USkeletalMeshComponent::ApplyMigratedPhysicsState()
+void USkeletalMeshComponent::ApplyMigratedPhysicsState(Chaos::FSerializedDataBufferPtr MigratedStateData)
 {
-	FPhysScene* PhysicsScene = FPhysicsInterface::GetCurrentSceneForComponent(this);
-	
-	if (Chaos::FSerializedDataBufferPtr RawMigratedData = PhysicsScene ? PhysicsScene->PopPendingSerializedStateForObject(this) : nullptr)
+	Chaos::FSerializedDataBufferPtr RawMigratedData;
+
+	if (MigratedStateData)
+	{
+		RawMigratedData = MoveTemp(MigratedStateData);
+	}
+	else if (FPhysScene* PhysicsScene = FPhysicsInterface::GetCurrentSceneForComponent(this))
+	{
+		RawMigratedData = PhysicsScene->PopPendingSerializedStateForObject(this);
+	}
+
+	if (RawMigratedData)
 	{
 		FMemoryReader MultiBodyStateReader(RawMigratedData->GetDataAsByteArrayRef());
 
