[HWRT] FStaticMeshRayTracingProxy adopts fallback meshes instead of preventing stripping.

diff --git a/Engine/Source/Runtime/Engine/Private/StaticMesh.cpp b/Engine/Source/Runtime/Engine/Private/StaticMesh.cpp
--- a/Engine/Source/Runtime/Engine/Private/StaticMesh.cpp
+++ b/Engine/Source/Runtime/Engine/Private/StaticMesh.cpp
@@ -199,8 +199,7 @@ static bool ShouldStripNaniteFallbackMesh(const ITargetPlatform* TargetPlatform,
 		&& StaticMesh->NaniteSettings.GenerateFallback != ENaniteGenerateFallback::Enabled
 		&& StaticMesh->HasValidNaniteData()
 		&& !StaticMesh->bAllowCPUAccess
-		&& !TargetPlatform->HasEditorOnlyData()
-		&& (!StaticMesh->bSupportRayTracing || !TargetPlatform->UsesRayTracing()); // TODO: ray tracing proxy should adopt the fallback mesh data
+		&& !TargetPlatform->HasEditorOnlyData();
 }
 
 static void FillMaterialName(const TArray<FStaticMaterial>& StaticMaterials, TMap<int32, FName>& OutMaterialMap)
@@ -1742,6 +1741,23 @@ void FStaticMeshRayTracingProxyLOD::Serialize(FArchive& Ar, UObject* Owner, int3
 {
 	UStaticMesh* OwnerStaticMesh = Cast<UStaticMesh>(Owner);
 
+#if WITH_EDITOR
+	// if the ray tracing proxy is using the fallback mesh but it will be stripped from the cooked build
+	// the ray tracing proxy needs to adopt the fallback mesh
+	const bool bAdoptFallbackMesh = ShouldStripNaniteFallbackMesh(Ar.CookingTarget(), OwnerStaticMesh);
+#else
+	const bool bAdoptFallbackMesh = false;
+#endif
+
+	const bool bOwnsBuffersBackup = bOwnsBuffers;
+	const bool bOwnsRayTracingGeometryBackup = bOwnsRayTracingGeometry;
+
+	if (bAdoptFallbackMesh)
+	{
+		bOwnsBuffers = true;
+		bOwnsRayTracingGeometry = true;
+	}
+
 	bool bTmpOwnsBuffers = bOwnsBuffers;
 	Ar << bTmpOwnsBuffers;
 	bOwnsBuffers = bTmpOwnsBuffers;
@@ -1825,6 +1841,13 @@ void FStaticMeshRayTracingProxyLOD::Serialize(FArchive& Ar, UObject* Owner, int3
 
 		SerializeMetaData(Ar);
 	}
+
+	if (bAdoptFallbackMesh)
+	{
+		// revert back to original state to prevent issues with ownership
+		bOwnsBuffers = bOwnsBuffersBackup;
+		bOwnsRayTracingGeometry = bOwnsRayTracingGeometryBackup;
+	}
 }
 
 void FStaticMeshRayTracingProxyLOD::SerializeBuffers(FArchive& Ar, UStaticMesh* OwnerStaticMesh, uint8 InStripFlags)
@@ -2076,6 +2099,23 @@ void FStaticMeshRayTracingProxy::Serialize(FArchive& Ar, UObject* Owner, FStatic
 {
 	FStripDataFlags StripFlags(Ar);
 
+	UStaticMesh* OwnerStaticMesh = Cast<UStaticMesh>(Owner);
+
+#if WITH_EDITOR
+	// if the ray tracing proxy is using the fallback mesh but it will be stripped from the cooked build
+	// the ray tracing proxy needs to adopt the fallback mesh
+	const bool bAdoptFallbackMesh = ShouldStripNaniteFallbackMesh(Ar.CookingTarget(), OwnerStaticMesh);
+#else
+	const bool bAdoptFallbackMesh = false;
+#endif
+
+	const bool bUsingRenderingLODsBackup = bUsingRenderingLODs;
+
+	if (bAdoptFallbackMesh)
+	{
+		bUsingRenderingLODs = false;
+	}
+
 	Ar << bUsingRenderingLODs;
 
 	if (!StripFlags.IsAudioVisualDataStripped())
@@ -2118,6 +2158,14 @@ void FStaticMeshRayTracingProxy::Serialize(FArchive& Ar, UObject* Owner, FStatic
 			}
 		}
 	}
+
+#if WITH_EDITOR
+	if (bAdoptFallbackMesh)
+	{
+		// revert back to original state to prevent issues with ownership
+		bUsingRenderingLODs = bUsingRenderingLODsBackup;
+	}
+#endif
 }
 
 /*------------------------------------------------------------------------------
@@ -3493,7 +3541,7 @@ static FString BuildStaticMeshDerivedDataKeySuffix(const ITargetPlatform* Target
 
 		if (ShouldStripNaniteFallbackMesh(TargetPlatform, Mesh))
 		{
-			KeySuffix += TEXT("NFMS");
+			KeySuffix += TEXT("NFMSV2");
 		}
 	}
 
