FSpline: Set to not be implemented by default

diff --git a/Engine/Source/Runtime/Engine/Private/Curves/Spline.cpp b/Engine/Source/Runtime/Engine/Private/Curves/Spline.cpp
--- a/Engine/Source/Runtime/Engine/Private/Curves/Spline.cpp
+++ b/Engine/Source/Runtime/Engine/Private/Curves/Spline.cpp
@@ -10,7 +10,7 @@
 namespace UE::Spline
 {
 
-int32 GImplementation = 1;
+int32 GImplementation = 0;
 #if WITH_EDITOR
 static FSimpleMulticastDelegate& GetOnSplineImplementationChanged()
 {
