PR #12165: Remove unnecessary FTransform ZeroRotator->Quat conversion in (H)ISM

diff --git a/Engine/Source/Runtime/Engine/Private/HierarchicalInstancedStaticMesh.cpp b/Engine/Source/Runtime/Engine/Private/HierarchicalInstancedStaticMesh.cpp
--- a/Engine/Source/Runtime/Engine/Private/HierarchicalInstancedStaticMesh.cpp
+++ b/Engine/Source/Runtime/Engine/Private/HierarchicalInstancedStaticMesh.cpp
@@ -1646,7 +1646,7 @@ void FHierarchicalStaticMeshSceneProxy::GetDynamicMeshElements(const TArray<cons
 				float MaxDrawDistanceScale = GetCachedScalabilityCVars().ViewDistanceScale;
 								
 				FVector AverageScale(InstanceParams.Tree[0].MinInstanceScale + (InstanceParams.Tree[0].MaxInstanceScale - InstanceParams.Tree[0].MinInstanceScale) / 2.0f);
-				FBoxSphereBounds ScaledBounds = RenderData->Bounds.TransformBy(FTransform(FRotator::ZeroRotator, FVector::ZeroVector, AverageScale));
+				FBoxSphereBounds ScaledBounds = RenderData->Bounds.TransformBy(FTransform(FQuat::Identity, FVector::ZeroVector, AverageScale));
 				float SphereRadius = ScaledBounds.SphereRadius + InstanceParams.MaxWPODisplacement;
 
 				float FinalCull = MAX_flt;
diff --git a/Engine/Source/Runtime/Engine/Private/InstancedStaticMesh.cpp b/Engine/Source/Runtime/Engine/Private/InstancedStaticMesh.cpp
--- a/Engine/Source/Runtime/Engine/Private/InstancedStaticMesh.cpp
+++ b/Engine/Source/Runtime/Engine/Private/InstancedStaticMesh.cpp
@@ -1244,7 +1244,7 @@ FInstancedStaticMeshVFLooseUniformShaderParametersRef FInstancedStaticMeshSceneP
 				FirstLOD = FMath::Max(FirstLOD, DebugMin);
 			}
 
-			FBoxSphereBounds ScaledBounds = InstancingUserData->MeshRenderData->Bounds.TransformBy(FTransform(FRotator::ZeroRotator, FVector::ZeroVector, InstancingUserData->AverageInstancesScale));
+			FBoxSphereBounds ScaledBounds = InstancingUserData->MeshRenderData->Bounds.TransformBy(FTransform(FQuat::Identity, FVector::ZeroVector, InstancingUserData->AverageInstancesScale));
 			float SphereRadius = ScaledBounds.SphereRadius;
 			float MinSize = View->ViewMatrices.IsPerspectiveProjection() ? CVarFoliageMinimumScreenSize.GetValueOnRenderThread() : 0.0f;
 			float LODScale = InstancingUserData->LODDistanceScale;
