Allow Streamable Handles to set and get download priority.

diff --git a/Engine/Source/Runtime/Engine/Classes/Engine/StreamableManager.h b/Engine/Source/Runtime/Engine/Classes/Engine/StreamableManager.h
--- a/Engine/Source/Runtime/Engine/Classes/Engine/StreamableManager.h
+++ b/Engine/Source/Runtime/Engine/Classes/Engine/StreamableManager.h
@@ -201,6 +201,17 @@ struct TStreamableHandleContextData : public FStreamableHandleContextDataBase
 	}
 };
 
+/** EXPERIMENTAL - download priority */
+enum class EStreamableDownloadPriority : uint8
+{
+	/** Below normal priority. */
+	Low,
+	/** Normal priority. */
+	Normal,
+	/** Above normal priority. */
+	High
+};
+
 /** A handle to a synchronous or async load. As long as the handle is Active, loaded assets will stay in memory */
 struct FStreamableHandle : public TSharedFromThis<FStreamableHandle>
 {
@@ -300,6 +311,15 @@ struct FStreamableHandle : public TSharedFromThis<FStreamableHandle>
 		return Priority;
 	}
 
+	/** Returns download priority. Only valid if a download has been requested and is not yet complete. */
+	ENGINE_API TOptional<EStreamableDownloadPriority> GetDownloadPriority() const;
+
+	/** 
+	 * Set download priority. Only valid if a download has been requested and is not yet complete. 
+	 * @return true if the priority could be set.
+	 */
+	ENGINE_API bool SetDownloadPriority(EStreamableDownloadPriority Priority);
+
 	/**
 	 * Release this handle. This can be called from normal gameplay code to indicate that the loaded assets are no longer needed
 	 * This will be called implicitly if all shared pointers to this handle are destroyed
@@ -675,16 +695,6 @@ namespace UE::StreamableManager
 	static constexpr TAsyncLoadPriority AsyncLoadHighPriority = 100;
 }
 
-enum class EStreamableDownloadPriority : uint8
-{
-	/** Below normal priority. */
-	Low,
-	/** Normal priority. */
-	Normal,
-	/** Above normal priority. */
-	High
-};
-
 /** EXPERIMENTAL -  Additional parameters for downloading packages */
 struct FStreamableDownloadParams
 {
diff --git a/Engine/Source/Runtime/Engine/Private/StreamableManager.cpp b/Engine/Source/Runtime/Engine/Private/StreamableManager.cpp
--- a/Engine/Source/Runtime/Engine/Private/StreamableManager.cpp
+++ b/Engine/Source/Runtime/Engine/Private/StreamableManager.cpp
@@ -380,6 +380,22 @@ TStreamableHandleContextDataTypeID FStreamableHandleContextDataBase::AllocateCla
 
 namespace UE::StreamableManager::Private
 {
+	static UE::IoStore::EOnDemandInstallPriority GetOnDemandInstallPriority(EStreamableDownloadPriority Priority)
+	{
+		UE::IoStore::EOnDemandInstallPriority OutPriority = UE::IoStore::EOnDemandInstallPriority::Normal;
+		switch (Priority)
+		{
+		case EStreamableDownloadPriority::Low:
+			OutPriority = UE::IoStore::EOnDemandInstallPriority::Low;
+			break;
+		case EStreamableDownloadPriority::High:
+			OutPriority = UE::IoStore::EOnDemandInstallPriority::High;
+			break;
+		}
+
+		return OutPriority;
+	}
+
 	// self-deleting context for managing package downloads for a load request
 	struct FDownloadContext
 	{
@@ -459,16 +475,7 @@ namespace UE::StreamableManager::Private
 				InstallArgs.Options |= UE::IoStore::EOnDemandInstallOptions::InstallSoftReferences;
 			}
 
-			InstallArgs.Priority = UE::IoStore::EOnDemandInstallPriority::Normal;
-			switch (Priority)
-			{
-			case EStreamableDownloadPriority::Low:
-				InstallArgs.Priority = UE::IoStore::EOnDemandInstallPriority::Low;
-				break;
-			case EStreamableDownloadPriority::High:
-				InstallArgs.Priority = UE::IoStore::EOnDemandInstallPriority::High;
-				break;
-			}
+			InstallArgs.Priority = GetOnDemandInstallPriority(Priority);
 
 			InstallRequest = IoStore->Install(MoveTemp(InstallArgs),
 				[WeakOwningHandle = OwningHandle.ToWeakPtr()](const UE::IoStore::FOnDemandInstallResult& OnDemandInstallResult)
@@ -751,6 +758,30 @@ void FStreamableHandle::SetDebugNameIfEmptyOrDefault(const FString& NewName)
 	}
 }
 
+TOptional<EStreamableDownloadPriority> FStreamableHandle::GetDownloadPriority() const
+{
+	if (IsActive() && DownloadContext)
+	{
+		return DownloadContext->Priority;
+	}
+
+	return {};
+}
+
+bool FStreamableHandle::SetDownloadPriority(EStreamableDownloadPriority InPriority)
+{
+	if (IsActive() && DownloadContext)
+	{
+		DownloadContext->Priority = InPriority;
+		DownloadContext->InstallRequest.UpdatePriority(
+			UE::StreamableManager::Private::GetOnDemandInstallPriority(InPriority));
+
+		return true;
+	}
+
+	return false;
+}
+
 void FStreamableHandle::GetRequestedAssets(TArray<FSoftObjectPath>& AssetList, bool bIncludeChildren /*= true*/) const
 {
 	AssetList = RequestedAssets;
