Fix crash when Nanite resources is overridden on UStaticMeshComponent using OnGetNaniteResources.

diff --git a/Engine/Source/Runtime/Engine/Private/StaticMeshSceneProxy.cpp b/Engine/Source/Runtime/Engine/Private/StaticMeshSceneProxy.cpp
--- a/Engine/Source/Runtime/Engine/Private/StaticMeshSceneProxy.cpp
+++ b/Engine/Source/Runtime/Engine/Private/StaticMeshSceneProxy.cpp
@@ -2861,7 +2861,10 @@ FPrimitiveSceneProxy* UStaticMeshComponent::CreateSceneProxy()
 
 		FSceneInterface* Scene = GetScene();
 		const EShaderPlatform ShaderPlatform = Scene ? Scene->GetShaderPlatform() : GMaxRHIShaderPlatform;
-		const bool bFallbackGenerated = GetStaticMesh()->HasNaniteFallbackMesh(ShaderPlatform);
+
+		// TODO: handle Nanite representation being overriden using OnGetNaniteResources
+		// for now need to check UStaticMesh::HasValidNaniteData() directly here
+		const bool bFallbackGenerated = !GetStaticMesh()->HasValidNaniteData() || GetStaticMesh()->HasNaniteFallbackMesh(ShaderPlatform);
 
 		if (!bFallbackGenerated)
 		{
