[NetPhysics] Post-resim render interpolation snapping distance to be based on velocity.

diff --git a/Engine/Source/Runtime/Engine/Private/PhysicsEngine/NetworkPhysicsSettingsComponent.cpp b/Engine/Source/Runtime/Engine/Private/PhysicsEngine/NetworkPhysicsSettingsComponent.cpp
--- a/Engine/Source/Runtime/Engine/Private/PhysicsEngine/NetworkPhysicsSettingsComponent.cpp
+++ b/Engine/Source/Runtime/Engine/Private/PhysicsEngine/NetworkPhysicsSettingsComponent.cpp
@@ -225,6 +225,7 @@ void FNetworkPhysicsSettingsResimulationErrorCorrection::ApplySettings_External(
 					FErrorInterpolationSettings& Settings = InterpolationError->GetOrCreateErrorInterpolationSettings();
 					Settings.ErrorCorrectionDuration = ResimErrorCorrectionDuration;
 					Settings.MaximumErrorCorrectionBeforeSnapping = ResimErrorMaximumDistanceBeforeSnapping;
+					Settings.MaximumErrorCorrectionDesyncTimeBeforeSnapping = ResimErrorMaximumDesyncTimeBeforeSnapping;
 					Settings.ErrorDirectionalDecayMultiplier = ResimErrorDirectionalDecayMultiplier;
 				}
 			};
diff --git a/Engine/Source/Runtime/Engine/Public/Physics/NetworkPhysicsSettingsComponent.h b/Engine/Source/Runtime/Engine/Public/Physics/NetworkPhysicsSettingsComponent.h
--- a/Engine/Source/Runtime/Engine/Public/Physics/NetworkPhysicsSettingsComponent.h
+++ b/Engine/Source/Runtime/Engine/Public/Physics/NetworkPhysicsSettingsComponent.h
@@ -251,8 +251,9 @@ struct FNetworkPhysicsSettingsResimulationErrorCorrection
 {
 	FNetworkPhysicsSettingsResimulationErrorCorrection()
 	: bOverrideResimErrorInterpolationSettings(0)
-	, ResimErrorCorrectionDuration(0.5f)
+	, ResimErrorCorrectionDuration(0.3f)
 	, ResimErrorMaximumDistanceBeforeSnapping(250.0f)
+	, ResimErrorMaximumDesyncTimeBeforeSnapping(0.6f)
 	, ResimErrorDirectionalDecayMultiplier(0.0f)
 	{};
 
@@ -262,19 +263,24 @@ struct FNetworkPhysicsSettingsResimulationErrorCorrection
 	void ApplySettings_External(Chaos::FPhysicsObjectHandle PhysicsObject);
 
 	/** Enable override for post-resimulation error correction settings during render interpolation
-	* NOTE: This currently does not work if the experimental p.RenderInterp.ErrorVelocityCorrection CVar is set to true (false by default).*/
+	* NOTE: This currently does not work if the experimental p.RenderInterp.ErrorVelocityCorrection CVar is set to true (false by default) */
 	UPROPERTY(config, EditDefaultsOnly, Category = "Overrides")
 	uint32 bOverrideResimErrorInterpolationSettings : 1;
 
-	// Overrides CVar: p.RenderInterp.ErrorCorrectionDuration -- How long in seconds to apply error correction over
+	/** Overrides CVar: p.RenderInterp.ErrorCorrectionDuration -- How long in seconds to apply error correction over */
 	UPROPERTY(config, EditDefaultsOnly, Category = "Overrides", Meta = (EditCondition = "bOverrideResimErrorInterpolationSettings", ClampMin = "0.0", UIMin = "0.0"))
 	float ResimErrorCorrectionDuration;
 
-	// Overrides CVar: p.RenderInterp.MaximumErrorCorrectionBeforeSnapping -- Maximum error correction in cm before we stop interpolating and snap to target
+	/** Overrides CVar : p.RenderInterp.MaximumErrorCorrectionBeforeSnapping -- Maximum error correction in cm before we stop interpolating and snap to target */
 	UPROPERTY(config, EditDefaultsOnly, Category = "Overrides", Meta = (EditCondition = "bOverrideResimErrorInterpolationSettings", ClampMin = "0.0", UIMin = "0.0"))
 	float ResimErrorMaximumDistanceBeforeSnapping;
 
-	// Overrides CVar: p.RenderInterp.DirectionalDecayMultiplier -- Decay error offset in the direction that the physics object is moving, value is multiplier of projected offset direction, 0.25 means a 25% decay of the magnitude in the direction of physics travel. Deactivate by setting to 0..
+	/** Overrides CVar: p.RenderInterp.MaximumErrorCorrectionDesyncTimeBeforeSnapping -- Time multiplied by the particles velocity to get the distance that error correction will be performed within without snapping, disable by setting a negative value
+	* NOTE: ResimErrorMaximumDistanceBeforeSnapping will act as a lowest distance clamp. */
+	UPROPERTY(config, EditDefaultsOnly, Category = "Overrides", Meta = (EditCondition = "bOverrideResimErrorInterpolationSettings", ClampMin = "-1.0", UIMin = "-1.0"))
+	float ResimErrorMaximumDesyncTimeBeforeSnapping;
+
+	/** Overrides CVar: p.RenderInterp.DirectionalDecayMultiplier -- Decay error offset in the direction that the physics object is moving, value is multiplier of projected offset direction, 0.25 means a 25 % decay of the magnitude in the direction of physics travel.Deactivate by setting to 0 */
 	UPROPERTY(config, EditDefaultsOnly, Category = "Overrides", Meta = (EditCondition = "bOverrideResimErrorInterpolationSettings", ClampMin = "0.0", UIMin = "0.0", ClampMax = "1.0", UIMax = "1.0"))
 	float ResimErrorDirectionalDecayMultiplier;
 };
diff --git a/Engine/Source/Runtime/Experimental/Chaos/Private/Framework/PhysicsProxyBase.cpp b/Engine/Source/Runtime/Experimental/Chaos/Private/Framework/PhysicsProxyBase.cpp
--- a/Engine/Source/Runtime/Experimental/Chaos/Private/Framework/PhysicsProxyBase.cpp
+++ b/Engine/Source/Runtime/Experimental/Chaos/Private/Framework/PhysicsProxyBase.cpp
@@ -5,13 +5,16 @@
 
 namespace RenderInterpolationCVars
 {
-	float RenderInterpErrorCorrectionDuration = 0.5f;
+	float RenderInterpErrorCorrectionDuration = 0.3f;
 	static FAutoConsoleVariableRef CVarRenderInterpErrorCorrectionDuration(TEXT("p.RenderInterp.ErrorCorrectionDuration"), RenderInterpErrorCorrectionDuration, TEXT("How long in seconds to apply error correction over."));
 
 	float RenderInterpMaximumErrorCorrectionBeforeSnapping = 250.0f;
-	static FAutoConsoleVariableRef CVarRenderInterpErrorCorrectionMaximumError(TEXT("p.RenderInterp.MaximumErrorCorrectionBeforeSnapping"), RenderInterpMaximumErrorCorrectionBeforeSnapping, TEXT("Maximum error correction in cm before we stop interpolating and snap to target."));
+	static FAutoConsoleVariableRef CVarRenderInterpErrorCorrectionMaximumError(TEXT("p.RenderInterp.MaximumErrorCorrectionBeforeSnapping"), RenderInterpMaximumErrorCorrectionBeforeSnapping, TEXT("Maximum error correction in cm before we stop interpolating and snap to target. NOTE: MaximumErrorCorrectionDesyncTimeBeforeSnapping can set a larger distance if enabled."));
 
-	float RenderInterpErrorVelocitySmoothingDuration = 0.5f;
+	float RenderInterpMaximumErrorCorrectionDesyncTimeBeforeSnapping = 0.6f;
+	static FAutoConsoleVariableRef CVarRenderInterpErrorCorrectionDesyncTime(TEXT("p.RenderInterp.MaximumErrorCorrectionDesyncTimeBeforeSnapping"), RenderInterpMaximumErrorCorrectionDesyncTimeBeforeSnapping, TEXT("Time multiplied by the particles velocity to get the distance that error correction will be performed within without snapping, disable by setting a negative value. NOTE: MaximumErrorCorrectionBeforeSnapping will act as a lowest distance clamp."));
+
+	float RenderInterpErrorVelocitySmoothingDuration = 0.3f;
 	static FAutoConsoleVariableRef CVarRenderInterpErrorVelocitySmoothingDuration(TEXT("p.RenderInterp.ErrorVelocitySmoothingDuration"), RenderInterpErrorVelocitySmoothingDuration, TEXT("How long in seconds to apply error velocity smoothing correction over, should be smaller than or equal to p.RenderInterp.ErrorCorrectionDuration. RENDERINTERPOLATION_VELOCITYSMOOTHING needs to be defined."));
 
 	float RenderInterpErrorDirectionalDecayMultiplier = 0.0f;
diff --git a/Engine/Source/Runtime/Experimental/Chaos/Private/PhysicsProxy/ClusterUnionPhysicsProxy.cpp b/Engine/Source/Runtime/Experimental/Chaos/Private/PhysicsProxy/ClusterUnionPhysicsProxy.cpp
--- a/Engine/Source/Runtime/Experimental/Chaos/Private/PhysicsProxy/ClusterUnionPhysicsProxy.cpp
+++ b/Engine/Source/Runtime/Experimental/Chaos/Private/PhysicsProxy/ClusterUnionPhysicsProxy.cpp
@@ -494,20 +494,34 @@ namespace Chaos
 			InterpData = GetOrCreateErrorInterpolationData<FProxyInterpolationError>();
 			check(InterpData);
 
-			FReal MaxErrorCorrection = RenderInterpolationCVars::RenderInterpMaximumErrorCorrectionBeforeSnapping;
+			const FReal ErrorDistanceSqr = Error->ErrorX.SizeSquared();
 			FReal ErrorCorrectionDuration = RenderInterpolationCVars::RenderInterpErrorCorrectionDuration;
+			FReal MaxErrorCorrection = RenderInterpolationCVars::RenderInterpMaximumErrorCorrectionBeforeSnapping;
+			FReal MaxErrorDesyncTime = RenderInterpolationCVars::RenderInterpMaximumErrorCorrectionDesyncTimeBeforeSnapping;
 			if (FErrorInterpolationSettings* ErrorSettings = InterpData->GetErrorInterpolationSettings())
 			{
-				MaxErrorCorrection = ErrorSettings->MaximumErrorCorrectionBeforeSnapping;
 				ErrorCorrectionDuration = ErrorSettings->ErrorCorrectionDuration;
+				MaxErrorCorrection = ErrorSettings->MaximumErrorCorrectionBeforeSnapping;
+				MaxErrorDesyncTime = ErrorSettings->MaximumErrorCorrectionDesyncTimeBeforeSnapping;
 			}
+			MaxErrorCorrection = MaxErrorCorrection * MaxErrorCorrection; // Square the value
+
+			auto VelocitySnapLimitHelper = [ErrorDistanceSqr, MaxErrorDesyncTime](const FVec3& Velocity) -> const bool
+				{
+					if (MaxErrorDesyncTime <= 0.0)
+					{
+						return false;
+					}
+					return ErrorDistanceSqr < (Velocity * MaxErrorDesyncTime).SizeSquared();
+				};
 
 			// If error is within interpolation limit, set the number of physics frames to interpolate over, else leave at 0 frames to instantly correct the error.
 			int32 RenderInterpErrorCorrectionDurationTicks = 0;
-			if (Error->ErrorX.SizeSquared() < (MaxErrorCorrection * MaxErrorCorrection))
+			if (ErrorDistanceSqr < MaxErrorCorrection || VelocitySnapLimitHelper(Particle_External->V()) || VelocitySnapLimitHelper(PullData.V) || (NextPullData && VelocitySnapLimitHelper(NextPullData->V)))
 			{
 				RenderInterpErrorCorrectionDurationTicks = FMath::FloorToInt32(ErrorCorrectionDuration / AsyncFixedTimeStep); // Convert duration from seconds to simulation ticks
 			}
+
 			InterpData->AccumlateErrorXR(Error->ErrorX, Error->ErrorR, SolverSyncTimestamp, RenderInterpErrorCorrectionDurationTicks);
 		}
 		const bool bHasInterpolationData = InterpData != nullptr;
diff --git a/Engine/Source/Runtime/Experimental/Chaos/Private/PhysicsProxy/SingleParticlePhysicsProxy.cpp b/Engine/Source/Runtime/Experimental/Chaos/Private/PhysicsProxy/SingleParticlePhysicsProxy.cpp
--- a/Engine/Source/Runtime/Experimental/Chaos/Private/PhysicsProxy/SingleParticlePhysicsProxy.cpp
+++ b/Engine/Source/Runtime/Experimental/Chaos/Private/PhysicsProxy/SingleParticlePhysicsProxy.cpp
@@ -309,7 +309,7 @@ bool FSingleParticlePhysicsProxy::PullFromPhysicsState(const Chaos::FDirtyRigidP
 
 	// Move buffered data into the TPBDRigidParticle without triggering invalidation of the physics state.
 	Chaos::FPBDRigidParticle* Rigid = Particle ? Particle->CastToRigidParticle() : nullptr;
-	if(Rigid)
+	if (Rigid)
 	{
 		bool bNeedUpdateShapeBounds = bPullPhysicsStateForceUpdateBounds;
 		// Note that kinematics should either be updated here (following simulation), or when the
@@ -331,20 +331,35 @@ bool FSingleParticlePhysicsProxy::PullFromPhysicsState(const Chaos::FDirtyRigidP
 			}
 			check(InterpData);
 			
-			FReal MaxErrorCorrection = RenderInterpolationCVars::RenderInterpMaximumErrorCorrectionBeforeSnapping;
+
+			const FReal ErrorDistanceSqr = Error->ErrorX.SizeSquared();
 			FReal ErrorCorrectionDuration = RenderInterpolationCVars::RenderInterpErrorCorrectionDuration;
+			FReal MaxErrorCorrection = RenderInterpolationCVars::RenderInterpMaximumErrorCorrectionBeforeSnapping;
+			FReal MaxErrorDesyncTime = RenderInterpolationCVars::RenderInterpMaximumErrorCorrectionDesyncTimeBeforeSnapping;
 			if (FErrorInterpolationSettings * ErrorSettings = InterpData->GetErrorInterpolationSettings())
 			{
-				MaxErrorCorrection = ErrorSettings->MaximumErrorCorrectionBeforeSnapping;
 				ErrorCorrectionDuration = ErrorSettings->ErrorCorrectionDuration;
+				MaxErrorCorrection = ErrorSettings->MaximumErrorCorrectionBeforeSnapping;
+				MaxErrorDesyncTime = ErrorSettings->MaximumErrorCorrectionDesyncTimeBeforeSnapping;
 			}
+			MaxErrorCorrection = MaxErrorCorrection * MaxErrorCorrection; // Square the value
+
+			auto VelocitySnapLimitHelper = [ErrorDistanceSqr, MaxErrorDesyncTime](const FVec3& Velocity) -> const bool
+			{
+				if (MaxErrorDesyncTime <= 0.0)
+				{
+					return false;
+				}
+				return ErrorDistanceSqr < (Velocity * MaxErrorDesyncTime).SizeSquared();
+			};
 
 			// If error is within interpolation limit, set the number of physics frames to interpolate over, else leave at 0 frames to instantly correct the error.
 			int32 RenderInterpErrorCorrectionDurationTicks = 0;
-			if (Error->ErrorX.SizeSquared() < (MaxErrorCorrection * MaxErrorCorrection))
+			if (ErrorDistanceSqr < MaxErrorCorrection || VelocitySnapLimitHelper(Rigid->V()) || VelocitySnapLimitHelper(PullData.V) || (NextPullData && VelocitySnapLimitHelper(NextPullData->V)))
 			{
 				RenderInterpErrorCorrectionDurationTicks = FMath::FloorToInt32(ErrorCorrectionDuration / AsyncFixedTimeStep); // Convert duration from seconds to simulation ticks
 			}
+
 			InterpData->AccumlateErrorXR(Error->ErrorX, Error->ErrorR, SolverSyncTimestamp, RenderInterpErrorCorrectionDurationTicks);
 			
 			const int32 RenderInterpErrorVelocitySmoothingDurationTicks = FMath::FloorToInt32(RenderInterpolationCVars::RenderInterpErrorVelocitySmoothingDuration / AsyncFixedTimeStep); // Convert duration from seconds to simulation ticks
diff --git a/Engine/Source/Runtime/Experimental/Chaos/Public/Chaos/Framework/PhysicsProxyBase.h b/Engine/Source/Runtime/Experimental/Chaos/Public/Chaos/Framework/PhysicsProxyBase.h
--- a/Engine/Source/Runtime/Experimental/Chaos/Public/Chaos/Framework/PhysicsProxyBase.h
+++ b/Engine/Source/Runtime/Experimental/Chaos/Public/Chaos/Framework/PhysicsProxyBase.h
@@ -31,6 +31,7 @@ namespace RenderInterpolationCVars
 {
 	extern float RenderInterpErrorCorrectionDuration;
 	extern float RenderInterpMaximumErrorCorrectionBeforeSnapping;
+	extern float RenderInterpMaximumErrorCorrectionDesyncTimeBeforeSnapping;
 	extern float RenderInterpErrorVelocitySmoothingDuration;
 	extern float RenderInterpErrorDirectionalDecayMultiplier;
 	extern bool bRenderInterpErrorVelocityCorrection;
@@ -206,19 +207,24 @@ struct PhysicsProxyWrapper
 struct FErrorInterpolationSettings
 {
 	FErrorInterpolationSettings()
-		: ErrorCorrectionDuration(0.5f)
+		: ErrorCorrectionDuration(0.3f)
 		, MaximumErrorCorrectionBeforeSnapping(250.0f)
+		, MaximumErrorCorrectionDesyncTimeBeforeSnapping(0.6f)
 		, ErrorDirectionalDecayMultiplier(0.0f)
 	{}
 	~FErrorInterpolationSettings() { }
 
-	// How long in seconds to apply error correction over
+	/** How long in seconds to apply error correction over */
 	float ErrorCorrectionDuration;
 
-	// Maximum error correction in cm before we stop interpolating and snap to target
+	/** Maximum error correction distance before we stop interpolating and snap to target. */
 	float MaximumErrorCorrectionBeforeSnapping;
 
-	// Decay error offset in the direction that the physics object is moving, value is multiplier of projected offset direction, 0.25 means a 25% decay of the magnitude in the direction of physics travel. Deactivate by setting to 0
+	/** Time multiplied by the particles velocity to get the distance that error correction will be performed within without snapping, disable by setting a negative value
+	* NOTE: At lower velocities MaximumErrorCorrectionBeforeSnapping will act as a lowest distance clamp. */
+	float MaximumErrorCorrectionDesyncTimeBeforeSnapping;
+
+	/** Decay error offset in the direction that the physics object is moving, value is multiplier of projected offset direction, 0.25 means a 25 % decay of the magnitude in the direction of physics travel.Deactivate by setting to 0 */
 	float ErrorDirectionalDecayMultiplier;
 };
 
