Crash fix for when an invalid component is added to a cluster union

diff --git a/Engine/Source/Runtime/Engine/Private/PhysicsEngine/ClusterUnionComponent.cpp b/Engine/Source/Runtime/Engine/Private/PhysicsEngine/ClusterUnionComponent.cpp
--- a/Engine/Source/Runtime/Engine/Private/PhysicsEngine/ClusterUnionComponent.cpp
+++ b/Engine/Source/Runtime/Engine/Private/PhysicsEngine/ClusterUnionComponent.cpp
@@ -203,7 +203,10 @@ void UClusterUnionComponent::AddComponentToCluster(UPrimitiveComponent* InCompon
 		if (Chaos::FPBDRigidParticle* Particle = Interface->GetRigidParticle(Object))
 		{
 			const int32 BoneId = Chaos::FPhysicsObjectInterface::GetId(Object);
-
+			if (BoneId == INDEX_NONE)
+			{
+				continue;
+			}
 			FClusterUnionBoneData BoneData;
 			BoneData.ID = BoneId;
 			BoneData.ParticleID = Particle ? Particle->UniqueIdx() : Chaos::FUniqueIdx();
@@ -226,29 +229,32 @@ void UClusterUnionComponent::AddComponentToCluster(UPrimitiveComponent* InCompon
 		}
 	}
 
-	// The first time we add a component, we need to manage the GT-side object state and set it to kinematic immediately.
-	if (bHasKinematicObject && PendingComponentSync.IsEmpty() && PerComponentData.IsEmpty())
+	if (!ChildToParents.IsEmpty())
 	{
-		Chaos::FPhysicsObjectHandle SelfHandle = PhysicsProxy->GetPhysicsObjectHandle();
-		Interface->ForceKinematic( { &SelfHandle, 1 });
-	}
+		// The first time we add a component, we need to manage the GT-side object state and set it to kinematic immediately.
+		if (bHasKinematicObject && PendingComponentSync.IsEmpty() && PerComponentData.IsEmpty())
+		{
+			Chaos::FPhysicsObjectHandle SelfHandle = PhysicsProxy->GetPhysicsObjectHandle();
+			Interface->ForceKinematic( { &SelfHandle, 1 });
+		}
 
-	// Need to listen to changes in the component's physics state. If it gets destroyed it should be removed from the cluster union as well.
-	// This technically is giving us a false sense of security because when we get this notification, the primitive component will already have no physics proxy on it so we can't actually grab physics objects/particles from it.
-	InComponent->OnComponentPhysicsStateChanged.AddUniqueDynamic(this, &UClusterUnionComponent::HandleComponentPhysicsStateChangePostAddIntoClusterUnion);
+		// Need to listen to changes in the component's physics state. If it gets destroyed it should be removed from the cluster union as well.
+		// This technically is giving us a false sense of security because when we get this notification, the primitive component will already have no physics proxy on it so we can't actually grab physics objects/particles from it.
+		InComponent->OnComponentPhysicsStateChanged.AddUniqueDynamic(this, &UClusterUnionComponent::HandleComponentPhysicsStateChangePostAddIntoClusterUnion);
 
-	if (bRebuildGeometry)
-	{
-		AddGTParticleGeometry(Objects);
-	}
+		if (bRebuildGeometry)
+		{
+			AddGTParticleGeometry(Objects);
+		}
 
-	PendingComponentSync.Add(ComponentKey, PendingData);
+		PendingComponentSync.Add(ComponentKey, PendingData);
 
-	PhysicsProxy->AddPhysicsObjects_External(Objects);
+		PhysicsProxy->AddPhysicsObjects_External(Objects);
 
-	TArray<int32> BoneIDs;
-	GetBoneIDsFromComponentData(BoneIDs, PendingData);
-	ForceSetChildToParent(InComponent, BoneIDs, ChildToParents);
+		TArray<int32> BoneIDs;
+		GetBoneIDsFromComponentData(BoneIDs, PendingData);
+		ForceSetChildToParent(InComponent, BoneIDs, ChildToParents);
+	}
 }
 
 void UClusterUnionComponent::RemoveComponentFromCluster(UPrimitiveComponent* InComponent)
