Added log to print the name of the subsystem not loaded when calling ImportTextItem.

diff --git a/Engine/Source/Runtime/Engine/Private/OnlineReplStructs.cpp b/Engine/Source/Runtime/Engine/Private/OnlineReplStructs.cpp
--- a/Engine/Source/Runtime/Engine/Private/OnlineReplStructs.cpp
+++ b/Engine/Source/Runtime/Engine/Private/OnlineReplStructs.cpp
@@ -555,26 +555,41 @@ bool FUniqueNetIdRepl::ImportTextItem(const TCHAR*& Buffer, int32 PortFlags, UOb
 			{
 				Buffer = NewBuffer1;
 
+				FName OnlineSubsystemName;
+				FString IdContents;
+
 				// Ids can be serialized with OSS prefix e.g.: PREFIX:48ac94b14ca949f3244f91f4a92afb86
 				if (Buffer[0] == TEXT(':'))
 				{
 					Buffer++;
 
-					FString IdStr;
-					if (const TCHAR* NewBuffer2 = FPropertyHelpers::ReadToken(Buffer, IdStr))
+					OnlineSubsystemName = FName(*Token);
+					if (const TCHAR* NewBuffer2 = FPropertyHelpers::ReadToken(Buffer, IdContents))
 					{
 						Buffer = NewBuffer2;
-						UniqueIdFromString(FName(*Token), IdStr);
 					}
 				}
-				else if (Token == InvalidUniqueNetIdStr)
+				else
 				{
-					bShouldWarn = false;
+					IdContents = Token;
 				}
-				else
+
+				if (IdContents != InvalidUniqueNetIdStr)
 				{
-					UniqueIdFromString(NAME_None, Token);
+					// Log the subsystem name used once
+					if (!UOnlineEngineInterface::Get()->IsLoaded(OnlineSubsystemName))
+					{
+						static bool bLogged = false;
+						UE_CLOG(!bLogged, LogNet, Log, TEXT("Attempted to ImportText to FUniqueNetIdRepl while OSS is not loaded. OSS:%s Parent:%s"), *OnlineSubsystemName.ToString(), *GetPathNameSafe(Parent));
+						bLogged = true;
+					}
+
+					UniqueIdFromString(OnlineSubsystemName, IdContents);
 				}
+				else
+				{
+					bShouldWarn = false;
+				}				
 			}
 		}
 	}
