AssetManager and StreamableManager predownload API

diff --git a/Engine/Source/Runtime/Engine/Classes/Engine/AssetManager.h b/Engine/Source/Runtime/Engine/Classes/Engine/AssetManager.h
--- a/Engine/Source/Runtime/Engine/Classes/Engine/AssetManager.h
+++ b/Engine/Source/Runtime/Engine/Classes/Engine/AssetManager.h
@@ -33,7 +33,15 @@ DECLARE_DELEGATE_OneParam(FAssetManagerAcquireResourceDelegate, bool /* bSuccess
 DECLARE_MULTICAST_DELEGATE_OneParam(FOnAddedAssetSearchRoot, const FString&);
 
 /** EXPERIMENTAL -  Additional parameters for downloading packages */
-using FAssetManagerDownloadParams = FStreamableDownloadParams;
+struct FAssetManagerDownloadParams
+{
+	/** Handle to keep packages in the local cache. If no cache pin is provided, a new one will be created. */
+	FStreamableDownloadCachePin CachePin;
+	/** Download priority */
+	EStreamableDownloadPriority Priority = EStreamableDownloadPriority::Normal;
+	/** If set, any soft references will also be installed */
+	bool bInstallSoftReferences = false;
+};
 
 struct FAssetManagerLoadParams
 {
@@ -420,6 +428,19 @@ public:
 	 */
 	ENGINE_API virtual TSharedPtr<FStreamableHandle> PreloadPrimaryAssets(const TArray<FPrimaryAssetId>& AssetsToLoad, const TArray<FName>& LoadBundles, bool bLoadRecursive, FAssetManagerLoadParams&& LoadParams);
 
+	/**
+	 * Downloads data for a set of assets in a specific bundle state, and returns a handle with a FStreamableDownloadCachePin.
+	 * The FStreamableDownloadCachePin must be kept active to hold the data in the download cache.
+	 * Does not load any assets.
+	 *
+	 * @param AssetsToLoad		List of primary assets to load
+	 * @param LoadBundles		List of bundles to load for those assets
+	 * @param bLoadRecursive	If true, this will call RecursivelyExpandBundleData and recurse into sub bundles of other primary assets loaded by a bundle reference
+	 * @param LoadParams		Contains callbacks and download params. Download params are required.
+	 * @return					Streamable Handle
+	 */
+	ENGINE_API virtual TSharedPtr<FStreamableHandle> PredownloadPrimaryAssets(const TArray<FPrimaryAssetId>& AssetsToLoad, const TArray<FName>& LoadBundles, bool bLoadRecursive, FAssetManagerLoadParams&& LoadParams);
+
 	/**
 	 * Load non primary assets with the primary streamable manager.
 	 * This will not auto release the handle, release it if needed.
@@ -465,6 +486,28 @@ public:
 		FAssetManagerLoadParams&& LoadParams,
 		DebugNameType&& DebugNameOrLocation = UE::FSourceLocation::Current());
 
+	/**
+	 * Download non primary assets with the primary streamable manager.
+	 * Does not load assets after downloading.
+	 * This will not auto release the handle, release it if needed.
+	 * 
+	 * @param AssetList			List of non primary assets to load
+	 * @param LoadParams		Contains callbacks and async load priority
+	 * @param DebguName			[optional] Name of this handle, either FString or anything that can construct FString, will be reported in debug tools, will report Source Location if DebugName is not specified explicitly
+	 * @return					Streamable Handle that must be stored to keep the preloaded assets from being freed
+	 */
+	template< typename DebugNameType = UE::FSourceLocation >
+	TSharedPtr<FStreamableHandle> DownloadAssetList(
+		const TArray<FSoftObjectPath>& AssetList,
+		FAssetManagerLoadParams&& LoadParams,
+		DebugNameType&& DebugNameOrLocation = UE::FSourceLocation::Current());
+	/** rvalue reference overload for Asset List */
+	template< typename DebugNameType = UE::FSourceLocation >
+	TSharedPtr<FStreamableHandle> DownloadAssetList(
+		TArray<FSoftObjectPath>&& AssetList,
+		FAssetManagerLoadParams&& LoadParams,
+		DebugNameType&& DebugNameOrLocation = UE::FSourceLocation::Current());
+
 	/** Returns a single AssetBundleInfo, matching Scope and Name */
 	ENGINE_API virtual FAssetBundleEntry GetAssetBundleEntry(const FPrimaryAssetId& BundleScope, FName BundleName) const;
 
@@ -871,6 +914,9 @@ protected:
 	/** Called to load asset list */
 	ENGINE_API virtual TSharedPtr<FStreamableHandle> LoadAssetListInternal(TArray<FSoftObjectPath>&& AssetList, FAssetManagerLoadParams&& LoadParams, FString&& DebugName);
 
+	/** Called to download asset list. Does not load assets. */
+	ENGINE_API TSharedPtr<FStreamableHandle> DownloadAssetListInternal(TArray<FSoftObjectPath>&& AssetList, FAssetManagerLoadParams&& LoadParams, FString&& DebugName);
+
 #if WITH_EDITOR
 	/** Function used during creating Management references to decide when to recurse and set references */
 	ENGINE_API virtual EAssetSetManagerResult::Type ShouldSetManager(const FAssetIdentifier& Manager, const FAssetIdentifier& Source, const FAssetIdentifier& Target,
@@ -1193,3 +1239,32 @@ TSharedPtr<FStreamableHandle> UAssetManager::LoadAssetList(
 			FString{ Forward<DebugNameType>(DebugNameOrLocation) });
 	}
 }
+
+template< typename DebugNameType >
+TSharedPtr<FStreamableHandle> UAssetManager::DownloadAssetList(
+	const TArray<FSoftObjectPath>& AssetList,
+	FAssetManagerLoadParams&& LoadParams,
+	DebugNameType&& DebugNameOrLocation)
+{
+	return DownloadAssetList(CopyTemp(AssetList), MoveTemp(LoadParams), Forward<DebugNameType>(DebugNameOrLocation));
+}
+
+template< typename DebugNameType >
+TSharedPtr<FStreamableHandle> UAssetManager::DownloadAssetList(
+	TArray<FSoftObjectPath>&& AssetList,
+	FAssetManagerLoadParams&& LoadParams,
+	DebugNameType&& DebugNameOrLocation)
+{
+	if constexpr (std::is_same_v<std::decay_t<DebugNameType>, UE::FSourceLocation>)
+	{
+		return DownloadAssetListInternal(
+			MoveTemp(AssetList), MoveTemp(LoadParams),
+			UE::SourceLocation::ToFileAndLineString(DebugNameOrLocation));
+	}
+	else
+	{
+		return DownloadAssetListInternal(
+			MoveTemp(AssetList), MoveTemp(LoadParams),
+			FString{ Forward<DebugNameType>(DebugNameOrLocation) });
+	}
+}
diff --git a/Engine/Source/Runtime/Engine/Classes/Engine/StreamableManager.h b/Engine/Source/Runtime/Engine/Classes/Engine/StreamableManager.h
--- a/Engine/Source/Runtime/Engine/Classes/Engine/StreamableManager.h
+++ b/Engine/Source/Runtime/Engine/Classes/Engine/StreamableManager.h
@@ -704,6 +704,8 @@ struct FStreamableDownloadParams
 	EStreamableDownloadPriority Priority = EStreamableDownloadPriority::Normal;
 	/** If set, any soft references will also be installed */
 	bool bInstallSoftReferences = false;
+	/** If set, request will complete after download, no assets will be loaded */
+	bool bDownloadOnly = false;
 };
 
 /** Parameters for an async load */
diff --git a/Engine/Source/Runtime/Engine/Private/AssetManager.cpp b/Engine/Source/Runtime/Engine/Private/AssetManager.cpp
--- a/Engine/Source/Runtime/Engine/Private/AssetManager.cpp
+++ b/Engine/Source/Runtime/Engine/Private/AssetManager.cpp
@@ -2164,6 +2164,20 @@ TSharedPtr<FStreamableHandle> UAssetManager::PreloadPrimaryAssets(const TArray<F
 	return ReturnHandle;
 }
 
+TSharedPtr<FStreamableHandle> UAssetManager::PredownloadPrimaryAssets(const TArray<FPrimaryAssetId>& AssetsToLoad, const TArray<FName>& LoadBundles, bool bLoadRecursive, FAssetManagerLoadParams&& LoadParams)
+{
+	TStringBuilder<256> DebugValid;
+	TStringBuilder<256> DebugInvalid;
+
+	TArray<FSoftObjectPath> PathsToLoad = UE::AssetManager::Private::GetPreloadPrimaryAssetsLoadList(this, AssetsToLoad, LoadBundles, bLoadRecursive, DebugValid, DebugInvalid);
+
+	TSharedPtr<FStreamableHandle> ReturnHandle = DownloadAssetList(MoveTemp(PathsToLoad), MoveTemp(LoadParams), FString(*DebugValid));
+
+	ensureMsgf(ReturnHandle.IsValid(), TEXT("Requested predownload of Primary Assets failed: %s"), *DebugInvalid);
+
+	return TSharedPtr<FStreamableHandle>();
+}
+
 void UAssetManager::OnAssetStateChangeCompleted(const TSharedPtr<FStreamableHandle>& BoundHandle, const FPrimaryAssetId& PrimaryAssetId, FStreamableDelegateWithHandle&& WrappedDelegate)
 {
 	FPrimaryAssetData* NameData = GetNameData(PrimaryAssetId);
@@ -2359,9 +2373,30 @@ int32 UAssetManager::UnloadPrimaryAssetsWithType(FPrimaryAssetType PrimaryAssetT
 	return UnloadPrimaryAssets(Assets);
 }
 
+namespace UE::AssetManager::Private
+{
+	static FStreamableAsyncLoadParams GetStreamableAsyncLoadParams(FAssetManagerLoadParams&& LoadParams)
+	{
+		FStreamableAsyncLoadParams Params;
+		Params.OnComplete = MoveTemp(LoadParams.OnComplete);
+		Params.OnCancel = MoveTemp(LoadParams.OnCancel);
+		Params.OnUpdate = MoveTemp(LoadParams.OnUpdate);
+		Params.Priority = LoadParams.Priority;
+
+		if (LoadParams.DownloadParams)
+		{
+			Params.DownloadParams.Emplace();
+			Params.DownloadParams->CachePin = MoveTemp(LoadParams.DownloadParams->CachePin);
+			Params.DownloadParams->Priority = LoadParams.DownloadParams->Priority;
+			Params.DownloadParams->bInstallSoftReferences = LoadParams.DownloadParams->bInstallSoftReferences;
+		}
+
+		return Params;
+	}
+}
+
 TSharedPtr<FStreamableHandle> UAssetManager::LoadAssetListInternal(TArray<FSoftObjectPath>&& AssetList, FAssetManagerLoadParams&& LoadParams, FString&& DebugName)
 {
-	TSharedPtr<FStreamableHandle> NewHandle;
 	TArray<int32> MissingChunks, ErrorChunks;
 
 	if (bShouldAcquireMissingChunksOnLoad)
@@ -2376,6 +2411,8 @@ TSharedPtr<FStreamableHandle> UAssetManager::LoadAssetListInternal(TArray<FSoftO
 		}
 	}
 
+	TSharedPtr<FStreamableHandle> NewHandle;
+
 	// SynchronousLoad doesn't make sense if chunks are missing
 	// SynchronousLoad doesn't make sense if we are going to download packages	
 	if (!LoadParams.DownloadParams && bShouldUseSynchronousLoad && MissingChunks.Num() == 0)
@@ -2385,15 +2422,10 @@ TSharedPtr<FStreamableHandle> UAssetManager::LoadAssetListInternal(TArray<FSoftO
 	}
 	else
 	{
-		FStreamableAsyncLoadParams Params;
+		FStreamableAsyncLoadParams Params = UE::AssetManager::Private::GetStreamableAsyncLoadParams(MoveTemp(LoadParams));
 		Params.TargetsToStream = MoveTemp(AssetList);
-		Params.OnComplete = MoveTemp(LoadParams.OnComplete);
-		Params.OnCancel = MoveTemp(LoadParams.OnCancel);
-		Params.OnUpdate = MoveTemp(LoadParams.OnUpdate);
-		Params.Priority = LoadParams.Priority;
 		Params.bManageActiveHandle = false;
 		Params.bStartStalled = MissingChunks.Num() > 0;
-		Params.DownloadParams = MoveTemp(LoadParams.DownloadParams);
 
 		NewHandle = StreamableManager.RequestAsyncLoad(MoveTemp(Params), MoveTemp(DebugName));
 
@@ -2406,6 +2438,22 @@ TSharedPtr<FStreamableHandle> UAssetManager::LoadAssetListInternal(TArray<FSoftO
 	return NewHandle;
 }
 
+TSharedPtr<FStreamableHandle> UAssetManager::DownloadAssetListInternal(TArray<FSoftObjectPath>&& AssetList, FAssetManagerLoadParams&& LoadParams, FString&& DebugName)
+{
+	if (!LoadParams.DownloadParams)
+	{
+		UE_LOG(LogAssetManager, Error, TEXT("Failure downloading %s, no download params!"), *DebugName);
+		return nullptr;
+	}
+
+	FStreamableAsyncLoadParams Params = UE::AssetManager::Private::GetStreamableAsyncLoadParams(MoveTemp(LoadParams));
+	Params.TargetsToStream = MoveTemp(AssetList);
+	Params.bManageActiveHandle = false;
+	Params.DownloadParams->bDownloadOnly = true;
+
+	return StreamableManager.RequestAsyncLoad(MoveTemp(Params), MoveTemp(DebugName));
+}
+
 FAssetBundleEntry UAssetManager::GetAssetBundleEntry(const FPrimaryAssetId& BundleScope, FName BundleName) const
 {
 	if (const TSharedPtr<FAssetBundleData, ESPMode::ThreadSafe>* FoundMap = CachedAssetBundles.Find(BundleScope))
diff --git a/Engine/Source/Runtime/Engine/Private/StreamableManager.cpp b/Engine/Source/Runtime/Engine/Private/StreamableManager.cpp
--- a/Engine/Source/Runtime/Engine/Private/StreamableManager.cpp
+++ b/Engine/Source/Runtime/Engine/Private/StreamableManager.cpp
@@ -403,6 +403,7 @@ namespace UE::StreamableManager::Private
 		UE::IoStore::FOnDemandInstallProgress InstallProgress;
 		EStreamableDownloadPriority Priority = EStreamableDownloadPriority::Normal;
 		bool bInstallSoftReferences = false;
+		bool bDownloadOnly = false;
 
 		bool IsStarted() const
 		{
@@ -517,6 +518,12 @@ namespace UE::StreamableManager::Private
 				return;
 			}
 
+			if (bDownloadOnly)
+			{
+				CompleteLoad(OwningHandle);
+				return;
+			}
+
 			if (OwningHandle->OwningManager)
 			{
 				// StartHandleRequests will progress to asset loading now that the download context is gone
@@ -1876,6 +1883,7 @@ TSharedPtr<FStreamableHandle> FStreamableManager::RequestAsyncLoadInternal(FStre
 		NewRequest->DownloadContext = MakeUnique<UE::StreamableManager::Private::FDownloadContext>();
 		NewRequest->DownloadContext->Priority = Params.DownloadParams->Priority;
 		NewRequest->DownloadContext->bInstallSoftReferences = Params.DownloadParams->bInstallSoftReferences;
+		NewRequest->DownloadContext->bDownloadOnly = Params.DownloadParams->bDownloadOnly;
 	}
 
 	if (Params.bStartStalled)
