Adding some missing includes

diff --git a/Engine/Source/Runtime/Engine/Private/GameplayStatics.cpp b/Engine/Source/Runtime/Engine/Private/GameplayStatics.cpp
--- a/Engine/Source/Runtime/Engine/Private/GameplayStatics.cpp
+++ b/Engine/Source/Runtime/Engine/Private/GameplayStatics.cpp
@@ -6,6 +6,8 @@
 #include "Engine/GameInstance.h"
 #include "Engine/GameViewportClient.h"
 #include "Engine/OverlapResult.h"
+#include "Serialization/MemoryReader.h"
+#include "Serialization/MemoryWriter.h"
 #include "Serialization/ObjectAndNameAsStringProxyArchive.h"
 #include "EngineLogs.h"
 #include "Misc/PackageName.h"
