Fixed incorrect tangents by skipping sections that aren't setup to process.

diff --git a/Engine/Source/Runtime/Engine/Private/GPUSkinCache.cpp b/Engine/Source/Runtime/Engine/Private/GPUSkinCache.cpp
--- a/Engine/Source/Runtime/Engine/Private/GPUSkinCache.cpp
+++ b/Engine/Source/Runtime/Engine/Private/GPUSkinCache.cpp
@@ -368,7 +368,7 @@ public:
 		FGPUSkinCache::FSkinCacheRWBuffer* TangentBuffer = nullptr;
 		FGPUSkinCache::FSkinCacheRWBuffer* IntermediateTangentBuffer = nullptr;
 		FGPUSkinCache::FSkinCacheRWBuffer* IntermediateAccumulatedTangentBuffer = nullptr;
-		uint32 IntermediateAccumulatedTangentBufferOffset = 0;
+		uint32 IntermediateAccumulatedTangentBufferOffset = INDEX_NONE;
 		FGPUSkinCache::FSkinCacheRWBuffer* PositionBuffer = nullptr;
 		FGPUSkinCache::FSkinCacheRWBuffer* PreviousPositionBuffer = nullptr;
 
@@ -1539,7 +1539,7 @@ void FGPUSkinCache::DispatchPassSetup(FTaskData& TaskData)
 
 		if (EnumHasAnyFlags(DispatchData.DispatchFlags, EGPUSkinCacheDispatchFlags::Position))
 		{
-			if (DispatchData.IntermediateTangentBuffer)
+			if (DispatchData.IntermediateTangentBuffer && DispatchData.IntermediateAccumulatedTangentBufferOffset != INDEX_NONE)
 			{
 				DispatchData.DispatchFlags |= EGPUSkinCacheDispatchFlags::RecomputeTangents;
 
