First Person: Added support for first person rendering in scene capture 2D components.

#12650

diff --git a/Engine/Plugins/Experimental/CineCameraSceneCapture/Source/CineCameraSceneCapture/Private/CineCameraSceneCaptureComponent.cpp b/Engine/Plugins/Experimental/CineCameraSceneCapture/Source/CineCameraSceneCapture/Private/CineCameraSceneCaptureComponent.cpp
--- a/Engine/Plugins/Experimental/CineCameraSceneCapture/Source/CineCameraSceneCapture/Private/CineCameraSceneCaptureComponent.cpp
+++ b/Engine/Plugins/Experimental/CineCameraSceneCapture/Source/CineCameraSceneCapture/Private/CineCameraSceneCaptureComponent.cpp
@@ -180,8 +180,12 @@ void UCineCaptureComponent2D::UpdateSceneCaptureContents(FSceneInterface* Scene,
 	
 	PostProcessSettings = CineCameraComponent->PostProcessSettings;
 	FOVAngle = CineCameraComponent->FieldOfView;
+	FirstPersonFieldOfView = CineCameraComponent->FirstPersonFieldOfView;
+	FirstPersonScale = CineCameraComponent->FirstPersonScale;
 	bOverride_CustomNearClippingPlane = CineCameraComponent->bOverride_CustomNearClippingPlane;
 	CustomNearClippingPlane = CineCameraComponent->CustomNearClippingPlane;
+	bEnableFirstPersonFieldOfView = CineCameraComponent->bEnableFirstPersonFieldOfView;
+	bEnableFirstPersonScale = CineCameraComponent->bEnableFirstPersonScale;
 
 	Scene->UpdateSceneCaptureContents(this, SceneRenderBuilder);
 }
diff --git a/Engine/Source/Runtime/Engine/Classes/Components/SceneCaptureComponent2D.h b/Engine/Source/Runtime/Engine/Classes/Components/SceneCaptureComponent2D.h
--- a/Engine/Source/Runtime/Engine/Classes/Components/SceneCaptureComponent2D.h
+++ b/Engine/Source/Runtime/Engine/Classes/Components/SceneCaptureComponent2D.h
@@ -43,6 +43,14 @@ public:
 	UPROPERTY(interp, EditAnywhere, BlueprintReadWrite, Category=Projection, meta=(DisplayName = "Field of View", UIMin = "5.0", UIMax = "170", ClampMin = "0.001", ClampMax = "360.0", editcondition = "ProjectionType==0"))
 	float FOVAngle;
 
+	/** The horizontal field of view (in degrees) used for primitives tagged as FirstPerson. */
+	UPROPERTY(Interp, EditAnywhere, BlueprintReadWrite, Category=Projection, meta = (UIMin = "5.0", UIMax = "170", ClampMin = "0.001", ClampMax = "360.0", Units = deg, EditCondition = "bEnableFirstPersonFieldOfView"))
+	float FirstPersonFieldOfView;
+
+	/** The scale to apply to primitives tagged as FirstPerson. This is used to scale down primitives towards the camera such that they are small enough not to intersect with the scene. */
+	UPROPERTY(Interp, EditAnywhere, BlueprintReadWrite, Category=Projection, meta=(UIMin = "0.001", UIMax = "1.0", ClampMin = "0.001", ClampMax = "1.0", EditCondition = "bEnableFirstPersonScale"))
+	float FirstPersonScale;
+
 	/** The desired width (in world units) of the orthographic view (ignored in Perspective mode) */
 	UPROPERTY(EditAnywhere, BlueprintReadWrite, Category=Projection, meta = (editcondition = "ProjectionType==1"))
 	float OrthoWidth;
@@ -151,6 +159,14 @@ public:
 	UPROPERTY(Transient, BlueprintReadWrite, Category = SceneCapture)
 	uint32 bCameraCutThisFrame : 1;
 
+	/** True if the first person field of view should be used for primitives tagged as FirstPerson. */
+	UPROPERTY(EditAnywhere, BlueprintReadWrite, Category=Projection)
+	uint32 bEnableFirstPersonFieldOfView : 1;
+
+	/** True if the first person scale should be used for primitives tagged as FirstPerson. */
+	UPROPERTY(EditAnywhere, BlueprintReadWrite, Category=Projection)
+	uint32 bEnableFirstPersonScale : 1;
+
 	/** Whether to only render exponential height fog on opaque pixels which were rendered by the scene capture. */
 	UPROPERTY(EditAnywhere, BlueprintReadWrite, Category = SceneCapture, meta = (DisplayName = "Fog only on rendered pixels"))
 	uint32 bConsiderUnrenderedOpaquePixelAsFullyTranslucent : 1;
diff --git a/Engine/Source/Runtime/Engine/Private/Components/SceneCaptureComponent.cpp b/Engine/Source/Runtime/Engine/Private/Components/SceneCaptureComponent.cpp
--- a/Engine/Source/Runtime/Engine/Private/Components/SceneCaptureComponent.cpp
+++ b/Engine/Source/Runtime/Engine/Private/Components/SceneCaptureComponent.cpp
@@ -629,6 +629,9 @@ USceneCaptureComponent2D::USceneCaptureComponent2D(const FObjectInitializer& Obj
 {
 	FOVAngle = 90.0f;
 
+	FirstPersonFieldOfView = 90.0f;
+	FirstPersonScale = 1.0f;
+
 	OrthoWidth = DEFAULT_ORTHOWIDTH;
 	bAutoCalculateOrthoPlanes = true;
 	AutoPlaneShift = 0.0f;
@@ -650,6 +653,8 @@ USceneCaptureComponent2D::USceneCaptureComponent2D(const FObjectInitializer& Obj
 	CustomProjectionMatrix.SetIdentity();
 	ClipPlaneNormal = FVector(0, 0, 1);
 	bCameraCutThisFrame = false;
+	bEnableFirstPersonFieldOfView = false;
+	bEnableFirstPersonScale = false;
 	bConsiderUnrenderedOpaquePixelAsFullyTranslucent = false;
 
 	bMainViewFamily = false;
@@ -766,6 +771,9 @@ void USceneCaptureComponent2D::SetCameraView(const FMinimalViewInfo& DesiredView
 	FOVAngle = DesiredView.FOV;
 	ProjectionType = DesiredView.ProjectionMode;
 	OrthoWidth = DesiredView.OrthoWidth;
+
+	FirstPersonFieldOfView = DesiredView.FirstPersonFOV;
+	FirstPersonScale = DesiredView.FirstPersonScale;
 }
 
 void USceneCaptureComponent2D::GetCameraView(float DeltaTime, FMinimalViewInfo& OutMinimalViewInfo)
@@ -782,6 +790,9 @@ void USceneCaptureComponent2D::GetCameraView(float DeltaTime, FMinimalViewInfo&
 	OutMinimalViewInfo.AutoPlaneShift = AutoPlaneShift;
 	OutMinimalViewInfo.bUpdateOrthoPlanes = bUpdateOrthoPlanes;
 	OutMinimalViewInfo.bUseCameraHeightAsViewTarget = bUseCameraHeightAsViewTarget;
+	OutMinimalViewInfo.FirstPersonFOV = bEnableFirstPersonFieldOfView ? FirstPersonFieldOfView : OutMinimalViewInfo.FOV;
+	OutMinimalViewInfo.FirstPersonScale = bEnableFirstPersonScale ? FirstPersonScale : 1.0f;
+	OutMinimalViewInfo.bUseFirstPersonParameters = bEnableFirstPersonFieldOfView || bEnableFirstPersonScale;
 
 	OutMinimalViewInfo.ApplyOverscan(Overscan);
 	
diff --git a/Engine/Source/Runtime/Renderer/Private/SceneCaptureRendering.cpp b/Engine/Source/Runtime/Renderer/Private/SceneCaptureRendering.cpp
--- a/Engine/Source/Runtime/Renderer/Private/SceneCaptureRendering.cpp
+++ b/Engine/Source/Runtime/Renderer/Private/SceneCaptureRendering.cpp
@@ -729,6 +729,14 @@ TArray<FSceneView*> SetupViewFamilyForSceneCapture(
 			ViewInitOptions.LODDistanceFactor = FMath::Clamp(SceneCaptureComponent->LODDistanceFactor, .01f, 100.0f);
 			ViewInitOptions.bIsSceneCaptureCube = SceneCaptureComponent->IsCube();
 			ViewInitOptions.bSceneCaptureUsesRayTracing = GRayTracingSceneCaptures == -1 ? SceneCaptureComponent->bUseRayTracingIfEnabled : GRayTracingSceneCaptures > 0;
+
+			USceneCaptureComponent2D* SceneCaptureComponent2D = Cast<USceneCaptureComponent2D>(SceneCaptureComponent);
+			if (IsValid(SceneCaptureComponent2D))
+			{
+				FMinimalViewInfo ViewInfo;
+				SceneCaptureComponent2D->GetCameraView(0.0f, ViewInfo);
+				ViewInitOptions.FirstPersonParams = FFirstPersonParameters(ViewInfo.CalculateFirstPersonFOVCorrectionFactor(), ViewInfo.FirstPersonScale, ViewInfo.bUseFirstPersonParameters);
+			}
 		}
 
 		FSceneView* View = new FSceneView(ViewInitOptions);
