Dataflow - Made Skeletal mesh render buffers export serialization and init functions for use of these structures in Dataflow graph nodes.

#rnx

diff --git a/Engine/Source/Runtime/Engine/Public/Rendering/MultiSizeIndexContainer.h b/Engine/Source/Runtime/Engine/Public/Rendering/MultiSizeIndexContainer.h
--- a/Engine/Source/Runtime/Engine/Public/Rendering/MultiSizeIndexContainer.h
+++ b/Engine/Source/Runtime/Engine/Public/Rendering/MultiSizeIndexContainer.h
@@ -46,7 +46,7 @@ public:
 	* @param bNeedsCPUAccess	If true, the loaded data will remain in CPU memory
 	*							even after the RHI index buffer has been initialized.
 	*/
-	void Serialize(FArchive& Ar, bool bNeedsCPUAccess);
+	ENGINE_API void Serialize(FArchive& Ar, bool bNeedsCPUAccess);
 
 	void SerializeMetaData(FArchive& Ar, bool bNeedsCPUAccess);
 
diff --git a/Engine/Source/Runtime/Engine/Public/Rendering/SkeletalMeshHalfEdgeBuffer.h b/Engine/Source/Runtime/Engine/Public/Rendering/SkeletalMeshHalfEdgeBuffer.h
--- a/Engine/Source/Runtime/Engine/Public/Rendering/SkeletalMeshHalfEdgeBuffer.h
+++ b/Engine/Source/Runtime/Engine/Public/Rendering/SkeletalMeshHalfEdgeBuffer.h
@@ -20,7 +20,7 @@ namespace SkeletalMeshHalfEdgeUtility
 class FSkeletalMeshHalfEdgeBuffer : public FRenderResource
 {
 public:
-	void Init(const FSkeletalMeshLODRenderData& InLodRenderData);
+	void ENGINE_API Init(const FSkeletalMeshLODRenderData& InLodRenderData);
 
 	struct FRHIInfo
 	{
@@ -37,11 +37,11 @@ public:
 
 	void ENGINE_API ReleaseRHI() override;
 
-	bool IsCPUDataValid() const;
+	bool ENGINE_API IsCPUDataValid() const;
 
 	bool ENGINE_API IsReadyForRendering() const;
 	
-	void CleanUp();
+	void ENGINE_API CleanUp();
 
 	int32 GetResourceSize() const;
 	
diff --git a/Engine/Source/Runtime/Engine/Public/Rendering/SkeletalMeshLODRenderData.h b/Engine/Source/Runtime/Engine/Public/Rendering/SkeletalMeshLODRenderData.h
--- a/Engine/Source/Runtime/Engine/Public/Rendering/SkeletalMeshLODRenderData.h
+++ b/Engine/Source/Runtime/Engine/Public/Rendering/SkeletalMeshLODRenderData.h
@@ -116,7 +116,7 @@ struct FSkelMeshRenderSection
 		return NumVertices;
 	}
 
-	friend FArchive& operator<<(FArchive& Ar, FSkelMeshRenderSection& S);
+	ENGINE_API friend FArchive& operator<<(FArchive& Ar, FSkelMeshRenderSection& S);
 };
 
 class FSkeletalMeshLODRenderData : public FRefCountBase
diff --git a/Engine/Source/Runtime/Engine/Public/Rendering/SkeletalMeshVertexClothBuffer.h b/Engine/Source/Runtime/Engine/Public/Rendering/SkeletalMeshVertexClothBuffer.h
--- a/Engine/Source/Runtime/Engine/Public/Rendering/SkeletalMeshVertexClothBuffer.h
+++ b/Engine/Source/Runtime/Engine/Public/Rendering/SkeletalMeshVertexClothBuffer.h
@@ -54,7 +54,7 @@ public:
 	* @param Ar - archive to serialize to
 	* @param B - data to serialize
 	*/
-	friend FArchive& operator<<(FArchive& Ar, FSkeletalMeshVertexClothBuffer& VertexBuffer);
+	ENGINE_API friend FArchive& operator<<(FArchive& Ar, FSkeletalMeshVertexClothBuffer& VertexBuffer);
 
 	void SerializeMetaData(FArchive& Ar);
 
diff --git a/Engine/Source/Runtime/Engine/Public/Rendering/SkinWeightVertexBuffer.h b/Engine/Source/Runtime/Engine/Public/Rendering/SkinWeightVertexBuffer.h
--- a/Engine/Source/Runtime/Engine/Public/Rendering/SkinWeightVertexBuffer.h
+++ b/Engine/Source/Runtime/Engine/Public/Rendering/SkinWeightVertexBuffer.h
@@ -390,7 +390,7 @@ public:
 
 	ENGINE_API void CopySkinWeightRawDataFromBuffer(const uint8* InSkinWeightData, uint32 InNumVertices);
 
-	friend FArchive& operator<<(FArchive& Ar, FSkinWeightVertexBuffer& VertexBuffer);
+	ENGINE_API friend FArchive& operator<<(FArchive& Ar, FSkinWeightVertexBuffer& VertexBuffer);
 
 	void SerializeMetaData(FArchive& Ar);
 	void CopyMetaData(const FSkinWeightVertexBuffer& Other);
