[materials] assume shadermap completion in test/shipping builds if platform requires cooked data. this significantly reduces cost of UMaterial::PostLoad and UMaterialInstance::PostLoad for builds where ODSC is never used and as such incomplete materials are rare (i.e. due to some kind of runtime/cooktime ShouldCompilePermutation mismatch; such errors will still be reported in debug/development builds).

diff --git a/Engine/Source/Runtime/Engine/Private/Materials/MaterialShared.cpp b/Engine/Source/Runtime/Engine/Private/Materials/MaterialShared.cpp
--- a/Engine/Source/Runtime/Engine/Private/Materials/MaterialShared.cpp
+++ b/Engine/Source/Runtime/Engine/Private/Materials/MaterialShared.cpp
@@ -1191,7 +1191,12 @@ bool FMaterial::MaterialGetSubstrateUsesComplexSpecialRenderPath_RenderThread()
 void FMaterial::SetGameThreadShaderMap(FMaterialShaderMap* InMaterialShaderMap)
 {
 	checkSlow(IsInGameThread() || IsInAsyncLoadingThread());
-	const bool bIsComplete = InMaterialShaderMap ? InMaterialShaderMap->IsComplete(this, true) : false;
+
+	bool bAssumeShaderMapIsComplete = false;
+#if UE_BUILD_SHIPPING || UE_BUILD_TEST
+	bAssumeShaderMapIsComplete = FPlatformProperties::RequiresCookedData();
+#endif
+	const bool bIsComplete = bAssumeShaderMapIsComplete || (InMaterialShaderMap ? InMaterialShaderMap->IsComplete(this, true) : false);
 	GameThreadShaderMap = InMaterialShaderMap;
 	if (LIKELY(GameThreadShaderMap))
 	{
@@ -1214,7 +1219,12 @@ void FMaterial::UpdateInlineShaderMapIsComplete()
 	check(bContainsInlineShaders);
 	// We expect inline shader maps to be complete, so we want to log missing shaders here
 	const bool bSilent = false;
-	const bool bIsComplete = GameThreadShaderMap->IsComplete(this, bSilent);
+
+	bool bAssumeShaderMapIsComplete = false;
+#if UE_BUILD_SHIPPING || UE_BUILD_TEST
+	bAssumeShaderMapIsComplete = FPlatformProperties::RequiresCookedData();
+#endif
+	const bool bIsComplete = bAssumeShaderMapIsComplete || GameThreadShaderMap->IsComplete(this, bSilent);
 
 	bGameThreadShaderMapIsComplete.store(bIsComplete, std::memory_order_relaxed);
 	TRefCountPtr<FMaterial> Material = this;
@@ -2857,7 +2867,11 @@ bool FMaterial::CacheShaders(const FMaterialShaderMapId& ShaderMapId, EShaderPla
 		// We can bypass this check in the case where we found a "required complete" shader map whose compilation is in progress;
 		// we will check if it's complete below and queue any jobs necessary if not (we only need to log warnings if we found a map without
 		// compilation in progress that is incomplete).
-		bShaderMapValid = GameThreadShaderMap->IsComplete(this, false);
+		bool bAssumeShaderMapIsComplete = false;
+#if UE_BUILD_SHIPPING || UE_BUILD_TEST
+		bAssumeShaderMapIsComplete = FPlatformProperties::RequiresCookedData();
+#endif
+		bShaderMapValid = bAssumeShaderMapIsComplete || GameThreadShaderMap->IsComplete(this, false);
 	}
 
 	if (!bShaderMapValid)
