Disallow Skeletal stream-in/out after render data invalidation
- Avoids initializing rhi resources after a skeletal object has begun destruction

diff --git a/Engine/Source/Runtime/Engine/Private/SkeletalMesh.cpp b/Engine/Source/Runtime/Engine/Private/SkeletalMesh.cpp
--- a/Engine/Source/Runtime/Engine/Private/SkeletalMesh.cpp
+++ b/Engine/Source/Runtime/Engine/Private/SkeletalMesh.cpp
@@ -1097,6 +1097,13 @@ bool USkeletalMesh::HasPendingRenderResourceInitialization() const
 bool USkeletalMesh::StreamOut(int32 NewMipCount)
 {
 	check(IsInGameThread());
+	
+	FSkeletalMeshRenderData* RenderData = GetResourceForRendering();
+	if (!RenderData || !RenderData->IsInitialized())
+	{
+		return false;
+	}
+	
 	if (!HasPendingInitOrStreaming() && CachedSRRState.StreamOut(NewMipCount))
 	{
 		PendingUpdate = new FSkeletalMeshStreamOut(this);
@@ -1108,6 +1115,13 @@ bool USkeletalMesh::StreamOut(int32 NewMipCount)
 bool USkeletalMesh::StreamIn(int32 NewMipCount, bool bHighPrio)
 {
 	check(IsInGameThread());
+	
+	FSkeletalMeshRenderData* RenderData = GetResourceForRendering();
+	if (!RenderData || !RenderData->IsInitialized())
+	{
+		return false;
+	}
+	
 	if (!HasPendingInitOrStreaming() && CachedSRRState.StreamIn(NewMipCount))
 	{
 		FRenderAssetUpdate::EThreadType CreateResourcesThread = GRHISupportsAsyncTextureCreation
