AnimNext: Fixes for socket attachment, animation curves and Niagara FX

Added callback to USceneComponent and flag in USkinnedMeshComponent to indicate whether socket update is required.
In AnimNext, queue a game thread task to update sockets when the flag is set.
Note: Not thread-safe yet - that is dependent on subsequent work for animation -> mesh abstraction
Post animation update in anim BP, skip socket attachment update if the flag is not set
Allow Niagara FX that use anim curves to work with AnimNext
Added query API for curve values to skeletal mesh component
Deprecated access to skel mesh comp curves

diff --git a/Engine/Plugins/Experimental/AnimNext/Source/AnimNext/Private/Component/SkinnedMeshComponentExtensions.cpp b/Engine/Plugins/Experimental/AnimNext/Source/AnimNext/Private/Component/SkinnedMeshComponentExtensions.cpp
--- a/Engine/Plugins/Experimental/AnimNext/Source/AnimNext/Private/Component/SkinnedMeshComponentExtensions.cpp
+++ b/Engine/Plugins/Experimental/AnimNext/Source/AnimNext/Private/Component/SkinnedMeshComponentExtensions.cpp
@@ -4,6 +4,7 @@
 #include "GenerationTools.h"
 #include "Components/SkinnedMeshComponent.h"
 #include "Animation/AnimTrace.h"
+#include "Module/AnimNextModuleInstance.h"
 
 namespace UE::Anim
 {
@@ -32,6 +33,20 @@ void FSkinnedMeshComponentExtensions::CompleteAndDispatch(
 		// Send updated transforms to the renderer
 		InComponent->SendRenderDynamicData_Concurrent();
 
+		if (InComponent->bHasSocketAttachments)
+		{
+			FAnimNextModuleInstance::RunTaskOnGameThread([WeakComponent = TWeakObjectPtr<USkinnedMeshComponent>(InComponent)]()
+			{
+				USkinnedMeshComponent* Component = WeakComponent.Get();
+				if (Component == nullptr)
+				{
+					return;
+				}
+
+				Component->UpdateChildTransforms(EUpdateTransformFlags::OnlyUpdateIfUsingSocket);
+			});
+		}
+		
 		if (const USkeletalMeshComponent* SkeletalMeshComponent = Cast<USkeletalMeshComponent>(InComponent))
 		{
 			TRACE_SKELETAL_MESH_COMPONENT(SkeletalMeshComponent);
diff --git a/Engine/Plugins/Experimental/AnimNext/Source/AnimNext/Private/Graph/RigUnit_AnimNextWriteSkeletalMeshComponentPose.cpp b/Engine/Plugins/Experimental/AnimNext/Source/AnimNext/Private/Graph/RigUnit_AnimNextWriteSkeletalMeshComponentPose.cpp
--- a/Engine/Plugins/Experimental/AnimNext/Source/AnimNext/Private/Graph/RigUnit_AnimNextWriteSkeletalMeshComponentPose.cpp
+++ b/Engine/Plugins/Experimental/AnimNext/Source/AnimNext/Private/Graph/RigUnit_AnimNextWriteSkeletalMeshComponentPose.cpp
@@ -93,14 +93,18 @@ FRigUnit_AnimNextWriteSkeletalMeshComponentPose_Execute()
 	}
 
 	// Clear our curves and attributes or we'll have leftovers from our last write.
+	PRAGMA_DISABLE_DEPRECATION_WARNINGS
 	OutputComponent->AnimCurves.Empty();
+	PRAGMA_ENABLE_DEPRECATION_WARNINGS
 	OutputComponent->GetEditableCustomAttributes().Empty();
 
 
 	// Map LOD pose into local-space scratch buffer
 	FGenerationTools::RemapPose(Pose.LODPose, LocalSpaceTransforms);
 
+	PRAGMA_DISABLE_DEPRECATION_WARNINGS
 	OutputComponent->AnimCurves.CopyFrom(Pose.Curves);
+	PRAGMA_ENABLE_DEPRECATION_WARNINGS
 
 	// Attributes require remapping since the indices are LOD indices and we want mesh indices
 	FGenerationTools::RemapAttributes(Pose.LODPose, Pose.Attributes, OutputComponent->GetEditableCustomAttributes());
diff --git a/Engine/Plugins/FX/Niagara/Source/NiagaraAnimNotifies/Private/AnimNotifyState_TimedNiagaraEffect.cpp b/Engine/Plugins/FX/Niagara/Source/NiagaraAnimNotifies/Private/AnimNotifyState_TimedNiagaraEffect.cpp
--- a/Engine/Plugins/FX/Niagara/Source/NiagaraAnimNotifies/Private/AnimNotifyState_TimedNiagaraEffect.cpp
+++ b/Engine/Plugins/FX/Niagara/Source/NiagaraAnimNotifies/Private/AnimNotifyState_TimedNiagaraEffect.cpp
@@ -236,6 +236,23 @@ void UAnimNotifyState_TimedNiagaraEffectAdvanced::NotifyTick(USkeletalMeshCompon
 					}
 				}
 			}
+			else
+			{
+				// No anim instance, defer to the mesh component's curves
+				for (int32 Index = 0; Index != AnimCurves.Num(); ++Index)
+				{
+					FName CurveName = AnimCurves[Index].AnimCurveName;
+					FName NiagaraUserVariableName = AnimCurves[Index].UserVariableName;
+					if ((!CurveName.IsNone()) && (!NiagaraUserVariableName.IsNone()))
+					{
+						float CurveValue = 0.0f;
+						if (MeshComp->GetCurveValue(CurveName, 0.0f, CurveValue))
+						{
+							FXComponent->SetFloatParameter(NiagaraUserVariableName, CurveValue);
+						}
+					}
+				}
+			}
 		}
 	}
 }
diff --git a/Engine/Plugins/Mutable/Source/CustomizableObjectEditor/Private/MuCOE/Nodes/CustomizableObjectNodeAnimationPose.cpp b/Engine/Plugins/Mutable/Source/CustomizableObjectEditor/Private/MuCOE/Nodes/CustomizableObjectNodeAnimationPose.cpp
--- a/Engine/Plugins/Mutable/Source/CustomizableObjectEditor/Private/MuCOE/Nodes/CustomizableObjectNodeAnimationPose.cpp
+++ b/Engine/Plugins/Mutable/Source/CustomizableObjectEditor/Private/MuCOE/Nodes/CustomizableObjectNodeAnimationPose.cpp
@@ -140,7 +140,7 @@ void UCustomizableObjectNodeAnimationPose::StaticRetrievePoseInformation(UPoseAs
 		OutPose = OutAnimData.GetPose();
 		OutCurve = OutAnimData.GetCurve();
 
-		OutCurve.CopyFrom(SkeletalMeshComponent->AnimCurves);
+		OutCurve.CopyFrom(SkeletalMeshComponent->GetAnimCurves());
 
 		// Assuming one single pose, with a weight set to 1.0
 		FAnimExtractContext ExtractionContext;
diff --git a/Engine/Source/Runtime/Engine/Classes/Components/SkeletalMeshComponent.h b/Engine/Source/Runtime/Engine/Classes/Components/SkeletalMeshComponent.h
--- a/Engine/Source/Runtime/Engine/Classes/Components/SkeletalMeshComponent.h
+++ b/Engine/Source/Runtime/Engine/Classes/Components/SkeletalMeshComponent.h
@@ -445,7 +445,7 @@ public:
 	UPROPERTY()
 	FVector LineCheckBoundsScale;
 
-	/** Temporary storage for curves */
+	UE_DEPRECATED(5.6, "Direct access to this property is deprecated, please use GetAnimCurves or GetCurveValue instead. We will move to private in the future.")
 	FBlendedHeapCurve AnimCurves;
 
 	/** Access cached component space transforms */
@@ -575,6 +575,19 @@ public:
 	UFUNCTION(BlueprintCallable, Category = CustomAttributes)
 	ENGINE_API bool GetStringAttribute(const FName& BoneName, const FName& AttributeName, FString DefaultValue, FString& OutValue, ECustomBoneAttributeLookup LookupType = ECustomBoneAttributeLookup::BoneOnly);
 
+	/** 
+	 * Get curve value
+	 * @param CurveName Name of the curve to retrieve
+	 * @param DefaultValue In case the curve could not be found
+	 * @param Value Retrieved curve value if found, otherwise is set to DefaultValue
+	 * @return Whether the curve was successfully retrieved
+	*/
+	UFUNCTION(BlueprintCallable, Category = Curve) 
+	ENGINE_API bool GetCurveValue(FName CurveName, float DefaultValue, float& Value) const;
+
+	/** Get the current set of curve values */
+	ENGINE_API const FBlendedHeapCurve& GetAnimCurves() const;
+
 protected:
 	/** Templated version to try and retrieve a typed bone attribute's value */
 	template<typename DataType, typename CustomAttributeType>
diff --git a/Engine/Source/Runtime/Engine/Classes/Components/SkinnedMeshComponent.h b/Engine/Source/Runtime/Engine/Classes/Components/SkinnedMeshComponent.h
--- a/Engine/Source/Runtime/Engine/Classes/Components/SkinnedMeshComponent.h
+++ b/Engine/Source/Runtime/Engine/Classes/Components/SkinnedMeshComponent.h
@@ -917,6 +917,9 @@ protected:
 	UPROPERTY(Transient)
 	uint8 bFollowerShouldTickPose : 1;
 
+	/** Whether this component currently has any attachments that are driven by sockets */
+	uint8 bHasSocketAttachments : 1;
+
 #if UE_ENABLE_DEBUG_DRAWING
 private:
 	/** Whether to draw this mesh's debug skeleton (regardless of showflags) */
@@ -1249,6 +1252,8 @@ public:
 	ENGINE_API virtual bool HasAnySockets() const override;
 	ENGINE_API virtual void QuerySupportedSockets(TArray<FComponentSocketDescription>& OutSockets) const override;
 	ENGINE_API virtual bool UpdateOverlapsImpl(const TOverlapArrayView* PendingOverlaps=nullptr, bool bDoNotifies=true, const TOverlapArrayView* OverlapsAtEndLocation=nullptr) override;
+	ENGINE_API virtual void OnChildAttached(USceneComponent* ChildComponent) override;
+	ENGINE_API virtual void OnChildDetached(USceneComponent* ChildComponent) override;
 	//~ End USceneComponent Interface
 
 	//~ Begin UPrimitiveComponent Interface
@@ -1758,6 +1763,8 @@ protected:
 	/** Get the default overlay material max draw distance */
 	ENGINE_API virtual float GetDefaultOverlayMaterialMaxDrawDistance() const override;
 
+	void UpdateCachedSocketAttachmentFlag();
+	
 public:
 	/** 
 	 * Refresh Follower Components if exists
@@ -1795,6 +1802,9 @@ public:
 	ENGINE_API void RemoveSocketOverrides(FName SourceSocketName);
 	ENGINE_API void RemoveAllSocketOverrides();
 
+	/** @return true if this component has any attach children that are attached via sockets */
+	ENGINE_API bool HasAnyAttachedSockets() const;
+
 	/** 
 	 * Get Bone Matrix from index
 	 *
diff --git a/Engine/Source/Runtime/Engine/Private/Components/SkeletalMeshComponent.cpp b/Engine/Source/Runtime/Engine/Private/Components/SkeletalMeshComponent.cpp
--- a/Engine/Source/Runtime/Engine/Private/Components/SkeletalMeshComponent.cpp
+++ b/Engine/Source/Runtime/Engine/Private/Components/SkeletalMeshComponent.cpp
@@ -2666,11 +2666,9 @@ void USkeletalMeshComponent::RefreshBoneTransforms(FActorComponentTickFunction*
 									 ((BoneSpaceTransforms.Num() != GetSkeletalMeshAsset()->GetRefSkeleton().GetNum())
 									 || (BoneSpaceTransforms.Num() != CachedBoneSpaceTransforms.Num())
 									 || (GetNumComponentSpaceTransforms() != CachedComponentSpaceTransforms.Num()));
-
-	PRAGMA_ENABLE_DEPRECATION_WARNINGS
-
+	
 	const bool bInvalidCachedCurve = bDoEvaluationRateOptimization && CachedCurve.Num() != AnimCurves.Num();
-
+	PRAGMA_ENABLE_DEPRECATION_WARNINGS
 	const bool bInvalidCachedAttributes = bDoEvaluationRateOptimization && CachedAttributes != CustomAttributes;
 
 	const bool bShouldDoEvaluation = !bDoEvaluationRateOptimization || bInvalidCachedBones || bInvalidCachedCurve || (bExternalTickRateControlled && bExternalUpdate) || (bCachedShouldUseUpdateRateOptimizations && !AnimUpdateRateParams->ShouldSkipEvaluation());
@@ -2779,9 +2777,9 @@ void USkeletalMeshComponent::RefreshBoneTransforms(FActorComponentTickFunction*
 					LocalEditableSpaceBases.Reset();
 					LocalEditableSpaceBases.Append(CachedComponentSpaceTransforms);
 				}
-
+				PRAGMA_DISABLE_DEPRECATION_WARNINGS
 				AnimCurves.CopyFrom(CachedCurve);
-
+				PRAGMA_ENABLE_DEPRECATION_WARNINGS
 				if (CachedAttributes.ContainsData())
 				{
 					CustomAttributes.CopyFrom(CachedAttributes);
@@ -2817,7 +2815,9 @@ void USkeletalMeshComponent::SwapEvaluationContextBuffers()
 	Exchange(AnimEvaluationContext.BoneSpaceTransforms, BoneSpaceTransforms);
 	PRAGMA_ENABLE_DEPRECATION_WARNINGS
 	Exchange(AnimEvaluationContext.CachedBoneSpaceTransforms, CachedBoneSpaceTransforms);
+	PRAGMA_DISABLE_DEPRECATION_WARNINGS
 	Exchange(AnimEvaluationContext.Curve, AnimCurves);
+	PRAGMA_ENABLE_DEPRECATION_WARNINGS
 	Exchange(AnimEvaluationContext.CachedCurve, CachedCurve);
 	Exchange(AnimEvaluationContext.RootBoneTranslation, RootBoneTranslation);
 
@@ -2894,7 +2894,9 @@ void USkeletalMeshComponent::DispatchParallelTickPose(FActorComponentTickFunctio
 				AnimEvaluationContext.AnimInstance = AnimScriptInstance;
 
 				// We dont set up the Curve here, as we dont use it in Update()
+				PRAGMA_DISABLE_DEPRECATION_WARNINGS
 				AnimCurves.Empty();
+				PRAGMA_ENABLE_DEPRECATION_WARNINGS
 
 				CustomAttributes.Empty();
 
@@ -2952,7 +2954,9 @@ void USkeletalMeshComponent::PostAnimEvaluation(FAnimationEvaluationContext& Eva
 	{
 		if (EvaluationContext.bDuplicateToCacheCurve)
 		{
+			PRAGMA_DISABLE_DEPRECATION_WARNINGS
 			CachedCurve.CopyFrom(AnimCurves);
+			PRAGMA_ENABLE_DEPRECATION_WARNINGS
 		}
 		
 		if (EvaluationContext.bDuplicateToCachedAttributes)
@@ -3004,10 +3008,10 @@ void USkeletalMeshComponent::PostAnimEvaluation(FAnimationEvaluationContext& Eva
 			PRAGMA_DISABLE_DEPRECATION_WARNINGS
 			FAnimationRuntime::LerpBoneTransforms(BoneSpaceTransforms, CachedBoneSpaceTransforms, Alpha, RequiredBones);
 			GetSkeletalMeshAsset()->FillComponentSpaceTransforms(BoneSpaceTransforms, FillComponentSpaceTransformsRequiredBones, GetEditableComponentSpaceTransforms());
-			PRAGMA_ENABLE_DEPRECATION_WARNINGS
+
 			// interpolate curve
 			AnimCurves.LerpTo(CachedCurve, Alpha);
-
+			PRAGMA_ENABLE_DEPRECATION_WARNINGS
 			// Interpolate custom attributes
 			UE::Anim::Attributes::InterpolateAttributes(CustomAttributes, CachedAttributes, Alpha);
 		}
@@ -3024,7 +3028,9 @@ void USkeletalMeshComponent::PostAnimEvaluation(FAnimationEvaluationContext& Eva
 		if(AnimScriptInstance)
 		{
 #if WITH_EDITOR
+			PRAGMA_DISABLE_DEPRECATION_WARNINGS
 			GetEditableAnimationCurves() = AnimCurves;
+			PRAGMA_ENABLE_DEPRECATION_WARNINGS
 #endif 
 			GetEditableCustomAttributes() = CustomAttributes;
 
@@ -5334,6 +5340,22 @@ bool USkeletalMeshComponent::FindAttributeChecked(const FName& BoneName, const F
 	return bFound;
 }
 
+bool USkeletalMeshComponent::GetCurveValue(FName CurveName, float DefaultValue, float& Value) const
+{
+	bool bHasElement = false;
+	PRAGMA_DISABLE_DEPRECATION_WARNINGS
+	Value = AnimCurves.Get(CurveName, bHasElement, DefaultValue);
+	PRAGMA_ENABLE_DEPRECATION_WARNINGS
+	return bHasElement;
+}
+
+const FBlendedHeapCurve& USkeletalMeshComponent::GetAnimCurves() const
+{
+	PRAGMA_DISABLE_DEPRECATION_WARNINGS
+	return AnimCurves;
+	PRAGMA_ENABLE_DEPRECATION_WARNINGS
+}
+
 const FName& USkeletalMeshComponent::GetAnimationModePropertyNameChecked()
 {
 	static FName Name = GET_MEMBER_NAME_CHECKED(USkeletalMeshComponent, AnimationMode);
diff --git a/Engine/Source/Runtime/Engine/Private/Components/SkinnedMeshComponent.cpp b/Engine/Source/Runtime/Engine/Private/Components/SkinnedMeshComponent.cpp
--- a/Engine/Source/Runtime/Engine/Private/Components/SkinnedMeshComponent.cpp
+++ b/Engine/Source/Runtime/Engine/Private/Components/SkinnedMeshComponent.cpp
@@ -481,6 +481,8 @@ USkinnedMeshComponent::USkinnedMeshComponent(const FObjectInitializer& ObjectIni
 
 	bFollowerShouldTickPose = false;
 
+	bHasSocketAttachments = false;
+
 #if UE_ENABLE_DEBUG_DRAWING
 	bDrawDebugSkeleton = false;
 #endif
@@ -3257,6 +3259,37 @@ void USkinnedMeshComponent::RemoveAllSocketOverrides()
 	SocketOverrideLookup.Reset();
 }
 
+void USkinnedMeshComponent::OnChildAttached(USceneComponent* ChildComponent)
+{
+	Super::OnChildAttached(ChildComponent);
+	UpdateCachedSocketAttachmentFlag();
+}
+
+void USkinnedMeshComponent::OnChildDetached(USceneComponent* ChildComponent)
+{
+	Super::OnChildDetached(ChildComponent);
+	UpdateCachedSocketAttachmentFlag();
+}
+
+void USkinnedMeshComponent::UpdateCachedSocketAttachmentFlag()
+{
+	// Update whether any of our attachments are socket-attached
+	bHasSocketAttachments = false;
+	for (const USceneComponent* AttachedComponent : GetAttachChildren())
+	{
+		if (AttachedComponent->GetAttachSocketName() != NAME_None)
+		{
+			bHasSocketAttachments = true;
+			break;
+		}
+	}
+}
+
+bool USkinnedMeshComponent::HasAnyAttachedSockets() const
+{
+	return bHasSocketAttachments;
+}
+
 bool USkinnedMeshComponent::DoesSocketExist(FName InSocketName) const
 {
 	return (GetSocketBoneName(InSocketName) != NAME_None);
diff --git a/Engine/Source/Runtime/Engine/Private/PhysicsEngine/PhysAnim.cpp b/Engine/Source/Runtime/Engine/Private/PhysicsEngine/PhysAnim.cpp
--- a/Engine/Source/Runtime/Engine/Private/PhysicsEngine/PhysAnim.cpp
+++ b/Engine/Source/Runtime/Engine/Private/PhysicsEngine/PhysAnim.cpp
@@ -478,7 +478,10 @@ void USkeletalMeshComponent::FinalizeAnimationUpdate()
 
 		// Update Child Transform - The above function changes bone transform, so will need to update child transform
 		// But only children attached to us via a socket.
-		UpdateChildTransforms(EUpdateTransformFlags::OnlyUpdateIfUsingSocket);
+		if (bHasSocketAttachments)
+		{
+			UpdateChildTransforms(EUpdateTransformFlags::OnlyUpdateIfUsingSocket);
+		}
 	}
 
 	if(bUpdateOverlapsOnAnimationFinalize)
