Remove IPrimitiveComponent & IStaticMeshComponent registered implementers at module shutdown
* Added RemoveImplementer() to the interface

diff --git a/Engine/Source/Runtime/Engine/Classes/Components/ComponentInterfaces.h b/Engine/Source/Runtime/Engine/Classes/Components/ComponentInterfaces.h
--- a/Engine/Source/Runtime/Engine/Classes/Components/ComponentInterfaces.h
+++ b/Engine/Source/Runtime/Engine/Classes/Components/ComponentInterfaces.h
@@ -101,7 +101,7 @@ public:
 	virtual HHitProxy* CreatePrimitiveHitProxies(TArray<TRefCountPtr<HHitProxy> >& OutHitProxies) = 0;
 	
 	ENGINE_API static void AddImplementer(const FComponentInterfaceImplementation& Implementer);
-
+	ENGINE_API static void RemoveImplementer(const UClass* ImplementerClass);
 protected:
 
 	template<class T> 
@@ -128,6 +128,7 @@ public:
 	}
 	
 	ENGINE_API static void AddImplementer(const FComponentInterfaceImplementation& Implementer);
+	ENGINE_API static void RemoveImplementer(const UClass* ImplementerClass);
 
 protected:
 
diff --git a/Engine/Source/Runtime/Engine/Private/Components/ComponentInterfaces.cpp b/Engine/Source/Runtime/Engine/Private/Components/ComponentInterfaces.cpp
--- a/Engine/Source/Runtime/Engine/Private/Components/ComponentInterfaces.cpp
+++ b/Engine/Source/Runtime/Engine/Private/Components/ComponentInterfaces.cpp
@@ -10,7 +10,17 @@ TArray<FComponentInterfaceImplementation> IPrimitiveComponent::Implementers;
 void IPrimitiveComponent::AddImplementer(const FComponentInterfaceImplementation& Implementer)
 {	
 	Implementers.Add(Implementer);
+}
 
+void IPrimitiveComponent::RemoveImplementer(const UClass* ImplementerClass)
+{
+	for (TArray<FComponentInterfaceImplementation>::TIterator ImplementerIt(Implementers); ImplementerIt; ++ImplementerIt)
+	{
+		if (ImplementerIt->Class == ImplementerClass)
+		{
+			ImplementerIt.RemoveCurrent();
+		}
+	}
 }
 
 TArray<FComponentInterfaceImplementation> IStaticMeshComponent::Implementers;
@@ -18,4 +28,15 @@ TArray<FComponentInterfaceImplementation> IStaticMeshComponent::Implementers;
 void IStaticMeshComponent::AddImplementer(const FComponentInterfaceImplementation& Implementer)
 {
 	Implementers.Add(Implementer);
-}
\ No newline at end of file
+}
+
+void IStaticMeshComponent::RemoveImplementer(const UClass* ImplementerClass)
+{
+	for (TArray<FComponentInterfaceImplementation>::TIterator ImplementerIt(Implementers); ImplementerIt; ++ImplementerIt)
+	{
+		if (ImplementerIt->Class == ImplementerClass)
+		{
+			ImplementerIt.RemoveCurrent();
+		}
+	}
+}
diff --git a/Engine/Source/Runtime/Engine/Private/UnrealEngine.cpp b/Engine/Source/Runtime/Engine/Private/UnrealEngine.cpp
--- a/Engine/Source/Runtime/Engine/Private/UnrealEngine.cpp
+++ b/Engine/Source/Runtime/Engine/Private/UnrealEngine.cpp
@@ -398,6 +398,9 @@ void FEngineModule::StartupModule()
 
 void FEngineModule::ShutdownModule()
 {
+	IPrimitiveComponent::RemoveImplementer(UPrimitiveComponent::StaticClass());
+	IStaticMeshComponent::RemoveImplementer(UStaticMeshComponent::StaticClass());
+
 #if UE_WITH_IRIS
 	UE::Net::ShutdownEngineNetObjectFactories();
 #endif
