[Materials] Add external view property declaration to external code registry.

- Implements UMaterialExpressionViewProperty in the new MIR translator.
- Moves all view properties from old HLSL translator to MaterialExternalCodeRegistry singleton to be shared between old and new material translators.
- Allow allocating dynamic strings in MIR::FEmitter that are referenced in the output IR module.

#rnx

diff --git a/Engine/Source/Runtime/Engine/Private/Materials/HLSLMaterialTranslator.cpp b/Engine/Source/Runtime/Engine/Private/Materials/HLSLMaterialTranslator.cpp
--- a/Engine/Source/Runtime/Engine/Private/Materials/HLSLMaterialTranslator.cpp
+++ b/Engine/Source/Runtime/Engine/Private/Materials/HLSLMaterialTranslator.cpp
@@ -5358,43 +5358,7 @@ int32 FHLSLMaterialTranslator::GenericConstant(const UE::Shader::FValue& Value)
 	
 int32 FHLSLMaterialTranslator::ViewProperty(EMaterialExposedViewProperty Property, bool InvProperty)
 {
-	check(Property < MEVP_MAX);
-
-	// Compile time struct storing all EMaterialExposedViewProperty's enumerations' HLSL compilation specific meta information
-	struct EMaterialExposedViewPropertyMeta
-	{
-		EMaterialExposedViewProperty EnumValue;
-		EMaterialValueType Type;
-		const TCHAR * PropertyCode;
-		const TCHAR * InvPropertyCode;
-	};
-
-	static const EMaterialExposedViewPropertyMeta ViewPropertyMetaArray[] = {
-		{MEVP_BufferSize, MCT_Float2, TEXT("View.BufferSizeAndInvSize.xy"), TEXT("View.BufferSizeAndInvSize.zw")},
-		{MEVP_FieldOfView, MCT_Float2, TEXT("View.<PREV>FieldOfViewWideAngles"), nullptr},
-		{MEVP_TanHalfFieldOfView, MCT_Float2, TEXT("Get<PREV>TanHalfFieldOfView()"), TEXT("Get<PREV>CotanHalfFieldOfView()")},
-		{MEVP_ViewSize, MCT_Float2, TEXT("View.ViewSizeAndInvSize.xy"), TEXT("View.ViewSizeAndInvSize.zw")},
-		{MEVP_WorldSpaceViewPosition, MCT_LWCVector3, TEXT("Get<PREV>WorldViewOrigin(Parameters)"), nullptr},
-		{MEVP_WorldSpaceCameraPosition, MCT_LWCVector3, TEXT("Get<PREV>WorldCameraOrigin(Parameters)"), nullptr},
-		{MEVP_ViewportOffset, MCT_Float2, TEXT("View.ViewRectMin.xy"), nullptr},
-		{MEVP_TemporalSampleCount, MCT_Float1, TEXT("View.TemporalAAParams.y"), nullptr},
-		{MEVP_TemporalSampleIndex, MCT_Float1, TEXT("View.TemporalAAParams.x"), nullptr},
-		{MEVP_TemporalSampleOffset, MCT_Float2, TEXT("View.TemporalAAParams.zw"), nullptr},
-		{MEVP_RuntimeVirtualTextureOutputLevel, MCT_Float1, TEXT("GetRuntimeVirtualTextureMipLevel().x"), nullptr},
-		{MEVP_RuntimeVirtualTextureOutputDerivative, MCT_Float2, TEXT("GetRuntimeVirtualTextureMipLevel().zw"), nullptr},
-		{MEVP_PreExposure, MCT_Float1, TEXT("View.PreExposure.x"), TEXT("View.OneOverPreExposure.x")},
-		{MEVP_RuntimeVirtualTextureMaxLevel, MCT_Float1, TEXT("GetRuntimeVirtualTextureMipLevel().y"), nullptr},
-		{MEVP_ResolutionFraction, MCT_Float1, TEXT("View.ResolutionFractionAndInv.x"), TEXT("View.ResolutionFractionAndInv.y")},
-		{MEVP_PostVolumeUserFlags, MCT_Float1, TEXT("View.PostVolumeUserFlags"), nullptr},
-		{MEVP_FirstPersonFieldOfView, MCT_Float2, TEXT("View.<PREV>FirstPersonFieldOfViewWideAngles"), nullptr},
-		{MEVP_FirstPersonTanHalfFieldOfView, MCT_Float2, TEXT("View.<PREV>FirstPersonTanAndInvTanHalfFOV.xy"), TEXT("View.<PREV>FirstPersonTanAndInvTanHalfFOV.zw")},
-		{MEVP_FirstPersonScale, MCT_Float, TEXT("View.<PREV>FirstPersonScale"), nullptr},
-		{MEVP_NearPlane, MCT_Float, TEXT("View.NearPlane"), nullptr},
-	};
-	static_assert((sizeof(ViewPropertyMetaArray) / sizeof(ViewPropertyMetaArray[0])) == MEVP_MAX, "incoherency between EMaterialExposedViewProperty and ViewPropertyMetaArray");
-
-	auto& PropertyMeta = ViewPropertyMetaArray[Property];
-	check(Property == PropertyMeta.EnumValue);
+	const FMaterialExposedViewPropertyMeta& PropertyMeta = MaterialExternalCodeRegistry::Get().GetExternalViewPropertyCode(Property);
 
 	FString Code = PropertyMeta.PropertyCode;
 
diff --git a/Engine/Source/Runtime/Engine/Private/Materials/HLSLMaterialTranslator.h b/Engine/Source/Runtime/Engine/Private/Materials/HLSLMaterialTranslator.h
--- a/Engine/Source/Runtime/Engine/Private/Materials/HLSLMaterialTranslator.h
+++ b/Engine/Source/Runtime/Engine/Private/Materials/HLSLMaterialTranslator.h
@@ -936,7 +936,7 @@ protected:
 	virtual int32 Constant4(float X, float Y, float Z, float W) override;
 	virtual int32 GenericConstant(const UE::Shader::FValue& Value) override;
 	
-	virtual int32 ViewProperty(EMaterialExposedViewProperty Property, bool InvProperty) override;
+	virtual int32 ViewProperty(EMaterialExposedViewProperty Property, bool bInvProperty) override;
 	virtual int32 IsOrthographic() override;
 
 	virtual int32 GameTime(bool bPeriodic, float Period) override;
diff --git a/Engine/Source/Runtime/Engine/Private/Materials/MaterialExpressionsIR.cpp b/Engine/Source/Runtime/Engine/Private/Materials/MaterialExpressionsIR.cpp
--- a/Engine/Source/Runtime/Engine/Private/Materials/MaterialExpressionsIR.cpp
+++ b/Engine/Source/Runtime/Engine/Private/Materials/MaterialExpressionsIR.cpp
@@ -1102,4 +1102,35 @@ void UMaterialExpressionConvert::Build(MIR::FEmitter& Emitter)
 	}
 }
 
+void UMaterialExpressionViewProperty::Build(MIR::FEmitter& Emitter)
+{
+	check(Property < MEVP_MAX);
+
+	const FMaterialExposedViewPropertyMeta& PropertyMeta = MaterialExternalCodeRegistry::Get().GetExternalViewPropertyCode(Property);
+
+	const bool bHasCustomInverseCode = PropertyMeta.InvPropertyCode != nullptr;
+
+	for (int32 OutputIndex = 0; OutputIndex < 2; ++OutputIndex)
+	{
+		const bool bInvProperty = OutputIndex == 1;
+
+		FString HLSLCode = bInvProperty && bHasCustomInverseCode ? PropertyMeta.InvPropertyCode : PropertyMeta.PropertyCode;
+		MIR::FTypePtr HLSLCodeType = MIR::FType::FromMaterialValueType(PropertyMeta.Type);
+
+		if (bInvProperty && !bHasCustomInverseCode)
+		{
+			// Fall back to compute the property's inverse from PropertyCode
+			Emitter.Output(OutputIndex,
+				Emitter.EmitOperator(MIR::BO_Divide,
+					Emitter.EmitConstantFloatScalar(1.0f),
+					Emitter.EmitExternalCode(HLSLCodeType, MoveTemp(HLSLCode))));
+		}
+		else
+		{
+			// @todo-laura.hermanns - old translator used CastToNonLWCIfDisabled(), but LWC not supported in MIR yet
+			Emitter.Output(OutputIndex, Emitter.EmitExternalCode(HLSLCodeType, MoveTemp(HLSLCode)));
+		}
+	}
+}
+
 #endif // WITH_EDITOR
diff --git a/Engine/Source/Runtime/Engine/Private/Materials/MaterialExternalCodeRegistry.cpp b/Engine/Source/Runtime/Engine/Private/Materials/MaterialExternalCodeRegistry.cpp
--- a/Engine/Source/Runtime/Engine/Private/Materials/MaterialExternalCodeRegistry.cpp
+++ b/Engine/Source/Runtime/Engine/Private/Materials/MaterialExternalCodeRegistry.cpp
@@ -40,3 +40,39 @@ const FMaterialExternalCodeDeclaration* MaterialExternalCodeRegistry::FindExtern
 	}
 	return nullptr;
 }
+
+const FMaterialExposedViewPropertyMeta& MaterialExternalCodeRegistry::GetExternalViewPropertyCode(const EMaterialExposedViewProperty InViewProperty) const
+{
+	// Compile time struct storing all EMaterialExposedViewProperty's enumerations' HLSL compilation specific meta information
+	static const FMaterialExposedViewPropertyMeta ViewPropertyMetaArray[] =
+	{
+		{MEVP_BufferSize, MCT_Float2, TEXT("View.BufferSizeAndInvSize.xy"), TEXT("View.BufferSizeAndInvSize.zw")},
+		{MEVP_FieldOfView, MCT_Float2, TEXT("View.<PREV>FieldOfViewWideAngles"), nullptr},
+		{MEVP_TanHalfFieldOfView, MCT_Float2, TEXT("Get<PREV>TanHalfFieldOfView()"), TEXT("Get<PREV>CotanHalfFieldOfView()")},
+		{MEVP_ViewSize, MCT_Float2, TEXT("View.ViewSizeAndInvSize.xy"), TEXT("View.ViewSizeAndInvSize.zw")},
+		{MEVP_WorldSpaceViewPosition, MCT_LWCVector3, TEXT("Get<PREV>WorldViewOrigin(Parameters)"), nullptr},
+		{MEVP_WorldSpaceCameraPosition, MCT_LWCVector3, TEXT("Get<PREV>WorldCameraOrigin(Parameters)"), nullptr},
+		{MEVP_ViewportOffset, MCT_Float2, TEXT("View.ViewRectMin.xy"), nullptr},
+		{MEVP_TemporalSampleCount, MCT_Float1, TEXT("View.TemporalAAParams.y"), nullptr},
+		{MEVP_TemporalSampleIndex, MCT_Float1, TEXT("View.TemporalAAParams.x"), nullptr},
+		{MEVP_TemporalSampleOffset, MCT_Float2, TEXT("View.TemporalAAParams.zw"), nullptr},
+		{MEVP_RuntimeVirtualTextureOutputLevel, MCT_Float1, TEXT("GetRuntimeVirtualTextureMipLevel().x"), nullptr},
+		{MEVP_RuntimeVirtualTextureOutputDerivative, MCT_Float2, TEXT("GetRuntimeVirtualTextureMipLevel().zw"), nullptr},
+		{MEVP_PreExposure, MCT_Float1, TEXT("View.PreExposure.x"), TEXT("View.OneOverPreExposure.x")},
+		{MEVP_RuntimeVirtualTextureMaxLevel, MCT_Float1, TEXT("GetRuntimeVirtualTextureMipLevel().y"), nullptr},
+		{MEVP_ResolutionFraction, MCT_Float1, TEXT("View.ResolutionFractionAndInv.x"), TEXT("View.ResolutionFractionAndInv.y")},
+		{MEVP_PostVolumeUserFlags, MCT_Float1, TEXT("View.PostVolumeUserFlags"), nullptr},
+		{MEVP_FirstPersonFieldOfView, MCT_Float2, TEXT("View.<PREV>FirstPersonFieldOfViewWideAngles"), nullptr},
+		{MEVP_FirstPersonTanHalfFieldOfView, MCT_Float2, TEXT("View.<PREV>FirstPersonTanAndInvTanHalfFOV.xy"), TEXT("View.<PREV>FirstPersonTanAndInvTanHalfFOV.zw")},
+		{MEVP_FirstPersonScale, MCT_Float, TEXT("View.<PREV>FirstPersonScale"), nullptr},
+		{MEVP_NearPlane, MCT_Float, TEXT("View.NearPlane"), nullptr},
+	};
+	static_assert((sizeof(ViewPropertyMetaArray) / sizeof(ViewPropertyMetaArray[0])) == MEVP_MAX, "incoherency between EMaterialExposedViewProperty and ViewPropertyMetaArray");
+
+	check(InViewProperty < MEVP_MAX);
+
+	const FMaterialExposedViewPropertyMeta& PropertyMeta = ViewPropertyMetaArray[InViewProperty];
+	check(InViewProperty == PropertyMeta.EnumValue);
+
+	return PropertyMeta;
+}
diff --git a/Engine/Source/Runtime/Engine/Private/Materials/MaterialExternalCodeRegistry.h b/Engine/Source/Runtime/Engine/Private/Materials/MaterialExternalCodeRegistry.h
--- a/Engine/Source/Runtime/Engine/Private/Materials/MaterialExternalCodeRegistry.h
+++ b/Engine/Source/Runtime/Engine/Private/Materials/MaterialExternalCodeRegistry.h
@@ -8,6 +8,7 @@
 
 #include "CoreMinimal.h"
 #include "Materials/HLSLMaterialDerivativeAutogen.h"
+#include "Materials/MaterialExpressionViewProperty.h"
 #include "RHIDefinitions.h"
 
 #include "MaterialExternalCodeRegistry.generated.h"
@@ -119,6 +120,15 @@ public:
 
 };
 
+/** Helper struct for hard coded external code expressions for view properties. */
+struct FMaterialExposedViewPropertyMeta
+{
+	EMaterialExposedViewProperty EnumValue;
+	EMaterialValueType Type;
+	const TCHAR* PropertyCode;
+	const TCHAR* InvPropertyCode;
+};
+
 /** Singleton class to register external HLSL function and input declarations for material IR modules. */
 class MaterialExternalCodeRegistry
 {
@@ -128,9 +138,12 @@ public:
 
 	static MaterialExternalCodeRegistry& Get();
 
-	/** Returns the external function declaration for the function with the specified name. */
+	/** Returns the external code declaration for the specified name. */
 	const FMaterialExternalCodeDeclaration* FindExternalCode(const FName& InExternalCodeIdentifier) const;
 
+	/** Returns the external code declaration for the specified view property. */
+	const FMaterialExposedViewPropertyMeta& GetExternalViewPropertyCode(const EMaterialExposedViewProperty InViewProperty) const;
+
 private:
 	MaterialExternalCodeRegistry();
 
diff --git a/Engine/Source/Runtime/Engine/Private/Materials/MaterialIRDebug.cpp b/Engine/Source/Runtime/Engine/Private/Materials/MaterialIRDebug.cpp
--- a/Engine/Source/Runtime/Engine/Private/Materials/MaterialIRDebug.cpp
+++ b/Engine/Source/Runtime/Engine/Private/Materials/MaterialIRDebug.cpp
@@ -310,14 +310,7 @@ struct FDebugDumpIRState
 		}
 		else if (const FExternalCode* ExternalCode = Instr->As<FExternalCode>())
 		{
-			if (const FMaterialExternalCodeDeclaration* ExternalCodeDeclaration = ExternalCode->ExternalCodeDeclaration)
-			{
-				Out.Appendf(TEXT(" \"%s\""), *ExternalCodeDeclaration->Definition);
-			}
-			else
-			{
-				Out.Append(TEXT(" <Missing Definition>"));
-			}
+			Out.Appendf(TEXT(" \"%s\""), ExternalCode->HLSLCode);
 		}
 	}
 
diff --git a/Engine/Source/Runtime/Engine/Private/Materials/MaterialIREmitter.cpp b/Engine/Source/Runtime/Engine/Private/Materials/MaterialIREmitter.cpp
--- a/Engine/Source/Runtime/Engine/Private/Materials/MaterialIREmitter.cpp
+++ b/Engine/Source/Runtime/Engine/Private/Materials/MaterialIREmitter.cpp
@@ -1835,13 +1835,19 @@ FValue* FEmitter::EmitTextureSample(FValue* Texture, FValue* TexCoord, ESamplerS
 	return EmitPrototype(this, Prototype);
 }
 
+FValue* FEmitter::EmitExternalCode(FTypePtr Type, FString InHLSLCode, bool bNoTokenSubstitution)
+{
+	MIR::FExternalCode Prototype = MakePrototype<MIR::FExternalCode>(Type);
+	Prototype.HLSLCode = AllocateModuleString(MoveTemp(InHLSLCode));
+	Prototype.bNoTokenSubstitution = bNoTokenSubstitution;
+	return EmitPrototype(this, Prototype);
+}
+
 FValue* FEmitter::EmitExternalCode(const FMaterialExternalCodeDeclaration* InExternalCodeDeclaration)
 {
 	check(InExternalCodeDeclaration != nullptr);
 	FTypePtr ReturnType = FType::FromMaterialValueType(InExternalCodeDeclaration->GetReturnTypeValue());
-	MIR::FExternalCode Prototype = MakePrototype<MIR::FExternalCode>(ReturnType);
-	Prototype.ExternalCodeDeclaration = InExternalCodeDeclaration;
-	return EmitPrototype(this, Prototype);
+	return  EmitExternalCode(ReturnType, InExternalCodeDeclaration->Definition);
 }
 
 FTypePtr FEmitter::TryGetCommonType(FTypePtr A, FTypePtr B)
@@ -1910,6 +1916,12 @@ void FEmitter::Initialize()
 	ConstantFalse = EmitPrototype(this, Temp);
 }
 
+const TCHAR* FEmitter::AllocateModuleString(FString InString)
+{
+	Module->Strings.Add(MoveTemp(InString));
+	return GetData(Module->Strings.Last());
+}
+
 bool FEmitter::FValueKeyFuncs::Matches(KeyInitType A, KeyInitType B)
 {
 	return A->Equals(B);
diff --git a/Engine/Source/Runtime/Engine/Private/Materials/MaterialIRToHLSLTranslator.cpp b/Engine/Source/Runtime/Engine/Private/Materials/MaterialIRToHLSLTranslator.cpp
--- a/Engine/Source/Runtime/Engine/Private/Materials/MaterialIRToHLSLTranslator.cpp
+++ b/Engine/Source/Runtime/Engine/Private/Materials/MaterialIRToHLSLTranslator.cpp
@@ -400,8 +400,23 @@ struct FTranslator : FMaterialIRToHLSLTranslation
 
 	void LowerExternalCode(const MIR::FExternalCode* ExternalCode)
 	{
-		check(ExternalCode->ExternalCodeDeclaration != nullptr);
-		Printer.Buffer.Append(*ExternalCode->ExternalCodeDeclaration->Definition);
+		check(ExternalCode->HLSLCode != nullptr);
+
+		// Substitute placeholder tokens now unless disabled for custom nodes
+		if (ExternalCode->bNoTokenSubstitution)
+		{
+			Printer.Buffer.Append(ExternalCode->HLSLCode);
+		}
+		else
+		{
+			FString HLSLCode = ExternalCode->HLSLCode;
+
+			// @todo-laura.hermanns - Must be replaced as soon as this state is available in MIR::FExternalCode
+			constexpr bool bCompilingPreviousFrame = false;
+			HLSLCode.ReplaceInline(TEXT("<PREV>"), bCompilingPreviousFrame ? TEXT("Prev") : TEXT(""));
+
+			Printer.Buffer.Append(MoveTemp(HLSLCode));
+		}
 	}
 
 	void LowerTextureObject(const MIR::FTextureObject* TextureObject)
diff --git a/Engine/Source/Runtime/Engine/Public/Materials/MaterialExpressionViewProperty.h b/Engine/Source/Runtime/Engine/Public/Materials/MaterialExpressionViewProperty.h
--- a/Engine/Source/Runtime/Engine/Public/Materials/MaterialExpressionViewProperty.h
+++ b/Engine/Source/Runtime/Engine/Public/Materials/MaterialExpressionViewProperty.h
@@ -66,6 +66,7 @@ class UMaterialExpressionViewProperty : public UMaterialExpression
 	
 	//~ Begin UMaterialExpression Interface
 #if WITH_EDITOR
+	virtual void Build(MIR::FEmitter& Emitter) override;
 	virtual int32 Compile(class FMaterialCompiler* Compiler, int32 OutputIndex) override;
 	virtual void GetCaption(TArray<FString>& OutCaptions) const override;
 #endif
diff --git a/Engine/Source/Runtime/Engine/Public/Materials/MaterialIR.h b/Engine/Source/Runtime/Engine/Public/Materials/MaterialIR.h
--- a/Engine/Source/Runtime/Engine/Public/Materials/MaterialIR.h
+++ b/Engine/Source/Runtime/Engine/Public/Materials/MaterialIR.h
@@ -359,7 +359,8 @@ struct FSetMaterialOutput : TInstruction<VK_SetMaterialOutput>
 // MIR instruction for external code declarations; See BaseMaterialExpressions.ini file.
 struct FExternalCode : TInstruction<VK_ExternalCode>
 {
-	const FMaterialExternalCodeDeclaration* ExternalCodeDeclaration;
+	const TCHAR* HLSLCode;
+	bool bNoTokenSubstitution; // Whether to skip substution of abstract tokens like "<PREV>". Only true for custom nodes.
 };
 
 // Note: if you modify this enum, update the implementations of the
diff --git a/Engine/Source/Runtime/Engine/Public/Materials/MaterialIREmitter.h b/Engine/Source/Runtime/Engine/Public/Materials/MaterialIREmitter.h
--- a/Engine/Source/Runtime/Engine/Public/Materials/MaterialIREmitter.h
+++ b/Engine/Source/Runtime/Engine/Public/Materials/MaterialIREmitter.h
@@ -161,6 +161,7 @@ public:
 	FValue* EmitBranch(FValue* Condition, FValue* True, FValue* False);
 	FValue* EmitConstruct(FTypePtr Type, FValue* Initializer);
 	FValue* EmitTextureSample(FValue* Texture, FValue* TexCoord, ESamplerSourceMode SamplerSourceMode, ETextureMipValueMode MipValueMode, EMaterialSamplerType SamplerType);
+	FValue* EmitExternalCode(FTypePtr Type, FString InHLSLCode, bool bNoTokenSubstitution = false);
 	FValue* EmitExternalCode(const FMaterialExternalCodeDeclaration* InExternalCodeDeclaration);
 
 	/* Types */
@@ -187,6 +188,8 @@ private:
 
 	void Initialize();
 
+	const TCHAR* AllocateModuleString(FString InString);
+
 	struct FValueKeyFuncs : DefaultKeyFuncs<FValue*>
 	{
 		static bool Matches(KeyInitType A, KeyInitType B);
diff --git a/Engine/Source/Runtime/Engine/Public/Materials/MaterialIRModule.h b/Engine/Source/Runtime/Engine/Public/Materials/MaterialIRModule.h
--- a/Engine/Source/Runtime/Engine/Public/Materials/MaterialIRModule.h
+++ b/Engine/Source/Runtime/Engine/Public/Materials/MaterialIRModule.h
@@ -53,6 +53,7 @@ private:
 	TArray<MIR::FValue*> Values;
 	TArray<MIR::FSetMaterialOutput*> Outputs[MIR::NumStages];
 	TArray<FError> Errors;
+	TArray<FString> Strings;
 	MIR::FBlock* RootBlock[MIR::NumStages];
 	FStatistics Statistics;
 	TMap<FMaterialParameterInfo, uint32> ParameterInfoToId;
