[DSTM] Fix potential serialization crash if we migrate a skeletal mesh component, but we failed to serialize data from some of its bodies/constraints.

The problem is that to avoid extra allocations we serialize all the physics data for each body and constraint into a single buffer and send that. But the receiving end needs to be able to split it into separate buffer (for now) so it can process it when each body/constraint is ready.

To do that we were storing the offset at which entry is in the unified buffer, and figuring out the sizes based in the next element offset, but that breaksdown if you have entries of zero size.

There was code to catch these issues and ensure/log errors, but there some edge cases that were not handled.

This fix serializes a tag per entry (which is the data offset + the size we expect to read per entry) and use that instead of trying to calculate the sizes (in hindsight, this should have been the initial implementation, at least until we are able to avoid the need of splitting the buffer again).

diff --git a/Engine/Source/Runtime/Engine/Classes/Components/SkeletalMeshComponent.h b/Engine/Source/Runtime/Engine/Classes/Components/SkeletalMeshComponent.h
--- a/Engine/Source/Runtime/Engine/Classes/Components/SkeletalMeshComponent.h
+++ b/Engine/Source/Runtime/Engine/Classes/Components/SkeletalMeshComponent.h
@@ -1997,16 +1997,19 @@ protected:
 private:
 	struct UE_INTERNAL FMigratedPhysicsState
 	{
-		int32 NumBodies = 0;
-		int32 NumConstraints = 0;
-
 		void ReadElementDataIntoBuffer(Chaos::FSerializedDataBuffer& TargetBuffer);
 
 		template<typename ElementType>
 		void AppendDataFromElement(ElementType& InElement)
 		{
-			DataOffsetsPerElementIndex.Emplace(GetSize());
+			FDataEntryTag DataTag;
+			DataTag.DataOffset = GetSize();
+			
 			InElement.SerializeState(MigratedStateAsBytes);
+
+			DataTag.DataSize = GetSize() - DataTag.DataOffset;
+
+			Header.DataTagPerElementIndex.Emplace(DataTag);
 		}
 
 		void Serialize(FArchive& Ar);
@@ -2016,12 +2019,71 @@ private:
 			return MigratedStateAsBytes.GetSize();
 		}
 
-	private:
-		int32 GetDataSizeForElementIndex(int32 ElementIndex) const;
-		int32 GetDataOffset(int32 ElementIndex) const;
+private:
 
 		Chaos::FSerializedDataBuffer MigratedStateAsBytes;
-		TArray<int32> DataOffsetsPerElementIndex;
+
+		struct FDataEntryTag
+		{
+			int32 DataOffset = INDEX_NONE;
+			int32 DataSize = INDEX_NONE;
+
+			friend FArchive& operator<<(FArchive& Ar, FDataEntryTag& Data)
+			{
+				Ar << Data.DataOffset;
+				Ar << Data.DataSize;
+				return Ar;
+			}
+
+			bool IsValid() const
+			{
+				return DataOffset != INDEX_NONE && DataSize != INDEX_NONE;
+			}
+		};
+
+		struct FHeader
+		{
+			int32 NumBodies = 0;
+			int32 NumConstraints = 0;
+
+			TArray<FDataEntryTag> DataTagPerElementIndex;
+			
+			friend FArchive& operator<<(FArchive& Ar, FHeader& Data)
+			{
+				Ar << Data.NumBodies;
+				Ar << Data.NumConstraints;
+				Ar << Data.DataTagPerElementIndex;
+
+				return Ar;
+			}
+
+			bool IsValid() const
+			{
+				return NumBodies + NumConstraints == DataTagPerElementIndex.Num();
+			}
+		};
+		
+		FHeader Header;
+
+	public:
+
+		void SetNumOfBodies(int32 NumBodies)
+		{
+			Header.NumBodies = NumBodies;
+		}
+
+		void SetNumOfConstraints(int32 NumConstraints)
+		{
+			Header.NumConstraints = NumConstraints;
+		}
+
+		const FHeader& GetHeader() const
+		{
+			return Header;
+		}
+
+	private:
+
 		int32 CurrentReadElementIndex = 0;
 	};
 
diff --git a/Engine/Source/Runtime/Engine/Private/Components/SkeletalMeshComponent.cpp b/Engine/Source/Runtime/Engine/Private/Components/SkeletalMeshComponent.cpp
--- a/Engine/Source/Runtime/Engine/Private/Components/SkeletalMeshComponent.cpp
+++ b/Engine/Source/Runtime/Engine/Private/Components/SkeletalMeshComponent.cpp
@@ -3811,59 +3811,36 @@ void USkeletalMeshComponent::GetResourceSizeEx(FResourceSizeEx& CumulativeResour
 
 void USkeletalMeshComponent::FMigratedPhysicsState::ReadElementDataIntoBuffer(Chaos::FSerializedDataBuffer& TargetBuffer)
 {
-	const int32 SizeToCopy = GetDataSizeForElementIndex(CurrentReadElementIndex);
-	const int32 StartPos = GetDataOffset(CurrentReadElementIndex);
+	FDataEntryTag DataTag = Header.DataTagPerElementIndex.IsValidIndex(CurrentReadElementIndex) ? Header.DataTagPerElementIndex[CurrentReadElementIndex] : FDataEntryTag();
 
 	TArray<uint8>& AsByteArray = MigratedStateAsBytes.GetDataAsByteArrayRef();
 
-	if (ensure(StartPos + SizeToCopy < AsByteArray.Num()))
-	{
-		TargetBuffer.GetDataAsByteArrayRef().Append(&AsByteArray[StartPos], SizeToCopy);
-	}
-	else
-	{
-		UE_LOG(LogSkeletalMesh, Error, TEXT("[Physics State Migration] Attempted to read data out of bounds! | Buffer size [%d] | Start Pos [%d]| SizeToCopy [%d]"), AsByteArray.Num(), StartPos, SizeToCopy);
-	}
-	
-	CurrentReadElementIndex++;
-}
+	const int32 BufferSize = GetSize();
 
-int32 USkeletalMeshComponent::FMigratedPhysicsState::GetDataSizeForElementIndex(int32 ElementIndex) const
-{
-	if (!ensure(DataOffsetsPerElementIndex.IsValidIndex(ElementIndex)))
-	{
-		UE_LOG(LogSkeletalMesh, Error, TEXT("[Physics State Migration] Attempted to get the data size for an invalid element index | Requested Index [%d] | Max Element Index [%d] "), ElementIndex, DataOffsetsPerElementIndex.Num());
-		return 0;
-	}
+	UE_LOG(LogSkeletalMesh, Verbose, TEXT("[Physics State Migration] Attempted to read an empty data entry. This might means we failed to obtain data to serialize a specific body or constraint."));
 
-	const int32 NextIndex = ElementIndex + 1;
+	// If we have entries past the available buffer size, but they were recorded as empty, the serialization code is ok to continue (it will just be a no-op)
+	const bool bHasValidStartingOffset = DataTag.DataOffset < BufferSize || DataTag.DataSize == 0;
+	const bool bThereIsEnoughDataToCopyFromBuffer = DataTag.DataOffset + DataTag.DataSize <= BufferSize;
 
-	if (DataOffsetsPerElementIndex.IsValidIndex(NextIndex))
+	if (ensure(DataTag.IsValid() && bHasValidStartingOffset && bThereIsEnoughDataToCopyFromBuffer))
 	{
-		return DataOffsetsPerElementIndex[NextIndex] - DataOffsetsPerElementIndex[ElementIndex];
+		if (DataTag.DataSize > 0)
+		{
+			TargetBuffer.GetDataAsByteArrayRef().Append(&AsByteArray[DataTag.DataOffset], DataTag.DataSize);
+		}
 	}
-
-	if (ensure(DataOffsetsPerElementIndex.Num() == NextIndex))
+	else
 	{
-		// If this is the last element, the size of data to read is the remaining of the buffer from where we currently are at
-		const int32 TotalInnerBufferSizeMinusOne = GetSize() - 1;
-		return TotalInnerBufferSizeMinusOne - DataOffsetsPerElementIndex[ElementIndex];
+		UE_LOG(LogSkeletalMesh, Error, TEXT("[Physics State Migration] Attempted to read data out of bounds! | Buffer size [%d] | Start Pos [%d]| SizeToCopy [%d]"), AsByteArray.Num(), DataTag.DataOffset, DataTag.DataSize);
 	}
-
-	UE_LOG(LogSkeletalMesh, Error, TEXT("[Physics State Migration] Failed to get the data size for Index [%d] "), ElementIndex);
-	return 0;
-}
-
-int32 USkeletalMeshComponent::FMigratedPhysicsState::GetDataOffset(int32 ElementIndex) const
-{
-	return DataOffsetsPerElementIndex.IsValidIndex(ElementIndex) ? DataOffsetsPerElementIndex[ElementIndex] : 0;
+	
+	CurrentReadElementIndex++;
 }
 
 void USkeletalMeshComponent::FMigratedPhysicsState::Serialize(FArchive& Ar)
 {
-	Ar << NumBodies;
-	Ar << NumConstraints;
-	Ar << DataOffsetsPerElementIndex;
+	Ar << Header;
 	Ar << MigratedStateAsBytes;
 }
 
@@ -3905,8 +3882,8 @@ void USkeletalMeshComponent::MigratePhysicsState(FArchive& Ar)
 		if (HasValidPhysicsState())
 		{
 			FMigratedPhysicsState MigratedState;
-			MigratedState.NumBodies = Bodies.Num();
-			MigratedState.NumConstraints = Constraints.Num();
+			MigratedState.SetNumOfBodies(Bodies.Num());
+			MigratedState.SetNumOfConstraints(Constraints.Num());
 
 			for (FBodyInstance* PhysicsBody : Bodies)
 			{
@@ -3965,7 +3942,7 @@ void USkeletalMeshComponent::ApplyMigratedPhysicsState(Chaos::FSerializedDataBuf
 		FMigratedPhysicsState MigratedState;
 		MigratedState.Serialize(MultiBodyStateReader);
 
-		const bool bIsValidMigratedState = MigratedState.NumBodies == Bodies.Num() || MigratedState.NumConstraints == Constraints.Num();
+		const bool bIsValidMigratedState = MigratedState.GetHeader().IsValid() && MigratedState.GetHeader().NumBodies == Bodies.Num() || MigratedState.GetHeader().NumConstraints == Constraints.Num();
 		if (!ensure(bIsValidMigratedState))
 		{
 			return;
