World Partition - HLOD: Introduce an IWorldPartitionHLODObject interface
* HLODRuntime subsystem cell handling & warmup logic now rely solely on this interface and not actors

diff --git a/Engine/Source/Runtime/Engine/Private/UnrealEngine.cpp b/Engine/Source/Runtime/Engine/Private/UnrealEngine.cpp
--- a/Engine/Source/Runtime/Engine/Private/UnrealEngine.cpp
+++ b/Engine/Source/Runtime/Engine/Private/UnrealEngine.cpp
@@ -12795,11 +12795,11 @@ float DrawMapWarnings(UWorld* World, FViewport* Viewport, FCanvas* Canvas, UCanv
 
 	if (UWorldPartitionHLODRuntimeSubsystem* HLODSubsystem = World->GetSubsystem<UWorldPartitionHLODRuntimeSubsystem>())
 	{
-		uint32 NumOutdatedHLODActors = HLODSubsystem->GetNumOutdatedHLODActors();
-		if (NumOutdatedHLODActors)
+		uint32 NumOutdatedHLODObjects = HLODSubsystem->GetNumOutdatedHLODObjects();
+		if (NumOutdatedHLODObjects)
 		{
 			SmallTextItem.SetColor(FLinearColor::Yellow);
-			SmallTextItem.Text = FText::Format(LOCTEXT("OUTDATED_HLOD_ACTORS_FMT", "Outdated HLOD actors detected, HLODs need to be rebuilt ({0} outdated {0}|plural(one=actor,other=actors))"), NumOutdatedHLODActors);
+			SmallTextItem.Text = FText::Format(LOCTEXT("OUTDATED_HLOD_ACTORS_FMT", "Outdated HLODs detected, HLODs need to be rebuilt ({0} outdated)"), NumOutdatedHLODObjects);
 			Canvas->DrawItem(SmallTextItem, FVector2D(MessageX, MessageY));
 			MessageY += FontSizeY;
 		}
diff --git a/Engine/Source/Runtime/Engine/Private/WorldPartition/HLOD/HLODActor.cpp b/Engine/Source/Runtime/Engine/Private/WorldPartition/HLOD/HLODActor.cpp
--- a/Engine/Source/Runtime/Engine/Private/WorldPartition/HLOD/HLODActor.cpp
+++ b/Engine/Source/Runtime/Engine/Private/WorldPartition/HLOD/HLODActor.cpp
@@ -2,8 +2,10 @@
 
 #include "WorldPartition/HLOD/HLODActor.h"
 #include "Engine/World.h"
+#include "Engine/StaticMesh.h"
 #include "WorldPartition/HLOD/HLODRuntimeSubsystem.h"
 #include "Components/PrimitiveComponent.h"
+#include "Components/InstancedStaticMeshComponent.h"
 #include "Misc/PackageName.h"
 #include "UObject/UE5MainStreamObjectVersion.h"
 #include "UObject/UE5PrivateFrostyStreamObjectVersion.h"
@@ -66,6 +68,16 @@ AWorldPartitionHLOD::AWorldPartitionHLOD(const FObjectInitializer& ObjectInitial
 #endif
 }
 
+ULevel* AWorldPartitionHLOD::GetHLODLevel() const
+{
+	return GetLevel();
+}
+
+FString AWorldPartitionHLOD::GetHLODNameOrLabel() const
+{
+	return GetActorNameOrLabel();
+}
+
 const FGuid& AWorldPartitionHLOD::GetSourceCellGuid() const
 {
 	// When no source cell guid was set, try resolving it through its associated world partition runtime cell
@@ -92,15 +104,44 @@ void AWorldPartitionHLOD::SetVisibility(bool bInVisible)
 	});
 }
 
+TSet<UObject*> AWorldPartitionHLOD::GetAssetsToWarmup() const
+{
+	TSet<UObject*> AssetsToWarmup;
+	
+	ForEachComponent<UStaticMeshComponent>(false, [&](UStaticMeshComponent* SMC)
+	{
+		// Assume ISM HLOD don't need warmup, as they are actually found in the source level
+		if (SMC->IsA<UInstancedStaticMeshComponent>())
+		{
+			return;
+		}
+
+		for (int32 iMaterialIndex = 0; iMaterialIndex < SMC->GetNumMaterials(); ++iMaterialIndex)
+		{
+			if (UMaterialInterface* Material = SMC->GetMaterial(iMaterialIndex))
+			{
+				AssetsToWarmup.Add(Material);
+			}
+		}
+
+		if (UStaticMesh* StaticMesh = SMC->GetStaticMesh())
+		{
+			AssetsToWarmup.Add(StaticMesh);
+		}
+	});
+
+	return AssetsToWarmup;
+}
+
 void AWorldPartitionHLOD::BeginPlay()
 {
 	Super::BeginPlay();
-	GetWorld()->GetSubsystem<UWorldPartitionHLODRuntimeSubsystem>()->RegisterHLODActor(this);
+	GetWorld()->GetSubsystem<UWorldPartitionHLODRuntimeSubsystem>()->RegisterHLODObject(this);
 }
 
 void AWorldPartitionHLOD::EndPlay(const EEndPlayReason::Type EndPlayReason)
 {
-	GetWorld()->GetSubsystem<UWorldPartitionHLODRuntimeSubsystem>()->UnregisterHLODActor(this);
+	GetWorld()->GetSubsystem<UWorldPartitionHLODRuntimeSubsystem>()->UnregisterHLODObject(this);
 	Super::EndPlay(EndPlayReason);
 }
 
diff --git a/Engine/Source/Runtime/Engine/Private/WorldPartition/HLOD/HLODDestruction.cpp b/Engine/Source/Runtime/Engine/Private/WorldPartition/HLOD/HLODDestruction.cpp
--- a/Engine/Source/Runtime/Engine/Private/WorldPartition/HLOD/HLODDestruction.cpp
+++ b/Engine/Source/Runtime/Engine/Private/WorldPartition/HLOD/HLODDestruction.cpp
@@ -33,18 +33,22 @@ namespace
 
 		if (RuntimeCell)
 		{
-			for (AWorldPartitionHLOD* HLODActor : InActor->GetWorld()->GetSubsystem<UWorldPartitionHLODRuntimeSubsystem>()->GetHLODActorsForCell(RuntimeCell))
+			for (IWorldPartitionHLODObject* HLODObject : InActor->GetWorld()->GetSubsystem<UWorldPartitionHLODRuntimeSubsystem>()->GetHLODObjectsForCell(RuntimeCell))
 			{
-				for (UActorComponent* ActorComponent : HLODActor->GetComponents())
+				// This is the only use of IWorldPartitionHLODObject::GetUObject() for now... Once we solve destruction for FastGeo HLODs, we will be able to remove it.
+				if (AWorldPartitionHLOD* HLODActor = Cast<AWorldPartitionHLOD>(HLODObject->GetUObject()))
 				{
-					if (UWorldPartitionDestructibleHLODComponent* DestructibleHLODComponent = Cast<UWorldPartitionDestructibleHLODComponent>(ActorComponent))
+					for (UActorComponent* ActorComponent : HLODActor->GetComponents())
 					{
-						int32 ActorIdx = DestructibleHLODComponent->GetDestructibleActors().IndexOfByKey(ActorFName);
-						if (ActorIdx != INDEX_NONE)
+						if (UWorldPartitionDestructibleHLODComponent* DestructibleHLODComponent = Cast<UWorldPartitionDestructibleHLODComponent>(ActorComponent))
 						{
-							HLODDestructionTag.ActorIndex = ActorIdx;
-							HLODDestructionTag.HLODDestructionComponent = DestructibleHLODComponent;
-							break;
+							int32 ActorIdx = DestructibleHLODComponent->GetDestructibleActors().IndexOfByKey(ActorFName);
+							if (ActorIdx != INDEX_NONE)
+							{
+								HLODDestructionTag.ActorIndex = ActorIdx;
+								HLODDestructionTag.HLODDestructionComponent = DestructibleHLODComponent;
+								break;
+							}
 						}
 					}
 				}
diff --git a/Engine/Source/Runtime/Engine/Private/WorldPartition/HLOD/HLODRuntimeSubsystem.cpp b/Engine/Source/Runtime/Engine/Private/WorldPartition/HLOD/HLODRuntimeSubsystem.cpp
--- a/Engine/Source/Runtime/Engine/Private/WorldPartition/HLOD/HLODRuntimeSubsystem.cpp
+++ b/Engine/Source/Runtime/Engine/Private/WorldPartition/HLOD/HLODRuntimeSubsystem.cpp
@@ -82,21 +82,21 @@ static FAutoConsoleVariableSink CVarHLODSink(FConsoleCommandDelegate::CreateStat
 
 namespace FHLODRuntimeSubsystem
 {
-    static const UWorldPartitionRuntimeCell* GetActorRuntimeCell(AActor* InActor)
+    static const UWorldPartitionRuntimeCell* GetRuntimeCell(IWorldPartitionHLODObject* InWorldPartitionHLOD)
     {
-	    const ULevel* ActorLevel = InActor->GetLevel();
-	    const ULevelStreaming* LevelStreaming = FLevelUtils::FindStreamingLevel(ActorLevel);
+	    const ULevel* Level = InWorldPartitionHLOD->GetHLODLevel();
+	    const ULevelStreaming* LevelStreaming = FLevelUtils::FindStreamingLevel(Level);
 	    return LevelStreaming ? Cast<const UWorldPartitionRuntimeCell>(LevelStreaming->GetWorldPartitionCell()) : nullptr;
     }
     
-    static UWorldPartition* GetWorldPartition(AWorldPartitionHLOD* InWorldPartitionHLOD)
+    static UWorldPartition* GetWorldPartition(IWorldPartitionHLODObject* InWorldPartitionHLOD)
     {
 	    // Alwaysloaded Cell level will have a WorldPartition
-	    if (UWorldPartition* WorldPartition = InWorldPartitionHLOD->GetLevel()->GetWorldPartition())
+	    if (UWorldPartition* WorldPartition = InWorldPartitionHLOD->GetHLODLevel()->GetWorldPartition())
 	    {
 		    return WorldPartition;
 	    } // If not find it through the cell
-	    else if (const UWorldPartitionRuntimeCell* RuntimeCell = GetActorRuntimeCell(InWorldPartitionHLOD))
+	    else if (const UWorldPartitionRuntimeCell* RuntimeCell = GetRuntimeCell(InWorldPartitionHLOD))
 	    {
 		    return RuntimeCell->GetOuterWorld()->GetWorldPartition();
 	    }
@@ -150,9 +150,9 @@ FAutoConsoleCommand UWorldPartitionHLODRuntimeSubsystem::EnableHLODCommand(
 						{
 							const FCellData& CellData = CellHLODMapping.Value;
 							bool bIsHLODVisible = UWorldPartitionHLODRuntimeSubsystem::WorldPartitionHLODEnabled && !CellData.bIsCellVisible;
-							for (AWorldPartitionHLOD* HLODActor : CellData.LoadedHLODs)
+							for (IWorldPartitionHLODObject* HLODObject : CellData.LoadedHLODs)
 							{
-								HLODActor->SetVisibility(bIsHLODVisible);
+								HLODObject->SetVisibility(bIsHLODVisible);
 							}
 						}
 					}
@@ -290,7 +290,7 @@ UWorldPartitionHLODRuntimeSubsystem::FCellData* UWorldPartitionHLODRuntimeSubsys
 	return nullptr;
 }
 
-UWorldPartitionHLODRuntimeSubsystem::FCellData* UWorldPartitionHLODRuntimeSubsystem::GetCellData(AWorldPartitionHLOD* InWorldPartitionHLOD)
+UWorldPartitionHLODRuntimeSubsystem::FCellData* UWorldPartitionHLODRuntimeSubsystem::GetCellData(IWorldPartitionHLODObject* InWorldPartitionHLOD)
 {
 	const UWorldPartition* WorldPartition = FHLODRuntimeSubsystem::GetWorldPartition(InWorldPartitionHLOD);
 	if (WorldPartition)
@@ -307,7 +307,7 @@ UWorldPartitionHLODRuntimeSubsystem::FCellData* UWorldPartitionHLODRuntimeSubsys
 	return nullptr;		
 }
 
-const TArray<AWorldPartitionHLOD*>& UWorldPartitionHLODRuntimeSubsystem::GetHLODActorsForCell(const UWorldPartitionRuntimeCell* InCell) const
+const TArray<IWorldPartitionHLODObject*>& UWorldPartitionHLODRuntimeSubsystem::GetHLODObjectsForCell(const UWorldPartitionRuntimeCell* InCell) const
 {
 	if (const FCellData* CellData = GetCellData(InCell))
 	{
@@ -315,53 +315,55 @@ const TArray<AWorldPartitionHLOD*>& UWorldPartitionHLODRuntimeSubsystem::GetHLOD
 	}
 
 	// No HLOD found for the given cell, return a dummy array
-	static const TArray<AWorldPartitionHLOD*> DummyArray;
+	static const TArray<IWorldPartitionHLODObject*> DummyArray;
 	return DummyArray;
 }
 
-void UWorldPartitionHLODRuntimeSubsystem::RegisterHLODActor(AWorldPartitionHLOD* InWorldPartitionHLOD)
+void UWorldPartitionHLODRuntimeSubsystem::RegisterHLODObject(IWorldPartitionHLODObject* InWorldPartitionHLOD)
 {
-	TRACE_CPUPROFILER_EVENT_SCOPE(UWorldPartitionHLODRuntimeSubsystem::RegisterHLODActor);
+	TRACE_CPUPROFILER_EVENT_SCOPE(UWorldPartitionHLODRuntimeSubsystem::RegisterHLODObject);
 	
-	if (FCellData* CellData = GetCellData(InWorldPartitionHLOD))
+	if (FCellData* SourceCellData = GetCellData(InWorldPartitionHLOD))
 	{
-		UE_LOG(LogHLODRuntimeSubsystem, Verbose, TEXT("Registering HLOD %s for cell %s"), *InWorldPartitionHLOD->GetActorNameOrLabel(), *InWorldPartitionHLOD->GetSourceCellGuid().ToString());
+		UE_LOG(LogHLODRuntimeSubsystem, Verbose, TEXT("Registering HLOD %s for cell %s"), *InWorldPartitionHLOD->GetHLODNameOrLabel(), *InWorldPartitionHLOD->GetSourceCellGuid().ToString());
 
-		CellData->LoadedHLODs.Add(InWorldPartitionHLOD);
-		InWorldPartitionHLOD->SetVisibility(UWorldPartitionHLODRuntimeSubsystem::WorldPartitionHLODEnabled && !CellData->bIsCellVisible);
+		// Register this HLOD as a representation of the source cell - Needed to warmup an HLOD object when it's source cell is unloading
+		SourceCellData->LoadedHLODs.Add(InWorldPartitionHLOD);
+		InWorldPartitionHLOD->SetVisibility(UWorldPartitionHLODRuntimeSubsystem::WorldPartitionHLODEnabled && !SourceCellData->bIsCellVisible);
 	}
 	else
 	{
-		UE_LOG(LogHLODRuntimeSubsystem, Verbose, TEXT("Found HLOD %s referencing nonexistent cell '%s'"), *InWorldPartitionHLOD->GetActorNameOrLabel(), *InWorldPartitionHLOD->GetSourceCellGuid().ToString());
+		UE_LOG(LogHLODRuntimeSubsystem, Verbose, TEXT("Found HLOD %s referencing nonexistent cell '%s'"), *InWorldPartitionHLOD->GetHLODNameOrLabel(), *InWorldPartitionHLOD->GetSourceCellGuid().ToString());
 		InWorldPartitionHLOD->SetVisibility(false);
 
 #if !UE_BUILD_SHIPPING
-		OutdatedHLODActors.Add(InWorldPartitionHLOD);
+		OutdatedHLODObjects.Add(InWorldPartitionHLOD);
 #endif
 	}
 
-	HLODActorsToWarmup.Remove(InWorldPartitionHLOD);
-	HLODActorRegisteredEvent.Broadcast(InWorldPartitionHLOD);
+	HLODObjectsToWarmup.Remove(InWorldPartitionHLOD);
+	HLODObjectRegisteredEvent.Broadcast(InWorldPartitionHLOD);
 }
 
-void UWorldPartitionHLODRuntimeSubsystem::UnregisterHLODActor(AWorldPartitionHLOD* InWorldPartitionHLOD)
+void UWorldPartitionHLODRuntimeSubsystem::UnregisterHLODObject(IWorldPartitionHLODObject* InWorldPartitionHLOD)
 {
-	TRACE_CPUPROFILER_EVENT_SCOPE(UWorldPartitionHLODRuntimeSubsystem::UnregisterHLODActor);
+	TRACE_CPUPROFILER_EVENT_SCOPE(UWorldPartitionHLODRuntimeSubsystem::UnregisterHLODObject);
 
 	if (FCellData* CellData = GetCellData(InWorldPartitionHLOD))
 	{
-		UE_LOG(LogHLODRuntimeSubsystem, Verbose, TEXT("Unregistering HLOD %s for cell %s"), *InWorldPartitionHLOD->GetActorNameOrLabel(), *InWorldPartitionHLOD->GetSourceCellGuid().ToString());
+		UE_LOG(LogHLODRuntimeSubsystem, Verbose, TEXT("Unregistering HLOD %s for cell %s"), *InWorldPartitionHLOD->GetHLODNameOrLabel(), *InWorldPartitionHLOD->GetSourceCellGuid().ToString());
 
 		verify(CellData->LoadedHLODs.Remove(InWorldPartitionHLOD));
 	}
 	else
 	{
 #if !UE_BUILD_SHIPPING
-		OutdatedHLODActors.Remove(InWorldPartitionHLOD);
+		OutdatedHLODObjects.Remove(InWorldPartitionHLOD);
 #endif
 	}
 
-	HLODActorUnregisteredEvent.Broadcast(InWorldPartitionHLOD);
+	HLODObjectsToWarmup.Remove(InWorldPartitionHLOD);
+	HLODObjectUnregisteredEvent.Broadcast(InWorldPartitionHLOD);
 }
 
 void UWorldPartitionHLODRuntimeSubsystem::OnCellShown(const UWorldPartitionRuntimeCell* InCell)
@@ -372,12 +374,12 @@ void UWorldPartitionHLODRuntimeSubsystem::OnCellShown(const UWorldPartitionRunti
 
 		if (!CellData->LoadedHLODs.IsEmpty())
 		{
-			UE_LOG(LogHLODRuntimeSubsystem, Verbose, TEXT("Cell shown - %s - hiding %d HLOD actors"), *InCell->GetGuid().ToString(), CellData->LoadedHLODs.Num());
+			UE_LOG(LogHLODRuntimeSubsystem, Verbose, TEXT("Cell shown - %s - hiding %d HLOD objects"), *InCell->GetGuid().ToString(), CellData->LoadedHLODs.Num());
 
-			for (AWorldPartitionHLOD* HLODActor : CellData->LoadedHLODs)
+			for (IWorldPartitionHLODObject* HLODObject : CellData->LoadedHLODs)
 			{
-				UE_LOG(LogHLODRuntimeSubsystem, Verbose, TEXT("\t\t* %s"), *HLODActor->GetActorNameOrLabel());
-				HLODActor->SetVisibility(false);
+				UE_LOG(LogHLODRuntimeSubsystem, Verbose, TEXT("\t\t* %s"), *HLODObject->GetHLODNameOrLabel());
+				HLODObject->SetVisibility(false);
 			}
 		}
 	}
@@ -391,20 +393,22 @@ void UWorldPartitionHLODRuntimeSubsystem::OnCellHidden(const UWorldPartitionRunt
 
 		if (!CellData->LoadedHLODs.IsEmpty())
 		{
-			UE_LOG(LogHLODRuntimeSubsystem, Verbose, TEXT("Cell hidden - %s - showing %d HLOD actors"), *InCell->GetGuid().ToString(), CellData->LoadedHLODs.Num());
+			UE_LOG(LogHLODRuntimeSubsystem, Verbose, TEXT("Cell hidden - %s - showing %d HLOD objects"), *InCell->GetGuid().ToString(), CellData->LoadedHLODs.Num());
 
-			for (AWorldPartitionHLOD* HLODActor : CellData->LoadedHLODs)
+			for (IWorldPartitionHLODObject* HLODObject : CellData->LoadedHLODs)
 			{
-				UE_LOG(LogHLODRuntimeSubsystem, Verbose, TEXT("\t\t* %s"), *HLODActor->GetActorNameOrLabel());
-				HLODActor->SetVisibility(UWorldPartitionHLODRuntimeSubsystem::WorldPartitionHLODEnabled);
-				HLODActorsToWarmup.Remove(HLODActor);
+				UE_LOG(LogHLODRuntimeSubsystem, Verbose, TEXT("\t\t* %s"), *HLODObject->GetHLODNameOrLabel());
+				HLODObject->SetVisibility(UWorldPartitionHLODRuntimeSubsystem::WorldPartitionHLODEnabled);
+				HLODObjectsToWarmup.Remove(HLODObject);
 			}
 		}
 	}
 }
 
-static void PrepareVTRequests(TMap<UMaterialInterface*, float>& InOutVTRequests, UStaticMeshComponent* InStaticMeshComponent, float InPixelSize)
+static void PrepareVTRequests(TMap<UMaterialInterface*, float>& InOutVTRequests, UMaterialInterface* InMaterialInterface, float InPixelSize)
 {
+	check(InMaterialInterface);
+
 	float PixelSize = InPixelSize;
 
 	// Assume the texture is wrapped around the object, so the screen size is actually less than the resolution we require.
@@ -413,37 +417,32 @@ static void PrepareVTRequests(TMap<UMaterialInterface*, float>& InOutVTRequests,
 	// Clamp for safety
 	PixelSize = FMath::Min(PixelSize, CVarHLODWarmupVTSizeClamp.GetValueOnAnyThread());
 
-	for (UMaterialInterface* MaterialInterface : InStaticMeshComponent->GetMaterials())
-	{
-		if (MaterialInterface)
-		{
-			// We have a VT we'd like to prefetch, add or update a request in our request map.
-			// If the texture was already requested by another component, fetch the highest required resolution only.
-			float& CurrentMaxPixel = InOutVTRequests.FindOrAdd(MaterialInterface);
-			CurrentMaxPixel = FMath::Max(CurrentMaxPixel, PixelSize);
-		}
-	}
+	// We have a VT we'd like to prefetch, add or update a request in our request map.
+	// If the texture was already requested by another component, fetch the highest required resolution only.
+	float& CurrentMaxPixel = InOutVTRequests.FindOrAdd(InMaterialInterface);
+	CurrentMaxPixel = FMath::Max(CurrentMaxPixel, PixelSize);
 }
 
-static void PrepareNaniteRequests(TMap<Nanite::FResources*, int32>& InOutNaniteRequests, UStaticMeshComponent* InStaticMeshComponent, int32 InNumFramesUntilRender)
+static void PrepareNaniteRequests(TMap<Nanite::FResources*, int32>& InOutNaniteRequests, UStaticMesh* InStaticMesh, int32 InNumFramesUntilRender)
 {
-	UStaticMesh* StaticMesh = InStaticMeshComponent->GetStaticMesh();
-	if (StaticMesh && StaticMesh->HasValidNaniteData())
+	check(InStaticMesh);
+
+	if (InStaticMesh->HasValidNaniteData())
 	{
-		int32& NumFramesUntilRender = InOutNaniteRequests.FindOrAdd(StaticMesh->GetRenderData()->NaniteResourcesPtr.Get());
+		int32& NumFramesUntilRender = InOutNaniteRequests.FindOrAdd(InStaticMesh->GetRenderData()->NaniteResourcesPtr.Get());
 		NumFramesUntilRender = FMath::Max(InNumFramesUntilRender, 1);
 	}
 }
 
-bool UWorldPartitionHLODRuntimeSubsystem::PrepareToWarmup(const UWorldPartitionRuntimeCell* InCell, AWorldPartitionHLOD* InHLODActor)
+bool UWorldPartitionHLODRuntimeSubsystem::PrepareToWarmup(const UWorldPartitionRuntimeCell* InCell, const IWorldPartitionHLODObject* InHLODObject)
 {
 	TRACE_CPUPROFILER_EVENT_SCOPE(UWorldPartitionHLODRuntimeSubsystem::PrepareToWarmup)
 
-	bool bHLODActorNeedsWarmUp = false;
+	bool bHLODObjectNeedsWarmUp = false;
 
-	if (InHLODActor->DoesRequireWarmup())
+	if (InHLODObject->DoesRequireWarmup())
 	{
-		FWorldPartitionHLODWarmupState& WarmupState = HLODActorsToWarmup.FindOrAdd(InHLODActor);
+		FWorldPartitionHLODWarmupState& WarmupState = HLODObjectsToWarmup.FindOrAdd(InHLODObject);
 
 		// Trigger warmup for CVarHLODWarmupNumFrames frames on the first request, or if a warmup wasn't requested in the last frame
 		const bool bResetWarmup =  WarmupState.WarmupLastRequestedFrame == INDEX_NONE ||
@@ -467,11 +466,11 @@ bool UWorldPartitionHLODRuntimeSubsystem::PrepareToWarmup(const UWorldPartitionR
 			WarmupState.WarmupCallsUntilReady--;
 		}
 		
-		bHLODActorNeedsWarmUp = WarmupState.WarmupCallsUntilReady != 0;
+		bHLODObjectNeedsWarmUp = WarmupState.WarmupCallsUntilReady != 0;
 		WarmupState.WarmupLastRequestedFrame = GFrameNumber;
 	}
 
-	return bHLODActorNeedsWarmUp;
+	return bHLODObjectNeedsWarmUp;
 }
 
 void UWorldPartitionHLODRuntimeSubsystem::OnCVarsChanged()
@@ -554,17 +553,25 @@ bool UWorldPartitionHLODRuntimeSubsystem::CanMakeVisible(const UWorldPartitionRu
 
 	bool bCanMakeVisible = true;
 
-	// Prevent cells containing HLODs actors from being made visible until warmup has been performed
+	// Prevent cells containing HLODs objects from being made visible until warmup has been performed
 	if (InCell->GetIsHLOD() && InCell->GetLevel())
 	{
+		TFunction<void(const IWorldPartitionHLODObject*)> PrepareForWarmup =
+			[&bCanMakeVisible, InCell, this](const IWorldPartitionHLODObject* HLODObject)
+			{
+				const bool bHLODObjectNeedsWarmup = PrepareToWarmup(InCell, HLODObject);
+				bCanMakeVisible &= !bHLODObjectNeedsWarmup;
+			};
+		
 		for (AActor* Actor : InCell->GetLevel()->Actors)
 		{
 			if (AWorldPartitionHLOD* HLODActor = Cast<AWorldPartitionHLOD>(Actor))
 			{
-				const bool bHLODActorNeedsWarmup = PrepareToWarmup(InCell, HLODActor);
-				bCanMakeVisible &= !bHLODActorNeedsWarmup;
+				PrepareForWarmup(HLODActor);
 			}
 		}
+
+		ForEachHLODObjectInCellEvent.Broadcast(InCell, PrepareForWarmup);
 	}
 
 	return bCanMakeVisible;
@@ -581,19 +588,24 @@ bool UWorldPartitionHLODRuntimeSubsystem::CanMakeInvisible(const UWorldPartition
 
 	bool bCanMakeInvisible = true;
 
-	// Prevent cells which have an HLOD representation from being hidden until their matching HLOD actors have been warmed up
+	// Prevent cells from being hidden until their matching HLOD objects have been warmed up
 	if (FCellData* CellData = GetCellData(InCell))
 	{
-		for (AWorldPartitionHLOD* HLODActor : CellData->LoadedHLODs)
+		for (IWorldPartitionHLODObject* HLODObject : CellData->LoadedHLODs)
 		{
-			const bool bHLODActorNeedsWarmup = PrepareToWarmup(InCell, HLODActor);
-			bCanMakeInvisible &= !bHLODActorNeedsWarmup;
+			const bool bHLODObjectNeedsWarmup = PrepareToWarmup(InCell, HLODObject);
+			bCanMakeInvisible &= !bHLODObjectNeedsWarmup;
 		}
 	}
 
 	return bCanMakeInvisible;
 }
 
+FWorldPartitionHLODForEachHLODObjectInCellEvent& UWorldPartitionHLODRuntimeSubsystem::GetForEachHLODObjectInCellEvent()
+{
+	return ForEachHLODObjectInCellEvent;
+}
+
 static float IsInView(const FVector& BoundsOrigin, const FVector& BoundsExtent, const FSceneViewFamily& InViewFamily, bool bComputeScreenSize, float& MaxScreenSizePixels)
 {
 	MaxScreenSizePixels = 0;
@@ -601,7 +613,7 @@ static float IsInView(const FVector& BoundsOrigin, const FVector& BoundsExtent,
 	// Estimate the highest screen pixel size of this component in the provided views
 	for (const FSceneView* View : InViewFamily.Views)
 	{
-		// Make sure the HLOD actor we're about to show is actually in the frustum
+		// Make sure the HLOD Object we're about to show is actually in the frustum
 		if (View->ViewFrustum.IntersectBox(BoundsOrigin, BoundsExtent))
 		{
 			if (bComputeScreenSize)
@@ -657,63 +669,53 @@ void UWorldPartitionHLODRuntimeSubsystem::OnBeginRenderViews(const FSceneViewFam
 	const bool bWarmupNanite = CVarHLODWarmupNanite.GetValueOnGameThread() != 0;
 	const bool bWarmupVT = CVarHLODWarmupVT.GetValueOnGameThread() != 0;
 
-	for (FHLODWarmupStateMap::TIterator HLODActorWarmupStateIt(HLODActorsToWarmup); HLODActorWarmupStateIt; ++HLODActorWarmupStateIt)
+	for (FHLODWarmupStateMap::TIterator HLODObjectWarmupStateIt(HLODObjectsToWarmup); HLODObjectWarmupStateIt; ++HLODObjectWarmupStateIt)
 	{
-		const FObjectKey& HLODActorObjectKey = HLODActorWarmupStateIt.Key();
-		FWorldPartitionHLODWarmupState& HLODWarmupState = HLODActorWarmupStateIt.Value();
-		const AWorldPartitionHLOD* HLODActor = Cast<AWorldPartitionHLOD>(HLODActorObjectKey.ResolveObjectPtr());
-		
-		if (HLODActor)
-		{
-			// Retrieve this component's bound - we must support getting the bounds before the component is even registered.
-			FVector BoundsOrigin;
-			FVector BoundsExtent;
-			HLODWarmupState.WarmupBounds.GetCenterAndExtents(BoundsOrigin, BoundsExtent);
+		const IWorldPartitionHLODObject* HLODObject = HLODObjectWarmupStateIt.Key();
+		FWorldPartitionHLODWarmupState& HLODWarmupState = HLODObjectWarmupStateIt.Value();
 
-			float ScreenSizePixels = 0;
-			if (IsInView(BoundsOrigin, BoundsExtent, InViewFamily, bWarmupVT, ScreenSizePixels))
+		// Retrieve this component's bound - we must support getting the bounds before the component is even registered.
+		FVector BoundsOrigin;
+		FVector BoundsExtent;
+		HLODWarmupState.WarmupBounds.GetCenterAndExtents(BoundsOrigin, BoundsExtent);
+
+		float ScreenSizePixels = 0;
+		if (IsInView(BoundsOrigin, BoundsExtent, InViewFamily, bWarmupVT, ScreenSizePixels))
+		{
+			for (UObject* AssetToWarmup : HLODObject->GetAssetsToWarmup())
 			{
-				HLODActor->ForEachComponent<UStaticMeshComponent>(false, [&](UStaticMeshComponent* SMC)
+				if (UMaterialInterface* MaterialInterface = Cast<UMaterialInterface>(AssetToWarmup))
 				{
-					// Assume ISM HLOD don't need warmup, as they are actually found in the source level
-					if (SMC->IsA<UInstancedStaticMeshComponent>())
-					{
-						return;
-					}
-				
-					if (bWarmupVT)
-					{
-						PrepareVTRequests(VTRequests, SMC, ScreenSizePixels);
-					}
-
-					if (bWarmupNanite)
-					{
-						if (HLODWarmupState.WarmupCallsUntilReady == CVarHLODWarmupNumFrames.GetValueOnGameThread())
-						{
-							// Send a nanite request to prepare for visibility in CVarHLODWarmupNumFrames frames
-							PrepareNaniteRequests(NaniteRequests, SMC, CVarHLODWarmupNumFrames.GetValueOnGameThread());
-						}
-						else if (HLODWarmupState.WarmupCallsUntilReady == 0)
-						{
-							// We expect HLOD to be visible at any moment (likely waiting for server visibility ack)
-							PrepareNaniteRequests(NaniteRequests, SMC, 1);
-						}
-					}
-				});
+				    if (bWarmupVT)
+				    {
+					    PrepareVTRequests(VTRequests, MaterialInterface, ScreenSizePixels);
+				    }
+				}
 
-#if ENABLE_DRAW_DEBUG
-				if (CVarHLODWarmupDebugDraw.GetValueOnAnyThread())
+				if (UStaticMesh* StaticMesh = Cast<UStaticMesh>(AssetToWarmup))
 				{
-					DrawDebugBox(HLODActor->GetWorld(), BoundsOrigin, BoundsExtent, FColor::Yellow, /*bPersistentLine*/ false, /*Lifetime*/ 1.0f);
+				    if (bWarmupNanite)
+				    {
+					    if (HLODWarmupState.WarmupCallsUntilReady == CVarHLODWarmupNumFrames.GetValueOnGameThread())
+					    {
+						    // Send a nanite request to prepare for visibility in CVarHLODWarmupNumFrames frames
+						    PrepareNaniteRequests(NaniteRequests, StaticMesh, CVarHLODWarmupNumFrames.GetValueOnGameThread());
+					    }
+					    else if (HLODWarmupState.WarmupCallsUntilReady == 0)
+					    {
+						    // We expect HLOD to be visible at any moment (likely waiting for server visibility ack)
+						    PrepareNaniteRequests(NaniteRequests, StaticMesh, 1);
+					    }
+				    }
 				}
-#endif
 			}
-		}
 
-		if (!HLODActor)
-		{
-			// HLOD actor has been unloaded, stop tracking it's warmup state
-			HLODActorWarmupStateIt.RemoveCurrent();
+#if ENABLE_DRAW_DEBUG
+			if (CVarHLODWarmupDebugDraw.GetValueOnAnyThread())
+			{
+				DrawDebugBox(HLODObject->GetHLODLevel()->GetWorld(), BoundsOrigin, BoundsExtent, HLODWarmupState.WarmupCallsUntilReady > 0 ? FColor::Yellow : FColor::Green, /*bPersistentLine*/ false, /*Lifetime*/ 0.5f);
+			}
+#endif
 		}
 	}
 
diff --git a/Engine/Source/Runtime/Engine/Public/WorldPartition/HLOD/HLODActor.h b/Engine/Source/Runtime/Engine/Public/WorldPartition/HLOD/HLODActor.h
--- a/Engine/Source/Runtime/Engine/Public/WorldPartition/HLOD/HLODActor.h
+++ b/Engine/Source/Runtime/Engine/Public/WorldPartition/HLOD/HLODActor.h
@@ -9,6 +9,7 @@
 #include "WorldPartition/WorldPartitionActorDesc.h"
 #include "WorldPartition/HLOD/HLODStats.h"
 #include "WorldPartition/HLOD/HLODBuilder.h"
+#include "WorldPartition/HLOD/IWorldPartitionHLODObject.h"
 
 #if WITH_EDITOR
 #include "UObject/ObjectSaveContext.h"
@@ -22,7 +23,7 @@ class UWorldPartitionHLODSourceActors;
 ENGINE_API DECLARE_LOG_CATEGORY_EXTERN(LogHLODHash, Log, All);
 
 UCLASS(NotPlaceable, MinimalAPI, HideCategories=(Rendering, Replication, Collision, Physics, Navigation, Networking, Input, Actor, LevelInstance, Cooking))
-class AWorldPartitionHLOD : public AActor
+class AWorldPartitionHLOD : public AActor, public IWorldPartitionHLODObject
 {
 	GENERATED_UCLASS_BODY()
 
@@ -30,14 +31,19 @@ class AWorldPartitionHLOD : public AActor
 	typedef TMap<FName, int64> FStats;
 
 public:
-	ENGINE_API void SetVisibility(bool bInVisible);
-
-	ENGINE_API const FGuid& GetSourceCellGuid() const;
 	inline uint32 GetLODLevel() const { return LODLevel; }
 
 	virtual bool IsHLODRelevant() const override { return true; }
 
-	bool DoesRequireWarmup() const { return bRequireWarmup; }
+	// ~Begin IWorldPartitionHLODObject interface
+	ENGINE_API virtual UObject* GetUObject() const override { return const_cast<AWorldPartitionHLOD*>(this); }
+	ENGINE_API virtual ULevel* GetHLODLevel() const override;
+	ENGINE_API virtual FString GetHLODNameOrLabel() const override;
+	ENGINE_API virtual bool DoesRequireWarmup() const override { return bRequireWarmup; }
+	ENGINE_API virtual TSet<UObject*> GetAssetsToWarmup() const override;
+	ENGINE_API virtual void SetVisibility(bool bIsVisible) override;
+	ENGINE_API virtual const FGuid& GetSourceCellGuid() const override;
+	// ~End IWorldPartitionHLODObject interface
 
 #if WITH_EDITOR
 	ENGINE_API void SetHLODComponents(const TArray<UActorComponent*>& InHLODComponents);
diff --git a/Engine/Source/Runtime/Engine/Public/WorldPartition/HLOD/HLODRuntimeSubsystem.h b/Engine/Source/Runtime/Engine/Public/WorldPartition/HLOD/HLODRuntimeSubsystem.h
--- a/Engine/Source/Runtime/Engine/Public/WorldPartition/HLOD/HLODRuntimeSubsystem.h
+++ b/Engine/Source/Runtime/Engine/Public/WorldPartition/HLOD/HLODRuntimeSubsystem.h
@@ -6,9 +6,10 @@
 #include "Containers/Map.h"
 #include "UObject/ObjectKey.h"
 #include "Containers/Ticker.h"
+#include "Templates/Function.h"
 #include "HLODRuntimeSubsystem.generated.h"
 
-
+class IWorldPartitionHLODObject;
 class AWorldPartitionHLOD;
 class FSceneViewFamily;
 class FHLODResourcesResidencySceneViewExtension;
@@ -17,10 +18,15 @@ class UWorldPartitionRuntimeCell;
 class URuntimeHashExternalStreamingObjectBase;
 class UWorld;
 
+DECLARE_MULTICAST_DELEGATE_OneParam(FWorldPartitionHLODObjectRegisteredEvent, IWorldPartitionHLODObject* /* InHLODObject */);
+DECLARE_MULTICAST_DELEGATE_OneParam(FWorldPartitionHLODObjectUnregisteredEvent, IWorldPartitionHLODObject* /* InHLODObject */);
+DECLARE_MULTICAST_DELEGATE_TwoParams(FWorldPartitionHLODForEachHLODObjectInCellEvent, const UWorldPartitionRuntimeCell* /* InCell */, TFunction<void(const IWorldPartitionHLODObject*)> /* InFunc */);
+
+UE_DEPRECATED(5.6, "Use FWorldPartitionHLODObjectRegisteredEvent instead")
 DECLARE_MULTICAST_DELEGATE_OneParam(FWorldPartitionHLODActorRegisteredEvent, AWorldPartitionHLOD* /* InHLODActor */);
+UE_DEPRECATED(5.6, "Use FWorldPartitionHLODObjectUnregisteredEvent instead")
 DECLARE_MULTICAST_DELEGATE_OneParam(FWorldPartitionHLODActorUnregisteredEvent, AWorldPartitionHLOD* /* InHLODActor */);
 
-
 /**
  * UWorldPartitionHLODRuntimeSubsystem
  */
@@ -41,8 +47,8 @@ public:
 	ENGINE_API virtual bool DoesSupportWorldType(const EWorldType::Type WorldType) const override;
 	//~ End UWorldSubsystem Interface.
 
-	ENGINE_API void RegisterHLODActor(AWorldPartitionHLOD* InWorldPartitionHLOD);
-	ENGINE_API void UnregisterHLODActor(AWorldPartitionHLOD* InWorldPartitionHLOD);
+	ENGINE_API void RegisterHLODObject(IWorldPartitionHLODObject* InWorldPartitionHLOD);
+	ENGINE_API void UnregisterHLODObject(IWorldPartitionHLODObject* InWorldPartitionHLOD);
 
 	ENGINE_API void OnCellShown(const UWorldPartitionRuntimeCell* InCell);
 	ENGINE_API void OnCellHidden(const UWorldPartitionRuntimeCell* InCell);
@@ -50,12 +56,13 @@ public:
 	ENGINE_API bool CanMakeVisible(const UWorldPartitionRuntimeCell* InCell);
 	ENGINE_API bool CanMakeInvisible(const UWorldPartitionRuntimeCell* InCell);
 
-	ENGINE_API const TArray<AWorldPartitionHLOD*>& GetHLODActorsForCell(const UWorldPartitionRuntimeCell* InCell) const;
+	ENGINE_API FWorldPartitionHLODForEachHLODObjectInCellEvent& GetForEachHLODObjectInCellEvent();
 
 	static ENGINE_API bool IsHLODEnabled();
 
-	FWorldPartitionHLODActorRegisteredEvent& OnHLODActorRegisteredEvent() { return HLODActorRegisteredEvent; }
-	FWorldPartitionHLODActorUnregisteredEvent& OnHLODActorUnregisteredEvent() { return HLODActorUnregisteredEvent; }
+	ENGINE_API const TArray<IWorldPartitionHLODObject*>& GetHLODObjectsForCell(const UWorldPartitionRuntimeCell* InCell) const;
+	FWorldPartitionHLODObjectRegisteredEvent& OnHLODObjectRegisteredEvent() { return HLODObjectRegisteredEvent; }
+	FWorldPartitionHLODObjectUnregisteredEvent& OnHLODObjectUnregisteredEvent() { return HLODObjectUnregisteredEvent; }
 
 	ENGINE_API void OnExternalStreamingObjectInjected(URuntimeHashExternalStreamingObjectBase* ExternalStreamingObject);
 	ENGINE_API void OnExternalStreamingObjectRemoved(URuntimeHashExternalStreamingObjectBase* ExternalStreamingObject);
@@ -63,9 +70,33 @@ public:
 	ENGINE_API void OnCVarsChanged();
 
 #if !UE_BUILD_SHIPPING
-	uint32 GetNumOutdatedHLODActors() const { return OutdatedHLODActors.Num(); }
+	uint32 GetNumOutdatedHLODObjects() const { return OutdatedHLODObjects.Num(); }
 #endif
 
+	UE_DEPRECATED(5.6, "Use RegisterHLODObject instead")
+	ENGINE_API void RegisterHLODActor(AWorldPartitionHLOD* InWorldPartitionHLOD) {}
+	UE_DEPRECATED(5.6, "Use UnregisterHLODObject instead")
+	ENGINE_API void UnregisterHLODActor(AWorldPartitionHLOD* InWorldPartitionHLOD) {}
+
+	UE_DEPRECATED(5.6, "Use GetHLODObjectsForCell instead")
+	ENGINE_API const TArray<AWorldPartitionHLOD*>& GetHLODActorsForCell(const UWorldPartitionRuntimeCell* InCell) const
+	{
+		static TArray<AWorldPartitionHLOD*> Empty;
+		return Empty;
+	}
+
+#if !UE_BUILD_SHIPPING
+	UE_DEPRECATED(5.6, "Use GetNumOutdatedHLODObjects instead")
+	uint32 GetNumOutdatedHLODActors() const { return 0; }
+#endif
+
+	PRAGMA_DISABLE_DEPRECATION_WARNINGS
+	UE_DEPRECATED(5.6, "Use OnHLODObjectRegisteredEvent instead")
+	FWorldPartitionHLODActorRegisteredEvent& OnHLODActorRegisteredEvent() { return HLODActorRegisteredEvent; }
+	UE_DEPRECATED(5.6, "Use OnHLODObjectUnregisteredEvent instead")
+	FWorldPartitionHLODActorUnregisteredEvent& OnHLODActorUnregisteredEvent() {	return HLODActorUnregisteredEvent; }
+	PRAGMA_ENABLE_DEPRECATION_WARNINGS
+
 #if WITH_EDITOR
 	UE_DEPRECATED(5.5, "Use UWorldPartitionHLODEditorSubsystem::WriteHLODStatsCSV()")
 	static ENGINE_API bool WriteHLODStatsCSV(UWorld* InWorld, const FString& InFilename) { return false; }
@@ -75,10 +106,12 @@ public:
 	ENGINE_API void SetHLODAlwaysLoadedCullDistance(int32 InCullDistance) {}
 	
 private:
+	FWorldPartitionHLODForEachHLODObjectInCellEvent ForEachHLODObjectInCellEvent;
+
 	struct FCellData
 	{
 		bool bIsCellVisible = false;
-		TArray<AWorldPartitionHLOD*> LoadedHLODs;
+		TArray<IWorldPartitionHLODObject*> LoadedHLODs;	// HLOD representation of the cell itself
 	};
 
 	struct FWorldPartitionHLODRuntimeData
@@ -89,7 +122,7 @@ private:
 	TMap<TObjectPtr<UWorldPartition>, FWorldPartitionHLODRuntimeData> WorldPartitionsHLODRuntimeData;
 	ENGINE_API const FCellData* GetCellData(const UWorldPartitionRuntimeCell* InCell) const;
 	ENGINE_API FCellData* GetCellData(const UWorldPartitionRuntimeCell* InCell);
-	ENGINE_API FCellData* GetCellData(AWorldPartitionHLOD* InWorldPartitionHLOD);
+	ENGINE_API FCellData* GetCellData(IWorldPartitionHLODObject* InWorldPartitionHLOD);
 
 	struct FWorldPartitionHLODWarmupState
 	{
@@ -98,11 +131,11 @@ private:
 		FBox WarmupBounds;
 	};
 
-	typedef TMap<FObjectKey, FWorldPartitionHLODWarmupState> FHLODWarmupStateMap;
-	FHLODWarmupStateMap HLODActorsToWarmup;
+	typedef TMap<const IWorldPartitionHLODObject*, FWorldPartitionHLODWarmupState> FHLODWarmupStateMap;
+	FHLODWarmupStateMap HLODObjectsToWarmup;
 
 	ENGINE_API void OnBeginRenderViews(const FSceneViewFamily& InViewFamily);
-	ENGINE_API bool PrepareToWarmup(const UWorldPartitionRuntimeCell* InCell, AWorldPartitionHLOD* InHLODActor);
+	ENGINE_API bool PrepareToWarmup(const UWorldPartitionRuntimeCell* InCell, const IWorldPartitionHLODObject* InHLODObject);
 	ENGINE_API bool ShouldPerformWarmup() const;
 	ENGINE_API bool ShouldPerformWarmupForCell(const UWorldPartitionRuntimeCell* InCell) const;
 	bool bCachedShouldPerformWarmup;
@@ -118,11 +151,15 @@ private:
 	ENGINE_API void OnWorldPartitionInitialized(UWorldPartition* InWorldPartition);
 	ENGINE_API void OnWorldPartitionUninitialized(UWorldPartition* InWorldPartition);
 
+	FWorldPartitionHLODObjectRegisteredEvent	HLODObjectRegisteredEvent;
+	FWorldPartitionHLODObjectUnregisteredEvent	HLODObjectUnregisteredEvent;
+
+	PRAGMA_DISABLE_DEPRECATION_WARNINGS
 	FWorldPartitionHLODActorRegisteredEvent		HLODActorRegisteredEvent;
 	FWorldPartitionHLODActorUnregisteredEvent	HLODActorUnregisteredEvent;
+	PRAGMA_ENABLE_DEPRECATION_WARNINGS
 
 #if !UE_BUILD_SHIPPING
-	TSet<AWorldPartitionHLOD*> OutdatedHLODActors;
+	TSet<IWorldPartitionHLODObject*> OutdatedHLODObjects;
 #endif
 };
-
diff --git a/Engine/Source/Runtime/Engine/Public/WorldPartition/HLOD/IWorldPartitionHLODObject.h b/Engine/Source/Runtime/Engine/Public/WorldPartition/HLOD/IWorldPartitionHLODObject.h
new file mode 100644
--- /dev/null
+++ b/Engine/Source/Runtime/Engine/Public/WorldPartition/HLOD/IWorldPartitionHLODObject.h
@@ -0,0 +1,35 @@
+// Copyright Epic Games, Inc. All Rights Reserved.
+
+#pragma once
+
+#include "Containers/UnrealString.h"
+#include "Containers/Set.h"
+#include "Misc/Guid.h"
+
+class ULevel;
+class UObject;
+
+class IWorldPartitionHLODObject
+{
+public:
+	/** Returns the associated UObject if any (nullptr if none) */
+	virtual UObject* GetUObject() const = 0;
+
+	/** Returns the HLOD object associated level */
+	virtual ULevel* GetHLODLevel() const = 0;
+
+	/** Return the name or the label of the HLOD object */
+	virtual FString GetHLODNameOrLabel() const = 0;
+
+	/** Returns whether the HLOD object requires warmup */
+	virtual bool DoesRequireWarmup() const = 0;
+
+	/** Gather the list of assets that should be warmed up before this HLOD object is made visible */
+	virtual TSet<UObject*> GetAssetsToWarmup() const = 0;
+
+	/** Changes the visibility of the HLOD object */
+	virtual void SetVisibility(bool bIsVisible) = 0;
+
+	/** Returns the associated source cell guid of this HLOD object. */
+	virtual const FGuid& GetSourceCellGuid() const = 0;
+};
\ No newline at end of file
