[shaders] report global shader errors as one single structured log entry with multiple lines instead of errors first followed by the fatal log. horde handles this properly now

diff --git a/Engine/Source/Runtime/Engine/Private/ShaderCompiler/ShaderCompilerEditor.cpp b/Engine/Source/Runtime/Engine/Private/ShaderCompiler/ShaderCompilerEditor.cpp
--- a/Engine/Source/Runtime/Engine/Private/ShaderCompiler/ShaderCompilerEditor.cpp
+++ b/Engine/Source/Runtime/Engine/Private/ShaderCompiler/ShaderCompilerEditor.cpp
@@ -580,40 +580,40 @@ void ProcessCompiledGlobalShaders(const TArray<FShaderCommonCompileJobPtr>& Comp
 
 	FShaderDiagnosticInfo ShaderDiagInfo(CompilationResults);
 
-	// Report unique errors for global shaders.
-	for (int32 ErrorIndex = 0; ErrorIndex < ShaderDiagInfo.UniqueErrors.Num(); ++ErrorIndex)
-	{
-		UE_LOGFMT_NSLOC(LogShaders, Error, "Shaders", "GlobalShaderCompileError", "{ErrorMessage}", 
-			("ErrorMessage", ShaderDiagInfo.UniqueErrors[ErrorIndex]));
-	}
 
-	for (const FString& WarningString : ShaderDiagInfo.UniqueWarnings)
-	{
-		UE_LOGFMT_NSLOC(LogShaders, Warning, "Shaders", "GlobalShaderCompileWarning", "{WarningMessage}", ("WarningMessage", WarningString));
-	}
 
 	const int32 UniqueErrorCount = ShaderDiagInfo.UniqueErrors.Num();
 	if (UniqueErrorCount)
 	{
-		const TCHAR* RetryMsg = TEXT(" Enable 'r.ShaderDevelopmentMode' in ConsoleVariables.ini for retries.");
+		// Report unique errors for global shaders.
+		FString AllUniqueErrors = FString::Join(ShaderDiagInfo.UniqueErrors, TEXT("\n"));
+
+		const TCHAR* RetryMsg = TEXT("\nEnable 'r.ShaderDevelopmentMode' in ConsoleVariables.ini for retries.");
 		if (AreShaderErrorsFatal())
 		{
-			UE_LOGFMT_NSLOC(LogShaders, Fatal, "Shaders", "GlobalShadersCompilationFailed", "{NumErrors} Shader compiler errors compiling GlobalShaders for platform {Platform}. {RetryMsg}",
+			UE_LOGFMT_NSLOC(LogShaders, Fatal, "Shaders", "GlobalShadersCompilationFailed", "{NumErrors} errors encountered compiling global shaders for platform {Platform}:\n{Errors}{RetryMsg}",
 				("NumErrors", UniqueErrorCount),
 				("Platform", ShaderDiagInfo.TargetShaderPlatformString),
-				("RetryMsg", IsRunningCommandlet() ? TEXT("") : RetryMsg)
+				("RetryMsg", IsRunningCommandlet() ? TEXT("") : RetryMsg),
+				("Errors", AllUniqueErrors)
 			);
 		}
 		else
 		{
-			UE_LOGFMT_NSLOC(LogShaders, Error, "Shaders", "GlobalShadersCompilationFailed", "{NumErrors} Shader compiler errors compiling GlobalShaders for platform {Platform}. {RetryMsg}",
+			UE_LOGFMT_NSLOC(LogShaders, Error, "Shaders", "GlobalShadersCompilationFailed", "{NumErrors} errors encountered compiling global shaders for platform {Platform}.\n{Errors}{RetryMsg}",
 				("NumErrors", UniqueErrorCount),
 				("Platform", ShaderDiagInfo.TargetShaderPlatformString),
-				("RetryMsg", IsRunningCommandlet() ? TEXT("") : RetryMsg)
+				("RetryMsg", IsRunningCommandlet() ? TEXT("") : RetryMsg),
+				("Errors", AllUniqueErrors)
 			);
 		}
 	}
 
+	for (const FString& WarningString : ShaderDiagInfo.UniqueWarnings)
+	{
+		UE_LOGFMT_NSLOC(LogShaders, Warning, "Shaders", "GlobalShaderCompileWarning", "{WarningMessage}", ("WarningMessage", WarningString));
+	}
+
 	TArray<EShaderPlatform> ShaderPlatformsProcessed;
 	TArray<const FShaderPipelineType*> SharedPipelines;
 
