Level Streaming - Avoid possible issue with the use of FindObjectFast() in ULevelStreaming::GetLevelStreamingStatus()
* When using async update streaming, the GC could be running while this call is done outside of the main thread

diff --git a/Engine/Source/Runtime/Engine/Private/LevelStreaming.cpp b/Engine/Source/Runtime/Engine/Private/LevelStreaming.cpp
--- a/Engine/Source/Runtime/Engine/Private/LevelStreaming.cpp
+++ b/Engine/Source/Runtime/Engine/Private/LevelStreaming.cpp
@@ -2392,6 +2392,13 @@ EStreamingStatus ULevelStreaming::GetLevelStreamingStatus() const
 		}
 
 		check(CurrentState == ELevelStreamingState::Removed || CurrentState == ELevelStreamingState::Unloaded);
+
+		// We can't use FindObjectFast() while the GC is running... Let's treat the level as unloaded at this point.
+		if (IsGarbageCollecting())
+		{
+			return LEVEL_Unloaded;
+		}
+
 		// See whether the level's world object is still around.
 		UPackage* LevelPackage = FindObjectFast<UPackage>(nullptr, GetWorldAssetPackageFName());
 		UWorld* LevelWorld = LevelPackage ? UWorld::FindWorldInPackage(LevelPackage) : nullptr;
