FPrimitiveSceneProxy - Added SetDrawnInGame_RenderThread()
* Can be used to change a scene proxy's visibility without having to recreate a component's render state

diff --git a/Engine/Source/Runtime/Engine/Private/PrimitiveSceneProxy.cpp b/Engine/Source/Runtime/Engine/Private/PrimitiveSceneProxy.cpp
--- a/Engine/Source/Runtime/Engine/Private/PrimitiveSceneProxy.cpp
+++ b/Engine/Source/Runtime/Engine/Private/PrimitiveSceneProxy.cpp
@@ -545,7 +545,7 @@ FPrimitiveSceneProxy::FPrimitiveSceneProxy(const FPrimitiveSceneProxyDesc& InPro
 ,	PrimitiveSceneInfo(nullptr)
 ,	OwnerName(InProxyDesc.GetOwner() ? InProxyDesc.GetOwner()->GetFName() : NAME_None)
 ,	ResourceName(InResourceName)
-,	LevelName(InProxyDesc.GetComponentLevel() ? InProxyDesc.GetComponentLevel()->GetOutermost()->GetFName() : NAME_None)
+,	LevelName(InProxyDesc.GetLevel() ? InProxyDesc.GetLevel()->GetOutermost()->GetFName() : NAME_None)
 #if WITH_EDITOR
 // by default we are always drawn
 ,	HiddenEditorViews(0)
@@ -565,7 +565,7 @@ FPrimitiveSceneProxy::FPrimitiveSceneProxy(const FPrimitiveSceneProxyDesc& InPro
 	check(Scene);
 
 	// Initialize ForceHidden flag based on Level's visibility (only if Level bRequireFullVisibilityToRender is set)
-	if (ULevel* Level = InProxyDesc.GetComponentLevel())
+	if (ULevel* Level = InProxyDesc.GetLevel())
 	{
 		bShouldNotifyOnWorldAddRemove = Level->bRequireFullVisibilityToRender;
 		if (bShouldNotifyOnWorldAddRemove)
@@ -1198,6 +1198,26 @@ void FPrimitiveSceneProxy::SetCustomDepthEnabled_RenderThread(const bool bInRend
 	}
 }
 
+/**
+* Set the draw in game flag (RENDER THREAD)
+*
+* @param the new value
+*/
+void FPrimitiveSceneProxy::SetDrawnInGame_RenderThread(bool bInDrawInGame)
+{
+	check(IsInRenderingThread());
+	if (DrawInGame != bInDrawInGame)
+	{
+		DrawInGame = bInDrawInGame;
+
+		if (PrimitiveSceneInfo)
+		{
+			Scene->RequestUniformBufferUpdate(*PrimitiveSceneInfo);
+			Scene->RequestGPUSceneUpdate(*PrimitiveSceneInfo, EPrimitiveDirtyState::ChangedOther);	
+		}
+	}
+}
+
 /**
 * Set the custom depth stencil value
 *
diff --git a/Engine/Source/Runtime/Engine/Public/PrimitiveSceneProxy.h b/Engine/Source/Runtime/Engine/Public/PrimitiveSceneProxy.h
--- a/Engine/Source/Runtime/Engine/Public/PrimitiveSceneProxy.h
+++ b/Engine/Source/Runtime/Engine/Public/PrimitiveSceneProxy.h
@@ -1088,7 +1088,13 @@ public:
 	virtual int32 GetLightMapCoordinateIndex() const { return INDEX_NONE; }
 
 	/** Tell us if this proxy is drawn in game.*/
-	virtual bool IsDrawnInGame() const { return DrawInGame; }
+	FORCEINLINE bool IsDrawnInGame() const { return DrawInGame; }
+
+	/**
+	 * Set the drawn in game flag, from the render thread.
+	 * @param bInDrawInGame		The new value.
+ 	 */
+	ENGINE_API void SetDrawnInGame_RenderThread(bool bInDrawInGame);
 
 	/** Tell us if this proxy is drawn in editor.*/
 	FORCEINLINE bool IsDrawnInEditor() const { return DrawInEditor; }
diff --git a/Engine/Source/Runtime/Engine/Public/PrimitiveSceneProxyDesc.h b/Engine/Source/Runtime/Engine/Public/PrimitiveSceneProxyDesc.h
--- a/Engine/Source/Runtime/Engine/Public/PrimitiveSceneProxyDesc.h
+++ b/Engine/Source/Runtime/Engine/Public/PrimitiveSceneProxyDesc.h
@@ -218,6 +218,8 @@ struct FPrimitiveSceneProxyDesc
 	T* GetOwner() const { return Cast<T>(Owner); }	
 
 	ULevel* GetLevel() const { return  Owner ? Owner->GetTypedOuter<ULevel>() : nullptr; }
+
+	UE_DEPRECATED(5.6, "Use GetLevel instead.")
 	ULevel* GetComponentLevel() const { return GetLevel(); }
 
 	FString GetPathName() const { return Component->GetPathName(); }
