Add new TaskSyncManager which allows registration of globally accessible tick functions that can be used to synchronize different runtime systems and efficiently batch per-frame updates. Support for manual task dispatch was added to FTickFunction so these functions can wait to be triggered halfway through a tick group.
If tick.CreateTaskSyncManager is enabled, it wll create the manager at startup and register sync points that are defined in the Task Synchronization section of project settings.
Expose OnWorldContextDestroyed in runtime builds to better support multiple worlds for testing or special use cases.
This version is still experimental and is disabled by default, it will change before the release of 5.6.
This also includes additional EngineRuntimeTests for Mass and the async message system, this will be split into a separate CL before checkin.

diff --git a/Engine/Source/Runtime/Engine/Classes/Engine/Engine.h b/Engine/Source/Runtime/Engine/Classes/Engine/Engine.h
--- a/Engine/Source/Runtime/Engine/Classes/Engine/Engine.h
+++ b/Engine/Source/Runtime/Engine/Classes/Engine/Engine.h
@@ -3482,23 +3482,17 @@ public:
 
 	ENGINE_API virtual void DestroyWorldContext(UWorld * InWorld);
 
-#if WITH_EDITOR
 	/** Triggered when a world context is destroyed. */
 	DECLARE_EVENT_OneParam(UEngine, FWorldContextDestroyedEvent, FWorldContext&);
 
 	/** Return the world context destroyed event. */
 	FWorldContextDestroyedEvent&	OnWorldContextDestroyed() { return WorldContextDestroyedEvent; }
-#endif // #if WITH_EDITOR
 
 private:
-
-#if WITH_EDITOR
 	/** Delegate broadcast when a world context is destroyed */
 	FWorldContextDestroyedEvent WorldContextDestroyedEvent;
-#endif // #if WITH_EDITOR
 
 public:
-
 	/** 
 	 * Called from GetFunctionCallspace on specific objects to check for authority/cosmetic function tags using global state
 	 *
diff --git a/Engine/Source/Runtime/Engine/Classes/Engine/EngineBaseTypes.h b/Engine/Source/Runtime/Engine/Classes/Engine/EngineBaseTypes.h
--- a/Engine/Source/Runtime/Engine/Classes/Engine/EngineBaseTypes.h
+++ b/Engine/Source/Runtime/Engine/Classes/Engine/EngineBaseTypes.h
@@ -223,7 +223,12 @@ public:
 	/** (experimental) if true, the tick function will be run transactionally */
 	uint8 bRunTransactionally:1;
 
+	/** True if this is an event that will not execute until it has been manually dispatched (and prerequisites complete) instead of being dispatched at start of the tick group */
+	uint8 bDispatchManually : 1;
+
 private:
+	/** For bDispatchManually functions, if this is set and the task is valid the task has been fully triggered and will execute */
+	uint8 bWasDispatchedManually:1;
 
 	enum class ETickState : uint8
 	{
@@ -237,7 +242,7 @@ private:
 	 * If CoolingDown, this tick has an interval frequency that is being adhered to currently
 	 * CAUTION: Do not set this directly
 	 **/
-	ETickState TickState;
+	ETickState TickState : 2;
 
 public:
 	/** The frequency in seconds at which this tick function will be executed.  If less than or equal to 0 then it will tick every frame */
@@ -347,6 +352,17 @@ public:
 	 */
 	ENGINE_API FGraphEventRef GetCompletionHandle() const;
 
+	/**
+	 * Returns true if bDispatchManually is set and this has a task that can be triggered right now.
+	 */
+	ENGINE_API virtual bool CanDispatchManually() const;
+
+	/** 
+	 * If bDispatchManually is set, attempt to start the task based on some event during a frame.
+	 * It will return true if it worked, which may have executed immediately or may be waiting for a prerequisite.
+	 */
+	ENGINE_API virtual bool DispatchManually();
+
 	/** 
 	 * Gets the action tick group that this function will be elligible to start in.
 	 * Only valid after StartFrame has been called through the end of the frame.
diff --git a/Engine/Source/Runtime/Engine/Private/TaskSyncManager.cpp b/Engine/Source/Runtime/Engine/Private/TaskSyncManager.cpp
new file mode 100644
--- /dev/null
+++ b/Engine/Source/Runtime/Engine/Private/TaskSyncManager.cpp
@@ -0,0 +1,1160 @@
+// Copyright Epic Games, Inc. All Rights Reserved.
+
+#include "TaskSyncManager.h"
+#include "Async/TaskGraphInterfaces.h"
+#include "Engine/Engine.h"
+#include "Engine/Level.h"
+#include "Engine/World.h"
+#include "Misc/ScopeLock.h"
+
+DEFINE_LOG_CATEGORY_STATIC(LogTaskSync, Log, All);
+
+using namespace UE::Tick;
+
+static bool bCreateTaskSyncManager = false;
+static FAutoConsoleVariableRef CVarCreateTaskSyncManager(
+	TEXT("tick.CreateTaskSyncManager"),
+	bCreateTaskSyncManager,
+	TEXT("If true, create the experimental Task Sync Manager (will always be true eventually)"));
+
+#if WITH_EDITOR
+void UTaskSyncManagerSettings::PostEditChangeProperty(FPropertyChangedEvent& PropertyChangedEvent)
+{
+	if (FTaskSyncManager* Manager = FTaskSyncManager::Get())
+	{
+		Manager->ReloadRegisteredData();
+	}
+}
+
+bool UTaskSyncManagerSettings::SupportsAutoRegistration() const
+{
+	return bCreateTaskSyncManager;
+}
+#endif 
+
+FWorldContextId FWorldContextId::GetContextIdForWorld(const UWorld* World)
+{
+	if (World)
+	{
+		// TODO this logic could be moved into GEngine so we can reuse it. Using separate int because of TIndirectArray
+		constexpr FInternalId PIEWorldStart = 1;
+		constexpr FInternalId TestWorldStart = 256;
+		uint32 CurrentWorldIndex = 0;
+		
+		for (const FWorldContext& Context : GEngine->GetWorldContexts())
+		{
+			if (World == Context.World())
+			{
+				if (CurrentWorldIndex == 0)
+				{
+					ensure(Context.WorldType == EWorldType::Game || Context.WorldType == EWorldType::Editor);
+					// The first world should either be a game or editor world
+					return FWorldContextId(FWorldContextId::DefaultWorldContextId);
+				}
+#if WITH_EDITOR
+				else if (Context.PIEInstance >= 0)
+				{
+					check(Context.PIEInstance < TestWorldStart - 1);
+					return FWorldContextId(Context.PIEInstance + PIEWorldStart);
+				}
+#endif
+				// Behavior is currently unclear for preview/test worlds, some are ticked separately and some are not with no explicit setting
+				else if (Context.WorldType == EWorldType::Game || Context.WorldType == EWorldType::GamePreview || Context.WorldType == EWorldType::EditorPreview)
+				{
+					return FWorldContextId(CurrentWorldIndex + TestWorldStart);
+				}
+				else
+				{
+					return FWorldContextId();
+				}
+
+			}
+
+			CurrentWorldIndex++;
+		}
+	}
+
+	return FWorldContextId();
+}
+
+FActiveSyncWorkHandle::~FActiveSyncWorkHandle()
+{
+	Reset();
+}
+
+bool FActiveSyncWorkHandle::IsValid() const
+{
+	return SyncPoint.IsValid();
+}
+
+FTickFunction* FActiveSyncWorkHandle::GetDependencyTickFunction()
+{
+	if (IsInGameThread())
+	{
+		return SyncPoint.Get();
+	}
+	return nullptr;
+}
+
+bool FActiveSyncWorkHandle::HasReservedWork() const
+{
+	return bWorkReserved;
+}
+
+bool FActiveSyncWorkHandle::HasRequestedWork() const
+{
+	return bWorkRequested;
+}
+
+bool FActiveSyncWorkHandle::AbandonWork()
+{
+	if (ensure(SyncPoint.IsValid()))
+	{
+		return SyncPoint->AbandonWork(*this);
+	}
+	return false;
+}
+
+bool FActiveSyncWorkHandle::Reset()
+{
+	if (SyncPoint.IsValid())
+	{
+		SyncPoint->ResetWorkHandle(*this);
+	}
+	return false;
+}
+
+bool FActiveSyncWorkHandle::ReserveFutureWork(ESyncWorkRepetition Repeat)
+{
+	if (ensure(SyncPoint.IsValid()))
+	{
+		return SyncPoint->ReserveFutureWork(*this, Repeat);
+	}
+	return false;
+}
+
+bool FActiveSyncWorkHandle::RequestWork(FTickFunction* WorkFunction, ESyncWorkRepetition Repeat)
+{
+	if (ensure(SyncPoint.IsValid()))
+	{
+		return SyncPoint->RequestWork(*this, WorkFunction, Repeat);
+	}
+	return false;
+}
+
+
+FActiveSyncPoint::FActiveSyncPoint()
+{
+	bCanEverTick = true;
+	bDispatchManually = true;
+	bStartWithTickEnabled = true;
+	FrameStatus = ESyncPointStatus::Unknown;
+}
+
+void FActiveSyncPoint::ResetWorkForFrame()
+{
+	// Don't lock as this is called from a very specific place before dispatching of tasks
+	check(IsInGameThread());
+
+	for (int32 i = 0; i < ActiveWork.Num(); i++)
+	{
+		FActiveSyncWork& CurrentWork = ActiveWork[i];
+
+		check(!CurrentWork.bCurrentlyExecuting);
+
+		CurrentWork.bAlreadyExecuted = false;
+
+		if (CurrentWork.bRequestEveryFrame)
+		{
+			CurrentWork.bWorkReserved = false;
+			CurrentWork.bWorkRequested = true;
+		}
+		else if (CurrentWork.bReserveEveryFrame)
+		{
+			CurrentWork.bWorkReserved = true;
+		}
+		// Don't reset reserved/requested in case they were set before the frame started
+	}
+}
+
+bool FActiveSyncPoint::HandleFirstTickGroup()
+{
+	// Don't lock as this is called from a very specific place before dispatching of tasks
+	check(IsInGameThread());
+
+	ensure(FrameStatus == ESyncPointStatus::TaskNotCreated);
+	switch (SyncPointDescription.ActivationRules)
+	{
+	case ESyncPointActivationRules::AlwaysActivate:
+		// Not technically true, but will be dispatched soon
+		FrameStatus = ESyncPointStatus::Dispatched;
+		break;
+	case ESyncPointActivationRules::WaitForTrigger:
+		FrameStatus = ESyncPointStatus::DispatchWaitingForTrigger;
+		break;
+	case ESyncPointActivationRules::WaitForAllWork:
+	case ESyncPointActivationRules::ActivateForAnyWork:
+		FrameStatus = ESyncPointStatus::DispatchWaitingForWork;
+		
+		return IsReadyToProcessWork();
+	}
+	return false;
+}
+
+bool FActiveSyncPoint::IsReadyToProcessWork()
+{
+	// See if we should dispatch a process task
+	bool bHasRequest = false;
+	bool bHasReservation = false;
+
+	for (FActiveSyncWork& CurrentWork : ActiveWork)
+	{
+		if (CurrentWork.bWorkRequested)
+		{
+			bHasRequest = true;
+		}
+		else if (CurrentWork.bWorkReserved)
+		{
+			bHasReservation = true;
+		}
+	}
+
+	if (bHasRequest)
+	{
+		// Needs to be dispatched if it has no reservations or will start for any work
+		return !bHasReservation || (SyncPointDescription.ActivationRules == ESyncPointActivationRules::ActivateForAnyWork);
+	}
+	return false;
+}
+
+bool FActiveSyncPoint::GetWorkToExecute(int32& NumReserved)
+{
+	// Copy work into executing work and possibly delete active work
+	UE::TScopeLock WorkScope(WorkLock);
+
+	check(ExecutingWork.Num() == 0); // ExecuteTick cannot be re-entrant and will not be resized during execution
+
+	check(FrameStatus >= ESyncPointStatus::Dispatched && FrameStatus != ESyncPointStatus::ExecutionComplete);
+	FrameStatus = ESyncPointStatus::Executing;
+	ExecutingWork.Reserve(ActiveWork.Num());
+	NumReserved = 0;
+
+	for (int32 i = 0; i < ActiveWork.Num(); i++)
+	{
+		FActiveSyncWork& CurrentWork = ActiveWork[i];
+		if (CurrentWork.bCurrentlyExecuting)
+		{
+			CurrentWork.bCurrentlyExecuting = false;
+			CurrentWork.bAlreadyExecuted = true;
+			if (!CurrentWork.bRequestEveryFrame)
+			{
+				// Clear function now
+				CurrentWork.TickFunction = nullptr;
+				CurrentWork.bWorkRequested = false;
+			}
+		}
+		else if (CurrentWork.bWorkReserved)
+		{
+			NumReserved++;
+		}
+		else if (CurrentWork.bWorkRequested && !CurrentWork.bAlreadyExecuted)
+		{
+			ExecutingWork.Add(CurrentWork);
+			CurrentWork.bCurrentlyExecuting = true;
+		}
+	}
+
+	if (ExecutingWork.Num() > 0)
+	{
+		return true;
+	}
+	else if (NumReserved > 0)
+	{
+		FrameStatus = ESyncPointStatus::WaitingForMoreWork;
+		return false;
+	}
+	else
+	{
+		FrameStatus = ESyncPointStatus::ExecutionComplete;
+		return false;
+	}
+}
+
+void FActiveSyncPoint::ExecuteTick(float DeltaTime, ELevelTick TickType, ENamedThreads::Type CurrentThread, const FGraphEventRef& MyCompletionGraphEvent)
+{
+	// Only lock during acquisition, GetWorkToExecute protects against recursive ticking
+	// This is a loop because the worker threads could add new work during execution
+
+	if (SyncPointDescription.EventType == ESyncPointEventType::SimpleEvent)
+	{
+		return;
+	}
+
+	// TODO add pre/post batch function calls
+
+	int32 NumReserved = 0;
+	while (GetWorkToExecute(NumReserved))
+	{
+		for (int32 i = 0; i < ExecutingWork.Num(); i++)
+		{
+			// Executing work cannot change size during execution
+			FTickFunction* TickFunction = ExecutingWork[i].TickFunction;
+			if (TickFunction && TickFunction->IsTickFunctionEnabled())
+			{
+				TickFunction->ExecuteNestedTick(DeltaTime, TickType, CurrentThread, MyCompletionGraphEvent);
+			}
+		}
+
+		ExecutingWork.SetNum(0, EAllowShrinking::No);
+	}
+
+	if (NumReserved > 0)
+	{
+		UE::TScopeLock WorkScope(WorkLock);
+
+		check(FrameStatus == ESyncPointStatus::WaitingForMoreWork);
+
+		if (ensure(SyncPointDescription.ActivationRules == ESyncPointActivationRules::ActivateForAnyWork))
+		{
+			ReactivationEvent = FGraphEvent::CreateGraphEvent();
+
+			// TODO create a new task to process the queue, need to add something to TickTaskManager
+		}
+		else
+		{
+			// This is invalid but set us to complete
+			FrameStatus = ESyncPointStatus::ExecutionComplete;
+		}
+	}
+}
+
+FString FActiveSyncPoint::DiagnosticMessage()
+{
+	return FString::Printf(TEXT("TaskSyncTickFunction: %s"), *SyncPointDescription.RegisteredName.ToString());
+}
+
+FName FActiveSyncPoint::DiagnosticContext(bool bDetailed)
+{
+	return SyncPointDescription.RegisteredName;
+}
+
+bool FActiveSyncPoint::DispatchManually()
+{
+	if (CanDispatchManually())
+	{
+		ensure(FrameStatus == ESyncPointStatus::DispatchWaitingForTrigger || FrameStatus == ESyncPointStatus::DispatchWaitingForWork);
+		// Update status
+		FrameStatus = ESyncPointStatus::Dispatched;
+		return FTickFunction::DispatchManually();
+	}
+	return false;
+}
+
+FActiveSyncWorkHandle FActiveSyncPoint::RegisterWorkHandle()
+{
+	UE::TScopeLock WorkScope(WorkLock);
+
+	uint32 WorkIndex = AllocateActiveWork();
+	FActiveSyncWork& CurrentWork = ActiveWork[WorkIndex];
+	check(!CurrentWork.bHasActiveHandle);
+
+	CurrentWork.bHasActiveHandle = true;
+
+	FActiveSyncWorkHandle ReturnHandle;
+	ReturnHandle.SyncPoint = AsShared();
+	ReturnHandle.WorkIndex = WorkIndex;
+	ReturnHandle.bWorkReserved = false;
+	ReturnHandle.bWorkRequested = false;
+
+	return MoveTemp(ReturnHandle);
+}
+
+bool FActiveSyncPoint::ReserveFutureWork(FActiveSyncWorkHandle& Handle, ESyncWorkRepetition Repeat)
+{
+	UE::TScopeLock WorkScope(WorkLock);
+
+	if (ensure(Handle.IsValid() && ActiveWork.IsValidIndex(Handle.WorkIndex)))
+	{
+		FActiveSyncWork& CurrentWork = ActiveWork[Handle.WorkIndex];
+
+		if (ensure(CurrentWork.bHasActiveHandle && !CurrentWork.bWorkReserved && !CurrentWork.bWorkRequested))
+		{
+			// TODO properly handle the repetition types
+
+			CurrentWork.bWorkReserved = true;
+			CurrentWork.bReserveEveryFrame = (Repeat == ESyncWorkRepetition::EveryFrame);
+
+			Handle.bWorkReserved = true;
+
+			if (FrameStatus > ESyncPointStatus::DispatchWaitingForWork && IsReadyToProcessWork())
+			{
+				WorkScope.Unlock();
+				DispatchWorkTask();
+			}
+
+			return true;
+		}
+	}
+	return false;
+}
+
+bool FActiveSyncPoint::RequestWork(FActiveSyncWorkHandle& Handle, FTickFunction* WorkFunction, ESyncWorkRepetition Repeat)
+{
+	UE::TScopeLock WorkScope(WorkLock);
+
+	if (ensure(Handle.IsValid() && ActiveWork.IsValidIndex(Handle.WorkIndex)))
+	{
+		FActiveSyncWork& CurrentWork = ActiveWork[Handle.WorkIndex];
+
+		if (ensure(CurrentWork.bHasActiveHandle && !CurrentWork.bWorkRequested))
+		{
+			// TODO properly handle the repetition types
+
+			CurrentWork.TickFunction = WorkFunction;
+			CurrentWork.bWorkReserved = false;
+			CurrentWork.bWorkRequested = true;
+			CurrentWork.bRequestEveryFrame = (Repeat == ESyncWorkRepetition::EveryFrame);
+
+			Handle.bWorkReserved = false;
+			Handle.bWorkRequested = true;
+
+			if (FrameStatus > ESyncPointStatus::DispatchWaitingForWork && IsReadyToProcessWork())
+			{
+				WorkScope.Unlock();
+				DispatchWorkTask();
+			}
+
+			return true;
+		}
+	}
+	return false;
+}
+
+bool FActiveSyncPoint::AbandonWork(FActiveSyncWorkHandle& Handle)
+{
+	UE::TScopeLock WorkScope(WorkLock);
+	
+	if (ensure(Handle.IsValid() && ActiveWork.IsValidIndex(Handle.WorkIndex)))
+	{
+		FActiveSyncWork& CurrentWork = ActiveWork[Handle.WorkIndex];
+
+		if (ensure(CurrentWork.bHasActiveHandle))
+		{
+			if (CurrentWork.bWorkReserved || CurrentWork.bWorkRequested || CurrentWork.bReserveEveryFrame || CurrentWork.bRequestEveryFrame)
+			{
+				// TODO What do we do if work has already been scheduled but not executed?
+
+				CurrentWork.bWorkReserved = false;
+				CurrentWork.bWorkRequested = false;
+				CurrentWork.bReserveEveryFrame = false;
+				CurrentWork.bRequestEveryFrame = false;
+
+				Handle.bWorkReserved = false;
+				Handle.bWorkRequested = false;
+			
+				if (FrameStatus > ESyncPointStatus::DispatchWaitingForWork && IsReadyToProcessWork())
+				{
+					WorkScope.Unlock();
+					DispatchWorkTask();
+				}
+			}
+			return true;
+		}
+	}
+	return false;
+}
+
+bool FActiveSyncPoint::ResetWorkHandle(FActiveSyncWorkHandle& Handle)
+{
+	UE::TScopeLock WorkScope(WorkLock);
+
+	if (ensure(Handle.IsValid() && ActiveWork.IsValidIndex(Handle.WorkIndex)))
+	{
+		FActiveSyncWork& CurrentWork = ActiveWork[Handle.WorkIndex];
+
+		if (ensure(CurrentWork.bHasActiveHandle))
+		{
+			bool bShouldCheckWork = CurrentWork.bWorkReserved || CurrentWork.bWorkRequested;
+
+			CurrentWork.Reset();
+			Handle.ResetInternal();
+
+			// TODO What do we do if work has already been scheduled but not executed?
+
+			if (bShouldCheckWork && FrameStatus > ESyncPointStatus::DispatchWaitingForWork && IsReadyToProcessWork())
+			{
+				WorkScope.Unlock();
+				DispatchWorkTask();
+			}
+			return true;
+		}
+	}
+	return false;
+}
+
+uint32 FActiveSyncPoint::AllocateActiveWork()
+{
+	// Look for deleted slot which will usually exist
+	for (int32 i = 0; i < ActiveWork.Num(); i++)
+	{
+		if (!ActiveWork[i].IsInitialized())
+		{
+			return i;
+		}
+	}
+
+	// Add new slot
+	check(ActiveWork.Num() < FActiveSyncWorkHandle::MaxWorkIndex - 1);
+
+	return ActiveWork.AddDefaulted();
+}
+
+void FActiveSyncPoint::DispatchWorkTask()
+{
+	if (FrameStatus < ESyncPointStatus::Dispatched)
+	{
+		ensure(DispatchManually());
+	}
+	else
+	{
+		// TODO Reschedule work with ReactivationEvent
+	}
+}
+
+FTaskSyncManager* FTaskSyncManager::Get()
+{
+	static TUniquePtr<FTaskSyncManager> SingletonInstance;
+
+	if (SingletonInstance.IsValid())
+	{
+		// 99% case so check ptr first
+		return SingletonInstance.Get();
+	}
+
+	if (!bCreateTaskSyncManager)
+	{
+		// Never created
+		return nullptr;
+	}
+
+	// Make on demand, but only on game thread. This happens early in tick startup so no async work should be happening before it
+	check(IsInGameThread());
+	SingletonInstance = MakeUnique<FTaskSyncManager>();
+	return SingletonInstance.Get();
+}
+
+FTaskSyncManager::FTaskSyncManager()
+{
+	ReloadRegisteredData();
+
+	if (GEngine)
+	{
+		GEngine->OnWorldContextDestroyed().AddRaw(this, &FTaskSyncManager::OnWorldContextRemove);
+	}
+}
+
+FTaskSyncManager::~FTaskSyncManager()
+{
+	if (GEngine)
+	{
+		GEngine->OnWorldContextDestroyed().RemoveAll(this);
+	}
+}
+
+void FTaskSyncManager::RegisterSyncPointInternal(const FSyncPointDescription& InDescription)
+{
+	FSyncPointId::FInternalId NewId = ++HighestSyncId;
+	FRegisteredSyncPointData& Data = RegisteredDataMap.FindOrAdd(NewId);
+	Data.RegisteredId = NewId;
+	Data.RegisteredPoint = InDescription;
+
+	RegisteredNameMap.Add(InDescription.RegisteredName, NewId);
+
+	// TODO what if we're already ticking?
+}
+
+void FTaskSyncManager::ReloadRegisteredData()
+{
+	check(IsInGameThread() && !IsTicking());
+
+	UE::TScopeLock Lock(ManagerLock);
+	const UTaskSyncManagerSettings* Settings = GetDefault<UTaskSyncManagerSettings>();
+
+	TArray<FSyncPointDescription> OtherSourceDescriptions;
+
+	for (TPair<FSyncPointId::FInternalId, FRegisteredSyncPointData>& DataPair : RegisteredDataMap)
+	{
+		FSyncPointDescription& Description = DataPair.Value.RegisteredPoint;
+		if (!Description.WasLoadedFromSettings())
+		{
+			OtherSourceDescriptions.Add(Description);
+		}
+	}
+
+	RegisteredDataMap.Reset();
+	RegisteredNameMap.Reset();
+
+	// Clear out all batch data, this is safe since we're not ticking
+	for (FBatchData& BatchData : BatchList)
+	{
+		BatchData.SyncPointData.Reset();
+	}
+
+	// Reset the sync point ids
+	HighestSyncId = FSyncPointId::InvalidSyncPoint;
+
+	for (const FSyncPointDescription& SyncPoint : Settings->RegisteredSyncPoints)
+	{
+		RegisterSyncPointInternal(SyncPoint);
+	}
+
+	for (const FSyncPointDescription& SyncPoint : OtherSourceDescriptions)
+	{
+		RegisterSyncPointInternal(SyncPoint);
+	}
+}
+
+bool FTaskSyncManager::GetSyncPointDescription(FName RegisteredName, FSyncPointDescription& OutDescription) const
+{
+	UE::TScopeLock Lock(ManagerLock);
+
+	if (const FSyncPointId::FInternalId* FoundId = RegisteredNameMap.Find(RegisteredName))
+	{
+		if (const FRegisteredSyncPointData* FoundData = RegisteredDataMap.Find(*FoundId))
+		{
+			OutDescription = FoundData->RegisteredPoint;
+			return true;
+		}
+	}
+	
+	return false;
+}
+
+bool FTaskSyncManager::RegisterNewSyncPoint(const FSyncPointDescription& NewDescription)
+{
+	check(IsInGameThread());
+
+	if (ensure(NewDescription.IsValid() && !NewDescription.WasLoadedFromSettings()))
+	{
+		UE::TScopeLock Lock(ManagerLock);
+
+		const FSyncPointId::FInternalId* FoundId = RegisteredNameMap.Find(NewDescription.RegisteredName);
+
+		if (!FoundId)
+		{
+			RegisterSyncPointInternal(NewDescription);
+			return true;
+		}
+		else
+		{
+			UE_LOG(LogTaskSync, Error, TEXT("Cannot register sync point %s from source %s as it already exists!"), *NewDescription.RegisteredName.ToString(), *NewDescription.SourceName.ToString());
+		}
+	}
+	else
+	{
+		UE_LOG(LogTaskSync, Error, TEXT("Cannot register invalid sync point description %s from source %s!"), *NewDescription.RegisteredName.ToString(), *NewDescription.SourceName.ToString());
+	}
+
+	return false;
+}
+
+bool FTaskSyncManager::UnregisterSyncPoint(FName RegisteredName, FName SourceName)
+{
+	check(IsInGameThread());
+
+	if (ensure(!RegisteredName.IsNone() && !SourceName.IsNone()))
+	{
+		const FSyncPointId::FInternalId* FoundId = RegisteredNameMap.Find(RegisteredName);
+
+		if (!FoundId)
+		{
+			UE_LOG(LogTaskSync, Error, TEXT("Cannot unregister sync point %s from source %s as it does not exist!"), *RegisteredName.ToString(), *SourceName.ToString());
+			return false;
+		}
+		else
+		{
+			if (const FRegisteredSyncPointData* FoundData = RegisteredDataMap.Find(*FoundId))
+			{
+				if (FoundData->RegisteredPoint.SourceName == SourceName)
+				{
+					// TODO clear out active batch map properly
+					RegisteredDataMap.Remove(*FoundId);
+					RegisteredNameMap.Remove(RegisteredName);
+				}
+				else
+				{
+					UE_LOG(LogTaskSync, Error, TEXT("Cannot unregister sync point %s from source %s as it was added by source %s!"), *RegisteredName.ToString(), *SourceName.ToString(), *FoundData->RegisteredPoint.SourceName.ToString());
+				}
+
+				return true;
+			}
+			else
+			{
+				UE_LOG(LogTaskSync, Error, TEXT("Cannot unregister sync point %s from source %s as it does not exist!"), *RegisteredName.ToString(), *SourceName.ToString());
+			}
+		}
+	}
+	else
+	{
+		UE_LOG(LogTaskSync, Error, TEXT("Cannot unregister invalid sync point %s from source %s!"), *RegisteredName.ToString(), *SourceName.ToString());
+	}
+
+	return false;
+}
+
+void FTaskSyncManager::StartFrame(const UWorld* InWorld)
+{
+	check(InWorld);	
+	check(IsInGameThread() && !IsTicking() && CurrentTickGroup == ETickingGroup::TG_MAX);
+
+	UE::TScopeLock Lock(ManagerLock);
+	CurrentTickWorld = InWorld;
+	ULevel* PersistentLevel = InWorld->PersistentLevel.Get();
+	FWorldContextId WorldContext(InWorld);
+
+	// For all batches that match the world
+	bool bFoundBatch = false;
+	for (FBatchData& BatchData : BatchList)
+	{
+		if (BatchData.WorldContext == WorldContext)
+		{
+			bFoundBatch = true;
+			InitializeBatchForFrame(BatchData, PersistentLevel);
+		}
+	}
+
+	// Add the default batch if there were 0
+	if (!bFoundBatch)
+	{
+		FBatchContextId BatchContext = CreateNewBatch(WorldContext);
+		check(BatchList.IsValidIndex(BatchContext.BatchId));
+		InitializeBatchForFrame(BatchList[BatchContext.BatchId], PersistentLevel);
+	}
+}
+
+void FTaskSyncManager::InitializeBatchForFrame(FBatchData& BatchData, ULevel* PersistentLevel)
+{
+	for (TPair<FSyncPointId::FInternalId, FRegisteredSyncPointData>& DataPair : RegisteredDataMap)
+	{
+		FActiveSyncPoint* SyncPoint = GetOrCreateSyncPoint(BatchData, DataPair.Value);
+		if (SyncPoint)
+		{
+			if (!SyncPoint->IsTickFunctionRegistered())
+			{
+				FSyncPointDescription& Description = DataPair.Value.RegisteredPoint;
+				
+				// Check our prerequisites now and try to create them if necessary
+				// This could resize BatchData.SyncPointData
+				for (FName Prerequisite : Description.PrerequisiteSyncGroups)
+				{
+					FSyncPointId::FInternalId* FoundId = RegisteredNameMap.Find(Prerequisite);
+
+					// TODO ignore or error on missing prereq?
+					if (FoundId)
+					{
+						FRegisteredSyncPointData* FoundData = RegisteredDataMap.Find(*FoundId);
+						if (ensure(FoundData))
+						{
+							FActiveSyncPoint* PrerequisiteSyncPoint = GetOrCreateSyncPoint(BatchData, *FoundData);
+
+							if (PrerequisiteSyncPoint)
+							{
+								SyncPoint->AddPrerequisite(PersistentLevel, *PrerequisiteSyncPoint);
+							}
+						}
+						
+					}
+				}
+
+				ensure(SyncPoint->FrameStatus == ESyncPointStatus::TaskNotRegistered);
+				SyncPoint->RegisterTickFunction(PersistentLevel);
+			}
+			SyncPoint->FrameStatus = ESyncPointStatus::TaskNotCreated;
+			SyncPoint->ResetWorkForFrame();
+		}
+	}
+}
+
+FActiveSyncPoint* FTaskSyncManager::GetOrCreateSyncPoint(FBatchData& BatchData, FRegisteredSyncPointData& SyncData)
+{
+	FSyncPointDescription& Description = SyncData.RegisteredPoint;
+	TSharedRef<FActiveSyncPoint>* FoundData = BatchData.SyncPointData.Find(SyncData.RegisteredId);
+
+	if (!FoundData)
+	{
+		// Check conditions for creation and possibly return 0
+
+		TSharedRef<FActiveSyncPoint> NewTickFunction = MakeShared<FActiveSyncPoint>();
+		NewTickFunction->TickGroup = Description.FirstPossibleTickGroup;
+		NewTickFunction->EndTickGroup = Description.LastPossibleTickGroup;
+		NewTickFunction->FrameStatus = ESyncPointStatus::TaskNotRegistered;
+
+		NewTickFunction->SyncPointDescription = Description;
+		NewTickFunction->SyncPointId.BatchContext = BatchData.BatchContext;
+		NewTickFunction->SyncPointId.SyncId = SyncData.RegisteredId;
+
+		switch (Description.EventType)
+		{
+		case ESyncPointEventType::SimpleEvent:
+			// Could possibly be implemented as a base FGraphEvent with some other changes
+			NewTickFunction->bRunOnAnyThread = true;
+			NewTickFunction->bHighPriority = true;
+			break;
+		case ESyncPointEventType::GameThreadTask:
+			NewTickFunction->bRunOnAnyThread = false;
+			NewTickFunction->bHighPriority = false;
+			break;
+		case ESyncPointEventType::GameThreadTask_HighPriority:
+			NewTickFunction->bRunOnAnyThread = false;
+			NewTickFunction->bHighPriority = true;
+			break;
+		case ESyncPointEventType::WorkerThreadTask:
+			NewTickFunction->bRunOnAnyThread = true;
+			NewTickFunction->bHighPriority = false;
+			break;
+		case ESyncPointEventType::WorkerThreadTask_HighPriority:
+			NewTickFunction->bRunOnAnyThread = true;
+			NewTickFunction->bHighPriority = false;
+			break;
+		default:
+			ensureMsgf(0, TEXT("Invalid event type %d on sync point %s!"), Description.EventType, *Description.RegisteredName.ToString());
+		}
+
+		switch (Description.ActivationRules)
+		{
+		case ESyncPointActivationRules::AlwaysActivate:
+			NewTickFunction->bDispatchManually = false;
+			break;
+		case ESyncPointActivationRules::WaitForTrigger:
+		case ESyncPointActivationRules::WaitForAllWork:
+		case ESyncPointActivationRules::ActivateForAnyWork:
+			NewTickFunction->bDispatchManually = true;
+			break;
+		default:
+			ensureMsgf(0, TEXT("Invalid activation rules %d on sync point %s!"), Description.ActivationRules, *Description.RegisteredName.ToString());
+		}
+
+		FoundData = &BatchData.SyncPointData.Add(SyncData.RegisteredId, MoveTemp(NewTickFunction));
+	}
+
+	return &FoundData->Get();
+}
+
+void FTaskSyncManager::StartTickGroup(const UWorld* InWorld, ETickingGroup Group, TArray<FTickFunction*>& TicksToManualDispatch)
+{
+	check(IsInGameThread() && CurrentTickWorld == InWorld && CurrentTickGroup == ETickingGroup::TG_MAX);
+
+	UE::TScopeLock Lock(ManagerLock);
+	FWorldContextId WorldContext(InWorld);
+
+	CurrentTickGroup = Group;
+
+	// For all batches that match the world
+	for (FBatchData& BatchData : BatchList)
+	{
+		if (BatchData.WorldContext == WorldContext)
+		{
+			for (TPair<FSyncPointId::FInternalId, TSharedRef<FActiveSyncPoint>>& PointPair : BatchData.SyncPointData)
+			{
+				FActiveSyncPoint& SyncPoint = PointPair.Value.Get();
+				if (SyncPoint.GetDescription().FirstPossibleTickGroup == Group && SyncPoint.HandleFirstTickGroup())
+				{
+					TicksToManualDispatch.Add(&SyncPoint);
+				}
+			}
+		}
+	}
+}
+
+void FTaskSyncManager::EndTickGroup(const UWorld* InWorld, ETickingGroup Group)
+{
+	check(IsInGameThread() && CurrentTickWorld == InWorld && CurrentTickGroup != ETickingGroup::TG_MAX);
+	CurrentTickGroup = ETickingGroup::TG_MAX;
+}
+
+void FTaskSyncManager::EndFrame(const UWorld* InWorld)
+{
+	check(IsInGameThread() && CurrentTickWorld == InWorld && CurrentTickGroup == ETickingGroup::TG_MAX);
+	CurrentTickWorld = nullptr;
+	CurrentTickGroup = ETickingGroup::TG_MAX;
+}
+
+void FTaskSyncManager::ReleaseWorldContext(FWorldContextId WorldContext)
+{
+	// We're completely destroying a world context so free the data
+	if (ensure(IsInGameThread() && !IsTicking()))
+	{
+		for (FBatchData& BatchData : BatchList)
+		{
+			if (BatchData.WorldContext == WorldContext)
+			{
+				BatchData.Reset();
+			}
+		}
+	}
+}
+
+void FTaskSyncManager::OnWorldContextRemove(FWorldContext& InWorldContext)
+{
+	if (InWorldContext.WorldType != EWorldType::None && InWorldContext.WorldType != EWorldType::Inactive && InWorldContext.World())
+	{
+		ReleaseWorldContext(FWorldContextId(InWorldContext.World()));
+	}
+}
+
+FWorldContextId FTaskSyncManager::GetCurrentWorldContext() const
+{
+	// Does this make sense to be different per thread? Looking up the world is slow and unnecessary in cooked
+	if (ensure(IsInGameThread()))
+	{
+		// TODO Pie logic
+		return FWorldContextId(FWorldContextId::DefaultWorldContextId);
+	}
+	return FWorldContextId(FWorldContextId::InvalidWorldContextId);
+}
+
+FBatchContextId FTaskSyncManager::FindDefaultBatch(FWorldContextId WorldContext) const
+{
+	UE::TScopeLock ManagerScope(ManagerLock);
+
+	if (WorldContext.IsValid())
+	{
+		for (const FBatchData& BatchData : BatchList)
+		{
+			if (BatchData.WorldContext == WorldContext)
+			{
+				// First batch is the default one
+				return BatchData.BatchContext;
+			}
+		}
+	}
+
+	return FBatchContextId();
+}
+
+FBatchContextId FTaskSyncManager::CreateNewBatch(FWorldContextId WorldContext)
+{
+	if (!ensure(IsInGameThread() && WorldContext.IsValid()))
+	{
+		return FBatchContextId();
+	}
+
+	UE::TScopeLock ManagerScope(ManagerLock);
+
+	// First look for an unused slot
+	for (int32 InternalId = 0; InternalId < BatchList.Num(); InternalId++)
+	{
+		FBatchData& ExistingData = BatchList[InternalId];
+		if (!ExistingData.BatchContext.IsValid())
+		{
+			ExistingData.BatchContext.BatchId = InternalId;
+			ExistingData.WorldContext = WorldContext;
+			return ExistingData.BatchContext;
+		}
+	}
+	
+	int32 NewId = BatchList.AddDefaulted();
+	BatchList[NewId].BatchContext.BatchId = NewId;
+	BatchList[NewId].WorldContext = WorldContext;
+	return FBatchContextId(NewId);
+}
+
+FSyncPointId FTaskSyncManager::FindSyncPoint(FBatchContextId Batch, FName RegisteredName)
+{
+	UE::TScopeLock ManagerScope(ManagerLock);
+	FSyncPointId ReturnSyncPoint;
+	FSyncPointId::FInternalId* FoundId = RegisteredNameMap.Find(RegisteredName);
+	if (FoundId)
+	{
+		ReturnSyncPoint.SyncId = *FoundId;
+		ReturnSyncPoint.BatchContext = Batch;
+	}
+
+	return ReturnSyncPoint;
+}
+
+FSyncPointId FTaskSyncManager::FindSyncPoint(FWorldContextId WorldContext, FName RegisteredName)
+{
+	UE::TScopeLock ManagerScope(ManagerLock);
+	return FindSyncPoint(FindDefaultBatch(WorldContext), RegisteredName);
+}
+
+FTickFunction* FTaskSyncManager::GetTickFunctionForSyncPoint(FSyncPointId SyncPoint)
+{
+	if (!ensure(IsInGameThread()))
+	{
+		return nullptr;
+	}
+
+	UE::TScopeLock ManagerScope(ManagerLock);
+	TSharedPtr<FActiveSyncPoint> ActiveData;
+	FTaskSyncResult ReturnStatus = FindActiveSyncPoint(SyncPoint, ActiveData);
+
+	if (ReturnStatus)
+	{
+		// TODO think about this return value more
+		return ActiveData.Get();
+	}
+
+	return nullptr;
+}
+
+FTaskSyncResult FTaskSyncManager::GetTaskGraphEvent(FSyncPointId SyncPoint, FGraphEventRef& OutEventRef)
+{
+	UE::TScopeLock ManagerScope(ManagerLock);
+	TSharedPtr<FActiveSyncPoint> ActiveData;
+	FTaskSyncResult ReturnStatus = FindActiveSyncPoint(SyncPoint, ActiveData);
+	if (!ReturnStatus)
+	{
+		return ReturnStatus;
+	}
+
+	if (ReturnStatus.WasTaskCreatedForFrame())
+	{
+		// TODO should we let it return during execution?	
+		if (!ReturnStatus.WasActivatedForFrame())
+		{
+			check(ActiveData.IsValid() && ActiveData->IsCompletionHandleValid());
+			OutEventRef = ActiveData->GetCompletionHandle();
+
+			check(!OutEventRef->IsCompleted());
+
+			return ReturnStatus;
+		}
+	}
+
+	ReturnStatus.OperationResult = ESyncOperationResult::SyncPointStatusIncorrect;
+	return ReturnStatus;
+}
+
+FTaskSyncResult FTaskSyncManager::TriggerSyncPoint(FSyncPointId SyncPoint)
+{
+	UE::TScopeLock ManagerScope(ManagerLock);
+	TSharedPtr<FActiveSyncPoint> ActiveData;
+	FTaskSyncResult ReturnStatus = FindActiveSyncPoint(SyncPoint, ActiveData);
+	if (!ReturnStatus)
+	{
+		return ReturnStatus;
+	}
+
+	if (ReturnStatus.WasTaskCreatedForFrame() && !ReturnStatus.WasActivatedForFrame())
+	{
+		// Check type
+		check(ActiveData);
+		FSyncPointDescription& Description = ActiveData->SyncPointDescription;
+		if (Description.ActivationRules != ESyncPointActivationRules::WaitForTrigger)
+		{
+			ReturnStatus.OperationResult = ESyncOperationResult::ActivationRulesIncorrect;
+			return ReturnStatus;
+		}
+
+		ManagerScope.Unlock();
+
+		// TODO Make sure this lock logic is correct, we want to drop the manager lock before it runs the tick on this thread
+		UE::TScopeLock WorkScope(ActiveData->WorkLock);
+
+		if (ActiveData->DispatchManually())
+		{
+			// Return success and the current status
+			return FTaskSyncResult(ESyncOperationResult::Success, ActiveData->GetFrameStatus());
+		}
+	}
+
+	ReturnStatus.OperationResult = ESyncOperationResult::SyncPointStatusIncorrect;
+	return ReturnStatus;
+}
+
+FTaskSyncResult FTaskSyncManager::TriggerSyncPointAfterEvent(FSyncPointId SyncPoint, FGraphEventRef EventToWaitFor)
+{
+	UE::TScopeLock ManagerScope(ManagerLock);
+	TSharedPtr<FActiveSyncPoint> ActiveData;
+	FTaskSyncResult ReturnStatus = FindActiveSyncPoint(SyncPoint, ActiveData);
+	if (!ReturnStatus)
+	{
+		return ReturnStatus;
+	}
+
+	check(ActiveData.IsValid());
+	if (ReturnStatus.WasTaskCreatedForFrame() && !ReturnStatus.WasActivatedForFrame())
+	{
+		// Check type
+		check(ActiveData);
+		FSyncPointDescription& Description = ActiveData->SyncPointDescription;
+		if (Description.ActivationRules != ESyncPointActivationRules::WaitForTrigger)
+		{
+			ReturnStatus.OperationResult = ESyncOperationResult::ActivationRulesIncorrect;
+			return ReturnStatus;
+		}
+
+		if (ActiveData->CanDispatchManually())
+		{
+			ManagerScope.Unlock();
+
+			// TODO Make sure this lock logic is correct, we want to drop the manager lock before it runs the tick on this thread
+			UE::TScopeLock WorkScope(ActiveData->WorkLock);
+
+			// Can't use DontCompleteUntil on normal tasks that haven't started executing yet
+			//ActiveData->GetCompletionHandle()->DontCompleteUntil(EventToWaitFor);
+
+			ActiveData->GetCompletionHandle()->AddPrerequisites(*EventToWaitFor);
+
+			if (ActiveData->DispatchManually())
+			{
+				// Return success and the current status
+				return FTaskSyncResult(ESyncOperationResult::Success, ActiveData->GetFrameStatus());
+			}
+		}
+	}
+
+	ReturnStatus.OperationResult = ESyncOperationResult::SyncPointStatusIncorrect;
+	return ReturnStatus;
+}
+
+FTaskSyncResult FTaskSyncManager::RegisterWorkHandle(FSyncPointId SyncPoint, FActiveSyncWorkHandle& OutWorkHandle)
+{
+	UE::TScopeLock ManagerScope(ManagerLock);
+	TSharedPtr<FActiveSyncPoint> ActiveData;
+	FTaskSyncResult ReturnStatus = FindActiveSyncPoint(SyncPoint, ActiveData);
+	if (!ReturnStatus)
+	{
+		return ReturnStatus;
+	}
+
+	check(ActiveData.IsValid());
+	FSyncPointDescription& Description = ActiveData->SyncPointDescription;
+	if (Description.EventType == ESyncPointEventType::SimpleEvent)
+	{
+		ReturnStatus.OperationResult = ESyncOperationResult::EventTypeIncorrect;
+		return ReturnStatus;
+	}
+
+	// TODO what do we do about inactive/disabled tick functions?
+
+	OutWorkHandle = ActiveData->RegisterWorkHandle();
+
+	return ReturnStatus;
+}
+
+FTaskSyncResult FTaskSyncManager::FindActiveSyncPoint(FSyncPointId SyncPoint, TSharedPtr<FActiveSyncPoint>& OutData)
+{
+	// Internal only, does not lock
+	if (!SyncPoint.IsValid())
+	{
+		return ESyncOperationResult::SyncPointInvalid;
+	}
+	int32 BatchIndex = SyncPoint.BatchContext.BatchId;
+	if (BatchList.IsValidIndex(BatchIndex))
+	{
+		FBatchData& FoundBatch = BatchList[BatchIndex];
+		TSharedRef<FActiveSyncPoint>* FoundData = FoundBatch.SyncPointData.Find(SyncPoint.SyncId);
+		if (FoundData)
+		{
+			OutData = *FoundData;
+			
+			return FTaskSyncResult(ESyncOperationResult::Success, (*FoundData)->GetFrameStatus());
+		}
+		return ESyncOperationResult::SyncPointNotFound;
+
+	}
+	return ESyncOperationResult::BatchNotFound;
+}
diff --git a/Engine/Source/Runtime/Engine/Private/TickTaskManager.cpp b/Engine/Source/Runtime/Engine/Private/TickTaskManager.cpp
--- a/Engine/Source/Runtime/Engine/Private/TickTaskManager.cpp
+++ b/Engine/Source/Runtime/Engine/Private/TickTaskManager.cpp
@@ -14,6 +14,10 @@
 #include "Misc/TimeGuard.h"
 #include "Templates/UniquePtr.h"
 #include "UObject/RemoteExecutor.h"
+#include "TaskSyncManager.h"
+
+using namespace UE::Tasks;
+using namespace UE::Tick;
 
 DEFINE_LOG_CATEGORY_STATIC(LogTick, Log, All);
 
@@ -499,6 +503,9 @@ class FTickTaskSequencer
 	/** Completion handles for each phase of ticks, this is compatible with FGraphEventArray */
 	TArrayWithThreadsafeAdd<FGraphEventRef, TInlineAllocator<4> > TickCompletionEvents[TG_MAX];
 
+	/** Event-triggered tick functions indexed by end group, these functions must be triggered manually before the group ends */
+	TArrayWithThreadsafeAdd<FTickFunction*> ManualDispatchTicks[TG_MAX];
+
 	/** HiPri Held tasks for each tick group. */
 	TArrayWithThreadsafeAdd<FTickGraphTask*> HiPriTickTasks[TG_MAX][TG_MAX];
 
@@ -690,6 +697,7 @@ public:
 				TickTasks[TickGroup][EndTickGroup].Reserve(NumTicks);
 			}
 			TickCompletionEvents[TickGroup].Reserve(NumTicks);
+			ManualDispatchTicks[TickGroup].Reserve(NumTicks);
 		}
 	}
 
@@ -804,7 +812,17 @@ public:
 		FTickGraphTask* Task = TGraphTask<FTickFunctionTask>::CreateTask(Prerequisites, ENamedThreads::GameThread).ConstructAndHold(TickFunction, &UseContext);
 		TickFunction->SetTaskPointer(FTickFunction::ETickTaskState::HasTask, Task);
 
-		AddTickTaskCompletion(TickFunction->InternalData->ActualStartTickGroup, TickFunction->InternalData->ActualEndTickGroup, Task, TickFunction->bHighPriority);
+		if (TickFunction->bDispatchManually)
+		{
+			const ETickingGroup TickGroup = TickFunction->InternalData->ActualEndTickGroup;
+			ManualDispatchTicks[TickGroup].Add(TickFunction);
+			TickCompletionEvents[TickGroup].Add(Task->GetCompletionEvent());
+			TickFunction->bWasDispatchedManually = false;
+		}
+		else
+		{
+			AddTickTaskCompletion(TickFunction->InternalData->ActualStartTickGroup, TickFunction->InternalData->ActualEndTickGroup, Task, TickFunction->bHighPriority);
+		}
 	}
 
 	/**
@@ -820,15 +838,40 @@ public:
 		FTickGraphTask* Task = TGraphTask<FTickFunctionTask>::CreateTask(Prerequisites, ENamedThreads::GameThread).ConstructAndHold(TickFunction, &UseContext);
 		TickFunction->SetTaskPointer(FTickFunction::ETickTaskState::HasTask, Task);
 
-		AddTickTaskCompletionParallel(TickFunction->InternalData->ActualStartTickGroup, TickFunction->InternalData->ActualEndTickGroup, Task, TickFunction->bHighPriority);
+		if (TickFunction->bDispatchManually)
+		{
+			const ETickingGroup TickGroup = TickFunction->InternalData->ActualEndTickGroup;
+			ManualDispatchTicks[TickGroup].AddThreadsafe(TickFunction);
+			TickCompletionEvents[TickGroup].AddThreadsafe(Task->GetCompletionEvent());
+			TickFunction->bWasDispatchedManually = false;
+		}
+		else
+		{
+			AddTickTaskCompletionParallel(TickFunction->InternalData->ActualStartTickGroup, TickFunction->InternalData->ActualEndTickGroup, Task, TickFunction->bHighPriority);
+		}
+	}
+
+	/** Make sure all manual tick dispatches have happened to avoid deadlocks */
+	void VerifyManualDispatch(ETickingGroup WorldTickGroup)
+	{
+		for (FTickFunction* ManualDispatchTick : ManualDispatchTicks[WorldTickGroup])
+		{
+			if (ManualDispatchTick->CanDispatchManually())
+			{
+				// TODO We may want to warn if this task has prerequisites as that would lead to confusing behavior
+				
+				ensure(ManualDispatchTick->DispatchManually());
+			}
+		}
 	}
 
 	/**
 	 * Release the queued ticks for a given tick group and process them.
 	 * @param WorldTickGroup - tick group to release
 	 * @param bBlockTillComplete - if true, do not return until all ticks are complete
+	 * @param TickFunctionsToManualDispatch - dispatch any manual tick functions in this lister after the normal ones
 	**/
-	void ReleaseTickGroup(ETickingGroup WorldTickGroup, bool bBlockTillComplete)
+	void ReleaseTickGroup(ETickingGroup WorldTickGroup, bool bBlockTillComplete, TArray<FTickFunction*>& TicksToManualDispatch)
 	{
 		if (bLogTicks)
 		{
@@ -850,6 +893,13 @@ public:
 			}
 		}
 
+		for (FTickFunction* TickToManualDispatch : TicksToManualDispatch)
+		{
+			// These must be safe to dispatch
+			ensure(TickToManualDispatch->DispatchManually());
+		}
+		TicksToManualDispatch.Reset();
+
 		if (bBlockTillComplete || bSingleThreadMode)
 		{
 			SCOPE_CYCLE_COUNTER(STAT_ReleaseTickGroup_Block);
@@ -859,7 +909,16 @@ public:
 				if (TickCompletionEvents[Block].Num())
 				{
 					TRACE_CPUPROFILER_EVENT_SCOPE(TickCompletionEvents);
-					FTaskGraphInterface::Get().ProcessUntilTasksComplete(TickCompletionEvents[Block], ENamedThreads::GameThread);
+
+					// Process tasks, and on first stall make sure all manual dispatch events have been dispatched
+					// TODO change this to a deadlock detector somehow? In theory this could happen too early if there are very few game thread ticks
+					// But right now it just runs directly inside the function as soon as the game thread runs out of work to do
+					FTaskGraphInterface::Get().ProcessUntilTasksComplete(TickCompletionEvents[Block], ENamedThreads::GameThread,
+						[this, WorldTickGroup]()
+						{
+							this->VerifyManualDispatch(WorldTickGroup);
+						});
+
 					if (bSingleThreadMode || Block == TG_NewlySpawned || CVarAllowAsyncTickCleanup.GetValueOnGameThread() == 0 || TickCompletionEvents[Block].Num() < 50)
 					{
 						ResetTickGroup(Block);
@@ -915,6 +974,7 @@ public:
 		{
 			check(!TickCompletionEvents[Index].Num());  // we should not be adding to these outside of a ticking proper and they were already cleared after they were ticked
 			TickCompletionEvents[Index].Reset();
+			ManualDispatchTicks[Index].Reset();
 			for (int32 IndexInner = 0; IndexInner < TG_MAX; IndexInner++)
 			{
 				check(!TickTasks[Index][IndexInner].Num() && !HiPriTickTasks[Index][IndexInner].Num());  // we should not be adding to these outside of a ticking proper and they were already cleared after they were ticked
@@ -977,6 +1037,7 @@ private:
 	{
 		QUICK_SCOPE_CYCLE_COUNTER(STAT_ResetTickGroup);
 		TickCompletionEvents[WorldTickGroup].Reset();
+		ManualDispatchTicks[WorldTickGroup].Reset();
 	}
 
 	void DispatchTickGroup(ENamedThreads::Type CurrentThread, ETickingGroup WorldTickGroup)
@@ -1812,6 +1873,12 @@ public:
 			FPlatformProcess::Sleep(CVarStallStartFrame.GetValueOnGameThread() / 1000.0f);
 		}
 #endif
+		if (FTaskSyncManager* SyncManager = FTaskSyncManager::Get())
+		{
+			// This can create tick functions
+			SyncManager->StartFrame(InWorld);
+		}
+
 		Context.TickGroup = ETickingGroup(0); // reset this to the start tick group
 		Context.DeltaSeconds = InDeltaSeconds;
 		Context.TickType = InTickType;
@@ -1913,7 +1980,16 @@ public:
 	{
 		check(Context.TickGroup == Group); // this should already be at the correct value, but we want to make sure things are happening in the right order
 		check(bTickNewlySpawned); // we should be in the middle of ticking
-		TickTaskSequencer.ReleaseTickGroup(Group, bBlockTillComplete);
+
+		TArray<FTickFunction*> TicksToManualDispatch;
+		FTaskSyncManager* SyncManager = FTaskSyncManager::Get();
+
+		if (SyncManager)
+		{
+			SyncManager->StartTickGroup(Context.World, Group, TicksToManualDispatch);
+		}
+
+		TickTaskSequencer.ReleaseTickGroup(Group, bBlockTillComplete, TicksToManualDispatch);
 		Context.TickGroup = ETickingGroup(Context.TickGroup + 1); // new actors go into the next tick group because this one is already gone
 		if (bBlockTillComplete) // we don't deal with newly spawned ticks within the async tick group, they wait until after the async stuff
 		{
@@ -1930,7 +2006,7 @@ public:
 				if (Num && Context.TickGroup == TG_NewlySpawned)
 				{
 					SCOPE_CYCLE_COUNTER(STAT_TG_NewlySpawned);
-					TickTaskSequencer.ReleaseTickGroup(TG_NewlySpawned, true);
+					TickTaskSequencer.ReleaseTickGroup(TG_NewlySpawned, true, TicksToManualDispatch);
 				}
 				else
 				{
@@ -1947,6 +2023,11 @@ public:
 				}
 			}
 		}
+
+		if (SyncManager)
+		{
+			SyncManager->EndTickGroup(Context.World, Group);
+		}
 	}
 
 	/** Finish a frame of ticks **/
@@ -1958,6 +2039,13 @@ public:
 		{
 			LevelList[LevelIndex]->EndFrame();
 		}
+
+		FTaskSyncManager* SyncManager = FTaskSyncManager::Get();
+		if (SyncManager)
+		{
+			SyncManager->EndFrame(Context.World);
+		}
+
 		Context.World = nullptr;
 		LevelList.Reset();
 	}
@@ -1994,6 +2082,9 @@ private:
 		, bTickNewlySpawned(false)
 	{
 		IConsoleManager::Get().RegisterConsoleCommand(TEXT("dumpticks"), TEXT("Dumps all tick functions registered with FTickTaskManager to log."));
+		
+		// Create the task sync manager if it is needed later
+		FTaskSyncManager* CreatedSyncManager = FTaskSyncManager::Get();
 	}
 
 	/** Fill the level list, only with levels that can actually tick */
@@ -2137,6 +2228,8 @@ FTickFunction::FTickFunction()
 	, bHighPriority(false)
 	, bRunOnAnyThread(false)
 	, bRunTransactionally(false)
+	, bDispatchManually(false)
+	, bWasDispatchedManually(false)
 	, TickState(ETickState::Enabled)
 	, TickInterval(0.f)
 {
@@ -2321,6 +2414,33 @@ FGraphEventRef FTickFunction::GetCompletionHandle() const
 	return Task->GetCompletionEvent();
 }
 
+bool FTickFunction::CanDispatchManually() const
+{
+	if (bDispatchManually && !bWasDispatchedManually && InternalData && InternalData->TaskState == ETickTaskState::HasTask)
+	{
+		return true;
+	}
+	return false;
+}
+
+bool FTickFunction::DispatchManually()
+{
+	if (ensure(bDispatchManually) && CanDispatchManually())
+	{
+		FTickGraphTask* Task = FTickTaskSequencer::GetGraphTask(this);
+		check(Task);
+
+		bWasDispatchedManually = true;
+				
+		// This could run the tick function and invalidate anything
+		Task->Unlock();
+
+		return true;
+	}
+
+	return false;
+}
+
 void FTickFunction::ExecuteNestedTick(float DeltaTime, ELevelTick TickType, ENamedThreads::Type CurrentThread, const FGraphEventRef& MyCompletionGraphEvent)
 {
 	// Not valid to call on a real scheduled function
@@ -2342,17 +2462,19 @@ void FTickFunction::ClearTaskInformation()
 
 void FTickFunction::SetTaskPointer(ETickTaskState NewState, void* InTaskPointer)
 {
-	check(InternalData);
-
-	if (NewState == ETickTaskState::NotQueued)
+	// If this is called on a tick function that was never registered, we don't want to allocate memory or crash
+	if (InternalData)
 	{
-		InternalData->TaskState = ETickTaskState::NotQueued;
-		InternalData->TaskPointer = nullptr;
-	}
-	else
-	{
-		InternalData->TaskState = NewState;
-		InternalData->TaskPointer = InTaskPointer;
+		if (NewState == ETickTaskState::NotQueued)
+		{
+			InternalData->TaskState = ETickTaskState::NotQueued;
+			InternalData->TaskPointer = nullptr;
+		}
+		else
+		{
+			InternalData->TaskState = NewState;
+			InternalData->TaskPointer = InTaskPointer;
+		}
 	}
 }
 
diff --git a/Engine/Source/Runtime/Engine/Private/UnrealEngine.cpp b/Engine/Source/Runtime/Engine/Private/UnrealEngine.cpp
--- a/Engine/Source/Runtime/Engine/Private/UnrealEngine.cpp
+++ b/Engine/Source/Runtime/Engine/Private/UnrealEngine.cpp
@@ -16724,9 +16724,8 @@ void UEngine::DestroyWorldContext(UWorld * InWorld)
 	{
 		if (WorldList[idx].World() == InWorld)
 		{
-#if WITH_EDITOR
 			WorldContextDestroyedEvent.Broadcast(WorldList[idx]);
-#endif
+
 			// Set the current world to NULL so that any external referencers are cleaned up before we remove
 			WorldList[idx].SetCurrentWorld(NULL);
 			WorldList.RemoveAt(idx);
diff --git a/Engine/Source/Runtime/Engine/Public/TaskSyncManager.h b/Engine/Source/Runtime/Engine/Public/TaskSyncManager.h
new file mode 100644
--- /dev/null
+++ b/Engine/Source/Runtime/Engine/Public/TaskSyncManager.h
@@ -0,0 +1,754 @@
+// Copyright Epic Games, Inc. All Rights Reserved.
+
+#pragma once
+
+#include "Engine/EngineBaseTypes.h"
+#include "Misc/TransactionallySafeCriticalSection.h"
+#include "Engine/DeveloperSettings.h"
+#include "Async/TaskGraphFwd.h"
+#include "Delegates/Delegate.h"
+
+#include "TaskSyncManager.generated.h"
+
+// NOTE: This header file is considered Experimental and will remain so until the release of 5.6. It will be split into multiple headers
+
+/** The type of event to execute at the sync point */
+UENUM()
+enum class ESyncPointEventType : uint8
+{
+	Invalid,
+
+	/** A simple event that cannot activate any code directly */
+	SimpleEvent,
+
+	/** A task that executes code on the game thread, can be used for batching */
+	GameThreadTask,
+
+	/** High priority game thread task, will run before normal ticks */
+	GameThreadTask_HighPriority,
+
+	/** A task that executes code on a worker thread, can be used for batching */
+	WorkerThreadTask,
+
+	/** A task that executes code on a worker thread, can be used for batching */
+	WorkerThreadTask_HighPriority,
+};
+
+/** Rules for when a sync point's task will be activated/dispatched during a frame. Execution will also need to wait on any task dependencies */
+UENUM()
+enum class ESyncPointActivationRules : uint8
+{
+	Invalid,
+
+	/** Always activate, dispatch during FirstPossibleTickGroup */
+	AlwaysActivate,
+
+	/** Triggered once manually with TriggerSyncPoint, or as a backup from LastPossibleTickGroup */
+	WaitForTrigger,
+
+	/** Trigger once when there is any requested work and no reserved work */
+	WaitForAllWork,
+
+	/** Can activate multiple times per frame, whenever there is any requested work to perform */
+	ActivateForAnyWork,
+};
+
+/**
+ * Description of a registered sync event that will happen once per frame according to specific rules
+ */
+USTRUCT()
+struct FSyncPointDescription
+{
+	GENERATED_BODY()
+
+	/** Name of the sync point that will be used for lookup */
+	UPROPERTY(EditDefaultsOnly, Category="Task")
+	FName RegisteredName;
+
+	/** Name of what added this description, default means it was loaded from settings */
+	FName SourceName;
+
+	/** The kind of task async task that is used to implement this sync point */
+	UPROPERTY(EditDefaultsOnly, Category="Task")
+	ESyncPointEventType EventType = ESyncPointEventType::Invalid;
+
+	/** Rules for when and how a sync point can be activated */
+	UPROPERTY(EditDefaultsOnly, Category="Task")
+	ESyncPointActivationRules ActivationRules = ESyncPointActivationRules::Invalid;
+
+	/** The first tick group this could be triggered during */
+	UPROPERTY(EditDefaultsOnly, Category="Task")
+	TEnumAsByte<enum ETickingGroup> FirstPossibleTickGroup = TG_PrePhysics;
+
+	/** The last possible tick group this will be triggered during, and when it will be forcibly triggered as a backup */
+	UPROPERTY(EditDefaultsOnly, Category="Task")
+	TEnumAsByte<enum ETickingGroup> LastPossibleTickGroup = TG_LastDemotable;
+
+	/** Array of other sync groups that this will tick after. To match the normal tick behavior these will be ignored if thee sync point is not active */
+	UPROPERTY(EditDefaultsOnly, Category="Task")
+	TArray<FName> PrerequisiteSyncGroups;
+
+	/** True if this valid and has a name */
+	inline bool IsValid() const
+	{
+		return !RegisteredName.IsNone();
+	}
+
+	/** True if this was loaded from settings and has a default source */
+	inline bool WasLoadedFromSettings() const 
+	{
+		return SourceName.IsNone();
+	}
+};
+
+/** Settings for synchronizing tasks and ticking across the engine */
+UCLASS(config = Engine, defaultconfig, meta = (DisplayName = "Task Synchronization"), MinimalAPI)
+class UTaskSyncManagerSettings : public UDeveloperSettings
+{
+	GENERATED_BODY()
+
+#if WITH_EDITOR
+	virtual void PostEditChangeProperty(FPropertyChangedEvent& PropertyChangedEvent) override;
+	virtual bool SupportsAutoRegistration() const override;
+#endif
+
+public:
+	/** List of all registered task events */
+	UPROPERTY(config, EditAnywhere, Category = "Task Synchronization")
+	TArray<FSyncPointDescription> RegisteredSyncPoints;
+};
+
+namespace UE::Tick
+{
+
+/** Enum describing the success/failure of a task sync operation */
+enum class ESyncOperationResult : uint8
+{
+	/** Unknown or invalid status */
+	Invalid,
+
+	/** The requested operation was successful */
+	Success,
+
+	// Error conditions TODO some of these are redundant
+
+	/** The FSyncPointId is completely invalid */
+	SyncPointInvalid,
+	/** Sync point is not registered */
+	SyncPointNotRegistered,
+	/** Sync point does not exist for the specific batch */
+	SyncPointNotFound,
+	/** Sync point was specifically disabled */
+	SyncPointDisabled,
+	
+	/** Sync Point status is incorrect for the requested operation, check the status */
+	SyncPointStatusIncorrect,
+	/** The sync point's event type does not support this operation, such as trying to add a tick function to a simple event */
+	EventTypeIncorrect,
+	/** This is not supported for the sync point's activation rules */
+	ActivationRulesIncorrect,
+
+	/** BatchContextId was not found */
+	BatchNotFound,
+	/** WorldContextId was not found */
+	WorldNotFound,
+	/** Operation called on incorrect thread or it cannot determine thread context */
+	ThreadIncorrect,
+};
+
+/** Enum describing the current status of a specific sync point in the current frame */
+enum class ESyncPointStatus : uint8
+{
+	/** Status of sync point is not known, probably because it could not be found */
+	Unknown,
+
+	/** Sync point tick function has not yet been registered with tick system */
+	TaskNotRegistered,
+	/** Sync point tick is registered but no task has been made yet, this is the state between tick frames */
+	TaskNotCreated,
+	/** Task has been created but not yet dispatched. It may be before the first tick group */
+	TaskCreated,
+
+	/** Has not been dispatched, waiting for a trigger */
+	DispatchWaitingForTrigger,
+	/** Has not been dispatched, waiting for work requests */
+	DispatchWaitingForWork,
+	
+	/** Has been dispatched, but has not started execution. This could be waiting on a prerequisite task */
+	Dispatched,
+	/** Was dispatched and has started execution */
+	Executing,
+	/** A follow up task has been dispatched to waiting for more work before executing again */
+	WaitingForMoreWork,
+	/** Completely done executing for the frame */
+	ExecutionComplete,
+};
+
+/** 
+ * Result structure returned from all task sync manager operations.
+ * This holds an error code as well as the status of the sync point at operation time, if known.
+ */
+struct FTaskSyncResult
+{
+	ESyncOperationResult OperationResult;
+	ESyncPointStatus SyncStatus;
+
+	/** Invalid result by default */
+	FTaskSyncResult()
+		: OperationResult(ESyncOperationResult::Invalid)
+		, SyncStatus(ESyncPointStatus::Unknown)
+	{}
+
+	/** Create from an operation result, defaulting to an unknown sync point status */
+	FTaskSyncResult(ESyncOperationResult Result, ESyncPointStatus Status = ESyncPointStatus::Unknown)
+		: OperationResult(Result)
+		, SyncStatus(Status)
+	{}
+
+	/** True if the result has been initialized by an operation */
+	inline bool IsValid() const
+	{
+		return OperationResult != ESyncOperationResult::Invalid;
+	}
+
+	/** True if the operation was successful */
+	inline bool WasSuccessful() const
+	{
+		return OperationResult == ESyncOperationResult::Success;
+	}
+
+	/** Treat successful as true for if statements */
+	inline operator bool() const
+	{
+		return WasSuccessful();
+	}
+
+	/** True if a low level task was created this frame for this sync point. The task may have been destroyed if it has completed */
+	inline bool WasTaskCreatedForFrame() const
+	{
+		return SyncStatus >= ESyncPointStatus::TaskCreated;
+	}
+
+	/** True if the sync point has already been dispatched for this frame */
+	inline bool WasActivatedForFrame() const 
+	{
+		return SyncStatus >= ESyncPointStatus::Dispatched;
+	}
+};
+
+
+/** Identifies a specific WorldContext that has objects which are allowed to tick and run tasks */
+struct FWorldContextId final
+{
+	// TODO Consider moving this concept into World.h
+	typedef int32 FInternalId;
+
+	static constexpr FInternalId InvalidWorldContextId = -1;
+
+	/** Default id which is used for the default game world outside the editor */
+	static constexpr FInternalId DefaultWorldContextId = 0;
+
+	/** Construct from a world pointer, if it is null or not part of a tickable world context it will return an invalid id */
+	static ENGINE_API FWorldContextId GetContextIdForWorld(const UWorld* World);
+
+	/** Construct an invalid world context id */
+	FWorldContextId()
+		: WorldId(InvalidWorldContextId)
+	{}
+
+	/** Construct from an internal id TODO: May make private */
+	explicit FWorldContextId(FInternalId InWorldId)
+		: WorldId(InWorldId)
+	{}
+
+	/** Construct from a world pointer, if it is null or not part of a tickable world context it will return an invalid id */
+	FWorldContextId(const UWorld* InWorld)
+	{
+		WorldId = GetContextIdForWorld(InWorld).WorldId;
+	}
+
+	/** Returns true if this points to a valid tickable world context */
+	inline bool IsValid() const
+	{
+		return WorldId > InvalidWorldContextId;
+	}
+
+	inline bool operator==(const FWorldContextId& Other) const
+	{
+		return WorldId == Other.WorldId;
+	}
+
+	inline bool operator!=(const FWorldContextId& Other) const
+	{
+		return !(*this == Other);
+	}
+
+	friend inline uint32 GetTypeHash(const FWorldContextId& Context)
+	{
+		return static_cast<uint32>(Context.WorldId);
+	}
+
+	FInternalId WorldId;
+};
+
+/** Identifies a specific batch of tasks, which will be associated with a set of objects in a world context */
+struct FBatchContextId final
+{
+	// Internal type is sometimes used as an array index
+	typedef int32 FInternalId;
+
+	static constexpr FInternalId InvalidBatch = -1;
+
+	/** Construct an invalid batch */
+	FBatchContextId()
+		: BatchId(InvalidBatch)
+	{}
+
+	/** Construct from an internal id TODO: May make private */
+	explicit FBatchContextId(FInternalId InBatchId)
+		: BatchId(InBatchId)
+	{}
+
+	/** Returns true if this has been initialized to a valid batch */
+	inline bool IsValid() const
+	{
+		return BatchId > InvalidBatch;
+	}
+
+	inline bool operator==(const FBatchContextId& Other) const
+	{
+		return BatchId == Other.BatchId;
+	}
+
+	inline bool operator!=(const FBatchContextId& Other) const
+	{
+		return !(*this == Other);
+	}
+
+	friend inline uint32 GetTypeHash(const FBatchContextId& Context)
+	{
+		return static_cast<uint32>(Context.BatchId);
+	}
+
+	FInternalId BatchId;
+};
+
+/** Identifies a specific registered sync point and batch, this can be safely passed around threads */
+struct FSyncPointId final
+{
+	// Internal type is sometimes used as an array index
+	typedef int32 FInternalId;
+
+	static constexpr FInternalId InvalidSyncPoint = -1;
+
+	/** Construct an invalid context */
+	FSyncPointId()
+		: SyncId(InvalidSyncPoint)
+	{}
+
+	/** Construct from explicit ids TODO: May make private */
+	explicit FSyncPointId(FInternalId InSyncId, FBatchContextId InBatchContext)
+		: SyncId(InSyncId)
+		, BatchContext(InBatchContext)
+	{}
+
+	/** Returns true if this has been initialized to a registered sync point */
+	inline bool IsValid() const
+	{
+		return SyncId > InvalidSyncPoint;
+	}
+
+	/** Returns the batch context this was initialized with */
+	inline FBatchContextId GetBatchContext() const
+	{
+		return BatchContext;
+	}
+
+	inline bool operator==(const FSyncPointId& Other) const
+	{
+		return BatchContext == Other.BatchContext && SyncId == Other.SyncId;
+	}
+
+	inline bool operator!=(const FSyncPointId& Other) const
+	{
+		return !(*this == Other);
+	}
+
+	friend inline uint32 GetTypeHash(const FSyncPointId& SyncPoint)
+	{
+		return GetTypeHash(SyncPoint.BatchContext) + SyncPoint.SyncId;
+	}
+
+	FInternalId SyncId;
+	FBatchContextId BatchContext;
+};
+
+/** Used to specify how many times work should be executed as part of a sync point */
+// TODO May need to change this enum to deal with in-frame updates
+enum class ESyncWorkRepetition : uint8
+{
+	/** Do not perform this work, used to cancel previous requests */
+	Never,
+
+	/** Work will be reserved or requested once per call to this function, and will reset for the next frame */
+	Once,
+
+	/** Work will be reserved or requested once every frame until it is abandoned, including the current frame if possible */
+	EveryFrame,
+};
+
+/** 
+ * Handle pointing to a FActiveSyncPoint that can be used to reserve or request work to execute as part of that sync point.
+ * These handles cannot be copied, but can be safely moved between threads and will cancel all reservations and requests on destruction.
+ */
+struct FActiveSyncWorkHandle final
+{
+	FActiveSyncWorkHandle() = default;
+	ENGINE_API ~FActiveSyncWorkHandle();
+
+	FActiveSyncWorkHandle(FActiveSyncWorkHandle&&) = default;
+	FActiveSyncWorkHandle& operator=(FActiveSyncWorkHandle&&) = default;
+
+	// These handles can be moved but not copied
+	FActiveSyncWorkHandle(const FActiveSyncWorkHandle&) = delete;
+	FActiveSyncWorkHandle& operator=(const FActiveSyncWorkHandle&) = delete;
+
+	/** True if this points to a real sync point */
+	ENGINE_API bool IsValid() const;
+
+	/**
+	 * Returns a raw pointer to the tick function that implements this handle which can be passed into AddPrerequisite. 
+	 * This is only valid to call on the game thread and should not be cached.
+	 */
+	ENGINE_API FTickFunction* GetDependencyTickFunction();
+
+	/** True if this handle has been used to reserve work. This will still be true after work has been requested */
+	ENGINE_API bool HasReservedWork() const;
+
+	/** True if this handle has been used to request work. This will still be true during and after it is executed */
+	ENGINE_API bool HasRequestedWork() const;
+
+	/** Reserve work using this handle, possibly for every frame */
+	ENGINE_API bool ReserveFutureWork(ESyncWorkRepetition Repeat);
+
+	/**
+	 * Requests that a tick function be executed by the sync point, this fills any reservations and may trigger tasks.
+	 * If EveryFrame is passed in, it will request the same function every frame until it is abandoned.
+	 */
+	ENGINE_API bool RequestWork(FTickFunction* FunctionToExecute, ESyncWorkRepetition Repeat);
+
+	/** 
+	 * Abandon any requested or reserved work and return to how it was at initial registration.
+	 * This will not stop any requested work that has already started on another thread.
+	 * This clears HasReserved/HasRequested, but does not invalidate the handle.
+	 * Return true if there was any requested/reserved work to abandon.
+	 */
+	ENGINE_API bool AbandonWork();
+
+	/** Completely reset this handle, which will abandon work and prevent any future use. This returns true if it was valid before */
+	ENGINE_API bool Reset();
+
+private:
+	/** The sync point this handle is from */
+	TSharedPtr<struct FActiveSyncPoint> SyncPoint;
+
+	/** Index into the ActiveWork array */
+	uint32 WorkIndex : 24 = 0;
+
+	/** Max amount of work per sync point */
+	static constexpr uint32 MaxWorkIndex = (2 << 24) - 1;
+
+	// TODO May remove these flags or add requested type
+
+	/** True if this is a reserved work handle that needs to be requested or abandoned */
+	uint32 bWorkReserved : 1 = 0;
+
+	/** True if this is a requested work handle that will be executed */
+	uint32 bWorkRequested : 1 = 0;
+
+	/** Reset a handle to point to nothing */
+	inline void ResetInternal()
+	{
+		SyncPoint.Reset();
+		WorkIndex = 0;
+		bWorkReserved = false;
+		bWorkRequested = false;
+	}
+
+	friend struct FActiveSyncPoint;
+	friend class FTaskSyncManager;
+};
+
+/** Struct defining an active sync point inside a specific world context, implemented as a tick function that may not be registered */
+struct FActiveSyncPoint final : public TSharedFromThis<FActiveSyncPoint, ESPMode::ThreadSafe>, private FTickFunction
+{
+	ENGINE_API FActiveSyncPoint();
+
+	/** Returns the registered data  */
+	inline const FSyncPointDescription& GetDescription() const { return SyncPointDescription; }
+
+	/** Returns the sync point id that created this, including the current batch context */
+	inline const FSyncPointId& GetSyncPointId() const { return SyncPointId; }
+
+	/** Returns current status */
+	inline ESyncPointStatus GetFrameStatus() const { return FrameStatus; }
+
+	/** Registers a new work handle that can be used to reserve and request work */
+	ENGINE_API FActiveSyncWorkHandle RegisterWorkHandle();
+
+	/** Uses handle to reserve work for later, that can be requested when ready */
+	ENGINE_API bool ReserveFutureWork(FActiveSyncWorkHandle& Handle, ESyncWorkRepetition Repeat);
+
+	/** Creates a new item of work bound to a specific function */
+	ENGINE_API bool RequestWork(FActiveSyncWorkHandle& Handle, FTickFunction* WorkFunction, ESyncWorkRepetition Repeat);
+
+	/** Abandons requested or reserved work, which could trigger other work to start */
+	ENGINE_API bool AbandonWork(FActiveSyncWorkHandle& Handle);
+
+	/** Resets a handle */
+	ENGINE_API bool ResetWorkHandle(FActiveSyncWorkHandle& Handle);
+
+private:
+
+	/** Struct defining a unit of work that is executed as part of an FActiveSyncPoint */
+	struct FActiveSyncWork
+	{
+		union
+		{
+			/** Default state of 0 means it can be reused immediately */
+			uint8 StateByte;
+
+			struct
+			{
+				/** If true, this work corresponds to an active handle */
+				uint8 bHasActiveHandle : 1;
+
+				/** If true, this work has been reserved for execution later in the frame */
+				uint8 bWorkReserved : 1;
+
+				/** If true, this work can be executed this frame */
+				uint8 bWorkRequested : 1;
+
+				/** If true, work will be reserved at the start of every frame */
+				uint8 bReserveEveryFrame : 1;
+
+				/** If true, work will be requested at the start of every frame */
+				uint8 bRequestEveryFrame : 1;
+
+				/** If true, this work is in the middle of being executed and is in the ExecutingWork array */
+				uint8 bCurrentlyExecuting : 1;
+
+				/** If true, this work has completely finished executing this frame */
+				uint8 bAlreadyExecuted : 1;
+			};
+		};
+
+		/** Function to actually execute */
+		FTickFunction* TickFunction;
+
+		FActiveSyncWork()
+			: StateByte(0)
+			, TickFunction(nullptr)
+		{
+		}
+
+		void Reset()
+		{
+			TickFunction = nullptr;
+			StateByte = 0;
+		}
+
+		bool IsInitialized() const
+		{
+			return StateByte != 0;
+		}
+	};
+
+
+	// FTickFunction interface
+	virtual void ExecuteTick(float DeltaTime, ELevelTick TickType, ENamedThreads::Type CurrentThread, const FGraphEventRef& MyCompletionGraphEvent) override;
+	virtual FString DiagnosticMessage() override;
+	virtual FName DiagnosticContext(bool bDetailed) override;
+	virtual bool DispatchManually() override;
+
+	/** Resets tick function at start of frame, this can automatically reserve/request work */
+	void ResetWorkForFrame();
+
+	/** Called on the first tick group to set status properly. Returns true if this function should be dispatched immediately */
+	bool HandleFirstTickGroup();
+
+	/** Returns true if this sync point is ready to trigger due to work */
+	bool IsReadyToProcessWork();
+
+	/** Handle copying of work into / out of ExecutingWork array, returns true if there is anything to do */
+	bool GetWorkToExecute(int32& NumReserved);
+
+	/** Gets index of new work slot, could resize array so must be called locked and not during iteration */
+	uint32 AllocateActiveWork();
+
+	/** Dispatches the task to handle work, this should be called outside the work lock */
+	void DispatchWorkTask();
+
+	/** Registered sync point that created this */
+	FSyncPointDescription SyncPointDescription;
+	
+	/** Batch this is inside */
+	FSyncPointId SyncPointId;
+
+	/** Current state of this sync point in the current frame */
+	ESyncPointStatus FrameStatus;
+
+	/** Event to signal if there is work to process after initial activation */
+	FGraphEventRef ReactivationEvent;
+
+	/** Lock to handle access to work array and other internal structs */
+	mutable FTransactionallySafeCriticalSection WorkLock;
+
+	/** Array of work to execute as part of tick */
+	TArray<FActiveSyncWork> ActiveWork;
+
+	/** Second buffer of work used during active execution */
+	TArray<FActiveSyncWork> ExecutingWork;
+
+	friend class FTaskSyncManager;
+	friend struct FActiveSyncWorkHandle;
+};
+
+
+/**
+ * Global singleton manager that can be used to synchronize tasks across different engine systems.
+ * This provides event registration and utility functions that wrap functionality in the base Task and TaskGraph systems.
+ */
+class FTaskSyncManager
+{
+public:
+	ENGINE_API FTaskSyncManager();
+	ENGINE_API ~FTaskSyncManager();
+	
+
+	// Thread safe accessors, safe at any time
+
+	/** Return the global singleton if it exists */
+	static ENGINE_API FTaskSyncManager* Get();
+
+	/** Gets the default batch for a world context. This will return an invalid id if the world context is not set up for ticking */
+	ENGINE_API FBatchContextId FindDefaultBatch(FWorldContextId WorldContext) const;
+
+	/** 
+	 * Initializes a sync point id that can be passed between threads and used in the functions below.
+	 * This will return an invalid id if the name is not currently registered.
+	 * This function does not check that the sync point is enabled for the specific batch, that is handled in the other functions below.
+	 */
+	ENGINE_API FSyncPointId FindSyncPoint(FBatchContextId Batch, FName RegisteredName);
+
+	/** Same as above, but using the default batch for the world context */
+	ENGINE_API FSyncPointId FindSyncPoint(FWorldContextId WorldContext, FName RegisteredName);
+
+
+	// These functions are thread safe and can be called by worker or game threads during a tick, if the sync point is in the proper state
+
+	/** 
+	 * Attempts to return the current frame's task graph event for this sync point, to pass as a dependency to other tasks.
+	 * This will fail if the task has not yet been created this frame or it has finished executing.
+	 */
+	ENGINE_API FTaskSyncResult GetTaskGraphEvent(FSyncPointId SyncPoint, FGraphEventRef& OutEventRef);
+
+	/** Manually triggers a sync point, this can only be called once per frame */
+	ENGINE_API FTaskSyncResult TriggerSyncPoint(FSyncPointId SyncPoint);
+
+	/** Tells the sync point to trigger at the completion of the passed in event/task, this is only possible if it hasn't already been triggered */
+	ENGINE_API FTaskSyncResult TriggerSyncPointAfterEvent(FSyncPointId SyncPoint, FGraphEventRef EventToWairFor);
+
+	/** 
+	 * Tries to creates a new sync work handle that can be used to reserve or request work to happen as part of a sync point
+	 * If it was successful, OutWorkHandle will be set to a handle you can pass between threads and later fulfill or abandon.
+	 */
+	ENGINE_API FTaskSyncResult RegisterWorkHandle(FSyncPointId SyncPoint, FActiveSyncWorkHandle& OutWorkHandle);
+
+
+	// Game thread-only functions, in general these will not affect the current frame
+
+	/** Refresh the registered data from settings */
+	ENGINE_API void ReloadRegisteredData();
+
+	/** Searches for a registered sync point, if ones is foundw ith that name it will return true and fill in OutDescription */
+	ENGINE_API bool GetSyncPointDescription(FName RegisteredName, FSyncPointDescription& OutDescription) const;
+
+	/** Registers a new sync point at runtime. RegisteredName and SourceName must be filled out */
+	ENGINE_API bool RegisterNewSyncPoint(const FSyncPointDescription& NewDescription);
+
+	/** Unregisters a sync point, this will only delete if if both the registered and source names match */
+	ENGINE_API bool UnregisterSyncPoint(FName RegisteredName, FName SourceName);
+
+	/** Allocates a new batch for the specified world, this will duplicate the events */
+	ENGINE_API FBatchContextId CreateNewBatch(FWorldContextId WorldContext);
+
+	/** Gets the correct task world context for the current thread context. This will return an invalid context if called on worker threads */
+	ENGINE_API FWorldContextId GetCurrentWorldContext() const;
+
+	/** Returns the tick function representing the specified event, if this returns a valid pointer you can use it for setting dependencies but it should not be modified or stored */
+	ENGINE_API FTickFunction* GetTickFunctionForSyncPoint(FSyncPointId SyncPoint);
+
+	/** Tells the manager that it is the start of a frame, which will register the appropriate events */
+	ENGINE_API void StartFrame(const UWorld* InWorld);
+
+	/** Tells the manager that a tick group is starting for a specific world */
+	ENGINE_API void StartTickGroup(const UWorld* InWorld, ETickingGroup TickGroup, TArray<FTickFunction*>& TicksToManualDispatch);
+
+	/** Tells the manager that a tick group is complete for a specific world, which could trigger events */
+	ENGINE_API void EndTickGroup(const UWorld* InWorld, ETickingGroup TickGroup);
+
+	/** Tells the manager that it is the end of a frame and it will not start/end any more tick groups */
+	ENGINE_API void EndFrame(const UWorld* InWorld);
+
+	/** Tells the manager that it should destroy all tracking info for a world context. This only needs to be called when destroying a test or preview world */
+	ENGINE_API void ReleaseWorldContext(FWorldContextId WorldContext);
+private:
+		
+	struct FBatchData
+	{
+		FBatchContextId BatchContext;
+		FWorldContextId WorldContext;
+		TMap<FSyncPointId::FInternalId, TSharedRef<FActiveSyncPoint>> SyncPointData;
+
+		inline void Reset()
+		{
+			SyncPointData.Reset();
+			BatchContext = FBatchContextId();
+			WorldContext = FWorldContextId();
+		}
+	};
+
+	struct FRegisteredSyncPointData
+	{
+		FSyncPointId::FInternalId RegisteredId = FSyncPointId::InvalidSyncPoint;
+		FSyncPointDescription RegisteredPoint;
+	};
+
+	// Store as a simple array for fast access, there are a small number of active batches
+	TArray<FBatchData> BatchList;
+
+	TMap<FSyncPointId::FInternalId, FRegisteredSyncPointData> RegisteredDataMap;
+	TMap<FName, FSyncPointId::FInternalId> RegisteredNameMap;
+	FSyncPointId::FInternalId HighestSyncId = FSyncPointId::InvalidSyncPoint;
+
+	/** Lock to support multithreaded access. TODO: Investigate if RWLock or command buffer is worth it */
+	mutable FTransactionallySafeCriticalSection ManagerLock;
+
+	/** World that is currently ticking, cannot modify certain operations during a tick or tick two worlds at once */
+	const UWorld* CurrentTickWorld = nullptr;
+	ETickingGroup CurrentTickGroup = ETickingGroup::TG_MAX;
+
+	/** True if this is currently ticking a world */
+	inline bool IsTicking() const { return CurrentTickWorld != nullptr; }
+
+	/** Internal accessor to get a tick function's raw data, does not lock */
+	FTaskSyncResult FindActiveSyncPoint(FSyncPointId SyncPoint, TSharedPtr<FActiveSyncPoint>& OutData);
+
+	void RegisterSyncPointInternal(const FSyncPointDescription& InDescription);
+	void OnWorldContextRemove(FWorldContext& InWorldContext);
+	void InitializeBatchForFrame(FBatchData& BatchData, ULevel* PersistentLevel);
+	FActiveSyncPoint* GetOrCreateSyncPoint(FBatchData& BatchData, FRegisteredSyncPointData& SyncData);
+};
+
+} // namespace UE::Task
