- Added new interface IAsyncPhysicsStateProcessor to support async physics state creation/destruction of other types than UPrimitiveComponent
- Added new interface IPhysicsBodyInstanceOwner to support having a FBodyInstance owner which is not a UPrimitiveComponent
(This development in required for an upcoming feature)

diff --git a/Engine/Source/Runtime/Engine/Classes/Components/ActorComponent.h b/Engine/Source/Runtime/Engine/Classes/Components/ActorComponent.h
--- a/Engine/Source/Runtime/Engine/Classes/Components/ActorComponent.h
+++ b/Engine/Source/Runtime/Engine/Classes/Components/ActorComponent.h
@@ -18,6 +18,7 @@
 #include "Experimental/ConcurrentLinearAllocator.h"
 #include "Misc/ScopeRWLock.h"
 #include "Misc/TransactionallySafeRWLock.h"
+#include "Physics/Experimental/AsyncPhysicsStateProcessorInterface.h"
 #include "ActorComponent.generated.h"
 
 struct FTypedElementHandle;
@@ -140,7 +141,7 @@ DECLARE_MULTICAST_DELEGATE_OneParam(FActorComponentGlobalDestroyPhysicsSignature
  * @see UPrimitiveComponent
  */
 UCLASS(DefaultToInstanced, BlueprintType, abstract, meta=(ShortTooltip="An ActorComponent is a reusable component that can be added to any actor."), config=Engine, MinimalAPI)
-class UActorComponent : public UObject, public IInterface_AssetUserData
+class UActorComponent : public UObject, public IInterface_AssetUserData, public IAsyncPhysicsStateProcessor
 {
 	GENERATED_BODY()
 
@@ -828,35 +829,23 @@ public:
 	/** Returns whether PreUnregisterComponentFromWorld should be called for this component when its Level is being removed from the world. */
 	ENGINE_API virtual bool ShouldIncrementalPreUnregister() const;
 
-	/** Returns whether this component allows having its physics state being created asynchronously (outside of the GameThread). */
+	//~ Begin IAsyncPhysicsStateProcessor
 	ENGINE_API virtual bool AllowsAsyncPhysicsStateCreation() const
 	{
 		return false;
 	}
-
-	/** Returns whether this component allows having its physics state being destroyed asynchronously (outside of the GameThread). */
 	ENGINE_API virtual bool AllowsAsyncPhysicsStateDestruction() const
 	{
 		return false;
 	}
-
-	/** Used to create any physics engine information for this component outside of the GameThread. */
-	ENGINE_API virtual void OnAsyncCreatePhysicsState();
-
-	/** Called on the GameThread once the component's physic engine information is created. */
-	ENGINE_API virtual void OnAsyncCreatePhysicsStateEnd_GameThread();
-
-	/** Used to destroy any physics engine information for this component outside of the GameThread. */
-	ENGINE_API virtual void OnAsyncDestroyPhysicsState()
-	{
-	}
-
-	/** Called on the GameThread before the component's physic engine information is destroyed. */
-	ENGINE_API virtual void OnAsyncDestroyPhysicsStateBegin_GameThread();
-
-	/** Called on the GameThread once the component's physic engine information is destroyed. */
-	ENGINE_API virtual void OnAsyncDestroyPhysicsStateEnd_GameThread();
-
+	ENGINE_API virtual bool IsAsyncPhysicsStateCreated() const override final;
+	ENGINE_API virtual UObject* GetAsyncPhysicsStateObject() const override final;
+	ENGINE_API virtual void OnAsyncCreatePhysicsState() override;
+	ENGINE_API virtual void OnAsyncCreatePhysicsStateEnd_GameThread() override;
+	ENGINE_API virtual void OnAsyncDestroyPhysicsState() override {}
+	ENGINE_API virtual void OnAsyncDestroyPhysicsStateBegin_GameThread() override;
+	ENGINE_API virtual void OnAsyncDestroyPhysicsStateEnd_GameThread() override;
+	//~ End IAsyncPhysicsStateProcessor
 	/**
 	 * Initializes the component.  Occurs at level startup or actor spawn. This is before BeginPlay (Actor or Component).  
 	 * All Components in the level will be Initialized on load before any Actor/Component gets BeginPlay
diff --git a/Engine/Source/Runtime/Engine/Classes/Components/PrimitiveComponent.h b/Engine/Source/Runtime/Engine/Classes/Components/PrimitiveComponent.h
--- a/Engine/Source/Runtime/Engine/Classes/Components/PrimitiveComponent.h
+++ b/Engine/Source/Runtime/Engine/Classes/Components/PrimitiveComponent.h
@@ -2482,10 +2482,6 @@ public:
 	ENGINE_API virtual void OnUnregister()  override;
 	ENGINE_API virtual void DestroyRenderState_Concurrent() override;
 	ENGINE_API virtual void OnCreatePhysicsState() override;
-	ENGINE_API virtual void OnAsyncCreatePhysicsStateEnd_GameThread() override;
-	ENGINE_API virtual void OnAsyncDestroyPhysicsState() override;
-	ENGINE_API virtual void OnAsyncDestroyPhysicsStateBegin_GameThread() override;
-	ENGINE_API virtual void OnAsyncDestroyPhysicsStateEnd_GameThread() override;
 	ENGINE_API virtual void OnDestroyPhysicsState() override;
 	ENGINE_API virtual void OnActorEnableCollisionChanged() override;
 	ENGINE_API virtual void InvalidateLightingCacheDetailed(bool bInvalidateBuildEnqueuedLighting, bool bTranslationOnly) override;
@@ -2500,6 +2496,14 @@ public:
 #endif // WITH_EDITOR	
 	//~ End UActorComponent Interface
 
+	//~ Begin IAsyncPhysicsStateProcessor
+	ENGINE_API virtual void OnAsyncCreatePhysicsStateEnd_GameThread() override;
+	ENGINE_API virtual void OnAsyncDestroyPhysicsState() override;
+	ENGINE_API virtual void OnAsyncDestroyPhysicsStateBegin_GameThread() override;
+	ENGINE_API virtual void OnAsyncDestroyPhysicsStateEnd_GameThread() override;
+	ENGINE_API virtual void CollectBodySetupsWithPhysicsMeshesToCreate(TSet<UBodySetup*>& OutBodySetups) const override;
+	//~ End IAsyncPhysicsStateProcessor
+
 protected:
 	/** Internal function that updates physics objects to match the component collision settings. */
 	ENGINE_API virtual void UpdatePhysicsToRBChannels();
diff --git a/Engine/Source/Runtime/Engine/Classes/PhysicsEngine/BodyInstance.h b/Engine/Source/Runtime/Engine/Classes/PhysicsEngine/BodyInstance.h
--- a/Engine/Source/Runtime/Engine/Classes/PhysicsEngine/BodyInstance.h
+++ b/Engine/Source/Runtime/Engine/Classes/PhysicsEngine/BodyInstance.h
@@ -19,6 +19,7 @@ class UBodySetup;
 class UPhysicalMaterial;
 class UPhysicalMaterialMask;
 class UPrimitiveComponent;
+class IPhysicsBodyInstanceOwner;
 struct FBodyInstance;
 struct FCollisionNotifyInfo;
 struct FCollisionShape;
@@ -26,10 +27,10 @@ struct FConstraintInstance;
 struct FPropertyChangedEvent;
 struct FRigidBodyState;
 struct FShapeData;
-class UPrimitiveComponent;
 
 struct FShapeData;
 
+ENGINE_API FArchive& operator<<(FArchive& Ar, FBodyInstance& BodyInst);
 ENGINE_API int32 FillInlineShapeArray_AssumesLocked(PhysicsInterfaceTypes::FInlineShapeArray& Array, const FPhysicsActorHandle& Actor);
 
 UENUM(BlueprintType)
@@ -99,12 +100,12 @@ struct TInitBodiesHelperBase : public Chaos::Private::FInitBodiesHelperBaseInter
 {
 public:
 	ENGINE_API TInitBodiesHelperBase(TArray<FBodyInstance*, InAllocatorType>& InBodies, TArray<FTransform, InAllocatorType>& InTransforms,
-		class UBodySetup* InBodySetup, class UPrimitiveComponent* InPrimitiveComp, FPhysScene* InRBScene, 
-		const FInitBodySpawnParams& InSpawnParams, FPhysicsAggregateHandle InAggregate);
+		class UBodySetup* InBodySetup, class UPrimitiveComponent* InPrimitiveComp, FPhysScene* InRBScene,
+		const FInitBodySpawnParams& InSpawnParams, FPhysicsAggregateHandle InAggregate, IPhysicsBodyInstanceOwner* InBodyInstanceOwner = nullptr);
 	
 	ENGINE_API TInitBodiesHelperBase(TArray<FBodyInstance*, InAllocatorType>& InBodies, TArray<FTransform, InAllocatorType>& InTransforms,
 		class UBodySetup* InBodySetup, class UPrimitiveComponent* InComponent, UObject* InSourceObject, FPhysScene* InRBScene,
-		const FInitBodySpawnParams& InSpawnParams, FPhysicsAggregateHandle InAggregate);
+		const FInitBodySpawnParams& InSpawnParams, FPhysicsAggregateHandle InAggregate, IPhysicsBodyInstanceOwner* InBodyInstanceOwner = nullptr);
 
 	TInitBodiesHelperBase(const TInitBodiesHelperBase& InHelper) = delete;
 	TInitBodiesHelperBase(TInitBodiesHelperBase&& InHelper) = delete;
@@ -119,6 +120,7 @@ public:
 
 	class UBodySetup* BodySetup;
 	class UPrimitiveComponent* PrimitiveComp;
+	IPhysicsBodyInstanceOwner* BodyInstanceOwner;
 	UObject* SourceObject;
 	FPhysScene* PhysScene;
 	FPhysicsAggregateHandle Aggregate;
@@ -160,12 +162,12 @@ struct FInitBodiesHelper : public TInitBodiesHelperBase<InAllocatorType>
 {
 	FInitBodiesHelper(TArray<FBodyInstance*, InAllocatorType>& InBodies, TArray<FTransform, InAllocatorType>& InTransforms, class UBodySetup* InBodySetup,
 		class UPrimitiveComponent* InPrimitiveComp, FPhysScene* InRBScene,
-		const FInitBodySpawnParams& InSpawnParams, FPhysicsAggregateHandle InAggregate);
+		const FInitBodySpawnParams& InSpawnParams, FPhysicsAggregateHandle InAggregate, IPhysicsBodyInstanceOwner* InBodyInstanceOwner = nullptr);
 
 	FInitBodiesHelper(TArray<FBodyInstance*, InAllocatorType>& InBodies, TArray<FTransform, InAllocatorType>& InTransforms, class UBodySetup* InBodySetup,
-		class UPrimitiveComponent* InComponent, class UObject* InOwnerObject, FPhysScene* InRBScene,
-		const FInitBodySpawnParams& InSpawnParams, FPhysicsAggregateHandle InAggregate)
-		: TInitBodiesHelperBase<InAllocatorType>(InBodies, InTransforms, InBodySetup, InComponent, InOwnerObject, InRBScene, InSpawnParams, InAggregate)
+		class UPrimitiveComponent* InComponent, class UObject* InSourceObject, FPhysScene* InRBScene,
+		const FInitBodySpawnParams& InSpawnParams, FPhysicsAggregateHandle InAggregate, IPhysicsBodyInstanceOwner* InBodyInstanceOwner = nullptr)
+		: TInitBodiesHelperBase<InAllocatorType>(InBodies, InTransforms, InBodySetup, InComponent, InSourceObject, InRBScene, InSpawnParams, InAggregate, InBodyInstanceOwner)
 	{
 		//Compute all the needed constants
 		this->bStatic = bCompileStatic || this->SpawnParams.bStaticPhysics;
@@ -181,8 +183,8 @@ template <bool bCompileStatic>
 struct FInitBodiesHelperWithData : public FInitBodiesHelperBase
 {
 	FInitBodiesHelperWithData() { check(false); }
-	FInitBodiesHelperWithData(TArray<FBodyInstance*>&& InBodies, TArray<FTransform>&& InTransforms, class UBodySetup* InBodySetup, class UPrimitiveComponent* InPrimitiveComp, FPhysScene* InRBScene, const FInitBodySpawnParams& InSpawnParams, FPhysicsAggregateHandle InAggregate)
-	: FInitBodiesHelperBase(OwnedBodies, OwnedTransforms, InBodySetup, InPrimitiveComp, InRBScene, InSpawnParams, InAggregate), OwnedBodies(MoveTemp(InBodies)), OwnedTransforms(MoveTemp(InTransforms)) //-V1050
+	FInitBodiesHelperWithData(TArray<FBodyInstance*>&& InBodies, TArray<FTransform>&& InTransforms, class UBodySetup* InBodySetup, class UPrimitiveComponent* InPrimitiveComp, FPhysScene* InRBScene, const FInitBodySpawnParams& InSpawnParams, FPhysicsAggregateHandle InAggregate, IPhysicsBodyInstanceOwner* InBodyInstanceOwner = nullptr)
+	: FInitBodiesHelperBase(OwnedBodies, OwnedTransforms, InBodySetup, InPrimitiveComp, InRBScene, InSpawnParams, InAggregate, InBodyInstanceOwner), OwnedBodies(MoveTemp(InBodies)), OwnedTransforms(MoveTemp(InTransforms)) //-V1050
 	{
 		//Compute all the needed constants
 		bStatic = bCompileStatic || SpawnParams.bStaticPhysics;
@@ -194,7 +196,7 @@ struct FInitBodiesHelperWithData : public FInitBodiesHelperBase
 	}
 
 	FInitBodiesHelperWithData(const FInitBodiesHelperWithData& InHelper)
-	: FInitBodiesHelperBase(OwnedBodies, OwnedTransforms, InHelper.BodySetup, InHelper.PrimitiveComp, InHelper.PhysScene, InHelper.SpawnParams, InHelper.Aggregate), OwnedBodies(InHelper.OwnedBodies), OwnedTransforms(InHelper.OwnedTransforms) //-V1050
+	: FInitBodiesHelperBase(OwnedBodies, OwnedTransforms, InHelper.BodySetup, InHelper.PrimitiveComp, InHelper.PhysScene, InHelper.SpawnParams, InHelper.Aggregate), OwnedBodies(InHelper.OwnedBodies), OwnedTransforms(InHelper.OwnedTransforms, InHelper.BodyInstanceOwner) //-V1050
 	{
 		ensure(false);
 	}
@@ -689,6 +691,9 @@ public:
 	/** PrimitiveComponent containing this body, if relevant.   */
 	TWeakObjectPtr<class UPrimitiveComponent> OwnerComponent;
 
+	/** Owner containing this body (when not a UPrimitiveComponent). */
+	IPhysicsBodyInstanceOwner* BodyInstanceOwner;
+
 	/** Source object for this body. If OwnerComponent is not null, this will typically be the owning actor. */
 	TWeakObjectPtr<UObject> SourceObject;
 
@@ -717,9 +722,9 @@ public:
 	*	@param PrimComp The owning component
 	*	@param InRBScene The physics scene to place the body into
 	*	@param SpawnParams The parameters for determining certain spawn behavior
-	*	@param InAggregate An aggregate to place the body into
+	*	@param BodyInstanceOwner The owning object of the body instance (optional and only used if PrimComp is not set)
 	*/
-	ENGINE_API void InitBody(UBodySetup* Setup, const FTransform& Transform, UPrimitiveComponent* PrimComp, FPhysScene* InRBScene, const FInitBodySpawnParams& SpawnParams);
+	ENGINE_API void InitBody(UBodySetup* Setup, const FTransform& Transform, UPrimitiveComponent* PrimComp, FPhysScene* InRBScene, const FInitBodySpawnParams& SpawnParams, IPhysicsBodyInstanceOwner* BodyInstanceOwner = nullptr);
 
 	/**
 	 * Initialise a single rigid body (this FBodyInstance) for the given body setup
@@ -730,11 +735,12 @@ public:
 	 *   to the owning actor.
 	 *  @param InRBScene The physics scene to place the body into
 	 *  @param SpawnParams The parameters for determining certain spawn behavior
-	 *  @param InAggregate An aggregate to place the body into
+	 *	@param BodyInstanceOwner The owning object of the body instance (optional and only used if PrimComp is not set)
 	 */
 	ENGINE_API void InitBody(UBodySetup* Setup, const FTransform& Transform, 
 		UPrimitiveComponent* PrimComp, UObject* SourceObject, 
-		FPhysScene* InRBScene, const FInitBodySpawnParams& SpawnParams);
+		FPhysScene* InRBScene, const FInitBodySpawnParams& SpawnParams, 
+		IPhysicsBodyInstanceOwner* BodyInstanceOwner = nullptr);
 
 	/** Validate a body transform, outputting debug info
 	 *	@param Transform Transform to debug
@@ -760,9 +766,10 @@ public:
 	 *	@param BodySetup The BodySetup to use
 	 *	@param PrimitiveComp The associated primitive component
 	 *	@param InRBScene The associated physics scene
+	 *	@param BodyInstanceOwner The associated owning object of the body instance (optional and only used if PrimitiveComp is not set)
 	 */
-	static ENGINE_API void InitStaticBodies(TArray<FBodyInstance*>&& Bodies, TArray<FTransform>&& Transforms, UBodySetup* BodySetup, class UPrimitiveComponent* PrimitiveComp, FPhysScene* InRBScene);
-	
+	static ENGINE_API void InitStaticBodies(TArray<FBodyInstance*>&& Bodies, TArray<FTransform>&& Transforms, UBodySetup* BodySetup, class UPrimitiveComponent* PrimitiveComp, FPhysScene* InRBScene, IPhysicsBodyInstanceOwner* BodyInstanceOwner = nullptr);
+
 	/** Get the scene that owns this body. */
 	ENGINE_API FPhysScene* GetPhysicsScene();
 	ENGINE_API const FPhysScene* GetPhysicsScene() const;
@@ -869,7 +876,7 @@ public:
 	ENGINE_API UPhysicalMaterial* GetSimplePhysicalMaterial() const;
 
 	/** Find the correct PhysicalMaterial for simple geometry on a given body and owner. This is really for internal use during serialization */
-	static ENGINE_API UPhysicalMaterial* GetSimplePhysicalMaterial(const FBodyInstance* BodyInstance, TWeakObjectPtr<UPrimitiveComponent> Owner, TWeakObjectPtr<UBodySetup> BodySetupPtr);
+	static ENGINE_API UPhysicalMaterial* GetSimplePhysicalMaterial(const FBodyInstance* BodyInstance, TWeakObjectPtr<UPrimitiveComponent> Owner, TWeakObjectPtr<UBodySetup> BodySetupPtr, IPhysicsBodyInstanceOwner* BodyInstanceOwner = nullptr);
 
 	/** Get the complex PhysicalMaterials array for this body */
 	ENGINE_API TArray<UPhysicalMaterial*> GetComplexPhysicalMaterials() const;
@@ -884,7 +891,7 @@ public:
 	ENGINE_API void GetComplexPhysicalMaterials(TArray<UPhysicalMaterial*> &OutPhysMaterials, TArray<FPhysicalMaterialMaskParams>& OutPhysMaterialMasks) const;
 
 	/** Find the correct PhysicalMaterial and PhysicalMaterialMasks for complex geometry on a given body and owner. This is really for internal use during serialization */
-	static ENGINE_API void GetComplexPhysicalMaterials(const FBodyInstance* BodyInstance, TWeakObjectPtr<UPrimitiveComponent> Owner, TArray<UPhysicalMaterial*>& OutPhysMaterials, TArray<FPhysicalMaterialMaskParams>* OutPhysMaterialMasks = nullptr);
+	static ENGINE_API void GetComplexPhysicalMaterials(const FBodyInstance* BodyInstance, TWeakObjectPtr<UPrimitiveComponent> Owner, TArray<UPhysicalMaterial*>& OutPhysMaterials, TArray<FPhysicalMaterialMaskParams>* OutPhysMaterialMasks = nullptr, IPhysicsBodyInstanceOwner* InBodyInstanceOwner = nullptr);
 
 	/** Returns the slope override struct for this instance. If we don't have our own custom setting, it will return the setting from the body setup. */
 	ENGINE_API const struct FWalkableSlopeOverride& GetWalkableSlopeOverride() const;
@@ -1144,6 +1151,9 @@ public:
 	/** Get current velocity of a point on this physics body, in world space. Point is specified in world space. */
 	ENGINE_API FVector GetUnrealWorldVelocityAtPoint_AssumesLocked(const FVector& Point) const;
 
+	/** Get physical material override for this body */
+	ENGINE_API UPhysicalMaterial* GetPhysMaterialOverride() const { return PhysMaterialOverride; }
+
 	/** Set physical material override for this body */
 	ENGINE_API void SetPhysMaterialOverride(class UPhysicalMaterial* NewPhysMaterial);
 
@@ -1256,7 +1266,7 @@ public:
 	/** Update the instances collision filtering data */ 
 	ENGINE_API void UpdatePhysicsFilterData();
 
-	friend FArchive& operator<<(FArchive& Ar,FBodyInstance& BodyInst);
+	friend ENGINE_API FArchive& operator<<(FArchive& Ar,FBodyInstance& BodyInst);
 
 	/** Get the name for this body, for use in debugging */
 	ENGINE_API FString GetBodyDebugName() const;
diff --git a/Engine/Source/Runtime/Engine/Classes/PhysicsEngine/PhysicsBodyInstanceOwnerInterface.h b/Engine/Source/Runtime/Engine/Classes/PhysicsEngine/PhysicsBodyInstanceOwnerInterface.h
new file mode 100644
--- /dev/null
+++ b/Engine/Source/Runtime/Engine/Classes/PhysicsEngine/PhysicsBodyInstanceOwnerInterface.h
@@ -0,0 +1,26 @@
+// Copyright Epic Games, Inc. All Rights Reserved.
+
+#pragma once
+
+#include "Containers/Array.h"
+
+class UObject;
+class UPhysicalMaterial;
+struct FPhysicalMaterialMaskParams;
+
+/** Interface representing the owner of a FBodyInstance (used when the owner is not an UPrimitiveComponent). */
+class IPhysicsBodyInstanceOwner
+{
+public:
+	/** Whether the physics is static. */
+	virtual bool IsStaticPhysics() const = 0;
+
+	/** Source object for this body. */
+	virtual UObject* GetSourceObject() const = 0;
+
+	/** Find the correct PhysicalMaterial for simple geometry on this body (used by FBodyInstance::GetSimplePhysicalMaterial). */
+	virtual UPhysicalMaterial* GetPhysicalMaterial() const = 0;
+
+	/** Get the complex PhysicalMaterials array for this body (used by FBodyInstance::GetComplexPhysicalMaterials). */
+	virtual void GetComplexPhysicalMaterials(TArray<UPhysicalMaterial*>& OutPhysMaterials, TArray<FPhysicalMaterialMaskParams>* OutPhysMaterialMasks) const = 0;
+};
\ No newline at end of file
diff --git a/Engine/Source/Runtime/Engine/Private/Collision/CollisionConversions.cpp b/Engine/Source/Runtime/Engine/Private/Collision/CollisionConversions.cpp
--- a/Engine/Source/Runtime/Engine/Private/Collision/CollisionConversions.cpp
+++ b/Engine/Source/Runtime/Engine/Private/Collision/CollisionConversions.cpp
@@ -11,6 +11,7 @@
 #include "EngineLogs.h"
 #include "PhysicalMaterials/PhysicalMaterial.h"
 #include "PhysicsEngine/PhysicsSettings.h"
+#include "PhysicsEngine/PhysicsBodyInstanceOwnerInterface.h"
 
 #include "Physics/Experimental/ChaosInterfaceWrapper.h"
 
@@ -135,7 +136,8 @@ static void SetHitResultFromShapeAndFaceIndex(const FPhysicsShape& Shape,  const
 	const int32 ShapeIndex = Shape.GetShapeIndex();
 	CHAOS_CHECK(ShapeIndex < (int32)TNumericLimits<uint8>::Max()); // I could just write < 256, but this makes it more clear *why*
 	OutResult.ElementIndex = (uint8)ShapeIndex;
-	
+
+	IPhysicsBodyInstanceOwner* BodyInstanceOwner = nullptr;
 	TWeakObjectPtr<UPrimitiveComponent> OwningComponent;
 	OutResult.PhysicsObject = nullptr;
 	if(const FBodyInstance* BodyInst = GetUserData(Actor))
@@ -151,6 +153,7 @@ static void SetHitResultFromShapeAndFaceIndex(const FPhysicsShape& Shape,  const
 		}
 
 		OwningComponent = BodyInst->OwnerComponent;
+		BodyInstanceOwner = BodyInst->BodyInstanceOwner;
 		OutResult.PhysicsObject = BodyInst->GetPhysicsActor() ? BodyInst->GetPhysicsActor()->GetPhysicsObject() : nullptr;
 	}
 	else
@@ -196,13 +199,16 @@ static void SetHitResultFromShapeAndFaceIndex(const FPhysicsShape& Shape,  const
 	OutResult.PhysMaterial = nullptr;
 
 	// Grab actor/component
-	if (!OwningComponent.IsExplicitlyNull())
+	if (!OwningComponent.IsExplicitlyNull() || BodyInstanceOwner)
 	{
-		OutResult.Component = OwningComponent;
+		if (!OwningComponent.IsExplicitlyNull())
+		{
+			OutResult.Component = OwningComponent;
 
-		// Create a handle that needs to be resolved later when accessed in game thread,
-		// since resolving FActorInstanceHandle may require to access related UObjects (e.g. owner actor), which is not safe outside game thread.
-		OutResult.HitObjectHandle = FActorInstanceHandle::MakeActorHandleToResolve(OwningComponent, OutResult.Item);
+			// Create a handle that needs to be resolved later when accessed in game thread,
+			// since resolving FActorInstanceHandle may require to access related UObjects (e.g. owner actor), which is not safe outside game thread.
+			OutResult.HitObjectHandle = FActorInstanceHandle::MakeActorHandleToResolve(OwningComponent, OutResult.Item);
+		}
 
 		if (bReturnPhysMat)
 		{
diff --git a/Engine/Source/Runtime/Engine/Private/Components/ActorComponent.cpp b/Engine/Source/Runtime/Engine/Private/Components/ActorComponent.cpp
--- a/Engine/Source/Runtime/Engine/Private/Components/ActorComponent.cpp
+++ b/Engine/Source/Runtime/Engine/Private/Components/ActorComponent.cpp
@@ -2017,6 +2017,16 @@ void UActorComponent::OnCreatePhysicsState()
 	bPhysicsStateCreated = true;
 }
 
+bool UActorComponent::IsAsyncPhysicsStateCreated() const
+{
+	return IsPhysicsStateCreated();
+}
+
+UObject* UActorComponent::GetAsyncPhysicsStateObject() const
+{
+	return const_cast<UActorComponent*>(this);
+}
+
 void UActorComponent::OnAsyncCreatePhysicsState()
 {
 	TRACE_CPUPROFILER_EVENT_SCOPE(UActorComponent::OnAsyncCreatePhysicsState);
diff --git a/Engine/Source/Runtime/Engine/Private/Components/PrimitiveComponent.cpp b/Engine/Source/Runtime/Engine/Private/Components/PrimitiveComponent.cpp
--- a/Engine/Source/Runtime/Engine/Private/Components/PrimitiveComponent.cpp
+++ b/Engine/Source/Runtime/Engine/Private/Components/PrimitiveComponent.cpp
@@ -1070,6 +1070,14 @@ void UPrimitiveComponent::OnAsyncDestroyPhysicsStateEnd_GameThread()
 	OnComponentPhysicsStateChanged.Broadcast(this, EComponentPhysicsStateChange::Destroyed);
 }
 
+void UPrimitiveComponent::CollectBodySetupsWithPhysicsMeshesToCreate(TSet<UBodySetup*>& OutBodySetups) const
+{
+	if (UBodySetup* BodySetup = const_cast<UPrimitiveComponent*>(this)->GetBodySetup(); BodySetup && !BodySetup->bCreatedPhysicsMeshes)
+	{
+		OutBodySetups.Add(BodySetup);
+	}
+}
+
 void UPrimitiveComponent::OnDestroyPhysicsState()
 {
 	// we remove welding related to this component
diff --git a/Engine/Source/Runtime/Engine/Private/PhysicsEngine/BodyInstance.cpp b/Engine/Source/Runtime/Engine/Private/PhysicsEngine/BodyInstance.cpp
--- a/Engine/Source/Runtime/Engine/Private/PhysicsEngine/BodyInstance.cpp
+++ b/Engine/Source/Runtime/Engine/Private/PhysicsEngine/BodyInstance.cpp
@@ -25,6 +25,7 @@
 #include "PhysicsEngine/SphylElem.h"
 #include "PhysicsEngine/BodySetup.h"
 #include "PhysicsEngine/BodyUtils.h"
+#include "PhysicsEngine/PhysicsBodyInstanceOwnerInterface.h"
 #include "Logging/MessageLog.h"
 
 #include "Chaos/Capsule.h"
@@ -406,6 +407,7 @@ FBodyInstance::FBodyInstance()
 	, StabilizationThresholdMultiplier(1.f)
 	, PhysicsBlendWeight(0.f)
 	, ActorHandle(DefaultPhysicsActorHandle())
+	, BodyInstanceOwner(nullptr)
 {
 	MaxAngularVelocity = UPhysicsSettings::Get()->MaxAngularVelocity;
 }
@@ -1211,11 +1213,11 @@ namespace Chaos::Private
 template<bool bCompileStatic, typename InAllocatorType>
 FInitBodiesHelper<bCompileStatic, InAllocatorType>::FInitBodiesHelper(
 	TArray<FBodyInstance*, InAllocatorType>& InBodies, TArray<FTransform, InAllocatorType>& InTransforms, UBodySetup* InBodySetup,
-	UPrimitiveComponent* InPrimitiveComp, FPhysScene* InRBScene, 
-	const FInitBodySpawnParams& InSpawnParams, FPhysicsAggregateHandle InAggregate)
+	UPrimitiveComponent* InPrimitiveComp, FPhysScene* InRBScene,
+	const FInitBodySpawnParams& InSpawnParams, FPhysicsAggregateHandle InAggregate, IPhysicsBodyInstanceOwner* InBodyInstanceOwner)
 	: FInitBodiesHelper<bCompileStatic, InAllocatorType>(InBodies, InTransforms, InBodySetup, InPrimitiveComp,
-		InPrimitiveComp ? InPrimitiveComp->GetOwner() : nullptr, 
-		InRBScene, InSpawnParams, InAggregate)
+		InPrimitiveComp ? InPrimitiveComp->GetOwner() : InBodyInstanceOwner ? InBodyInstanceOwner->GetSourceObject() : nullptr,
+		InRBScene, InSpawnParams, InAggregate, InBodyInstanceOwner)
 {
 }
 
@@ -1230,11 +1232,12 @@ using FInitSingleBodyHelperDynamic = FInitBodiesHelper<false, BodyInstanceSingle
 template<typename InAllocatorType>
 TInitBodiesHelperBase<InAllocatorType>::TInitBodiesHelperBase(TArray<FBodyInstance*, InAllocatorType>& InBodies, TArray<FTransform, InAllocatorType>& InTransforms,
 	class UBodySetup* InBodySetup, UPrimitiveComponent* InComponent, UObject* InSourceObject, FPhysScene* InRBScene,
-	const FInitBodySpawnParams& InSpawnParams, FPhysicsAggregateHandle InAggregate)
+	const FInitBodySpawnParams& InSpawnParams, FPhysicsAggregateHandle InAggregate, IPhysicsBodyInstanceOwner* InBodyInstanceOwner)
 	: Bodies(InBodies)
 	, Transforms(InTransforms)
 	, BodySetup(InBodySetup)
 	, PrimitiveComp(InComponent)
+	, BodyInstanceOwner(InBodyInstanceOwner)
 	, SourceObject(InSourceObject)
 	, PhysScene(InRBScene)
 	, Aggregate(InAggregate)
@@ -1255,12 +1258,12 @@ TInitBodiesHelperBase<InAllocatorType>::TInitBodiesHelperBase(TArray<FBodyInstan
 template<typename InAllocatorType>
 TInitBodiesHelperBase<InAllocatorType>::TInitBodiesHelperBase(TArray<FBodyInstance*, InAllocatorType>& InBodies, TArray<FTransform, InAllocatorType>& InTransforms,
 	class UBodySetup* InBodySetup, class UPrimitiveComponent* InPrimitiveComp, FPhysScene* InRBScene,
-	const FInitBodySpawnParams& InSpawnParams, FPhysicsAggregateHandle InAggregate)
+	const FInitBodySpawnParams& InSpawnParams, FPhysicsAggregateHandle InAggregate, IPhysicsBodyInstanceOwner* InBodyInstanceOwner)
 	// Forward to the other constructor
 	: TInitBodiesHelperBase<InAllocatorType>(InBodies, InTransforms, InBodySetup,
 		InPrimitiveComp,
-		InPrimitiveComp ? InPrimitiveComp->GetOwner() : nullptr,
-		InRBScene, InSpawnParams, InAggregate)
+		InPrimitiveComp ? InPrimitiveComp->GetOwner() : InBodyInstanceOwner ? InBodyInstanceOwner->GetSourceObject() : nullptr,
+		InRBScene, InSpawnParams, InAggregate, InBodyInstanceOwner)
 {
 }
 
@@ -1473,6 +1476,7 @@ bool TInitBodiesHelperBase<InAllocatorType>::CreateShapesAndActors()
 		}
 
 		Instance->OwnerComponent = PrimitiveComp;
+		Instance->BodyInstanceOwner = BodyInstanceOwner;
 		Instance->SourceObject = SourceObject;
 		Instance->BodySetup = BodySetup;
 		Instance->Scale3D = Transform.GetScale3D();
@@ -1508,6 +1512,7 @@ bool TInitBodiesHelperBase<InAllocatorType>::CreateShapesAndActors()
 		if (FPhysicsInterface::IsValid(Instance->GetPhysicsActor()))
 		{
 			Instance->OwnerComponent = nullptr;
+			Instance->BodyInstanceOwner = nullptr;
 			Instance->SourceObject = nullptr;
 			Instance->BodySetup = nullptr;
 			constexpr uint32 Count = 1;
@@ -1544,6 +1549,7 @@ bool TInitBodiesHelperBase<InAllocatorType>::CreateShapesAndActors()
 			Instance->ReleasePhysicsActor();
 
 			Instance->OwnerComponent = nullptr;
+			Instance->BodyInstanceOwner = nullptr;
 			Instance->SourceObject = nullptr;
 			Instance->BodySetup = nullptr;
 			SetBodyInstanceExternalCollisionProfileBodySetup(Instance, nullptr);
@@ -1693,12 +1699,12 @@ static TAutoConsoleVariable<int32> CVarAllowCreatePhysxBodies(
 	ECVF_ReadOnly);
 
 
-void FBodyInstance::InitBody(UBodySetup* Setup, const FTransform& Transform, UPrimitiveComponent* PrimComp, FPhysScene* InRBScene, const FInitBodySpawnParams& SpawnParams)
+void FBodyInstance::InitBody(UBodySetup* Setup, const FTransform& Transform, UPrimitiveComponent* PrimComp, FPhysScene* InRBScene, const FInitBodySpawnParams& SpawnParams, IPhysicsBodyInstanceOwner* InBodyInstanceOwner)
 {
-	InitBody(Setup, Transform, PrimComp, nullptr, InRBScene, SpawnParams);
+	InitBody(Setup, Transform, PrimComp, nullptr, InRBScene, SpawnParams, InBodyInstanceOwner);
 }
 
-void FBodyInstance::InitBody(class UBodySetup* Setup, const FTransform& Transform, UPrimitiveComponent* PrimComp, UObject* InSourceObject, FPhysScene* InRBScene, const FInitBodySpawnParams& SpawnParams)
+void FBodyInstance::InitBody(UBodySetup* Setup, const FTransform& Transform, UPrimitiveComponent* PrimComp, UObject* InSourceObject, FPhysScene* InRBScene, const FInitBodySpawnParams& SpawnParams, IPhysicsBodyInstanceOwner* InBodyInstanceOwner)
 {
 	if (CVarAllowCreatePhysxBodies.GetValueOnAnyThread() == 0)
 	{
@@ -1719,16 +1725,20 @@ void FBodyInstance::InitBody(class UBodySetup* Setup, const FTransform& Transfor
 	{
 		SourceObjectToUse = PrimComp->GetOwner();
 	}
+	if (!SourceObjectToUse && InBodyInstanceOwner)
+	{
+		SourceObjectToUse = InBodyInstanceOwner->GetSourceObject();
+	}
 
 	bool bIsStatic = SpawnParams.bStaticPhysics;
 	if(bIsStatic)
 	{
-		FInitSingleBodyHelperStatic InitBodyHelper(Bodies, Transforms, Setup, PrimComp, SourceObjectToUse, InRBScene, SpawnParams, SpawnParams.Aggregate);
+		FInitSingleBodyHelperStatic InitBodyHelper(Bodies, Transforms, Setup, PrimComp, SourceObjectToUse, InRBScene, SpawnParams, SpawnParams.Aggregate, InBodyInstanceOwner);
 		InitBodyHelper.InitBodies();
 	}
 	else
 	{
-		FInitSingleBodyHelperDynamic InitBodyHelper(Bodies, Transforms, Setup, PrimComp, SourceObjectToUse, InRBScene, SpawnParams, SpawnParams.Aggregate);
+		FInitSingleBodyHelperDynamic InitBodyHelper(Bodies, Transforms, Setup, PrimComp, SourceObjectToUse, InRBScene, SpawnParams, SpawnParams.Aggregate, InBodyInstanceOwner);
 		InitBodyHelper.InitBodies();
 	}
 
@@ -1779,6 +1789,7 @@ FBodyInstance::FAsyncTermBodyPayload FBodyInstance::StartAsyncTermBody_GameThrea
 	Payload.bNotifyRigidBodyCollision = bNotifyRigidBodyCollision;
 	SetPhysicsActor(nullptr);
 	OwnerComponent = nullptr;
+	BodyInstanceOwner = nullptr;
 
 	CurrentSceneState = BodyInstanceSceneState::NotAdded;
 	BodySetup = NULL;
@@ -3088,10 +3099,10 @@ void FBodyInstance::SetPhysMaterialOverride( UPhysicalMaterial* NewPhysMaterial
 
 UPhysicalMaterial* FBodyInstance::GetSimplePhysicalMaterial() const
 {
-	return GetSimplePhysicalMaterial(this, OwnerComponent, GetBodySetup());
+	return GetSimplePhysicalMaterial(this, OwnerComponent, GetBodySetup(), BodyInstanceOwner);
 }
 
-UPhysicalMaterial* FBodyInstance::GetSimplePhysicalMaterial(const FBodyInstance* BodyInstance, TWeakObjectPtr<UPrimitiveComponent> OwnerComp, TWeakObjectPtr<UBodySetup> BodySetupPtr)
+UPhysicalMaterial* FBodyInstance::GetSimplePhysicalMaterial(const FBodyInstance* BodyInstance, TWeakObjectPtr<UPrimitiveComponent> OwnerComp, TWeakObjectPtr<UBodySetup> BodySetupPtr, IPhysicsBodyInstanceOwner* BodyInstanceOwner)
 {
 	if(!GEngine || !GEngine->DefaultPhysMaterial)
 	{
@@ -3142,6 +3153,10 @@ UPhysicalMaterial* FBodyInstance::GetSimplePhysicalMaterial(const FBodyInstance*
 						PhysMatFromMaterial = Material->GetPhysicalMaterial();
 					}
 				}
+				else if (BodyInstanceOwner)
+				{
+					PhysMatFromMaterial = BodyInstanceOwner->GetPhysicalMaterial();
+				}
 
 				if (PhysMatFromMaterial != NULL)
 				{
@@ -3177,15 +3192,15 @@ TArray<UPhysicalMaterial*> FBodyInstance::GetComplexPhysicalMaterials(TArray<FPh
 
 void FBodyInstance::GetComplexPhysicalMaterials(TArray<UPhysicalMaterial*>& OutPhysMaterials) const
 {
-	GetComplexPhysicalMaterials(this, OwnerComponent, OutPhysMaterials);
+	GetComplexPhysicalMaterials(this, OwnerComponent, OutPhysMaterials, nullptr, BodyInstanceOwner);
 }
 
 void FBodyInstance::GetComplexPhysicalMaterials(TArray<UPhysicalMaterial*>& OutPhysMaterials, TArray<FPhysicalMaterialMaskParams>& OutPhysMaterialMasks) const
 {
-	GetComplexPhysicalMaterials(this, OwnerComponent, OutPhysMaterials, &OutPhysMaterialMasks);
+	GetComplexPhysicalMaterials(this, OwnerComponent, OutPhysMaterials, &OutPhysMaterialMasks, BodyInstanceOwner);
 }
 
-void FBodyInstance::GetComplexPhysicalMaterials(const FBodyInstance* BodyInstance, TWeakObjectPtr<UPrimitiveComponent> OwnerComp, TArray<UPhysicalMaterial*>& OutPhysMaterials, TArray<FPhysicalMaterialMaskParams>* OutPhysMaterialMasks)
+void FBodyInstance::GetComplexPhysicalMaterials(const FBodyInstance* BodyInstance, TWeakObjectPtr<UPrimitiveComponent> OwnerComp, TArray<UPhysicalMaterial*>& OutPhysMaterials, TArray<FPhysicalMaterialMaskParams>* OutPhysMaterialMasks, IPhysicsBodyInstanceOwner* InBodyInstanceOwner)
 {
 	if(!GEngine || !GEngine->DefaultPhysMaterial)
 	{
@@ -3207,18 +3222,17 @@ void FBodyInstance::GetComplexPhysicalMaterials(const FBodyInstance* BodyInstanc
 	else
 	{
 		// Component override
-		UPrimitiveComponent* OwnerPrimComponent = OwnerComp.Get();
-		if (OwnerPrimComponent && OwnerPrimComponent->BodyInstance.PhysMaterialOverride != nullptr)
-		{
-			OutPhysMaterials.SetNum(1);
-			OutPhysMaterials[0] = OwnerComp->BodyInstance.PhysMaterialOverride;
-			check(!OutPhysMaterials[0] || OutPhysMaterials[0]->IsValidLowLevel());
-		}
-		else
+		if (UPrimitiveComponent* OwnerPrimComponent = OwnerComp.Get())
 		{
-			// See if the Material has a PhysicalMaterial
-			if (OwnerPrimComponent)
+			if (OwnerPrimComponent->BodyInstance.PhysMaterialOverride != nullptr)
 			{
+				OutPhysMaterials.SetNum(1);
+				OutPhysMaterials[0] = OwnerComp->BodyInstance.PhysMaterialOverride;
+				check(!OutPhysMaterials[0] || OutPhysMaterials[0]->IsValidLowLevel());
+			}
+			else
+			{
+				// See if the Material has a PhysicalMaterial
 				const int32 NumMaterials = OwnerPrimComponent->GetNumMaterials();
 				OutPhysMaterials.SetNum(NumMaterials);
 
@@ -3256,7 +3270,11 @@ void FBodyInstance::GetComplexPhysicalMaterials(const FBodyInstance* BodyInstanc
 						(*OutPhysMaterialMasks)[MatIdx].PhysicalMaterialMap = PhysMatMap;
 					}
 				}
-			}			
+			}
+		}
+		else if (InBodyInstanceOwner)
+		{
+			InBodyInstanceOwner->GetComplexPhysicalMaterials(OutPhysMaterials, OutPhysMaterialMasks);
 		}
 	}
 }
@@ -4712,13 +4730,14 @@ void FBodyInstance::InitStaticBodies(const TArray<FBodyInstance*>& Bodies, const
 	TransformsStatic.Reset();
 }
 
-void FBodyInstance::InitStaticBodies(TArray<FBodyInstance*>&& Bodies, TArray<FTransform>&& Transforms, UBodySetup* BodySetup, UPrimitiveComponent* PrimitiveComp, FPhysScene* InRBScene)
+void FBodyInstance::InitStaticBodies(TArray<FBodyInstance*>&& Bodies, TArray<FTransform>&& Transforms, UBodySetup* BodySetup, UPrimitiveComponent* PrimitiveComp, FPhysScene* InRBScene, IPhysicsBodyInstanceOwner* InBodyInstanceOwner)
 {
 	SCOPE_CYCLE_COUNTER(STAT_StaticInitBodies);
 	check(BodySetup);
 	check(InRBScene);
 	check(Bodies.Num() > 0);
-	FInitBodiesHelperWithData<true> InitBodiesHelper(MoveTemp(Bodies), MoveTemp(Transforms), BodySetup, PrimitiveComp, InRBScene, FInitBodySpawnParams(PrimitiveComp), FPhysicsAggregateHandle());
+	check(!InBodyInstanceOwner || InBodyInstanceOwner->IsStaticPhysics());
+	FInitBodiesHelperWithData<true> InitBodiesHelper(MoveTemp(Bodies), MoveTemp(Transforms), BodySetup, PrimitiveComp, InRBScene, FInitBodySpawnParams(PrimitiveComp), FPhysicsAggregateHandle(), InBodyInstanceOwner);
 	InitBodiesHelper.InitBodies();
 }
 
diff --git a/Engine/Source/Runtime/Engine/Private/PhysicsEngine/Experimental/PhysScene_AsyncPhysicsStateJobQueue.cpp b/Engine/Source/Runtime/Engine/Private/PhysicsEngine/Experimental/PhysScene_AsyncPhysicsStateJobQueue.cpp
--- a/Engine/Source/Runtime/Engine/Private/PhysicsEngine/Experimental/PhysScene_AsyncPhysicsStateJobQueue.cpp
+++ b/Engine/Source/Runtime/Engine/Private/PhysicsEngine/Experimental/PhysScene_AsyncPhysicsStateJobQueue.cpp
@@ -2,8 +2,6 @@
 
 #include "Physics/Experimental/PhysScene_AsyncPhysicsStateJobQueue.h"
 #include "Physics/Experimental/PhysScene_Chaos.h"
-#include "Components/ActorComponent.h"
-#include "Components/PrimitiveComponent.h"
 #include "PhysicsEngine/BodySetup.h"
 #include "Engine/World.h"
 
@@ -76,12 +74,7 @@ bool FPhysScene_AsyncPhysicsStateJobQueue::AreAsyncTasksExecuted() const
 
 bool FPhysScene_AsyncPhysicsStateJobQueue::IsCompleted() const
 { 
-	if (!PendingJobs.IsValid() && PipedJobs.IsEmpty() && PipeTasks.IsEmpty())
-	{
-		check(AreAsyncTasksExecuted());
-		return true;
-	}
-	return false;
+	return AreAsyncTasksExecuted() && PipedJobs.IsEmpty() && PipeTasks.IsEmpty();
 }
 
 void FPhysScene_AsyncPhysicsStateJobQueue::Launch()
@@ -131,7 +124,7 @@ void FPhysScene_AsyncPhysicsStateJobQueue::WaitForAsyncTasksExecution()
 	UE::Tasks::Wait(PipeTasks);
 }
 
-void FPhysScene_AsyncPhysicsStateJobQueue::Tick()
+void FPhysScene_AsyncPhysicsStateJobQueue::Tick(bool bWaitForCompletion)
 {
 	TRACE_CPUPROFILER_EVENT_SCOPE(FPhysScene_AsyncPhysicsStateJobQueue::Tick);
 	check(IsInGameThread());
@@ -140,7 +133,10 @@ void FPhysScene_AsyncPhysicsStateJobQueue::Tick()
 	Launch();
 
 	// Wait for tasks if the world is inside a blocking load
-	if (PhysScene->GetOwningWorld()->GetIsInBlockTillLevelStreamingCompleted())
+	if (bWaitForCompletion || 
+		PhysScene->GetOwningWorld()->GetIsInBlockTillLevelStreamingCompleted() ||
+		PhysScene->GetOwningWorld()->GetShouldForceUnloadStreamingLevels() ||
+		PhysScene->GetOwningWorld()->IsBeingCleanedUp())
 	{
 		WaitForAsyncTasksExecution();
 	}
@@ -189,12 +185,7 @@ void FPhysScene_AsyncPhysicsStateJobQueue::FJobSet::OnPreExecute_GameThread() co
 		{
 			if (Job.Type == EJobType::CreatePhysicsState)
 			{
-				UPrimitiveComponent* PrimitiveComponent = Cast<UPrimitiveComponent>(Job.Component);
-				UBodySetup* BodySetup = PrimitiveComponent ? PrimitiveComponent->GetBodySetup() : nullptr;
-				if (BodySetup && !BodySetup->bCreatedPhysicsMeshes)
-				{
-					UniqueBodySetups.Add(BodySetup);
-				}
+				Job.Processor->CollectBodySetupsWithPhysicsMeshesToCreate(UniqueBodySetups);
 			}
 		}
 		if (UniqueBodySetups.Num())
@@ -217,7 +208,7 @@ void FPhysScene_AsyncPhysicsStateJobQueue::FJob::OnPreExecute_GameThread() const
 	{
 		case EJobType::DestroyPhysicsState:
 		{
-			Component->OnAsyncDestroyPhysicsStateBegin_GameThread();
+			Processor->OnAsyncDestroyPhysicsStateBegin_GameThread();
 		}
 		break;
 	}
@@ -234,20 +225,20 @@ void FPhysScene_AsyncPhysicsStateJobQueue::FJobSet::OnPostExecute_GameThread() c
 
 void FPhysScene_AsyncPhysicsStateJobQueue::FJob::Execute() const
 {
-	if (const bool bPendingKill = Component->GetOwner() && !::IsValid(Component->GetOwner()); !bPendingKill)
+	if (IsValid())
 	{
 		switch (Type)
 		{
 			case EJobType::CreatePhysicsState:
 			{
 				TRACE_CPUPROFILER_EVENT_SCOPE(FPhysScene_AsyncPhysicsStateJobQueue::CreatePhysicsState);
-				Component->OnAsyncCreatePhysicsState();
+				Processor->OnAsyncCreatePhysicsState();
 			}
 			break;
 			case EJobType::DestroyPhysicsState:
 			{
 				TRACE_CPUPROFILER_EVENT_SCOPE(FPhysScene_AsyncPhysicsStateJobQueue::DestroyPhysicsState);
-				Component->OnAsyncDestroyPhysicsState();
+				Processor->OnAsyncDestroyPhysicsState();
 			}
 			break;
 		}
@@ -256,18 +247,18 @@ void FPhysScene_AsyncPhysicsStateJobQueue::FJob::Execute() const
 
 void FPhysScene_AsyncPhysicsStateJobQueue::FJob::OnPostExecute_GameThread() const
 {
-	if (const bool bPendingKill = Component->GetOwner() && !::IsValid(Component->GetOwner()); !bPendingKill)
+	if (IsValid()) 
 	{
 		switch (Type)
 		{
 			case EJobType::CreatePhysicsState:
 			{
-				Component->OnAsyncCreatePhysicsStateEnd_GameThread();
+				Processor->OnAsyncCreatePhysicsStateEnd_GameThread();
 			}
 			break;
 			case EJobType::DestroyPhysicsState:
 			{
-				Component->OnAsyncDestroyPhysicsStateEnd_GameThread();
+				Processor->OnAsyncDestroyPhysicsStateEnd_GameThread();
 			}
 			break;
 		}
diff --git a/Engine/Source/Runtime/Engine/Private/PhysicsEngine/Experimental/PhysScene_Chaos.cpp b/Engine/Source/Runtime/Engine/Private/PhysicsEngine/Experimental/PhysScene_Chaos.cpp
--- a/Engine/Source/Runtime/Engine/Private/PhysicsEngine/Experimental/PhysScene_Chaos.cpp
+++ b/Engine/Source/Runtime/Engine/Private/PhysicsEngine/Experimental/PhysScene_Chaos.cpp
@@ -1829,59 +1829,58 @@ bool FPhysScene_Chaos::SupportsAsyncPhysicsStateDestruction()
 	return SupportsAsyncPhysicsStateCreation();
 }
 
-bool FPhysScene_Chaos::PushAsyncCreatePhysicsStateJob(UActorComponent* Component)
+bool FPhysScene_Chaos::PushAsyncCreatePhysicsStateJob(IAsyncPhysicsStateProcessor* Processor)
 {
 	if (SupportsAsyncPhysicsStateCreation())
 	{
 		check(IsInGameThread());
-		check(!Component->IsPhysicsStateCreated());
-		check(Component->AllowsAsyncPhysicsStateCreation());
+		check(!Processor->IsAsyncPhysicsStateCreated());
+		check(Processor->AllowsAsyncPhysicsStateCreation());
 
 		if (!AsyncPhysicsStateJobQueue.IsValid())
 		{
 			AsyncPhysicsStateJobQueue = MakeUnique<FPhysScene_AsyncPhysicsStateJobQueue>(this);
 		}
-		AsyncPhysicsStateJobQueue->AddJob({ Component, FPhysScene_AsyncPhysicsStateJobQueue::EJobType::CreatePhysicsState });
+		AsyncPhysicsStateJobQueue->AddJob({ Processor, FPhysScene_AsyncPhysicsStateJobQueue::EJobType::CreatePhysicsState });
 		return true;
 	}
 
 	return false;
 }
 
-bool FPhysScene_Chaos::PushAsyncDestroyPhysicsStateJob(UActorComponent* Component)
+bool FPhysScene_Chaos::PushAsyncDestroyPhysicsStateJob(IAsyncPhysicsStateProcessor* Processor)
 {
 	if (SupportsAsyncPhysicsStateDestruction())
 	{
 		check(IsInGameThread());
-		check(Component->IsPhysicsStateCreated());
-		check(Component->AllowsAsyncPhysicsStateDestruction());
+		check(Processor->IsAsyncPhysicsStateCreated());
+		check(Processor->AllowsAsyncPhysicsStateDestruction());
 
 		if (!AsyncPhysicsStateJobQueue.IsValid())
 		{
 			AsyncPhysicsStateJobQueue = MakeUnique<FPhysScene_AsyncPhysicsStateJobQueue>(this);
 		}
-
-		AsyncPhysicsStateJobQueue->AddJob({ Component, FPhysScene_AsyncPhysicsStateJobQueue::EJobType::DestroyPhysicsState });
+		AsyncPhysicsStateJobQueue->AddJob({ Processor, FPhysScene_AsyncPhysicsStateJobQueue::EJobType::DestroyPhysicsState });
 	}
 	return true;
 }
 
-void FPhysScene_Chaos::RemoveAsyncPhysicsStateJobs(UActorComponent* Component)
+void FPhysScene_Chaos::RemoveAsyncPhysicsStateJobs(IAsyncPhysicsStateProcessor* Processor)
 {
 	check(IsInGameThread());
 	if (AsyncPhysicsStateJobQueue.IsValid())
 	{
-		AsyncPhysicsStateJobQueue->RemoveJob({ Component, FPhysScene_AsyncPhysicsStateJobQueue::EJobType::CreatePhysicsState });
-		AsyncPhysicsStateJobQueue->RemoveJob({ Component, FPhysScene_AsyncPhysicsStateJobQueue::EJobType::DestroyPhysicsState });
+		AsyncPhysicsStateJobQueue->RemoveJob({ Processor, FPhysScene_AsyncPhysicsStateJobQueue::EJobType::CreatePhysicsState });
+		AsyncPhysicsStateJobQueue->RemoveJob({ Processor, FPhysScene_AsyncPhysicsStateJobQueue::EJobType::DestroyPhysicsState });
 	}
 }
 
-void FPhysScene_Chaos::ProcessAsyncPhysicsStateJobs()
+void FPhysScene_Chaos::ProcessAsyncPhysicsStateJobs(bool bWaitForCompletion)
 {
 	check(IsInGameThread());
 	if (AsyncPhysicsStateJobQueue.IsValid())
 	{
-		AsyncPhysicsStateJobQueue->Tick();
+		AsyncPhysicsStateJobQueue->Tick(bWaitForCompletion);
 		if (AsyncPhysicsStateJobQueue->IsCompleted())
 		{
 			AsyncPhysicsStateJobQueue.Reset();
diff --git a/Engine/Source/Runtime/Engine/Public/Physics/Experimental/AsyncPhysicsStateProcessorInterface.h b/Engine/Source/Runtime/Engine/Public/Physics/Experimental/AsyncPhysicsStateProcessorInterface.h
new file mode 100644
--- /dev/null
+++ b/Engine/Source/Runtime/Engine/Public/Physics/Experimental/AsyncPhysicsStateProcessorInterface.h
@@ -0,0 +1,44 @@
+// Copyright Epic Games, Inc. All Rights Reserved.
+
+#pragma once
+
+#include "Containers/Set.h"
+
+// Interface used by FPhysScene_AsyncPhysicsStateJobQueue
+
+class UObject;
+class UBodySetup;
+
+class IAsyncPhysicsStateProcessor
+{
+public:
+	/** Returns whether this component allows having its physics state being created asynchronously (outside of the GameThread). */
+	virtual bool AllowsAsyncPhysicsStateCreation() const { return false; }
+	
+	/** Returns whether this component allows having its physics state being destroyed asynchronously (outside of the GameThread). */
+	virtual bool AllowsAsyncPhysicsStateDestruction() const { return false; }
+
+	/** Returns whether the physics state is created. */
+	virtual bool IsAsyncPhysicsStateCreated() const { return false; }
+
+	/** Returns the associated UObject for this processor. */
+	virtual UObject* GetAsyncPhysicsStateObject() const { return nullptr; }
+
+	/** Returns body setups that needs to create their physics meshes before physics state asynchronous creation */
+	virtual void CollectBodySetupsWithPhysicsMeshesToCreate(TSet<UBodySetup*>& OutBodySetups) const {}
+
+	/** Used to create any physics engine information for this component outside of the GameThread. */
+	virtual void OnAsyncCreatePhysicsState() {}
+
+	/** Called on the GameThread once the component's physic engine information is created. */
+	virtual void OnAsyncCreatePhysicsStateEnd_GameThread() {}
+
+	/** Called on the GameThread before the component's physic engine information is destroyed. */
+	virtual void OnAsyncDestroyPhysicsStateBegin_GameThread() {}
+
+	/** Used to destroy any physics engine information for this component outside of the GameThread. */
+	virtual void OnAsyncDestroyPhysicsState() {}
+
+	/** Called on the GameThread once the component's physic engine information is destroyed. */
+	virtual void OnAsyncDestroyPhysicsStateEnd_GameThread() {}
+};
\ No newline at end of file
diff --git a/Engine/Source/Runtime/Engine/Public/Physics/Experimental/PhysScene_AsyncPhysicsStateJobQueue.h b/Engine/Source/Runtime/Engine/Public/Physics/Experimental/PhysScene_AsyncPhysicsStateJobQueue.h
--- a/Engine/Source/Runtime/Engine/Public/Physics/Experimental/PhysScene_AsyncPhysicsStateJobQueue.h
+++ b/Engine/Source/Runtime/Engine/Public/Physics/Experimental/PhysScene_AsyncPhysicsStateJobQueue.h
@@ -5,10 +5,9 @@
 #include "Tasks/Pipe.h"
 #include "Tasks/Task.h"
 #include "Async/Fundamental/Task.h"
-#include "Components/ActorComponent.h"
+#include "UObject/Object.h"
 #include "PhysicsInterfaceDeclaresCore.h"
-
-class UActorComponent;
+#include "Physics/Experimental/AsyncPhysicsStateProcessorInterface.h"
 
 struct FPhysScene_AsyncPhysicsStateJobQueue
 {
@@ -23,27 +22,30 @@ struct FPhysScene_AsyncPhysicsStateJobQueue
 
 	struct FJob
 	{
-		FJob(UActorComponent* InComponent, EJobType InType)
-			: Component(InComponent)
+		FJob(IAsyncPhysicsStateProcessor* InProcessor, EJobType InType)
+			: Processor(InProcessor)
 			, Type(InType)
 		{}
 
-		bool IsValid() const { return ::IsValid(Component); }
+		bool IsValid() const
+		{
+			return Processor && ::IsValid(Processor->GetAsyncPhysicsStateObject());
+		}
 		void Execute() const;
 		void OnPreExecute_GameThread() const;
 		void OnPostExecute_GameThread() const;
 
 		inline bool operator == (const FJob& Other) const
 		{
-			return (Component == Other.Component) && (Type == Other.Type);
+			return (Processor == Other.Processor) && (Type == Other.Type);
 		}
 
 		friend inline uint32 GetTypeHash(const FJob& InJob)
 		{
-			return GetTypeHash(TPair<UActorComponent*, EJobType>{InJob.Component, InJob.Type});
+			return GetTypeHash(TPair<IAsyncPhysicsStateProcessor*, EJobType>{InJob.Processor, InJob.Type});
 		}
 
-		UActorComponent* Component;
+		IAsyncPhysicsStateProcessor* Processor;
 		EJobType Type;
 	};
 
@@ -72,7 +74,7 @@ struct FPhysScene_AsyncPhysicsStateJobQueue
 		TSet<FJob> Jobs;
 	};
 
-	void Tick();
+	void Tick(bool bWaitForCompletion = false);
 	void AddJob(const FJob& Job);
 	void RemoveJob(const FJob& Job);
 	bool IsCompleted() const;
diff --git a/Engine/Source/Runtime/Engine/Public/Physics/Experimental/PhysScene_Chaos.h b/Engine/Source/Runtime/Engine/Public/Physics/Experimental/PhysScene_Chaos.h
--- a/Engine/Source/Runtime/Engine/Public/Physics/Experimental/PhysScene_Chaos.h
+++ b/Engine/Source/Runtime/Engine/Public/Physics/Experimental/PhysScene_Chaos.h
@@ -26,6 +26,7 @@
 
 class UPrimitiveComponent;
 
+class IAsyncPhysicsStateProcessor;
 class AdvanceOneTimeStepTask;
 class IPhysicsReplication;
 class IPhysicsReplicationLOD;
@@ -270,10 +271,10 @@ public:
 
 	ENGINE_API static bool SupportsAsyncPhysicsStateCreation();
 	ENGINE_API static bool SupportsAsyncPhysicsStateDestruction();
-	ENGINE_API bool PushAsyncCreatePhysicsStateJob(UActorComponent* Component);
-	ENGINE_API bool PushAsyncDestroyPhysicsStateJob(UActorComponent* Component);
-	ENGINE_API void RemoveAsyncPhysicsStateJobs(UActorComponent* Component);
-	ENGINE_API void ProcessAsyncPhysicsStateJobs();
+	ENGINE_API bool PushAsyncCreatePhysicsStateJob(IAsyncPhysicsStateProcessor* Processor);
+	ENGINE_API bool PushAsyncDestroyPhysicsStateJob(IAsyncPhysicsStateProcessor* Processor);
+	ENGINE_API void RemoveAsyncPhysicsStateJobs(IAsyncPhysicsStateProcessor* Processor);
+	ENGINE_API void ProcessAsyncPhysicsStateJobs(bool bWaitForCompletion = false);
 
 	UChaosEventRelay* GetChaosEventRelay() const { return ChaosEventRelay.Get(); }
 
