- Added UWorld::OnAddLevelToWorldExtension and UWorld::OnRemoveLevelFromWorldExtension to allow external systems to add custom code to AddToWorld/RemoveFromWorld
- Added FWorldDelegates::LevelComponentsCleared to notify external systems when level components are all unregistered
- Added FWorldDelegates::LevelComponentsUpdated to notify external systems when level components are all updated synchronously
- Made UWorld::MakingVisibleLevels thread-safe
(This development in required for an upcoming feature)

diff --git a/Engine/Source/Runtime/Engine/Classes/Engine/Level.h b/Engine/Source/Runtime/Engine/Classes/Engine/Level.h
--- a/Engine/Source/Runtime/Engine/Classes/Engine/Level.h
+++ b/Engine/Source/Runtime/Engine/Classes/Engine/Level.h
@@ -648,6 +648,8 @@ public:
 	uint8										bWasDuplicatedForPIE:1;
 	/** Whether the level is currently being removed from the world */
 	uint8										bIsBeingRemoved:1;
+	/** Whether the level incremental unregistration of components is completed */
+	uint8										bIncrementalUnregisterComponentsCompleted:1;
 	/** Whether this level has gone through a complete rerun construction script pass. */
 	uint8										bHasRerunConstructionScripts:1;
 	/** Whether the level had its actor cluster created. This doesn't mean that the creation was successful. */
diff --git a/Engine/Source/Runtime/Engine/Classes/Engine/World.h b/Engine/Source/Runtime/Engine/Classes/Engine/World.h
--- a/Engine/Source/Runtime/Engine/Classes/Engine/World.h
+++ b/Engine/Source/Runtime/Engine/Classes/Engine/World.h
@@ -995,7 +995,7 @@ public:
 	bool IsLevelMakingVisible(const ULevel* InLevel) const;
 
 	/** Returns true if any level is in the process of being made visible */
-	bool HasAnyLevelMakingVisible() const { return MakingVisibleLevels.Num() > 0; }
+	bool HasAnyLevelMakingVisible() const;
 
 	/** Returns the level, if any, in the process of being made invisible */
 	UE_DEPRECATED(5.6, "Use IsLevelMakingInvisible instead.")
@@ -1005,7 +1005,7 @@ public:
 	bool IsLevelMakingInvisible(const ULevel* InLevel) const;
 
 	/** Returns true if any level is in the process of being made invisible */
-	bool HasAnyLevelMakingInvisible() const { return MakingInvisibleLevels.Num() > 0; }
+	bool HasAnyLevelMakingInvisible() const;
 
 	/** Add a streaming level to the list of streamed levels to consider. */
 	void AddStreamingLevel(ULevelStreaming* StreamingLevelToAdd);
@@ -1093,6 +1093,9 @@ private:
 	/** Stores whether the game world supports for a client to use "making invisible" transaction requests to the server. */
 	mutable TOptional<bool> bSupportsMakingInvisibleTransactionRequests;
 
+	/** Used for MT access to MakingVisibleLevels and MakingInvisibleLevels */
+	mutable FRWLock PendingVisibilityLock;
+
 	/** Level(s) in the queue to be made visible. */
 	UPROPERTY(Transient)
 	TArray<TObjectPtr<class ULevel>>						MakingVisibleLevels;
@@ -1692,6 +1695,18 @@ private:
 		OnWorldPartitionUninitializedEvent.Broadcast(InWorldPartition);
 	}
 
+	/** Broadcasted when Level is being added to the world (used to add custom logic).
+      * Set bOutHasCompleted to false if custom execution has not completed.
+	  */
+	DECLARE_EVENT_ThreeParams(UWorld, FAddLevelToWorldExtensionEvent, ULevel* /*InLevel*/, const bool /*bWaitForCompletion*/, bool& /*bOutHasCompleted*/);
+	FAddLevelToWorldExtensionEvent AddLevelToWorldExtensionEvent;
+
+	/** Broadcasted when Level is being removed from the world (used to add custom logic).
+      * Set bOutHasCompleted to false if custom execution has not completed.
+	  */
+	DECLARE_EVENT_ThreeParams(UWorld, FRemoveLevelFromWorldExtensionEvent, ULevel* /*InLevel*/, const bool /*bWaitForCompletion*/, bool& /*bOutHasCompleted*/);
+	FRemoveLevelFromWorldExtensionEvent RemoveLevelFromWorldExtensionEvent;
+
 	friend FWorldPartitionEvents;
 
 #if WITH_EDITOR
@@ -3235,6 +3250,18 @@ public:
 	 */
 	void EnsureCollisionTreeIsBuilt();
 
+	/** Sent when Level is being added to the world (used to add custom logic). */
+	FAddLevelToWorldExtensionEvent& OnAddLevelToWorldExtension()
+	{
+		return AddLevelToWorldExtensionEvent;
+	}
+
+	/** Sent when Level is being removed from the world (used to add custom logic). */
+	FRemoveLevelFromWorldExtensionEvent& OnRemoveLevelFromWorldExtension()
+	{
+		return RemoveLevelFromWorldExtensionEvent;
+	}
+
 #if WITH_EDITOR	
 	/** Returns the SelectedLevelsChangedEvent member. */
 	FOnSelectedLevelsChangedEvent& OnSelectedLevelsChanged() { return SelectedLevelsChangedEvent; }
@@ -4300,6 +4327,14 @@ public:
 	// removed from the world without a RemoveFromWorld call for each)
 	static FOnLevelChanged			LevelRemovedFromWorld;
 
+	DECLARE_MULTICAST_DELEGATE_TwoParams(FLevelComponentsEvent, UWorld*, ULevel*);
+
+	/** Sent when level components are all unregistered (ULevel::ClearLevelComponents). */
+	static FLevelComponentsEvent	LevelComponentsCleared;
+
+	/** Sent when level components are all updated (ULevel::UpdateLevelComponents). */
+	static FLevelComponentsEvent	LevelComponentsUpdated;
+
 	// Called after offset was applied to a level
 	DECLARE_MULTICAST_DELEGATE_FourParams(FLevelOffsetEvent, ULevel*,  UWorld*, const FVector&, bool);
 	static FLevelOffsetEvent		PostApplyLevelOffset;
diff --git a/Engine/Source/Runtime/Engine/Private/Level.cpp b/Engine/Source/Runtime/Engine/Private/Level.cpp
--- a/Engine/Source/Runtime/Engine/Private/Level.cpp
+++ b/Engine/Source/Runtime/Engine/Private/Level.cpp
@@ -446,6 +446,7 @@ ULevel::ULevel( const FObjectInitializer& ObjectInitializer )
 	bGarbageCollectionClusteringEnabled = true;
 	bIsPartitioned = false;
 	bStaticComponentsRegisteredInStreamingManager = false;
+	bIncrementalUnregisterComponentsCompleted = false;
 	IncrementalComponentState = EIncrementalComponentState::Init;
 	IncrementalUnregisterComponentState = EIncrementalUnregisterComponentState::Init;
 	CurrentActorIndexForIncrementalUpdate = 0;
@@ -1562,6 +1563,8 @@ void ULevel::ClearLevelComponents()
 			Actor->UnregisterAllComponents();
 		}
 	}
+	bIncrementalUnregisterComponentsCompleted = false;
+	FWorldDelegates::LevelComponentsCleared.Broadcast(GetWorld(), this);
 }
 
 void ULevel::BeginDestroy()
@@ -1641,6 +1644,7 @@ void ULevel::UpdateLevelComponents(bool bRerunConstructionScripts, FRegisterComp
 {
 	// Update all components in one swoop.
 	IncrementalUpdateComponents( 0, bRerunConstructionScripts, Context);
+	FWorldDelegates::LevelComponentsUpdated.Broadcast(GetWorld(), this);
 }
 
 /**
@@ -2148,6 +2152,13 @@ bool ULevel::IncrementalUnregisterComponents(int32 NumComponentsToUnregister, FU
 		checkf(OwningWorld->IsGameWorld(), TEXT("Cannot call IncrementalUnregisterComponents with non 0 argument in the Editor/ commandlets."));
 	}
 
+	if (bIncrementalUnregisterComponentsCompleted)
+	{
+		check(bIncrementalUnregisterComponents);
+		check(IncrementalUnregisterComponentState == EIncrementalUnregisterComponentState::Init);
+		return true;
+	}
+
 	auto IsServer = [this]()
 	{
 		const ENetMode NetMode = OwningWorld->GetNetMode();
@@ -2229,6 +2240,7 @@ bool ULevel::IncrementalUnregisterComponents(int32 NumComponentsToUnregister, FU
 			}
 			break;
 		case EIncrementalUnregisterComponentState::Finalize:
+			bIncrementalUnregisterComponentsCompleted = true;
 			IncrementalUnregisterComponentState = EIncrementalUnregisterComponentState::Init;
 			return true;
 			break;
diff --git a/Engine/Source/Runtime/Engine/Private/World.cpp b/Engine/Source/Runtime/Engine/Private/World.cpp
--- a/Engine/Source/Runtime/Engine/Private/World.cpp
+++ b/Engine/Source/Runtime/Engine/Private/World.cpp
@@ -644,6 +644,8 @@ FWorldDelegates::FWorldEvent FWorldDelegates::OnPreWorldFinishDestroy;
 FWorldDelegates::FOnLevelChanged FWorldDelegates::LevelAddedToWorld;
 FWorldDelegates::FOnLevelChanged FWorldDelegates::PreLevelRemovedFromWorld;
 FWorldDelegates::FOnLevelChanged FWorldDelegates::LevelRemovedFromWorld;
+FWorldDelegates::FLevelComponentsEvent FWorldDelegates::LevelComponentsCleared;
+FWorldDelegates::FLevelComponentsEvent FWorldDelegates::LevelComponentsUpdated;
 FWorldDelegates::FLevelOffsetEvent FWorldDelegates::PostApplyLevelOffset;
 FWorldDelegates::FLevelTransformEvent FWorldDelegates::PostApplyLevelTransform;
 PRAGMA_DISABLE_DEPRECATION_WARNINGS;
@@ -3251,8 +3253,21 @@ const AServerStreamingLevelsVisibility* UWorld::GetServerStreamingLevelsVisibili
 	return ServerStreamingLevelsVisibility;
 }
 
+bool UWorld::HasAnyLevelMakingVisible() const
+{
+	FReadScopeLock ReadLock(PendingVisibilityLock);
+	return MakingVisibleLevels.Num() > 0; 
+}
+
+bool UWorld::HasAnyLevelMakingInvisible() const
+{
+	FReadScopeLock ReadLock(PendingVisibilityLock);
+	return MakingInvisibleLevels.Num() > 0; 
+}
+
 bool UWorld::IsLevelMakingVisible(const ULevel* InLevel) const
 {
+	FReadScopeLock ReadLock(PendingVisibilityLock);
 	const bool bIsMakingVisible = InLevel && (InLevel->GetPendingVisibilityState() == ELevelPendingVisibilityState::MakingVisible);
 	check(bIsMakingVisible == MakingVisibleLevels.Contains(InLevel));
 	return bIsMakingVisible;
@@ -3260,6 +3275,7 @@ bool UWorld::IsLevelMakingVisible(const ULevel* InLevel) const
 
 bool UWorld::IsLevelMakingInvisible(const ULevel* InLevel) const
 {
+	FReadScopeLock ReadLock(PendingVisibilityLock);
 	const bool bIsMakingInvisible = InLevel && (InLevel->GetPendingVisibilityState() == ELevelPendingVisibilityState::MakingInvisible);
 	check(bIsMakingInvisible == MakingInvisibleLevels.Contains(InLevel));
 	return bIsMakingInvisible;
@@ -3283,6 +3299,12 @@ void UWorld::EnableActorTickAndUserCallbacks(bool bEnable)
 
 void UWorld::SetLevelPendingVisibilityState(ULevel* InLevel, ELevelPendingVisibilityState InState)
 {
+	FWriteScopeLock WriteLock(PendingVisibilityLock);
+
+	if (InLevel->GetPendingVisibilityState() == InState)
+	{
+		return;
+	}
 	FLevelPendingVisibilityStateAccessor::SetPendingVisibilityState(InLevel, InState);
 
 	switch (InState)
@@ -3290,14 +3312,19 @@ void UWorld::SetLevelPendingVisibilityState(ULevel* InLevel, ELevelPendingVisibi
 	case ELevelPendingVisibilityState::None:
 		MakingVisibleLevels.Remove(InLevel);
 		MakingInvisibleLevels.Remove(InLevel);
+		check(!MakingVisibleLevels.Contains(InLevel));
+		check(!MakingInvisibleLevels.Contains(InLevel));
 		break;
 	case ELevelPendingVisibilityState::MakingVisible:
 		check(!MakingInvisibleLevels.Contains(InLevel));
 		MakingVisibleLevels.Add(InLevel);
+		check(MakingVisibleLevels.Num() <= FMath::Max<int32>(1, GMaximumMakingVisibleLevels));
 		break;
 	case ELevelPendingVisibilityState::MakingInvisible:
 		check(!MakingVisibleLevels.Contains(InLevel));
 		MakingInvisibleLevels.Add(InLevel);
+		check(!InLevel->bIncrementalUnregisterComponentsCompleted);
+		check(MakingInvisibleLevels.Num() <= FMath::Max<int32>(1, GMaximumMakingInvisibleLevels));
 		break;
 	}
 }
@@ -3561,6 +3588,7 @@ void UWorld::AddToWorld( ULevel* Level, const FTransform& LevelTransform, bool b
 		// Initialize all actors and start execution.
 		if (bExecuteNextStep && !(Level->bAlreadyInitializedNetworkActors && Level->bAlreadyClearedActorsSeamlessTravelFlag))
 		{
+			TRACE_CPUPROFILER_EVENT_SCOPE(AddToWorld_InitializeNetworkActors);
 			QUICK_SCOPE_CYCLE_COUNTER(STAT_AddToWorldTime_InitializeNetworkActors);
 			CSV_SCOPED_TIMING_STAT(LevelStreamingDetail, AddToWorld_InitNetworkActors);
 			SCOPE_TIME_TO_VAR(&InitActorTime);
@@ -3585,6 +3613,7 @@ void UWorld::AddToWorld( ULevel* Level, const FTransform& LevelTransform, bool b
 		// Route various initialization functions and set volumes.
 		if (bExecuteNextStep && !Level->IsFinishedRouteActorInitialization())
 		{
+			TRACE_CPUPROFILER_EVENT_SCOPE(AddToWorld_RouteActorInit);
 			CSV_SCOPED_TIMING_STAT(LevelStreamingDetail, AddToWorld_RouteActorInit);
 			QUICK_SCOPE_CYCLE_COUNTER(STAT_AddToWorldTime_RouteActorInitialize);
 			SCOPE_TIME_TO_VAR(&RouteActorInitializeTime);
@@ -3599,15 +3628,20 @@ void UWorld::AddToWorld( ULevel* Level, const FTransform& LevelTransform, bool b
 			bExecuteNextStep = Level->IsFinishedRouteActorInitialization() && (!bConsiderTimeLimit || !IsTimeLimitExceeded( TEXT("routing Initialize on actors"), StartTime, Level, TimeLimit ));
 		}
 		
-		// Sort the actor list; can't do this on save as the relevant properties for sorting might have been changed by code
-		if( bExecuteNextStep && !Level->bAlreadySortedActorList )
+		if( bExecuteNextStep )
 		{
-			CSV_SCOPED_TIMING_STAT(LevelStreamingDetail, AddToWorld_SortActorList);
-			SCOPE_TIME_TO_VAR(&SortActorListTime);
+		// Sort the actor list; can't do this on save as the relevant properties for sorting might have been changed by code
+			if (!Level->bAlreadySortedActorList)
+			{
+				TRACE_CPUPROFILER_EVENT_SCOPE(AddToWorld_SortActorList);
+				CSV_SCOPED_TIMING_STAT(LevelStreamingDetail, AddToWorld_SortActorList);
+				SCOPE_TIME_TO_VAR(&SortActorListTime);
+
+				Level->SortActorList();
+				Level->bAlreadySortedActorList = true;
+				bExecuteNextStep = (!bConsiderTimeLimit || !IsTimeLimitExceeded( TEXT("sorting actor list"), StartTime, Level, TimeLimit ));
+			}
 
-			Level->SortActorList();
-			Level->bAlreadySortedActorList = true;
-			bExecuteNextStep = (!bConsiderTimeLimit || !IsTimeLimitExceeded( TEXT("sorting actor list"), StartTime, Level, TimeLimit ));
 			bPerformedLastStep = true;
 		}
 	}
@@ -3616,6 +3650,17 @@ void UWorld::AddToWorld( ULevel* Level, const FTransform& LevelTransform, bool b
 		bPerformedLastStep = true;
 	}
 
+	{
+		TRACE_CPUPROFILER_EVENT_SCOPE(AddToWorld_Extension);
+		bool bWaitForCompletion = !bConsiderTimeLimit;
+		bool bOutHasCompleted = true;
+		AddLevelToWorldExtensionEvent.Broadcast(Level, bWaitForCompletion, bOutHasCompleted);
+		if (!bOutHasCompleted)
+		{
+			bPerformedLastStep = false;
+		}
+	}
+
 	Level->bIsAssociatingLevel = false;
 
 	// We're done.
@@ -3783,7 +3828,7 @@ void UWorld::RemoveFromWorld( ULevel* Level, bool bAllowIncrementalRemoval, FNet
 		// Keep track of timing.
 		double PerfTrackStartTime = FPlatformTime::Seconds();
 #endif
-
+		bool bWaitForCompletion = false;
 		bool bFinishRemovingLevel = true;
 		if ( bAllowIncrementalRemoval && GLevelStreamingUnregisterComponentsTimeLimit > 0.0 )
 		{
@@ -3826,8 +3871,6 @@ void UWorld::RemoveFromWorld( ULevel* Level, bool bAllowIncrementalRemoval, FNet
 					{
 						if (Level->IncrementalUnregisterComponents(NumComponentsToUnregister, &Context))
 						{
-							// We're done, so the level can be removed
-							SetLevelPendingVisibilityState(Level, ELevelPendingVisibilityState::None);
 							bFinishRemovingLevel = true;
 							break;
 						}
@@ -3847,16 +3890,31 @@ void UWorld::RemoveFromWorld( ULevel* Level, bool bAllowIncrementalRemoval, FNet
 			// Finish current level pending invisibility without time limit
 			check(Level->bIsBeingRemoved);
 			while (!Level->IncrementalUnregisterComponents(MAX_int32)) {}
-			SetLevelPendingVisibilityState(Level, ELevelPendingVisibilityState::None);
-			bFinishRemovingLevel = true;
+			bWaitForCompletion = true;
+			check(bFinishRemovingLevel);
 		}
 		else
 		{
 			BeginRemoval();
+			bWaitForCompletion = true;
+			check(bFinishRemovingLevel);
+		}
+
+		{
+			TRACE_CPUPROFILER_EVENT_SCOPE(RemoveFromWorld_Extension);
+			check(bWaitForCompletion == !(bAllowIncrementalRemoval && GLevelStreamingUnregisterComponentsTimeLimit > 0.0));
+			bool bOutHasCompleted = true;
+			RemoveLevelFromWorldExtensionEvent.Broadcast(Level, bWaitForCompletion, bOutHasCompleted);
+			if (!bOutHasCompleted)
+			{
+				bFinishRemovingLevel = false;
+			}
 		}
 
 		if (bFinishRemovingLevel)
 		{
+			// The level can be removed from the pending visibility state
+			SetLevelPendingVisibilityState(Level, ELevelPendingVisibilityState::None);
 			TRACE_CPUPROFILER_EVENT_SCOPE(RemoveFromWorld_FinishRemovingLevel);
 			double StartTime = FPlatformTime::Seconds();
 
