Fixed CIS - AutoRTFM : replaced FRWLock with FTransactionallySafeRWLock

diff --git a/Engine/Source/Runtime/Engine/Classes/Engine/World.h b/Engine/Source/Runtime/Engine/Classes/Engine/World.h
--- a/Engine/Source/Runtime/Engine/Classes/Engine/World.h
+++ b/Engine/Source/Runtime/Engine/Classes/Engine/World.h
@@ -31,6 +31,7 @@
 #include "CollisionProfile.h"
 #include "RHIFeatureLevel.h"
 #include "Engine/WorldInitializationValues.h"
+#include "Misc/TransactionallySafeRWLock.h"
 
 #include "World.generated.h"
 
@@ -1094,7 +1095,7 @@ private:
 	mutable TOptional<bool> bSupportsMakingInvisibleTransactionRequests;
 
 	/** Used for MT access to MakingVisibleLevels and MakingInvisibleLevels */
-	mutable FRWLock PendingVisibilityLock;
+	mutable FTransactionallySafeRWLock PendingVisibilityLock;
 
 	/** Level(s) in the queue to be made visible. */
 	UPROPERTY(Transient)
diff --git a/Engine/Source/Runtime/Engine/Private/World.cpp b/Engine/Source/Runtime/Engine/Private/World.cpp
--- a/Engine/Source/Runtime/Engine/Private/World.cpp
+++ b/Engine/Source/Runtime/Engine/Private/World.cpp
@@ -84,6 +84,7 @@
 #include "Streaming/LevelStreamingDelegates.h"
 #include "Streaming/StreamingWorldSubsystemInterface.h"
 #include "Components/UnregisterComponentContext.h"
+#include "Misc/ScopeRWLock.h"
 
 #if WITH_EDITOR
 	#include "DerivedDataCacheInterface.h"
@@ -3255,19 +3256,19 @@ const AServerStreamingLevelsVisibility* UWorld::GetServerStreamingLevelsVisibili
 
 bool UWorld::HasAnyLevelMakingVisible() const
 {
-	FReadScopeLock ReadLock(PendingVisibilityLock);
+	UE::TReadScopeLock ReadLock(PendingVisibilityLock);
 	return MakingVisibleLevels.Num() > 0; 
 }
 
 bool UWorld::HasAnyLevelMakingInvisible() const
 {
-	FReadScopeLock ReadLock(PendingVisibilityLock);
+	UE::TReadScopeLock ReadLock(PendingVisibilityLock);
 	return MakingInvisibleLevels.Num() > 0; 
 }
 
 bool UWorld::IsLevelMakingVisible(const ULevel* InLevel) const
 {
-	FReadScopeLock ReadLock(PendingVisibilityLock);
+	UE::TReadScopeLock ReadLock(PendingVisibilityLock);
 	const bool bIsMakingVisible = InLevel && (InLevel->GetPendingVisibilityState() == ELevelPendingVisibilityState::MakingVisible);
 	check(bIsMakingVisible == MakingVisibleLevels.Contains(InLevel));
 	return bIsMakingVisible;
@@ -3275,7 +3276,7 @@ bool UWorld::IsLevelMakingVisible(const ULevel* InLevel) const
 
 bool UWorld::IsLevelMakingInvisible(const ULevel* InLevel) const
 {
-	FReadScopeLock ReadLock(PendingVisibilityLock);
+	UE::TReadScopeLock ReadLock(PendingVisibilityLock);
 	const bool bIsMakingInvisible = InLevel && (InLevel->GetPendingVisibilityState() == ELevelPendingVisibilityState::MakingInvisible);
 	check(bIsMakingInvisible == MakingInvisibleLevels.Contains(InLevel));
 	return bIsMakingInvisible;
@@ -3299,7 +3300,7 @@ void UWorld::EnableActorTickAndUserCallbacks(bool bEnable)
 
 void UWorld::SetLevelPendingVisibilityState(ULevel* InLevel, ELevelPendingVisibilityState InState)
 {
-	FWriteScopeLock WriteLock(PendingVisibilityLock);
+	UE::TWriteScopeLock WriteLock(PendingVisibilityLock);
 
 	if (InLevel->GetPendingVisibilityState() == InState)
 	{
