[Skeletal Mesh Asset] Added support for building half edge buffers automatically when the LOD supports mesh deformer and the asset itself has default mesh deformer assigned

diff --git a/Engine/Source/Runtime/Engine/Classes/Components/SkinnedMeshComponent.h b/Engine/Source/Runtime/Engine/Classes/Components/SkinnedMeshComponent.h
--- a/Engine/Source/Runtime/Engine/Classes/Components/SkinnedMeshComponent.h
+++ b/Engine/Source/Runtime/Engine/Classes/Components/SkinnedMeshComponent.h
@@ -331,6 +331,7 @@ protected:
 public:
 	const FMeshDeformerInstanceSet& GetMeshDeformerInstances() const { return MeshDeformerInstances; }
 
+	ENGINE_API TObjectPtr<UMeshDeformer> GetComponentMeshDeformer() const { return MeshDeformer; }
 private:
 	/** Initializes the MeshDeformerInstances property */
 	void CreateMeshDeformerInstances(const FMeshDeformerSet& DeformerSet);
diff --git a/Engine/Source/Runtime/Engine/Classes/Engine/SkinnedAssetCommon.h b/Engine/Source/Runtime/Engine/Classes/Engine/SkinnedAssetCommon.h
--- a/Engine/Source/Runtime/Engine/Classes/Engine/SkinnedAssetCommon.h
+++ b/Engine/Source/Runtime/Engine/Classes/Engine/SkinnedAssetCommon.h
@@ -224,10 +224,11 @@ struct FSkeletalMeshLODInfo
 	uint8 bAllowCPUAccess : 1;
 	
 	/**
-	 * If true, we will cache/cook half edge data that provides vertex connectivity information across material sections, which
- 	 * may be useful for other systems like Mesh Deformer.
+	 * Whether to force cache/cook half edge data that provides vertex connectivity information across material sections, which
+ 	 * may be useful for other systems like Mesh Deformer. Please note that Half Edge buffers are always built if the Skeletal Mesh Asset
+ 	 * has Default Mesh Deformer assigned, and the LOD has AllowMeshDeformer enabled, 
 	 */
-	UPROPERTY(EditAnywhere, Category = SkeletalMeshLODInfo)
+	UPROPERTY(EditAnywhere, Category = SkeletalMeshLODInfo, meta=(DisplayName="Force Build Half Edge Buffers"))
 	uint8 bBuildHalfEdgeBuffers : 1;
 
 	/** Whether a Mesh Deformer applied to the mesh asset or Skinned Mesh Component should be used on this LOD or not */
diff --git a/Engine/Source/Runtime/Engine/Private/SkeletalMeshRenderData.cpp b/Engine/Source/Runtime/Engine/Private/SkeletalMeshRenderData.cpp
--- a/Engine/Source/Runtime/Engine/Private/SkeletalMeshRenderData.cpp
+++ b/Engine/Source/Runtime/Engine/Private/SkeletalMeshRenderData.cpp
@@ -349,6 +349,13 @@ void FSkeletalMeshRenderData::Cache(const ITargetPlatform* TargetPlatform, USkin
 				FSkeletalMeshLODRenderData::FBuildSettings BuildSettings;
 				BuildSettings.BuildFlags = VertexBufferBuildFlags;
 				BuildSettings.bBuildHalfEdgeBuffers = LODInfo->bBuildHalfEdgeBuffers;
+
+				// The presence of a deformer on the skinned asset almost guarantees dependency on HalfEdgeBuffer,
+				// so let's build half edge buffer for them automatically
+				if (LODInfo->bAllowMeshDeformer && Owner->GetDefaultMeshDeformer())
+				{
+					BuildSettings.bBuildHalfEdgeBuffers = true;
+				}
 				
 				LODData->BuildFromLODModel(LODModel, LODInfo->VertexAttributes, BuildSettings);
 			}
