Added CVars to disable cloth / morph target rendering for debugging purposes.

diff --git a/Engine/Source/Runtime/Engine/Private/SkeletalRenderGPUSkin.cpp b/Engine/Source/Runtime/Engine/Private/SkeletalRenderGPUSkin.cpp
--- a/Engine/Source/Runtime/Engine/Private/SkeletalRenderGPUSkin.cpp
+++ b/Engine/Source/Runtime/Engine/Private/SkeletalRenderGPUSkin.cpp
@@ -40,6 +40,26 @@ static TAutoConsoleVariable<int32> CVarMotionBlurDebug(
 	TEXT(" 1: on"),
 	ECVF_Cheat | ECVF_RenderThreadSafe);
 
+static bool GEnableMorphTargets = true;
+static FAutoConsoleVariableRef CVarMorphTargetEnable(
+	TEXT("r.MorphTarget.Enable"),
+	GEnableMorphTargets,
+	TEXT("Enables morph target rendering.\n")
+	TEXT(" 0: Disable\n")
+	TEXT(" 1: Enable (default)\n"),
+	ECVF_RenderThreadSafe
+	);
+
+static bool GEnableCloth = true;
+static FAutoConsoleVariableRef CVarClothEnable(
+	TEXT("r.Cloth.Enable"),
+	GEnableCloth,
+	TEXT("Enables cloth rendering.\n")
+	TEXT(" 0: Disable\n")
+	TEXT(" 1: Enable (default)\n"),
+	ECVF_RenderThreadSafe
+	);
+
 static int32 GUseGPUMorphTargets = 1;
 static FAutoConsoleVariableRef CVarUseGPUMorphTargets(
 	TEXT("r.MorphTarget.Mode"),
@@ -875,7 +895,7 @@ void FSkeletalMeshObjectGPUSkin::ProcessUpdatedDynamicData(FRHICommandList& RHIC
 		PreviousRevisionNumber = bHasPreviousReferenceToLocal ? DynamicData->PreviousRevisionNumber : INDEX_NONE;
 
 		// RayTracing-specific LOD's can't have a separate morph target buffer at the moment because there is only one morph vertex buffer across the entire mesh.
-		if (LODData.GetNumVertices() > 0 && (bHasWeightedActiveMorphs || bHasExternalMorphs))
+		if (GEnableMorphTargets && LODData.GetNumVertices() > 0 && (bHasWeightedActiveMorphs || bHasExternalMorphs))
 		{
 			bMorphNeedsUpdate |= GForceUpdateMorphTargets != 0;
 			MorphVertexBuffer = &LOD.MorphVertexBufferPool->GetMorphVertexBufferForWriting();
@@ -955,7 +975,7 @@ void FSkeletalMeshObjectGPUSkin::ProcessUpdatedDynamicData(FRHICommandList& RHIC
 		FGPUBaseSkinAPEXClothVertexFactory* ClothVertexFactory = VertexFactory->GetClothVertexFactory();
 		FGPUBaseSkinAPEXClothVertexFactory::ClothShaderType* ClothShaderData = nullptr;
 
-		const bool bSectionUsingCloth = ClothVertexFactory != nullptr;
+		const bool bSectionUsingCloth = GEnableCloth && ClothVertexFactory != nullptr;
 		const bool bSectionUsingMorph = MorphVertexBuffer && !bSectionUsingCloth && (bHasExternalMorphs || (bHasWeightedActiveMorphs && DynamicData->SectionIdsUseByActiveMorphTargets.Contains(SectionIdx)));
 		bool bSectionUsingSkinCache = bAllowAddToSkinCache ? Section.MaxBoneInfluences != 0 : false;
 
@@ -965,7 +985,7 @@ void FSkeletalMeshObjectGPUSkin::ProcessUpdatedDynamicData(FRHICommandList& RHIC
 		const FClothSimulData* ClothSimulationData = nullptr;
 
 		// Update uniform buffer for APEX cloth simulation mesh positions and normals
-		if (ClothVertexFactory)
+		if (bSectionUsingCloth)
 		{
 			ClothShaderData = &ClothVertexFactory->GetClothShaderData();
 			ClothSimulationData = DynamicData->ClothingSimData.Find(Section.CorrespondClothAssetIndex);
