Anim BP: Fix reinstancing crash in the edge case of a nested linked anim graph

This could cause the reinstancing callback to be called on an anim instance that was outered to an unregistered skeletal mesh

diff --git a/Engine/Source/Runtime/Engine/Private/Animation/AnimNode_LinkedAnimGraph.cpp b/Engine/Source/Runtime/Engine/Private/Animation/AnimNode_LinkedAnimGraph.cpp
--- a/Engine/Source/Runtime/Engine/Private/Animation/AnimNode_LinkedAnimGraph.cpp
+++ b/Engine/Source/Runtime/Engine/Private/Animation/AnimNode_LinkedAnimGraph.cpp
@@ -537,7 +537,7 @@ void FAnimNode_LinkedAnimGraph::HandleObjectsReinstanced_Impl(UObject* InSourceO
 		// Ensure we have a valid mesh at this point, as calling into the graph without one can result in crashes
 		// as we assume a valid bone container/reference skeleton is present
 		USkeletalMeshComponent* MeshComponent = SourceAnimInstance->GetSkelMeshComponent();
-		if(MeshComponent && MeshComponent->GetSkeletalMeshAsset())
+		if(MeshComponent && MeshComponent->IsRegistered() && MeshComponent->GetSkeletalMeshAsset())
 		{
 			SourceAnimInstance->RecalcRequiredBones();
 
