Added render in main/depth pass to ISM descriptors.

#rnx

diff --git a/Engine/Source/Runtime/Engine/Private/ISMPartition/ISMComponentDescriptor.cpp b/Engine/Source/Runtime/Engine/Private/ISMPartition/ISMComponentDescriptor.cpp
--- a/Engine/Source/Runtime/Engine/Private/ISMPartition/ISMComponentDescriptor.cpp
+++ b/Engine/Source/Runtime/Engine/Private/ISMPartition/ISMComponentDescriptor.cpp
@@ -124,6 +124,8 @@ void FISMComponentDescriptorBase::InitFrom(const UStaticMeshComponent* Template,
 	bSelectable = Template->bSelectable;
 	bHasPerInstanceHitProxies = Template->bHasPerInstanceHitProxies;
 	bVisibleInRayTracing = Template->bVisibleInRayTracing;
+	bRenderInMainPass = Template->bRenderInMainPass;
+	bRenderInDepthPass = Template->bRenderInDepthPass;	
 	bEvaluateWorldPositionOffset = Template->bEvaluateWorldPositionOffset;
 	WorldPositionOffsetDisableDistance = Template->WorldPositionOffsetDisableDistance;
 	ShadowCacheInvalidationBehavior = Template->ShadowCacheInvalidationBehavior;
@@ -266,6 +268,8 @@ bool FISMComponentDescriptorBase::operator==(const FISMComponentDescriptorBase&
 	bIsEditorOnly == Other.bIsEditorOnly &&
 	bVisible == Other.bVisible &&
 	bVisibleInRayTracing == Other.bVisibleInRayTracing &&
+	bRenderInMainPass == Other.bRenderInMainPass &&
+	bRenderInDepthPass == Other.bRenderInDepthPass &&
 	bSelectable == Other.bSelectable &&
 	bHasPerInstanceHitProxies == Other.bHasPerInstanceHitProxies &&
 	bEvaluateWorldPositionOffset == Other.bEvaluateWorldPositionOffset &&
@@ -396,6 +400,8 @@ void FISMComponentDescriptorBase::InitComponent(UInstancedStaticMeshComponent* I
 	ISMComponent->bRenderCustomDepth = bRenderCustomDepth;
 	ISMComponent->bHiddenInGame = bHiddenInGame;
 	ISMComponent->bIsEditorOnly = bIsEditorOnly;
+	ISMComponent->bRenderInMainPass = bRenderInMainPass;
+	ISMComponent->bRenderInDepthPass = bRenderInDepthPass;
 	ISMComponent->SetVisibleFlag(bVisible);
 	ISMComponent->bVisibleInRayTracing = bVisibleInRayTracing;
 	ISMComponent->bSelectable = bSelectable;
diff --git a/Engine/Source/Runtime/Engine/Public/ISMPartition/ISMComponentDescriptor.h b/Engine/Source/Runtime/Engine/Public/ISMPartition/ISMComponentDescriptor.h
--- a/Engine/Source/Runtime/Engine/Public/ISMPartition/ISMComponentDescriptor.h
+++ b/Engine/Source/Runtime/Engine/Public/ISMPartition/ISMComponentDescriptor.h
@@ -165,6 +165,12 @@ public:
 	UPROPERTY(EditAnywhere, Category = "Component Settings")
 	uint8 bVisibleInRayTracing : 1;
 
+	UPROPERTY(EditAnywhere, Category = "Component Settings")
+	uint8 bRenderInMainPass : 1;
+
+	UPROPERTY(EditAnywhere, Category = "Component Settings")
+	uint8 bRenderInDepthPass : 1;
+
 	UPROPERTY(EditAnywhere, Category = "Component Settings")
 	uint8 bHiddenInGame : 1;
 
