[Input] Remove the old and deprecated IsFloatAxis and IsVectorAxis functions from FKey

diff --git a/Engine/Source/Runtime/Engine/Classes/Kismet/KismetInputLibrary.h b/Engine/Source/Runtime/Engine/Classes/Kismet/KismetInputLibrary.h
--- a/Engine/Source/Runtime/Engine/Classes/Kismet/KismetInputLibrary.h
+++ b/Engine/Source/Runtime/Engine/Classes/Kismet/KismetInputLibrary.h
@@ -87,47 +87,40 @@ class UKismetInputLibrary : public UBlueprintFunctionLibrary
 	UFUNCTION(BlueprintPure, meta = (DisplayName = "Is Keyboard Key"), Category="Input|Key")
 	static ENGINE_API bool Key_IsKeyboardKey(const FKey& Key);
 
-	/**
-	 * Returns true if the key is a vector axis
-	 * @note Deprecated. Use Is Axis 2D/3D instead.
-	 */
-	UFUNCTION(BlueprintPure, meta = (DisplayName = "Is Vector Axis", DeprecatedFunction, DeprecationMessage = "Use Is Axis 2D/3D instead."), Category = "Input|Key")
-	static ENGINE_API bool Key_IsVectorAxis(const FKey& Key);
-
 	/**
 	 * Returns true if the key is a 1D (float) axis
 	 */
-	UFUNCTION(BlueprintPure, meta = (Keywords = "IsFloatAxis", DisplayName = "Is Axis 1D"), Category="Input|Key")
+	UFUNCTION(BlueprintPure, meta = (DisplayName = "Is Axis 1D"), Category="Input|Key")
 	static ENGINE_API bool Key_IsAxis1D(const FKey& Key);
 
 	/**
 	 * Returns true if the key is a 2D (vector) axis
 	 */
-	UFUNCTION(BlueprintPure, meta = (Keywords = "IsVectorAxis", DisplayName = "Is Axis 2D"), Category="Input|Key")
+	UFUNCTION(BlueprintPure, meta = (DisplayName = "Is Axis 2D"), Category="Input|Key")
 	static ENGINE_API bool Key_IsAxis2D(const FKey& Key);
 
 	/**
 	 * Returns true if the key is a 3D (vector) axis
 	 */
-	UFUNCTION(BlueprintPure, meta = (Keywords = "IsVectorAxis", DisplayName = "Is Axis 3D"), Category="Input|Key")
+	UFUNCTION(BlueprintPure, meta = (DisplayName = "Is Axis 3D"), Category="Input|Key")
 	static ENGINE_API bool Key_IsAxis3D(const FKey& Key);
 
 	/**
 	 * Returns true if the key is a 1D axis emulating a digital button press.
 	 */
-	UFUNCTION(BlueprintPure, meta = (Keywords = "IsFloatAxis", DisplayName = "Is Button Axis"), Category = "Input|Key")
+	UFUNCTION(BlueprintPure, meta = (DisplayName = "Is Button Axis"), Category = "Input|Key")
 	static ENGINE_API bool Key_IsButtonAxis(const FKey& Key);
 
 	/**
 	 * Returns true if the key is an analog axis
 	 */
-	UFUNCTION(BlueprintPure, meta = (Keywords = "IsFloatAxis, IsVectorAxis", DisplayName = "Is Analog"), Category = "Input|Key")
+	UFUNCTION(BlueprintPure, meta = (DisplayName = "Is Analog"), Category = "Input|Key")
 	static ENGINE_API bool Key_IsAnalog(const FKey& Key);
 
 	/**
 	 * Returns true if the key is a digital button press
 	 */
-	UFUNCTION(BlueprintPure, meta = (Keywords = "IsFloatAxis, IsVectorAxis", DisplayName = "Is Digital"), Category = "Input|Key")
+	UFUNCTION(BlueprintPure, meta = (DisplayName = "Is Digital"), Category = "Input|Key")
 	static ENGINE_API bool Key_IsDigital(const FKey& Key);
 
 	/**
diff --git a/Engine/Source/Runtime/Engine/Private/KismetInputLibrary.cpp b/Engine/Source/Runtime/Engine/Private/KismetInputLibrary.cpp
--- a/Engine/Source/Runtime/Engine/Private/KismetInputLibrary.cpp
+++ b/Engine/Source/Runtime/Engine/Private/KismetInputLibrary.cpp
@@ -64,11 +64,6 @@ bool UKismetInputLibrary::Key_IsKeyboardKey(const FKey& Key)
 	return Key.IsBindableInBlueprints() && (Key.IsGamepadKey() == false && Key.IsMouseButton() == false);
 }
 
-bool UKismetInputLibrary::Key_IsVectorAxis(const FKey& Key)
-{
-	return Key.IsAxis2D() || Key.IsAxis3D();
-}
-
 bool UKismetInputLibrary::Key_IsAxis1D(const FKey& Key)
 {
 	return Key.IsAxis1D();
diff --git a/Engine/Source/Runtime/InputCore/Classes/InputCoreTypes.h b/Engine/Source/Runtime/InputCore/Classes/InputCoreTypes.h
--- a/Engine/Source/Runtime/InputCore/Classes/InputCoreTypes.h
+++ b/Engine/Source/Runtime/InputCore/Classes/InputCoreTypes.h
@@ -78,10 +78,6 @@ struct FKey
 	INPUTCORE_API bool IsAxis1D() const;
 	INPUTCORE_API bool IsAxis2D() const;
 	INPUTCORE_API bool IsAxis3D() const;
-	UE_DEPRECATED(4.26, "Please use IsAxis1D instead.")
-	INPUTCORE_API bool IsFloatAxis() const;
-	UE_DEPRECATED(4.26, "Please use IsAxis2D/IsAxis3D instead.")
-	INPUTCORE_API bool IsVectorAxis() const;
 	INPUTCORE_API bool IsDigital() const;
 	INPUTCORE_API bool IsAnalog() const;
 	INPUTCORE_API bool IsBindableInBlueprints() const;
@@ -160,11 +156,6 @@ struct FKeyDetails
 		// Gesture input types such as Flick, Pinch, and Rotate
 		Gesture					 = 1 << 12,
 
-		// Deprecated. Replace with axis definitions for clarity.
-
-		FloatAxis  UE_DEPRECATED(4.26, "Please use Axis1D instead.") = Axis1D,
-		VectorAxis UE_DEPRECATED(4.26, "Please use Axis2D/Axis3D instead.") = Axis3D,
-
 		NoFlags                  = 0,
 	};
 
@@ -179,10 +170,6 @@ struct FKeyDetails
 	FORCEINLINE bool IsAxis2D() const { return AxisType == EInputAxisType::Axis2D; }
 	FORCEINLINE bool IsAxis3D() const { return AxisType == EInputAxisType::Axis3D; }
 	FORCEINLINE bool IsButtonAxis() const { return AxisType == EInputAxisType::Button; }	// Analog 1D axis emulating a digital button press.
-	UE_DEPRECATED(4.26, "Please use IsAxis1D instead.")
-	FORCEINLINE bool IsFloatAxis() const { return IsAxis1D(); }
-	UE_DEPRECATED(4.26, "Please use IsAxis2D/IsAxis3D instead.")
-	FORCEINLINE bool IsVectorAxis() const { return IsAxis2D() || IsAxis3D(); }
 	FORCEINLINE bool IsAnalog() const { return IsAxis1D() || IsAxis2D() || IsAxis3D(); }
 	FORCEINLINE bool IsDigital() const { return !IsAnalog(); }
 	FORCEINLINE bool IsBindableInBlueprints() const { return bIsBindableInBlueprints != 0; }
diff --git a/Engine/Source/Runtime/InputCore/Private/InputCoreTypes.cpp b/Engine/Source/Runtime/InputCore/Private/InputCoreTypes.cpp
--- a/Engine/Source/Runtime/InputCore/Private/InputCoreTypes.cpp
+++ b/Engine/Source/Runtime/InputCore/Private/InputCoreTypes.cpp
@@ -1367,17 +1367,6 @@ bool FKey::IsAxis3D() const
 	return (KeyDetails.IsValid() ? KeyDetails->IsAxis3D() : false);
 }
 
-bool FKey::IsFloatAxis() const
-{
-	return IsAxis1D();
-}
-
-bool FKey::IsVectorAxis() const
-{
-	ConditionalLookupKeyDetails();
-	return (KeyDetails.IsValid() ? KeyDetails->IsAnalog() && !KeyDetails->IsAxis1D() : false);
-}
-
 bool FKey::IsDigital() const
 {
 	ConditionalLookupKeyDetails();
