Added support for printing a Class object's Default Object.

#rnx

diff --git a/Engine/Source/Runtime/Engine/Private/PrintObjectUtils.cpp b/Engine/Source/Runtime/Engine/Private/PrintObjectUtils.cpp
--- a/Engine/Source/Runtime/Engine/Private/PrintObjectUtils.cpp
+++ b/Engine/Source/Runtime/Engine/Private/PrintObjectUtils.cpp
@@ -371,6 +371,36 @@ void PrintObjectIDO(UObject* Object, EPrintObjectFlag Flags, FOutputDevice* InOu
 	}
 }
 
+void PrintClassDefaultObject(UClass* Class, EPrintObjectFlag Flags, FOutputDevice* InOutputDevice)
+{
+	FOutputDevice* OutputDevice = InOutputDevice ? InOutputDevice : GLog;
+	if (OutputDevice)
+	{
+		if (Class)
+		{
+			FString ClassPath = Class->GetPathName();
+
+			UObject* CDO = Class->GetDefaultObject(/*bCreateIfNeeded*/false);
+
+			bool bVerbose = EnumHasAnyFlags(Flags, EPrintObjectFlag::Verbose);
+
+			if (CDO)
+			{
+				LogObjectState(Class, OutputDevice, bVerbose, TEXT("Printing Class Default Object for class: "));
+				LogObjectState(CDO, OutputDevice, bVerbose);
+			}
+			else
+			{
+				LogObjectState(Class, OutputDevice, bVerbose, TEXT("PrintClassDefaultObject: No Class Default Object found for class: "));
+			}
+		}
+		else
+		{
+			OutputDevice->Logf(TEXT("PrintClassDefaultObject: NULL class"));
+		}
+	}
+}
+
 static void ParseObjectIdAndFlags(const TArray<FString>& Args, FString& OutObjectIdentifier, EPrintObjectFlag& OutFlags)
 {
 	OutFlags = EPrintObjectFlag::None;
@@ -541,6 +571,35 @@ FAutoConsoleCommandWithArgsAndOutputDevice CVarCommandPrintObjectIDO(
 		}
 	}));
 
+FAutoConsoleCommandWithArgsAndOutputDevice CVarCommandPrintCDO(
+	TEXT("Obj.DumpCDO"),
+	TEXT("Outputs a class' Class Default Object."),
+	FConsoleCommandWithArgsAndOutputDeviceDelegate::CreateLambda([](const TArray<FString>& Args, FOutputDevice& OutputDevice)
+	{
+		if (Args.Num() > 0)
+		{
+			FString ClassPath;
+			EPrintObjectFlag Flags;
+			ParseObjectIdAndFlags(Args, ClassPath, Flags);
+
+			UClass* Class = FindFirstObjectSafe<UClass>(*ClassPath);
+
+			if (Class)
+			{
+				PrintClassDefaultObject(Class, Flags, &OutputDevice);
+			}
+			else
+			{
+				OutputDevice.Logf(TEXT("PrintClassDefaultObject: failed to find any classes for path: '%s'"), *ClassPath);
+			}
+		}
+		else
+		{
+			OutputDevice.Logf(TEXT("PrintClassDefaultObject: no class name or path specified (example usage: Obj.DumpCDO MyActor)"));
+			LogHelp(OutputDevice);
+		}
+	}));
+	
 } // namespace UE
 
 #endif // !UE_BUILD_SHIPPING
diff --git a/Engine/Source/Runtime/Engine/Public/PrintObjectUtils.h b/Engine/Source/Runtime/Engine/Public/PrintObjectUtils.h
--- a/Engine/Source/Runtime/Engine/Public/PrintObjectUtils.h
+++ b/Engine/Source/Runtime/Engine/Public/PrintObjectUtils.h
@@ -9,6 +9,7 @@
 
 class UObject;
 class UStruct;
+class UClass;
 class FOutputDevice;
 
 /**
@@ -90,6 +91,15 @@ UE_INTERNAL void PrintObjectArchetype(UObject* Object, EPrintObjectFlag Flags =
  */
 UE_INTERNAL void PrintObjectIDO(UObject* Object, EPrintObjectFlag Flags = EPrintObjectFlag::None, FOutputDevice* OutputDevice = nullptr);
 
+/**
+ * Outputs a class' Class Default Object.
+ *
+ * @param Class The class to use.
+ * @param Flags For specifying optional behaviors.
+ * @param OutputDevice The output device to print to (can be NULL).
+ */
+UE_INTERNAL void PrintClassDefaultObject(UClass* Class, EPrintObjectFlag Flags = EPrintObjectFlag::None, FOutputDevice* OutputDevice = nullptr);
+
 } // namespace UE
 
 #endif // !UE_BUILD_SHIPPING
