Asset Manager
- Don't allow users to pass in cache pin. This can lead to hard to reason about behavior because asset manager internally tracks everything per-primary asset and shared pins can lead to unexpected results. Users can still gather the pins it creates for a primary asset from the returned streamable handle. If we find some need for user speficified pins in the future, we will need to come up with some method that allows the user to map pins to specific primary assets.

diff --git a/Engine/Source/Runtime/Engine/Classes/Engine/AssetManager.h b/Engine/Source/Runtime/Engine/Classes/Engine/AssetManager.h
--- a/Engine/Source/Runtime/Engine/Classes/Engine/AssetManager.h
+++ b/Engine/Source/Runtime/Engine/Classes/Engine/AssetManager.h
@@ -32,11 +32,13 @@ DECLARE_DELEGATE_OneParam(FAssetManagerAcquireResourceDelegate, bool /* bSuccess
 /** Delegate called when new asset search root is registered due to runtime asset mounting, path will not have a trailing slash */
 DECLARE_MULTICAST_DELEGATE_OneParam(FOnAddedAssetSearchRoot, const FString&);
 
-/** EXPERIMENTAL -  Additional parameters for downloading packages */
+/** EXPERIMENTAL -  Additional parameters for downloading packages 
+ * Asset Manager will automatically create a FStreamableDownloadCachePin for each primary asset requested with download params.
+ * Cache pins can be retrieved with FStreamableHandle::GetDownloadCachePins()
+ * Cache pins will be released along with the primary asset or when bundle state changes.
+ */
 struct FAssetManagerDownloadParams
 {
-	/** Handle to keep packages in the local cache. If no cache pin is provided, a new one will be created. */
-	FStreamableDownloadCachePin CachePin;
 	/** Download priority */
 	EStreamableDownloadPriority Priority = EStreamableDownloadPriority::Normal;
 	/** If set, any soft references will also be installed */
diff --git a/Engine/Source/Runtime/Engine/Private/AssetManager.cpp b/Engine/Source/Runtime/Engine/Private/AssetManager.cpp
--- a/Engine/Source/Runtime/Engine/Private/AssetManager.cpp
+++ b/Engine/Source/Runtime/Engine/Private/AssetManager.cpp
@@ -424,19 +424,6 @@ namespace UE::AssetManager::Private
 
 		return PathsToLoad;
 	}
-
-	void GetDownloadCachePinDebugName(const TArray<FPrimaryAssetId>& AssetsToLoad, FStringBuilderBase& Builder)
-	{
-		for (const FPrimaryAssetId& PrimaryAssetId : AssetsToLoad)
-		{
-			if (Builder.Len() >= MaxDebugWarningLen)
-			{
-				break;
-			}
-
-			Builder << (Builder.Len() > 0 ? TEXT(", ") : TEXT("")) << PrimaryAssetId.ToString();
-		}
-	}
 }
 
 UAssetManager::UAssetManager()
@@ -1820,13 +1807,6 @@ TSharedPtr<FStreamableHandle> UAssetManager::ChangeBundleStateForPrimaryAssets(
 	bool bRemoveAllBundles,
 	FAssetManagerLoadParams&& LoadParams)
 {
-	if (LoadParams.DownloadParams && !LoadParams.DownloadParams->CachePin)
-	{
-		TStringBuilder<256> DebugName;
-		UE::AssetManager::Private::GetDownloadCachePinDebugName(AssetsToChange, DebugName);
-		LoadParams.DownloadParams->CachePin = FStreamableDownloadCachePin::Create(DebugName);
-	}
-
 	TArray<TSharedPtr<FStreamableHandle> > NewHandles, ExistingHandles;
 	TArray<FPrimaryAssetId> NewAssets;
 	int32 MoviePlayerNumAssets = 0;
@@ -1847,10 +1827,10 @@ TSharedPtr<FStreamableHandle> UAssetManager::ChangeBundleStateForPrimaryAssets(
 		{
 			FPlatformMisc::PumpEssentialAppMessages();
 
-			// If a request has ever been made with a cache pin, require one for any subsequent requests.
+			// If a request has ever been made with a download params, require a download for any subsequent requests.
 			// This can help prevent user error if downloads aren't consistently requested and prevent losing cache pins
 			// when a handle is destroyed.
-			const bool bDownloadCachePinRequired = 
+			const bool bDownloadRequired = 
 				(NameData->PendingState.Handle && NameData->PendingState.Handle->GetDownloadCachePin().IsValid()) ||
 				(NameData->CurrentState.Handle && NameData->CurrentState.Handle->GetDownloadCachePin().IsValid());
 
@@ -1960,14 +1940,14 @@ TSharedPtr<FStreamableHandle> UAssetManager::ChangeBundleStateForPrimaryAssets(
 				LoadParamsWithoutCallbacks.DownloadParams = CopyTemp(*LoadParams.DownloadParams);
 				NewHandle = LoadAssetList(MoveTemp(PathsToLoad), MoveTemp(LoadParamsWithoutCallbacks), FString(DebugName));
 			}
-			else if (!bDownloadCachePinRequired)
+			else if (!bDownloadRequired)
 			{
 				NewHandle = LoadAssetList(MoveTemp(PathsToLoad), MoveTemp(LoadParamsWithoutCallbacks), FString(DebugName));
 			}
 			else
 			{
 				// The user previously use download params for requests with this primary asset, so we expect them going forward.
-				// See comment above where bDownloadCachePinRequired is delcared.
+				// See comment above where bDownloadRequired is delcared.
 				FStreamableAsyncLoadParams StreamableLoadParams;
 				StreamableLoadParams.TargetsToStream = MoveTemp(PathsToLoad);
 				NewHandle = StreamableManager.CreateErrorHandle(
@@ -2386,7 +2366,6 @@ namespace UE::AssetManager::Private
 		if (LoadParams.DownloadParams)
 		{
 			Params.DownloadParams.Emplace();
-			Params.DownloadParams->CachePin = MoveTemp(LoadParams.DownloadParams->CachePin);
 			Params.DownloadParams->Priority = LoadParams.DownloadParams->Priority;
 			Params.DownloadParams->bInstallSoftReferences = LoadParams.DownloadParams->bInstallSoftReferences;
 		}
