[Waveform Editor] - An Exported Sound Wave through Waveform Editor does not retain its channels

diff --git a/Engine/Source/Runtime/Engine/Private/SampleBufferIO.cpp b/Engine/Source/Runtime/Engine/Private/SampleBufferIO.cpp
--- a/Engine/Source/Runtime/Engine/Private/SampleBufferIO.cpp
+++ b/Engine/Source/Runtime/Engine/Private/SampleBufferIO.cpp
@@ -210,6 +210,10 @@ namespace Audio
 		CurrentOperation.Reset(new FAsyncSoundWavePCMWriterTask(this, ESoundWavePCMWriteTaskType::GenerateAndWriteSoundWave, OnSuccess));
 		CurrentOperation->StartBackgroundTask();
 
+		// Invalidate compressed data on new soundwave so that it gets refreshed or the soundwave is not seekable
+		// and causes this error: FMixerBuffer::CreateStreamingBuffer failed to StreamCompressedInfo on SoundWave
+		CurrentSoundWave->InvalidateCompressedData(true);
+
 		return true;
 	}
 
