Scene Capture:  Default post process blend weight to 1.0 for Cube captures, to match 2D captures.  A lot of confusion can result due to all post process features not working until that's set, when toggling to one of the Final Color modes that run post processing.

#rnx

diff --git a/Engine/Source/Runtime/Engine/Private/Components/SceneCaptureComponent.cpp b/Engine/Source/Runtime/Engine/Private/Components/SceneCaptureComponent.cpp
--- a/Engine/Source/Runtime/Engine/Private/Components/SceneCaptureComponent.cpp
+++ b/Engine/Source/Runtime/Engine/Private/Components/SceneCaptureComponent.cpp
@@ -1309,6 +1309,9 @@ USceneCaptureComponentCube::USceneCaptureComponentCube(const FObjectInitializer&
 	bTickInEditor = true;
 	bCaptureRotation = false;
 
+	// default to full blend weight..
+	PostProcessBlendWeight = 1.0f;
+
 	// Update default post-processing values to reflect default values (previously and still) applied in code.
 	PostProcessSettings.DynamicGlobalIlluminationMethod = EDynamicGlobalIlluminationMethod::None;
 	PostProcessSettings.ReflectionMethod = EReflectionMethod::None;
