With UE_WITH_REMOTE_OBJECT_HANDLE and a transactional NetDriver, move the transaction scope to a higher level to include deserialization of incoming RPC parameters.

diff --git a/Engine/Source/Runtime/Engine/Private/DataReplication.cpp b/Engine/Source/Runtime/Engine/Private/DataReplication.cpp
--- a/Engine/Source/Runtime/Engine/Private/DataReplication.cpp
+++ b/Engine/Source/Runtime/Engine/Private/DataReplication.cpp
@@ -1163,7 +1163,23 @@ bool FObjectReplicator::ReceivedBunch(FNetBitReader& Bunch, const FReplicationFl
 		{
 			bool bDelayFunction = false;
 			TSet<FNetworkGUID> UnmappedGuids;
-			bool bSuccess = ReceivedRPC(Reader, RepFlags, FieldCache, bCanDelayRPCs, bDelayFunction, UnmappedGuids);
+
+			bool bSuccess = false;
+
+#if UE_WITH_REMOTE_OBJECT_HANDLE
+			if (Connection->Driver->bReplicateTransactionally)
+			{
+				static FName TransactionalWorkName(TEXT("FObjectReplicator ReceivedRPC"));
+				UE::RemoteExecutor::ExecuteTransactional(TransactionalWorkName, [this, &bSuccess, &Reader, RepFlags, FieldCache, bCanDelayRPCs, &bDelayFunction, &UnmappedGuids]()
+					{
+						bSuccess = ReceivedRPC(Reader, RepFlags, FieldCache, bCanDelayRPCs, bDelayFunction, UnmappedGuids);
+					});
+			}
+			else
+#endif
+			{
+				bSuccess = ReceivedRPC(Reader, RepFlags, FieldCache, bCanDelayRPCs, bDelayFunction, UnmappedGuids);
+			}
 
 			if (!bSuccess)
 			{
@@ -1387,20 +1403,7 @@ bool FObjectReplicator::ReceivedRPC(FNetBitReader& Reader, const FReplicationFla
 				// Call the function.
 				if (Connection->Driver->IsExecuteRPCFunctionsEnabled())
 				{
-#if UE_WITH_REMOTE_OBJECT_HANDLE
-					if (Connection->Driver->bReplicateTransactionally)
-					{
-						static FName TransactionalWorkName(TEXT("FObjectReplicator ProcessEvent"));
-						UE::RemoteExecutor::ExecuteTransactional(TransactionalWorkName, [Object, Function, Parms]()
-							{
-								CallProcessEventForReceivedRPC(Object, Function, Parms);
-							});
-					}
-					else
-#endif
-					{
-						CallProcessEventForReceivedRPC(Object, Function, Parms);
-					}
+					CallProcessEventForReceivedRPC(Object, Function, Parms);
 				}
 			}
 		}
