Added additional checks before replicating subobjects to avoid an edge case where an object replicator could be improperly initialized.

#rnx

diff --git a/Engine/Source/Runtime/Engine/Private/DataChannel.cpp b/Engine/Source/Runtime/Engine/Private/DataChannel.cpp
--- a/Engine/Source/Runtime/Engine/Private/DataChannel.cpp
+++ b/Engine/Source/Runtime/Engine/Private/DataChannel.cpp
@@ -5536,7 +5536,9 @@ void UActorChannel::AddedToChannelPool()
 
 bool UActorChannel::WriteSubObjectInBunch(UObject* Obj, FOutBunch& Bunch, FReplicationFlags RepFlags)
 {
-	if (!IsValid(Obj))
+	// This matches the check in FObjectReplicator::InitWithObject if we do not early out here then the connection will be improperly initialized.
+	TWeakObjectPtr<UObject> WeakObjectPtr = Obj;
+	if (!IsValid(Obj) || !WeakObjectPtr.IsValid())
 	{
 		return false;
 	}
