Skeletal Mesh: Ensure source models and LOD infos are synced when loading cooked assets.

#rnx

diff --git a/Engine/Source/Runtime/Engine/Private/SkeletalMesh.cpp b/Engine/Source/Runtime/Engine/Private/SkeletalMesh.cpp
--- a/Engine/Source/Runtime/Engine/Private/SkeletalMesh.cpp
+++ b/Engine/Source/Runtime/Engine/Private/SkeletalMesh.cpp
@@ -3776,6 +3776,12 @@ void USkeletalMesh::ExecutePostLoadInternal(FSkinnedAssetPostLoadContext& Contex
 			Context.bHasCachedDerivedData = true;
 		}
 	}
+	else
+	{
+		// Ensure we have source model storage that matches the number of LODs defined on this mesh since we're not doing any translation from
+		// old models. 
+		SetNumSourceModels(GetLODNum());
+	}
 
 	// check the MinLOD values are all within range
 	FPerQualityLevelInt QualityLocalMinLOD;
