Allow queued subtitles to have their timing controlled by modules other than the subtitles system.

diff --git a/Engine/Source/Runtime/Engine/Private/SoundWave.cpp b/Engine/Source/Runtime/Engine/Private/SoundWave.cpp
--- a/Engine/Source/Runtime/Engine/Private/SoundWave.cpp
+++ b/Engine/Source/Runtime/Engine/Private/SoundWave.cpp
@@ -3388,7 +3388,7 @@ void USoundWave::Parse(FAudioDevice* AudioDevice, const UPTRINT NodeWaveInstance
 			const ISubtitlesAndClosedCaptionsModule& SubtitlesModule = ModularFeatures.GetModularFeature<ISubtitlesAndClosedCaptionsModule>(ISubtitlesAndClosedCaptionsModule::GetModularFeatureName());
 			if (const UAssetUserData * Subtitle{ GetAssetUserDataOfClass(SubtitlesModule.GetAssetUserDataClass()) })
 			{
-				FSubtitlesAndClosedCaptionsDelegates::QueueSubtitle.ExecuteIfBound(FQueueSubtitleParameters{*Subtitle});
+				FSubtitlesAndClosedCaptionsDelegates::QueueSubtitle.ExecuteIfBound(FQueueSubtitleParameters{*Subtitle}, ESubtitleTiming::InternallyTimed);
 			}
 		}
 	}
@@ -3510,7 +3510,10 @@ void USoundWave::Parse(FAudioDevice* AudioDevice, const UPTRINT NodeWaveInstance
 			if (SoundClassProperties->bApplyEffects)
 			{
 				const UAudioSettings* Settings = GetDefault<UAudioSettings>();
-				WaveInstance->SoundSubmix = Cast<USoundSubmix>(FSoftObjectPtr(Settings->EQSubmix).Get());
+				if (LIKELY(IsValid(Settings)))
+				{
+					WaveInstance->SoundSubmix = Cast<USoundSubmix>(FSoftObjectPtr(Settings->EQSubmix).Get());
+				}
 			}
 			else if (SoundClassProperties->DefaultSubmix)
 			{
diff --git a/Engine/Source/Runtime/Engine/Private/Subtitles/SubtitlesAndClosedCaptionsDelegates.cpp b/Engine/Source/Runtime/Engine/Private/Subtitles/SubtitlesAndClosedCaptionsDelegates.cpp
--- a/Engine/Source/Runtime/Engine/Private/Subtitles/SubtitlesAndClosedCaptionsDelegates.cpp
+++ b/Engine/Source/Runtime/Engine/Private/Subtitles/SubtitlesAndClosedCaptionsDelegates.cpp
@@ -2,7 +2,7 @@
 
 #include "Subtitles/SubtitlesAndClosedCaptionsDelegates.h"
 
-TDelegate<void(const FQueueSubtitleParameters&)> FSubtitlesAndClosedCaptionsDelegates::QueueSubtitle;
+TDelegate<void(const FQueueSubtitleParameters&, const ESubtitleTiming)> FSubtitlesAndClosedCaptionsDelegates::QueueSubtitle;
 TDelegate<bool(const UAssetUserData&)> FSubtitlesAndClosedCaptionsDelegates::IsSubtitleActive;
 TDelegate<void(const UAssetUserData&)> FSubtitlesAndClosedCaptionsDelegates::StopSubtitle;
 TDelegate<void()> FSubtitlesAndClosedCaptionsDelegates::StopAllSubtitles;
diff --git a/Engine/Source/Runtime/Engine/Public/Subtitles/SubtitlesAndClosedCaptionsDelegates.h b/Engine/Source/Runtime/Engine/Public/Subtitles/SubtitlesAndClosedCaptionsDelegates.h
--- a/Engine/Source/Runtime/Engine/Public/Subtitles/SubtitlesAndClosedCaptionsDelegates.h
+++ b/Engine/Source/Runtime/Engine/Public/Subtitles/SubtitlesAndClosedCaptionsDelegates.h
@@ -3,9 +3,19 @@
 #pragma once
 #include "Delegates/Delegate.h"
 
+#include "SubtitlesAndClosedCaptionsDelegates.generated.h"
+
 class FCanvas;
 class UAssetUserData;
 
+UENUM()
+enum class ESubtitleTiming : uint8
+{
+	InternallyTimed,
+	ExternallyTimed
+};
+
+
 struct ENGINE_API FQueueSubtitleParameters
 {
 	const UAssetUserData& Subtitle;
@@ -17,7 +27,7 @@ class ENGINE_API FSubtitlesAndClosedCaptionsDelegates
 public:
 
 	// Have the subtitle subsystem to queue a subtitle to be displayed
-	static TDelegate<void(const FQueueSubtitleParameters&)> QueueSubtitle;
+	static TDelegate<void(const FQueueSubtitleParameters&, const ESubtitleTiming)> QueueSubtitle;
 
 	static TDelegate<bool(const UAssetUserData&)> IsSubtitleActive;
 
