Exposed GetStaticMeshComponentBoostPSOPrecachePriority() in StaticMeshSceneProxy.h
- Missing header change from 40390609

diff --git a/Engine/Source/Runtime/Engine/Public/StaticMeshSceneProxy.h b/Engine/Source/Runtime/Engine/Public/StaticMeshSceneProxy.h
--- a/Engine/Source/Runtime/Engine/Public/StaticMeshSceneProxy.h
+++ b/Engine/Source/Runtime/Engine/Public/StaticMeshSceneProxy.h
@@ -20,6 +20,9 @@ struct FStaticMeshVertexFactories;
 using FStaticMeshVertexFactoriesArray = TArray<FStaticMeshVertexFactories>;
 struct FStaticMeshSceneProxyDesc;
 class FTextureResource;
+enum class EPSOPrecachePriority : uint8;
+
+ENGINE_API EPSOPrecachePriority GetStaticMeshComponentBoostPSOPrecachePriority();
 
 /**
  * A static mesh component scene proxy.
