Add Preset Convert Nodes to Context Menu, and add all Convert nodes to their own category

- Added a method for Material Expressions to register additional menu actions (in addition to just their default creation action)
- Added a 'PostCreationDelegate' that allows FMaterialGraphSchemaAction_NewNode Actions to perform custom setup immediately after the new material expression is created
- Allow renaming ConvertNode, which the preset nodes use to replace 'Convert' with 'Make Float4' for example

diff --git a/Engine/Source/Editor/MaterialEditor/Private/MaterialEditorUtilities.cpp b/Engine/Source/Editor/MaterialEditor/Private/MaterialEditorUtilities.cpp
--- a/Engine/Source/Editor/MaterialEditor/Private/MaterialEditorUtilities.cpp
+++ b/Engine/Source/Editor/MaterialEditor/Private/MaterialEditorUtilities.cpp
@@ -185,21 +185,6 @@ int32 FMaterialEditorUtilities::GetNumberOfSelectedNodes(const class UEdGraph* G
 
 void FMaterialEditorUtilities::GetMaterialExpressionActions(FGraphActionMenuBuilder& ActionMenuBuilder, bool bMaterialFunction)
 {
-	// TODO: Not sure if this is necessary/usable anymore
-	// Get all menu extenders for this context menu from the material editor module
-	/*IMaterialEditorModule& MaterialEditor = FModuleManager::GetModuleChecked<IMaterialEditorModule>( TEXT("MaterialEditor") );
-	TArray<IMaterialEditorModule::FMaterialMenuExtender> MenuExtenderDelegates = MaterialEditor.GetAllMaterialCanvasMenuExtenders();
-
-	TArray<TSharedPtr<FExtender>> Extenders;
-	for (int32 i = 0; i < MenuExtenderDelegates.Num(); ++i)
-	{
-		if (MenuExtenderDelegates[i].IsBound())
-		{
-			Extenders.Add(MenuExtenderDelegates[i].Execute(MaterialEditorPtr.Pin()->GetToolkitCommands()));
-		}
-	}
-	TSharedPtr<FExtender> MenuExtender = FExtender::Combine(Extenders);*/
-
 	bool bUseUnsortedMenus = false;
 	MaterialExpressionClasses* ExpressionClasses = MaterialExpressionClasses::Get();
 
@@ -713,12 +698,20 @@ void FMaterialEditorUtilities::AddMaterialExpressionCategory(FGraphActionMenuBui
 					CreationName = MaterialExpression.CreationName;
 				}
 
+				UMaterialExpression* MaterialExpressionDefaultObject = CastChecked<UMaterialExpression>(MaterialExpression.MaterialClass->GetDefaultObject());
+
 				TSharedPtr<FMaterialGraphSchemaAction_NewNode> NewNodeAction(new FMaterialGraphSchemaAction_NewNode(
 					CategoryName,
 					CreationName,
-					ToolTip, 0, CastChecked<UMaterialExpression>(MaterialExpression.MaterialClass->GetDefaultObject())->GetKeywords()));
+					ToolTip,
+					0, 
+					MaterialExpressionDefaultObject->GetKeywords()
+				));
 				NewNodeAction->MaterialExpressionClass = MaterialExpression.MaterialClass;
 				ActionMenuBuilder.AddAction(NewNodeAction);
+
+				// Register Any Additional Menu Actions for this Expression
+				MaterialExpressionDefaultObject->RegisterAdditionalMenuActions(ActionMenuBuilder, CategoryName);
 			}
 		}
 	}
diff --git a/Engine/Source/Editor/UnrealEd/Classes/MaterialGraph/MaterialGraphSchema.h b/Engine/Source/Editor/UnrealEd/Classes/MaterialGraph/MaterialGraphSchema.h
--- a/Engine/Source/Editor/UnrealEd/Classes/MaterialGraph/MaterialGraphSchema.h
+++ b/Engine/Source/Editor/UnrealEd/Classes/MaterialGraph/MaterialGraphSchema.h
@@ -3,11 +3,13 @@
 #pragma once
 
 #include "CoreMinimal.h"
+#include "Delegates/Delegate.h"
 #include "UObject/ObjectMacros.h"
 #include "EdGraph/EdGraphSchema.h"
 #include "MaterialGraphSchema.generated.h"
 
 class UEdGraph;
+class UMaterialExpression;
 struct FToolMenuContext;
 
 /** Action to add an expression node to the graph */
@@ -20,6 +22,9 @@ struct FMaterialGraphSchemaAction_NewNode : public FEdGraphSchemaAction
 	UPROPERTY()
 	TObjectPtr<class UClass> MaterialExpressionClass;
 
+	/** Delegate to be run immediately after expression of class 'MaterialExpressionClass' is created */
+	TDelegate<void(UMaterialExpression*)> PostCreationDelegate;
+
 	// Simple type info
 	static FName StaticGetTypeId() {static FName Type("FMaterialGraphSchemaAction_NewNode"); return Type;}
 	virtual FName GetTypeId() const override { return StaticGetTypeId(); } 
diff --git a/Engine/Source/Editor/UnrealEd/Private/MaterialGraphSchema.cpp b/Engine/Source/Editor/UnrealEd/Private/MaterialGraphSchema.cpp
--- a/Engine/Source/Editor/UnrealEd/Private/MaterialGraphSchema.cpp
+++ b/Engine/Source/Editor/UnrealEd/Private/MaterialGraphSchema.cpp
@@ -85,6 +85,7 @@ UEdGraphNode* FMaterialGraphSchemaAction_NewNode::PerformAction(class UEdGraph*
 	}
 
 	UMaterialExpression* NewExpression = FMaterialEditorUtilities::CreateNewMaterialExpression(ParentGraph, MaterialExpressionClass, Location, bSelectNewNode, /*bAutoAssignResource*/true);
+	PostCreationDelegate.ExecuteIfBound(NewExpression);
 
 	if (NewExpression)
 	{
diff --git a/Engine/Source/Runtime/Engine/Private/Materials/MaterialExpressions.cpp b/Engine/Source/Runtime/Engine/Private/Materials/MaterialExpressions.cpp
--- a/Engine/Source/Runtime/Engine/Private/Materials/MaterialExpressions.cpp
+++ b/Engine/Source/Runtime/Engine/Private/Materials/MaterialExpressions.cpp
@@ -23237,6 +23237,7 @@ UMaterialExpressionConvert::UMaterialExpressionConvert(const FObjectInitializer&
 {
 #if WITH_EDITORONLY_DATA
 	bHidePreviewWindow = true;
+	MenuCategories.Add(LOCTEXT("Convert", "Convert"));
 #endif
 }
 
@@ -23246,7 +23247,8 @@ void UMaterialExpressionConvert::PostEditChangeProperty(FPropertyChangedEvent& P
 	const bool bNeedsNodeUpdate = 
 			PropertyChangedEvent.GetPropertyName() == GET_MEMBER_NAME_CHECKED(ThisClass, ConvertInputs)
 		|| 	PropertyChangedEvent.GetPropertyName() == GET_MEMBER_NAME_CHECKED(ThisClass, ConvertMappings)
-		|| 	PropertyChangedEvent.GetPropertyName() == GET_MEMBER_NAME_CHECKED(ThisClass, ConvertOutputs);
+		|| 	PropertyChangedEvent.GetPropertyName() == GET_MEMBER_NAME_CHECKED(ThisClass, ConvertOutputs)
+		||	PropertyChangedEvent.GetPropertyName() == GET_MEMBER_NAME_CHECKED(ThisClass, NodeName);
 
 	if (bNeedsNodeUpdate)
 	{
@@ -23378,8 +23380,8 @@ int32 UMaterialExpressionConvert::Compile(class FMaterialCompiler* Compiler, int
 }
 
 void UMaterialExpressionConvert::GetCaption(TArray<FString>& OutCaptions) const 
-{ 
-	OutCaptions.Add(TEXT("Convert")); 
+{
+	OutCaptions.Add(NodeName);
 }
 
 void UMaterialExpressionConvert::GetConnectorToolTip(int32 InputIndex, int32 OutputIndex, TArray<FString>& OutToolTip)
@@ -23493,6 +23495,111 @@ TSharedPtr<class SGraphNodeMaterialBase> UMaterialExpressionConvert::CreateCusto
 	return nullptr;
 }
 
+void UMaterialExpressionConvert::RegisterAdditionalMenuActions(FGraphActionMenuBuilder& ActionMenuBuilder, const FText& CategoryName)
+{
+	// Helper Lambda to Create a Make Node of a specified ConvertType
+	auto CreateMakeNode = [this, &ActionMenuBuilder, &CategoryName](EMaterialExpressionConvertType ConvertType)
+	{
+		const int32 NumComponents = MaterialExpressionConvertType::GetComponentCount(ConvertType);
+		const FText ActionName = FText::Format(LOCTEXT("Convert_MakeNode_Name", "Make Float{0}"), NumComponents);
+		const FText ActionTooltip = FText::Format(LOCTEXT("Convert_MakeNode_Tooltip", "Adds a Make Float{0} Node"), NumComponents);
+
+		TSharedPtr<FMaterialGraphSchemaAction_NewNode> Action_MakeNode(new FMaterialGraphSchemaAction_NewNode(
+			CategoryName,
+			ActionName,
+			ActionTooltip,
+			0,
+			GetKeywords()
+		));
+		Action_MakeNode->MaterialExpressionClass = GetClass();
+		Action_MakeNode->PostCreationDelegate = TDelegate<void(UMaterialExpression*)>::CreateLambda(
+			[ConvertType, NumComponents, ActionName](UMaterialExpression* NewExpression)
+			{
+				UMaterialExpressionConvert* NewConvertExpression = CastChecked<UMaterialExpressionConvert>(NewExpression);
+				NewConvertExpression->NodeName = ActionName.ToString();
+			
+				// Initialize NumComponents Float1 Inputs
+				FMaterialExpressionConvertInput Float1Input;
+				Float1Input.Type = EMaterialExpressionConvertType::Float1;
+				NewConvertExpression->ConvertInputs.Init(Float1Input, NumComponents);
+
+				// Add a single Output of our desired Convert Type
+				FMaterialExpressionConvertOutput ConvertOutput;
+				ConvertOutput.Type = ConvertType;
+				NewConvertExpression->ConvertOutputs = { ConvertOutput };
+
+				// Create the relevant mappings
+				for (int32 MappingIndex = 0; MappingIndex < NumComponents; ++MappingIndex)
+				{
+					FMaterialExpressionConvertMapping NewMapping(MappingIndex, 0, 0, MappingIndex);
+					NewConvertExpression->ConvertMappings.Add(NewMapping);
+				}
+
+				if (NewExpression->GraphNode)
+				{
+					NewExpression->GraphNode->ReconstructNode();
+				}
+			}
+		);
+		ActionMenuBuilder.AddAction(Action_MakeNode);
+	};
+
+	CreateMakeNode(EMaterialExpressionConvertType::Float2);
+	CreateMakeNode(EMaterialExpressionConvertType::Float3);
+	CreateMakeNode(EMaterialExpressionConvertType::Float4);
+
+	// Helper Lambda to Create a Break Node of a specified ConvertType
+	auto CreateBreakNode = [this, &ActionMenuBuilder, &CategoryName](EMaterialExpressionConvertType ConvertType)
+	{
+		const int32 NumComponents = MaterialExpressionConvertType::GetComponentCount(ConvertType);
+		const FText ActionName = FText::Format(LOCTEXT("Convert_MakeNode_Name", "Break Float{0}"), NumComponents);
+		const FText ActionTooltip = FText::Format(LOCTEXT("Convert_MakeNode_Tooltip", "Adds a Break Float{0} Node"), NumComponents);
+
+		TSharedPtr<FMaterialGraphSchemaAction_NewNode> Action_BreakNode(new FMaterialGraphSchemaAction_NewNode(
+			CategoryName,
+			ActionName,
+			ActionTooltip,
+			0,
+			GetKeywords()
+		));
+		Action_BreakNode->MaterialExpressionClass = GetClass();
+		Action_BreakNode->PostCreationDelegate = TDelegate<void(UMaterialExpression*)>::CreateLambda(
+			[ConvertType, NumComponents, ActionName](UMaterialExpression* NewExpression)
+			{
+				UMaterialExpressionConvert* NewConvertExpression = CastChecked<UMaterialExpressionConvert>(NewExpression);
+				NewConvertExpression->NodeName = ActionName.ToString();
+			
+				// Add a single Input of our desired Convert Type
+				FMaterialExpressionConvertInput ConvertInput;
+				ConvertInput.Type = ConvertType;
+				NewConvertExpression->ConvertInputs = { ConvertInput };
+
+				// Initialize NumComponents Float1 Outputs
+				FMaterialExpressionConvertOutput Float1Output;
+				Float1Output.Type = EMaterialExpressionConvertType::Float1;
+				NewConvertExpression->ConvertOutputs.Init(Float1Output, NumComponents);
+
+				// Create the relevant mappings
+				for (int32 MappingIndex = 0; MappingIndex < NumComponents; ++MappingIndex)
+				{
+					FMaterialExpressionConvertMapping NewMapping(0, MappingIndex, MappingIndex, 0);
+					NewConvertExpression->ConvertMappings.Add(NewMapping);
+				}
+
+				if (NewExpression->GraphNode)
+				{
+					NewExpression->GraphNode->ReconstructNode();
+				}
+			}
+		);
+		ActionMenuBuilder.AddAction(Action_BreakNode);
+	};
+
+	CreateBreakNode(EMaterialExpressionConvertType::Float2);
+	CreateBreakNode(EMaterialExpressionConvertType::Float3);
+	CreateBreakNode(EMaterialExpressionConvertType::Float4);
+}
+
 void UMaterialExpressionConvert::RecreateOutputs()
 {
 	Outputs.Reset();
diff --git a/Engine/Source/Runtime/Engine/Public/Materials/MaterialExpression.h b/Engine/Source/Runtime/Engine/Public/Materials/MaterialExpression.h
--- a/Engine/Source/Runtime/Engine/Public/Materials/MaterialExpression.h
+++ b/Engine/Source/Runtime/Engine/Public/Materials/MaterialExpression.h
@@ -15,6 +15,7 @@ class UEdGraphNode;
 class UMaterial;
 class UTexture;
 class UTextureCollection;
+struct FGraphActionMenuBuilder;
 struct FPropertyChangedEvent;
 struct FMaterialParameterMetadata;
 struct FMaterialShadingModelField;
@@ -566,6 +567,13 @@ class UMaterialExpression : public UObject
 	 */
 	virtual void DeletePin(EEdGraphPinDirection PinDirection, int32 PinIndex) {}
 
+	/**
+	 * Allows Material Expressions to register additional actions in the context menu
+	 * @param ActionMenuBuilder action menu builder to add actions to
+	 * @param CategoryName category name for actions
+	 */
+	virtual void RegisterAdditionalMenuActions(FGraphActionMenuBuilder& ActionMenuBuilder, const FText& CategoryName) {}
+
 protected:
 	/** Caches the list of fixed expression inputs this expression has. */
 	TArray<FExpressionInput*> CachedInputs;
diff --git a/Engine/Source/Runtime/Engine/Public/Materials/MaterialExpressionConvert.h b/Engine/Source/Runtime/Engine/Public/Materials/MaterialExpressionConvert.h
--- a/Engine/Source/Runtime/Engine/Public/Materials/MaterialExpressionConvert.h
+++ b/Engine/Source/Runtime/Engine/Public/Materials/MaterialExpressionConvert.h
@@ -96,6 +96,14 @@ struct FMaterialExpressionConvertMapping
 {
 	GENERATED_BODY()
 
+	FMaterialExpressionConvertMapping() = default;
+	FMaterialExpressionConvertMapping(int32 InputIndex, int32 InputComponentIndex, int32 OutputIndex, int32 OutputComponentIndex)
+	: InputIndex(InputIndex)
+	, InputComponentIndex(InputComponentIndex)
+	, OutputIndex(OutputIndex)
+	, OutputComponentIndex(OutputComponentIndex)
+	{}
+
 	/** Which Input to map from */
 	UPROPERTY(EditAnywhere, Category = "MaterialExpressionConvert")
 	int32 InputIndex = 0;
@@ -128,6 +136,9 @@ public:
 	UPROPERTY()
 	TArray<FMaterialExpressionConvertMapping> ConvertMappings;
 
+	UPROPERTY(EditAnywhere, Category = "MaterialExpressionConvert")
+	FString NodeName = TEXT("Convert");
+
 	UMaterialExpressionConvert(const FObjectInitializer& ObjectInitializer);
 
 	//~Begin UObject Interface
@@ -154,6 +165,7 @@ public:
 	virtual bool CanDeletePin(EEdGraphPinDirection PinDirection, int32 PinIndex) const override;
 	virtual void DeletePin(EEdGraphPinDirection PinDirection, int32 PinIndex) override;
 	virtual TSharedPtr<class SGraphNodeMaterialBase> CreateCustomGraphNodeWidget() override;
+	virtual void RegisterAdditionalMenuActions(FGraphActionMenuBuilder& ActionMenuBuilder, const FText& CategoryName) override;
 	#endif //WITH_EDITOR
 	//~ End UMaterialExpression Interface
 private:
