Fix issues where async asset loads could complete without intended callbacks being made
This fix to UE's Asset Manager class functionality makes sure OnAssetStateChangeCompleted() is called at the time a bundle state is changed, even when it has already completed by the time ChangeBundleStateForPrimaryAssets() is called.
Without these changes I am able to repro the underlying issue 1 out of 2 or 3 matches, with them I am currently 0 for 15.
Testing performed in cooked Main build 40299630
#rnx

diff --git a/Engine/Source/Runtime/Engine/Private/AssetManager.cpp b/Engine/Source/Runtime/Engine/Private/AssetManager.cpp
--- a/Engine/Source/Runtime/Engine/Private/AssetManager.cpp
+++ b/Engine/Source/Runtime/Engine/Private/AssetManager.cpp
@@ -2011,13 +2011,16 @@ TSharedPtr<FStreamableHandle> UAssetManager::ChangeBundleStateForPrimaryAssets(
 		ReturnHandle = NewHandles[0];
 		ensure(NewAssets.Num() == 1);
 
-		// If only one handle, return it and add callback
-		// Call internal callback and external callback when it finishes
-		LoadParams.OnComplete = FStreamableDelegateWithHandle::CreateWeakLambda(this,
-			[this, PrimaryAssetId = NewAssets[0], DelegateToCall = MoveTemp(LoadParams.OnComplete)](const TSharedPtr<FStreamableHandle>& Handle) mutable
-			{
-				OnAssetStateChangeCompleted(Handle, PrimaryAssetId, MoveTemp(DelegateToCall));
-			});
+		if (!ReturnHandle->HasLoadCompleted()) // OnAssetStateChangeCompleted should only be called for pending handles
+		{
+			// If only one handle, return it and add callback
+			// Call internal callback and external callback when it finishes
+			LoadParams.OnComplete = FStreamableDelegateWithHandle::CreateWeakLambda(this,
+				[this, PrimaryAssetId = NewAssets[0], DelegateToCall = MoveTemp(LoadParams.OnComplete)](const TSharedPtr<FStreamableHandle>& Handle) mutable
+				{
+					OnAssetStateChangeCompleted(Handle, PrimaryAssetId, MoveTemp(DelegateToCall));
+				});
+		}
 	}
 
 	if (ReturnHandle)
