Fix crash in gather references in AudioDeviceManager
#rnx

diff --git a/Engine/Source/Runtime/Engine/Private/AudioDeviceManager.cpp b/Engine/Source/Runtime/Engine/Private/AudioDeviceManager.cpp
--- a/Engine/Source/Runtime/Engine/Private/AudioDeviceManager.cpp
+++ b/Engine/Source/Runtime/Engine/Private/AudioDeviceManager.cpp
@@ -844,7 +844,9 @@ void FAudioDeviceManager::AddReferencedObjects(FReferenceCollector& Collector)
 	// occur on level load which hit this callstack even though they are not part
 	// of GarbageCollection. This inherent race condition has existed for at least
 	// 10 years and has not caused known issues. It remains as tech-debt. 
-	checkf(!IsAudioThreadRunning(), TEXT("The audio thread must be disabled or suspended while collecting object references"));
+	
+	//checkf(!IsAudioThreadRunning(), TEXT("The audio thread must be disabled or suspended while collecting object references"));
+	// This check should be renabled with the resolution of UE-253226
 #endif
 
 	IterateOverAllDevices(
