"stat gpu" improvements with the new GPU profiler

Display the busy, wait and idle stats all on the same line in the table:
- The stats themselves are still independent TStatIds, but an extra rendering path has been added to group the stats together for display.
- The stats are tagged with a "IsGPU" flag. Due to complexities in the stats system, the easiest way to pass this flag down without refactoring lots of public APIs and code paths was via the stat "category" string pointer. FStatNameAndInfo::GpuStatCategory is exported from the Core module, and a pointer comparison is performed when setting the flag.
- The type of the GPU stat (busy, wait or idle) is encoded in the number field of the stat's FName. This is retrieved later when the stats are rendered to merge them.

Fixed bug in stats system
- FThreadSafeStaticStatBase::HighPerformanceEnable was not initialized, with the comment claiming we need "atomic" initialization. However, the GetStatId() relies on this value being nullptr on the first call.
- I presume this was always the case because the field ended up compiled into the .bss section, but we're now hitting cases where this field has uninitialized garbage on Android.

Added FFrameTime event type to GPU profiler
- Will be used specifically for OpenGL to override the stat unit GPU time to the one provided by Swappy.

diff --git a/Engine/Source/Runtime/Core/Private/Stats/StatsCommand.cpp b/Engine/Source/Runtime/Core/Private/Stats/StatsCommand.cpp
--- a/Engine/Source/Runtime/Core/Private/Stats/StatsCommand.cpp
+++ b/Engine/Source/Runtime/Core/Private/Stats/StatsCommand.cpp
@@ -1525,7 +1525,9 @@ struct FHUDGroupManager
 				{
 					const FComplexStatMessage& AggregatedStatMessage = AggregatedNonStackStatsHistory[Index];
 					const bool bIsMemory = AggregatedStatMessage.NameAndInfo.GetFlag( EStatMetaFlags::IsMemory );
-					TArray<FComplexStatMessage>& Dest = bIsMemory ? HudGroup.MemoryAggregate : HudGroup.CountersAggregate; 
+					const bool bIsGpu = AggregatedStatMessage.NameAndInfo.GetFlag( EStatMetaFlags::IsGPU );
+
+					TArray<FComplexStatMessage>& Dest = bIsGpu ? HudGroup.GpuStatsAggregate : (bIsMemory ? HudGroup.MemoryAggregate : HudGroup.CountersAggregate); 
 
 					const bool bToBeAdded = InternalGroup.EnabledItems.Contains( AggregatedStatMessage.NameAndInfo.GetRawName() );
 					if( bToBeAdded )
diff --git a/Engine/Source/Runtime/Core/Private/Stats/StatsSystemTypes.cpp b/Engine/Source/Runtime/Core/Private/Stats/StatsSystemTypes.cpp
--- a/Engine/Source/Runtime/Core/Private/Stats/StatsSystemTypes.cpp
+++ b/Engine/Source/Runtime/Core/Private/Stats/StatsSystemTypes.cpp
@@ -270,14 +270,12 @@ FStartupMessages& FStartupMessages::Get()
 	FThreadSafeStaticStatBase
 -----------------------------------------------------------------------------*/
 
-const TStatIdData* FThreadSafeStaticStatBase::DoSetup(const char* InStatName, const TCHAR* InStatDesc, const char* InGroupName, const char* InGroupCategory, const TCHAR* InGroupDesc, bool bDefaultEnable, bool bShouldClearEveryFrame, EStatDataType::Type InStatType, bool bCycleStat, bool bSortByName, FPlatformMemory::EMemoryCounterRegion InMemoryRegion) const
+const TStatIdData* FThreadSafeStaticStatBase::DoSetup(FName InStatName, const TCHAR* InStatDesc, const char* InGroupName, const char* InGroupCategory, const TCHAR* InGroupDesc, bool bDefaultEnable, bool bShouldClearEveryFrame, EStatDataType::Type InStatType, bool bCycleStat, bool bSortByName, FPlatformMemory::EMemoryCounterRegion InMemoryRegion) const
 {
-	FName TempName(InStatName);
-
 	// send meta data, we don't use normal messages because the stats thread might not be running yet
-	FStartupMessages::Get().AddMetadata(TempName, InStatDesc, InGroupName, InGroupCategory, InGroupDesc, bShouldClearEveryFrame, InStatType, bCycleStat, bSortByName, InMemoryRegion);
+	FStartupMessages::Get().AddMetadata(InStatName, InStatDesc, InGroupName, InGroupCategory, InGroupDesc, bShouldClearEveryFrame, InStatType, bCycleStat, bSortByName, InMemoryRegion);
 
-	TStatIdData const* LocalHighPerformanceEnable(IStatGroupEnableManager::Get().GetHighPerformanceEnableForStat(FName(InStatName), InGroupName, InGroupCategory, bDefaultEnable, bShouldClearEveryFrame, InStatType, InStatDesc, bCycleStat, bSortByName, InMemoryRegion).GetRawPointer());
+	TStatIdData const* LocalHighPerformanceEnable(IStatGroupEnableManager::Get().GetHighPerformanceEnableForStat(InStatName, InGroupName, InGroupCategory, bDefaultEnable, bShouldClearEveryFrame, InStatType, InStatDesc, bCycleStat, bSortByName, InMemoryRegion).GetRawPointer());
 	TStatIdData const* OldHighPerformanceEnable = HighPerformanceEnable.Exchange(LocalHighPerformanceEnable);
 	check(!OldHighPerformanceEnable || LocalHighPerformanceEnable == OldHighPerformanceEnable); // we are assigned two different groups?
 
@@ -661,6 +659,8 @@ IStatGroupEnableManager& IStatGroupEnableManager::Get()
 	FStatNameAndInfo
 -----------------------------------------------------------------------------*/
 
+CORE_API const char* FStatNameAndInfo::GpuStatCategory = "STATCAT_GPU";
+
 FName FStatNameAndInfo::ToLongName(FName InStatName, char const* InGroup, char const* InCategory, TCHAR const* InDescription, bool InSortByName)
 {
 	TStringBuilder<256> LongName;
diff --git a/Engine/Source/Runtime/Core/Public/Stats/StatsData.h b/Engine/Source/Runtime/Core/Public/Stats/StatsData.h
--- a/Engine/Source/Runtime/Core/Public/Stats/StatsData.h
+++ b/Engine/Source/Runtime/Core/Public/Stats/StatsData.h
@@ -858,6 +858,9 @@ struct FActiveStatGroupInfo
 	/** Counters aggregates. */
 	TArray<FComplexStatMessage> CountersAggregate;
 
+	/** GPU aggregates. */
+	TArray<FComplexStatMessage> GpuStatsAggregate;
+
 	/** Children stats that should not be used when adding up group cost **/
 	TSet<FName> BudgetIgnoreStats;
 
diff --git a/Engine/Source/Runtime/Core/Public/Stats/StatsSystemTypes.h b/Engine/Source/Runtime/Core/Public/Stats/StatsSystemTypes.h
--- a/Engine/Source/Runtime/Core/Public/Stats/StatsSystemTypes.h
+++ b/Engine/Source/Runtime/Core/Public/Stats/StatsSystemTypes.h
@@ -262,8 +262,8 @@ struct EStatMetaFlags
 		Invalid = 0x00,
 		/** this bit is always one and is used for error checking. */
 		DummyAlwaysOne = 0x01,
-		/** if true, then this message contains all meta data as well and the name is long. NOT USED. */
-		HasLongNameAndMetaInfo = 0x02,
+		/** if true, then this message contains a GPU stat */
+		IsGPU = 0x02,
 		/** if true, then this message contains and int64 cycle or IsPackedCCAndDuration. */
 		IsCycle = 0x04,
 		/** if true, then this message contains a memory stat. */
@@ -359,6 +359,8 @@ class FStatNameAndInfo
 	int32 Number;
 
 public:
+	static CORE_API const char* GpuStatCategory;
+
 	FORCEINLINE_STATS FStatNameAndInfo()
 	{
 	}
@@ -390,7 +392,7 @@ public:
 
 		// ok, you can't have numbered stat FNames too large
 		checkStats(!(Number >> EStatAllFields::StartShift));
-		Number |= (EStatMetaFlags::DummyAlwaysOne | EStatMetaFlags::HasLongNameAndMetaInfo) << (EStatMetaFlags::Shift + EStatAllFields::StartShift);
+		Number |= EStatMetaFlags::DummyAlwaysOne << (EStatMetaFlags::Shift + EStatAllFields::StartShift);
 
 		SetField<EStatDataType>(InStatType);
 		SetFlag(EStatMetaFlags::ShouldClearEveryFrame, bShouldClearEveryFrame);
@@ -401,6 +403,11 @@ public:
 			SetField<EMemoryRegion>(EMemoryRegion::Type(MemoryRegion));
 		}
 
+		if (InCategory == GpuStatCategory)
+		{
+			SetFlag(EStatMetaFlags::IsGPU, true);
+		}
+
 		CheckInvariants();
 	}
 
@@ -1774,8 +1781,8 @@ public:
 struct FThreadSafeStaticStatBase
 {
 protected:
-	mutable TAtomic<const TStatIdData*> HighPerformanceEnable; // must be uninitialized, because we need atomic initialization
-	CORE_API const TStatIdData* DoSetup(const char* InStatName, const TCHAR* InStatDesc, const char* InGroupName, const char* InGroupCategory, const TCHAR* InGroupDesc, bool bDefaultEnable, bool bShouldClearEveryFrame, EStatDataType::Type InStatType, bool bCycleStat, bool bSortByName, FPlatformMemory::EMemoryCounterRegion InMemoryRegion) const;
+	mutable TAtomic<const TStatIdData*> HighPerformanceEnable { nullptr };
+	CORE_API const TStatIdData* DoSetup(FName InStatName, const TCHAR* InStatDesc, const char* InGroupName, const char* InGroupCategory, const TCHAR* InGroupDesc, bool bDefaultEnable, bool bShouldClearEveryFrame, EStatDataType::Type InStatType, bool bCycleStat, bool bSortByName, FPlatformMemory::EMemoryCounterRegion InMemoryRegion) const;
 };
 
 template<class TStatData, bool TCompiledIn>
@@ -1787,7 +1794,7 @@ struct FThreadSafeStaticStatInner : public FThreadSafeStaticStatBase
 		if (!LocalHighPerformanceEnable)
 		{
 			LocalHighPerformanceEnable = DoSetup(
-				TStatData::GetStatName(),
+				FName(TStatData::GetStatName()),
 				TStatData::GetDescription(),
 				TStatData::TGroup::GetGroupName(),
 				TStatData::TGroup::GetGroupCategory(),
@@ -1812,7 +1819,7 @@ struct FThreadSafeStaticStatInner : public FThreadSafeStaticStatBase
 struct FDynamicStat : public FThreadSafeStaticStatBase
 {
 	FDynamicStat(
-		const char* InStatName,
+		FName InStatName,
 		const TCHAR* InStatDesc,
 		const char* InGroupName,
 		const char* InGroupCategory,
diff --git a/Engine/Source/Runtime/Engine/Private/StatsRender2.cpp b/Engine/Source/Runtime/Engine/Private/StatsRender2.cpp
--- a/Engine/Source/Runtime/Engine/Private/StatsRender2.cpp
+++ b/Engine/Source/Runtime/Engine/Private/StatsRender2.cpp
@@ -85,58 +85,54 @@ void LexFromString( EStatFontTypes& OutValue, const TCHAR* Buffer )
 struct FStatRenderGlobals
 {
 	/** Rendering offset for first column from stat label. */
-	int32 AfterNameColumnOffset;
+	int32 AfterNameColumnOffset = 0;
+
+	/** Same as above, but for stat groups that don't need to display as many characters in the name column. */
+	int32 AfterNameColumnOffsetNarrow = 0;
 
 	/** Rendering offsets for additional columns. */
-	int32 InterColumnOffset;
+	int32 InterColumnOffset = 0;
+
+	/** Same as above, but for stat columns that don't need to display as many value characters. */
+	int32 InterColumnOffsetNarrow = 0;
 
 	/** Color for rendering stats. */
-	FLinearColor StatColor;
+	FLinearColor StatColor = { 0.f, 1.f, 0.f };
 
 	/** Color to use when rendering headings. */
-	FLinearColor HeadingColor;
+	FLinearColor HeadingColor = { 1.f, 0.2f, 0.f };
 
 	/** Color to use when rendering group names. */
-	FLinearColor GroupColor;
+	FLinearColor GroupColor = FLinearColor::White;
 
 	/** Color used as the background for every other stat item to make it easier to read across lines. */
 	FLinearColor BackgroundColors[2];
 
 	/** The font used for rendering stats. */
-	UFont* StatFont;
+	UFont* StatFont = nullptr;
 
 	/** The cached max char height of StatFont. */
-	int32 StatFontHeight;
+	int32 StatFontHeight = 12;
 
 	/** The legacy font size of StatFont the last time Initialize was called. */
-	int32 CachedLegacyFontSize;
+	int32 CachedLegacyFontSize = 0;
 
 	/** Current size of the viewport, used to detect resolution changes. */
 	FIntPoint SizeXY;
 
 	/** The X size that can be used to render the stats. */
-	int32 SafeSizeX;
+	int32 SafeSizeX = 0;
 
 	/** Current stat font type. */
-	EStatFontTypes StatFontType;
+	EStatFontTypes StatFontType = EStatFontTypes::NumFonts;
 
 	/** Whether we are in the stereo mode. */
-	bool bIsStereo;
+	bool bIsStereo = false;
 
 	/** Whether we need update internals. */
-	bool bNeedRefresh;
-
-	FStatRenderGlobals() :
-		AfterNameColumnOffset(0),
-		InterColumnOffset(0),
-		StatColor(0.f,1.f,0.f),
-		HeadingColor(1.f,0.2f,0.f),
-		GroupColor(FLinearColor::White),
-		StatFont(nullptr),
-		StatFontHeight(12),
-		CachedLegacyFontSize(0),
-		StatFontType(EStatFontTypes::NumFonts),
-		bNeedRefresh(true)
+	bool bNeedRefresh = true;
+
+	FStatRenderGlobals()
 	{
 		BackgroundColors[0] = FLinearColor(0.05f, 0.05f, 0.05f, 0.92f); // dark gray mostly occluding the background
 		BackgroundColors[1] = FLinearColor(0.02f, 0.02f, 0.02f, 0.88f); // slightly different to help make long lines more readable
@@ -193,9 +189,11 @@ struct FStatRenderGlobals
 		{
 			// This is the number of W characters to leave spacing for in the stat name column.
 			const FString STATNAME_COLUMN_WIDTH = FString::ChrN( GetNumCharsForStatName(), 'S' );
+			const FString STATNAME_COLUMN_WIDTH_NARROW = FString::ChrN(GetNumCharsForStatName() / 2, 'S');
 
 			// This is the number of W characters to leave spacing for in the other stat columns.
 			const FString STATDATA_COLUMN_WIDTH = FString::ChrN( StatFontType == EStatFontTypes::Small ? 8 : 7, 'W' );
+			const FString STATDATA_COLUMN_WIDTH_NARROW = FString::ChrN(StatFontType == EStatFontTypes::Small ? 5 : 4, 'W');
 
 			// Get the size of the spaces, since we can't render the width calculation strings.
 			int32 StatColumnSpaceSizeY = 0;
@@ -204,9 +202,11 @@ struct FStatRenderGlobals
 			// #TODO: Compute on the stats thread to get the exact measurement
 			// Determine where the first column goes.
 			StringSize(StatFont, AfterNameColumnOffset, StatColumnSpaceSizeY, *STATNAME_COLUMN_WIDTH);
+			StringSize(StatFont, AfterNameColumnOffsetNarrow, StatColumnSpaceSizeY, *STATNAME_COLUMN_WIDTH_NARROW);
 
 			// Determine the width of subsequent columns.
 			StringSize(StatFont, InterColumnOffset, TimeColumnSpaceSizeY, *STATDATA_COLUMN_WIDTH);
+			StringSize(StatFont, InterColumnOffsetNarrow, TimeColumnSpaceSizeY, *STATDATA_COLUMN_WIDTH_NARROW);
 
 			bNeedRefresh = false;
 		}
@@ -312,7 +312,7 @@ protected:
 }
 StatCmdEngineExec;
 
-static void RightJustify( FCanvas* Canvas, const int32 X, const int32 Y, TCHAR const* Text, FLinearColor const& Color )
+static void RightJustify(FCanvas* Canvas, const int32 X, const int32 Y, TCHAR const* Text, FLinearColor const& Color)
 {
 	const FStatRenderGlobals& Globals = GetStatRenderGlobals();
 
@@ -590,67 +590,80 @@ static int32 RenderMemoryCounter( const FGameThreadStatsData& ViewData, const FC
 	return Globals.GetFontHeight();
 }
 
-static int32 RenderCounter( const FGameThreadStatsData& ViewData, const FComplexStatMessage& All, class FCanvas* Canvas, const int32 X, const int32 Y, const float Budget, const bool bIsBudgetIgnored )
+static void RenderCounterLabel(const FComplexStatMessage& All, class FCanvas* Canvas, const int32 X, const int32 Y)
 {
 	const FStatRenderGlobals& Globals = GetStatRenderGlobals();
 
-	// If this is a cycle, render it as a cycle. This is a special case for manually set cycle counters.
-	const bool bIsCycle = All.NameAndInfo.GetFlag(EStatMetaFlags::IsCycle);
-	if( bIsCycle )
-	{
-		return RenderCycle( All, Canvas, X, Y, 0, false, Budget, bIsBudgetIgnored );
-	}
-
-	const bool bDisplayAll = All.NameAndInfo.GetFlag(EStatMetaFlags::ShouldClearEveryFrame);
-
-	// Draw the label
 	const FString StatDesc = All.GetDescription();
 	const FString StatDisplay = StatDesc.Len() == 0 ? All.GetShortName().GetPlainNameString() : StatDesc;
 	Canvas->DrawShadowedString(X, Y, *ShortenName(*StatDisplay), Globals.StatFont, Globals.StatColor);
-	int32 CurrX = X + Globals.AfterNameColumnOffset;
+}
 
-	if( bDisplayAll )
+static void RenderCounterValues(const FComplexStatMessage& All, class FCanvas* Canvas, int32& X, const int32 Y, int32 Spacing, float ZeroStatAlphaMult = 1.0f)
+{
+	const FStatRenderGlobals& Globals = GetStatRenderGlobals();
+
+	auto RenderValue = [&](EComplexStatField::Type Field)
 	{
-		// Append the average.
 		if (All.NameAndInfo.GetField<EStatDataType>() == EStatDataType::ST_double)
 		{
-			const FString ValueFormatted = FormatStatValueFloat( All.GetValue_double( EComplexStatField::IncAve ) );
-			RightJustify(Canvas, CurrX, Y, *ValueFormatted, Globals.StatColor);
+			const FString ValueFormatted = FormatStatValueFloat(All.GetValue_double(Field));
+
+			FLinearColor Color = Globals.StatColor;
+			if (ValueFormatted == TEXT("0.00"))
+			{
+				Color.A *= ZeroStatAlphaMult;
+			}
+
+			RightJustify(Canvas, X, Y, *ValueFormatted, Color);
 		}
 		else if (All.NameAndInfo.GetField<EStatDataType>() == EStatDataType::ST_int64)
 		{
-			const FString ValueFormatted = FormatStatValueInt64( All.GetValue_int64( EComplexStatField::IncAve ) );
-			RightJustify(Canvas, CurrX, Y, *ValueFormatted, Globals.StatColor);
+			int64 Value = All.GetValue_int64(Field);
+			const FString ValueFormatted = FormatStatValueInt64(Value);
+
+			FLinearColor Color = Globals.StatColor;
+			if (Value == 0)
+			{
+				Color.A *= ZeroStatAlphaMult;
+			}
+
+			RightJustify(Canvas, X, Y, *ValueFormatted, Color);
 		}
-	}
-	CurrX += Globals.InterColumnOffset;
+	};
 
-	// Append the maximum.
-	if (All.NameAndInfo.GetField<EStatDataType>() == EStatDataType::ST_double)
+	if (All.NameAndInfo.GetFlag(EStatMetaFlags::ShouldClearEveryFrame))
 	{
-		const FString ValueFormatted = FormatStatValueFloat( All.GetValue_double( EComplexStatField::IncMax ) );
-		RightJustify(Canvas, CurrX, Y, *ValueFormatted, Globals.StatColor);
-	}
-	else if (All.NameAndInfo.GetField<EStatDataType>() == EStatDataType::ST_int64)
-	{
-		const FString ValueFormatted = FormatStatValueInt64( All.GetValue_int64( EComplexStatField::IncMax ) );
-		RightJustify(Canvas, CurrX, Y, *ValueFormatted, Globals.StatColor);
+		// Append the average.
+		RenderValue(EComplexStatField::IncAve);
 	}
+	X += Spacing;
 
-	CurrX += Globals.InterColumnOffset;
+	// Append the maximum.
+	RenderValue(EComplexStatField::IncMax);
+	X += Spacing;
 
 	// Append the minimum.
-	if (All.NameAndInfo.GetField<EStatDataType>() == EStatDataType::ST_double)
-	{
-		const FString ValueFormatted = FormatStatValueFloat(All.GetValue_double(EComplexStatField::IncMin));
-		RightJustify(Canvas, CurrX, Y, *ValueFormatted, Globals.StatColor);
-	}
-	else if (All.NameAndInfo.GetField<EStatDataType>() == EStatDataType::ST_int64)
+	RenderValue(EComplexStatField::IncMin);
+	X += Spacing;
+}
+
+static int32 RenderCounter( const FGameThreadStatsData& ViewData, const FComplexStatMessage& All, class FCanvas* Canvas, const int32 X, const int32 Y, const float Budget, const bool bIsBudgetIgnored )
+{
+	const FStatRenderGlobals& Globals = GetStatRenderGlobals();
+
+	// If this is a cycle, render it as a cycle. This is a special case for manually set cycle counters.
+	const bool bIsCycle = All.NameAndInfo.GetFlag(EStatMetaFlags::IsCycle);
+	if( bIsCycle )
 	{
-		const FString ValueFormatted = FormatStatValueInt64(All.GetValue_int64(EComplexStatField::IncMin));
-		RightJustify(Canvas, CurrX, Y, *ValueFormatted, Globals.StatColor);
+		return RenderCycle( All, Canvas, X, Y, 0, false, Budget, bIsBudgetIgnored );
 	}
 
+	// Draw the label
+	int32 CurrX = X;
+	RenderCounterLabel(All, Canvas, CurrX, Y);
+	CurrX += Globals.AfterNameColumnOffset;
+	RenderCounterValues(All, Canvas, CurrX, Y, Globals.InterColumnOffset);
 
 	return Globals.GetFontHeight();
 }
@@ -769,6 +782,143 @@ void RenderArrayOfStats( FCanvas* Canvas, const int32 X, int32& Y, const TArray<
 	}
 }
 
+#if RHI_NEW_GPU_PROFILER
+struct FGpuStatMessage
+{
+	const FComplexStatMessage* Busy;
+	const FComplexStatMessage* Wait;
+	const FComplexStatMessage* Idle;
+};
+
+static int32 RenderGpuStatHeadings(class FCanvas* Canvas, const int32 X, const int32 Y)
+{
+	const FStatRenderGlobals& Globals = GetStatRenderGlobals();
+
+	Canvas->DrawShadowedString(X, Y, TEXT("Counters"), Globals.StatFont, Globals.HeadingColor);
+
+	TCHAR const* Types[] =
+	{
+		TEXT("Busy"),
+		TEXT("Wait"),
+		TEXT("Idle")
+	};
+
+	// Draw headers
+	int32 CurrX = X + Globals.AfterNameColumnOffsetNarrow + Globals.InterColumnOffsetNarrow;
+	for (TCHAR const* Type : Types)
+	{
+		RightJustify(Canvas, CurrX, Y - Globals.GetFontHeight(), Type, Globals.HeadingColor);
+		CurrX += Globals.InterColumnOffsetNarrow * 2 + Globals.InterColumnOffset;
+	}
+
+	// Draw column labels
+	CurrX = X + Globals.AfterNameColumnOffsetNarrow;
+	for (TCHAR const* Type : Types)
+	{
+		// Draw the average column label.
+		RightJustify(Canvas, CurrX, Y, TEXT("Avg"), Globals.HeadingColor);
+		CurrX += Globals.InterColumnOffsetNarrow;
+
+		// Draw the max column label.
+		RightJustify(Canvas, CurrX, Y, TEXT("Max"), Globals.HeadingColor);
+		CurrX += Globals.InterColumnOffsetNarrow;
+
+		// Draw the max column label.
+		RightJustify(Canvas, CurrX, Y, TEXT("Min"), Globals.HeadingColor);
+
+		// Use a larger space between the stat types
+		CurrX += Globals.InterColumnOffset;
+	}
+
+	return Globals.GetFontHeight() + (Globals.GetFontHeight() / 3);
+}
+
+void RenderArrayOfGpuStats(FCanvas* Canvas, const int32 X, int32& Y, const TMap<FName, FGpuStatMessage>& Aggregates, const FGameThreadStatsData& ViewData)
+{
+	const FStatRenderGlobals& Globals = GetStatRenderGlobals();
+	const FTexture* BackgroundTexture = Globals.GetBackgroundTexture();
+
+	int32 RowIndex = 0;
+	auto DrawRowBackground = [&]()
+	{
+		if (BackgroundTexture != nullptr)
+		{
+			Canvas->DrawTile(
+				X,
+				Y + Globals.GetYOffset(),
+				Globals.AfterNameColumnOffsetNarrow + Globals.InterColumnOffsetNarrow * 6 + Globals.InterColumnOffset * 3,
+				Globals.GetFontHeight(),
+				0, 0, 1, 1,
+				Globals.BackgroundColors[RowIndex & 1],
+				BackgroundTexture,
+				true
+			);
+		}
+	};
+
+	// Render all counters.
+	int32 MaxStatsPerGroup = CVarNumStatsPerGroup.GetValueOnGameThread();
+	for (auto const& [Name, Stat] : Aggregates)
+	{
+		if (RowIndex >= MaxStatsPerGroup)
+		{
+			break;
+		}
+		
+		DrawRowBackground();
+
+		int32 CurrX = X;
+		const FComplexStatMessage* StatForLabel = 
+			Stat.Busy ? Stat.Busy :
+			Stat.Wait ? Stat.Wait : 
+			Stat.Idle ? Stat.Idle : nullptr;
+
+		if (StatForLabel)
+		{
+			RenderCounterLabel(*StatForLabel, Canvas, CurrX, Y);
+			CurrX += Globals.AfterNameColumnOffsetNarrow;
+
+			// Render the Busy, Wait and Idle stats all on the same row
+			auto RenderCol = [&](const FComplexStatMessage* Message)
+			{
+				if (Message)
+				{
+					RenderCounterValues(*Message, Canvas, CurrX, Y, Globals.InterColumnOffsetNarrow, 0.4f);
+				}
+				else
+				{
+					CurrX += Globals.InterColumnOffsetNarrow * 3;
+				}
+
+				// Adjust for a wider space between stat types
+				CurrX -= Globals.InterColumnOffsetNarrow;
+				CurrX += Globals.InterColumnOffset;
+			};
+
+			RenderCol(Stat.Busy);
+			RenderCol(Stat.Wait);
+			RenderCol(Stat.Idle);
+
+			Y += Globals.GetFontHeight();
+		}
+
+		// Stop drawing after we have gone off the bottom of the screen
+		if (Y > (Globals.SizeXY.Y + Globals.GetFontHeight()))
+		{
+			break;
+		}
+
+		RowIndex++;
+	}
+
+	if (MaxStatsPerGroup < Aggregates.Num())
+	{
+		DrawRowBackground();
+		Y += RenderMoreStatsLine(Canvas, X, Y, Aggregates.Num() - MaxStatsPerGroup);
+	}
+}
+#endif
+
 static int32 RenderFlatCycle( const FGameThreadStatsData& ViewData, const FComplexStatMessage& Item, class FCanvas* Canvas, const int32 X, const int32 Y, const  float Budget, const bool bIsBudgetIgnored )
 {
 	return RenderCycle( Item, Canvas, X, Y, 0, true, Budget, bIsBudgetIgnored);
@@ -873,6 +1023,36 @@ static void RenderGroupedWithHierarchy(const FGameThreadStatsData& ViewData, FVi
 			RenderArrayOfStats(Canvas, X, Y, StatGroup.CountersAggregate, ViewData, StatGroup.BudgetIgnoreStats, -1.f, RenderCounter);
 			Y += Globals.GetFontHeight();
 		}
+
+		// Render GPU stats
+	#if RHI_NEW_GPU_PROFILER
+		if (StatGroup.GpuStatsAggregate.Num())
+		{
+			// Group together the individual GPU Busy, Wait and Idle stats 
+			TMap<FName, FGpuStatMessage> GroupedStats;
+			for (const FComplexStatMessage& Message : StatGroup.GpuStatsAggregate)
+			{
+				FName ShortName = Message.GetShortName();
+				int32 Number = ShortName.GetNumber();
+				ShortName.SetNumber(0);
+
+				using EType = UE::RHI::GPUProfiler::FGPUStat::EType;
+
+				FGpuStatMessage& Stat = GroupedStats.FindOrAdd(ShortName);
+				switch (EType(Number))
+				{
+				default: checkNoEntry(); [[fallthrough]];
+				case EType::Busy: Stat.Busy = &Message; break;
+				case EType::Wait: Stat.Wait = &Message; break;
+				case EType::Idle: Stat.Idle = &Message; break;
+				}
+			}
+
+			Y += RenderGpuStatHeadings(Canvas, X, Y);
+			RenderArrayOfGpuStats(Canvas, X, Y, GroupedStats, ViewData);
+			Y += Globals.GetFontHeight();
+		}
+	#endif
 	}
 }
 
@@ -915,4 +1095,4 @@ void RenderStats(FViewport* Viewport, class FCanvas* Canvas, int32 X, int32 Y, i
 	}
 }
 
-#endif
+#endif
\ No newline at end of file
diff --git a/Engine/Source/Runtime/RHI/Private/GPUProfiler.cpp b/Engine/Source/Runtime/RHI/Private/GPUProfiler.cpp
--- a/Engine/Source/Runtime/RHI/Private/GPUProfiler.cpp
+++ b/Engine/Source/Runtime/RHI/Private/GPUProfiler.cpp
@@ -1326,84 +1326,83 @@ namespace UE::RHI::GPUProfiler
 	static FGPUStat GPUStat_Total(TEXT("Queue Total"));
 #endif
 
-	TMap<FGPUStat::FStatIdentifier, TUniquePtr<FGPUStat::FStatCategory>> FGPUStat::Categories;
+	TMap<FQueue, TUniquePtr<FGPUStat::FStatCategory>> FGPUStat::Categories;
 
-	FGPUStat::FStatCategory& FGPUStat::GetCategory(FStatIdentifier ID)
+	FGPUStat::FStatCategory& FGPUStat::GetCategory(FQueue Queue)
 	{
-		TUniquePtr<FStatCategory>& Category = Categories.FindOrAdd(ID);
+		TUniquePtr<FStatCategory>& Category = Categories.FindOrAdd(Queue);
 		if (!Category)
 		{
-			Category = MakeUnique<FStatCategory>(ID);
+			Category = MakeUnique<FStatCategory>(Queue);
 		}
 		return *Category;
 	}
 
-	FString FGPUStat::FStatIdentifier::GetIDString(bool bFriendly) const
+	FString FGPUStat::GetIDString(FQueue Queue, bool bFriendly)
 	{
-		TCHAR const* TypeStr;
-		switch (Type)
-		{
-		default: checkNoEntry(); [[fallthrough]];
-		case EType::Busy: TypeStr = TEXT("Busy"); break;
-		case EType::Idle: TypeStr = TEXT("Idle"); break;
-		case EType::Wait: TypeStr = TEXT("Wait"); break;
-		}
-
 		if (bFriendly)
 		{
-			return FString::Printf(TEXT("GPU %d %s Queue %d (%s)")
+			return FString::Printf(TEXT("GPU %d %s Queue %d")
 				, Queue.GPU
 				, Queue.GetTypeString()
 				, Queue.Index
-				, TypeStr
 			);
 		}
 		else
 		{
-			return FString::Printf(TEXT("GPU%d_%s%d_%s")
+			return FString::Printf(TEXT("GPU%d_%s%d")
 				, Queue.GPU
 				, Queue.GetTypeString()
 				, Queue.Index
-				, TypeStr
 			);
 		}
 	}
 
-	FGPUStat::FStatCategory_Stats::FStatCategory_Stats(FStatIdentifier const& ID)
+	FGPUStat::FStatCategory_Stats::FStatCategory_Stats(FQueue Queue)
 	#if STATS
-		: GroupName(FString::Printf(TEXT("STATGROUP_%s"), *ID.GetIDString(false)))
-		, GroupDesc(FString::Printf(TEXT("%s Timing"), *ID.GetIDString(true)))
+		: GroupName(FString::Printf(TEXT("STATGROUP_%s"), *GetIDString(Queue, false)))
+		, GroupDesc(FString::Printf(TEXT("%s Timing"), *GetIDString(Queue, true)))
 	#endif
 	{}
 
-	FGPUStat::FStatCategory_Csv::FStatCategory_Csv(FStatIdentifier const& ID)
+	FGPUStat::FStatCategory_Csv::FStatCategory_Csv(FQueue Queue)
 	#if CSV_PROFILER_STATS
-		: Csv(*ID.GetIDString(false), true)
+		: Csv(*GetIDString(Queue, false), true)
 	#endif
 	{}
 
-	FGPUStat::FStatCategory::FStatCategory(FStatIdentifier const& ID)
-		: FStatCategory_Stats(ID)
-		, FStatCategory_Csv(ID)
+	FGPUStat::FStatCategory::FStatCategory(FQueue Queue)
+		: FStatCategory_Stats(Queue)
+		, FStatCategory_Csv(Queue)
 	{}
 
 #if STATS
 	TStatId FGPUStat::GetStatId(FQueue Queue, EType Type)
 	{
-		FStatIdentifier ID = { Queue, Type };
-		FStatInstance& Instance = Instances.FindOrAdd(ID);
+		FStatInstance& Instance = Instances.FindOrAdd(Queue);
+
+		FStatInstance::FInner* Inner;
+		switch (Type)
+		{
+		default: checkNoEntry(); [[fallthrough]];
+		case EType::Busy: Inner = &Instance.Busy; break;
+		case EType::Wait: Inner = &Instance.Wait; break;
+		case EType::Idle: Inner = &Instance.Idle; break;
+		}
 
-		if (!Instance.Stat)
+		if (!Inner->Stat)
 		{
-			FStatCategory& Category = GetCategory(ID);
+			FStatCategory& Category = GetCategory(Queue);
 
-			Instance.StatName = MakeUnique<FAnsiString>(FString::Printf(TEXT("STAT_%s_%s"), *ID.GetIDString(false), Name));
+			// Encode the stat type in the FName number
+			int32 Number = int32(Type);
+			Inner->StatName = FName(*FString::Printf(TEXT("STAT_%s_%s"), *GetIDString(Queue, false), Name), Number);
 
-			Instance.Stat = MakeUnique<FDynamicStat>(
-				**Instance.StatName,
+			Inner->Stat = MakeUnique<FDynamicStat>(
+				Inner->StatName,
 				Name,
 				*Category.GroupName,
-				"STATCAT_GPU",
+				FStatNameAndInfo::GpuStatCategory,
 				*Category.GroupDesc,
 				true, // IsDefaultEnabled
 				true, // IsClearEveryFrame
@@ -1414,23 +1413,22 @@ namespace UE::RHI::GPUProfiler
 			);
 		}
 
-		return Instance.Stat->GetStatId();
+		return Inner->Stat->GetStatId();
 	}
 #endif
 
 #if CSV_PROFILER_STATS
 	FCsvDeclaredStat& FGPUStat::GetCsvStat(FQueue Queue)
 	{
-		FStatIdentifier ID = { Queue, EType::Busy };
-		FStatInstance& Instance = Instances.FindOrAdd(ID);
+		FStatInstance& Instance = Instances.FindOrAdd(Queue);
 
-		if (!Instance.CsvStat.IsSet())
+		if (!Instance.Busy.CsvStat.IsSet())
 		{
-			FStatCategory& Category = GetCategory(ID);
-			Instance.CsvStat.Emplace(Name, Category.Csv.Index);
+			FStatCategory& Category = GetCategory(Queue);
+			Instance.Busy.CsvStat.Emplace(Name, Category.Csv.Index);
 		}
 
-		return *Instance.CsvStat;
+		return *Instance.Busy.CsvStat;
 	}
 #endif
 
@@ -1603,6 +1601,9 @@ namespace UE::RHI::GPUProfiler
 
 			uint64 CPUFrameBoundary = 0;
 
+			// Used to override the GPU time calculation for this queue, if an FFrameTime event is in the stream
+			TOptional<uint64> TotalBusyCycles;
+
 		#if WITH_RHI_BREADCRUMBS
 			TMap<FRHIBreadcrumbData_Stats, FStatState> Stats;
 		#endif
@@ -1717,15 +1718,15 @@ namespace UE::RHI::GPUProfiler
 						}
 					}
 
-					auto Recurse = [&](auto& Recurse, FNode* Root, FNode* Current, bool bParentMatchedFilter, int32 Level) -> bool
+					auto Recurse = [&](auto& Recurse, FNode* Root, FNode* CurrentNode, bool bParentMatchedFilter, int32 Level) -> bool
 					{
 						// Percent that this node was of the total frame time
-						const double Percent = Root ? (Current->Inclusive.GetBusyMilliseconds() / Root->Inclusive.GetBusyMilliseconds()) * 100.0 : 100.0;
+						const double Percent = Root ? (CurrentNode->Inclusive.GetBusyMilliseconds() / Root->Inclusive.GetBusyMilliseconds()) * 100.0 : 100.0;
 
 						// Filter nodes according to cvar settings
 						const bool bAboveThreshold = Percent >= PercentThreshold;
-						const bool bNameMatches = bParentMatchedFilter || RootWildcard.IsMatch(Current->Name);
-						const bool bHasWork = bShowEmptyNodes || Current->Inclusive.HasWork();
+						const bool bNameMatches = bParentMatchedFilter || RootWildcard.IsMatch(CurrentNode->Name);
+						const bool bHasWork = bShowEmptyNodes || CurrentNode->Inclusive.HasWork();
 
 						const bool bDisplayEvent = bNameMatches && bHasWork && bAboveThreshold;
 			
@@ -1733,14 +1734,14 @@ namespace UE::RHI::GPUProfiler
 						{
 							if (Root == nullptr)
 							{
-								Root = Current;
+								Root = CurrentNode;
 							}
 
 							Table.AddRow(
 								Root,
-								Current->Inclusive,
-								Current->Exclusive,
-								Current->Name,
+								CurrentNode->Inclusive,
+								CurrentNode->Exclusive,
+								CurrentNode->Name,
 								Level
 							);
 						}
@@ -1749,7 +1750,7 @@ namespace UE::RHI::GPUProfiler
 						FNode::FStats OtherChildrenExclusive;
 						uint32 NumHiddenChildren = 0;
 
-						for (FNode* Child : Current->Children)
+						for (FNode* Child : CurrentNode->Children)
 						{
 							bool bChildShown = Recurse(Recurse, Root, Child, bDisplayEvent, bDisplayEvent ? Level + 1 : Level);
 							if (!bChildShown)
@@ -2053,6 +2054,11 @@ namespace UE::RHI::GPUProfiler
 				}
 			}
 
+			void FrameTime(uint64 TotalGPUTime)
+			{
+				Timestamps.TotalBusyCycles = TotalGPUTime;
+			}
+
 			void FrameBoundary(FEvent::FFrameBoundary const& Event, FFrameState& FrameState, uint32 FrameNumber)
 			{
 				check(!bBusy);
@@ -2254,6 +2260,13 @@ namespace UE::RHI::GPUProfiler
 					}
 					break;
 
+				case FEvent::EType::FrameTime:
+					{
+						const FEvent::FFrameTime& FrameTime = Event->Value.Get<FEvent::FFrameTime>();
+						QueueState.FrameTime(FrameTime.TotalGPUTime);
+					}
+					break;
+
 				case FEvent::EType::FrameBoundary:
 					{
 						FEvent::FFrameBoundary const& FrameBoundary = Event->Value.Get<FEvent::FFrameBoundary>();
@@ -2328,6 +2341,8 @@ namespace UE::RHI::GPUProfiler
 			const bool bCsvStatsEnabled = false;
 		#endif
 
+			TOptional<uint64> MaxQueueBusyCycles;
+
 			for (auto const& [Queue, QueueTimestamps] : FrameState)
 			{
 			#if WITH_RHI_BREADCRUMBS
@@ -2360,18 +2375,32 @@ namespace UE::RHI::GPUProfiler
 					, bCsvStatsEnabled ? CsvProfiler : nullptr
 				#endif
 				);
+
+				if (QueueTimestamps.TotalBusyCycles.IsSet())
+				{
+					uint64 CurrentMax = MaxQueueBusyCycles ? *MaxQueueBusyCycles : 0;
+					MaxQueueBusyCycles = FMath::Max(CurrentMax, *QueueTimestamps.TotalBusyCycles);
+				}
 			}
 
-			// Compute the whole-frame total GPU time.
-			TArray<FTimestampStream::FState, TInlineAllocator<GetRHIPipelineCount() * MAX_NUM_GPUS>> StreamPointers;
-			for (auto const& [Queue, State] : FrameState)
+			if (MaxQueueBusyCycles.IsSet())
 			{
-				StreamPointers.Emplace(State.Queue);
+				// Set the total GPU time stat according to the value directly provided by the platform RHI
+				GGPUFrameTime = FPlatformMath::TruncToInt(FPlatformTime::ToSeconds64(*MaxQueueBusyCycles) / FPlatformTime::GetSecondsPerCycle());
 			}
-			uint64 WholeFrameUnion = FTimestampStream::ComputeUnion(StreamPointers);
+			else
+			{
+				// Compute the whole-frame total GPU time.
+				TArray<FTimestampStream::FState, TInlineAllocator<GetRHIPipelineCount() * MAX_NUM_GPUS>> StreamPointers;
+				for (auto const& [Queue, State] : FrameState)
+				{
+					StreamPointers.Emplace(State.Queue);
+				}
+				uint64 WholeFrameUnion = FTimestampStream::ComputeUnion(StreamPointers);
 
-			// Update the global GPU frame time stats - need to convert to Cycles32 rather than Cycles64.
-			GGPUFrameTime = FPlatformMath::TruncToInt(FPlatformTime::ToSeconds64(WholeFrameUnion) / FPlatformTime::GetSecondsPerCycle());
+				// Update the global GPU frame time stats - need to convert to Cycles32 rather than Cycles64.
+				GGPUFrameTime = FPlatformMath::TruncToInt(FPlatformTime::ToSeconds64(WholeFrameUnion) / FPlatformTime::GetSecondsPerCycle());
+			}
 
 			// @todo set global csv GPU time
 			//RHISetGPUStatTotals(bCsvStatsEnabled, FPlatformTime::ToMilliseconds64(WholeFrameUnion));
diff --git a/Engine/Source/Runtime/RHI/Public/GPUProfiler.h b/Engine/Source/Runtime/RHI/Public/GPUProfiler.h
--- a/Engine/Source/Runtime/RHI/Public/GPUProfiler.h
+++ b/Engine/Source/Runtime/RHI/Public/GPUProfiler.h
@@ -127,6 +127,18 @@ namespace UE::RHI::GPUProfiler
 			{}
 		};
 
+		// When present in the stream, overrides the total GPU time stat with the value it contains.
+		// Used for platform RHIs which don't support accurate GPU timing.
+		struct FFrameTime
+		{
+			// Same frequency as FPlatformTime::Cycles64()
+			uint64 TotalGPUTime;
+
+			FFrameTime(uint64 InTotalGPUTime)
+				: TotalGPUTime(InTotalGPUTime)
+			{}
+		};
+
 	#if WITH_RHI_BREADCRUMBS
 		struct FBeginBreadcrumb
 		{
@@ -243,6 +255,7 @@ namespace UE::RHI::GPUProfiler
 		
 		using FStorage = TVariant<
 			  FFrameBoundary
+			, FFrameTime
 		#if WITH_RHI_BREADCRUMBS
 			, FBeginBreadcrumb
 			, FEndBreadcrumb
@@ -259,6 +272,7 @@ namespace UE::RHI::GPUProfiler
 		enum class EType
 		{
 			FrameBoundary   = FStorage::IndexOfType<FFrameBoundary  >(),
+			FrameTime       = FStorage::IndexOfType<FFrameTime      >(),
 		#if WITH_RHI_BREADCRUMBS
 			BeginBreadcrumb = FStorage::IndexOfType<FBeginBreadcrumb>(),
 			EndBreadcrumb   = FStorage::IndexOfType<FEndBreadcrumb  >(),
@@ -495,26 +509,12 @@ namespace UE::RHI::GPUProfiler
 	{
 		enum class EType
 		{
-			Busy, Idle, Wait
+			Busy,
+			Wait,
+			Idle
 		};
 
-		struct FStatIdentifier
-		{
-			FQueue Queue;
-			EType Type;
-
-			FString GetIDString(bool bFriendly) const;
-
-			bool operator == (FStatIdentifier const& RHS) const
-			{
-				return Queue == RHS.Queue && Type == RHS.Type;
-			}
-
-			friend uint32 GetTypeHash(FStatIdentifier const& ID)
-			{
-				return HashCombineFast(GetTypeHash(ID.Queue), GetTypeHash(ID.Type));
-			}
-		};
+		static FString GetIDString(FQueue Queue, bool bFriendly);
 
 	private:
 		struct FStatCategory_Stats
@@ -524,7 +524,7 @@ namespace UE::RHI::GPUProfiler
 			FString     const GroupDesc;
 		#endif
 
-			FStatCategory_Stats(FStatIdentifier const& ID);
+			FStatCategory_Stats(FQueue Queue);
 		};
 
 		struct FStatCategory_Csv
@@ -533,31 +533,36 @@ namespace UE::RHI::GPUProfiler
 			FCsvCategory const Csv;
 		#endif
 
-			FStatCategory_Csv(FStatIdentifier const& ID);
+			FStatCategory_Csv(FQueue Queue);
 		};
 
 		struct FStatCategory
 			: public FStatCategory_Stats
 			, public FStatCategory_Csv
 		{
-			FStatCategory(FStatIdentifier const& ID);
+			FStatCategory(FQueue Queue);
 		};
 
 		struct FStatInstance
 		{
-		#if STATS
-			TUniquePtr<FAnsiString> StatName;
-			TUniquePtr<FDynamicStat> Stat;
-		#endif
-		#if CSV_PROFILER_STATS
-			TOptional<FCsvDeclaredStat> CsvStat;
-		#endif
+			struct FInner
+			{
+			#if STATS
+				FName StatName;
+				TUniquePtr<FDynamicStat> Stat;
+			#endif
+			#if CSV_PROFILER_STATS
+				TOptional<FCsvDeclaredStat> CsvStat;
+			#endif
+			};
+
+			FInner Busy, Wait, Idle;
 		};
 
-		static TMap<FStatIdentifier, TUniquePtr<FStatCategory>> Categories;
-		static FStatCategory& GetCategory(FStatIdentifier ID);
+		static TMap<FQueue, TUniquePtr<FStatCategory>> Categories;
+		static FStatCategory& GetCategory(FQueue Queue);
 
-		TMap<FStatIdentifier, FStatInstance> Instances;
+		TMap<FQueue, FStatInstance> Instances;
 
 	public:
 		TCHAR const* const Name;
