Adding an option to trigger an incremental gc using obj gc

diff --git a/Engine/Source/Runtime/Engine/Private/UnrealEngine.cpp b/Engine/Source/Runtime/Engine/Private/UnrealEngine.cpp
--- a/Engine/Source/Runtime/Engine/Private/UnrealEngine.cpp
+++ b/Engine/Source/Runtime/Engine/Private/UnrealEngine.cpp
@@ -9044,16 +9044,20 @@ bool UEngine::HandleObjCommand( const TCHAR* Cmd, FOutputDevice& Ar )
 {
 	if( FParse::Command(&Cmd,TEXT("GARBAGE")) || FParse::Command(&Cmd,TEXT("GC")) )
 	{
+		bool bIncremental = FParse::Command(&Cmd, TEXT("INCREMENTAL"));
+
 		// Purge unclaimed objects.
 		Ar.Logf(TEXT("Collecting garbage and resetting GC timer."));
-		CollectGarbage( GARBAGE_COLLECTION_KEEPFLAGS );
+		CollectGarbage(GARBAGE_COLLECTION_KEEPFLAGS, !bIncremental);
 		TimeSinceLastPendingKillPurge = 0.f;
 		return true;
 	}
 	else if (FParse::Command(&Cmd, TEXT("TRYGC")))
 	{
+		bool bIncremental = FParse::Command(&Cmd, TEXT("INCREMENTAL"));
+
 		// Purge unclaimed objects.		
-		if (TryCollectGarbage(GARBAGE_COLLECTION_KEEPFLAGS))
+		if (TryCollectGarbage(GARBAGE_COLLECTION_KEEPFLAGS, !bIncremental))
 		{
 			Ar.Logf(TEXT("Collecting garbage and resetting GC timer."));
 			TimeSinceLastPendingKillPurge = 0.f;
