Fix unnnecessary clamping on calculation of NumMips.
The itention of the clamping wass to ensure that the texture can fit in a page table. But it is incorrectly using the texture size in texels here instead of the texture size in tiles.
We already return early from this function when the texture is too large for the page table.
All 4K and above virtual textures are affected, and we have many of those in CitySample etc.
The error has gone under the radar because NumMips isn't used anywhere important. It is used for UI reporting. But anywhere that reads actual VT data relies on the VT "MaxLevel" calculation which is the number of mips down to the VT tile size and no further.
I'm taking this opportunity to also fix the mip calculation error that is already commented here.
Note that that there is agreement that although this will change the build artifacts for 4K and above textures, it's not worth bumping the DDC keys and taking a big patching hit.
The change will trigger cook determinism warnings for some fixed assets, but we will take that into account in the relevant tests.

diff --git a/Engine/Source/Runtime/Engine/Private/VT/VirtualTextureDataBuilder.cpp b/Engine/Source/Runtime/Engine/Private/VT/VirtualTextureDataBuilder.cpp
--- a/Engine/Source/Runtime/Engine/Private/VT/VirtualTextureDataBuilder.cpp
+++ b/Engine/Source/Runtime/Engine/Private/VT/VirtualTextureDataBuilder.cpp
@@ -385,14 +385,8 @@ bool FVirtualTextureBuilderDerivedInfo::InitializeFromBuildSettings(const FTextu
 
 	const uint32 Size = FMath::Max(SizeX, SizeY);
 
-	// Mip down to 1x1 pixels, but don't create more mips than can fit in max page table texture
-
-	// This is incorrect - should be floor not ceil. Testing _so far_ has found there to be no difference noted
-	// in projects, however saving this work to a separate CL since this is old and presumably stable code, due to
-	// the fact that the Min(x, VIRTUALTEXTURE_LOG2_MAX_PAGETABLE_SIZE) is load bearing. In order to get an incorrect
-	// mip count, you need Size to be non pow2 and the result to be larger than VIRTUALTEXTURE_LOG2_MAX_PAGETABLE_SIZE.
-	//	(Size can be non-pow2 because UDIM blocks must be pow2 but UDIM count can be non-pow2)
-	NumMips = FMath::Min<uint32>(FMath::CeilLogTwo(Size) + 1, VIRTUALTEXTURE_LOG2_MAX_PAGETABLE_SIZE);
+	// Mip down to 1x1 pixels
+	NumMips = FMath::FloorLog2(Size) + 1;
 
 	return true;
 }
