Deprecate bPrefixParameterNames since the code for this was removed a long time ago and no longer functions.

[RN] minor | Rendering - Materials and Shaders

diff --git a/Engine/Source/Runtime/Engine/Public/Materials/MaterialFunction.h b/Engine/Source/Runtime/Engine/Public/Materials/MaterialFunction.h
--- a/Engine/Source/Runtime/Engine/Public/Materials/MaterialFunction.h
+++ b/Engine/Source/Runtime/Engine/Public/Materials/MaterialFunction.h
@@ -64,10 +64,12 @@ class UMaterialFunction : public UMaterialFunctionInterface
 	/** Whether to list this function in the material function library, which is a window in the material editor that lists categorized functions. */
 	UPROPERTY(EditAnywhere, Category=MaterialFunction, AssetRegistrySearchable)
 	uint8 bExposeToLibrary:1;
-	
-	/** If true, parameters in this function will have a prefix added to their group name. */
-	UPROPERTY(EditAnywhere, Category=MaterialFunction)
+
+#if WITH_EDITORONLY_DATA	
+	UE_DEPRECATED(5.6, "bPrefixParameterNames has been deprecated.")
+	UPROPERTY()
 	uint8 bPrefixParameterNames:1;
+#endif // WITH_EDITORONLY_DATA
 
 	UPROPERTY(EditAnywhere, Category = MaterialFunction)
 	uint8 bEnableExecWire : 1; // maxtodo
