Prevent AISystem being instantiated and preserved on client worlds during replay.

1. During replay, a client World gets initialized and as its NetMode still unset a call to CreateAISystem() is made to create it's AISystem.

2. Replay (demo) loads and calls UEngine::MovePendingLevel(FWorldContext &Context) which sets the Demo Net Driver that would have prevented the AISystem to get created originally.

3. CreateAISystem() gets called again and correctly do not create AI System because conditions aren't met this time but was not clearing any previously created AISystem.

diff --git a/Engine/Source/Runtime/Engine/Private/World.cpp b/Engine/Source/Runtime/Engine/Private/World.cpp
--- a/Engine/Source/Runtime/Engine/Private/World.cpp
+++ b/Engine/Source/Runtime/Engine/Private/World.cpp
@@ -1860,8 +1860,20 @@ void UWorld::GetCollisionProfileChannelAndResponseParams(FName ProfileName, ECol
 
 UAISystemBase* UWorld::CreateAISystem()
 {
+	const bool bShouldHaveAISystemInThisNetMode = UAISystemBase::ShouldInstantiateInNetMode(GetNetMode());
+
+	// If we should not have an AI System in this NetMode, clean up possible one that could have been created already for this world BEFORE getting assigned a proper NetMode (could happens during Replays as follow:)	
+	// 1. During replay, a client World gets initialized and as its NetMode still unset and a call to CreateAISystem() is made to create it's AISystem.
+	// 2. Replay (demo) loads and calls UEngine::MovePendingLevel(FWorldContext &Context) which sets the Demo Net Driver that would have prevented the AISystem to get created initially.
+	// 3. CreateAISystem() gets called again and will correctly not create AI System because conditions aren't met this time but we should be clearing any previously created AISystem.
+	if (AISystem && !bShouldHaveAISystemInThisNetMode)
+	{
+		AISystem->CleanupWorld(/*bSessionEnded*/false, /*bCleanupResources*/true);
+		AISystem = nullptr;
+	}
+
 	// create navigation system for editor and server targets, but remove it from game clients
-	if (AISystem == NULL && UAISystemBase::ShouldInstantiateInNetMode(GetNetMode()) && PersistentLevel)
+	if (AISystem == nullptr && bShouldHaveAISystemInThisNetMode && PersistentLevel)
 	{
 		const FName AIModuleName = UAISystemBase::GetAISystemModuleName();
 		const AWorldSettings* WorldSettings = PersistentLevel->GetWorldSettings(false);
