Making UAnimInstance::Montage_PlayInternal virtual

diff --git a/Engine/Source/Runtime/Engine/Classes/Animation/AnimInstance.h b/Engine/Source/Runtime/Engine/Classes/Animation/AnimInstance.h
--- a/Engine/Source/Runtime/Engine/Classes/Animation/AnimInstance.h
+++ b/Engine/Source/Runtime/Engine/Classes/Animation/AnimInstance.h
@@ -910,7 +910,7 @@ protected:
 	ENGINE_API virtual void Montage_Advance(float DeltaSeconds);
 
 	ENGINE_API void Montage_StopInternal(TFunctionRef<FMontageBlendSettings(const FAnimMontageInstance*)> AlphaBlendSelectorFunction, const UAnimMontage* Montage = nullptr);
-	ENGINE_API float Montage_PlayInternal(UAnimMontage* MontageToPlay, const FMontageBlendSettings& BlendInSettings, float InPlayRate = 1.f, EMontagePlayReturnType ReturnValueType = EMontagePlayReturnType::MontageLength, float InTimeToStartMontageAt = 0.f, bool bStopAllMontages = true);
+	ENGINE_API virtual float Montage_PlayInternal(UAnimMontage* MontageToPlay, const FMontageBlendSettings& BlendInSettings, float InPlayRate = 1.f, EMontagePlayReturnType ReturnValueType = EMontagePlayReturnType::MontageLength, float InTimeToStartMontageAt = 0.f, bool bStopAllMontages = true);
 
 public:
 
