Add UE_INTERNAL to API wrapper scene query functions

diff --git a/Engine/Source/Runtime/Engine/Public/Physics/Experimental/ChaosInterfaceWrapper.h b/Engine/Source/Runtime/Engine/Public/Physics/Experimental/ChaosInterfaceWrapper.h
--- a/Engine/Source/Runtime/Engine/Public/Physics/Experimental/ChaosInterfaceWrapper.h
+++ b/Engine/Source/Runtime/Engine/Public/Physics/Experimental/ChaosInterfaceWrapper.h
@@ -58,7 +58,7 @@ namespace Chaos::Private
 }
 
 // Export raycast functions
-#define DECLARE_LOW_LEVEL_RAYCAST_ACCEL_HIT(TACCEL, THIT) template<> void ENGINE_API Chaos::Private::LowLevelRaycast(const TACCEL& Container, const FVector& Start, const FVector& Dir, float DeltaMag, FPhysicsHitCallback<THIT>& HitBuffer, EHitFlags OutputFlags, FQueryFlags QueryFlags, const FCollisionFilterData& Filter, const ChaosInterface::FQueryFilterData& QueryFilterData, ICollisionQueryFilterCallbackBase* QueryCallback, const ChaosInterface::FQueryDebugParams& DebugParams);
+#define DECLARE_LOW_LEVEL_RAYCAST_ACCEL_HIT(TACCEL, THIT) template<> UE_INTERNAL ENGINE_API void Chaos::Private::LowLevelRaycast(const TACCEL& Container, const FVector& Start, const FVector& Dir, float DeltaMag, FPhysicsHitCallback<THIT>& HitBuffer, EHitFlags OutputFlags, FQueryFlags QueryFlags, const FCollisionFilterData& Filter, const ChaosInterface::FQueryFilterData& QueryFilterData, ICollisionQueryFilterCallbackBase* QueryCallback, const ChaosInterface::FQueryDebugParams& DebugParams);
 #define DECLARE_LOW_LEVEL_RAYCAST_ACCEL(TACCEL) \
 DECLARE_LOW_LEVEL_RAYCAST_ACCEL_HIT(TACCEL, FHitRaycast); \
 DECLARE_LOW_LEVEL_RAYCAST_ACCEL_HIT(TACCEL, ChaosInterface::FPTRaycastHit);
@@ -67,7 +67,7 @@ DECLARE_LOW_LEVEL_RAYCAST_ACCEL(IExternalSpatialAcceleration)
 DECLARE_LOW_LEVEL_RAYCAST_ACCEL(FPhysScene)
 
 // Export sweep functions
-#define DECLARE_LOW_LEVEL_SWEEP_ACCEL_HIT(TACCEL, THIT) template<> ENGINE_API void Chaos::Private::LowLevelSweep(const TACCEL& Container, const FPhysicsGeometry& QueryGeom, const FTransform& StartTM, const FVector& Dir, float DeltaMag, FPhysicsHitCallback<THIT>& HitBuffer, EHitFlags OutputFlags, FQueryFlags QueryFlags, const FCollisionFilterData& Filter, const ChaosInterface::FQueryFilterData& QueryFilterData, ICollisionQueryFilterCallbackBase* QueryCallback, const ChaosInterface::FQueryDebugParams& DebugParams);
+#define DECLARE_LOW_LEVEL_SWEEP_ACCEL_HIT(TACCEL, THIT) template<> UE_INTERNAL ENGINE_API void Chaos::Private::LowLevelSweep(const TACCEL& Container, const FPhysicsGeometry& QueryGeom, const FTransform& StartTM, const FVector& Dir, float DeltaMag, FPhysicsHitCallback<THIT>& HitBuffer, EHitFlags OutputFlags, FQueryFlags QueryFlags, const FCollisionFilterData& Filter, const ChaosInterface::FQueryFilterData& QueryFilterData, ICollisionQueryFilterCallbackBase* QueryCallback, const ChaosInterface::FQueryDebugParams& DebugParams);
 #define DECLARE_LOW_LEVEL_SWEEP_ACCEL(TACCEL) \
 DECLARE_LOW_LEVEL_SWEEP_ACCEL_HIT(TACCEL, FHitSweep); \
 DECLARE_LOW_LEVEL_SWEEP_ACCEL_HIT(TACCEL, ChaosInterface::FPTSweepHit);
@@ -76,7 +76,7 @@ DECLARE_LOW_LEVEL_SWEEP_ACCEL(IExternalSpatialAcceleration)
 DECLARE_LOW_LEVEL_SWEEP_ACCEL(FPhysScene)
 
 // Export overlap functions
-#define DECLARE_LOW_LEVEL_OVERLAP_ACCEL_HIT(TACCEL, THIT) template<> ENGINE_API void Chaos::Private::LowLevelOverlap(const TACCEL& Container, const FPhysicsGeometry& QueryGeom, const FTransform& GeomPose, FPhysicsHitCallback<THIT>& HitBuffer, FQueryFlags QueryFlags, const FCollisionFilterData& Filter, const ChaosInterface::FQueryFilterData& QueryFilterData, ICollisionQueryFilterCallbackBase* QueryCallback, const ChaosInterface::FQueryDebugParams& DebugParams)
+#define DECLARE_LOW_LEVEL_OVERLAP_ACCEL_HIT(TACCEL, THIT) template<> UE_INTERNAL ENGINE_API void Chaos::Private::LowLevelOverlap(const TACCEL& Container, const FPhysicsGeometry& QueryGeom, const FTransform& GeomPose, FPhysicsHitCallback<THIT>& HitBuffer, FQueryFlags QueryFlags, const FCollisionFilterData& Filter, const ChaosInterface::FQueryFilterData& QueryFilterData, ICollisionQueryFilterCallbackBase* QueryCallback, const ChaosInterface::FQueryDebugParams& DebugParams)
 #define DECLARE_LOW_LEVEL_OVERLAP_ACCEL(TACCEL) \
 DECLARE_LOW_LEVEL_OVERLAP_ACCEL_HIT(TACCEL, FHitOverlap); \
 DECLARE_LOW_LEVEL_OVERLAP_ACCEL_HIT(TACCEL, ChaosInterface::FPTOverlapHit);
