Print out the Shader Job Cache's total shader code size and number of shaders we have size for.

`LogShaderCompilers: Display: Shader code size: total 784.475 MiB, numShaders 52,942, average 15.173 KiB, min 1.16 KiB, max 339.584 KiB`

[RN] minor | Rendering - Materials and Shaders

diff --git a/Engine/Source/Runtime/Engine/Private/ShaderCompiler/ShaderCompilerStats.cpp b/Engine/Source/Runtime/Engine/Private/ShaderCompiler/ShaderCompilerStats.cpp
--- a/Engine/Source/Runtime/Engine/Private/ShaderCompiler/ShaderCompilerStats.cpp
+++ b/Engine/Source/Runtime/Engine/Private/ShaderCompiler/ShaderCompilerStats.cpp
@@ -329,10 +329,12 @@ void FShaderCompilerStats::WriteStatSummary()
 
 	if (Counters.NumAccumulatedShaderCodes > 0)
 	{
+		const FString TotalCodeSizeStr = FText::AsMemory(Counters.AccumulatedShaderCodeSize).ToString();
+		const FString NumShadersStr = FormatNumber(Counters.NumAccumulatedShaderCodes);
 		const FString AvgCodeSizeStr = FText::AsMemory((uint64)((double)Counters.AccumulatedShaderCodeSize / (double)Counters.NumAccumulatedShaderCodes)).ToString();
 		const FString MinCodeSizeStr = FText::AsMemory((uint64)Counters.MinShaderCodeSize).ToString();
 		const FString MaxCodeSizeStr = FText::AsMemory((uint64)Counters.MaxShaderCodeSize).ToString();
-		UE_LOG(LogShaderCompilers, Display, TEXT("Shader code size: average %s, min %s, max %s"), *AvgCodeSizeStr, *MinCodeSizeStr, *MaxCodeSizeStr);
+		UE_LOG(LogShaderCompilers, Display, TEXT("Shader code size: total %s, numShaders %s, average %s, min %s, max %s"), *TotalCodeSizeStr, *NumShadersStr, *AvgCodeSizeStr, *MinCodeSizeStr, *MaxCodeSizeStr);
 	}
 
 	UE_LOG(LogShaderCompilers, Display, TEXT("Time at least one job was in flight (either pending or executed): %.2f s"), TotalTimeAtLeastOneJobWasInFlight);
