Speculative fix for a rare assertion that could be triggered when requesting a level that's still in memory while ShouldReuseUnloadedButStillAroundLevels is disabled.
Replaced the assertion with a fallback path that trashes the package and forces it to be reloaded.

diff --git a/Engine/Source/Runtime/Engine/Private/LevelStreaming.cpp b/Engine/Source/Runtime/Engine/Private/LevelStreaming.cpp
--- a/Engine/Source/Runtime/Engine/Private/LevelStreaming.cpp
+++ b/Engine/Source/Runtime/Engine/Private/LevelStreaming.cpp
@@ -1675,7 +1675,12 @@ bool ULevelStreaming::RequestLevel(UWorld* PersistentWorld, bool bAllowLevelLoad
 				return false;
 			}
 #endif
-			check(ULevelStreaming::ShouldReuseUnloadedButStillAroundLevels(World->PersistentLevel));
+			if (!ensure(ULevelStreaming::ShouldReuseUnloadedButStillAroundLevels(World->PersistentLevel)))
+			{
+				UE_LOG(LogLevelStreaming, Log, TEXT("Trying to reuse level %s but ShouldReuseUnloadedButStillAroundLevels returns false. Trash the current package and reload"), *LevelPackage->GetName())
+				FLevelStreamingGCHelper::TrashPackage(LevelPackage);
+				return false;
+			}
 			if (World->PersistentLevel != LoadedLevel)
 			{
 				// Level already exists but may have the wrong type due to being inactive before, so copy data over
