- Fix for morph updating both raster / ray tracing LODs. Don't enable morph on ray tracing skin cache vertex factories.
- Fix to force skin cache for a particular mesh if it supports ray tracing; otherwise both the inline and skin cache passes get run on the same LOD, which causes multiple bone buffer updates.

diff --git a/Engine/Source/Runtime/Engine/Private/Components/SkinnedMeshComponent.cpp b/Engine/Source/Runtime/Engine/Private/Components/SkinnedMeshComponent.cpp
--- a/Engine/Source/Runtime/Engine/Private/Components/SkinnedMeshComponent.cpp
+++ b/Engine/Source/Runtime/Engine/Private/Components/SkinnedMeshComponent.cpp
@@ -2233,7 +2233,10 @@ bool USkinnedMeshComponent::GetTwistAndSwingAngleOfDeltaRotationFromRefPose(FNam
 
 bool USkinnedMeshComponent::IsSkinCacheAllowed(int32 LodIdx) const
 {
-	const bool bGlobalDefault = GetSkinCacheDefaultBehavior() == ESkinCacheDefaultBehavior::Inclusive;
+	if (!GEnableGPUSkinCache)
+	{
+		return false;
+	}
 
 	if (GetMeshDeformerInstanceForLOD(LodIdx) != nullptr)
 	{
@@ -2242,15 +2245,25 @@ bool USkinnedMeshComponent::IsSkinCacheAllowed(int32 LodIdx) const
 		return false;
 	}
 
+	const bool bGlobalDefault = GetSkinCacheDefaultBehavior() == ESkinCacheDefaultBehavior::Inclusive;
+
 	if (!GetSkinnedAsset())
 	{
-		return GEnableGPUSkinCache && bGlobalDefault;
+		return bGlobalDefault;
+	}
+
+	if (FSkeletalMeshRenderData* SkelMeshRenderData = GetSkinnedAsset()->GetResourceForRendering())
+	{
+		if (SkelMeshRenderData->bSupportRayTracing)
+		{
+			return true;
+		}
 	}
 
 	FSkeletalMeshLODInfo* LodInfo = GetSkinnedAsset()->GetLODInfo(LodIdx);
 	if (!LodInfo)
 	{
-		return GEnableGPUSkinCache && bGlobalDefault;
+		return bGlobalDefault;
 	}
 
 	bool bLodEnabled = LodInfo->SkinCacheUsage == ESkinCacheUsage::Auto ?
@@ -2259,14 +2272,14 @@ bool USkinnedMeshComponent::IsSkinCacheAllowed(int32 LodIdx) const
 
 	if (!SkinCacheUsage.IsValidIndex(LodIdx))
 	{
-		return GEnableGPUSkinCache && bLodEnabled;
+		return bLodEnabled;
 	}
 
 	bool bComponentEnabled = SkinCacheUsage[LodIdx] == ESkinCacheUsage::Auto ? 
 		bLodEnabled :
 		SkinCacheUsage[LodIdx] == ESkinCacheUsage::Enabled;
 
-	return GEnableGPUSkinCache && bComponentEnabled;
+	return bComponentEnabled;
 }
 
 void USkinnedMeshComponent::GetBoneNames(TArray<FName>& BoneNames)
diff --git a/Engine/Source/Runtime/Engine/Private/SkeletalRenderGPUSkin.cpp b/Engine/Source/Runtime/Engine/Private/SkeletalRenderGPUSkin.cpp
--- a/Engine/Source/Runtime/Engine/Private/SkeletalRenderGPUSkin.cpp
+++ b/Engine/Source/Runtime/Engine/Private/SkeletalRenderGPUSkin.cpp
@@ -898,6 +898,12 @@ void FSkeletalMeshObjectGPUSkin::ProcessUpdatedDynamicData(FRHICommandList& RHIC
 		if (GEnableMorphTargets && LODData.GetNumVertices() > 0 && (bHasWeightedActiveMorphs || bHasExternalMorphs))
 		{
 			bMorphNeedsUpdate |= GForceUpdateMorphTargets != 0;
+
+			if (bMorphNeedsUpdate)
+			{
+				LOD.MorphVertexBufferPool->SetCurrentRevisionNumber(CurrentRevisionNumber);
+			}
+
 			MorphVertexBuffer = &LOD.MorphVertexBufferPool->GetMorphVertexBufferForWriting();
 
 			// Force an update if this is the first use of the buffer
@@ -908,7 +914,6 @@ void FSkeletalMeshObjectGPUSkin::ProcessUpdatedDynamicData(FRHICommandList& RHIC
 
 			if (bMorphNeedsUpdate)
 			{
-				LOD.MorphVertexBufferPool->SetCurrentRevisionNumber(CurrentRevisionNumber);
 				LOD.MorphVertexBufferPool->SetUpdatedFrameNumber(BoneTransformFrameNumber);
 			}
 		}
@@ -919,12 +924,9 @@ void FSkeletalMeshObjectGPUSkin::ProcessUpdatedDynamicData(FRHICommandList& RHIC
 		}
 	}
 
-	const bool bGPUSkinCacheEnabled = GPUSkinCache && GEnableGPUSkinCache &&
-		// Force skin cache enabled for ray tracing if the inline skinning technique was requested.
-		(DynamicData->GPUSkinTechnique == ESkeletalMeshGPUSkinTechnique::GPUSkinCache || (DynamicData->GPUSkinTechnique == ESkeletalMeshGPUSkinTechnique::Inline && Mode == EGPUSkinCacheEntryMode::RayTracing));
+	const bool bGPUSkinCacheEnabled = GPUSkinCache && GEnableGPUSkinCache && DynamicData->GPUSkinTechnique == ESkeletalMeshGPUSkinTechnique::GPUSkinCache;
 
 	// Immediately release any stale entry if we've recently switched to a LOD level that disallows skin cache
-	// This saves memory and avoids confusing ShouldUseSeparateSkinCacheEntryForRayTracing() which checks SkinCacheEntry == nullptr
 	if (!bGPUSkinCacheEnabled)
 	{
 #if RHI_RAYTRACING
@@ -976,7 +978,7 @@ void FSkeletalMeshObjectGPUSkin::ProcessUpdatedDynamicData(FRHICommandList& RHIC
 		FGPUBaseSkinAPEXClothVertexFactory::ClothShaderType* ClothShaderData = nullptr;
 
 		const bool bSectionUsingCloth = GEnableCloth && ClothVertexFactory != nullptr;
-		const bool bSectionUsingMorph = MorphVertexBuffer && !bSectionUsingCloth && (bHasExternalMorphs || (bHasWeightedActiveMorphs && DynamicData->SectionIdsUseByActiveMorphTargets.Contains(SectionIdx)));
+		const bool bSectionUsingMorph = Mode == EGPUSkinCacheEntryMode::Raster && MorphVertexBuffer && !bSectionUsingCloth && (bHasExternalMorphs || (bHasWeightedActiveMorphs && DynamicData->SectionIdsUseByActiveMorphTargets.Contains(SectionIdx)));
 		bool bSectionUsingSkinCache = bAllowAddToSkinCache ? Section.MaxBoneInfluences != 0 : false;
 
 		VertexFactory->UpdateMorphState(RHICmdList, bSectionUsingMorph);
diff --git a/Engine/Source/Runtime/Engine/Private/SkinnedMeshSceneProxyDesc.cpp b/Engine/Source/Runtime/Engine/Private/SkinnedMeshSceneProxyDesc.cpp
--- a/Engine/Source/Runtime/Engine/Private/SkinnedMeshSceneProxyDesc.cpp
+++ b/Engine/Source/Runtime/Engine/Private/SkinnedMeshSceneProxyDesc.cpp
@@ -299,6 +299,11 @@ bool FSkinnedMeshSceneProxyDesc::IsSkinCacheAllowed(int32 LodIdx) const
 {
 	const bool bGlobalDefault = GetSkinCacheDefaultBehavior() == ESkinCacheDefaultBehavior::Inclusive;
 
+	if (!GEnableGPUSkinCache)
+	{
+		return false;
+	}
+
 	if (GetMeshDeformerInstance() != nullptr)
 	{
 		// Disable skin cache if a mesh deformer is in use.
@@ -308,13 +313,21 @@ bool FSkinnedMeshSceneProxyDesc::IsSkinCacheAllowed(int32 LodIdx) const
 
 	if (!GetSkinnedAsset())
 	{
-		return GEnableGPUSkinCache && bGlobalDefault;
+		return bGlobalDefault;
+	}
+
+	if (FSkeletalMeshRenderData* SkelMeshRenderData = GetSkinnedAsset()->GetResourceForRendering())
+	{
+		if (SkelMeshRenderData->bSupportRayTracing)
+		{
+			return true;
+		}
 	}
 
 	FSkeletalMeshLODInfo* LodInfo = GetSkinnedAsset()->GetLODInfo(LodIdx);
 	if (!LodInfo)
 	{
-		return GEnableGPUSkinCache && bGlobalDefault;
+		return bGlobalDefault;
 	}
 
 	bool bLodEnabled = LodInfo->SkinCacheUsage == ESkinCacheUsage::Auto ?
@@ -323,14 +336,14 @@ bool FSkinnedMeshSceneProxyDesc::IsSkinCacheAllowed(int32 LodIdx) const
 
 	if (!SkinCacheUsage.IsValidIndex(LodIdx))
 	{
-		return GEnableGPUSkinCache && bLodEnabled;
+		return bLodEnabled;
 	}
 
 	bool bComponentEnabled = SkinCacheUsage[LodIdx] == ESkinCacheUsage::Auto ? 
 		bLodEnabled :
 		SkinCacheUsage[LodIdx] == ESkinCacheUsage::Enabled;
 
-	return GEnableGPUSkinCache && bComponentEnabled;
+	return bComponentEnabled;
 }
 
 UMeshDeformerInstance* FSkinnedMeshSceneProxyDesc::GetMeshDeformerInstance() const
@@ -341,7 +354,12 @@ UMeshDeformerInstance* FSkinnedMeshSceneProxyDesc::GetMeshDeformerInstance() con
 bool FSkinnedMeshSceneProxyDynamicData::IsSkinCacheAllowed(int32 LodIdx, const USkinnedAsset* InSkinnedAsset) const
 {
 	const bool bGlobalDefault = GetSkinCacheDefaultBehavior() == ESkinCacheDefaultBehavior::Inclusive;
-	
+
+	if (!GEnableGPUSkinCache)
+	{
+		return false;
+	}
+
 	if (bHasMeshDeformerInstance)
 	{
 		// Disable skin cache if a mesh deformer is in use.
@@ -351,13 +369,21 @@ bool FSkinnedMeshSceneProxyDynamicData::IsSkinCacheAllowed(int32 LodIdx, const U
 
 	if (!InSkinnedAsset)
 	{
-		return GEnableGPUSkinCache && bGlobalDefault;
+		return bGlobalDefault;
+	}
+
+	if (FSkeletalMeshRenderData* SkelMeshRenderData = InSkinnedAsset->GetResourceForRendering())
+	{
+		if (SkelMeshRenderData->bSupportRayTracing)
+		{
+			return true;
+		}
 	}
 
 	const FSkeletalMeshLODInfo* LodInfo = InSkinnedAsset->GetLODInfo(LodIdx);
 	if (!LodInfo)
 	{
-		return GEnableGPUSkinCache && bGlobalDefault;
+		return bGlobalDefault;
 	}
 
 	bool bLodEnabled = LodInfo->SkinCacheUsage == ESkinCacheUsage::Auto ?
@@ -366,14 +392,14 @@ bool FSkinnedMeshSceneProxyDynamicData::IsSkinCacheAllowed(int32 LodIdx, const U
 
 	if (!SkinCacheUsage.IsValidIndex(LodIdx))
 	{
-		return GEnableGPUSkinCache && bLodEnabled;
+		return bLodEnabled;
 	}
 
 	bool bComponentEnabled = SkinCacheUsage[LodIdx] == ESkinCacheUsage::Auto ? 
 		bLodEnabled :
 		SkinCacheUsage[LodIdx] == ESkinCacheUsage::Enabled;
 
-	return GEnableGPUSkinCache && bComponentEnabled;
+	return bComponentEnabled;
 }
 
 extern TAutoConsoleVariable<int32> CVarMeshDeformerMaxLod;
diff --git a/Engine/Source/Runtime/Engine/Public/SkeletalRenderPublic.h b/Engine/Source/Runtime/Engine/Public/SkeletalRenderPublic.h
--- a/Engine/Source/Runtime/Engine/Public/SkeletalRenderPublic.h
+++ b/Engine/Source/Runtime/Engine/Public/SkeletalRenderPublic.h
@@ -286,7 +286,7 @@ public:
 	virtual FRWBuffer* GetRayTracingDynamicVertexBuffer() { return nullptr; }
 	virtual int32 GetRayTracingLOD() const { return GetLOD(); }
 
-	virtual bool ShouldUseSeparateSkinCacheEntryForRayTracing() const { return GetLOD() != GetRayTracingLOD() || SkinCacheEntry == nullptr; }
+	virtual bool ShouldUseSeparateSkinCacheEntryForRayTracing() const { return GetLOD() != GetRayTracingLOD(); }
 	virtual FGPUSkinCacheEntry* GetSkinCacheEntryForRayTracing() const { return ShouldUseSeparateSkinCacheEntryForRayTracing() ? SkinCacheEntryForRayTracing : SkinCacheEntry; }
 
 	/** Request ray tracing geometry to be updated with specified LOD Vertex Buffers */
