Expose Chaos API functions to support Control Rig Physics

diff --git a/Engine/Source/Runtime/Engine/Private/PhysicsEngine/ImmediatePhysics/ImmediatePhysicsChaos/ImmediatePhysicsJointHandle_Chaos.cpp b/Engine/Source/Runtime/Engine/Private/PhysicsEngine/ImmediatePhysics/ImmediatePhysicsChaos/ImmediatePhysicsJointHandle_Chaos.cpp
--- a/Engine/Source/Runtime/Engine/Private/PhysicsEngine/ImmediatePhysics/ImmediatePhysicsChaos/ImmediatePhysicsJointHandle_Chaos.cpp
+++ b/Engine/Source/Runtime/Engine/Private/PhysicsEngine/ImmediatePhysics/ImmediatePhysicsChaos/ImmediatePhysicsJointHandle_Chaos.cpp
@@ -54,6 +54,73 @@ namespace ImmediatePhysics_Chaos
 		}
 	}
 
+	void UpdateJointSettingsFromLinearConstraint(const FLinearConstraint& Constraint, Chaos::FPBDJointSettings& JointSettings)
+	{
+		using namespace Chaos;
+
+		JointSettings.LinearMotionTypes[0] = ConvertToJointMotionType(Constraint.XMotion);
+		JointSettings.LinearMotionTypes[1] = ConvertToJointMotionType(Constraint.YMotion);
+		JointSettings.LinearMotionTypes[2] = ConvertToJointMotionType(Constraint.ZMotion);
+
+		JointSettings.LinearLimit = Constraint.Limit;
+		JointSettings.bSoftLinearLimitsEnabled = Constraint.bSoftConstraint;
+
+		JointSettings.SoftLinearStiffness =
+			Chaos::ConstraintSettings::SoftLinearStiffnessScale() * Constraint.Stiffness;
+		JointSettings.SoftLinearDamping =
+			Chaos::ConstraintSettings::SoftLinearDampingScale() * Constraint.Damping;
+		JointSettings.LinearRestitution = Constraint.Restitution;
+		JointSettings.LinearContactDistance = Constraint.ContactDistance;
+
+		// UE Disables Soft Limits when the Limit is less than some threshold. This is not necessary
+		// in Chaos but for now we also do it for parity's sake (See FLinearConstraint::UpdateLinearLimit_AssumesLocked).
+		if (JointSettings.LinearLimit < RB_MinSizeToLockDOF)
+		{
+			for (int32 AxisIndex = 0; AxisIndex < 3; ++AxisIndex)
+			{
+				if (JointSettings.LinearMotionTypes[AxisIndex] == EJointMotionType::Limited)
+				{
+					JointSettings.LinearMotionTypes[AxisIndex] = EJointMotionType::Locked;
+				}
+			}
+		}
+	}
+
+	// twist/swing Order is twist, swing1, swing2 and in degrees
+
+	void UpdateJointSettingsFromConeConstraint(const FConeConstraint& Constraint, Chaos::FPBDJointSettings& JointSettings)
+	{
+		JointSettings.AngularMotionTypes[(int32)Chaos::EJointAngularConstraintIndex::Swing1] = ConvertToJointMotionType(Constraint.Swing1Motion);
+		JointSettings.AngularMotionTypes[(int32)Chaos::EJointAngularConstraintIndex::Swing2] = ConvertToJointMotionType(Constraint.Swing2Motion);
+
+		JointSettings.AngularLimits[(int32)Chaos::EJointAngularConstraintIndex::Swing1] = FMath::DegreesToRadians(Constraint.Swing1LimitDegrees);
+		JointSettings.AngularLimits[(int32)Chaos::EJointAngularConstraintIndex::Swing2] = FMath::DegreesToRadians(Constraint.Swing2LimitDegrees);
+
+		JointSettings.bSoftSwingLimitsEnabled = Constraint.bSoftConstraint;
+
+		JointSettings.SoftSwingStiffness =
+			Chaos::ConstraintSettings::SoftAngularStiffnessScale() * Constraint.Stiffness;
+		JointSettings.SoftSwingDamping =
+			Chaos::ConstraintSettings::SoftAngularDampingScale() * Constraint.Damping;
+
+		JointSettings.SwingRestitution = Constraint.Restitution;
+
+		JointSettings.SwingContactDistance = Constraint.ContactDistance;
+	}
+
+	void UpdateJointSettingsFromTwistConstraint(const FTwistConstraint& Constraint, Chaos::FPBDJointSettings& JointSettings)
+	{
+		JointSettings.AngularMotionTypes[(int32)Chaos::EJointAngularConstraintIndex::Twist] = ConvertToJointMotionType(Constraint.TwistMotion);
+		JointSettings.AngularLimits[(int32)Chaos::EJointAngularConstraintIndex::Twist] = FMath::DegreesToRadians(Constraint.TwistLimitDegrees);
+		JointSettings.bSoftTwistLimitsEnabled = Constraint.bSoftConstraint;
+		JointSettings.SoftTwistStiffness =
+			Chaos::ConstraintSettings::SoftAngularStiffnessScale() * Constraint.Stiffness;
+		JointSettings.SoftTwistDamping =
+			Chaos::ConstraintSettings::SoftAngularDampingScale() * Constraint.Damping;
+		JointSettings.TwistRestitution = Constraint.Restitution;
+		JointSettings.TwistContactDistance = Constraint.ContactDistance;
+	}
+
 	// Copies data from a profile into settings
 	void UpdateJointSettingsFromConstraintProfile(const FConstraintProfileProperties& Profile, Chaos::FPBDJointSettings& JointSettings)
 	{
@@ -75,51 +142,15 @@ namespace ImmediatePhysics_Chaos
 
 		JointSettings.bUseLinearSolver = Profile.bUseLinearJointSolver;
 
-		JointSettings.LinearMotionTypes[0] = ConvertToJointMotionType(Profile.LinearLimit.XMotion);
-		JointSettings.LinearMotionTypes[1] = ConvertToJointMotionType(Profile.LinearLimit.YMotion);
-		JointSettings.LinearMotionTypes[2] = ConvertToJointMotionType(Profile.LinearLimit.ZMotion);
-
-		JointSettings.LinearLimit = Profile.LinearLimit.Limit;
-
-		// Order is twist, swing1, swing2 and in degrees
-		JointSettings.AngularMotionTypes[(int32)Chaos::EJointAngularConstraintIndex::Twist] = ConvertToJointMotionType(Profile.TwistLimit.TwistMotion);
-		JointSettings.AngularMotionTypes[(int32)Chaos::EJointAngularConstraintIndex::Swing1] = ConvertToJointMotionType(Profile.ConeLimit.Swing1Motion);
-		JointSettings.AngularMotionTypes[(int32)Chaos::EJointAngularConstraintIndex::Swing2] = ConvertToJointMotionType(Profile.ConeLimit.Swing2Motion);
-
-		JointSettings.AngularLimits[(int32)Chaos::EJointAngularConstraintIndex::Twist] = FMath::DegreesToRadians(Profile.TwistLimit.TwistLimitDegrees);
-		JointSettings.AngularLimits[(int32)Chaos::EJointAngularConstraintIndex::Swing1] = FMath::DegreesToRadians(Profile.ConeLimit.Swing1LimitDegrees);
-		JointSettings.AngularLimits[(int32)Chaos::EJointAngularConstraintIndex::Swing2] = FMath::DegreesToRadians(Profile.ConeLimit.Swing2LimitDegrees);
-
-		JointSettings.bSoftLinearLimitsEnabled = Profile.LinearLimit.bSoftConstraint;
-		JointSettings.bSoftTwistLimitsEnabled = Profile.TwistLimit.bSoftConstraint;
-		JointSettings.bSoftSwingLimitsEnabled = Profile.ConeLimit.bSoftConstraint;
+		UpdateJointSettingsFromLinearConstraint(Profile.LinearLimit, JointSettings);
+		UpdateJointSettingsFromConeConstraint(Profile.ConeLimit, JointSettings);
+		UpdateJointSettingsFromTwistConstraint(Profile.TwistLimit, JointSettings);
 
 		JointSettings.LinearSoftForceMode = (ConstraintSettings::SoftLinearForceMode() == 0) ?
 			EJointForceMode::Acceleration : EJointForceMode::Force;
 		JointSettings.AngularSoftForceMode = (ConstraintSettings::SoftAngularForceMode() == 0) ?
 			EJointForceMode::Acceleration : EJointForceMode::Force;
 
-		JointSettings.SoftLinearStiffness =
-			Chaos::ConstraintSettings::SoftLinearStiffnessScale() * Profile.LinearLimit.Stiffness;
-		JointSettings.SoftLinearDamping =
-			Chaos::ConstraintSettings::SoftLinearDampingScale() * Profile.LinearLimit.Damping;
-		JointSettings.SoftTwistStiffness =
-			Chaos::ConstraintSettings::SoftAngularStiffnessScale() * Profile.TwistLimit.Stiffness;
-		JointSettings.SoftTwistDamping =
-			Chaos::ConstraintSettings::SoftAngularDampingScale() * Profile.TwistLimit.Damping;
-		JointSettings.SoftSwingStiffness =
-			Chaos::ConstraintSettings::SoftAngularStiffnessScale() * Profile.ConeLimit.Stiffness;
-		JointSettings.SoftSwingDamping =
-			Chaos::ConstraintSettings::SoftAngularDampingScale() * Profile.ConeLimit.Damping;
-
-		JointSettings.LinearRestitution = Profile.LinearLimit.Restitution;
-		JointSettings.TwistRestitution = Profile.TwistLimit.Restitution;
-		JointSettings.SwingRestitution = Profile.ConeLimit.Restitution;
-
-		JointSettings.LinearContactDistance = Profile.LinearLimit.ContactDistance;
-		JointSettings.TwistContactDistance = Profile.TwistLimit.ContactDistance;
-		JointSettings.SwingContactDistance = Profile.ConeLimit.ContactDistance;
-
 		JointSettings.LinearDrivePositionTarget = Profile.LinearDrive.PositionTarget;
 		JointSettings.LinearDriveVelocityTarget = Profile.LinearDrive.VelocityTarget;
 		JointSettings.bLinearPositionDriveEnabled[0] = Profile.LinearDrive.XDrive.bEnablePositionDrive;
@@ -195,19 +226,6 @@ namespace ImmediatePhysics_Chaos
 			FMath::Clamp((float)Profile.AngularPlasticityThreshold, 0.0f, 1.0f) : FLT_MAX;;
 
 		JointSettings.ContactTransferScale = Profile.ContactTransferScale;
-
-
-		// UE Disables Soft Limits when the Limit is less than some threshold. This is not necessary in Chaos but for now we also do it for parity's sake (See FLinearConstraint::UpdateLinearLimit_AssumesLocked).
-		if (JointSettings.LinearLimit < RB_MinSizeToLockDOF)
-		{
-			for (int32 AxisIndex = 0; AxisIndex < 3; ++AxisIndex)
-			{
-				if (JointSettings.LinearMotionTypes[AxisIndex] == EJointMotionType::Limited)
-				{
-					JointSettings.LinearMotionTypes[AxisIndex] = EJointMotionType::Locked;
-				}
-			}
-		}
 	}
 
 	FJointHandle::FJointHandle(FChaosConstraintContainer* InConstraints, FConstraintInstance* ConstraintInstance, FActorHandle* Actor1, FActorHandle* Actor2)
diff --git a/Engine/Source/Runtime/Engine/Public/Physics/ImmediatePhysics/ImmediatePhysicsChaos/ImmediatePhysicsJointHandle_Chaos.h b/Engine/Source/Runtime/Engine/Public/Physics/ImmediatePhysics/ImmediatePhysicsChaos/ImmediatePhysicsJointHandle_Chaos.h
--- a/Engine/Source/Runtime/Engine/Public/Physics/ImmediatePhysics/ImmediatePhysicsChaos/ImmediatePhysicsJointHandle_Chaos.h
+++ b/Engine/Source/Runtime/Engine/Public/Physics/ImmediatePhysics/ImmediatePhysicsChaos/ImmediatePhysicsJointHandle_Chaos.h
@@ -15,6 +15,9 @@ namespace Chaos
 }
 
 struct FConstraintProfileProperties;
+struct FLinearConstraint;
+struct FConeConstraint;
+struct FTwistConstraint;
 
 namespace ImmediatePhysics_Chaos
 {
@@ -49,4 +52,15 @@ namespace ImmediatePhysics_Chaos
 	/** Transfers parameters from a constraint profile into the joint settings */
 	ENGINE_API void UpdateJointSettingsFromConstraintProfile(
 		const FConstraintProfileProperties& Profile, Chaos::FPBDJointSettings& JointSettings);
+
+	/** Extracts parameters from a constraint constraint into the joint settings */
+	ENGINE_API void UpdateJointSettingsFromLinearConstraint(
+		const FLinearConstraint& Constraint, Chaos::FPBDJointSettings& JointSettings);
+	/** Extracts parameters from a constraint constraint into the joint settings */
+	ENGINE_API void UpdateJointSettingsFromConeConstraint(
+		const FConeConstraint& Constraint, Chaos::FPBDJointSettings& JointSettings);
+	/** Extracts parameters from a constraint constraint into the joint settings */
+	ENGINE_API void UpdateJointSettingsFromTwistConstraint(
+		const FTwistConstraint& Constraint, Chaos::FPBDJointSettings& JointSettings);
+
 }
