Remove RVT continuous update from SVT build hash.
This setting has no impact on whether an SVT is compatible with the associated RVT.
#rnx

diff --git a/Engine/Source/Runtime/Engine/Private/Components/RuntimeVirtualTextureComponent.cpp b/Engine/Source/Runtime/Engine/Private/Components/RuntimeVirtualTextureComponent.cpp
--- a/Engine/Source/Runtime/Engine/Private/Components/RuntimeVirtualTextureComponent.cpp
+++ b/Engine/Source/Runtime/Engine/Private/Components/RuntimeVirtualTextureComponent.cpp
@@ -275,31 +275,41 @@ FTransform URuntimeVirtualTextureComponent::GetTexelSnapTransform() const
 	return FTransform(Offset);
 }
 
+// If a setting change can cause the streaming texture to no longer be valid then it should be included in this hash.
+union FPackedStreamingMipsSettings
+{
+	uint64 PackedValue;
+	struct
+	{
+		uint32 PackedSettingsVersion : 4;
+		uint32 MaterialType : 4;
+		uint32 TileSize : 12;
+		uint32 TileBorderSize : 4;
+		uint32 LODGroup : 8;
+		uint32 CompressTextures : 1;
+		uint32 SinglePhysicalSpace : 1;
+		uint32 ContinuousUpdate_Deprecated : 1;
+		uint32 bUseLowQualityCompression : 1;
+		uint32 LossyCompressionAmount : 4;
+	};
+
+	static uint64 RemoveDeprecated(uint64 InHash)
+	{
+		// Remove deprecated parts of the the packed build hash.
+		// Need to do this at runtime to handle old data.
+		FPackedStreamingMipsSettings Settings;
+		Settings.PackedValue = InHash;
+		Settings.ContinuousUpdate_Deprecated = 0;
+		return Settings.PackedValue;
+	}
+};
+
 uint64 URuntimeVirtualTextureComponent::CalculateStreamingTextureSettingsHash() const
 {
 	// Shouldn't need to call this when there is VirtualTexture == nullptr
 	check(VirtualTexture != nullptr);
 
-	// If a setting change can cause the streaming texture to no longer be valid then it should be included in this hash.
-	union FPackedSettings
-	{
-		uint64 PackedValue;
-		struct
-		{
-			uint32 PackedSettingsVersion : 4;
-			uint32 MaterialType : 4;
-			uint32 TileSize : 12;
-			uint32 TileBorderSize : 4;
-			uint32 LODGroup : 8;
-			uint32 CompressTextures : 1;
-			uint32 SinglePhysicalSpace : 1;
-			uint32 ContinuousUpdate : 1;
-			uint32 bUseLowQualityCompression : 1;
-			uint32 LossyCompressionAmount : 4;
-		};
-	};
-
-	FPackedSettings Settings;
+	FPackedStreamingMipsSettings Settings;
 	Settings.PackedValue = 0;
 	Settings.PackedSettingsVersion = 2;
 	Settings.MaterialType = (uint32)VirtualTexture->GetMaterialType();
@@ -307,7 +317,6 @@ uint64 URuntimeVirtualTextureComponent::CalculateStreamingTextureSettingsHash()
 	Settings.TileBorderSize = (uint32)VirtualTexture->GetTileBorderSize();
 	Settings.LODGroup = (uint32)VirtualTexture->GetLODGroup();
 	Settings.CompressTextures = (uint32)VirtualTexture->GetCompressTextures();
-	Settings.ContinuousUpdate = (uint32)VirtualTexture->GetContinuousUpdate();
 	Settings.SinglePhysicalSpace = (uint32)VirtualTexture->GetSinglePhysicalSpace();
 	Settings.bUseLowQualityCompression = (uint32)VirtualTexture->GetLQCompression();
 	Settings.LossyCompressionAmount = (uint32)GetLossyCompressionAmount();
@@ -395,7 +404,8 @@ URuntimeVirtualTextureComponent::EStreamingTextureStatusFlags URuntimeVirtualTex
 		}
 		else
 		{
-			if (StreamingTexture->BuildHash != CalculateStreamingTextureSettingsHash() || !IsCompatibleFormat(*VirtualTexture, *BuiltTexture))
+			const uint64 StreamingMipsBuildHash = FPackedStreamingMipsSettings::RemoveDeprecated(StreamingTexture->BuildHash);
+			if (StreamingMipsBuildHash != CalculateStreamingTextureSettingsHash() || !IsCompatibleFormat(*VirtualTexture, *BuiltTexture))
 			{
 				Status |= EStreamingTextureStatusFlags::NonMatchingStreamingTextureSettings;
 			}
