New material translator push:
- Evaluate the Normal attribute before all other attributes, as its value may be referenced by some expression that contributes to the other material attributes.

diff --git a/Engine/Source/Runtime/Engine/Private/Materials/MaterialIRInternal.cpp b/Engine/Source/Runtime/Engine/Private/Materials/MaterialIRInternal.cpp
--- a/Engine/Source/Runtime/Engine/Private/Materials/MaterialIRInternal.cpp
+++ b/Engine/Source/Runtime/Engine/Private/Materials/MaterialIRInternal.cpp
@@ -14,40 +14,39 @@ bool IsMaterialPropertyEnabled(EMaterialProperty InProperty)
 {
 	switch (InProperty)
 	{
-	case MP_Normal:
-	case MP_Tangent:
-	case MP_EmissiveColor:
-	case MP_Opacity:
-	case MP_OpacityMask:
-	case MP_BaseColor:
-	case MP_Metallic:
-	case MP_Specular:
-	case MP_Roughness:
-	case MP_Anisotropy:
-	case MP_AmbientOcclusion:
-	case MP_Refraction:
-	case MP_PixelDepthOffset:
-	case MP_SubsurfaceColor:
-	case MP_ShadingModel:
-	case MP_SurfaceThickness:
-	case MP_FrontMaterial:
-	case MP_Displacement:
-	case MP_WorldPositionOffset:
-		return true;
-	default:
-		return false;
+		case MP_Normal:
+		case MP_Tangent:
+		case MP_EmissiveColor:
+		case MP_Opacity:
+		case MP_OpacityMask:
+		case MP_BaseColor:
+		case MP_Metallic:
+		case MP_Specular:
+		case MP_Roughness:
+		case MP_Anisotropy:
+		case MP_AmbientOcclusion:
+		case MP_Refraction:
+		case MP_PixelDepthOffset:
+		case MP_SubsurfaceColor:
+		case MP_ShadingModel:
+		case MP_SurfaceThickness:
+		case MP_FrontMaterial:
+		case MP_Displacement:
+		case MP_WorldPositionOffset:
+			return true;
+		default:
+			return false;
 	}
 }
 
-bool NextMaterialAttributeInput(UMaterial* BaseMaterial, int32& PropertyIndex, FMaterialInputDescription& Input)
+bool NextMaterialAttributeInput(UMaterial* BaseMaterial, int32& PropertyIndex)
 {
 	for (; PropertyIndex < MP_MAX; ++PropertyIndex)
 	{
 		EMaterialProperty Property = (EMaterialProperty)PropertyIndex;
 		if (MIR::Internal::IsMaterialPropertyEnabled(Property) 
 			&& Property != MP_SubsurfaceColor
-			&& Property != MP_FrontMaterial
-			&& BaseMaterial->GetExpressionInputDescription(Property, Input))
+			&& Property != MP_FrontMaterial)
 		{
 			return true;
 		}
diff --git a/Engine/Source/Runtime/Engine/Private/Materials/MaterialIRInternal.h b/Engine/Source/Runtime/Engine/Private/Materials/MaterialIRInternal.h
--- a/Engine/Source/Runtime/Engine/Private/Materials/MaterialIRInternal.h
+++ b/Engine/Source/Runtime/Engine/Private/Materials/MaterialIRInternal.h
@@ -19,7 +19,7 @@ namespace MIR::Internal {
 bool IsMaterialPropertyEnabled(EMaterialProperty InProperty);
 
 //
-bool NextMaterialAttributeInput(UMaterial* BaseMaterial, int32& PropertyIndex, FMaterialInputDescription& Input);
+bool NextMaterialAttributeInput(UMaterial* BaseMaterial, int32& PropertyIndex);
 
 //
 FValue* CreateMaterialAttributeDefaultValue(FEmitter& Emitter, const UMaterial* Material, EMaterialProperty Property);
diff --git a/Engine/Source/Runtime/Engine/Private/Materials/MaterialIRModuleBuilder.cpp b/Engine/Source/Runtime/Engine/Private/Materials/MaterialIRModuleBuilder.cpp
--- a/Engine/Source/Runtime/Engine/Private/Materials/MaterialIRModuleBuilder.cpp
+++ b/Engine/Source/Runtime/Engine/Private/Materials/MaterialIRModuleBuilder.cpp
@@ -82,29 +82,42 @@ struct FMaterialIRModuleBuilderImpl
 
 	void Step_GenerateOutputInstructions()
 	{
-		// Prepare the array of FSetMaterialOutputInstr outputs from the material attributes inputs.
-		FMaterialInputDescription Input;
-		for (int32 Index = 0; MIR::Internal::NextMaterialAttributeInput(Builder->Material, Index, Input); ++Index)
-		{
-			EMaterialProperty Property = (EMaterialProperty)Index;
-
-			MIR::FSetMaterialOutput* Output = Emitter->EmitSetMaterialOutput(Property, nullptr);
+		// The normal input is read back from the value set in the material attribute.
+		// For this reason, the normal attribute is evaluated and set first, ensuring that
+		// other inputs can read its value.
+		PrepareSingleMaterialAttribute(MP_Normal);
 
-			if (Input.bUseConstant)
-			{
-				Output->Arg = Emitter->EmitConstantFromShaderValue(Input.ConstantValue);
-			}
-			else if (!Input.Input->IsConnected())
-			{
-				Output->Arg = MIR::Internal::CreateMaterialAttributeDefaultValue(*Emitter, Builder->Material, Property);
-			}
-			else
+		// Then prepare all the other material attributes.
+		for (int32 Index = 0; MIR::Internal::NextMaterialAttributeInput(Builder->Material, Index); ++Index)
+		{
+			if (Index != MP_Normal)
 			{
-				AnalysisContextStack.Last().ExpressionStack.Add(Input.Input->Expression);
+				PrepareSingleMaterialAttribute((EMaterialProperty)Index);
 			}
 		}
 	}
 
+	void PrepareSingleMaterialAttribute(EMaterialProperty Property)
+	{
+		FMaterialInputDescription Input;
+		Builder->Material->GetExpressionInputDescription(Property, Input);
+
+		MIR::FSetMaterialOutput* Output = Emitter->EmitSetMaterialOutput(Property, nullptr);
+
+		if (Input.bUseConstant)
+		{
+			Output->Arg = Emitter->EmitConstantFromShaderValue(Input.ConstantValue);
+		}
+		else if (!Input.Input->IsConnected())
+		{
+			Output->Arg = MIR::Internal::CreateMaterialAttributeDefaultValue(*Emitter, Builder->Material, Property);
+		}
+		else
+		{
+			AnalysisContextStack.Last().ExpressionStack.Add(Input.Input->Expression);
+		}
+	}
+
 	void Step_BuildMaterialExpressionsToIRGraph()
 	{
 		while (true)
@@ -273,7 +286,7 @@ struct FMaterialIRModuleBuilderImpl
 			// Get the function output type.
 			MIR::FTypePtr OutputType = MIR::FType::FromMaterialValueType((EMaterialValueType)FunctionOutput.ExpressionOutput->GetOutputType(0));
 
-			// Cast the value to the expected output type. This may fail (value will be null). 
+			// Cast the value to the expected output type. This may fail (value will be poison). 
 			Value = Emitter->EmitConstruct(OutputType, Value);
 
 			// And flow it to the relative function *call* output in the parent context.
diff --git a/Engine/Source/Runtime/Engine/Private/Materials/MaterialIRToHLSLTranslator.cpp b/Engine/Source/Runtime/Engine/Private/Materials/MaterialIRToHLSLTranslator.cpp
--- a/Engine/Source/Runtime/Engine/Private/Materials/MaterialIRToHLSLTranslator.cpp
+++ b/Engine/Source/Runtime/Engine/Private/Materials/MaterialIRToHLSLTranslator.cpp
@@ -230,6 +230,7 @@ struct FTranslator : FMaterialIRToHLSLTranslation
 	FHLSLPrinter Printer;
 	FString PixelAttributesHLSL;
 	FString WorldPositionOffsetHLSL;
+	FString EvaluateNormalMaterialAttributeHLSL;
 	FString EvaluateOtherMaterialAttributesHLSL;
 
 	void GeneratePixelAttributesHLSL()
@@ -253,23 +254,33 @@ struct FTranslator : FMaterialIRToHLSLTranslation
 		}
 	}
 	
-	FString GenerateStageHLSL(MIR::EStage Stage)
+	void GenerateVertexStageHLSL()
 	{
-		Printer = {};
-		Printer.Tabs = 1;
-		Printer << Indentation;
+		BeginStage(MIR::Stage_Vertex);
 
-		CurrentStage = Stage;
+		LowerBlock(Module->GetRootBlock(MIR::Stage_Vertex));
+
+		WorldPositionOffsetHLSL = MoveTemp(Printer.Buffer);
+	}
+	
+	void GenerateFragmentStageHLSL()
+	{
+		BeginStage(MIR::Stage_Fragment);
 
 		LowerBlock(Module->GetRootBlock(CurrentStage));
 
-		if (Stage == MIR::Stage_Fragment)
-		{
-			Printer << TEXT("PixelMaterialInputs.FrontMaterial = GetInitialisedSubstrateData()") << EndOfStatement;
-			Printer << TEXT("PixelMaterialInputs.Subsurface = 0") << EndOfStatement;
-		}
+		Printer << TEXT("PixelMaterialInputs.FrontMaterial = GetInitialisedSubstrateData()") << EndOfStatement;
+		Printer << TEXT("PixelMaterialInputs.Subsurface = 0") << EndOfStatement;
+
+		EvaluateOtherMaterialAttributesHLSL = MoveTemp(Printer.Buffer);
+	}
 
-		return MoveTemp(Printer.Buffer);
+	void BeginStage(MIR::EStage Stage)
+	{
+		Printer = {};
+		Printer.Tabs = 1;
+		Printer << Indentation;
+		CurrentStage = Stage;
 	}
 	
 	ENoOp LowerBlock(const MIR::FBlock& Block)
@@ -306,6 +317,14 @@ struct FTranslator : FMaterialIRToHLSLTranslation
 			{
 				Printer << EndOfStatement;
 			}
+
+			// Store the code needed to evaluate the normal in a separate chunk than the other material attributes
+			// since this needs to be emitted before the others in the material template.
+			const MIR::FSetMaterialOutput* SetMaterialOutput = MIR::As<MIR::FSetMaterialOutput>(Instr);
+			if (SetMaterialOutput && SetMaterialOutput->Property == MP_Normal) 
+			{
+				EvaluateNormalMaterialAttributeHLSL = MoveTemp(Printer.Buffer);
+			}
         }
 
         NumLocals = OldNumLocals;
@@ -780,8 +799,15 @@ struct FTranslator : FMaterialIRToHLSLTranslation
 		};
 
 		Params.Add(TEXT("pixel_material_inputs"), MoveTemp(PixelAttributesHLSL));
-		Params.Add(TEXT("calc_pixel_material_inputs_initial_calculations"), EvaluateOtherMaterialAttributesHLSL);
-		Params.Add(TEXT("calc_pixel_material_inputs_analytic_derivatives_initial"), MoveTemp(EvaluateOtherMaterialAttributesHLSL));
+
+		// "Normal" is treated in a special way because the rest of the attributes may lead back to reading it.
+		// Therefore, in the way MaterialTemplate.ush is structured, it needs to be evaluated before other attributes.
+		Params.Add(TEXT("calc_pixel_material_inputs_analytic_derivatives_normal"), EvaluateNormalMaterialAttributeHLSL);
+		Params.Add(TEXT("calc_pixel_material_inputs_normal"), EvaluateNormalMaterialAttributeHLSL);
+
+		// Then the other other attributes.
+		Params.Add(TEXT("calc_pixel_material_inputs_analytic_derivatives_other_inputs"), EvaluateOtherMaterialAttributesHLSL);
+		Params.Add(TEXT("calc_pixel_material_inputs_other_inputs"), EvaluateOtherMaterialAttributesHLSL);
 		
 		// MaterialAttributes
 		TArray<FGuid> OrderedVisibleAttributes = FMaterialAttributeDefinitionMap::GetOrderedVisibleAttributeList();
@@ -992,8 +1018,8 @@ void FMaterialIRToHLSLTranslation::Run(TMap<FString, FString>& OutParameters, FS
 
 	FTranslator Translator{ *this };
 	Translator.GeneratePixelAttributesHLSL();
-	Translator.WorldPositionOffsetHLSL = Translator.GenerateStageHLSL(MIR::Stage_Vertex);
-	Translator.EvaluateOtherMaterialAttributesHLSL = Translator.GenerateStageHLSL(MIR::Stage_Fragment);
+	Translator.GenerateVertexStageHLSL();
+	Translator.GenerateFragmentStageHLSL();
 	Translator.GenerateTemplateStringParameters(OutParameters);
 	Translator.GetShaderCompilerEnvironment(OutEnvironment);
 }
