Character Movement: Fixed bug where a simulated character's rotation was not shown correctly on the client when the player is on a rotating platform.

Problem: On remote clients, when a character was on a sufficiently rotating platform the character's target rotation would update to the direction of the character's capsule's newly adjusted rotation (i.e. the intended rotation would be wiped out with a very small change in rotation.) Note: this bug did NOT affect player character's whose rotation was based on their controller's rotation (i.e. players that aim with the camera.)
Solution: In that situation set the character's target rotation to the character's target rotation + rotation delta. This preservers the intent of the replicated rotation while adjusting for the rotation of the object they are standing on.
Testing: tested that normal movement replicated correctly. Tested a variety of different moving platforms including platforms that move with no rotations, platforms that twist like a screw on its side, platforms that tilt, hamster wheels, and platforms that move and yaw.

diff --git a/Engine/Source/Runtime/Engine/Private/Components/CharacterMovementComponent.cpp b/Engine/Source/Runtime/Engine/Private/Components/CharacterMovementComponent.cpp
--- a/Engine/Source/Runtime/Engine/Private/Components/CharacterMovementComponent.cpp
+++ b/Engine/Source/Runtime/Engine/Private/Components/CharacterMovementComponent.cpp
@@ -2002,8 +2002,13 @@ void UCharacterMovementComponent::SimulatedTick(float DeltaSeconds)
 				const FQuat NewCapsuleRotation = UpdatedComponent->GetComponentQuat();
 				if (Mesh == CharacterOwner->GetMesh() && !NewCapsuleRotation.Equals(SavedCapsuleRotation, 1e-6f) && ClientPredictionData)
 				{
-					// Smoothing should lerp toward this new rotation target, otherwise it will just try to go back toward the old rotation.
-					ClientPredictionData->MeshRotationTarget = NewCapsuleRotation;
+					// Add delta rotation to the target rotation and original offset. Otherwise this object will move back toward the old rotation.
+					const FQuat RotationDelta = NewCapsuleRotation - SavedCapsuleRotation;
+					ClientPredictionData->MeshRotationTarget += RotationDelta;
+					ClientPredictionData->OriginalMeshRotationOffset += RotationDelta;
+
+					// Update the MeshRotationOffset to match the capsule rotation.
+					ClientPredictionData->MeshRotationOffset = NewCapsuleRotation;
 					Mesh->SetRelativeLocationAndRotation(SavedMeshRelativeLocation, CharacterOwner->GetBaseRotationOffset());
 				}
 			}
