Remove legacy code that incorrectly overwrites physics mesh UV data.

FinishCreatingPhysicsMeshes_Chaos moves UVInfo from AsyncPhysicsCookHelper, invalidating the data there. As a result, the write to UVInfo in BodySetup causes it to be overwritten with invalid data. Since the move is now handled by FinishCreatingPhysicsMeshes_Chaos, this old line should be removed.

diff --git a/Engine/Source/Runtime/Engine/Private/PhysicsEngine/BodySetup.cpp b/Engine/Source/Runtime/Engine/Private/PhysicsEngine/BodySetup.cpp
--- a/Engine/Source/Runtime/Engine/Private/PhysicsEngine/BodySetup.cpp
+++ b/Engine/Source/Runtime/Engine/Private/PhysicsEngine/BodySetup.cpp
@@ -527,7 +527,6 @@ void UBodySetup::FinishCreatePhysicsMeshesAsync(FAsyncCookHelper* AsyncPhysicsCo
 	if(AsyncPhysicsCookHelper)
 	{
 		FinishCreatingPhysicsMeshes_Chaos(*AsyncPhysicsCookHelper);
-		UVInfo = AsyncPhysicsCookHelper->UVInfo;
 
 		delete AsyncPhysicsCookHelper;
 	}
