Fixing a bug with async physics queries running in an auto-rtfm context without a query handler.

diff --git a/Engine/Source/Runtime/Engine/Private/Collision/WorldCollisionAsync.cpp b/Engine/Source/Runtime/Engine/Private/Collision/WorldCollisionAsync.cpp
--- a/Engine/Source/Runtime/Engine/Private/Collision/WorldCollisionAsync.cpp
+++ b/Engine/Source/Runtime/Engine/Private/Collision/WorldCollisionAsync.cpp
@@ -6,6 +6,7 @@
 
 #include "Engine/World.h"
 #include "Async/TaskGraphInterfaces.h"
+#include "AutoRTFM/AutoRTFM.h"
 #include "Engine/HitResult.h"
 #include "Engine/OverlapResult.h"
 #include "HAL/IConsoleManager.h"
@@ -25,7 +26,7 @@ DEFINE_LOG_CATEGORY_STATIC(LogWorldCollision, Log, All);
  * if no delegate, you can query trace data using QueryTraceData or QueryOverlapData
  * the data is available only in the next frame after request is made - in other words, if request is made in frame X, you can get the result in frame (X+1)
  *
- * @param InDeleagte	Delegate function to be called - to see example, search FTraceDelegate
+ * @param InDelegate	Delegate function to be called - to see example, search FTraceDelegate
  *						Example can be void MyActor::TraceDone(const FTraceHandle& TraceHandle, FTraceDatum & TraceData)
  *						Before sending to the function, 
  *						
@@ -387,7 +388,7 @@ namespace
 	template <typename DatumType>
 	void ExecuteAsyncTraceIfAvailable(FWorldAsyncTraceState& State, bool bExecuteAll)
 	{
-		auto& DataBuffer = State.GetBufferForCurrentFrame();
+		AsyncTraceData& DataBuffer = State.GetBufferForCurrentFrame();
 
 		FBufferIndexPair Next(GetTraceIndex<DatumType>(DataBuffer));
 
@@ -398,17 +399,17 @@ namespace
 			Next.Block = Next.Block - 1;
 			Next.Index = ASYNC_TRACE_BUFFER_SIZE;
 		}
-		// don't execute if we haven't been explicitly requested to OR there's nothing to runs
+		// don't execute if we haven't been explicitly requested to OR there's nothing to run
 		else if (!bExecuteAll || Next.Index == 0)
 		{
 			return;
 		}
 
-		auto* Datum      = GetTraceContainer<DatumType>(DataBuffer)[Next.Block]->Buffer;
+		DatumType* Datum      = GetTraceContainer<DatumType>(DataBuffer)[Next.Block]->Buffer;
 		const bool bRunAsyncTraceOnWorkerThread = AsyncTraceCVars::IsAsyncTraceOnWorkerThreads();
-		if (bRunAsyncTraceOnWorkerThread)
+		if (bRunAsyncTraceOnWorkerThread && !AutoRTFM::IsTransactional())
 		{
-			DataBuffer.AsyncTraceCompletionEvent.Emplace(TGraphTask<FAsyncTraceTask>::CreateTask(NULL, ENamedThreads::GameThread).ConstructAndDispatchWhenReady(Datum, Next.Index));
+			DataBuffer.AsyncTraceCompletionEvent.Emplace(TGraphTask<FAsyncTraceTask>::CreateTask(nullptr, ENamedThreads::GameThread).ConstructAndDispatchWhenReady(Datum, Next.Index));
 		}
 		else
 		{
@@ -448,7 +449,7 @@ namespace
 	}
 
 	template <typename DatumType>
-	FTraceHandle StartNewTrace(FWorldAsyncTraceState& State, const DatumType& Val)
+	FTraceHandle StartNewTrace(UWorld* World, FWorldAsyncTraceState& State, const DatumType& Val)
 	{
 		// Using async traces outside of the game thread can cause memory corruption
 		check(IsInGameThread());
@@ -459,9 +460,9 @@ namespace
 		// Check we're allowed to do an async call here
 		check(DataBuffer.bAsyncAllowed);
 
-		// Handle transactional async queries separately. These need to be thrown into a separate queue 
+		// Handle transactional async queries separately if we have a query handler. These need to be thrown into a separate queue 
 		// that is processed outside of a transaction (we cannot process in "batches" like normal async).
-		if (AutoRTFM::IsClosed())
+		if (AutoRTFM::IsClosed() && World->PhysicsQueryHandler != nullptr)
 		{
 			FTransactionalAsyncTraceBuffer<DatumType>& TransactionalData = GetTransactionalData<DatumType>(DataBuffer);
 			TransactionalData.Data.Add(MakeUnique<TTraceThreadData<DatumType>>());
@@ -524,12 +525,12 @@ FWorldAsyncTraceState::FWorldAsyncTraceState()
 
 FTraceHandle UWorld::AsyncLineTraceByChannel(EAsyncTraceType InTraceType, const FVector& Start, const FVector& End, ECollisionChannel TraceChannel, const FCollisionQueryParams& Params /* = FCollisionQueryParams::DefaultQueryParam */, const FCollisionResponseParams& ResponseParam /* = FCollisionResponseParams::DefaultResponseParam */, const FTraceDelegate* InDelegate/* = nullptr */, uint32 UserData /* = 0 */)
 {
-	return StartNewTrace(AsyncTraceState, FTraceDatum(this, FCollisionShape::LineShape, Params, ResponseParam, FCollisionObjectQueryParams::DefaultObjectQueryParam, TraceChannel, UserData, InTraceType, Start, End, FQuat::Identity, InDelegate, AsyncTraceState.CurrentFrame));
+	return StartNewTrace(this, AsyncTraceState, FTraceDatum(this, FCollisionShape::LineShape, Params, ResponseParam, FCollisionObjectQueryParams::DefaultObjectQueryParam, TraceChannel, UserData, InTraceType, Start, End, FQuat::Identity, InDelegate, AsyncTraceState.CurrentFrame));
 }
 
 FTraceHandle UWorld::AsyncLineTraceByObjectType(EAsyncTraceType InTraceType, const FVector& Start, const FVector& End, const FCollisionObjectQueryParams& ObjectQueryParams, const FCollisionQueryParams& Params /* = FCollisionQueryParams::DefaultQueryParam */, const FTraceDelegate* InDelegate/* = nullptr */, uint32 UserData /* = 0 */)
 {
-	return StartNewTrace(AsyncTraceState, FTraceDatum(this, FCollisionShape::LineShape, Params, FCollisionResponseParams::DefaultResponseParam, ObjectQueryParams, DefaultCollisionChannel, UserData, InTraceType, Start, End, FQuat::Identity, InDelegate, AsyncTraceState.CurrentFrame));
+	return StartNewTrace(this, AsyncTraceState, FTraceDatum(this, FCollisionShape::LineShape, Params, FCollisionResponseParams::DefaultResponseParam, ObjectQueryParams, DefaultCollisionChannel, UserData, InTraceType, Start, End, FQuat::Identity, InDelegate, AsyncTraceState.CurrentFrame));
 }
 
 FTraceHandle UWorld::AsyncLineTraceByProfile(EAsyncTraceType InTraceType, const FVector& Start, const FVector& End, FName ProfileName, const FCollisionQueryParams& Params /* = FCollisionQueryParams::DefaultQueryParam */, const FTraceDelegate* InDelegate/* = nullptr */, uint32 UserData /* = 0 */)
@@ -538,17 +539,17 @@ FTraceHandle UWorld::AsyncLineTraceByProfile(EAsyncTraceType InTraceType, const
 	FCollisionResponseParams ResponseParam;
 	GetCollisionProfileChannelAndResponseParams(ProfileName, TraceChannel, ResponseParam);
 
-	return StartNewTrace(AsyncTraceState, FTraceDatum(this, FCollisionShape::LineShape, Params, ResponseParam, FCollisionObjectQueryParams::DefaultObjectQueryParam, TraceChannel, UserData, InTraceType, Start, End, FQuat::Identity, InDelegate, AsyncTraceState.CurrentFrame));
+	return StartNewTrace(this, AsyncTraceState, FTraceDatum(this, FCollisionShape::LineShape, Params, ResponseParam, FCollisionObjectQueryParams::DefaultObjectQueryParam, TraceChannel, UserData, InTraceType, Start, End, FQuat::Identity, InDelegate, AsyncTraceState.CurrentFrame));
 }
 
 FTraceHandle UWorld::AsyncSweepByChannel(EAsyncTraceType InTraceType, const FVector& Start, const FVector& End, const FQuat& Rot, ECollisionChannel TraceChannel, const FCollisionShape& CollisionShape, const FCollisionQueryParams& Params /* = FCollisionQueryParams::DefaultQueryParam */, const FCollisionResponseParams& ResponseParam /* = FCollisionResponseParams::DefaultResponseParam */, const FTraceDelegate* InDelegate/* = nullptr */, uint32 UserData /* = 0 */)
 {
-	return StartNewTrace(AsyncTraceState, FTraceDatum(this, CollisionShape, Params, ResponseParam, FCollisionObjectQueryParams::DefaultObjectQueryParam, TraceChannel, UserData, InTraceType, Start, End, Rot, InDelegate, AsyncTraceState.CurrentFrame));
+	return StartNewTrace(this, AsyncTraceState, FTraceDatum(this, CollisionShape, Params, ResponseParam, FCollisionObjectQueryParams::DefaultObjectQueryParam, TraceChannel, UserData, InTraceType, Start, End, Rot, InDelegate, AsyncTraceState.CurrentFrame));
 }
 
 FTraceHandle UWorld::AsyncSweepByObjectType(EAsyncTraceType InTraceType, const FVector& Start, const FVector& End, const FQuat& Rot, const FCollisionObjectQueryParams& ObjectQueryParams, const FCollisionShape& CollisionShape, const FCollisionQueryParams& Params /* = FCollisionQueryParams::DefaultQueryParam */, const FTraceDelegate* InDelegate/* = nullptr */, uint32 UserData /* = 0 */)
 {
-	return StartNewTrace(AsyncTraceState, FTraceDatum(this, CollisionShape, Params, FCollisionResponseParams::DefaultResponseParam, ObjectQueryParams, DefaultCollisionChannel, UserData, InTraceType, Start, End, Rot, InDelegate, AsyncTraceState.CurrentFrame));
+	return StartNewTrace(this, AsyncTraceState, FTraceDatum(this, CollisionShape, Params, FCollisionResponseParams::DefaultResponseParam, ObjectQueryParams, DefaultCollisionChannel, UserData, InTraceType, Start, End, Rot, InDelegate, AsyncTraceState.CurrentFrame));
 }
 
 FTraceHandle UWorld::AsyncSweepByProfile(EAsyncTraceType InTraceType, const FVector& Start, const FVector& End, const FQuat& Rot, FName ProfileName, const FCollisionShape& CollisionShape, const FCollisionQueryParams& Params /* = FCollisionQueryParams::DefaultQueryParam */, const FTraceDelegate* InDelegate/* = nullptr */, uint32 UserData /* = 0 */)
@@ -557,18 +558,18 @@ FTraceHandle UWorld::AsyncSweepByProfile(EAsyncTraceType InTraceType, const FVec
 	FCollisionResponseParams ResponseParam;
 	GetCollisionProfileChannelAndResponseParams(ProfileName, TraceChannel, ResponseParam);
 
-	return StartNewTrace(AsyncTraceState, FTraceDatum(this, CollisionShape, Params, ResponseParam, FCollisionObjectQueryParams::DefaultObjectQueryParam, TraceChannel, UserData, InTraceType, Start, End, Rot, InDelegate, AsyncTraceState.CurrentFrame));
+	return StartNewTrace(this, AsyncTraceState, FTraceDatum(this, CollisionShape, Params, ResponseParam, FCollisionObjectQueryParams::DefaultObjectQueryParam, TraceChannel, UserData, InTraceType, Start, End, Rot, InDelegate, AsyncTraceState.CurrentFrame));
 }
 
 // overlap functions
 FTraceHandle UWorld::AsyncOverlapByChannel(const FVector& Pos, const FQuat& Rot, ECollisionChannel TraceChannel, const FCollisionShape& CollisionShape, const FCollisionQueryParams& Params /* = FCollisionQueryParams::DefaultQueryParam */, const FCollisionResponseParams& ResponseParam /* = FCollisionResponseParams::DefaultResponseParam */, const FOverlapDelegate* InDelegate/* = nullptr */, uint32 UserData /* = 0 */)
 {
-	return StartNewTrace(AsyncTraceState, FOverlapDatum(this, CollisionShape, Params, ResponseParam, FCollisionObjectQueryParams::DefaultObjectQueryParam, TraceChannel, UserData, Pos, Rot, InDelegate, AsyncTraceState.CurrentFrame));
+	return StartNewTrace(this, AsyncTraceState, FOverlapDatum(this, CollisionShape, Params, ResponseParam, FCollisionObjectQueryParams::DefaultObjectQueryParam, TraceChannel, UserData, Pos, Rot, InDelegate, AsyncTraceState.CurrentFrame));
 }
 
 FTraceHandle UWorld::AsyncOverlapByObjectType(const FVector& Pos, const FQuat& Rot, const FCollisionObjectQueryParams& ObjectQueryParams, const FCollisionShape& CollisionShape, const FCollisionQueryParams& Params /* = FCollisionQueryParams::DefaultQueryParam */, const FOverlapDelegate* InDelegate/* = nullptr */, uint32 UserData /* = 0 */)
 {
-	return StartNewTrace(AsyncTraceState, FOverlapDatum(this, CollisionShape, Params, FCollisionResponseParams::DefaultResponseParam, ObjectQueryParams, DefaultCollisionChannel, UserData, Pos, Rot, InDelegate, AsyncTraceState.CurrentFrame));
+	return StartNewTrace(this, AsyncTraceState, FOverlapDatum(this, CollisionShape, Params, FCollisionResponseParams::DefaultResponseParam, ObjectQueryParams, DefaultCollisionChannel, UserData, Pos, Rot, InDelegate, AsyncTraceState.CurrentFrame));
 }
 
 FTraceHandle UWorld::AsyncOverlapByProfile(const FVector& Pos, const FQuat& Rot, FName ProfileName, const FCollisionShape& CollisionShape, const FCollisionQueryParams& Params /* = FCollisionQueryParams::DefaultQueryParam */, const FOverlapDelegate* InDelegate /* = nullptr */, uint32 UserData /* = 0 */)
@@ -577,7 +578,7 @@ FTraceHandle UWorld::AsyncOverlapByProfile(const FVector& Pos, const FQuat& Rot,
 	FCollisionResponseParams ResponseParam;
 	GetCollisionProfileChannelAndResponseParams(ProfileName, TraceChannel, ResponseParam);
 
-	return StartNewTrace(AsyncTraceState, FOverlapDatum(this, CollisionShape, Params, ResponseParam, FCollisionObjectQueryParams::DefaultObjectQueryParam, TraceChannel, UserData, Pos, Rot, InDelegate, AsyncTraceState.CurrentFrame));
+	return StartNewTrace(this, AsyncTraceState, FOverlapDatum(this, CollisionShape, Params, ResponseParam, FCollisionObjectQueryParams::DefaultObjectQueryParam, TraceChannel, UserData, Pos, Rot, InDelegate, AsyncTraceState.CurrentFrame));
 }
 
 bool UWorld::IsTraceHandleValid(const FTraceHandle& Handle, bool bOverlapTrace)
