Fix right eye appearing black when launching VR PIE in Vulkan on ES3.1 with MMV enabled.
Reverts a potion of CL 39587028 that led to editor primitive compositing being mistakenly enabled in VR PIE.

diff --git a/Engine/Source/Runtime/Engine/Public/ShowFlags.h b/Engine/Source/Runtime/Engine/Public/ShowFlags.h
--- a/Engine/Source/Runtime/Engine/Public/ShowFlags.h
+++ b/Engine/Source/Runtime/Engine/Public/ShowFlags.h
@@ -409,6 +409,7 @@ private:
 		SetVisualizeVirtualShadowMap(false);
 		SetVectorFields(false);
 		SetGBufferHints(false);
+		SetCompositeEditorPrimitives(InitMode == ESFIM_Editor || InitMode == ESFIM_VREditing);
 #if UE_ENABLE_DEBUG_DRAWING
 		SetCompositeDebugPrimitives(true);
 #else
