[shaders] add an extra line to compile state debug info to indicate resource's platform and whether the resource state indicates it's currently compiling shaders

diff --git a/Engine/Source/Runtime/Engine/Private/Materials/Material.cpp b/Engine/Source/Runtime/Engine/Private/Materials/Material.cpp
--- a/Engine/Source/Runtime/Engine/Private/Materials/Material.cpp
+++ b/Engine/Source/Runtime/Engine/Private/Materials/Material.cpp
@@ -4701,6 +4701,7 @@ void UMaterial::AppendCompileStateDebugInfo(FStringBuilderBase& OutDebugInfo) co
 		{
 			for (const FMaterialResourceForCooking& Resource : Iter.Value())
 			{
+				OutDebugInfo << "Resource for platform " << LexToString(Resource.Platform) << ", IsCachingShaders=" << (Resource.Resource->IsCachingShaders() ? "true" : "false") << "\n";
 				Resource.Resource->AppendCompileStateDebugInfo(OutDebugInfo);
 			}
 		}
diff --git a/Engine/Source/Runtime/Engine/Private/Materials/MaterialInstance.cpp b/Engine/Source/Runtime/Engine/Private/Materials/MaterialInstance.cpp
--- a/Engine/Source/Runtime/Engine/Private/Materials/MaterialInstance.cpp
+++ b/Engine/Source/Runtime/Engine/Private/Materials/MaterialInstance.cpp
@@ -3065,6 +3065,7 @@ void UMaterialInstance::AppendCompileStateDebugInfo(FStringBuilderBase& OutDebug
 			{
 				for (const FMaterialResourceForCooking& CookResource : Iter.Value())
 				{
+					OutDebugInfo << "Resource for platform " << LexToString(CookResource.Platform) << ", IsCachingShaders=" << (CookResource.Resource->IsCachingShaders() ? "true" : "false") << "\n";
 					bAnyResources = true;
 					CookResource.Resource->AppendCompileStateDebugInfo(OutDebugInfo);
 				}
