[NetPhysics] Fix uninitialized pointer access and invalid pointer access.

diff --git a/Engine/Source/Runtime/Engine/Private/PhysicsEngine/NetworkPhysicsComponent.cpp b/Engine/Source/Runtime/Engine/Private/PhysicsEngine/NetworkPhysicsComponent.cpp
--- a/Engine/Source/Runtime/Engine/Private/PhysicsEngine/NetworkPhysicsComponent.cpp
+++ b/Engine/Source/Runtime/Engine/Private/PhysicsEngine/NetworkPhysicsComponent.cpp
@@ -555,6 +555,7 @@ void UNetworkPhysicsComponent::UninitializeComponent()
 		if (FAsyncNetworkPhysicsComponentInput* AsyncInput = NetworkPhysicsComponent_Internal->GetProducerInputData_External())
 		{
 			AsyncInput->ActorComponent = nullptr;
+			AsyncInput->PhysicsObject = nullptr;
 		}
 	}
 
@@ -1475,6 +1476,7 @@ FAsyncNetworkPhysicsComponent::FAsyncNetworkPhysicsComponent() : TSimCallbackObj
 	, NetworkPhysicsTickOffset(0)
 	, PhysicsReplicationMode(EPhysicsReplicationMode::Default)
 	, SettingsComponent(nullptr)
+	, ActorComponent(nullptr)
 	, PhysicsObject(nullptr)
 	, bCompareStateToTriggerRewind(false)
 	, bCompareStateToTriggerRewindIncludeSimProxies(false)
diff --git a/Engine/Source/Runtime/Engine/Public/Physics/NetworkPhysicsComponent.h b/Engine/Source/Runtime/Engine/Public/Physics/NetworkPhysicsComponent.h
--- a/Engine/Source/Runtime/Engine/Public/Physics/NetworkPhysicsComponent.h
+++ b/Engine/Source/Runtime/Engine/Public/Physics/NetworkPhysicsComponent.h
@@ -1189,11 +1189,11 @@ private:
 
 public:
 	// Actor component that will be used to fill the histories
-	TWeakObjectPtr<UActorComponent> ActorComponent;
+	TWeakObjectPtr<UActorComponent> ActorComponent = nullptr;
 
 private:
 	// Root components physics object
-	Chaos::FConstPhysicsObjectHandle PhysicsObject;
+	Chaos::FConstPhysicsObjectHandle PhysicsObject = nullptr;
 
 	// Locally relayed inputs makes this component act as if it's a locally controlled pawn.
 	bool bIsRelayingLocalInputs = false;
