Fix for not properly clearing notifies queued due to server moves in the CharacterMovementComponent, this caused that later on during parellel anim evaluation the already triggered notifies would be dispatched again during the same tick.

Additionally, improved how we clear notifies when OnlyTickMontagesWhenNotRendered, to keep track of them in ActiveAnimNotifiesSinceLastTick. This will support users who extend the usage of ActiveAnimNotifiesSinceLastTick and use OnlyTickMontagesWhenNotRendered as their tick visibility option.

#rnx

diff --git a/Engine/Source/Runtime/Engine/Classes/Animation/AnimInstance.h b/Engine/Source/Runtime/Engine/Classes/Animation/AnimInstance.h
--- a/Engine/Source/Runtime/Engine/Classes/Animation/AnimInstance.h
+++ b/Engine/Source/Runtime/Engine/Classes/Animation/AnimInstance.h
@@ -1711,6 +1711,9 @@ public:
 	/** Dispatch AnimEvents (AnimNotifies, Montage Events) queued during UpdateAnimation() */
 	ENGINE_API void DispatchQueuedAnimEvents();
 
+	/** Clear AnimEvents (AnimNotifies, Montage Events) queued. Note this is mainly needed for very specific scenarios where our notify queue is populated more than once per tick, avoid calling it if not needed. */
+	ENGINE_API void ClearQueuedAnimEvents(bool bShouldUpdateActiveAnimNotifiesSinceLastTick);
+	
 #if WITH_EDITOR
 	// Helper function to handle reinstancing in editor
 	ENGINE_API virtual void HandleObjectsReinstanced(const TMap<UObject*, UObject*>& OldToNewInstanceMap);
diff --git a/Engine/Source/Runtime/Engine/Private/Animation/AnimInstance.cpp b/Engine/Source/Runtime/Engine/Private/Animation/AnimInstance.cpp
--- a/Engine/Source/Runtime/Engine/Private/Animation/AnimInstance.cpp
+++ b/Engine/Source/Runtime/Engine/Private/Animation/AnimInstance.cpp
@@ -552,7 +552,7 @@ void UAnimInstance::UpdateAnimation(float DeltaSeconds, bool bNeedsValidRootMoti
 				A side effect of this, is that we're stopping all state notifies in the graph, until ticking resumes.
 				This should be fine. But if it is ever a problem, we should keep two versions of them. One for montages and one for the graph.
 			*/
-			NotifyQueue.Reset(GetSkelMeshComponent());
+			ClearQueuedAnimEvents(true);
 
 			/** 
 				Reset UpdateCounter(), this will force Update to occur if Eval is triggered without an Update.
@@ -674,9 +674,8 @@ void UAnimInstance::PreUpdateAnimation(float DeltaSeconds)
 
 	bNeedsUpdate = true;
 
-	GetProxyOnGameThread<FAnimInstanceProxy>().UpdateActiveAnimNotifiesSinceLastTick(NotifyQueue);
-	
-	NotifyQueue.Reset(GetSkelMeshComponent());
+	ClearQueuedAnimEvents(true);
+
 	RootMotionBlendQueue.Reset();
 
 	GetProxyOnGameThread<FAnimInstanceProxy>().PreUpdate(this, DeltaSeconds);
@@ -805,6 +804,24 @@ void UAnimInstance::DispatchQueuedAnimEvents()
 	}
 }
 
+void UAnimInstance::ClearQueuedAnimEvents(bool bShouldUpdateActiveAnimNotifiesSinceLastTick)
+{
+	FAnimInstanceProxy& Proxy = GetProxyOnGameThread<FAnimInstanceProxy>();
+	
+	if (bShouldUpdateActiveAnimNotifiesSinceLastTick)
+	{
+		// Before clearing notify queue we need to ensure we record the notifies that where triggered last tick.
+		Proxy.UpdateActiveAnimNotifiesSinceLastTick(NotifyQueue);
+	}
+	else
+	{
+		// In the case this function is manually called multiple times a tick, we need to ensure we keep track of previous anim notifies properly.
+		Proxy.ActiveAnimNotifiesSinceLastTick.Append(NotifyQueue.AnimNotifies);
+	}
+	
+	NotifyQueue.Reset(GetSkelMeshComponent());
+}
+
 void UAnimInstance::ParallelUpdateAnimation()
 {
 	GetProxyOnAnyThread<FAnimInstanceProxy>().UpdateAnimation();
diff --git a/Engine/Source/Runtime/Engine/Private/Components/CharacterMovementComponent.cpp b/Engine/Source/Runtime/Engine/Private/Components/CharacterMovementComponent.cpp
--- a/Engine/Source/Runtime/Engine/Private/Components/CharacterMovementComponent.cpp
+++ b/Engine/Source/Runtime/Engine/Private/Components/CharacterMovementComponent.cpp
@@ -10464,13 +10464,18 @@ void UCharacterMovementComponent::MoveAutonomous
 		
 		if (CharacterMovementCVars::EnableQueuedAnimEventsOnServer)
 		{
-			if (const UAnimInstance* AnimInstance = OwnerMesh->GetAnimInstance())
+			if (UAnimInstance* AnimInstance = OwnerMesh->GetAnimInstance())
 			{
 				if (OwnerMesh->VisibilityBasedAnimTickOption <= EVisibilityBasedAnimTickOption::AlwaysTickPose && AnimInstance->NeedsUpdate())
 				{
 					// If we are doing a full graph update on the server but its doing a parallel update,
-					// trigger events right away since these are notifies queued from the montage update and we could be receiving multiple ServerMoves per frame.
+					// trigger events right away since these are notifies queued from the montage update, and we could be receiving multiple ServerMoves per frame.
 					OwnerMesh->ConditionallyDispatchQueuedAnimEvents();
+
+					// We need to manually clear the anim notify queue (since normally its only is cleared in PreUpdateAnimation()) otherwise if animation ticks, the notifies queued from the ServerMove would fire twice.
+					AnimInstance->ClearQueuedAnimEvents(false);
+
+					// When animation ticks, we want its queued events to be triggered.
 					OwnerMesh->AllowQueuedAnimEventsNextDispatch();
 				}
 			}
