Add wp.Editor.IsStreamingGenerationLogAllowed cvar to control whether we should dump the streaming generation logs or not, enabled by default.

diff --git a/Engine/Source/Runtime/Engine/Private/WorldPartition/WorldPartitionStreamingGeneration.cpp b/Engine/Source/Runtime/Engine/Private/WorldPartition/WorldPartitionStreamingGeneration.cpp
--- a/Engine/Source/Runtime/Engine/Private/WorldPartition/WorldPartitionStreamingGeneration.cpp
+++ b/Engine/Source/Runtime/Engine/Private/WorldPartition/WorldPartitionStreamingGeneration.cpp
@@ -42,6 +42,13 @@
 
 #define LOCTEXT_NAMESPACE "WorldPartition"
 
+static bool GIsStreamingGenerationLogAllowed = true;
+static FAutoConsoleVariableRef CVarIsStreamingGenerationLogAllowed(
+	TEXT("wp.Editor.IsStreamingGenerationLogAllowed"),
+	GIsStreamingGenerationLogAllowed,
+	TEXT("Can be used to disable the streaming generation logs."),
+	ECVF_Default);
+
 static FAutoConsoleCommand DumpStreamingGenerationLog(
 	TEXT("wp.Editor.DumpStreamingGenerationLog"),
 	TEXT("Dump the streaming generation log."),
@@ -2008,7 +2015,7 @@ bool UWorldPartition::GenerateContainerStreaming(const FGenerateStreamingParams&
 	TUniquePtr<FArchive> LogFileAr;
 	TUniquePtr<FHierarchicalLogArchive> HierarchicalLogAr;
 
-	const bool bIsStreamingGenerationLogAllowed = (!bIsPIE || IWorldPartitionEditorModule::Get().GetEnableStreamingGenerationLogOnPIE()) && (!InParams.OutputLogType.IsSet() || !InParams.OutputLogType.GetValue().IsEmpty());
+	const bool bIsStreamingGenerationLogAllowed = GIsStreamingGenerationLogAllowed && (!bIsPIE || IWorldPartitionEditorModule::Get().GetEnableStreamingGenerationLogOnPIE()) && (!InParams.OutputLogType.IsSet() || !InParams.OutputLogType.GetValue().IsEmpty());
 	const bool bIsStreamingGenerationLogRelevant = IsMainWorldPartition() && (!GIsBuildMachine || GIsAutomationTesting || IsRunningCookCommandlet());
 
 	if (bIsStreamingGenerationLogAllowed && bIsStreamingGenerationLogRelevant)
