Clarify comment about cache pins

diff --git a/Engine/Source/Runtime/Engine/Private/AssetManager.cpp b/Engine/Source/Runtime/Engine/Private/AssetManager.cpp
--- a/Engine/Source/Runtime/Engine/Private/AssetManager.cpp
+++ b/Engine/Source/Runtime/Engine/Private/AssetManager.cpp
@@ -1828,8 +1828,8 @@ TSharedPtr<FStreamableHandle> UAssetManager::ChangeBundleStateForPrimaryAssets(
 			FPlatformMisc::PumpEssentialAppMessages();
 
 			// If a request has ever been made with a download params, require a download for any subsequent requests.
-			// This can help prevent user error if downloads aren't consistently requested and prevent losing cache pins
-			// when a handle is destroyed.
+			// This can help prevent user error if downloads aren't consistently requested and ensures that the correct
+			// content remains pinned in the iostore download cache.
 			const bool bDownloadRequired = 
 				(NameData->PendingState.Handle && NameData->PendingState.Handle->GetDownloadCachePin().IsValid()) ||
 				(NameData->CurrentState.Handle && NameData->CurrentState.Handle->GetDownloadCachePin().IsValid());
