[Shaders] Export c'tor of FMaterialShaderParameters and add internal wrapper for FShaderCompileUtilities::ExecuteShaderCompileJob() to be accessible in plugins.

This is a preparation for a new plugin to analyze shader compile jobs which relies on these functions.

#rnx

diff --git a/Engine/Source/Runtime/Engine/Internal/ShaderCompiler/ShaderCompilerInternal.cpp b/Engine/Source/Runtime/Engine/Internal/ShaderCompiler/ShaderCompilerInternal.cpp
new file mode 100644
--- /dev/null
+++ b/Engine/Source/Runtime/Engine/Internal/ShaderCompiler/ShaderCompilerInternal.cpp
@@ -0,0 +1,14 @@
+// Copyright Epic Games, Inc. All Rights Reserved.
+
+/*=============================================================================
+	ShaderCompilerInternal.cpp: Implements FShaderCompileInternalUtilities.
+=============================================================================*/
+
+#include "ShaderCompilerInternal.h"
+#include "ShaderCompiler.h"
+
+
+void FShaderCompileInternalUtilities::ExecuteShaderCompileJob(FShaderCommonCompileJob& Job)
+{
+	FShaderCompileUtilities::ExecuteShaderCompileJob(Job);
+}
diff --git a/Engine/Source/Runtime/Engine/Internal/ShaderCompiler/ShaderCompilerInternal.h b/Engine/Source/Runtime/Engine/Internal/ShaderCompiler/ShaderCompilerInternal.h
new file mode 100644
--- /dev/null
+++ b/Engine/Source/Runtime/Engine/Internal/ShaderCompiler/ShaderCompilerInternal.h
@@ -0,0 +1,20 @@
+// Copyright Epic Games, Inc. All Rights Reserved.
+
+/*=============================================================================
+	ShaderCompiler.h: Platform independent shader compilation definitions.
+=============================================================================*/
+
+#pragma once
+
+#include "ShaderCompilerCore.h"
+#include "ShaderCompilerJobTypes.h"
+
+
+/** Wrapper for internal shader compiler utilities that can be accessed by plugins for internal use. */
+class FShaderCompileInternalUtilities
+{
+public:
+	/** Execute the specified (single or pipeline) shader compile job. */
+	static ENGINE_API void ExecuteShaderCompileJob(FShaderCommonCompileJob& Job);
+};
+
diff --git a/Engine/Source/Runtime/Engine/Public/MaterialShared.h b/Engine/Source/Runtime/Engine/Public/MaterialShared.h
--- a/Engine/Source/Runtime/Engine/Public/MaterialShared.h
+++ b/Engine/Source/Runtime/Engine/Public/MaterialShared.h
@@ -1095,7 +1095,7 @@ private:
 // FMaterialShaderParametersBase split out to ease serializing of raw data
 struct FMaterialShaderParameters : public FMaterialShaderParametersBase
 {
-	FMaterialShaderParameters(const FMaterial* InMaterial = nullptr);
+	ENGINE_API FMaterialShaderParameters(const FMaterial* InMaterial = nullptr);
 
 public:
 #if WITH_EDITOR
