Remove some remaining outdated references to r.SceneCulling

#rnx
#tests citysample,qagame

diff --git a/Engine/Source/Runtime/Engine/Private/InstanceData/InstanceDataManager.cpp b/Engine/Source/Runtime/Engine/Private/InstanceData/InstanceDataManager.cpp
--- a/Engine/Source/Runtime/Engine/Private/InstanceData/InstanceDataManager.cpp
+++ b/Engine/Source/Runtime/Engine/Private/InstanceData/InstanceDataManager.cpp
@@ -921,10 +921,8 @@ SIZE_T FInstanceDataManager::GetAllocatedSize() const
 bool FInstanceDataManager::ShouldUsePrecomputed()
 {
 	static const auto CVarPrecomputed = IConsoleManager::Get().FindTConsoleVariableDataInt(TEXT("r.SceneCulling.Precomputed"));
-	static const auto CVarSceneCull = IConsoleManager::Get().FindTConsoleVariableDataInt(TEXT("r.SceneCulling"));
-
-	return CVarSceneCull && CVarSceneCull->GetValueOnAnyThread() != 0
-		&& CVarPrecomputed && CVarPrecomputed->GetValueOnAnyThread() != 0;
+	
+	return CVarPrecomputed && CVarPrecomputed->GetValueOnAnyThread() != 0;
 }
 
 
diff --git a/Engine/Source/Runtime/Engine/Private/InstancedStaticMesh/ISMInstanceDataManager.cpp b/Engine/Source/Runtime/Engine/Private/InstancedStaticMesh/ISMInstanceDataManager.cpp
--- a/Engine/Source/Runtime/Engine/Private/InstancedStaticMesh/ISMInstanceDataManager.cpp
+++ b/Engine/Source/Runtime/Engine/Private/InstancedStaticMesh/ISMInstanceDataManager.cpp
@@ -979,10 +979,8 @@ void FPrimitiveInstanceDataManager::OnRegister(int32 InNumInstances)
 bool FPrimitiveInstanceDataManager::ShouldUsePrecomputed()
 {
 	static const auto CVarPrecomputed = IConsoleManager::Get().FindTConsoleVariableDataInt(TEXT("r.SceneCulling.Precomputed"));
-	static const auto CVarSceneCull = IConsoleManager::Get().FindTConsoleVariableDataInt(TEXT("r.SceneCulling"));
-
-	return CVarSceneCull && CVarSceneCull->GetValueOnAnyThread() != 0
-		&& CVarPrecomputed && CVarPrecomputed->GetValueOnAnyThread() != 0;
+	
+	return CVarPrecomputed && CVarPrecomputed->GetValueOnAnyThread() != 0;
 }
 
 #if WITH_EDITOR
