[UE] Fix an old regression that causes components attached to Blueprint actors dynamically spawned at runtime using the fast path option to fail to instance subobject templates in a cooked build.

Notes:
- Regression was first introduced at CL [REDACTED_PHONE]).

diff --git a/Engine/Source/Runtime/Engine/Private/ActorConstruction.cpp b/Engine/Source/Runtime/Engine/Private/ActorConstruction.cpp
--- a/Engine/Source/Runtime/Engine/Private/ActorConstruction.cpp
+++ b/Engine/Source/Runtime/Engine/Private/ActorConstruction.cpp
@@ -1168,6 +1168,12 @@ UActorComponent* AActor::CreateComponentFromTemplateData(const FBlueprintCookedC
 		FBlueprintComponentInstanceDataLoader ComponentInstanceDataLoader(TemplateData->GetCachedPropertyData(), TemplateData->GetCachedPropertyList());
 		NewActorComp->Serialize(ComponentInstanceDataLoader);
 
+		// Instance any subobject values initialized from template data.
+		if (NewActorComp->GetClass()->HasAnyClassFlags(CLASS_HasInstancedReference))
+		{
+			NewActorComp->InstanceSubobjectTemplates();
+		}
+
 		// Handle tasks that would normally occur post-duplication w/ SDO.
 		NewActorComp->PostDuplicate(EDuplicateMode::Normal);
 		{
