UI material permutation fix

diff --git a/Engine/Source/Runtime/Engine/Private/LocalVertexFactory.cpp b/Engine/Source/Runtime/Engine/Private/LocalVertexFactory.cpp
--- a/Engine/Source/Runtime/Engine/Private/LocalVertexFactory.cpp
+++ b/Engine/Source/Runtime/Engine/Private/LocalVertexFactory.cpp
@@ -252,10 +252,11 @@ void FLocalVertexFactoryShaderParameters::GetElementShaderBindings(
  */
 bool FLocalVertexFactory::ShouldCompilePermutation(const FVertexFactoryShaderPermutationParameters& Parameters)
 {
-	// Only compile this permutation inside the editor - it's not applicable in games, but occasionally the editor needs it.
+	// Compile local vertex factory permutations of materials in the UI domain to maintain
+	// our historical (but unintentional) support for using them in this domain.
 	if (Parameters.MaterialParameters.MaterialDomain == MD_UI)
 	{
-		return EnumHasAllFlags(Parameters.Flags, EShaderPermutationFlags::HasEditorOnlyData);
+		return true;
 	}
 
 	return true; 
