Workaround for invalid motion vectors on animations with constraints.
- When a constraint is active, the previous motion vectors sent to the renderer are not correct. They appear to be double processed, which results in a UpdatePrevious op. The artifacts are present when skin cache is inactive on the mesh.
- When skin cache is active, it has a condition that will avoid recalculating the previous position buffer if the revision and pointer match. This was working around the incorrect transforms by effectively ignoring them.
- This change restores the behavior that allows the skin cache to match up the previous bone buffer and revision to avoid recalculating the previous buffer.

diff --git a/Engine/Source/Runtime/Engine/Public/GPUSkinVertexFactory.h b/Engine/Source/Runtime/Engine/Public/GPUSkinVertexFactory.h
--- a/Engine/Source/Runtime/Engine/Public/GPUSkinVertexFactory.h
+++ b/Engine/Source/Runtime/Engine/Public/GPUSkinVertexFactory.h
@@ -267,19 +267,9 @@ public:
 
 		void SetRevisionNumbers(uint32 InCurrentRevisionNumber, uint32 InPreviousRevisionNumber)
 		{
-			// Both previous and current revision numbers are being assigned, so both buffers will be filled.
-			if (InPreviousRevisionNumber != INDEX_NONE)
-			{
-				CurrentRevisionNumber = InCurrentRevisionNumber;
-				PreviousRevisionNumber = InPreviousRevisionNumber;
-			}
-			// Only flip the buffers if the current revision number differs from the previous one.
-			else if (CurrentRevisionNumber != InCurrentRevisionNumber)
-			{
-				PreviousRevisionNumber = CurrentRevisionNumber;
-				CurrentRevisionNumber = InCurrentRevisionNumber;
-				CurrentBuffer = 1 - CurrentBuffer;
-			}
+			PreviousRevisionNumber = InPreviousRevisionNumber != INDEX_NONE ? InPreviousRevisionNumber : CurrentRevisionNumber;
+			CurrentRevisionNumber = InCurrentRevisionNumber;
+			CurrentBuffer = 1 - CurrentBuffer;
 		}
 
 		uint32 GetRevisionNumber(bool bPrevious) const
@@ -814,6 +804,8 @@ public:
 
 	void GetOverrideVertexStreams(FVertexInputStreamArray& VertexStreams) const;
 
+	EPixelFormat GetTangentFormat() const { return TangentFormat; }
+
 public:
 	TUniformBufferRef<FGPUSkinPassThroughFactoryLooseParameters> LooseParametersUniformBuffer;
 
@@ -833,6 +825,7 @@ private:
 	uint32 UpdatedFrameNumber = ~0u;
 	EVertexAttributeFlags VertexAttributesRequested;
 	EVertexAttributeFlags VertexAttributesToBind = EVertexAttributeFlags::None;
+	EPixelFormat TangentFormat = PF_Unknown;
 };
 
 ENUM_CLASS_FLAGS(FGPUSkinPassthroughVertexFactory::EVertexAttributeFlags)
\ No newline at end of file
