RBAN - ImmediatePhysics no longer depends directly on BodyInstance or ConstraintInstance.

This allows the immediate physics actors to be created on a worker thread. It is also step 1 in optimizing the RBAN creation flow to avoid creating the temporary instance of the physics asset.

diff --git a/Engine/Plugins/Experimental/PhysicsControl/Source/ControlRigPhysics/Private/RigPhysicsSimulation.cpp b/Engine/Plugins/Experimental/PhysicsControl/Source/ControlRigPhysics/Private/RigPhysicsSimulation.cpp
--- a/Engine/Plugins/Experimental/PhysicsControl/Source/ControlRigPhysics/Private/RigPhysicsSimulation.cpp
+++ b/Engine/Plugins/Experimental/PhysicsControl/Source/ControlRigPhysics/Private/RigPhysicsSimulation.cpp
@@ -18,6 +18,7 @@
 #include "Engine/World.h"
 
 #include "Physics/ImmediatePhysics/ImmediatePhysicsActorHandle.h"
+#include "Physics/ImmediatePhysics/ImmediatePhysicsAdapters.h"
 #include "Physics/ImmediatePhysics/ImmediatePhysicsJointHandle.h"
 #include "Physics/ImmediatePhysics/ImmediatePhysicsSimulation.h"
 #include "Physics/ImmediatePhysics/ImmediatePhysicsChaos/ImmediatePhysicsSimulation_Chaos.h"
@@ -373,9 +374,8 @@ ImmediatePhysics::FActorHandle* FRigPhysicsSimulation::CreateBody(
 	}
 	BodyInstance.InitBody(BodySetup, BodyRelSimSpaceTM, nullptr, nullptr);
 
-	ImmediatePhysics::FActorHandle* ActorHandle = Dynamics ?
-		Simulation->CreateDynamicActor(&BodyInstance, BodyRelSimSpaceTM) :
-		Simulation->CreateKinematicActor(&BodyInstance, BodyRelSimSpaceTM);
+	ImmediatePhysics::EActorType ActorType = Dynamics ? ImmediatePhysics::EActorType::DynamicActor : ImmediatePhysics::EActorType::KinematicActor;
+	ImmediatePhysics::FActorHandle* ActorHandle = Simulation->CreateActor(ImmediatePhysics::MakeActorSetup(ActorType, &BodyInstance, BodyRelSimSpaceTM));
 
 	// Creation is guaranteed to return something
 	check(ActorHandle);
@@ -621,7 +621,7 @@ void FRigPhysicsSimulation::InstantiatePhysicsJoints(const FRigPhysicsSolverComp
 			ConstraintSettings.bUseLinearSolver = SolverSettings.bUseLinearJointSolver;
 
 			ChildBodyRecord.JointHandle = Simulation->CreateJoint(
-				ConstraintSettings, ChildActorHandle, ParentActorHandle);
+				ImmediatePhysics::FJointSetup(ConstraintSettings, ChildActorHandle, ParentActorHandle));
 
 			// Now set up the drive, if requested
 			const FRigPhysicsDrive& Drive = PhysicsJoint.Drive;
diff --git a/Engine/Plugins/Experimental/PhysicsControl/Source/PhysicsControl/Private/AnimNode_RigidBodyWithControl.cpp b/Engine/Plugins/Experimental/PhysicsControl/Source/PhysicsControl/Private/AnimNode_RigidBodyWithControl.cpp
--- a/Engine/Plugins/Experimental/PhysicsControl/Source/PhysicsControl/Private/AnimNode_RigidBodyWithControl.cpp
+++ b/Engine/Plugins/Experimental/PhysicsControl/Source/PhysicsControl/Private/AnimNode_RigidBodyWithControl.cpp
@@ -17,6 +17,7 @@
 #include "GameFramework/PawnMovementComponent.h"
 #include "Physics/PhysicsInterfaceCore.h"
 #include "Physics/ImmediatePhysics/ImmediatePhysicsActorHandle.h"
+#include "Physics/ImmediatePhysics/ImmediatePhysicsAdapters.h"
 #include "Physics/ImmediatePhysics/ImmediatePhysicsJointHandle.h"
 #include "Physics/ImmediatePhysics/ImmediatePhysicsSimulation.h"
 #include "Physics/ImmediatePhysics/ImmediatePhysicsStats.h"
@@ -1081,7 +1082,7 @@ void FAnimNode_RigidBodyWithControl::InitPhysics(const UAnimInstance* InAnimInst
 				bool bSimulated = (BodySetup->PhysicsType == EPhysicsType::PhysType_Simulated);
 				ImmediatePhysics::EActorType ActorType = ImmediatePhysics::EActorType::DynamicActor;
 				ImmediatePhysics::FActorHandle* ActorHandle = Simulation->CreateActor(
-					ActorType, BodyInstance, BodyInstance->GetUnrealWorldTransform());
+					ImmediatePhysics::MakeActorSetup(ActorType, BodyInstance, BodyInstance->GetUnrealWorldTransform()));
 				if (ActorHandle)
 				{
 					ActorHandle->InitWorldTransform(BodyInstance->GetUnrealWorldTransform());
@@ -1163,7 +1164,7 @@ void FAnimNode_RigidBodyWithControl::InitPhysics(const UAnimInstance* InAnimInst
 					if (bMakeKinematicConstraints || (Body1Handle->IsSimulated() || Body2Handle->IsSimulated()))
 					{
 						ImmediatePhysics::FJointHandle* JointHandle =
-							Simulation->CreateJoint(CI, Body1Handle, Body2Handle);
+							Simulation->CreateJoint(ImmediatePhysics::MakeJointSetup(CI, Body1Handle, Body2Handle));
 
 						// Record the joint handle under the child bone (i.e. more leaf-ward bone),
 						// since each bone may have multiple children, but only one parent.
@@ -1572,7 +1573,7 @@ void FAnimNode_RigidBodyWithControl::CollectClothColliderObjects(const USkeletal
 				{
 					FBodyInstance* const BodyInstance = BodyInstances[BodyInstanceIndex];
 
-					ImmediatePhysics::FActorHandle* const ActorHandle = Simulation->CreateActor(ImmediatePhysics::EActorType::KinematicActor, BodyInstance, BodyInstance->GetUnrealWorldTransform());
+					ImmediatePhysics::FActorHandle* const ActorHandle = Simulation->CreateActor(ImmediatePhysics::MakeKinematicActorSetup(BodyInstance, BodyInstance->GetUnrealWorldTransform()));
 					Simulation->AddToCollidingPairs(ActorHandle); // <-allow collision between this actor and all dynamic actors.
 					ClothColliders.Add(FClothCollider(ActorHandle, SourceComponent, BodyInstance->InstanceBoneIndex));
 
@@ -1678,7 +1679,7 @@ void FAnimNode_RigidBodyWithControl::CollectWorldObjects()
 						// Create a kinematic actor. Not using Static as world-static objects may
 						// move in the simulation's frame of reference. Note that TM may be in the wrong space here
 						ImmediatePhysics::FActorHandle* ActorHandle = Simulation->CreateActor(
-							ImmediatePhysics::EActorType::KinematicActor, &OverlapComp->BodyInstance, TM);
+							ImmediatePhysics::MakeKinematicActorSetup(&OverlapComp->BodyInstance, TM));
 						ActorHandle->SetName(FName(OverlapComp->GetName()));
 						Simulation->AddToCollidingPairs(ActorHandle);
 						ComponentsInSim.Add(OverlapComp, FWorldObject(ActorHandle, ComponentsInSimTick));
diff --git a/Engine/Plugins/Experimental/PhysicsControl/Source/PhysicsControl/Private/RigidBodyWithControl.cpp b/Engine/Plugins/Experimental/PhysicsControl/Source/PhysicsControl/Private/RigidBodyWithControl.cpp
--- a/Engine/Plugins/Experimental/PhysicsControl/Source/PhysicsControl/Private/RigidBodyWithControl.cpp
+++ b/Engine/Plugins/Experimental/PhysicsControl/Source/PhysicsControl/Private/RigidBodyWithControl.cpp
@@ -14,6 +14,7 @@
 #include "Components/StaticMeshComponent.h"
 #include "Engine/SkeletalMesh.h"
 #include "Physics/ImmediatePhysics/ImmediatePhysicsActorHandle.h"
+#include "Physics/ImmediatePhysics/ImmediatePhysicsAdapters.h"
 #include "Physics/ImmediatePhysics/ImmediatePhysicsJointHandle.h"
 #include "Physics/ImmediatePhysics/ImmediatePhysicsSimulation.h"
 #include "PhysicsEngine/PhysicsConstraintTemplate.h"
@@ -214,7 +215,7 @@ ImmediatePhysics::FJointHandle* CreatePhysicsJoint(
 		Settings.ConnectorTransforms[ConstraintChildIndex].SetLocation(ChildCoMPositionOffset);
 
 		Settings.Sanitize();
-		JointHandle = Simulation->CreateJoint(Settings, ChildActorHandle, ParentActorHandle);
+		JointHandle = Simulation->CreateJoint(ImmediatePhysics::FJointSetup(Settings, ChildActorHandle, ParentActorHandle));
 	}
 
 	if (JointHandle)
@@ -240,7 +241,7 @@ void FAnimNode_RigidBodyWithControl::CreateWorldSpaceControlRootBody(UPhysicsAss
 {
 	// Note that this never moves - it's just defined as being an actor at the root of the world
 	WorldSpaceControlActorHandle = Simulation->CreateActor(
-		ImmediatePhysics::EActorType::StaticActor, nullptr, FTransform());
+		ImmediatePhysics::MakeStaticActorSetup(nullptr, FTransform()));
 	if (WorldSpaceControlActorHandle)
 	{
 		WorldSpaceControlActorHandle->SetName(FName("WorldSpaceControlHandle"));
@@ -980,7 +981,7 @@ void FAnimNode_RigidBodyWithControl::ApplyCurrentConstraintProfile()
 
 					FPBDJointSettings JointSettings = ConstraintHandle->GetSettings();
 					
-					ImmediatePhysics_Chaos::UpdateJointSettingsFromConstraintProfile(Profile, JointSettings);
+					ImmediatePhysics::UpdateJointSettingsFromConstraintProfile(Profile, JointSettings);
 
 					ConstraintHandle->SetSettings(JointSettings);
 				}
diff --git a/Engine/Plugins/Experimental/PhysicsControl/Source/PhysicsControlEditor/Private/PhysicsControlAssetEditorAnimInstanceProxy.cpp b/Engine/Plugins/Experimental/PhysicsControl/Source/PhysicsControlEditor/Private/PhysicsControlAssetEditorAnimInstanceProxy.cpp
--- a/Engine/Plugins/Experimental/PhysicsControl/Source/PhysicsControlEditor/Private/PhysicsControlAssetEditorAnimInstanceProxy.cpp
+++ b/Engine/Plugins/Experimental/PhysicsControl/Source/PhysicsControlEditor/Private/PhysicsControlAssetEditorAnimInstanceProxy.cpp
@@ -4,8 +4,10 @@
 #include "PhysicsControlAssetEditorAnimInstance.h"
 
 #include "Physics/ImmediatePhysics/ImmediatePhysicsActorHandle.h"
+#include "Physics/ImmediatePhysics/ImmediatePhysicsAdapters.h"
 #include "Physics/ImmediatePhysics/ImmediatePhysicsJointHandle.h"
 #include "Physics/ImmediatePhysics/ImmediatePhysicsSimulation.h"
+#include "Physics/ImmediatePhysics/ImmediatePhysicsAdapters.h"
 
 FPhysicsControlAssetEditorAnimInstanceProxy::FPhysicsControlAssetEditorAnimInstanceProxy()
 	: TargetActor(nullptr)
@@ -132,11 +134,11 @@ void FPhysicsControlAssetEditorAnimInstanceProxy::Grab(FName InBoneName, const F
 	if (TargetActor != nullptr)
 	{
 		FTransform HandleTransform = FTransform(Rotation, Location);
-		HandleActor = Simulation->CreateKinematicActor(nullptr, HandleTransform);
+		HandleActor = Simulation->CreateActor(ImmediatePhysics::MakeKinematicActorSetup(nullptr, HandleTransform));
 		HandleActor->SetWorldTransform(HandleTransform);
 		HandleActor->SetKinematicTarget(HandleTransform);
 
-		HandleJoint = Simulation->CreateJoint(nullptr, TargetActor, HandleActor);
+		HandleJoint = Simulation->CreateJoint(ImmediatePhysics::MakeJointSetup(nullptr, TargetActor, HandleActor));
 	}
 }
 
@@ -180,7 +182,7 @@ void FPhysicsControlAssetEditorAnimInstanceProxy::CreateSimulationFloor(
 	ImmediatePhysics::FSimulation* Simulation = RagdollNode.GetSimulation();
 	if (Simulation != nullptr)
 	{
-		FloorActor = Simulation->CreateKinematicActor(FloorBodyInstance, Transform);
+		FloorActor = Simulation->CreateActor(ImmediatePhysics::MakeKinematicActorSetup(FloorBodyInstance, Transform));
 		if (FloorActor != nullptr)
 		{
 			Simulation->AddToCollidingPairs(FloorActor);
diff --git a/Engine/Source/Editor/PhysicsAssetEditor/Private/PhysicsAssetEditorAnimInstanceProxy.cpp b/Engine/Source/Editor/PhysicsAssetEditor/Private/PhysicsAssetEditorAnimInstanceProxy.cpp
--- a/Engine/Source/Editor/PhysicsAssetEditor/Private/PhysicsAssetEditorAnimInstanceProxy.cpp
+++ b/Engine/Source/Editor/PhysicsAssetEditor/Private/PhysicsAssetEditorAnimInstanceProxy.cpp
@@ -4,6 +4,7 @@
 #include "PhysicsAssetEditorAnimInstance.h"
 
 #include "Physics/ImmediatePhysics/ImmediatePhysicsActorHandle.h"
+#include "Physics/ImmediatePhysics/ImmediatePhysicsAdapters.h"
 #include "Physics/ImmediatePhysics/ImmediatePhysicsJointHandle.h"
 #include "Physics/ImmediatePhysics/ImmediatePhysicsSimulation.h"
 
@@ -134,11 +135,11 @@ void FPhysicsAssetEditorAnimInstanceProxy::Grab(FName InBoneName, const FVector&
 	if (TargetActor != nullptr)
 	{
 		FTransform HandleTransform = FTransform(Rotation, Location);
-		HandleActor = Simulation->CreateKinematicActor(nullptr, HandleTransform);
+		HandleActor = Simulation->CreateActor(ImmediatePhysics::MakeKinematicActorSetup(nullptr, HandleTransform));
 		HandleActor->SetWorldTransform(HandleTransform);
 		HandleActor->SetKinematicTarget(HandleTransform);
 
-		HandleJoint = Simulation->CreateJoint(nullptr, TargetActor, HandleActor);
+		HandleJoint = Simulation->CreateJoint(ImmediatePhysics::MakeJointSetup(nullptr, TargetActor, HandleActor));
 	}
 }
 
@@ -181,7 +182,7 @@ void FPhysicsAssetEditorAnimInstanceProxy::CreateSimulationFloor(FBodyInstance*
 	ImmediatePhysics::FSimulation* Simulation = RagdollNode.GetSimulation();
 	if (Simulation != nullptr)
 	{
-		FloorActor = Simulation->CreateKinematicActor(FloorBodyInstance, Transform);
+		FloorActor = Simulation->CreateActor(ImmediatePhysics::MakeKinematicActorSetup(FloorBodyInstance, Transform));
 		if (FloorActor != nullptr)
 		{
 			Simulation->AddToCollidingPairs(FloorActor);
diff --git a/Engine/Source/Runtime/AnimGraphRuntime/Private/BoneControllers/AnimNode_RigidBody.cpp b/Engine/Source/Runtime/AnimGraphRuntime/Private/BoneControllers/AnimNode_RigidBody.cpp
--- a/Engine/Source/Runtime/AnimGraphRuntime/Private/BoneControllers/AnimNode_RigidBody.cpp
+++ b/Engine/Source/Runtime/AnimGraphRuntime/Private/BoneControllers/AnimNode_RigidBody.cpp
@@ -20,6 +20,8 @@
 #include "HAL/LowLevelMemTracker.h"
 #include "Physics/Experimental/PhysScene_Chaos.h"
 #include "Physics/ImmediatePhysics/ImmediatePhysicsActorHandle.h"
+#include "Physics/ImmediatePhysics/ImmediatePhysicsAdapters.h"
+#include "Physics/ImmediatePhysics/ImmediatePhysicsJointHandle.h"
 #include "Physics/ImmediatePhysics/ImmediatePhysicsSimulation.h"
 #include "Physics/ImmediatePhysics/ImmediatePhysicsStats.h"
 #include "Physics/PhysicsInterfaceCore.h"
@@ -1139,7 +1141,7 @@ void FAnimNode_RigidBody::InitPhysics(const UAnimInstance* InAnimInstance)
 
 				bool bSimulated = (BodySetup->PhysicsType == EPhysicsType::PhysType_Simulated);
 				ImmediatePhysics::EActorType ActorType = bSimulated ?  ImmediatePhysics::EActorType::DynamicActor : ImmediatePhysics::EActorType::KinematicActor;
-				ImmediatePhysics::FActorHandle* NewBodyHandle = PhysicsSimulation->CreateActor(ActorType, BodyInstance, BodyInstance->GetUnrealWorldTransform());
+				ImmediatePhysics::FActorHandle* NewBodyHandle = PhysicsSimulation->CreateActor(ImmediatePhysics::MakeActorSetup(ActorType, BodyInstance, BodyInstance->GetUnrealWorldTransform()));
 				if (NewBodyHandle)
 				{
 					if (bSimulated)
@@ -1211,7 +1213,7 @@ void FAnimNode_RigidBody::InitPhysics(const UAnimInstance* InAnimInstance)
 				{
 					if (Body1Handle->IsSimulated() || Body2Handle->IsSimulated())
 					{
-						PhysicsSimulation->CreateJoint(CI, Body1Handle, Body2Handle);
+						PhysicsSimulation->CreateJoint(ImmediatePhysics::MakeJointSetup(CI, Body1Handle, Body2Handle));
 						if (bForceDisableCollisionBetweenConstraintBodies)
 						{
 							int32 BodyIndex1 = UsePhysicsAsset->FindBodyIndex(CI->ConstraintBone1);
@@ -1631,7 +1633,7 @@ void FAnimNode_RigidBody::CollectClothColliderObjects(const USkeletalMeshCompone
 
 					if (BodyInstance->InstanceBoneIndex != INDEX_NONE) // Invalid Bone index can occur if the physics asset references a bone name that does not exist in the Cloth Collision Source.
 					{
-						ImmediatePhysics::FActorHandle* const ActorHandle = PhysicsSimulation->CreateActor(ImmediatePhysics::EActorType::KinematicActor, BodyInstance, BodyInstance->GetUnrealWorldTransform());
+						ImmediatePhysics::FActorHandle* const ActorHandle = PhysicsSimulation->CreateActor(ImmediatePhysics::MakeKinematicActorSetup(BodyInstance, BodyInstance->GetUnrealWorldTransform()));
 						PhysicsSimulation->AddToCollidingPairs(ActorHandle); // <-allow collision between this actor and all dynamic actors.
 						ClothColliders.Add(FClothCollider(ActorHandle, SourceComponent, BodyInstance->InstanceBoneIndex));
 					}
@@ -1706,7 +1708,7 @@ void FAnimNode_RigidBody::CollectWorldObjects()
 					if (!bIsSelf)
 					{
 						// Create a kinematic actor. Not using Static as world-static objects may move in the simulation's frame of reference
-						ImmediatePhysics::FActorHandle* ActorHandle = PhysicsSimulation->CreateActor(ImmediatePhysics::EActorType::KinematicActor, &OverlapComp->BodyInstance, OverlapComp->GetComponentTransform());
+						ImmediatePhysics::FActorHandle* ActorHandle = PhysicsSimulation->CreateActor(ImmediatePhysics::MakeKinematicActorSetup(&OverlapComp->BodyInstance, OverlapComp->GetComponentTransform()));
 						PhysicsSimulation->AddToCollidingPairs(ActorHandle);
 						ComponentsInSim.Add(OverlapComp, FWorldObject(ActorHandle, ComponentsInSimTick));
 					}
diff --git a/Engine/Source/Runtime/Engine/Classes/PhysicsEngine/BodyInstance.h b/Engine/Source/Runtime/Engine/Classes/PhysicsEngine/BodyInstance.h
--- a/Engine/Source/Runtime/Engine/Classes/PhysicsEngine/BodyInstance.h
+++ b/Engine/Source/Runtime/Engine/Classes/PhysicsEngine/BodyInstance.h
@@ -918,7 +918,7 @@ public:
 	ENGINE_API FVector GetBodyInertiaTensor() const;
 
 	/** Whether inertia conditioning is enabled. @see bInertiaConditioning */
-	bool IsInertiaConditioningEnabled() { return bInertiaConditioning; }
+	bool IsInertiaConditioningEnabled() const { return bInertiaConditioning; }
 
 	/** Enable or disable inertia conditionin.  @see bInertiaConditioning */
 	ENGINE_API void SetInertiaConditioningEnabled(bool bEnabled);
diff --git a/Engine/Source/Runtime/Engine/Private/PhysicsEngine/ImmediatePhysics/ImmediatePhysicsAdapters.cpp b/Engine/Source/Runtime/Engine/Private/PhysicsEngine/ImmediatePhysics/ImmediatePhysicsAdapters.cpp
new file mode 100644
--- /dev/null
+++ b/Engine/Source/Runtime/Engine/Private/PhysicsEngine/ImmediatePhysics/ImmediatePhysicsAdapters.cpp
@@ -0,0 +1,512 @@
+// Copyright Epic Games, Inc. All Rights Reserved.
+#include "Physics/ImmediatePhysics/ImmediatePhysicsAdapters.h"
+
+#include "BodySetupEnums.h"
+#include "Chaos/ChaosConstraintSettings.h"
+#include "Chaos/MassProperties.h"
+#include "Chaos/PBDRigidsSOAs.h"
+#include "Chaos/Particle/ParticleUtilities.h"
+#include "Chaos/TriangleMeshImplicitObject.h"
+#include "PhysicalMaterials/PhysicalMaterial.h"
+#include "Physics/Experimental/ChaosInterfaceUtils.h"
+#include "Physics/ImmediatePhysics/ImmediatePhysicsActorHandle.h"
+#include "Physics/ImmediatePhysics/ImmediatePhysicsJointHandle.h"
+#include "Physics/PhysicsInterfaceTypes.h"
+#include "PhysicsEngine/BodyInstance.h"
+#include "PhysicsEngine/BodySetup.h"
+#include "PhysicsEngine/BodyUtils.h"
+#include "PhysicsEngine/ConstraintInstance.h"
+#include "PhysicsProxy/SingleParticlePhysicsProxy.h"
+
+namespace ImmediatePhysics
+{
+	bool CreateDefaultGeometry(const FVector& Scale, Chaos::FReal& OutMass, Chaos::FVec3& OutInertia, Chaos::FRigidTransform3& OutCoMTransform, Chaos::FImplicitObjectPtr& OutGeom, TArray<TUniquePtr<Chaos::FPerShapeData>>& OutShapes)
+	{
+		using namespace Chaos;
+
+		const FReal Mass = 1.0f;
+		const FReal Radius = 1.0f * Scale.GetMax();
+
+		Chaos::FImplicitObjectPtr ImplicitSphere = MakeImplicitObjectPtr<Chaos::FImplicitSphere3>(FVec3(0), Radius);
+		TUniquePtr<FPerShapeData> NewShape = Chaos::FShapeInstance::Make(OutShapes.Num(), ImplicitSphere);
+		NewShape->UpdateShapeBounds(FTransform::Identity);
+		NewShape->SetUserData(nullptr);
+		NewShape->SetQueryEnabled(false);
+		NewShape->SetSimEnabled(false);
+
+		OutMass = Mass;
+		OutInertia = Chaos::FSphere::GetInertiaTensor(Mass, Radius).GetDiagonal();
+		OutCoMTransform = FTransform::Identity;
+		OutShapes.Emplace(MoveTemp(NewShape));
+		OutGeom = MoveTemp(ImplicitSphere);
+
+		return true;
+	}
+
+	// Clone wrapper geometry if it is wrapped in Instanced or Transformed, without cloning the inner shape. This is used to instance
+	// world geometry like heightfields and trimeshes into the immediate physics simulation.
+	Chaos::FImplicitObjectPtr CloneWrappedGeometry(const Chaos::FImplicitObject* Geom, TArray<TUniquePtr<Chaos::FPerShapeData>>& OutShapes)
+	{
+		using namespace Chaos;
+
+		EImplicitObjectType GeomType = GetInnerType(Geom->GetCollisionType());
+		bool bIsInstanced = IsInstanced(Geom->GetCollisionType());
+		bool bIsScaled = IsScaled(Geom->GetCollisionType());
+
+		// Transformed HeightField
+		if (GeomType == ImplicitObjectType::Transformed)
+		{
+			const TImplicitObjectTransformed<FReal, 3>* SrcTransformed = Geom->template GetObject<const TImplicitObjectTransformed<FReal, 3>>();
+			if ((SrcTransformed != nullptr) && (SrcTransformed->GetTransformedObject()->GetType() == ImplicitObjectType::HeightField))
+			{
+				FImplicitObject* InnerGeom = const_cast<FImplicitObject*>(SrcTransformed->GetTransformedObject());
+				return MakeImplicitObjectPtr<Chaos::TImplicitObjectTransformed<FReal, 3, false>>(InnerGeom, SrcTransformed->GetTransform());
+			}
+		}
+
+		// Instanced Trimesh
+		if (bIsInstanced && (GeomType == ImplicitObjectType::TriangleMesh))
+		{
+			const TImplicitObjectInstanced<FTriangleMeshImplicitObject>* SrcInstanced = Geom->template GetObject<const TImplicitObjectInstanced<FTriangleMeshImplicitObject>>();
+			if (SrcInstanced != nullptr)
+			{
+				const TImplicitObjectInstanced<FTriangleMeshImplicitObject>::ObjectType InnerGeom = SrcInstanced->Object();
+				return  MakeImplicitObjectPtr<TImplicitObjectInstanced<FTriangleMeshImplicitObject>>(InnerGeom);
+			}
+		}
+
+		return nullptr;
+	}
+
+	// Intended for use with Tri Mesh and Heightfields when cloning world simulation objects into the immediate scene
+	bool CloneWrappedGeometry(const FBodyInstance* BodyInstance, EActorType ActorType, const FVector& Scale, Chaos::FReal& OutMass, Chaos::FVec3& OutInertia, Chaos::FRigidTransform3& OutCoMTransform, Chaos::FImplicitObjectPtr& OutGeom, TArray<TUniquePtr<Chaos::FPerShapeData>>& OutShapes)
+	{
+		// We should only get non-simulated objects through this path, but you never know...
+		if ((BodyInstance != nullptr) && !BodyInstance->bSimulatePhysics && BodyInstance->GetPhysicsActor())
+		{
+			OutMass = 0.0f;
+			OutInertia = FVector::ZeroVector;
+			OutCoMTransform = FTransform::Identity;
+			OutGeom = CloneWrappedGeometry(BodyInstance->GetPhysicsActor()->GetGameThreadAPI().GetGeometry(), OutShapes);
+			if (OutGeom != nullptr)
+			{
+				return true;
+			}
+		}
+
+		return CreateDefaultGeometry(Scale, OutMass, OutInertia, OutCoMTransform, OutGeom, OutShapes);
+	}
+
+	bool CreateGeometry(const FBodyInstance* BodyInstance, EActorType ActorType, const FVector& Scale, Chaos::FReal& OutMass, Chaos::FVec3& OutInertia, Chaos::FRigidTransform3& OutCoMTransform, Chaos::FImplicitObjectPtr& OutGeom, TArray<TUniquePtr<Chaos::FPerShapeData>>& OutShapes)
+	{
+		using namespace Chaos;
+
+		OutMass = 0.0f;
+		OutInertia = FVector::ZeroVector;
+		OutCoMTransform = FTransform::Identity;
+
+		// If there's no BodySetup, we may be cloning an in-world object and probably have a TriMesh or HeightField so try to just copy references
+		// @todo(ccaulfield): make this cleaner - we should have a separate path for this
+		if ((BodyInstance == nullptr) || (BodyInstance->BodySetup == nullptr) || (BodyInstance->BodySetup->CollisionTraceFlag == ECollisionTraceFlag::CTF_UseComplexAsSimple))
+		{
+			return CloneWrappedGeometry(BodyInstance, ActorType, Scale, OutMass, OutInertia, OutCoMTransform, OutGeom, OutShapes);
+		}
+
+		UBodySetup* BodySetup = BodyInstance->GetBodySetup();
+
+		// Set the filter to collide with everything (we use a broad phase that only contains particle pairs that are explicitly set to collide)
+		FBodyCollisionData BodyCollisionData;
+		// @todo(chaos): we need an API for setting up filters
+		BodyCollisionData.CollisionFilterData.SimFilter.Word1 = 0xFFFF;
+		BodyCollisionData.CollisionFilterData.SimFilter.Word3 = 0xFFFF;
+		BodyCollisionData.CollisionFlags.bEnableSimCollisionSimple = true;
+
+		BodyInstance->BuildBodyCollisionFlags(BodyCollisionData.CollisionFlags, BodyInstance->GetCollisionEnabled(), BodyInstance->BodySetup->GetCollisionTraceFlag() == CTF_UseComplexAsSimple);
+
+		FGeometryAddParams AddParams;
+		AddParams.bDoubleSided = BodySetup->bDoubleSidedGeometry;
+		AddParams.CollisionData = BodyCollisionData;
+		AddParams.CollisionTraceType = BodySetup->GetCollisionTraceFlag();
+		AddParams.Scale = Scale;
+		//AddParams.SimpleMaterial = SimpleMaterial;
+		//AddParams.ComplexMaterials = TArrayView<UPhysicalMaterial*>(ComplexMaterials);
+		AddParams.LocalTransform = FTransform::Identity;
+		AddParams.WorldTransform = BodyInstance->GetUnrealWorldTransform();
+		AddParams.Geometry = &BodySetup->AggGeom;
+		AddParams.TriMeshGeometries = MakeArrayView(BodySetup->TriMeshGeometries);
+
+		TArray<Chaos::FImplicitObjectPtr> Geoms;
+		FShapesArray Shapes;
+		ChaosInterface::CreateGeometry(AddParams, Geoms, Shapes);
+
+		if (Geoms.Num() == 0)
+		{
+			return false;
+		}
+
+		if (ActorType == EActorType::DynamicActor)
+		{
+			// Whether each shape contributes to mass
+			// @todo(chaos): it would be easier if ComputeMassProperties knew how to extract this info. Maybe it should be a flag in PerShapeData
+			TArray<bool> bContributesToMass;
+			bContributesToMass.Reserve(Shapes.Num());
+			for (int32 ShapeIndex = 0; ShapeIndex < Shapes.Num(); ++ShapeIndex)
+			{
+				const TUniquePtr<FPerShapeData>& Shape = Shapes[ShapeIndex];
+				const FKShapeElem* ShapeElem = FChaosUserData::Get<FKShapeElem>(Shape->GetUserData());
+				bool bHasMass = ShapeElem && ShapeElem->GetContributeToMass();
+				bContributesToMass.Add(bHasMass);
+			}
+
+			// bInertaScaleIncludeMass = true is to match legacy physics behaviour. This will scale the inertia by the change in mass (density x volumescale) 
+			// as well as the dimension change even though we don't actually change the mass.
+			const bool bInertaScaleIncludeMass = true;
+			FMassProperties MassProperties = BodyUtils::ComputeMassProperties(BodyInstance, Shapes, bContributesToMass, FTransform::Identity, bInertaScaleIncludeMass);
+			OutMass = MassProperties.Mass;
+			OutInertia = MassProperties.InertiaTensor.GetDiagonal();
+			OutCoMTransform = FTransform(MassProperties.RotationOfMass, MassProperties.CenterOfMass);
+		}
+
+		// If we have multiple root shapes, wrap them in a union
+		if (Geoms.Num() == 1)
+		{
+			OutGeom = MoveTemp(Geoms[0]);
+		}
+		else
+		{
+			OutGeom = MakeImplicitObjectPtr<FImplicitObjectUnion>(MoveTemp(Geoms));
+		}
+
+		for (TUniquePtr<FPerShapeData>& Shape : Shapes)
+		{
+			OutShapes.Emplace(MoveTemp(Shape));
+		}
+
+		return true;
+	}
+
+	FActorSetup MakeStaticActorSetup(const FBodyInstance* InBodyInstance, const FTransform& InTransform)
+	{
+		return MakeActorSetup(EActorType::StaticActor, InBodyInstance, InTransform);
+	}
+
+	FActorSetup MakeKinematicActorSetup(const FBodyInstance* InBodyInstance, const FTransform& InTransform)
+	{
+		return MakeActorSetup(EActorType::KinematicActor, InBodyInstance, InTransform);
+	}
+
+	FActorSetup MakeDynamicActorSetup(const FBodyInstance* InBodyInstance, const FTransform& InTransform)
+	{
+		return MakeActorSetup(EActorType::DynamicActor, InBodyInstance, InTransform);
+	}
+
+	FActorSetup MakeActorSetup(EActorType InActorType, const FBodyInstance* InBodyInstance, const FTransform& InTransform)
+	{
+		FActorSetup ActorSetup;
+		ActorSetup.ActorType = InActorType;
+		ActorSetup.Transform = InTransform;
+
+		Chaos::FReal BodyMass = 0.0f;
+		Chaos::FVec3 BodyInertia = Chaos::FVec3::OneVector;
+		Chaos::FRigidTransform3 BodyCoMTransform = Chaos::FRigidTransform3::Identity;
+		Chaos::FImplicitObjectPtr BodyGeom;
+		TArray<TUniquePtr<Chaos::FPerShapeData>> BodyShapes;
+		bool bHasGeometry = CreateGeometry(InBodyInstance, InActorType, InTransform.GetScale3D(), BodyMass, BodyInertia, BodyCoMTransform, BodyGeom, BodyShapes);
+		if (!bHasGeometry)
+		{
+			return MoveTemp(ActorSetup);
+		}
+
+		ActorSetup.Geometry = MoveTemp(BodyGeom);
+		ActorSetup.Shapes = MoveTemp(BodyShapes);
+		ActorSetup.Material = MakeUnique<Chaos::FChaosPhysicsMaterial>();
+
+		ActorSetup.Mass = BodyMass;
+		ActorSetup.Inertia = BodyInertia;
+		ActorSetup.CoMTransform = BodyCoMTransform;
+
+		if (InBodyInstance != nullptr)
+		{
+			ActorSetup.LinearDamping = InBodyInstance->LinearDamping;
+			ActorSetup.AngularDamping = InBodyInstance->AngularDamping;
+
+			ActorSetup.CollisionEnabled = InBodyInstance->GetCollisionEnabled();
+			ActorSetup.GravityGroupIndex = InBodyInstance->GravityGroupIndex;
+			ActorSetup.bInertiaConditioningEnabled = InBodyInstance->IsInertiaConditioningEnabled();
+			ActorSetup.bEnableGravity = InBodyInstance->bEnableGravity;
+			ActorSetup.bUpdateKinematicFromSimulation = InBodyInstance->bUpdateKinematicFromSimulation;
+			ActorSetup.bGyroscopicTorqueEnabled = InBodyInstance->bGyroscopicTorqueEnabled;
+
+			// @todo(ccaulfield): We cannot ask for the physical material on a task thread, because FMICReentranceGuard in UMaterialInstance will assert (in editor). Fix this...
+			// For now we just use material defaults when actors are created on a task thread. This happens when adding world-objects to a RigidBody AnimNode simulation.
+			if (IsInGameThread())
+			{
+				if (UPhysicalMaterial* SimplePhysMat = InBodyInstance->GetSimplePhysicalMaterial())
+				{
+					ActorSetup.Material->Friction = SimplePhysMat->Friction;
+					ActorSetup.Material->Restitution = SimplePhysMat->Restitution;
+				}
+			}
+		}
+
+		return MoveTemp(ActorSetup);
+	}
+
+	static_assert((int32)Chaos::EJointMotionType::Free == (int32)EAngularConstraintMotion::ACM_Free, "Chaos::EJointMotionType and EAngularConstraintMotion mismatch");
+	static_assert((int32)Chaos::EJointMotionType::Limited == (int32)EAngularConstraintMotion::ACM_Limited, "Chaos::EJointMotionType and EAngularConstraintMotion mismatch");
+	static_assert((int32)Chaos::EJointMotionType::Locked == (int32)EAngularConstraintMotion::ACM_Locked, "Chaos::EJointMotionType and EAngularConstraintMotion mismatch");
+
+	// NOTE: Hard dependence on EJointAngularConstraintIndex - the following will break if we change the order (but can be easily fixed). See FJointHandle::FJointHandle
+	static_assert((int32)Chaos::EJointAngularConstraintIndex::Twist == 0, "Angular drive targets have hard dependency on constraint order");
+	static_assert((int32)Chaos::EJointAngularConstraintIndex::Swing1 == 2, "Angular drive targets have hard dependency on constraint order");
+
+	static Chaos::EJointMotionType ConvertToJointMotionType(ELinearConstraintMotion InType)
+	{
+		switch (InType)
+		{
+		case ELinearConstraintMotion::LCM_Free: return Chaos::EJointMotionType::Free;
+		case ELinearConstraintMotion::LCM_Limited: return Chaos::EJointMotionType::Limited;
+		case ELinearConstraintMotion::LCM_Locked: return Chaos::EJointMotionType::Locked;
+		default: ensure(false); return Chaos::EJointMotionType::Locked;
+		}
+	};
+
+	static Chaos::EJointMotionType ConvertToJointMotionType(EAngularConstraintMotion InType)
+	{
+		switch (InType)
+		{
+		case EAngularConstraintMotion::ACM_Free: return Chaos::EJointMotionType::Free;
+		case EAngularConstraintMotion::ACM_Limited: return Chaos::EJointMotionType::Limited;
+		case EAngularConstraintMotion::ACM_Locked: return Chaos::EJointMotionType::Locked;
+		default: ensure(false); return Chaos::EJointMotionType::Locked;
+		}
+	};
+
+	static Chaos::EPlasticityType ConvertToPlasticityType(EConstraintPlasticityType InType)
+	{
+		switch (InType)
+		{
+		case EConstraintPlasticityType::CCPT_Free: return Chaos::EPlasticityType::Free;
+		case EConstraintPlasticityType::CCPT_Shrink: return Chaos::EPlasticityType::Shrink;
+		case EConstraintPlasticityType::CCPT_Grow: return Chaos::EPlasticityType::Grow;
+		default: return Chaos::EPlasticityType::Free;
+		}
+	}
+
+	void UpdateJointSettingsFromLinearConstraint(const FLinearConstraint& Constraint, Chaos::FPBDJointSettings& JointSettings)
+	{
+		using namespace Chaos;
+
+		JointSettings.LinearMotionTypes[0] = ConvertToJointMotionType(Constraint.XMotion);
+		JointSettings.LinearMotionTypes[1] = ConvertToJointMotionType(Constraint.YMotion);
+		JointSettings.LinearMotionTypes[2] = ConvertToJointMotionType(Constraint.ZMotion);
+
+		JointSettings.LinearLimit = Constraint.Limit;
+		JointSettings.bSoftLinearLimitsEnabled = Constraint.bSoftConstraint;
+
+		JointSettings.SoftLinearStiffness =
+			Chaos::ConstraintSettings::SoftLinearStiffnessScale() * Constraint.Stiffness;
+		JointSettings.SoftLinearDamping =
+			Chaos::ConstraintSettings::SoftLinearDampingScale() * Constraint.Damping;
+		JointSettings.LinearRestitution = Constraint.Restitution;
+		JointSettings.LinearContactDistance = Constraint.ContactDistance;
+
+		// UE Disables Soft Limits when the Limit is less than some threshold. This is not necessary
+		// in Chaos but for now we also do it for parity's sake (See FLinearConstraint::UpdateLinearLimit_AssumesLocked).
+		if (JointSettings.LinearLimit < RB_MinSizeToLockDOF)
+		{
+			for (int32 AxisIndex = 0; AxisIndex < 3; ++AxisIndex)
+			{
+				if (JointSettings.LinearMotionTypes[AxisIndex] == EJointMotionType::Limited)
+				{
+					JointSettings.LinearMotionTypes[AxisIndex] = EJointMotionType::Locked;
+				}
+			}
+		}
+	}
+
+	// twist/swing Order is twist, swing1, swing2 and in degrees
+
+	void UpdateJointSettingsFromConeConstraint(const FConeConstraint& Constraint, Chaos::FPBDJointSettings& JointSettings)
+	{
+		JointSettings.AngularMotionTypes[(int32)Chaos::EJointAngularConstraintIndex::Swing1] = ConvertToJointMotionType(Constraint.Swing1Motion);
+		JointSettings.AngularMotionTypes[(int32)Chaos::EJointAngularConstraintIndex::Swing2] = ConvertToJointMotionType(Constraint.Swing2Motion);
+
+		JointSettings.AngularLimits[(int32)Chaos::EJointAngularConstraintIndex::Swing1] = FMath::DegreesToRadians(Constraint.Swing1LimitDegrees);
+		JointSettings.AngularLimits[(int32)Chaos::EJointAngularConstraintIndex::Swing2] = FMath::DegreesToRadians(Constraint.Swing2LimitDegrees);
+
+		JointSettings.bSoftSwingLimitsEnabled = Constraint.bSoftConstraint;
+
+		JointSettings.SoftSwingStiffness =
+			Chaos::ConstraintSettings::SoftAngularStiffnessScale() * Constraint.Stiffness;
+		JointSettings.SoftSwingDamping =
+			Chaos::ConstraintSettings::SoftAngularDampingScale() * Constraint.Damping;
+
+		JointSettings.SwingRestitution = Constraint.Restitution;
+
+		JointSettings.SwingContactDistance = Constraint.ContactDistance;
+	}
+
+	void UpdateJointSettingsFromTwistConstraint(const FTwistConstraint& Constraint, Chaos::FPBDJointSettings& JointSettings)
+	{
+		JointSettings.AngularMotionTypes[(int32)Chaos::EJointAngularConstraintIndex::Twist] = ConvertToJointMotionType(Constraint.TwistMotion);
+		JointSettings.AngularLimits[(int32)Chaos::EJointAngularConstraintIndex::Twist] = FMath::DegreesToRadians(Constraint.TwistLimitDegrees);
+		JointSettings.bSoftTwistLimitsEnabled = Constraint.bSoftConstraint;
+		JointSettings.SoftTwistStiffness =
+			Chaos::ConstraintSettings::SoftAngularStiffnessScale() * Constraint.Stiffness;
+		JointSettings.SoftTwistDamping =
+			Chaos::ConstraintSettings::SoftAngularDampingScale() * Constraint.Damping;
+		JointSettings.TwistRestitution = Constraint.Restitution;
+		JointSettings.TwistContactDistance = Constraint.ContactDistance;
+	}
+
+	// Copies data from a profile into settings
+	void UpdateJointSettingsFromConstraintProfile(const FConstraintProfileProperties& Profile, Chaos::FPBDJointSettings& JointSettings)
+	{
+		using namespace Chaos;
+
+		JointSettings.Stiffness = ConstraintSettings::JointStiffness();
+		JointSettings.LinearProjection = Profile.bEnableProjection ? Profile.ProjectionLinearAlpha : 0.0f;
+		JointSettings.AngularProjection = Profile.bEnableProjection ? Profile.ProjectionAngularAlpha : 0.0f;
+		JointSettings.ShockPropagation = Profile.bEnableShockPropagation ? Profile.ShockPropagationAlpha : 0.0f;
+		JointSettings.TeleportDistance = Profile.bEnableProjection ? Profile.ProjectionLinearTolerance : -1.0f;
+		JointSettings.TeleportAngle = Profile.bEnableProjection ?
+			FMath::DegreesToRadians(Profile.ProjectionAngularTolerance) : -1.0f;
+		JointSettings.ParentInvMassScale = Profile.bParentDominates ? (FReal)0 : (FReal)1;
+
+		JointSettings.bCollisionEnabled = !Profile.bDisableCollision;
+		JointSettings.bProjectionEnabled = Profile.bEnableProjection;
+		JointSettings.bShockPropagationEnabled = Profile.bEnableShockPropagation;
+		JointSettings.bMassConditioningEnabled = Profile.bEnableMassConditioning;
+
+		JointSettings.bUseLinearSolver = Profile.bUseLinearJointSolver;
+
+		UpdateJointSettingsFromLinearConstraint(Profile.LinearLimit, JointSettings);
+		UpdateJointSettingsFromConeConstraint(Profile.ConeLimit, JointSettings);
+		UpdateJointSettingsFromTwistConstraint(Profile.TwistLimit, JointSettings);
+
+		JointSettings.LinearSoftForceMode = (ConstraintSettings::SoftLinearForceMode() == 0) ?
+			EJointForceMode::Acceleration : EJointForceMode::Force;
+		JointSettings.AngularSoftForceMode = (ConstraintSettings::SoftAngularForceMode() == 0) ?
+			EJointForceMode::Acceleration : EJointForceMode::Force;
+
+		JointSettings.LinearDrivePositionTarget = Profile.LinearDrive.PositionTarget;
+		JointSettings.LinearDriveVelocityTarget = Profile.LinearDrive.VelocityTarget;
+		JointSettings.bLinearPositionDriveEnabled[0] = Profile.LinearDrive.XDrive.bEnablePositionDrive;
+		JointSettings.bLinearPositionDriveEnabled[1] = Profile.LinearDrive.YDrive.bEnablePositionDrive;
+		JointSettings.bLinearPositionDriveEnabled[2] = Profile.LinearDrive.ZDrive.bEnablePositionDrive;
+		JointSettings.bLinearVelocityDriveEnabled[0] = Profile.LinearDrive.XDrive.bEnableVelocityDrive;
+		JointSettings.bLinearVelocityDriveEnabled[1] = Profile.LinearDrive.YDrive.bEnableVelocityDrive;
+		JointSettings.bLinearVelocityDriveEnabled[2] = Profile.LinearDrive.ZDrive.bEnableVelocityDrive;
+
+		JointSettings.LinearDriveForceMode = EJointForceMode::Acceleration; // hardcoded!
+		JointSettings.LinearDriveStiffness = ConstraintSettings::LinearDriveStiffnessScale() * FVec3(
+			Profile.LinearDrive.XDrive.Stiffness,
+			Profile.LinearDrive.YDrive.Stiffness,
+			Profile.LinearDrive.ZDrive.Stiffness);
+		JointSettings.LinearDriveDamping = ConstraintSettings::LinearDriveDampingScale() * FVec3(
+			Profile.LinearDrive.XDrive.Damping,
+			Profile.LinearDrive.YDrive.Damping,
+			Profile.LinearDrive.ZDrive.Damping);
+		JointSettings.LinearDriveMaxForce[0] = Profile.LinearDrive.XDrive.MaxForce;
+		JointSettings.LinearDriveMaxForce[1] = Profile.LinearDrive.YDrive.MaxForce;
+		JointSettings.LinearDriveMaxForce[2] = Profile.LinearDrive.ZDrive.MaxForce;
+
+		JointSettings.AngularDrivePositionTarget = FQuat(Profile.AngularDrive.OrientationTarget);
+		JointSettings.AngularDriveVelocityTarget = Profile.AngularDrive.AngularVelocityTarget * UE_TWO_PI; // rev/s to rad/s
+
+		JointSettings.AngularDriveForceMode = EJointForceMode::Acceleration; // hardcoded!
+		if (Profile.AngularDrive.AngularDriveMode == EAngularDriveMode::SLERP)
+		{
+			JointSettings.AngularDriveStiffness = FVec3(
+				ConstraintSettings::AngularDriveStiffnessScale() * Profile.AngularDrive.SlerpDrive.Stiffness);
+			JointSettings.AngularDriveDamping = FVec3(
+				ConstraintSettings::AngularDriveDampingScale() * Profile.AngularDrive.SlerpDrive.Damping);
+			JointSettings.AngularDriveMaxTorque = FVec3(
+				Profile.AngularDrive.SlerpDrive.MaxForce);
+			JointSettings.bAngularSLerpPositionDriveEnabled = Profile.AngularDrive.SlerpDrive.bEnablePositionDrive;
+			JointSettings.bAngularSLerpVelocityDriveEnabled = Profile.AngularDrive.SlerpDrive.bEnableVelocityDrive;
+			JointSettings.bAngularTwistPositionDriveEnabled = false;
+			JointSettings.bAngularTwistVelocityDriveEnabled = false;
+			JointSettings.bAngularSwingPositionDriveEnabled = false;
+			JointSettings.bAngularSwingVelocityDriveEnabled = false;
+		}
+		else
+		{
+			JointSettings.AngularDriveStiffness = ConstraintSettings::AngularDriveStiffnessScale() * FVec3(
+				Profile.AngularDrive.TwistDrive.Stiffness,
+				Profile.AngularDrive.SwingDrive.Stiffness,
+				Profile.AngularDrive.SwingDrive.Stiffness);
+			JointSettings.AngularDriveDamping = ConstraintSettings::AngularDriveDampingScale() * FVec3(
+				Profile.AngularDrive.TwistDrive.Damping,
+				Profile.AngularDrive.SwingDrive.Damping,
+				Profile.AngularDrive.SwingDrive.Damping);
+			JointSettings.AngularDriveMaxTorque[0] = Profile.AngularDrive.TwistDrive.MaxForce;
+			JointSettings.AngularDriveMaxTorque[1] = Profile.AngularDrive.SwingDrive.MaxForce;
+			JointSettings.AngularDriveMaxTorque[2] = Profile.AngularDrive.SwingDrive.MaxForce;
+			JointSettings.bAngularSLerpPositionDriveEnabled = false;
+			JointSettings.bAngularSLerpVelocityDriveEnabled = false;
+			JointSettings.bAngularTwistPositionDriveEnabled = Profile.AngularDrive.TwistDrive.bEnablePositionDrive;
+			JointSettings.bAngularTwistVelocityDriveEnabled = Profile.AngularDrive.TwistDrive.bEnableVelocityDrive;
+			JointSettings.bAngularSwingPositionDriveEnabled = Profile.AngularDrive.SwingDrive.bEnablePositionDrive;
+			JointSettings.bAngularSwingVelocityDriveEnabled = Profile.AngularDrive.SwingDrive.bEnableVelocityDrive;
+		}
+
+		JointSettings.LinearBreakForce = Profile.bLinearBreakable ?
+			ConstraintSettings::LinearBreakScale() * Profile.LinearBreakThreshold : FLT_MAX;
+		JointSettings.LinearPlasticityLimit = Profile.bLinearPlasticity ?
+			FMath::Clamp((float)Profile.LinearPlasticityThreshold, 0.0f, 1.0f) : FLT_MAX;
+		JointSettings.LinearPlasticityType = ConvertToPlasticityType(Profile.LinearPlasticityType);
+		// JointSettings.LinearPlasticityInitialDistanceSquared = ; // What do we do with this?
+
+		JointSettings.AngularBreakTorque = Profile.bAngularBreakable ?
+			ConstraintSettings::AngularBreakScale() * Profile.AngularBreakThreshold : FLT_MAX;
+		JointSettings.AngularPlasticityLimit = Profile.bAngularPlasticity ?
+			FMath::Clamp((float)Profile.AngularPlasticityThreshold, 0.0f, 1.0f) : FLT_MAX;;
+
+		JointSettings.ContactTransferScale = Profile.ContactTransferScale;
+	}
+
+
+	FJointSetup MakeJointSetup(const FConstraintInstance* InConstraintInstance, FActorHandle* InActor1, FActorHandle* InActor2)
+	{
+		using namespace Chaos;
+
+		FJointSetup JointSetup;
+		JointSetup.Actor1 = InActor1;
+		JointSetup.Actor2 = InActor2;
+
+		if (InConstraintInstance != nullptr)
+		{
+			UpdateJointSettingsFromConstraintProfile(InConstraintInstance->ProfileInstance, JointSetup.ConstraintSettings);
+			FReal JointScale = InConstraintInstance->GetLastKnownScale();
+			JointSetup.ConstraintSettings.ConnectorTransforms[0] = InConstraintInstance->GetRefFrame(EConstraintFrame::Frame1);
+			JointSetup.ConstraintSettings.ConnectorTransforms[1] = InConstraintInstance->GetRefFrame(EConstraintFrame::Frame2);
+			JointSetup.ConstraintSettings.ConnectorTransforms[0].ScaleTranslation(JointScale);
+			JointSetup.ConstraintSettings.ConnectorTransforms[1].ScaleTranslation(JointScale);
+
+			JointSetup.ConstraintSettings.Sanitize();
+		}
+		else
+		{
+			// TEMP: all creation with null ConstraintIndex for PhAT handles
+			FTransform Actor1Transform = (InActor1 != nullptr) ? InActor1->GetWorldTransform() : FTransform::Identity;
+			FTransform Actor2Transform = (InActor2 != nullptr) ? InActor2->GetWorldTransform() : FTransform::Identity;
+			JointSetup.ConstraintSettings.ConnectorTransforms[0] = Actor2Transform.GetRelativeTransform(Actor1Transform);
+			JointSetup.ConstraintSettings.ConnectorTransforms[1] = FRigidTransform3();
+			JointSetup.ConstraintSettings.LinearMotionTypes = { EJointMotionType::Limited, EJointMotionType::Limited, EJointMotionType::Limited };
+			JointSetup.ConstraintSettings.LinearLimit = 0.1f;
+			JointSetup.ConstraintSettings.SoftLinearStiffness = 500.0f;
+			JointSetup.ConstraintSettings.SoftLinearDamping = 100.0f;
+			JointSetup.ConstraintSettings.bSoftLinearLimitsEnabled = true;
+			JointSetup.ConstraintSettings.LinearSoftForceMode = EJointForceMode::Acceleration;
+			JointSetup.ConstraintSettings.LinearProjection = 0.0f;
+			JointSetup.ConstraintSettings.AngularProjection = 0.0f;
+			JointSetup.ConstraintSettings.TeleportDistance = -1.0f;
+			JointSetup.ConstraintSettings.TeleportAngle = -1.0f;
+		}
+
+		return JointSetup;
+	}
+}
+
diff --git a/Engine/Source/Runtime/Engine/Private/PhysicsEngine/ImmediatePhysics/ImmediatePhysicsChaos/ImmediatePhysicsActorHandle_Chaos.cpp b/Engine/Source/Runtime/Engine/Private/PhysicsEngine/ImmediatePhysics/ImmediatePhysicsChaos/ImmediatePhysicsActorHandle_Chaos.cpp
--- a/Engine/Source/Runtime/Engine/Private/PhysicsEngine/ImmediatePhysics/ImmediatePhysicsChaos/ImmediatePhysicsActorHandle_Chaos.cpp
+++ b/Engine/Source/Runtime/Engine/Private/PhysicsEngine/ImmediatePhysics/ImmediatePhysicsChaos/ImmediatePhysicsActorHandle_Chaos.cpp
@@ -2,203 +2,71 @@
 
 #include "Physics/ImmediatePhysics/ImmediatePhysicsChaos/ImmediatePhysicsActorHandle_Chaos.h"
 
-#include "BodySetupEnums.h"
-#include "Chaos/MassProperties.h"
 #include "Chaos/PBDRigidsSOAs.h"
 #include "Chaos/Particle/ParticleUtilities.h"
-#include "Chaos/TriangleMeshImplicitObject.h"
-
-#include "Physics/Experimental/ChaosInterfaceUtils.h"
-
 #include "Physics/ImmediatePhysics/ImmediatePhysicsShared/ImmediatePhysicsCore.h"
-#include "PhysicsEngine/BodySetup.h"
-#include "Physics/PhysicsInterfaceTypes.h"
-#include "PhysicsEngine/BodyUtils.h"
 #include "PhysicsProxy/SingleParticlePhysicsProxy.h"
 
-//UE_DISABLE_OPTIMIZATION
-
 namespace ImmediatePhysics_Chaos
 {
 	//
-	// Utils
+	// FActorSetup
 	//
 
-	bool CreateDefaultGeometry(const FVector& Scale, FReal& OutMass, Chaos::FVec3& OutInertia, Chaos::FRigidTransform3& OutCoMTransform, Chaos::FImplicitObjectPtr& OutGeom, TArray<TUniquePtr<Chaos::FPerShapeData>>& OutShapes)
-	{
-		using namespace Chaos;
-
-		const FReal Mass = 1.0f;
-		const FReal Radius = 1.0f * Scale.GetMax();
-
-		Chaos::FImplicitObjectPtr ImplicitSphere = MakeImplicitObjectPtr<Chaos::FImplicitSphere3>(FVec3(0), Radius);
-		TUniquePtr<FPerShapeData> NewShape = Chaos::FShapeInstance::Make(OutShapes.Num(), ImplicitSphere);
-		NewShape->UpdateShapeBounds(FTransform::Identity);
-		NewShape->SetUserData(nullptr);
-		NewShape->SetQueryEnabled(false);
-		NewShape->SetSimEnabled(false);
-
-		OutMass = Mass;
-		OutInertia = Chaos::FSphere::GetInertiaTensor(Mass, Radius).GetDiagonal();
-		OutCoMTransform = FTransform::Identity;
-		OutShapes.Emplace(MoveTemp(NewShape));
-		OutGeom = MoveTemp(ImplicitSphere);
-
-		return true;
-	}
-
-	Chaos::FImplicitObjectPtr CloneGeometry(const Chaos::FImplicitObject* Geom, TArray<TUniquePtr<Chaos::FPerShapeData>>& OutShapes)
-	{
-		using namespace Chaos;
-
-		EImplicitObjectType GeomType = GetInnerType(Geom->GetCollisionType());
-		bool bIsInstanced = IsInstanced(Geom->GetCollisionType());
-		bool bIsScaled = IsScaled(Geom->GetCollisionType());
-
-		// Transformed HeightField
-		if (GeomType == ImplicitObjectType::Transformed)
-		{
-			const TImplicitObjectTransformed<FReal, 3>* SrcTransformed = Geom->template GetObject<const TImplicitObjectTransformed<FReal, 3>>();
-			if ((SrcTransformed != nullptr) && (SrcTransformed->GetTransformedObject()->GetType() == ImplicitObjectType::HeightField))
-			{
-				FImplicitObject* InnerGeom = const_cast<FImplicitObject*>(SrcTransformed->GetTransformedObject());
-				return MakeImplicitObjectPtr<Chaos::TImplicitObjectTransformed<FReal, 3, false>>(InnerGeom, SrcTransformed->GetTransform());
-			}
-		}
-
-		// Instanced Trimesh
-		if (bIsInstanced && (GeomType == ImplicitObjectType::TriangleMesh))
-		{
-			const TImplicitObjectInstanced<FTriangleMeshImplicitObject>* SrcInstanced = Geom->template GetObject<const TImplicitObjectInstanced<FTriangleMeshImplicitObject>>();
-			if (SrcInstanced != nullptr)
-			{
-				const TImplicitObjectInstanced<FTriangleMeshImplicitObject>::ObjectType InnerGeom = SrcInstanced->Object();
-				return  MakeImplicitObjectPtr<TImplicitObjectInstanced<FTriangleMeshImplicitObject>>(InnerGeom);
-			}
-		}
-
-		return nullptr;
-	}
-
-	// Intended for use with Tri Mesh and Heightfields when cloning world simulation objects into the immediate scene
-	bool CloneGeometry(FBodyInstance* BodyInstance, EActorType ActorType, const FVector& Scale, FReal& OutMass, Chaos::FVec3& OutInertia, Chaos::FRigidTransform3& OutCoMTransform, Chaos::FImplicitObjectPtr& OutGeom, TArray<TUniquePtr<Chaos::FPerShapeData>>& OutShapes)
+	FActorSetup::FActorSetup()
+		: Transform(FTransform::Identity)
+		, Geometry{}
+		, Shapes{}
+		, Material{}
+		, Mass(1.0)
+		, Inertia(FVector::OneVector)
+		, CoMTransform(FTransform::Identity)
+		, LinearDamping(0.0)
+		, AngularDamping(0.0)
+		, ActorType(EActorType::KinematicActor)
+		, CollisionEnabled(ECollisionEnabled::PhysicsOnly)
+		, GravityGroupIndex(0)
+		, bInertiaConditioningEnabled(true)
+		, bEnableGravity(true)
+		, bUpdateKinematicFromSimulation(true)
+		, bGyroscopicTorqueEnabled(false)
+	{
+	}
+
+	FActorSetup::FActorSetup(FActorSetup&& R)
+		: Transform(R.Transform)
+		, Geometry(R.Geometry)
+		, Shapes(MoveTemp(R.Shapes))
+		, Material(MoveTemp(R.Material))
+		, Mass(R.Mass)
+		, Inertia(R.Inertia)
+		, CoMTransform(R.CoMTransform)
+		, LinearDamping(R.LinearDamping)
+		, AngularDamping(R.AngularDamping)
+		, ActorType(R.ActorType)
+		, CollisionEnabled(R.CollisionEnabled)
+		, GravityGroupIndex(R.GravityGroupIndex)
+		, bInertiaConditioningEnabled(R.bInertiaConditioningEnabled)
+		, bEnableGravity(R.bEnableGravity)
+		, bUpdateKinematicFromSimulation(R.bUpdateKinematicFromSimulation)
+		, bGyroscopicTorqueEnabled(R.bGyroscopicTorqueEnabled)
+	{
+	}
+
+	FActorSetup::~FActorSetup()
 	{
-		// We should only get non-simulated objects through this path, but you never know...
-		if ((BodyInstance != nullptr) && !BodyInstance->bSimulatePhysics && BodyInstance->GetPhysicsActor())
-		{
-			OutMass = 0.0f;
-			OutInertia = FVector::ZeroVector;
-			OutCoMTransform = FTransform::Identity;
-			OutGeom = CloneGeometry(BodyInstance->GetPhysicsActor()->GetGameThreadAPI().GetGeometry(), OutShapes);
-			if (OutGeom != nullptr)
-			{
-				return true;
-			}
-		}
-
-		return CreateDefaultGeometry(Scale, OutMass, OutInertia, OutCoMTransform, OutGeom, OutShapes);
 	}
 
-	bool CreateGeometry(FBodyInstance* BodyInstance, EActorType ActorType, const FVector& Scale, FReal& OutMass, Chaos::FVec3& OutInertia, Chaos::FRigidTransform3& OutCoMTransform, Chaos::FImplicitObjectPtr& OutGeom, TArray<TUniquePtr<Chaos::FPerShapeData>>& OutShapes)
-	{
-		using namespace Chaos;
-
-		OutMass = 0.0f;
-		OutInertia = FVector::ZeroVector;
-		OutCoMTransform = FTransform::Identity;
-
-		// If there's no BodySetup, we may be cloning an in-world object and probably have a TriMesh or HeightField so try to just copy references
-		// @todo(ccaulfield): make this cleaner - we should have a separate path for this
-		if ((BodyInstance == nullptr) || (BodyInstance->BodySetup == nullptr) || (BodyInstance->BodySetup->CollisionTraceFlag == ECollisionTraceFlag::CTF_UseComplexAsSimple))
-		{
-			return CloneGeometry(BodyInstance, ActorType, Scale, OutMass, OutInertia, OutCoMTransform, OutGeom, OutShapes);
-		}
-
-		UBodySetup* BodySetup = BodyInstance->GetBodySetup();
-
-		// Set the filter to collide with everything (we use a broad phase that only contains particle pairs that are explicitly set to collide)
-		FBodyCollisionData BodyCollisionData;
-		// @todo(chaos): we need an API for setting up filters
-		BodyCollisionData.CollisionFilterData.SimFilter.Word1 = 0xFFFF;
-		BodyCollisionData.CollisionFilterData.SimFilter.Word3 = 0xFFFF;
-		BodyCollisionData.CollisionFlags.bEnableSimCollisionSimple = true;
-
-		//BodyInstance->BuildBodyFilterData(BodyCollisionData.CollisionFilterData);
-		BodyInstance->BuildBodyCollisionFlags(BodyCollisionData.CollisionFlags, BodyInstance->GetCollisionEnabled(), BodyInstance->BodySetup->GetCollisionTraceFlag() == CTF_UseComplexAsSimple);
-
-		FGeometryAddParams AddParams;
-		AddParams.bDoubleSided = BodySetup->bDoubleSidedGeometry;
-		AddParams.CollisionData = BodyCollisionData;
-		AddParams.CollisionTraceType = BodySetup->GetCollisionTraceFlag();
-		AddParams.Scale = Scale;
-		//AddParams.SimpleMaterial = SimpleMaterial;
-		//AddParams.ComplexMaterials = TArrayView<UPhysicalMaterial*>(ComplexMaterials);
-		AddParams.LocalTransform = FTransform::Identity;
-		AddParams.WorldTransform = BodyInstance->GetUnrealWorldTransform();
-		AddParams.Geometry = &BodySetup->AggGeom;
-		AddParams.TriMeshGeometries = MakeArrayView(BodySetup->TriMeshGeometries);
-
-		TArray<Chaos::FImplicitObjectPtr> Geoms;
-		FShapesArray Shapes;
-		ChaosInterface::CreateGeometry(AddParams, Geoms, Shapes);
-
-		if (Geoms.Num() == 0)
-		{
-			return false;
-		}
-
-		if (ActorType == EActorType::DynamicActor)
-		{
-			// Whether each shape contributes to mass
-			// @todo(chaos): it would be easier if ComputeMassProperties knew how to extract this info. Maybe it should be a flag in PerShapeData
-			TArray<bool> bContributesToMass;
-			bContributesToMass.Reserve(Shapes.Num());
-			for (int32 ShapeIndex = 0; ShapeIndex < Shapes.Num(); ++ShapeIndex)
-			{
-				const TUniquePtr<FPerShapeData>& Shape = Shapes[ShapeIndex];
-				const FKShapeElem* ShapeElem = FChaosUserData::Get<FKShapeElem>(Shape->GetUserData());
-				bool bHasMass = ShapeElem && ShapeElem->GetContributeToMass();
-				bContributesToMass.Add(bHasMass);
-			}
-
-			// bInertaScaleIncludeMass = true is to match legacy physics behaviour. This will scale the inertia by the change in mass (density x volumescale) 
-			// as well as the dimension change even though we don't actually change the mass.
-			const bool bInertaScaleIncludeMass = true;
-			FMassProperties MassProperties = BodyUtils::ComputeMassProperties(BodyInstance, Shapes, bContributesToMass, FTransform::Identity, bInertaScaleIncludeMass);
-			OutMass = MassProperties.Mass;
-			OutInertia = MassProperties.InertiaTensor.GetDiagonal();
-			OutCoMTransform = FTransform(MassProperties.RotationOfMass, MassProperties.CenterOfMass);
-		}
-
-		// If we have multiple root shapes, wrap them in a union
-		if (Geoms.Num() == 1)
-		{
-			OutGeom = MoveTemp(Geoms[0]);
-		}
-		else
-		{
-			OutGeom = MakeImplicitObjectPtr<FImplicitObjectUnion>(MoveTemp(Geoms));
-		}
-
-		for (TUniquePtr<FPerShapeData>& Shape : Shapes)
-		{
-			OutShapes.Emplace(MoveTemp(Shape));
-		}
-
-		return true;
-	}
+	//
+	// FActorHandle
+	//
 
 	void FActorHandle::CreateParticleHandle(
-		FBodyInstance*                 BodyInstance,
-		const EActorType               ActorType,
-		const FTransform&              WorldTransform,
-		const FReal                    Mass,
-		const Chaos::FVec3             Inertia,
-		const Chaos::FRigidTransform3& CoMTransform)
+		FActorSetup&& ActorSetup)
 	{
 		using namespace Chaos;
 
-		switch (ActorType)
+		switch (ActorSetup.ActorType)
 		{
 		case EActorType::StaticActor:
 			ParticleHandle = Particles.CreateStaticParticles(1, nullptr, FGeometryParticleParameters())[0];
@@ -213,9 +81,9 @@ namespace ImmediatePhysics_Chaos
 
 		if (ParticleHandle != nullptr)
 		{
-			SetWorldTransform(WorldTransform);
+			SetWorldTransform(ActorSetup.Transform);
 
-			ParticleHandle->SetGeometry(Geometry);
+			ParticleHandle->SetGeometry(ActorSetup.Geometry);
 
 			// Set the collision filter data for the shapes to collide with everything.
 			// Even though we already tried to do this when we created the original shapes array, 
@@ -229,10 +97,10 @@ namespace ImmediatePhysics_Chaos
 				Shape->SetCollisionData(CollisionData);
 			}
 
-			if (Geometry && Geometry->HasBoundingBox())
+			if (ParticleHandle->GetGeometry() && ParticleHandle->GetGeometry()->HasBoundingBox())
 			{
 				ParticleHandle->SetHasBounds(true);
-				ParticleHandle->SetLocalBounds(Geometry->BoundingBox());
+				ParticleHandle->SetLocalBounds(ParticleHandle->GetGeometry()->BoundingBox());
 				ParticleHandle->UpdateWorldSpaceState(FRigidTransform3(ParticleHandle->GetX(), ParticleHandle->GetR()), FVec3(0));
 			}
 
@@ -245,24 +113,21 @@ namespace ImmediatePhysics_Chaos
 			FPBDRigidParticleHandle* Dynamic = ParticleHandle->CastToRigidParticle();
 			if (Dynamic && Dynamic->ObjectState() == EObjectStateType::Dynamic)
 			{
-				FReal MassInv = (Mass > 0.0f) ? 1.0f / Mass : 0.0f;
-				FVec3 InertiaInv = (Mass > 0.0f) ? Inertia.Reciprocal() : FVec3::ZeroVector;
-				Dynamic->SetM(Mass);
+				FReal MassInv = (ActorSetup.Mass > 0.0f) ? 1.0f / ActorSetup.Mass : 0.0f;
+				FVec3 InertiaInv = (ActorSetup.Mass > 0.0f) ? ActorSetup.Inertia.Reciprocal() : FVec3::ZeroVector;
+				Dynamic->SetM(ActorSetup.Mass);
 				Dynamic->SetInvM(MassInv);
-				Dynamic->SetCenterOfMass(CoMTransform.GetTranslation());
-				Dynamic->SetRotationOfMass(CoMTransform.GetRotation());
-				Dynamic->SetI(TVec3<FRealSingle>(Inertia.X, Inertia.Y, Inertia.Z));
+				Dynamic->SetCenterOfMass(ActorSetup.CoMTransform.GetTranslation());
+				Dynamic->SetRotationOfMass(ActorSetup.CoMTransform.GetRotation());
+				Dynamic->SetI(TVec3<FRealSingle>(ActorSetup.Inertia.X, ActorSetup.Inertia.Y, ActorSetup.Inertia.Z));
 				Dynamic->SetInvI(TVec3<FRealSingle>(InertiaInv.X, InertiaInv.Y, InertiaInv.Z));
-				if (BodyInstance != nullptr)
-				{
-					Dynamic->SetInertiaConditioningEnabled(BodyInstance->IsInertiaConditioningEnabled());
-					Dynamic->SetLinearEtherDrag(BodyInstance->LinearDamping);
-					Dynamic->SetAngularEtherDrag(BodyInstance->AngularDamping);
-					Dynamic->SetGravityEnabled(BodyInstance->bEnableGravity);
-					Dynamic->SetGravityGroupIndex(BodyInstance->GravityGroupIndex);
-					Dynamic->SetUpdateKinematicFromSimulation(BodyInstance->bUpdateKinematicFromSimulation);
-					Dynamic->SetGyroscopicTorqueEnabled(BodyInstance->bGyroscopicTorqueEnabled);
-				}
+				Dynamic->SetInertiaConditioningEnabled(ActorSetup.bInertiaConditioningEnabled);
+				Dynamic->SetLinearEtherDrag(ActorSetup.LinearDamping);
+				Dynamic->SetAngularEtherDrag(ActorSetup.AngularDamping);
+				Dynamic->SetGravityEnabled(ActorSetup.bEnableGravity);
+				Dynamic->SetGravityGroupIndex(ActorSetup.GravityGroupIndex);
+				Dynamic->SetUpdateKinematicFromSimulation(ActorSetup.bUpdateKinematicFromSimulation);
+				Dynamic->SetGyroscopicTorqueEnabled(ActorSetup.bGyroscopicTorqueEnabled);
 				Dynamic->SetDisabled(true);
 			}
 		}
@@ -276,26 +141,13 @@ namespace ImmediatePhysics_Chaos
 		Chaos::FPBDRigidsSOAs& InParticles, 
 		Chaos::TArrayCollectionArray<Chaos::FVec3>& InParticlePrevXs, 
 		Chaos::TArrayCollectionArray<Chaos::FRotation3>& InParticlePrevRs, 
-		EActorType ActorType, 
-		FBodyInstance* BodyInstance, 
-		const FTransform& InTransform)
+		FActorSetup&& ActorSetup)
 		: Particles(InParticles)
 		, ParticleHandle(nullptr)
 		, ParticlePrevXs(InParticlePrevXs)
 		, ParticlePrevRs(InParticlePrevRs)
 	{
-		using namespace Chaos;
-
-		const FTransform Transform = FTransform(InTransform.GetRotation(), InTransform.GetTranslation());
-		const FVector Scale = InTransform.GetScale3D();
-
-		FReal Mass = 0;
-		FVec3 Inertia = FVec3::OneVector;
-		FRigidTransform3 CoMTransform = FRigidTransform3::Identity;
-		if (CreateGeometry(BodyInstance, ActorType, Scale, Mass, Inertia, CoMTransform, Geometry, Shapes))
-		{
-			CreateParticleHandle(BodyInstance, ActorType, Transform, Mass, Inertia, CoMTransform);
-		}
+		CreateParticleHandle(MoveTemp(ActorSetup));
 	}
 
 	FActorHandle::~FActorHandle()
@@ -304,7 +156,6 @@ namespace ImmediatePhysics_Chaos
 		{
 			Particles.DestroyParticle(ParticleHandle);
 			ParticleHandle = nullptr;
-			Geometry = nullptr;
 		}
 	}
 
diff --git a/Engine/Source/Runtime/Engine/Private/PhysicsEngine/ImmediatePhysics/ImmediatePhysicsChaos/ImmediatePhysicsJointHandle_Chaos.cpp b/Engine/Source/Runtime/Engine/Private/PhysicsEngine/ImmediatePhysics/ImmediatePhysicsChaos/ImmediatePhysicsJointHandle_Chaos.cpp
--- a/Engine/Source/Runtime/Engine/Private/PhysicsEngine/ImmediatePhysics/ImmediatePhysicsChaos/ImmediatePhysicsJointHandle_Chaos.cpp
+++ b/Engine/Source/Runtime/Engine/Private/PhysicsEngine/ImmediatePhysics/ImmediatePhysicsChaos/ImmediatePhysicsJointHandle_Chaos.cpp
@@ -2,6 +2,7 @@
 
 #include "Physics/ImmediatePhysics/ImmediatePhysicsChaos/ImmediatePhysicsJointHandle_Chaos.h"
 #include "Physics/ImmediatePhysics/ImmediatePhysicsChaos/ImmediatePhysicsActorHandle_Chaos.h"
+#include "Physics/ImmediatePhysics/ImmediatePhysicsAdapters.h"
 
 #include "Chaos/Particle/ParticleUtilities.h"
 #include "Chaos/PBDJointConstraints.h"
@@ -9,271 +10,43 @@
 
 #include "PhysicsEngine/ConstraintInstance.h"
 
-//UE_DISABLE_OPTIMIZATION
-
-static_assert((int32)Chaos::EJointMotionType::Free == (int32)EAngularConstraintMotion::ACM_Free, "Chaos::EJointMotionType and EAngularConstraintMotion mismatch");
-static_assert((int32)Chaos::EJointMotionType::Limited == (int32)EAngularConstraintMotion::ACM_Limited, "Chaos::EJointMotionType and EAngularConstraintMotion mismatch");
-static_assert((int32)Chaos::EJointMotionType::Locked == (int32)EAngularConstraintMotion::ACM_Locked, "Chaos::EJointMotionType and EAngularConstraintMotion mismatch");
-
-// NOTE: Hard dependence on EJointAngularConstraintIndex - the following will break if we change the order (but can be easily fixed). See FJointHandle::FJointHandle
-static_assert((int32)Chaos::EJointAngularConstraintIndex::Twist == 0, "Angular drive targets have hard dependency on constraint order");
-static_assert((int32)Chaos::EJointAngularConstraintIndex::Swing1 == 2, "Angular drive targets have hard dependency on constraint order");
-
 namespace ImmediatePhysics_Chaos
 {
-	static Chaos::EJointMotionType ConvertToJointMotionType(ELinearConstraintMotion InType)
-	{
-		switch (InType)
-		{
-		case ELinearConstraintMotion::LCM_Free: return Chaos::EJointMotionType::Free;
-		case ELinearConstraintMotion::LCM_Limited: return Chaos::EJointMotionType::Limited;
-		case ELinearConstraintMotion::LCM_Locked: return Chaos::EJointMotionType::Locked;
-		default: ensure(false); return Chaos::EJointMotionType::Locked;
-		}
-	};
-
-	static Chaos::EJointMotionType ConvertToJointMotionType(EAngularConstraintMotion InType)
-	{
-		switch (InType)
-		{
-		case EAngularConstraintMotion::ACM_Free: return Chaos::EJointMotionType::Free;
-		case EAngularConstraintMotion::ACM_Limited: return Chaos::EJointMotionType::Limited;
-		case EAngularConstraintMotion::ACM_Locked: return Chaos::EJointMotionType::Locked;
-		default: ensure(false); return Chaos::EJointMotionType::Locked;
-		}
-	};
-
-	static Chaos::EPlasticityType ConvertToPlasticityType(EConstraintPlasticityType InType)
-	{
-		switch (InType)
-		{
-		case EConstraintPlasticityType::CCPT_Free: return Chaos::EPlasticityType::Free;
-		case EConstraintPlasticityType::CCPT_Shrink: return Chaos::EPlasticityType::Shrink;
-		case EConstraintPlasticityType::CCPT_Grow: return Chaos::EPlasticityType::Grow;
-		default: return Chaos::EPlasticityType::Free;
-		}
-	}
-
-	void UpdateJointSettingsFromLinearConstraint(const FLinearConstraint& Constraint, Chaos::FPBDJointSettings& JointSettings)
+	// TODO: Deprecated - remove this and decl in ImmediatePhysicsJointHandle_Chaos.h
+	void UpdateJointSettingsFromConstraintProfile(const FConstraintProfileProperties& Profile, Chaos::FPBDJointSettings& JointSettings)
 	{
-		using namespace Chaos;
-
-		JointSettings.LinearMotionTypes[0] = ConvertToJointMotionType(Constraint.XMotion);
-		JointSettings.LinearMotionTypes[1] = ConvertToJointMotionType(Constraint.YMotion);
-		JointSettings.LinearMotionTypes[2] = ConvertToJointMotionType(Constraint.ZMotion);
-
-		JointSettings.LinearLimit = Constraint.Limit;
-		JointSettings.bSoftLinearLimitsEnabled = Constraint.bSoftConstraint;
-
-		JointSettings.SoftLinearStiffness =
-			Chaos::ConstraintSettings::SoftLinearStiffnessScale() * Constraint.Stiffness;
-		JointSettings.SoftLinearDamping =
-			Chaos::ConstraintSettings::SoftLinearDampingScale() * Constraint.Damping;
-		JointSettings.LinearRestitution = Constraint.Restitution;
-		JointSettings.LinearContactDistance = Constraint.ContactDistance;
-
-		// UE Disables Soft Limits when the Limit is less than some threshold. This is not necessary
-		// in Chaos but for now we also do it for parity's sake (See FLinearConstraint::UpdateLinearLimit_AssumesLocked).
-		if (JointSettings.LinearLimit < RB_MinSizeToLockDOF)
-		{
-			for (int32 AxisIndex = 0; AxisIndex < 3; ++AxisIndex)
-			{
-				if (JointSettings.LinearMotionTypes[AxisIndex] == EJointMotionType::Limited)
-				{
-					JointSettings.LinearMotionTypes[AxisIndex] = EJointMotionType::Locked;
-				}
-			}
-		}
+		return ImmediatePhysics::UpdateJointSettingsFromConstraintProfile(Profile, JointSettings);
 	}
 
-	// twist/swing Order is twist, swing1, swing2 and in degrees
 
-	void UpdateJointSettingsFromConeConstraint(const FConeConstraint& Constraint, Chaos::FPBDJointSettings& JointSettings)
+	FJointHandle::FJointHandle(FChaosConstraintContainer* InConstraints, FConstraintInstance* ConstraintInstance, FActorHandle* Actor1, FActorHandle* Actor2)
+		: ActorHandles({ nullptr, nullptr })
+		, Constraints(nullptr)
 	{
-		JointSettings.AngularMotionTypes[(int32)Chaos::EJointAngularConstraintIndex::Swing1] = ConvertToJointMotionType(Constraint.Swing1Motion);
-		JointSettings.AngularMotionTypes[(int32)Chaos::EJointAngularConstraintIndex::Swing2] = ConvertToJointMotionType(Constraint.Swing2Motion);
-
-		JointSettings.AngularLimits[(int32)Chaos::EJointAngularConstraintIndex::Swing1] = FMath::DegreesToRadians(Constraint.Swing1LimitDegrees);
-		JointSettings.AngularLimits[(int32)Chaos::EJointAngularConstraintIndex::Swing2] = FMath::DegreesToRadians(Constraint.Swing2LimitDegrees);
-
-		JointSettings.bSoftSwingLimitsEnabled = Constraint.bSoftConstraint;
-
-		JointSettings.SoftSwingStiffness =
-			Chaos::ConstraintSettings::SoftAngularStiffnessScale() * Constraint.Stiffness;
-		JointSettings.SoftSwingDamping =
-			Chaos::ConstraintSettings::SoftAngularDampingScale() * Constraint.Damping;
-
-		JointSettings.SwingRestitution = Constraint.Restitution;
-
-		JointSettings.SwingContactDistance = Constraint.ContactDistance;
+		FJointSetup JointSetup = ImmediatePhysics::MakeJointSetup(ConstraintInstance, Actor1, Actor2);
+		CreateJoint(InConstraints, JointSetup.ConstraintSettings, JointSetup.Actor1, JointSetup.Actor2);
 	}
 
-	void UpdateJointSettingsFromTwistConstraint(const FTwistConstraint& Constraint, Chaos::FPBDJointSettings& JointSettings)
+	FJointHandle::FJointHandle(FChaosConstraintContainer* InConstraints, const Chaos::FPBDJointSettings& ConstraintSettings, FActorHandle* Actor1, FActorHandle* Actor2)
+		: ActorHandles({ nullptr, nullptr })
+		, Constraints(nullptr)
 	{
-		JointSettings.AngularMotionTypes[(int32)Chaos::EJointAngularConstraintIndex::Twist] = ConvertToJointMotionType(Constraint.TwistMotion);
-		JointSettings.AngularLimits[(int32)Chaos::EJointAngularConstraintIndex::Twist] = FMath::DegreesToRadians(Constraint.TwistLimitDegrees);
-		JointSettings.bSoftTwistLimitsEnabled = Constraint.bSoftConstraint;
-		JointSettings.SoftTwistStiffness =
-			Chaos::ConstraintSettings::SoftAngularStiffnessScale() * Constraint.Stiffness;
-		JointSettings.SoftTwistDamping =
-			Chaos::ConstraintSettings::SoftAngularDampingScale() * Constraint.Damping;
-		JointSettings.TwistRestitution = Constraint.Restitution;
-		JointSettings.TwistContactDistance = Constraint.ContactDistance;
+		CreateJoint(InConstraints, ConstraintSettings, Actor1, Actor2);
 	}
 
-	// Copies data from a profile into settings
-	void UpdateJointSettingsFromConstraintProfile(const FConstraintProfileProperties& Profile, Chaos::FPBDJointSettings& JointSettings)
+	FJointHandle::FJointHandle(FChaosConstraintContainer* InConstraints, const FJointSetup& JointSetup)
+		: ActorHandles({ nullptr, nullptr })
+		, Constraints(nullptr)
 	{
-		using namespace Chaos;
-
-		JointSettings.Stiffness = ConstraintSettings::JointStiffness();
-		JointSettings.LinearProjection = Profile.bEnableProjection ? Profile.ProjectionLinearAlpha : 0.0f;
-		JointSettings.AngularProjection = Profile.bEnableProjection ? Profile.ProjectionAngularAlpha : 0.0f;
-		JointSettings.ShockPropagation = Profile.bEnableShockPropagation ? Profile.ShockPropagationAlpha : 0.0f;
-		JointSettings.TeleportDistance = Profile.bEnableProjection ? Profile.ProjectionLinearTolerance : -1.0f;
-		JointSettings.TeleportAngle = Profile.bEnableProjection ?
-			FMath::DegreesToRadians(Profile.ProjectionAngularTolerance) : -1.0f;
-		JointSettings.ParentInvMassScale = Profile.bParentDominates ? (FReal)0 : (FReal)1;
-
-		JointSettings.bCollisionEnabled = !Profile.bDisableCollision;
-		JointSettings.bProjectionEnabled = Profile.bEnableProjection;
-		JointSettings.bShockPropagationEnabled = Profile.bEnableShockPropagation;
-		JointSettings.bMassConditioningEnabled = Profile.bEnableMassConditioning;
-
-		JointSettings.bUseLinearSolver = Profile.bUseLinearJointSolver;
-
-		UpdateJointSettingsFromLinearConstraint(Profile.LinearLimit, JointSettings);
-		UpdateJointSettingsFromConeConstraint(Profile.ConeLimit, JointSettings);
-		UpdateJointSettingsFromTwistConstraint(Profile.TwistLimit, JointSettings);
-
-		JointSettings.LinearSoftForceMode = (ConstraintSettings::SoftLinearForceMode() == 0) ?
-			EJointForceMode::Acceleration : EJointForceMode::Force;
-		JointSettings.AngularSoftForceMode = (ConstraintSettings::SoftAngularForceMode() == 0) ?
-			EJointForceMode::Acceleration : EJointForceMode::Force;
-
-		JointSettings.LinearDrivePositionTarget = Profile.LinearDrive.PositionTarget;
-		JointSettings.LinearDriveVelocityTarget = Profile.LinearDrive.VelocityTarget;
-		JointSettings.bLinearPositionDriveEnabled[0] = Profile.LinearDrive.XDrive.bEnablePositionDrive;
-		JointSettings.bLinearPositionDriveEnabled[1] = Profile.LinearDrive.YDrive.bEnablePositionDrive;
-		JointSettings.bLinearPositionDriveEnabled[2] = Profile.LinearDrive.ZDrive.bEnablePositionDrive;
-		JointSettings.bLinearVelocityDriveEnabled[0] = Profile.LinearDrive.XDrive.bEnableVelocityDrive;
-		JointSettings.bLinearVelocityDriveEnabled[1] = Profile.LinearDrive.YDrive.bEnableVelocityDrive;
-		JointSettings.bLinearVelocityDriveEnabled[2] = Profile.LinearDrive.ZDrive.bEnableVelocityDrive;
-
-		JointSettings.LinearDriveForceMode = EJointForceMode::Acceleration; // hardcoded!
-		JointSettings.LinearDriveStiffness = Chaos::ConstraintSettings::LinearDriveStiffnessScale() * FVec3(
-			Profile.LinearDrive.XDrive.Stiffness,
-			Profile.LinearDrive.YDrive.Stiffness,
-			Profile.LinearDrive.ZDrive.Stiffness);
-		JointSettings.LinearDriveDamping = Chaos::ConstraintSettings::LinearDriveDampingScale() * FVec3(
-			Profile.LinearDrive.XDrive.Damping,
-			Profile.LinearDrive.YDrive.Damping,
-			Profile.LinearDrive.ZDrive.Damping);
-		JointSettings.LinearDriveMaxForce[0] = Profile.LinearDrive.XDrive.MaxForce;
-		JointSettings.LinearDriveMaxForce[1] = Profile.LinearDrive.YDrive.MaxForce;
-		JointSettings.LinearDriveMaxForce[2] = Profile.LinearDrive.ZDrive.MaxForce;
-
-		JointSettings.AngularDrivePositionTarget = FQuat(Profile.AngularDrive.OrientationTarget);
-		JointSettings.AngularDriveVelocityTarget = Profile.AngularDrive.AngularVelocityTarget * UE_TWO_PI; // rev/s to rad/s
-
-		JointSettings.AngularDriveForceMode = EJointForceMode::Acceleration; // hardcoded!
-		if (Profile.AngularDrive.AngularDriveMode == EAngularDriveMode::SLERP)
-		{
-			JointSettings.AngularDriveStiffness = FVec3(
-				ConstraintSettings::AngularDriveStiffnessScale() * Profile.AngularDrive.SlerpDrive.Stiffness);
-			JointSettings.AngularDriveDamping = FVec3(
-				ConstraintSettings::AngularDriveDampingScale() * Profile.AngularDrive.SlerpDrive.Damping);
-			JointSettings.AngularDriveMaxTorque = FVec3(
-				Profile.AngularDrive.SlerpDrive.MaxForce);
-			JointSettings.bAngularSLerpPositionDriveEnabled = Profile.AngularDrive.SlerpDrive.bEnablePositionDrive;
-			JointSettings.bAngularSLerpVelocityDriveEnabled = Profile.AngularDrive.SlerpDrive.bEnableVelocityDrive;
-			JointSettings.bAngularTwistPositionDriveEnabled = false;
-			JointSettings.bAngularTwistVelocityDriveEnabled = false;
-			JointSettings.bAngularSwingPositionDriveEnabled = false;
-			JointSettings.bAngularSwingVelocityDriveEnabled = false;
-		}
-		else
-		{
-			JointSettings.AngularDriveStiffness = ConstraintSettings::AngularDriveStiffnessScale() * FVec3(
-				Profile.AngularDrive.TwistDrive.Stiffness,
-				Profile.AngularDrive.SwingDrive.Stiffness,
-				Profile.AngularDrive.SwingDrive.Stiffness);
-			JointSettings.AngularDriveDamping = ConstraintSettings::AngularDriveDampingScale() * FVec3(
-				Profile.AngularDrive.TwistDrive.Damping,
-				Profile.AngularDrive.SwingDrive.Damping,
-				Profile.AngularDrive.SwingDrive.Damping);
-			JointSettings.AngularDriveMaxTorque[0] = Profile.AngularDrive.TwistDrive.MaxForce;
-			JointSettings.AngularDriveMaxTorque[1] = Profile.AngularDrive.SwingDrive.MaxForce;
-			JointSettings.AngularDriveMaxTorque[2] = Profile.AngularDrive.SwingDrive.MaxForce;
-			JointSettings.bAngularSLerpPositionDriveEnabled = false;
-			JointSettings.bAngularSLerpVelocityDriveEnabled = false;
-			JointSettings.bAngularTwistPositionDriveEnabled = Profile.AngularDrive.TwistDrive.bEnablePositionDrive;
-			JointSettings.bAngularTwistVelocityDriveEnabled = Profile.AngularDrive.TwistDrive.bEnableVelocityDrive;
-			JointSettings.bAngularSwingPositionDriveEnabled = Profile.AngularDrive.SwingDrive.bEnablePositionDrive;
-			JointSettings.bAngularSwingVelocityDriveEnabled = Profile.AngularDrive.SwingDrive.bEnableVelocityDrive;
-		}
-
-		JointSettings.LinearBreakForce = Profile.bLinearBreakable ?
-			Chaos::ConstraintSettings::LinearBreakScale() * Profile.LinearBreakThreshold : FLT_MAX;
-		JointSettings.LinearPlasticityLimit = Profile.bLinearPlasticity ?
-			FMath::Clamp((float)Profile.LinearPlasticityThreshold, 0.0f, 1.0f) : FLT_MAX;
-		JointSettings.LinearPlasticityType = ConvertToPlasticityType(Profile.LinearPlasticityType);
-		// JointSettings.LinearPlasticityInitialDistanceSquared = ; // What do we do with this?
-
-		JointSettings.AngularBreakTorque = Profile.bAngularBreakable ?
-			Chaos::ConstraintSettings::AngularBreakScale() * Profile.AngularBreakThreshold : FLT_MAX;
-		JointSettings.AngularPlasticityLimit = Profile.bAngularPlasticity ?
-			FMath::Clamp((float)Profile.AngularPlasticityThreshold, 0.0f, 1.0f) : FLT_MAX;;
-
-		JointSettings.ContactTransferScale = Profile.ContactTransferScale;
+		CreateJoint(InConstraints, JointSetup.ConstraintSettings, JointSetup.Actor1, JointSetup.Actor2);
 	}
 
-	FJointHandle::FJointHandle(FChaosConstraintContainer* InConstraints, FConstraintInstance* ConstraintInstance, FActorHandle* Actor1, FActorHandle* Actor2)
-		: ActorHandles({ Actor1, Actor2 })
-		, Constraints(InConstraints)
+	void FJointHandle::CreateJoint(FChaosConstraintContainer* InConstraints, const Chaos::FPBDJointSettings& ConstraintSettings, FActorHandle* Actor1, FActorHandle* Actor2)
 	{
-		using namespace Chaos;
-
-		FPBDJointSettings ConstraintSettings;
-
-		if (ConstraintInstance != nullptr)
-		{
-			// BodyInstance/PhysX has the constraint locations in actor-space, but we need them in Center-of-Mass space
-			UpdateJointSettingsFromConstraintProfile(ConstraintInstance->ProfileInstance, ConstraintSettings);
-			FReal JointScale = ConstraintInstance->GetLastKnownScale();
-			ConstraintSettings.ConnectorTransforms[0] = FParticleUtilities::ActorLocalToParticleLocal(FGenericParticleHandle(Actor1->GetParticle()), ConstraintInstance->GetRefFrame(EConstraintFrame::Frame1));
-			ConstraintSettings.ConnectorTransforms[1] = FParticleUtilities::ActorLocalToParticleLocal(FGenericParticleHandle(Actor2->GetParticle()), ConstraintInstance->GetRefFrame(EConstraintFrame::Frame2));
-			ConstraintSettings.ConnectorTransforms[0].ScaleTranslation(JointScale);
-			ConstraintSettings.ConnectorTransforms[1].ScaleTranslation(JointScale);
-		}
-		else
-		{
-			// TEMP: all creation with null ConstraintIndex for PhAt handles
-			ConstraintSettings.ConnectorTransforms[0] = Actor2->GetWorldTransform().GetRelativeTransform(Actor1->GetWorldTransform());
-			ConstraintSettings.ConnectorTransforms[1] = FRigidTransform3();
-			ConstraintSettings.LinearMotionTypes = { EJointMotionType::Limited, EJointMotionType::Limited, EJointMotionType::Limited };
-			ConstraintSettings.LinearLimit = 0.1f;
-			ConstraintSettings.SoftLinearStiffness = 500.0f;
-			ConstraintSettings.SoftLinearDamping = 100.0f;
-			ConstraintSettings.bSoftLinearLimitsEnabled = true;
-			ConstraintSettings.LinearSoftForceMode = EJointForceMode::Acceleration;
-			ConstraintSettings.LinearProjection = 0.0f;
-			ConstraintSettings.AngularProjection = 0.0f;
-			ConstraintSettings.TeleportDistance = -1.0f;
-			ConstraintSettings.TeleportAngle = -1.0f;
-		}
+		Constraints = InConstraints;
+		ActorHandles[0] = Actor1;
+		ActorHandles[1] = Actor2;
 
-		ConstraintSettings.Sanitize();
-
-		ConstraintHandle = Constraints->AddConstraint({ Actor1->ParticleHandle, Actor2->ParticleHandle }, ConstraintSettings);
-
-		SetActorInertiaConditioningDirty();
-	}
-
-	FJointHandle::FJointHandle(FChaosConstraintContainer* InConstraints, const FPBDJointSettings& ConstraintSettings, FActorHandle* const Actor1, FActorHandle* const Actor2)
-		: ActorHandles({ Actor1, Actor2 })
-		, Constraints(InConstraints)
-	{
 		ConstraintHandle = Constraints->AddConstraint({ Actor1->ParticleHandle, Actor2->ParticleHandle }, ConstraintSettings);
 
 		SetActorInertiaConditioningDirty();
@@ -308,7 +81,7 @@ namespace ImmediatePhysics_Chaos
 	void FJointHandle::SetSoftLinearSettings(bool bLinearSoft, FReal LinearStiffness, FReal LinearDamping)
 	{
 		using namespace Chaos;
-		FPBDJointSettings JointSettings = ConstraintHandle->GetSettings();
+		Chaos::FPBDJointSettings JointSettings = ConstraintHandle->GetSettings();
 		JointSettings.bSoftLinearLimitsEnabled = bLinearSoft;
 		JointSettings.SoftLinearStiffness = bLinearSoft ? LinearStiffness : 0.0f;
 		JointSettings.SoftLinearDamping = bLinearSoft ? LinearDamping : 0.0f;
diff --git a/Engine/Source/Runtime/Engine/Private/PhysicsEngine/ImmediatePhysics/ImmediatePhysicsChaos/ImmediatePhysicsSimulation_Chaos.cpp b/Engine/Source/Runtime/Engine/Private/PhysicsEngine/ImmediatePhysics/ImmediatePhysicsChaos/ImmediatePhysicsSimulation_Chaos.cpp
--- a/Engine/Source/Runtime/Engine/Private/PhysicsEngine/ImmediatePhysics/ImmediatePhysicsChaos/ImmediatePhysicsSimulation_Chaos.cpp
+++ b/Engine/Source/Runtime/Engine/Private/PhysicsEngine/ImmediatePhysics/ImmediatePhysicsChaos/ImmediatePhysicsSimulation_Chaos.cpp
@@ -2,6 +2,7 @@
 
 #include "Physics/ImmediatePhysics/ImmediatePhysicsChaos/ImmediatePhysicsSimulation_Chaos.h"
 #include "Chaos/PBDRigidsSOAs.h"
+#include "Physics/ImmediatePhysics/ImmediatePhysicsAdapters.h"
 #include "Physics/ImmediatePhysics/ImmediatePhysicsChaos/ImmediatePhysicsActorHandle_Chaos.h"
 #include "Physics/ImmediatePhysics/ImmediatePhysicsChaos/ImmediatePhysicsJointHandle_Chaos.h"
 #include "Physics/ImmediatePhysics/ImmediatePhysicsShared/ImmediatePhysicsCore.h"
@@ -250,7 +251,6 @@ FAutoConsoleVariableRef CVarChaosImmPhysScale(TEXT("p.Chaos.ImmPhys.DebugDraw.Sc
 
 namespace ImmediatePhysics_Chaos
 {
-
 	class FSimpleParticleUniqueIndices : public Chaos::IParticleUniqueIndices
 	{
 	public:
@@ -428,47 +428,31 @@ namespace ImmediatePhysics_Chaos
 
 	FActorHandle* FSimulation::CreateStaticActor(FBodyInstance* BodyInstance)
 	{
-		return CreateActor(EActorType::StaticActor, BodyInstance, BodyInstance->GetUnrealWorldTransform());
+		return CreateActor(ImmediatePhysics::MakeActorSetup(EActorType::StaticActor, BodyInstance, BodyInstance->GetUnrealWorldTransform()));
 	}
 
 	FActorHandle* FSimulation::CreateKinematicActor(FBodyInstance* BodyInstance, const FTransform& Transform)
 	{
-		return CreateActor(EActorType::KinematicActor, BodyInstance, Transform);
+		return CreateActor(ImmediatePhysics::MakeActorSetup(EActorType::KinematicActor, BodyInstance, Transform));
 	}
 
 	FActorHandle* FSimulation::CreateDynamicActor(FBodyInstance* BodyInstance, const FTransform& Transform)
 	{
-		return CreateActor(EActorType::DynamicActor, BodyInstance, Transform);
+		return CreateActor(ImmediatePhysics::MakeActorSetup(EActorType::DynamicActor, BodyInstance, Transform));
 	}
 
 	FActorHandle* FSimulation::CreateActor(EActorType ActorType, FBodyInstance* BodyInstance, const FTransform& Transform)
 	{
-		// @todo(ccaulfield): Shared materials
-		// @todo(ccaulfield): Add colliding particle pairs
-
-		using namespace Chaos;
+		return CreateActor(ImmediatePhysics::MakeActorSetup(ActorType, BodyInstance, Transform));
+	}
 
-		FActorHandle* ActorHandle = new FActorHandle(Implementation->Particles, Implementation->ParticlePrevXs, Implementation->ParticlePrevRs, ActorType, BodyInstance, Transform);
+	FActorHandle* FSimulation::CreateActor(FActorSetup&& ActorSetup)
+	{
+		FActorHandle* ActorHandle = new FActorHandle(Implementation->Particles, Implementation->ParticlePrevXs, Implementation->ParticlePrevRs, MoveTemp(ActorSetup));
 		int ActorIndex = Implementation->ActorHandles.Add(ActorHandle);
 
-		TUniquePtr<FChaosPhysicsMaterial> Material = MakeUnique<FChaosPhysicsMaterial>();
-		if (BodyInstance != nullptr)
-		{
-			// @todo(ccaulfield): We cannot ask for the physical material on a task thread, because FMICReentranceGuard in UMaterialInstance will assert (in editor). Fix this...
-			// For now we just use material defaults when actors are created on a task thread. This happens when adding world-objects to a RigidBody AnimNode simulation.
-			if (IsInGameThread())
-			{
-				UPhysicalMaterial* SimplePhysMat = BodyInstance->GetSimplePhysicalMaterial();
-				if (SimplePhysMat != nullptr)
-				{
-					Material->Friction = SimplePhysMat->Friction;
-					Material->Restitution = SimplePhysMat->Restitution;
-				}
-			}
-		}
-
-		ActorHandle->GetParticle()->AuxilaryValue(Implementation->ParticleMaterials) = MakeSerializable(Material);
-		ActorHandle->GetParticle()->AuxilaryValue(Implementation->PerParticleMaterials) = MoveTemp(Material);
+		ActorHandle->GetParticle()->AuxilaryValue(Implementation->ParticleMaterials) = MakeSerializable(ActorSetup.Material);
+		ActorHandle->GetParticle()->AuxilaryValue(Implementation->PerParticleMaterials) = MoveTemp(ActorSetup.Material);
 		ActorHandle->GetParticle()->AuxilaryValue(Implementation->CollidedParticles) = false;
 		ActorHandle->GetParticle()->AuxilaryValue(Implementation->ParticlePrevXs) = ActorHandle->GetParticle()->GetX();
 		ActorHandle->GetParticle()->AuxilaryValue(Implementation->ParticlePrevRs) = ActorHandle->GetParticle()->GetR();
@@ -478,6 +462,7 @@ namespace ImmediatePhysics_Chaos
 		return ActorHandle;
 	}
 
+
 	void FSimulation::DestroyActor(FActorHandle* ActorHandle)
 	{
 		// @todo(ccaulfield): FActorHandle could remember its index to optimize this
@@ -553,31 +538,20 @@ namespace ImmediatePhysics_Chaos
 
 	FJointHandle* FSimulation::CreateJoint(FConstraintInstance* ConstraintInstance, FActorHandle* Body1, FActorHandle* Body2)
 	{
-		if (ConstraintInstance && !ConstraintInstance->ProfileInstance.bUseLinearJointSolver)
-		{
-			FJointHandle* JointHandle = new FJointHandle(&(Implementation->Joints.NonLinearConstraints), ConstraintInstance, Body1, Body2);
-			Implementation->JointHandles.Add(JointHandle);
-			return JointHandle;
-		}
-		else
-		{
-			FJointHandle* JointHandle = new FJointHandle(&(Implementation->Joints.LinearConstraints), ConstraintInstance, Body1, Body2);
-			Implementation->JointHandles.Add(JointHandle);
-			return JointHandle;
-		}
+		return CreateJoint(ImmediatePhysics::MakeJointSetup(ConstraintInstance, Body1, Body2));
 	}
 
-	FJointHandle* FSimulation::CreateJoint(const Chaos::FPBDJointSettings& ConstraintSettings, FActorHandle* const Body1, FActorHandle* const Body2)
+	FJointHandle* FSimulation::CreateJoint(const FJointSetup& JointSetup)
 	{
-		if (ConstraintSettings.bUseLinearSolver)
+		if (JointSetup.ConstraintSettings.bUseLinearSolver)
 		{
-			FJointHandle* JointHandle = new FJointHandle(&(Implementation->Joints.LinearConstraints), ConstraintSettings, Body1, Body2);
+			FJointHandle* JointHandle = new FJointHandle(&(Implementation->Joints.LinearConstraints), JointSetup);
 			Implementation->JointHandles.Add(JointHandle);
 			return JointHandle;
 		}
 		else
 		{
-			FJointHandle* JointHandle = new FJointHandle(&(Implementation->Joints.NonLinearConstraints), ConstraintSettings, Body1, Body2);
+			FJointHandle* JointHandle = new FJointHandle(&(Implementation->Joints.NonLinearConstraints), JointSetup);
 			Implementation->JointHandles.Add(JointHandle);
 			return JointHandle;
 		}
diff --git a/Engine/Source/Runtime/Engine/Public/Physics/ImmediatePhysics/ImmediatePhysicsAdapters.h b/Engine/Source/Runtime/Engine/Public/Physics/ImmediatePhysics/ImmediatePhysicsAdapters.h
new file mode 100644
--- /dev/null
+++ b/Engine/Source/Runtime/Engine/Public/Physics/ImmediatePhysics/ImmediatePhysicsAdapters.h
@@ -0,0 +1,43 @@
+// Copyright Epic Games, Inc. All Rights Reserved.
+
+#pragma once
+
+#include "Physics/ImmediatePhysics/ImmediatePhysicsDeclares.h"
+#include "Chaos/ChaosEngineInterface.h"
+
+#include "Engine/EngineTypes.h"
+
+struct FConeConstraint;
+struct FConstraintProfileProperties;
+struct FLinearConstraint;
+struct FTwistConstraint;
+
+namespace ImmediatePhysics
+{
+	// Make an ActorSetup from a BodyInstance
+	ENGINE_API FActorSetup MakeStaticActorSetup(const FBodyInstance* InBodyInstance, const FTransform& InTransform);
+
+	// Make an ActorSetup from a BodyInstance
+	ENGINE_API FActorSetup MakeKinematicActorSetup(const FBodyInstance* InBodyInstance, const FTransform& InTransform);
+
+	// Make an ActorSetup from a BodyInstance
+	ENGINE_API FActorSetup MakeDynamicActorSetup(const FBodyInstance* InBodyInstance, const FTransform& InTransform);
+
+	// Make an ActorSetup from a BodyInstance
+	ENGINE_API FActorSetup MakeActorSetup(EActorType InActorType, const FBodyInstance* InBodyInstance, const FTransform& InTransform);
+
+	// Make a JointSetup from a ConstraintInstance
+	ENGINE_API FJointSetup MakeJointSetup(const FConstraintInstance* InConstraintInstance, FActorHandle* InActor1, FActorHandle* InActor2);
+
+	// Transfers parameters from a constraint profile into the joint settings
+	ENGINE_API void UpdateJointSettingsFromConstraintProfile(const FConstraintProfileProperties& Profile, Chaos::FPBDJointSettings& JointSettings);
+
+	// Extracts parameters from a constraint constraint into the joint settings
+	ENGINE_API void UpdateJointSettingsFromLinearConstraint(const FLinearConstraint& Constraint, Chaos::FPBDJointSettings& JointSettings);
+	
+	//Extracts parameters from a constraint constraint into the joint settings
+	ENGINE_API void UpdateJointSettingsFromConeConstraint(const FConeConstraint& Constraint, Chaos::FPBDJointSettings& JointSettings);
+	
+	//Extracts parameters from a constraint constraint into the joint settings
+	ENGINE_API void UpdateJointSettingsFromTwistConstraint(const FTwistConstraint& Constraint, Chaos::FPBDJointSettings& JointSettings);
+}
diff --git a/Engine/Source/Runtime/Engine/Public/Physics/ImmediatePhysics/ImmediatePhysicsChaos/ImmediatePhysicsActorHandle_Chaos.h b/Engine/Source/Runtime/Engine/Public/Physics/ImmediatePhysics/ImmediatePhysicsChaos/ImmediatePhysicsActorHandle_Chaos.h
--- a/Engine/Source/Runtime/Engine/Public/Physics/ImmediatePhysics/ImmediatePhysicsChaos/ImmediatePhysicsActorHandle_Chaos.h
+++ b/Engine/Source/Runtime/Engine/Public/Physics/ImmediatePhysics/ImmediatePhysicsChaos/ImmediatePhysicsActorHandle_Chaos.h
@@ -4,11 +4,43 @@
 
 #include "Physics/ImmediatePhysics/ImmediatePhysicsChaos/ImmediatePhysicsCore_Chaos.h"
 #include "Chaos/ChaosEngineInterface.h"
+#include "Chaos/ImplicitObject.h"
 
 #include "Engine/EngineTypes.h"
 
 namespace ImmediatePhysics_Chaos
 {
+	struct FActorSetup
+	{
+	public:
+		ENGINE_API FActorSetup();
+		ENGINE_API FActorSetup(FActorSetup&& R);
+		ENGINE_API ~FActorSetup();
+
+		FTransform Transform;
+
+		// NOTE: Shapes and Material are Moved out the setup structure when
+		// ued to create an Actor.
+		Chaos::FImplicitObjectPtr Geometry;
+		TArray<TUniquePtr<Chaos::FPerShapeData>> Shapes;
+		TUniquePtr<Chaos::FChaosPhysicsMaterial> Material;
+
+		FReal Mass;
+		FVector Inertia;
+		FTransform CoMTransform;
+
+		FReal LinearDamping;
+		FReal AngularDamping;
+
+		EActorType ActorType;
+		ECollisionEnabled::Type CollisionEnabled;
+		int8 GravityGroupIndex;
+		uint8 bInertiaConditioningEnabled : 1;
+		uint8 bEnableGravity : 1;
+		uint8 bUpdateKinematicFromSimulation : 1;
+		uint8 bGyroscopicTorqueEnabled : 1;
+	};
+
 	/** handle associated with a physics actor. This is the proper way to read/write to the physics simulation */
 	struct FActorHandle
 	{
@@ -66,9 +98,6 @@ namespace ImmediatePhysics_Chaos
 		/** Returns true if the world gravity applies to this actor */
 		ENGINE_API bool IsGravityEnabled() const;
 
-		/** Whether the body is static */
-		ENGINE_API bool IsStatic() const;
-
 		/** Returns true the body is or could be dynamic (i.e. currently simulated, or a rigid body that is currently kinematic) */
 		ENGINE_API bool CouldBeDynamic() const;
 
@@ -164,38 +193,28 @@ namespace ImmediatePhysics_Chaos
 		ENGINE_API void SetLevel(int32 InLevel);
 
 	private:
-		ENGINE_API FKinematicTarget& GetKinematicTarget();
+		FKinematicTarget& GetKinematicTarget();
 
-		ENGINE_API void CreateParticleHandle(
-			FBodyInstance*                 BodyInstance,
-			const EActorType               ActorType,
-			const FTransform&              WorldTransform,
-			const FReal                    Mass,
-			const Chaos::FVec3             Inertia,
-			const Chaos::FRigidTransform3& CoMTransform);
+		void CreateParticleHandle(
+			FActorSetup&& ActorSetup);
 
 		friend struct FSimulation;
 		friend struct FJointHandle;
 
-		ENGINE_API FActorHandle(
+		FActorHandle(
 			Chaos::FPBDRigidsSOAs& InParticles,
 			Chaos::TArrayCollectionArray<Chaos::FVec3>& InParticlePrevXs,
 			Chaos::TArrayCollectionArray<Chaos::FRotation3>& InParticlePrevRs,
-			EActorType ActorType,
-			FBodyInstance* BodyInstance,
-			const FTransform& Transform);
+			FActorSetup&& ActorSetup);
 
 
-		ENGINE_API Chaos::FGenericParticleHandle Handle() const;
+		Chaos::FGenericParticleHandle Handle() const;
 
 		FName Name;
 		Chaos::FPBDRigidsSOAs& Particles;
 		Chaos::FGeometryParticleHandle* ParticleHandle;
 		Chaos::TArrayCollectionArray<Chaos::FVec3>& ParticlePrevXs;
 		Chaos::TArrayCollectionArray<Chaos::FRotation3>& ParticlePrevRs;
-		Chaos::FImplicitObjectPtr Geometry;
-		TArray<TUniquePtr<Chaos::FPerShapeData>> Shapes;
 		int32 Level;
 	};
-
 }
diff --git a/Engine/Source/Runtime/Engine/Public/Physics/ImmediatePhysics/ImmediatePhysicsChaos/ImmediatePhysicsDeclares_Chaos.h b/Engine/Source/Runtime/Engine/Public/Physics/ImmediatePhysics/ImmediatePhysicsChaos/ImmediatePhysicsDeclares_Chaos.h
--- a/Engine/Source/Runtime/Engine/Public/Physics/ImmediatePhysics/ImmediatePhysicsChaos/ImmediatePhysicsDeclares_Chaos.h
+++ b/Engine/Source/Runtime/Engine/Public/Physics/ImmediatePhysics/ImmediatePhysicsChaos/ImmediatePhysicsDeclares_Chaos.h
@@ -2,9 +2,16 @@
 
 #pragma once
 
+namespace Chaos
+{
+	class FPBDJointSettings;
+}
+
 namespace ImmediatePhysics_Chaos
 {
-	struct FSimulation;
 	struct FActorHandle;
+	struct FActorSetup;
 	struct FJointHandle;
+	struct FJointSetup;
+	struct FSimulation;
 }
diff --git a/Engine/Source/Runtime/Engine/Public/Physics/ImmediatePhysics/ImmediatePhysicsChaos/ImmediatePhysicsJointHandle_Chaos.h b/Engine/Source/Runtime/Engine/Public/Physics/ImmediatePhysics/ImmediatePhysicsChaos/ImmediatePhysicsJointHandle_Chaos.h
--- a/Engine/Source/Runtime/Engine/Public/Physics/ImmediatePhysics/ImmediatePhysicsChaos/ImmediatePhysicsJointHandle_Chaos.h
+++ b/Engine/Source/Runtime/Engine/Public/Physics/ImmediatePhysics/ImmediatePhysicsChaos/ImmediatePhysicsJointHandle_Chaos.h
@@ -5,32 +5,51 @@
 #include "Physics/ImmediatePhysics/ImmediatePhysicsChaos/ImmediatePhysicsCore_Chaos.h"
 #include "Physics/ImmediatePhysics/ImmediatePhysicsChaos/ImmediatePhysicsActorHandle_Chaos.h"
 
+#include "Chaos/PBDJointConstraintTypes.h"
 #include "Chaos/Vector.h"
 
 #include "Engine/EngineTypes.h"
 
-namespace Chaos
+namespace ImmediatePhysics_Chaos
 {
-	class FPBDJointSettings;
-}
+	struct FJointSetup
+	{
+	public:
+		FJointSetup()
+			: ConstraintSettings()
+			, Actor1(nullptr)
+			, Actor2(nullptr)
+		{
+		}
 
-struct FConstraintProfileProperties;
-struct FLinearConstraint;
-struct FConeConstraint;
-struct FTwistConstraint;
+		FJointSetup(const Chaos::FPBDJointSettings& InConstraintSettings, FActorHandle* InActor1, FActorHandle* InActor2)
+			: ConstraintSettings(InConstraintSettings)
+			, Actor1(InActor1)
+			, Actor2(InActor2)
+		{
+		}
+
+		Chaos::FPBDJointSettings ConstraintSettings;
+		FActorHandle* Actor1;
+		FActorHandle* Actor2;
+	};
 
-namespace ImmediatePhysics_Chaos
-{
 	/** handle associated with a physics joint. This is the proper way to read/write to the physics simulation */
 	struct FJointHandle
 	{
 	public:
 		using FChaosConstraintContainer = Chaos::FPBDJointConstraints;
 		using FChaosConstraintHandle = typename Chaos::FPBDJointConstraintHandle;
-		using FPBDJointSettings = Chaos::FPBDJointSettings;
 
+		using FPBDJointSettings UE_DEPRECATED(5.6, "Not used") = Chaos::FPBDJointSettings;
+
+		UE_DEPRECATED(5.6, "Use constructor which takes FJointSetup")
 		ENGINE_API FJointHandle(FChaosConstraintContainer* InConstraints, FConstraintInstance* ConstraintInstance, FActorHandle* InActor1, FActorHandle* InActor2);
-		ENGINE_API FJointHandle(FChaosConstraintContainer* InConstraints, const FPBDJointSettings& ConstraintSettings, FActorHandle* InActor1, FActorHandle* InActor2);
+
+		UE_DEPRECATED(5.6, "Use constructor which takes FJointSetup")
+		ENGINE_API FJointHandle(FChaosConstraintContainer* InConstraints, const Chaos::FPBDJointSettings& ConstraintSettings, FActorHandle* InActor1, FActorHandle* InActor2);
+
+		ENGINE_API FJointHandle(FChaosConstraintContainer* InConstraints, const FJointSetup& JointSetup);
 		ENGINE_API ~FJointHandle();
 
 		ENGINE_API FChaosConstraintHandle* GetConstraint();
@@ -42,25 +61,20 @@ namespace ImmediatePhysics_Chaos
 		ENGINE_API void SetSoftLinearSettings(bool bLinearSoft, FReal LinearStiffness, FReal LinearDamping);
 
 	private:
+		void CreateJoint(FChaosConstraintContainer* InConstraints, const Chaos::FPBDJointSettings& ConstraintSettings, FActorHandle* const Actor1, FActorHandle* const Actor2);
 		void SetActorInertiaConditioningDirty();
 		
 		Chaos::TVec2<FActorHandle*> ActorHandles;
 		FChaosConstraintContainer* Constraints;
 		FChaosConstraintHandle* ConstraintHandle;
 	};
+}
 
-	/** Transfers parameters from a constraint profile into the joint settings */
-	ENGINE_API void UpdateJointSettingsFromConstraintProfile(
-		const FConstraintProfileProperties& Profile, Chaos::FPBDJointSettings& JointSettings);
-
-	/** Extracts parameters from a constraint constraint into the joint settings */
-	ENGINE_API void UpdateJointSettingsFromLinearConstraint(
-		const FLinearConstraint& Constraint, Chaos::FPBDJointSettings& JointSettings);
-	/** Extracts parameters from a constraint constraint into the joint settings */
-	ENGINE_API void UpdateJointSettingsFromConeConstraint(
-		const FConeConstraint& Constraint, Chaos::FPBDJointSettings& JointSettings);
-	/** Extracts parameters from a constraint constraint into the joint settings */
-	ENGINE_API void UpdateJointSettingsFromTwistConstraint(
-		const FTwistConstraint& Constraint, Chaos::FPBDJointSettings& JointSettings);
+// TODO: REMOVE - only for deprecated function
+struct FConstraintProfileProperties;
 
+namespace ImmediatePhysics_Chaos
+{
+	UE_DEPRECATED(5.6, "Moved to ImmediatePhysics namespace")
+	ENGINE_API void UpdateJointSettingsFromConstraintProfile(const FConstraintProfileProperties& Profile, Chaos::FPBDJointSettings& JointSettings);
 }
diff --git a/Engine/Source/Runtime/Engine/Public/Physics/ImmediatePhysics/ImmediatePhysicsChaos/ImmediatePhysicsSimulation_Chaos.h b/Engine/Source/Runtime/Engine/Public/Physics/ImmediatePhysics/ImmediatePhysicsChaos/ImmediatePhysicsSimulation_Chaos.h
--- a/Engine/Source/Runtime/Engine/Public/Physics/ImmediatePhysics/ImmediatePhysicsChaos/ImmediatePhysicsSimulation_Chaos.h
+++ b/Engine/Source/Runtime/Engine/Public/Physics/ImmediatePhysics/ImmediatePhysicsChaos/ImmediatePhysicsSimulation_Chaos.h
@@ -32,16 +32,20 @@ namespace ImmediatePhysics_Chaos
 		ENGINE_API FActorHandle* GetActorHandle(int32 ActorHandleIndex);
 		ENGINE_API const FActorHandle* GetActorHandle(int32 ActorHandleIndex) const;
 
-		/** Create a static body and add it to the simulation */
+		UE_DEPRECATED(5.6, "Use CreateActor with FActorSetup")
 		ENGINE_API FActorHandle* CreateStaticActor(FBodyInstance* BodyInstance);
-
-		/** Create a kinematic body and add it to the simulation */
+		
+		UE_DEPRECATED(5.6, "Use CreateActor with FActorSetup")
 		ENGINE_API FActorHandle* CreateKinematicActor(FBodyInstance* BodyInstance, const FTransform& Transform);
-
-		/** Create a dynamic body and add it to the simulation */
+		
+		UE_DEPRECATED(5.6, "Use CreateActor with FActorSetup")
 		ENGINE_API FActorHandle* CreateDynamicActor(FBodyInstance* BodyInstance, const FTransform& Transform);
 
+		UE_DEPRECATED(5.6, "Use CreateActor with FActorSetup")
 		ENGINE_API FActorHandle* CreateActor(EActorType ActorType, FBodyInstance* BodyInstance, const FTransform& Transform);
+
+		ENGINE_API FActorHandle* CreateActor(FActorSetup&& ActorSetup);
+
 		ENGINE_API void DestroyActor(FActorHandle* ActorHandle);
 
 		ENGINE_API void DestroyActorCollisions(FActorHandle* ActorHandle);
@@ -53,10 +57,11 @@ namespace ImmediatePhysics_Chaos
 		ENGINE_API void SetHasCollision(FActorHandle* ActorHandle, bool bHasCollision);
 
 		/** Create a physical joint and add it to the simulation */
+		UE_DEPRECATED(5.6, "Use CreateActor with FJointSetup")
 		ENGINE_API FJointHandle* CreateJoint(FConstraintInstance* ConstraintInstance, FActorHandle* Body1, FActorHandle* Body2);
 
 		/** Create a physical joint and add it to the simulation */
-		ENGINE_API FJointHandle* CreateJoint(const Chaos::FPBDJointSettings& ConstraintSettings, FActorHandle* const Body1, FActorHandle* const Body2);
+		ENGINE_API FJointHandle* CreateJoint(const FJointSetup& JointSetup);
 
 		ENGINE_API void DestroyJoint(FJointHandle* JointHandle);
 
@@ -94,7 +99,8 @@ namespace ImmediatePhysics_Chaos
 			Chaos::FPBDJointSolverSettings* JointSolverSettings = nullptr, 
 			Chaos::FPBDCollisionSolverSettings* CollisionSolverSettings = nullptr,
 			Chaos::FCollisionDetectorSettings* CollisionDetectorSettings = nullptr);
-		void Simulate_AssumesLocked(FReal DeltaTime, FReal MaxStepTime, int32 MaxSubSteps, const FVector& InGravity,
+
+		ENGINE_API void Simulate_AssumesLocked(FReal DeltaTime, FReal MaxStepTime, int32 MaxSubSteps, const FVector& InGravity,
 			Chaos::FPBDJointSolverSettings* JointSolverSettings = nullptr,
 			Chaos::FPBDCollisionSolverSettings* CollisionSolverSettings = nullptr,
 			Chaos::FCollisionDetectorSettings* CollisionDetectorSettings = nullptr)
@@ -140,19 +146,19 @@ namespace ImmediatePhysics_Chaos
 		ENGINE_API void DebugDraw();
 
 	private:
-		ENGINE_API void RemoveFromCollidingPairs(FActorHandle* ActorHandle);
-		ENGINE_API void UpdateInertiaConditioning(const FVector& Gravity);
-		ENGINE_API void PackCollidingPairs();
-		ENGINE_API void UpdateActivePotentiallyCollidingPairs();
-		ENGINE_API void EnableDisableJoints();
-		ENGINE_API FReal UpdateStepTime(const FReal DeltaTime, const FReal MaxStepTime);
-
-		ENGINE_API void UpdateStatCounters();
-		ENGINE_API void DebugDrawStaticParticles();
-		ENGINE_API void DebugDrawKinematicParticles();
-		ENGINE_API void DebugDrawDynamicParticles();
-		ENGINE_API void DebugDrawConstraints();
-		ENGINE_API void DebugDrawSimulationSpace();
+		void RemoveFromCollidingPairs(FActorHandle* ActorHandle);
+		void UpdateInertiaConditioning(const FVector& Gravity);
+		void PackCollidingPairs();
+		void UpdateActivePotentiallyCollidingPairs();
+		void EnableDisableJoints();
+		FReal UpdateStepTime(const FReal DeltaTime, const FReal MaxStepTime);
+
+		void UpdateStatCounters();
+		void DebugDrawStaticParticles();
+		void DebugDrawKinematicParticles();
+		void DebugDrawDynamicParticles();
+		void DebugDrawConstraints();
+		void DebugDrawSimulationSpace();
 
 		struct FImplementation;
 		TUniquePtr<FImplementation> Implementation;
diff --git a/Engine/Source/Runtime/Engine/Public/Physics/ImmediatePhysics/ImmediatePhysicsDeclares.h b/Engine/Source/Runtime/Engine/Public/Physics/ImmediatePhysics/ImmediatePhysicsDeclares.h
--- a/Engine/Source/Runtime/Engine/Public/Physics/ImmediatePhysics/ImmediatePhysicsDeclares.h
+++ b/Engine/Source/Runtime/Engine/Public/Physics/ImmediatePhysics/ImmediatePhysicsDeclares.h
@@ -11,5 +11,7 @@ namespace ImmediatePhysics
 	using EForceType = ImmediatePhysics_Shared::EForceType;
 	using FSimulation = ImmediatePhysics_Chaos::FSimulation;
 	using FActorHandle = ImmediatePhysics_Chaos::FActorHandle;
+	using FActorSetup = ImmediatePhysics_Chaos::FActorSetup;
 	using FJointHandle = ImmediatePhysics_Chaos::FJointHandle;
+	using FJointSetup = ImmediatePhysics_Chaos::FJointSetup;
 }
\ No newline at end of file
