Only force GPU skin cache for raytraced meshes if ray tracing is enabled.

diff --git a/Engine/Source/Runtime/Engine/Private/Components/SkinnedMeshComponent.cpp b/Engine/Source/Runtime/Engine/Private/Components/SkinnedMeshComponent.cpp
--- a/Engine/Source/Runtime/Engine/Private/Components/SkinnedMeshComponent.cpp
+++ b/Engine/Source/Runtime/Engine/Private/Components/SkinnedMeshComponent.cpp
@@ -2254,7 +2254,7 @@ bool USkinnedMeshComponent::IsSkinCacheAllowed(int32 LodIdx) const
 
 	if (FSkeletalMeshRenderData* SkelMeshRenderData = GetSkinnedAsset()->GetResourceForRendering())
 	{
-		if (SkelMeshRenderData->bSupportRayTracing)
+		if (SkelMeshRenderData->bSupportRayTracing && FGPUSkinCache::IsGPUSkinCacheRayTracingSupported() && IsRayTracingEnabled())
 		{
 			return true;
 		}
diff --git a/Engine/Source/Runtime/Engine/Private/SkinnedMeshSceneProxyDesc.cpp b/Engine/Source/Runtime/Engine/Private/SkinnedMeshSceneProxyDesc.cpp
--- a/Engine/Source/Runtime/Engine/Private/SkinnedMeshSceneProxyDesc.cpp
+++ b/Engine/Source/Runtime/Engine/Private/SkinnedMeshSceneProxyDesc.cpp
@@ -318,7 +318,7 @@ bool FSkinnedMeshSceneProxyDesc::IsSkinCacheAllowed(int32 LodIdx) const
 
 	if (FSkeletalMeshRenderData* SkelMeshRenderData = GetSkinnedAsset()->GetResourceForRendering())
 	{
-		if (SkelMeshRenderData->bSupportRayTracing)
+		if (SkelMeshRenderData->bSupportRayTracing && FGPUSkinCache::IsGPUSkinCacheRayTracingSupported() && IsRayTracingEnabled())
 		{
 			return true;
 		}
@@ -374,7 +374,7 @@ bool FSkinnedMeshSceneProxyDynamicData::IsSkinCacheAllowed(int32 LodIdx, const U
 
 	if (FSkeletalMeshRenderData* SkelMeshRenderData = InSkinnedAsset->GetResourceForRendering())
 	{
-		if (SkelMeshRenderData->bSupportRayTracing)
+		if (SkelMeshRenderData->bSupportRayTracing && FGPUSkinCache::IsGPUSkinCacheRayTracingSupported() && IsRayTracingEnabled())
 		{
 			return true;
 		}
