[shaders] fix global shader error message string so it matches between two if conditions (required by localization)

diff --git a/Engine/Source/Runtime/Engine/Private/ShaderCompiler/ShaderCompilerEditor.cpp b/Engine/Source/Runtime/Engine/Private/ShaderCompiler/ShaderCompilerEditor.cpp
--- a/Engine/Source/Runtime/Engine/Private/ShaderCompiler/ShaderCompilerEditor.cpp
+++ b/Engine/Source/Runtime/Engine/Private/ShaderCompiler/ShaderCompilerEditor.cpp
@@ -600,7 +600,7 @@ void ProcessCompiledGlobalShaders(const TArray<FShaderCommonCompileJobPtr>& Comp
 		}
 		else
 		{
-			UE_LOGFMT_NSLOC(LogShaders, Error, "Shaders", "GlobalShadersCompilationFailed", "{NumErrors} errors encountered compiling global shaders for platform {Platform}.\n{Errors}{RetryMsg}",
+			UE_LOGFMT_NSLOC(LogShaders, Error, "Shaders", "GlobalShadersCompilationFailed", "{NumErrors} errors encountered compiling global shaders for platform {Platform}:\n{Errors}{RetryMsg}",
 				("NumErrors", UniqueErrorCount),
 				("Platform", ShaderDiagInfo.TargetShaderPlatformString),
 				("RetryMsg", IsRunningCommandlet() ? TEXT("") : RetryMsg),
