Improved UMaterialExpressionConvert implementation in the new material translator. It also now generalizes to any scalar type, for instance it takes "Vector2" intead of "Float2" as input.

#rnx

diff --git a/Engine/Config/BaseEngine.ini b/Engine/Config/BaseEngine.ini
--- a/Engine/Config/BaseEngine.ini
+++ b/Engine/Config/BaseEngine.ini
@@ -1353,6 +1353,8 @@ ClassName=/Script/OnlineSubsystemUtils.OnlineEngineInterfaceImpl
 +ClassRedirects=(OldName="/Script/HierarchyEditor.HierarchySection",NewName="/Script/DataHierarchyEditor.HierarchySection")
 +ClassRedirects=(OldName="/Script/HierarchyEditor.HierarchyCategory",NewName="/Script/DataHierarchyEditor.HierarchyCategory")
 
++EnumRedirects=(OldName="EMaterialExpressionConvertType", ValueChanges=(("Float1", "Scalar"), ("Float2", "Vector2"), ("Float3", "Vector3"), ("Float4", "Vector4")))
+
 ; End of releases redirects
 
 ; Fixing the output of tags for blueprints to be reset.
diff --git a/Engine/Source/Editor/GraphEditor/Private/MaterialNodes/SGraphNodeMaterialConvert.cpp b/Engine/Source/Editor/GraphEditor/Private/MaterialNodes/SGraphNodeMaterialConvert.cpp
--- a/Engine/Source/Editor/GraphEditor/Private/MaterialNodes/SGraphNodeMaterialConvert.cpp
+++ b/Engine/Source/Editor/GraphEditor/Private/MaterialNodes/SGraphNodeMaterialConvert.cpp
@@ -585,12 +585,12 @@ TSharedRef<SWidget> SGraphNodeMaterialConvert::CreateAddPinContextMenu(bool bInp
 {
 	FMenuBuilder ContextMenuBuilder(true, NULL);
 
-	static const TArray<EMaterialExpressionConvertType> ConvertTypes =
+	static const EMaterialExpressionConvertType ConvertTypes[] =
 	{ 
-		EMaterialExpressionConvertType::Float1, 
-		EMaterialExpressionConvertType::Float2, 
-		EMaterialExpressionConvertType::Float3, 
-		EMaterialExpressionConvertType::Float4 
+		EMaterialExpressionConvertType::Scalar, 
+		EMaterialExpressionConvertType::Vector2, 
+		EMaterialExpressionConvertType::Vector3, 
+		EMaterialExpressionConvertType::Vector4 
 	};
 
 	for (const EMaterialExpressionConvertType ConvertType : ConvertTypes)
diff --git a/Engine/Source/Runtime/Engine/Private/Materials/MaterialExpressions.cpp b/Engine/Source/Runtime/Engine/Private/Materials/MaterialExpressions.cpp
--- a/Engine/Source/Runtime/Engine/Private/Materials/MaterialExpressions.cpp
+++ b/Engine/Source/Runtime/Engine/Private/Materials/MaterialExpressions.cpp
@@ -23596,7 +23596,7 @@ void UMaterialExpressionConvert::RegisterAdditionalMenuActions(FGraphActionMenuB
 			
 				// Initialize NumComponents Float1 Inputs
 				FMaterialExpressionConvertInput Float1Input;
-				Float1Input.Type = EMaterialExpressionConvertType::Float1;
+				Float1Input.Type = EMaterialExpressionConvertType::Scalar;
 				NewConvertExpression->ConvertInputs.Init(Float1Input, NumComponents);
 
 				// Add a single Output of our desired Convert Type
@@ -23620,9 +23620,9 @@ void UMaterialExpressionConvert::RegisterAdditionalMenuActions(FGraphActionMenuB
 		ActionMenuBuilder.AddAction(Action_MakeNode);
 	};
 
-	CreateMakeNode(EMaterialExpressionConvertType::Float2);
-	CreateMakeNode(EMaterialExpressionConvertType::Float3);
-	CreateMakeNode(EMaterialExpressionConvertType::Float4);
+	CreateMakeNode(EMaterialExpressionConvertType::Vector2);
+	CreateMakeNode(EMaterialExpressionConvertType::Vector3);
+	CreateMakeNode(EMaterialExpressionConvertType::Vector4);
 
 	// Helper Lambda to Create a Break Node of a specified ConvertType
 	auto CreateBreakNode = [this, &ActionMenuBuilder, &CategoryName](EMaterialExpressionConvertType ConvertType)
@@ -23651,9 +23651,9 @@ void UMaterialExpressionConvert::RegisterAdditionalMenuActions(FGraphActionMenuB
 				NewConvertExpression->ConvertInputs = { ConvertInput };
 
 				// Initialize NumComponents Float1 Outputs
-				FMaterialExpressionConvertOutput Float1Output;
-				Float1Output.Type = EMaterialExpressionConvertType::Float1;
-				NewConvertExpression->ConvertOutputs.Init(Float1Output, NumComponents);
+				FMaterialExpressionConvertOutput ScalarOutput;
+				ScalarOutput.Type = EMaterialExpressionConvertType::Scalar;
+				NewConvertExpression->ConvertOutputs.Init(ScalarOutput, NumComponents);
 
 				// Create the relevant mappings
 				for (int32 MappingIndex = 0; MappingIndex < NumComponents; ++MappingIndex)
@@ -23671,9 +23671,9 @@ void UMaterialExpressionConvert::RegisterAdditionalMenuActions(FGraphActionMenuB
 		ActionMenuBuilder.AddAction(Action_BreakNode);
 	};
 
-	CreateBreakNode(EMaterialExpressionConvertType::Float2);
-	CreateBreakNode(EMaterialExpressionConvertType::Float3);
-	CreateBreakNode(EMaterialExpressionConvertType::Float4);
+	CreateBreakNode(EMaterialExpressionConvertType::Vector2);
+	CreateBreakNode(EMaterialExpressionConvertType::Vector3);
+	CreateBreakNode(EMaterialExpressionConvertType::Vector4);
 }
 
 void UMaterialExpressionConvert::RecreateOutputs()
diff --git a/Engine/Source/Runtime/Engine/Private/Materials/MaterialExpressionsIR.cpp b/Engine/Source/Runtime/Engine/Private/Materials/MaterialExpressionsIR.cpp
--- a/Engine/Source/Runtime/Engine/Private/Materials/MaterialExpressionsIR.cpp
+++ b/Engine/Source/Runtime/Engine/Private/Materials/MaterialExpressionsIR.cpp
@@ -1008,94 +1008,97 @@ void UMaterialExpressionSmoothStep::Build(MIR::FEmitter& Emitter)
 
 void UMaterialExpressionConvert::Build(MIR::FEmitter& Emitter)
 {
-	// Cache Compiled Inputs as they may be referred in multiple ConvertMappings/Outputs
-	TArray<MIR::FValue*> CachedCompiledInputs;
-	CachedCompiledInputs.Init(nullptr, ConvertInputs.Num());
+	TArray<MIR::FValue*, TInlineAllocator<8>> InputValues;
+	InputValues.Init(nullptr, ConvertInputs.Num());
 
 	for (int32 OutputIndex = 0; OutputIndex < ConvertOutputs.Num(); ++OutputIndex)
 	{
 		const FMaterialExpressionConvertOutput& ConvertOutput = ConvertOutputs[OutputIndex];
-		const EMaterialExpressionConvertType OutputType = ConvertOutput.Type;
-		const int32 OutputNumComponents = MaterialExpressionConvertType::GetComponentCount(OutputType);
-
 		MIR::FValue* OutComponents[4] = { nullptr, nullptr, nullptr, nullptr };
 
-		// Determine value of any connected outputs
-		for (const FMaterialExpressionConvertMapping& ConvertMapping : ConvertMappings)
+		for (const FMaterialExpressionConvertMapping& Mapping : ConvertMappings)
 		{
-			const int32 CurrentOutputIndex = ConvertMapping.OutputIndex;
-
 			// We only care about mappings relevant to this output
-			if (CurrentOutputIndex != OutputIndex)
+			if (Mapping.OutputIndex != OutputIndex)
 			{
 				continue;
 			}
 
-			const int32 OutputComponentIndex = ConvertMapping.OutputComponentIndex;
-			if (!IsValidComponentIndex(OutputComponentIndex, OutputType))
+			const int32 OutputComponentIndex = Mapping.OutputComponentIndex;
+			if (!IsValidComponentIndex(OutputComponentIndex, ConvertOutput.Type))
 			{
-				Emitter.Errorf(TEXT("Convert mapping's output component (%d) is invalid."), OutputComponentIndex);
+				Emitter.Errorf(TEXT("Convert mapping's output component `%d` is invalid."), OutputComponentIndex);
 				continue;
 			}
 
-			const int32 InputIndex = ConvertMapping.InputIndex;
+			const int32 InputIndex = Mapping.InputIndex;
 			if (!ConvertInputs.IsValidIndex(InputIndex))
 			{
-				Emitter.Errorf(TEXT("Convert mapping's input (%d) is invalid."), InputIndex);
+				Emitter.Errorf(TEXT("Convert mapping's input `%d` is invalid."), InputIndex);
 				continue;
 			}
 
 			FMaterialExpressionConvertInput& ConvertInput = ConvertInputs[InputIndex];
-			const int32 InputComponentIndex = ConvertMapping.InputComponentIndex;
+			const int32 InputComponentIndex = Mapping.InputComponentIndex;
 			if (!IsValidComponentIndex(InputComponentIndex, ConvertInput.Type))
 			{
-				Emitter.Errorf(TEXT("Convert mapping's input component (%d) is invalid."), InputComponentIndex);
+				Emitter.Errorf(TEXT("Convert mapping's input component `%d` is invalid."), InputComponentIndex);
 				continue;
 			}
 
-			// Search CachedCompiledInputs for previously compiled value. If it isn't found, compile that expression input now
-			MIR::FValue*& CompiledInput = CachedCompiledInputs[InputIndex];
-			if (CompiledInput == nullptr)
+			// If not already emitted, read the input value, cast it to the specified input
+			// type and cache it into an array, as each input could be used multiple times
+			// by output values.
+			if (!InputValues[InputIndex])
 			{
-				// Need to compile input if its connected or compile a default-value constant if not
-				CompiledInput = Emitter.TryInput(&ConvertInput.ExpressionInput)
-					? Emitter.Input(&ConvertInput.ExpressionInput)
-					: Emitter.EmitConstantFloat4(ConvertInput.DefaultValue);
+				// Read the input's value (or read float zero if disconnected).
+				FExpressionInput* Input = &ConvertInputs[InputIndex].ExpressionInput;
+				InputValues[InputIndex] = Emitter.DefaultInputToFloatScalarZero(Input).Input(Input);
+
+				// Expect type to be primitive.
+				MIR::FPrimitiveTypePtr InputPrimitiveType = InputValues[InputIndex]->Type->AsPrimitive();
+				if (!InputPrimitiveType)
+				{
+					Emitter.Errorf(TEXT("Input `%d` of type `%s` is not primitive."), InputComponentIndex, InputValues[InputIndex]->Type->GetSpelling().GetData());
+					continue;
+				}
+
+				// Determine the target type.
+				MIR::FTypePtr InputType = MIR::FPrimitiveType::GetVector(InputPrimitiveType->ScalarKind, MaterialExpressionConvertType::GetComponentCount(ConvertInput.Type));
+
+				// Cast the input value to the target type.
+				InputValues[InputIndex] = Emitter.EmitConstruct(InputType, InputValues[InputIndex]);
 			}
 
-			// Choose the specific component/subscript we want
-			OutComponents[OutputComponentIndex] = Emitter.EmitSubscript(CompiledInput, InputComponentIndex);
+			// Subscript the input value to the specified component index.
+			OutComponents[OutputComponentIndex] = Emitter.EmitSubscript(InputValues[InputIndex], InputComponentIndex);
 		}
 
-		// For any ComponentCompileResults that are still INDEX_NONE, set to ConvertOutput's default value
+		const int32 OutputNumComponents = MaterialExpressionConvertType::GetComponentCount(ConvertOutput.Type);
+
+		// For any component still unset, give assign it to the default value.
 		for (int32 OutputComponentIndex = 0; OutputComponentIndex < OutputNumComponents; ++OutputComponentIndex)
 		{
-			MIR::FValue*& ComponentCompileResult = OutComponents[OutputComponentIndex];
 			// If we don't have a compile result here, default it to a that component's default value
-			if (ComponentCompileResult == nullptr)
+			if (!OutComponents[OutputComponentIndex])
 			{
-				ComponentCompileResult = Emitter.EmitConstantFloatScalar(ConvertOutput.DefaultValue.Component(OutputComponentIndex));
+				OutComponents[OutputComponentIndex] = Emitter.EmitConstantFloatScalar(ConvertOutput.DefaultValue.Component(OutputComponentIndex));
 			}
 		}
 
+		// Finally create the output dimensional value by combining the output components.
 		MIR::FValue* OutValue = nullptr;
 		switch (OutputNumComponents)
 		{
-			case 1:
-				OutValue = OutComponents[0];
-				break;
-			case 2:
-				OutValue = Emitter.EmitVector2(OutComponents[0], OutComponents[1]);
-				break;
-			case 3:
-				OutValue = Emitter.EmitVector3(OutComponents[0], OutComponents[1], OutComponents[2]);
-				break;
-			case 4:
-				OutValue = Emitter.EmitVector4(OutComponents[0], OutComponents[1], OutComponents[2], OutComponents[3]);
-				break;
+			case 1: OutValue = OutComponents[0]; break;
+			case 2: OutValue = Emitter.EmitVector2(OutComponents[0], OutComponents[1]); break;
+			case 3: OutValue = Emitter.EmitVector3(OutComponents[0], OutComponents[1], OutComponents[2]); break;
+			case 4: OutValue = Emitter.EmitVector4(OutComponents[0], OutComponents[1], OutComponents[2], OutComponents[3]); break;
+			
 			default:
+				OutValue = MIR::FPoison::Get();
 				Emitter.Errorf(TEXT("Convert node has an invalid component count of %d"), OutputNumComponents);
-				continue;
+				break;
 		}
 
 		Emitter.Output(OutputIndex, OutValue);
diff --git a/Engine/Source/Runtime/Engine/Public/Materials/MaterialExpressionConvert.h b/Engine/Source/Runtime/Engine/Public/Materials/MaterialExpressionConvert.h
--- a/Engine/Source/Runtime/Engine/Public/Materials/MaterialExpressionConvert.h
+++ b/Engine/Source/Runtime/Engine/Public/Materials/MaterialExpressionConvert.h
@@ -17,10 +17,10 @@ class UEdGraphPin;
 UENUM()
 enum class EMaterialExpressionConvertType : uint8
 {
-	Float1,
-	Float2,
-	Float3,
-	Float4,
+	Scalar,
+	Vector2,
+	Vector3,
+	Vector4,
 };
 
 namespace MaterialExpressionConvertType
@@ -29,10 +29,10 @@ namespace MaterialExpressionConvertType
 	{
 		switch (InConvertType)
 		{
-			case EMaterialExpressionConvertType::Float1: return 1;
-			case EMaterialExpressionConvertType::Float2: return 2;
-			case EMaterialExpressionConvertType::Float3: return 3;
-			case EMaterialExpressionConvertType::Float4: return 4;
+			case EMaterialExpressionConvertType::Scalar: return 1;
+			case EMaterialExpressionConvertType::Vector2: return 2;
+			case EMaterialExpressionConvertType::Vector3: return 3;
+			case EMaterialExpressionConvertType::Vector4: return 4;
 			default: checkNoEntry(); return INDEX_NONE;
 		}
 	}
@@ -41,10 +41,10 @@ namespace MaterialExpressionConvertType
 	{
 		switch (InConvertType)
 		{
-			case EMaterialExpressionConvertType::Float1: return MCT_Float1;
-			case EMaterialExpressionConvertType::Float2: return MCT_Float2;
-			case EMaterialExpressionConvertType::Float3: return MCT_Float3;
-			case EMaterialExpressionConvertType::Float4: return MCT_Float4;
+			case EMaterialExpressionConvertType::Scalar: return MCT_Float1;
+			case EMaterialExpressionConvertType::Vector2: return MCT_Float2;
+			case EMaterialExpressionConvertType::Vector3: return MCT_Float3;
+			case EMaterialExpressionConvertType::Vector4: return MCT_Float4;
 			default: checkNoEntry(); return MCT_Float;
 		}
 	}
@@ -53,10 +53,10 @@ namespace MaterialExpressionConvertType
 	{
 		switch (InConvertType)
 		{
-			case EMaterialExpressionConvertType::Float1: return LOCTEXT("Float1", "Float1");
-			case EMaterialExpressionConvertType::Float2: return LOCTEXT("Float2", "Float2");
-			case EMaterialExpressionConvertType::Float3: return LOCTEXT("Float3", "Float3");
-			case EMaterialExpressionConvertType::Float4: return LOCTEXT("Float4", "Float4");
+			case EMaterialExpressionConvertType::Scalar:  return LOCTEXT("Scalar", "Scalar");
+			case EMaterialExpressionConvertType::Vector2: return LOCTEXT("Vector2", "Vector2");
+			case EMaterialExpressionConvertType::Vector3: return LOCTEXT("Vector3", "Vector3");
+			case EMaterialExpressionConvertType::Vector4: return LOCTEXT("Vector4", "Vector4");
 			default: checkNoEntry(); return INVTEXT("");
 		}
 	}
@@ -71,7 +71,7 @@ struct FMaterialExpressionConvertInput
 	FExpressionInput ExpressionInput;
 
 	UPROPERTY(EditAnywhere, Category = "MaterialExpressionConvert")
-	EMaterialExpressionConvertType Type = EMaterialExpressionConvertType::Float1;
+	EMaterialExpressionConvertType Type = EMaterialExpressionConvertType::Scalar;
 
 	/** Default Value used when this input has no incoming connection */
 	UPROPERTY(EditAnywhere, Category = "MaterialExpressionConvert")
@@ -84,7 +84,7 @@ struct FMaterialExpressionConvertOutput
 	GENERATED_BODY()
 
 	UPROPERTY(EditAnywhere, Category = "MaterialExpressionConvert")
-	EMaterialExpressionConvertType Type = EMaterialExpressionConvertType::Float1;
+	EMaterialExpressionConvertType Type = EMaterialExpressionConvertType::Scalar;
 
 	/** Default Value used when this output has no connection */
 	UPROPERTY(EditAnywhere, Category = "MaterialExpressionConvert")
