Add CVars to allow choosing which slice of a texture array we use for screenshots when in Mobile Multi-View, or to collate multiplie slices in a single image, vertically or horizontally.

diff --git a/Engine/Source/Runtime/Engine/Private/GameViewportClient.cpp b/Engine/Source/Runtime/Engine/Private/GameViewportClient.cpp
--- a/Engine/Source/Runtime/Engine/Private/GameViewportClient.cpp
+++ b/Engine/Source/Runtime/Engine/Private/GameViewportClient.cpp
@@ -150,6 +150,23 @@ static FAutoConsoleCommandWithWorld SplitScreenSwapControllersCommand
 		})
 );
 
+static TAutoConsoleVariable<int32> CVarScreenshotTextureArrayMode(
+	TEXT("r.Screenshot.TextureArrayMode"),
+	0,
+	TEXT("Choose how screenshots should be handled when the viewport is using a texture array. Collation is supported only for non-HDR, non-UI screenshots.")
+	TEXT(" 0: Include only a single array slice (default, controlled by r.Screenshot.TextureArrayIndex)\n")
+	TEXT(" 1: Collate all array slices horizontally (left-to-right)\n")
+	TEXT(" 2: Collate all array slices vertically (top-to-bottom)\n")
+	TEXT("When collating, the screenshot resolution is used for each individual array slice, so requesting a 400x400 screenshot from\n")
+	TEXT("a viewport with two slices and then collating horizontally will yield a final screenshot that's 800x400."),
+	ECVF_Default);
+
+static TAutoConsoleVariable<int32> CVarScreenshotTextureArrayIndex(
+	TEXT("r.Screenshot.TextureArrayIndex"),
+	0,
+	TEXT("Texture array slice to use for screenshots when r.Screenshot.TextureArrayMode = 0"),
+	ECVF_Default);
+
 #if CSV_PROFILER_STATS
 struct FCsvLocalPlayer
 {
@@ -2202,10 +2219,96 @@ bool UGameViewportClient::ProcessScreenShots(FViewport* InViewport)
 			bScreenshotSuccessful = GetViewportScreenShotHDR(InViewport, BitmapHDR, FIntRect(0, 0, Size.X, Size.Y));
 			ConvertPixelDataToSCRGB(BitmapHDR, ViewportOutputFormat);
 		}
+		else if (InViewport->GetShaderResourceTexture() && InViewport->GetShaderResourceTexture()->GetDesc().Dimension == ETextureDimension::Texture2DArray)
+		{
+			// Handle Texture2DArray, we need to determine which array slices to include
+			FReadSurfaceDataFlags ReadSurfaceFlags = FReadSurfaceDataFlags();
+			const int32 ScreenshotTextureArrayMode = CVarScreenshotTextureArrayMode.GetValueOnAnyThread();
+			const int32 ArraySize = InViewport->GetShaderResourceTexture()->GetDesc().ArraySize;
+
+			ensureAlwaysMsgf(ScreenshotTextureArrayMode >= 0 && ScreenshotTextureArrayMode < 3, TEXT("r.Screenshot.TextureArrayMode %d is not a valid mode, defaulting to mode 0"), ScreenshotTextureArrayMode);
+
+			// Screenshot all slices, then collate horizontally or vertically
+			if (ScreenshotTextureArrayMode == 1 || ScreenshotTextureArrayMode == 2)
+			{
+				TArray<TArray<FColor>> BitmapArray;
+				BitmapArray.SetNum(ArraySize);
+
+				for (int CurrentArrayIndex = 0; CurrentArrayIndex < ArraySize; CurrentArrayIndex++)
+				{
+					ReadSurfaceFlags.SetArrayIndex(CurrentArrayIndex);
+					bScreenshotSuccessful = GetViewportScreenShot(InViewport, BitmapArray[CurrentArrayIndex], FIntRect(0, 0, Size.X, Size.Y), ReadSurfaceFlags);
+
+					if (!bScreenshotSuccessful)
+					{
+						break;
+					}
+				}
+
+				if (bScreenshotSuccessful)
+				{
+					const int32 NumPixelsPerSlice = Size.X * Size.Y;
+					const int32 NumPixelsTotal = NumPixelsPerSlice * ArraySize;
+
+					Bitmap.Reset();
+					Bitmap.SetNumUninitialized(NumPixelsTotal);
+
+					// Horizontal collation
+					if (ScreenshotTextureArrayMode == 1)
+					{
+						for (int CurrentPixel = 0; CurrentPixel < NumPixelsTotal; CurrentPixel++)
+						{
+							// Complex case, need to alternate between array entries for each horizontal line
+							const int32 CurrentRow = CurrentPixel / (ArraySize * Size.X);
+							const int32 CurrentColumn = CurrentPixel % Size.X;
+								 
+							const int32 ArrayIndex = (CurrentPixel / Size.X) % ArraySize;
+							const int32 PixelIndex = CurrentColumn + CurrentRow * Size.X;
+							Bitmap[CurrentPixel] = BitmapArray[ArrayIndex][PixelIndex];
+						}
+
+						Size.X *= ArraySize;
+						GScreenshotResolutionX = Size.X;
+					}
+
+					// Vertical collation
+					else
+					{
+						for (int CurrentPixel = 0; CurrentPixel < NumPixelsTotal; CurrentPixel++)
+						{
+							// Simple case, just concatenate all bitmaps
+							const int32 ArrayIndex = CurrentPixel / NumPixelsPerSlice;
+							const int32 PixelIndex = CurrentPixel % NumPixelsPerSlice;
+							Bitmap[CurrentPixel] = BitmapArray[ArrayIndex][PixelIndex];
+						}
+
+						Size.Y *= ArraySize;
+						GScreenshotResolutionY = Size.Y;
+					}
+				}
+			}
+
+			// Screenshot a specific array slice (ScreenshotTextureArrayMode == 0 or other invalid values)
+			else
+			{
+				int32 ArrayIndex = CVarScreenshotTextureArrayIndex.GetValueOnAnyThread();
+				const int32 ClampedArrayIndex = FMath::Clamp(ArrayIndex, 0, ArraySize - 1);
+
+				if (ArrayIndex != ClampedArrayIndex)
+				{
+					ensureAlwaysMsgf(false, TEXT("r.Screenshot.TextureArrayIndex %d is out of range for the current viewport, which has %d texture array slices, defaulting to slice %d"), ArrayIndex, ArraySize, ClampedArrayIndex);
+					ArrayIndex = ClampedArrayIndex;
+				}
+
+				ReadSurfaceFlags.SetArrayIndex(ArrayIndex);
+				bScreenshotSuccessful = GetViewportScreenShot(InViewport, Bitmap, FIntRect(0, 0, Size.X, Size.Y), ReadSurfaceFlags);
+			}
+		}
 		else
 		{
 			bScreenshotSuccessful = GetViewportScreenShot(InViewport, Bitmap, FIntRect(0, 0, Size.X, Size.Y));
 		}
+		
 
 		if (bScreenshotSuccessful)
 		{
diff --git a/Engine/Source/Runtime/Engine/Private/UnrealClient.cpp b/Engine/Source/Runtime/Engine/Private/UnrealClient.cpp
--- a/Engine/Source/Runtime/Engine/Private/UnrealClient.cpp
+++ b/Engine/Source/Runtime/Engine/Private/UnrealClient.cpp
@@ -2397,10 +2397,10 @@ void FViewport::SetInitialSize( FIntPoint InitialSizeXY )
 }
 
 
-ENGINE_API bool GetViewportScreenShot(FViewport* Viewport, TArray<FColor>& Bitmap, const FIntRect& ViewRect /*= FIntRect()*/)
+ENGINE_API bool GetViewportScreenShot(FViewport* Viewport, TArray<FColor>& Bitmap, const FIntRect& ViewRect /*= FIntRect()*/, const FReadSurfaceDataFlags& ReadSurfaceFlags /*= FReadSurfaceDataFlags()*/)
 {
 	// Read the contents of the viewport into an array.
-	if (Viewport->ReadPixels(Bitmap, FReadSurfaceDataFlags(), ViewRect))
+	if (Viewport->ReadPixels(Bitmap, ReadSurfaceFlags, ViewRect))
 	{
 		check(Bitmap.Num() == ViewRect.Area() || (Bitmap.Num() == Viewport->GetRenderTargetTextureSizeXY().X * Viewport->GetRenderTargetTextureSizeXY().Y));
 		return true;
@@ -2409,10 +2409,10 @@ ENGINE_API bool GetViewportScreenShot(FViewport* Viewport, TArray<FColor>& Bitma
 	return false;
 }
 
-ENGINE_API bool GetViewportScreenShotHDR(FViewport* Viewport, TArray<FLinearColor>& Bitmap, const FIntRect& ViewRect /*= FIntRect()*/)
+ENGINE_API bool GetViewportScreenShotHDR(FViewport* Viewport, TArray<FLinearColor>& Bitmap, const FIntRect& ViewRect /*= FIntRect()*/, const FReadSurfaceDataFlags& ReadSurfaceFlags /*= FReadSurfaceDataFlags(RCM_MinMax)*/)
 {
 	// Read the contents of the viewport into an array.
-	if (Viewport->ReadLinearColorPixels(Bitmap, FReadSurfaceDataFlags(RCM_MinMax), ViewRect))
+	if (Viewport->ReadLinearColorPixels(Bitmap, ReadSurfaceFlags, ViewRect))
 	{
 		check(Bitmap.Num() == ViewRect.Area() || (Bitmap.Num() == Viewport->GetRenderTargetTextureSizeXY().X * Viewport->GetRenderTargetTextureSizeXY().Y));
 		return true;
diff --git a/Engine/Source/Runtime/Engine/Public/UnrealClient.h b/Engine/Source/Runtime/Engine/Public/UnrealClient.h
--- a/Engine/Source/Runtime/Engine/Public/UnrealClient.h
+++ b/Engine/Source/Runtime/Engine/Public/UnrealClient.h
@@ -860,8 +860,8 @@ extern ENGINE_API bool IsCtrlDown(FViewport* Viewport);
 extern ENGINE_API bool IsShiftDown(FViewport* Viewport);
 extern ENGINE_API bool IsAltDown(FViewport* Viewport);
 
-extern ENGINE_API bool GetViewportScreenShot(FViewport* Viewport, TArray<FColor>& Bitmap, const FIntRect& ViewRect = FIntRect());
-extern ENGINE_API bool GetViewportScreenShotHDR(FViewport* Viewport, TArray<FLinearColor>& Bitmap, const FIntRect& ViewRect = FIntRect());
+extern ENGINE_API bool GetViewportScreenShot(FViewport* Viewport, TArray<FColor>& Bitmap, const FIntRect& ViewRect = FIntRect(), const FReadSurfaceDataFlags& ReadSurfaceFlags = FReadSurfaceDataFlags());
+extern ENGINE_API bool GetViewportScreenShotHDR(FViewport* Viewport, TArray<FLinearColor>& Bitmap, const FIntRect& ViewRect = FIntRect(), const FReadSurfaceDataFlags& ReadSurfaceFlags = FReadSurfaceDataFlags(RCM_MinMax));
 extern ENGINE_API bool GetHighResScreenShotInput(const TCHAR* Cmd, FOutputDevice& Ar, uint32& OutXRes, uint32& OutYRes, float& OutResMult, FIntRect& OutCaptureRegion, bool& OutShouldEnableMask, bool& OutDumpBufferVisualizationTargets, bool& OutCaptureHDR, FString& OutFilenameOverride, bool& OutUseDateTimeAsFileName);
 
 /** Tracks the viewport client that should process the stat command, can be NULL */
