[shaders] reset job status object in path that re-runs compilation jobs to dump debug info. rarely this could result in a job that was previously a duplicate actually executing, and so triggering the check in the stats that expects a zero compile time for any jobs that were skipped for any reason (usually due to cache hits or being a duplicate of another job).

diff --git a/Engine/Source/Runtime/Engine/Private/ShaderCompiler/ShaderCompilerUtil.cpp b/Engine/Source/Runtime/Engine/Private/ShaderCompiler/ShaderCompilerUtil.cpp
--- a/Engine/Source/Runtime/Engine/Private/ShaderCompiler/ShaderCompilerUtil.cpp
+++ b/Engine/Source/Runtime/Engine/Private/ShaderCompiler/ShaderCompilerUtil.cpp
@@ -556,6 +556,8 @@ static bool ReadSingleJob(FShaderCompileJob* CurrentJob, FArchive& WorkerOutputF
 	{
 		// Build debug info path and create the directory if it doesn't already exist
 		CurrentJob->Input.DumpDebugInfoPath = GShaderCompilingManager->CreateShaderDebugInfoPath(CurrentJob->Input);
+		// Reset job status object to ensure both job state and "duplicate" flag are reset.
+		CurrentJob->JobStatusPtr.Reset();
 
 		// We failed to compile a shader, we will retry and dump the shader source. so increment the number of shader sources we've dumped so far.
 		GShaderCompilingManager->IncrementNumDumpedShaderSources();
