[Engine][Animation] Root Space Additive and Layered Blend Per Bone Options
- Add Options for ApplyMeshSpaceAdditive and LayeredBlendPerBone to Blend in "Root Space"
- Essentially ignore the root rotation when calculating the additive

diff --git a/Engine/Source/Runtime/AnimGraphRuntime/Private/AnimNodes/AnimNode_LayeredBoneBlend.cpp b/Engine/Source/Runtime/AnimGraphRuntime/Private/AnimNodes/AnimNode_LayeredBoneBlend.cpp
--- a/Engine/Source/Runtime/AnimGraphRuntime/Private/AnimNodes/AnimNode_LayeredBoneBlend.cpp
+++ b/Engine/Source/Runtime/AnimGraphRuntime/Private/AnimNodes/AnimNode_LayeredBoneBlend.cpp
@@ -257,6 +257,11 @@ void FAnimNode_LayeredBoneBlend::Evaluate_AnyThread(FPoseContext& Output)
 		if (bMeshSpaceRotationBlend)
 		{
 			BlendFlags |= FAnimationRuntime::EBlendPosesPerBoneFilterFlags::MeshSpaceRotation;
+
+			if (bRootSpaceRotationBlend)
+			{
+				BlendFlags |= FAnimationRuntime::EBlendPosesPerBoneFilterFlags::RootSpaceRotation;
+			}
 		}
 		if (bMeshSpaceScaleBlend)
 		{
diff --git a/Engine/Source/Runtime/AnimGraphRuntime/Public/AnimNodes/AnimNode_LayeredBoneBlend.h b/Engine/Source/Runtime/AnimGraphRuntime/Public/AnimNodes/AnimNode_LayeredBoneBlend.h
--- a/Engine/Source/Runtime/AnimGraphRuntime/Public/AnimNodes/AnimNode_LayeredBoneBlend.h
+++ b/Engine/Source/Runtime/AnimGraphRuntime/Public/AnimNodes/AnimNode_LayeredBoneBlend.h
@@ -93,6 +93,10 @@ public:
 	UPROPERTY(EditAnywhere, BlueprintReadWrite, Category=Config)
 	bool bMeshSpaceRotationBlend;
 
+	/** Whether to blend bone rotations in root space or in mesh space */
+	UPROPERTY(EditAnywhere, BlueprintReadWrite, Category=Config, Meta = (Editcondition = bMeshSpaceRotationBlend))
+	bool bRootSpaceRotationBlend;
+
 	/** Whether to blend bone scales in mesh space or in local space */
 	UPROPERTY(EditAnywhere, BlueprintReadWrite, Category=Config)
 	bool bMeshSpaceScaleBlend;
diff --git a/Engine/Source/Runtime/Engine/Classes/Animation/AnimNode_ApplyMeshSpaceAdditive.h b/Engine/Source/Runtime/Engine/Classes/Animation/AnimNode_ApplyMeshSpaceAdditive.h
--- a/Engine/Source/Runtime/Engine/Classes/Animation/AnimNode_ApplyMeshSpaceAdditive.h
+++ b/Engine/Source/Runtime/Engine/Classes/Animation/AnimNode_ApplyMeshSpaceAdditive.h
@@ -19,6 +19,9 @@ struct FAnimNode_ApplyMeshSpaceAdditive : public FAnimNode_Base
 	UPROPERTY(EditAnywhere, BlueprintReadWrite, Category = Links)
 	FPoseLink Additive;
 
+	UPROPERTY(EditAnywhere, BlueprintReadWrite, Category = Settings)
+	uint8 bRootSpaceAdditive : 1;
+
 	/** The data type used to control the alpha blending of the additive pose.
 		Note: Changing this value will disconnect alpha input pins.
 	*/
diff --git a/Engine/Source/Runtime/Engine/Private/Animation/AnimNode_ApplyMeshSpaceAdditive.cpp b/Engine/Source/Runtime/Engine/Private/Animation/AnimNode_ApplyMeshSpaceAdditive.cpp
--- a/Engine/Source/Runtime/Engine/Private/Animation/AnimNode_ApplyMeshSpaceAdditive.cpp
+++ b/Engine/Source/Runtime/Engine/Private/Animation/AnimNode_ApplyMeshSpaceAdditive.cpp
@@ -91,7 +91,17 @@ void FAnimNode_ApplyMeshSpaceAdditive::Evaluate_AnyThread(FPoseContext& Output)
 
 		FAnimationPoseData BaseAnimationPoseData(Output);
 		const FAnimationPoseData AdditiveAnimationPoseData(AdditiveEvalContext);
-		FAnimationRuntime::AccumulateAdditivePose(BaseAnimationPoseData, AdditiveAnimationPoseData, ActualAlpha, AAT_RotationOffsetMeshSpace);
+
+		
+		if (bRootSpaceAdditive)
+		{
+			FAnimationRuntime::AccumulateRootSpaceRotationAdditiveToLocalPose(BaseAnimationPoseData, AdditiveAnimationPoseData, ActualAlpha);
+		}
+		else
+		{
+			FAnimationRuntime::AccumulateAdditivePose(BaseAnimationPoseData, AdditiveAnimationPoseData, ActualAlpha, AAT_RotationOffsetMeshSpace);
+		}
+
 		Output.Pose.NormalizeRotations();
 	}
 	else
diff --git a/Engine/Source/Runtime/Engine/Private/Animation/AnimationRuntime.cpp b/Engine/Source/Runtime/Engine/Private/Animation/AnimationRuntime.cpp
--- a/Engine/Source/Runtime/Engine/Private/Animation/AnimationRuntime.cpp
+++ b/Engine/Source/Runtime/Engine/Private/Animation/AnimationRuntime.cpp
@@ -1213,6 +1213,19 @@ void FAnimationRuntime::AccumulateMeshSpaceRotationAdditiveToLocalPose(FAnimatio
 	AccumulateAdditivePose(BaseAnimationPoseData, MeshSpaceRotationAdditiveAnimationPoseData, Weight, EAdditiveAnimationType::AAT_RotationOffsetMeshSpace);
 }
 
+void FAnimationRuntime::AccumulateRootSpaceRotationAdditiveToLocalPose(FAnimationPoseData& BaseAnimationPoseData, const FAnimationPoseData& RootSpaceRotationAdditiveAnimationPoseData, float Weight)
+{
+	AccumulateRootSpaceAdditivePoseInternal(BaseAnimationPoseData.GetPose(), RootSpaceRotationAdditiveAnimationPoseData.GetPose(), Weight);
+
+	// if curve exists, accumulate with the weight, 
+	BaseAnimationPoseData.GetCurve().Accumulate(RootSpaceRotationAdditiveAnimationPoseData.GetCurve(), Weight);
+
+	UE::Anim::Attributes::AccumulateAttributes(RootSpaceRotationAdditiveAnimationPoseData.GetAttributes(), BaseAnimationPoseData.GetAttributes(), Weight, EAdditiveAnimationType::AAT_RotationOffsetMeshSpace);
+
+	// normalize
+	BaseAnimationPoseData.GetPose().NormalizeRotations();
+}
+
 void FAnimationRuntime::AccumulateAdditivePose(FAnimationPoseData& BaseAnimationPoseData, const FAnimationPoseData& AdditiveAnimationPoseData, float Weight, enum EAdditiveAnimationType AdditiveType)
 {
 	if (AdditiveType == AAT_RotationOffsetMeshSpace)
@@ -1295,6 +1308,29 @@ void FAnimationRuntime::AccumulateMeshSpaceRotationAdditiveToLocalPoseInternal(F
 	}
 }
 
+ENGINE_API void FAnimationRuntime::AccumulateRootSpaceAdditivePoseInternal(FCompactPose& BasePose, const FCompactPose& RootSpaceAdditivePose, float Weight)
+{
+	if (FAnimWeight::IsRelevant(Weight))
+	{		
+		// Cache off the Root Transform and Zero it Out
+		FCompactPoseBoneIndex BoneIndex(0);
+		FTransform RootTransform = BasePose[BoneIndex];
+		BasePose[BoneIndex] = FTransform::Identity;
+
+		// Convert base pose from local space to mesh space rotation.
+		FAnimationRuntime::ConvertPoseToMeshRotation(BasePose);
+
+		// Add MeshSpaceRotAdditive to it
+		FAnimationRuntime::AccumulateLocalSpaceAdditivePoseInternal(BasePose, RootSpaceAdditivePose, Weight);
+
+		// Convert back to local space
+		FAnimationRuntime::ConvertMeshRotationPoseToLocalSpace(BasePose);
+
+		// Restore the original Root Transform
+		BasePose[BoneIndex] = RootTransform;
+	}
+}
+
 void FAnimationRuntime::MirrorCurves(FBlendedCurve& Curves, const UMirrorDataTable& MirrorDataTable)
 {
 	for(const TPair<FName, FName>& MirrorPair : MirrorDataTable.CurveToMirrorCurveMap)
@@ -1865,8 +1901,9 @@ void FAnimationRuntime::BlendPosesPerBoneFilter(FCompactPose& BasePose, const TA
 
 	const bool bMeshSpaceRotationBlend = EnumHasAnyFlags(BlendFlags, EBlendPosesPerBoneFilterFlags::MeshSpaceRotation);
 	const bool bMeshSpaceScaleBlend = EnumHasAnyFlags(BlendFlags, EBlendPosesPerBoneFilterFlags::MeshSpaceScale);
+	const bool bRootSpaceRotationBlend = EnumHasAnyFlags(BlendFlags, EBlendPosesPerBoneFilterFlags::RootSpaceRotation);
 
-	if (bMeshSpaceRotationBlend)
+	if (bRootSpaceRotationBlend || bMeshSpaceRotationBlend)
 	{
 		SourceRotations.Reset();
 		SourceRotations.AddUninitialized(NumBones);
@@ -1929,8 +1966,105 @@ void FAnimationRuntime::BlendPosesPerBoneFilter(FCompactPose& BasePose, const TA
 
 	const FBoneContainer& BoneContainer = BasePose.GetBoneContainer();
 
+	if (bRootSpaceRotationBlend && bMeshSpaceScaleBlend)
+	{
+		for (const FCompactPoseBoneIndex BoneIndex : BasePose.ForEachBoneIndex())
+		{
+			const int32 PoseIndex = BoneBlendWeights[BoneIndex.GetInt()].SourceIndex;
+			const FCompactPoseBoneIndex ParentIndex = BoneContainer.GetParentBoneIndex(BoneIndex);
+
+			if (ParentIndex != INDEX_NONE)
+			{
+				AccumulateMeshSpaceRotation(PoseIndex, BoneIndex, SourceRotations[ParentIndex], TargetRotations[ParentIndex]);
+				AccumulateMeshSpaceScale(PoseIndex, BoneIndex, SourceScales[ParentIndex], TargetScales[ParentIndex]);
+			}
+			else
+			{
+				AccumulateMeshSpaceScale(PoseIndex, BoneIndex, FVector(1.0f), FVector(1.0f));
+			}
+
+			const FTransform& BaseAtom = BasePose[BoneIndex];
+			const FTransform& TargetAtom = BlendPoses[PoseIndex][BoneIndex];
+			FTransform BlendAtom;
+
+			const float BlendWeight = FMath::Clamp(BoneBlendWeights[BoneIndex.GetInt()].BlendWeight, 0.f, 1.f);
+			MaxPoseWeights[PoseIndex] = FMath::Max(MaxPoseWeights[PoseIndex], BlendWeight);
+
+			if (!FAnimWeight::IsRelevant(BlendWeight))
+			{
+				BlendAtom = BaseAtom;
+				BlendRotations[BoneIndex] = SourceRotations[BoneIndex];
+				BlendScales[BoneIndex] = SourceScales[BoneIndex];
+			}
+			else if (FAnimWeight::IsFullWeight(BlendWeight))
+			{
+				BlendAtom = TargetAtom;
+				BlendRotations[BoneIndex] = TargetRotations[BoneIndex];
+				BlendScales[BoneIndex] = TargetScales[BoneIndex];
+			}
+			else
+			{
+				BlendAtom = BaseAtom;
+				BlendAtom.BlendWith(TargetAtom, BlendWeight);
+				LerpMeshSpaceRotation(BoneIndex, BlendWeight);
+				LerpMeshSpaceScale(BoneIndex, BlendWeight);
+			}
+
+			if (ParentIndex != INDEX_NONE)
+			{
+				ConvertMeshToLocalSpaceRotation(BlendAtom, ParentIndex, BoneIndex);
+				ConvertMeshToLocalSpaceScale(BlendAtom, ParentIndex, BoneIndex);
+			}
+
+			OutPose[BoneIndex] = BlendAtom;
+		}
+	}
+	else if (bRootSpaceRotationBlend)
+	{
+		for (const FCompactPoseBoneIndex BoneIndex : BasePose.ForEachBoneIndex())
+		{
+			const int32 PoseIndex = BoneBlendWeights[BoneIndex.GetInt()].SourceIndex;
+			const FCompactPoseBoneIndex ParentIndex = BoneContainer.GetParentBoneIndex(BoneIndex);
+
+			if (ParentIndex != INDEX_NONE)
+			{
+				AccumulateMeshSpaceRotation(PoseIndex, BoneIndex, SourceRotations[ParentIndex], TargetRotations[ParentIndex]);
+			}
+
+			const FTransform& BaseAtom = BasePose[BoneIndex];
+			const FTransform& TargetAtom = BlendPoses[PoseIndex][BoneIndex];
+			FTransform BlendAtom;
+
+			const float BlendWeight = FMath::Clamp(BoneBlendWeights[BoneIndex.GetInt()].BlendWeight, 0.f, 1.f);
+			MaxPoseWeights[PoseIndex] = FMath::Max(MaxPoseWeights[PoseIndex], BlendWeight);
+
+			if (!FAnimWeight::IsRelevant(BlendWeight))
+			{
+				BlendAtom = BaseAtom;
+				BlendRotations[BoneIndex] = SourceRotations[BoneIndex];
+			}
+			else if (FAnimWeight::IsFullWeight(BlendWeight))
+			{
+				BlendAtom = TargetAtom;
+				BlendRotations[BoneIndex] = TargetRotations[BoneIndex];
+			}
+			else
+			{
+				BlendAtom = BaseAtom;
+				BlendAtom.BlendWith(TargetAtom, BlendWeight);
+				LerpMeshSpaceRotation(BoneIndex, BlendWeight);
+			}
+
+			if (ParentIndex != INDEX_NONE)
+			{
+				ConvertMeshToLocalSpaceRotation(BlendAtom, ParentIndex, BoneIndex);
+			}
+
+			OutPose[BoneIndex] = BlendAtom;
+		}
+	}
 	// blend poses with both mesh space rotation and scaling (we assume uniform scale)
-	if (bMeshSpaceRotationBlend && bMeshSpaceScaleBlend)
+	else if (bMeshSpaceRotationBlend && bMeshSpaceScaleBlend)
 	{
 #if INTEL_ISPC
 		if (bAnim_BlendPosesPerBoneFilter_ISPC_Enabled)
diff --git a/Engine/Source/Runtime/Engine/Public/AnimationRuntime.h b/Engine/Source/Runtime/Engine/Public/AnimationRuntime.h
--- a/Engine/Source/Runtime/Engine/Public/AnimationRuntime.h
+++ b/Engine/Source/Runtime/Engine/Public/AnimationRuntime.h
@@ -351,6 +351,7 @@ public:
 		None = 0,
 		MeshSpaceRotation = (1 << 0),	//! Blend bone rotations in mesh space instead of local space
 		MeshSpaceScale    = (1 << 1),	//! Blend bone scales in mesh space instead of local space
+		RootSpaceRotation = (1 << 2),	//! Blend bone rotations in root space instead of local space
 	};
 
 	/**
@@ -452,6 +453,10 @@ private:
 
 	/** Accumulate a MeshSpaceRotation Additive pose to a local pose. Rotations are NOT normalized */
 	static ENGINE_API void AccumulateMeshSpaceRotationAdditiveToLocalPoseInternal(FCompactPose& BasePose, const FCompactPose& MeshSpaceRotationAdditive, float Weight);
+
+	/** Accumulate a RootSpaceRotation Additive pose to a local pose. Rotations are NOT normalized */
+	static ENGINE_API void AccumulateRootSpaceAdditivePoseInternal(FCompactPose& BasePose, const FCompactPose& RootSpaceAdditivePose, float Weight);
+
 public:
 
 	/** Accumulates weighted AdditivePose to BasePose. Rotations are NOT normalized. */
@@ -466,6 +471,7 @@ public:
 
 	static ENGINE_API void AccumulateMeshSpaceRotationAdditiveToLocalPose(FAnimationPoseData& BaseAnimationPoseData, const FAnimationPoseData& MeshSpaceRotationAdditiveAnimationPoseData, float Weight);
 
+	static ENGINE_API void AccumulateRootSpaceRotationAdditiveToLocalPose(FAnimationPoseData& BaseAnimationPoseData, const FAnimationPoseData& MeshSpaceRotationAdditiveAnimationPoseData, float Weight);
 
 	/** Lerp for FCompactPose. Stores results in PoseA. Performs PoseA = Lerp(PoseA, PoseB, Alpha); */
 	static ENGINE_API void LerpPoses(FCompactPose& PoseA, const FCompactPose& PoseB, FBlendedCurve& CurveA, const FBlendedCurve& CurveB, float Alpha);
