Animation: Speculative fix for a null attach component in UpdateCachedSocketAttachmentFlag

diff --git a/Engine/Source/Runtime/Engine/Private/Components/SkinnedMeshComponent.cpp b/Engine/Source/Runtime/Engine/Private/Components/SkinnedMeshComponent.cpp
--- a/Engine/Source/Runtime/Engine/Private/Components/SkinnedMeshComponent.cpp
+++ b/Engine/Source/Runtime/Engine/Private/Components/SkinnedMeshComponent.cpp
@@ -3290,7 +3290,7 @@ void USkinnedMeshComponent::UpdateCachedSocketAttachmentFlag()
 	bHasSocketAttachments = false;
 	for (const USceneComponent* AttachedComponent : GetAttachChildren())
 	{
-		if (AttachedComponent->GetAttachSocketName() != NAME_None)
+		if (AttachedComponent && AttachedComponent->GetAttachSocketName() != NAME_None)
 		{
 			bHasSocketAttachments = true;
 			break;
