Added a delegate to be able to react when a montage section changes

#rnx

diff --git a/Engine/Source/Runtime/Engine/Classes/Animation/AnimInstance.h b/Engine/Source/Runtime/Engine/Classes/Animation/AnimInstance.h
--- a/Engine/Source/Runtime/Engine/Classes/Animation/AnimInstance.h
+++ b/Engine/Source/Runtime/Engine/Classes/Animation/AnimInstance.h
@@ -76,6 +76,8 @@ DECLARE_DELEGATE_OneParam(FOnMontageStarted, UAnimMontage*)
 DECLARE_DELEGATE_TwoParams(FOnMontageEnded, UAnimMontage*, bool /*bInterrupted*/)
 DECLARE_DELEGATE_TwoParams(FOnMontageBlendingOutStarted, UAnimMontage*, bool /*bInterrupted*/)
 DECLARE_DELEGATE_OneParam(FOnMontageBlendedInEnded, UAnimMontage*)
+DECLARE_DELEGATE_ThreeParams(FOnMontageSectionChanged, UAnimMontage*, FName /*SectionName*/ , bool /*bLooped*/)
+
 /**
 * Delegate for when Montage is started
 */
@@ -103,7 +105,10 @@ DECLARE_DYNAMIC_MULTICAST_DELEGATE_TwoParams(FOnMontageBlendingOutStartedMCDeleg
 /** Delegate for when Montage finished blending in */
 DECLARE_DYNAMIC_MULTICAST_DELEGATE_OneParam(FOnMontageBlendedInEndedMCDelegate, UAnimMontage*, Montage);
 
-/** Delegate that native code can hook to to provide additional transition logic */
+/** Delegate for when Montage changes sections */
+DECLARE_DYNAMIC_MULTICAST_DELEGATE_ThreeParams(FOnMontageSectionChangedMCDelegate, UAnimMontage*, Montage, FName, SectionName, bool, bLooped);
+
+/** Delegate that native code can hook to provide additional transition logic */
 DECLARE_DELEGATE_RetVal(bool, FCanTakeTransition);
 
 /** Delegate that native code can hook into to handle state entry/exit */
@@ -235,6 +240,31 @@ struct FQueuedMontageEndedEvent
 	{}
 };
 
+/** Helper struct to store a Queued Montage Ended event. */
+struct FQueuedMontageSectionChangedEvent
+{
+	TObjectPtr<class UAnimMontage> Montage;
+	int32 MontageInstanceID;
+	FName SectionName;
+	bool bLooped;
+	FOnMontageSectionChanged Delegate;
+
+	FQueuedMontageSectionChangedEvent()
+		: Montage(nullptr)
+		, MontageInstanceID(INDEX_NONE)
+		, SectionName(NAME_None)
+		, bLooped(false)
+	{}
+
+	FQueuedMontageSectionChangedEvent(class UAnimMontage* InMontage, int32 InInstanceID, FName SectionName, bool InbLooped, FOnMontageSectionChanged InDelegate)
+		: Montage(InMontage)
+		, MontageInstanceID(InInstanceID)
+		, SectionName(SectionName)
+		, bLooped(InbLooped)
+		, Delegate(InDelegate)
+	{}
+};
+
 /** Binding allowing native transition rule evaluation */
 struct FNativeTransitionBinding
 {
@@ -730,6 +760,10 @@ public:
 	UPROPERTY(BlueprintAssignable)
 	FOnAllMontageInstancesEndedMCDelegate OnAllMontageInstancesEnded;
 
+	/** Called when a montage section changes */
+	UPROPERTY(BlueprintAssignable)
+	FOnMontageSectionChangedMCDelegate OnMontageSectionChanged;
+	
 	/*********************************************************************************************
 	* AnimMontage native C++ interface
 	********************************************************************************************* */
@@ -748,6 +782,12 @@ public:
 	If Montage reference is NULL, it will pick the first active montage found. */
 	ENGINE_API FOnMontageBlendingOutStarted* Montage_GetBlendingOutDelegate(UAnimMontage* Montage = NULL);
 
+	ENGINE_API void Montage_SetSectionChangedDelegate(FOnMontageSectionChanged& InOnMontageSectionChanged, UAnimMontage* Montage = NULL);
+	
+	/** Get pointer to SectionChanged delegate for Montage.
+	If Montage reference is NULL, it will pick the first active montage found. */
+	ENGINE_API FOnMontageSectionChanged* Montage_GetSectionChangedDelegate(UAnimMontage* Montage = NULL);
+	
 	/** Get next sectionID for given section ID */
 	ENGINE_API int32 Montage_GetNextSectionID(const UAnimMontage* Montage, int32 const & CurrentSectionID) const;
 
@@ -937,6 +977,9 @@ public:
 	/** Queue a Montage Ended Event to be triggered. */
 	ENGINE_API void QueueMontageEndedEvent(const FQueuedMontageEndedEvent& MontageEndedEvent);
 
+	/** Queue a Montage Section Changed Event to be triggered. */
+	ENGINE_API void QueueMontageSectionChangedEvent(const FQueuedMontageSectionChangedEvent& MontageSectionChangedEvent);
+	
 private:
 	/** Trigger queued Montage events. */
 	ENGINE_API void TriggerQueuedMontageEvents();
@@ -950,6 +993,9 @@ private:
 	/** Queued Montage Ended Events */
 	TArray<FQueuedMontageEndedEvent> QueuedMontageEndedEvents;
 
+	/** Queued Montage Section Changed Events */
+	TArray<FQueuedMontageSectionChangedEvent> QueuedMontageSectionChangedEvents;
+	
 	/** Trigger a Montage BlendingOut event */
 	ENGINE_API void TriggerMontageBlendingOutEvent(const FQueuedMontageBlendingOutEvent& MontageBlendingOutEvent);
 
@@ -959,6 +1005,9 @@ private:
 	/** Trigger a Montage Ended event */
 	ENGINE_API void TriggerMontageEndedEvent(const FQueuedMontageEndedEvent& MontageEndedEvent);
 
+	/** Trigger a Montage Ended event */
+	ENGINE_API void TriggerMontageSectionChangedEvent(const FQueuedMontageSectionChangedEvent& MontageSectionChangedEvent);
+	
 public:
 
 #if DO_CHECK
diff --git a/Engine/Source/Runtime/Engine/Classes/Animation/AnimMontage.h b/Engine/Source/Runtime/Engine/Classes/Animation/AnimMontage.h
--- a/Engine/Source/Runtime/Engine/Classes/Animation/AnimMontage.h
+++ b/Engine/Source/Runtime/Engine/Classes/Animation/AnimMontage.h
@@ -195,6 +195,13 @@ DECLARE_DELEGATE_TwoParams( FOnMontageBlendingOutStarted, class UAnimMontage*, b
 
 DECLARE_DELEGATE_OneParam(FOnMontageBlendedInEnded, class UAnimMontage*)
 
+/**
+ * Delegate for when Montage changes section.
+ *
+ * bLooped = true if a requested to change to the current section.
+ */
+DECLARE_DELEGATE_ThreeParams( FOnMontageSectionChanged, class UAnimMontage*, FName /*SectionName*/ , bool /*bInterrupted*/)
+
 /**
 	Helper struct to sub step through Montages when advancing time.
 	These require stopping at sections and branching points to potential jumps and loops.
@@ -337,7 +344,8 @@ struct FAnimMontageInstance
 	FOnMontageEnded OnMontageEnded;
 	FOnMontageBlendingOutStarted OnMontageBlendingOutStarted;
 	FOnMontageBlendedInEnded OnMontageBlendedInEnded;
-
+	FOnMontageSectionChanged OnMontageSectionChanged;
+	
 	UPROPERTY()
 	bool bPlaying;
 
diff --git a/Engine/Source/Runtime/Engine/Private/Animation/AnimInstance.cpp b/Engine/Source/Runtime/Engine/Private/Animation/AnimInstance.cpp
--- a/Engine/Source/Runtime/Engine/Private/Animation/AnimInstance.cpp
+++ b/Engine/Source/Runtime/Engine/Private/Animation/AnimInstance.cpp
@@ -2137,6 +2137,18 @@ void UAnimInstance::QueueMontageEndedEvent(const FQueuedMontageEndedEvent& Monta
 	}
 }
 
+void UAnimInstance::QueueMontageSectionChangedEvent(const FQueuedMontageSectionChangedEvent& MontageSectionChangedEvent)
+{
+	if (bQueueMontageEvents)
+	{
+		QueuedMontageSectionChangedEvents.Add(MontageSectionChangedEvent);
+	}
+	else
+	{
+		TriggerMontageSectionChangedEvent(MontageSectionChangedEvent);
+	}
+}
+
 void UAnimInstance::TriggerMontageEndedEvent(const FQueuedMontageEndedEvent& MontageEndedEvent)
 {
 	// Send end notifications for anim notify state when we are stopped
@@ -2195,6 +2207,12 @@ void UAnimInstance::TriggerMontageEndedEvent(const FQueuedMontageEndedEvent& Mon
 	OnMontageEnded.Broadcast(MontageEndedEvent.Montage, MontageEndedEvent.bInterrupted);
 }
 
+void UAnimInstance::TriggerMontageSectionChangedEvent(const FQueuedMontageSectionChangedEvent& MontageSectionChangedEvent)
+{
+	MontageSectionChangedEvent.Delegate.ExecuteIfBound(MontageSectionChangedEvent.Montage, MontageSectionChangedEvent.SectionName, MontageSectionChangedEvent.bLooped);
+	OnMontageSectionChanged.Broadcast(MontageSectionChangedEvent.Montage, MontageSectionChangedEvent.SectionName, MontageSectionChangedEvent.bLooped);
+}
+
 void UAnimInstance::TriggerQueuedMontageEvents()
 {
 	// We don't need to queue montage events anymore.
@@ -2219,6 +2237,15 @@ void UAnimInstance::TriggerQueuedMontageEvents()
 		QueuedMontageBlendedInEvents.Reset();
 	}
 
+	if (QueuedMontageSectionChangedEvents.Num() > 0)
+	{
+		for (const FQueuedMontageSectionChangedEvent& MontageSectionChangedEvent : QueuedMontageSectionChangedEvents)
+		{
+			TriggerMontageSectionChangedEvent(MontageSectionChangedEvent);
+		}
+		QueuedMontageSectionChangedEvents.Reset();
+	}
+	
 	if (QueuedMontageEndedEvents.Num() > 0)
 	{
 		for (const FQueuedMontageEndedEvent& MontageEndedEvent : QueuedMontageEndedEvents)
@@ -2912,6 +2939,56 @@ FOnMontageBlendingOutStarted* UAnimInstance::Montage_GetBlendingOutDelegate(UAni
 	return NULL;
 }
 
+void UAnimInstance::Montage_SetSectionChangedDelegate(FOnMontageSectionChanged& InOnMontageSectionChanged, UAnimMontage* Montage)
+{
+	if (Montage)
+	{
+		FAnimMontageInstance* MontageInstance = GetActiveInstanceForMontage(Montage);
+		if (MontageInstance)
+		{
+			MontageInstance->OnMontageSectionChanged = InOnMontageSectionChanged;
+		}
+	}
+	else
+	{
+		// If no Montage reference, do it on all active ones.
+		for (int32 InstanceIndex = 0; InstanceIndex < MontageInstances.Num(); InstanceIndex++)
+		{
+			FAnimMontageInstance* MontageInstance = MontageInstances[InstanceIndex];
+			if (MontageInstance && MontageInstance->IsActive())
+			{
+				MontageInstance->OnMontageSectionChanged = InOnMontageSectionChanged;
+			}
+		}
+	}
+}
+
+FOnMontageSectionChanged* UAnimInstance::Montage_GetSectionChangedDelegate(UAnimMontage* Montage)
+{
+	if (Montage)
+	{
+		FAnimMontageInstance* MontageInstance = GetActiveInstanceForMontage(Montage);
+		if (MontageInstance)
+		{
+			return &MontageInstance->OnMontageSectionChanged;
+		}
+	}
+	else
+	{
+		// If no Montage reference, use first active one found.
+		for (int32 InstanceIndex = 0; InstanceIndex < MontageInstances.Num(); InstanceIndex++)
+		{
+			FAnimMontageInstance* MontageInstance = MontageInstances[InstanceIndex];
+			if (MontageInstance && MontageInstance->IsActive())
+			{
+				return &MontageInstance->OnMontageSectionChanged;
+			}
+		}
+	}
+
+	return NULL;
+}
+
 void UAnimInstance::Montage_SetPosition(const UAnimMontage* Montage, float NewPosition)
 {
 	if (Montage)
@@ -4015,6 +4092,11 @@ void UAnimInstance::AddReferencedObjects(UObject* InThis, FReferenceCollector& C
 		Collector.AddReferencedObject(This->QueuedMontageEndedEvents[I].Montage);
 	}
 
+	for (int32 I = 0; I < This->QueuedMontageSectionChangedEvents.Num(); ++I)
+	{
+		Collector.AddReferencedObject(This->QueuedMontageSectionChangedEvents[I].Montage);
+	}
+	
 	if (This->AnimInstanceProxy)
 	{
 		This->AnimInstanceProxy->AddReferencedObjects(This, Collector);
diff --git a/Engine/Source/Runtime/Engine/Private/Animation/AnimMontage.cpp b/Engine/Source/Runtime/Engine/Private/Animation/AnimMontage.cpp
--- a/Engine/Source/Runtime/Engine/Private/Animation/AnimMontage.cpp
+++ b/Engine/Source/Runtime/Engine/Private/Animation/AnimMontage.cpp
@@ -1778,6 +1778,11 @@ bool FAnimMontageInstance::JumpToSectionName(FName const & SectionName, bool bEn
 		const float NewPosition = Montage->CalculatePos(CurSection, bEndOfSection ? Montage->GetSectionLength(SectionID) - UE_KINDA_SMALL_NUMBER : 0.0f);
 		SetPosition(NewPosition);
 		OnMontagePositionChanged(SectionName);
+		
+		// changed sections, trigger event
+		const bool bLooped = Montage->IsValidSectionIndex(MontageSubStepper.GetCurrentSectionIndex()) ? MontageSubStepper.GetCurrentSectionIndex() == SectionID : false;
+		AnimInstance->QueueMontageSectionChangedEvent(FQueuedMontageSectionChangedEvent(Montage, InstanceID, SectionName, bLooped, OnMontageSectionChanged));
+		
 		return true;
 	}
 
@@ -2633,6 +2638,9 @@ void FAnimMontageInstance::Advance(float DeltaTime, struct FRootMotionMovementPa
 							// Jump to next section's appropriate starting point (start or end).
 							Position = bPlayingForward ? LatestNextSectionStartTime : (LatestNextSectionEndTime - EndOffset);
 							SubStepResult = EMontageSubStepResult::Moved;
+							
+							const bool bLooped = CurrentSectionIndex == RecentNextSectionIndex;
+							AnimInstance->QueueMontageSectionChangedEvent(FQueuedMontageSectionChangedEvent(Montage, InstanceID, Montage->GetSectionName(RecentNextSectionIndex), bLooped, OnMontageSectionChanged));
 						}
 						else
 						{
