[Engine][Animation] Root Space Additive and Layered Blend Per Bone Options
- Initialize Uninitialized Variables

diff --git a/Engine/Source/Runtime/AnimGraphRuntime/Public/AnimNodes/AnimNode_LayeredBoneBlend.h b/Engine/Source/Runtime/AnimGraphRuntime/Public/AnimNodes/AnimNode_LayeredBoneBlend.h
--- a/Engine/Source/Runtime/AnimGraphRuntime/Public/AnimNodes/AnimNode_LayeredBoneBlend.h
+++ b/Engine/Source/Runtime/AnimGraphRuntime/Public/AnimNodes/AnimNode_LayeredBoneBlend.h
@@ -116,6 +116,7 @@ public:
 		, RequiredBonesSerialNumber(0)
 		, LODThreshold(INDEX_NONE)
 		, bMeshSpaceRotationBlend(false)
+		, bRootSpaceRotationBlend(false)
 		, bMeshSpaceScaleBlend(false)
 		, CurveBlendOption(ECurveBlendOption::Override)
 		, bBlendRootMotionBasedOnRootBone(true)
diff --git a/Engine/Source/Runtime/Engine/Private/Animation/AnimNode_ApplyMeshSpaceAdditive.cpp b/Engine/Source/Runtime/Engine/Private/Animation/AnimNode_ApplyMeshSpaceAdditive.cpp
--- a/Engine/Source/Runtime/Engine/Private/Animation/AnimNode_ApplyMeshSpaceAdditive.cpp
+++ b/Engine/Source/Runtime/Engine/Private/Animation/AnimNode_ApplyMeshSpaceAdditive.cpp
@@ -111,7 +111,8 @@ void FAnimNode_ApplyMeshSpaceAdditive::Evaluate_AnyThread(FPoseContext& Output)
 }
 
 FAnimNode_ApplyMeshSpaceAdditive::FAnimNode_ApplyMeshSpaceAdditive()
-	: AlphaInputType(EAnimAlphaInputType::Float)
+	: bRootSpaceAdditive(false)
+	, AlphaInputType(EAnimAlphaInputType::Float)
 	, Alpha(1.0f)
 	, bAlphaBoolEnabled(true)
 	, AlphaCurveName(NAME_None)
