Expose Chaos API functions to support Control Rig Physics (extends previous submit like this)

diff --git a/Engine/Source/Runtime/Engine/Private/PhysicsEngine/ImmediatePhysics/ImmediatePhysicsAdapters.cpp b/Engine/Source/Runtime/Engine/Private/PhysicsEngine/ImmediatePhysics/ImmediatePhysicsAdapters.cpp
--- a/Engine/Source/Runtime/Engine/Private/PhysicsEngine/ImmediatePhysics/ImmediatePhysicsAdapters.cpp
+++ b/Engine/Source/Runtime/Engine/Private/PhysicsEngine/ImmediatePhysics/ImmediatePhysicsAdapters.cpp
@@ -360,6 +360,79 @@ namespace ImmediatePhysics
 		JointSettings.TwistContactDistance = Constraint.ContactDistance;
 	}
 
+	void UpdateJointSettingsFromLinearDriveConstraint(const FLinearDriveConstraint& LinearDriveConstraint, Chaos::FPBDJointSettings& JointSettings)
+	{
+		using namespace Chaos;
+
+		JointSettings.LinearDrivePositionTarget = LinearDriveConstraint.PositionTarget;
+		JointSettings.LinearDriveVelocityTarget = LinearDriveConstraint.VelocityTarget;
+		JointSettings.bLinearPositionDriveEnabled[0] = LinearDriveConstraint.XDrive.bEnablePositionDrive;
+		JointSettings.bLinearPositionDriveEnabled[1] = LinearDriveConstraint.YDrive.bEnablePositionDrive;
+		JointSettings.bLinearPositionDriveEnabled[2] = LinearDriveConstraint.ZDrive.bEnablePositionDrive;
+		JointSettings.bLinearVelocityDriveEnabled[0] = LinearDriveConstraint.XDrive.bEnableVelocityDrive;
+		JointSettings.bLinearVelocityDriveEnabled[1] = LinearDriveConstraint.YDrive.bEnableVelocityDrive;
+		JointSettings.bLinearVelocityDriveEnabled[2] = LinearDriveConstraint.ZDrive.bEnableVelocityDrive;
+
+		JointSettings.LinearDriveForceMode = EJointForceMode::Acceleration; // hardcoded!
+		JointSettings.LinearDriveStiffness = ConstraintSettings::LinearDriveStiffnessScale() * FVec3(
+			LinearDriveConstraint.XDrive.Stiffness,
+			LinearDriveConstraint.YDrive.Stiffness,
+			LinearDriveConstraint.ZDrive.Stiffness);
+		JointSettings.LinearDriveDamping = ConstraintSettings::LinearDriveDampingScale() * FVec3(
+			LinearDriveConstraint.XDrive.Damping,
+			LinearDriveConstraint.YDrive.Damping,
+			LinearDriveConstraint.ZDrive.Damping);
+		JointSettings.LinearDriveMaxForce[0] = LinearDriveConstraint.XDrive.MaxForce;
+		JointSettings.LinearDriveMaxForce[1] = LinearDriveConstraint.YDrive.MaxForce;
+		JointSettings.LinearDriveMaxForce[2] = LinearDriveConstraint.ZDrive.MaxForce;
+	}
+
+	void UpdateJointSettingsFromAngularDriveConstraint(
+		const FAngularDriveConstraint& AngularDriveConstraint, Chaos::FPBDJointSettings& JointSettings)
+	{
+		using namespace Chaos;
+
+		JointSettings.AngularDrivePositionTarget = FQuat(AngularDriveConstraint.OrientationTarget);
+		JointSettings.AngularDriveVelocityTarget = AngularDriveConstraint.AngularVelocityTarget * UE_TWO_PI; // rev/s to rad/s
+
+		JointSettings.AngularDriveForceMode = EJointForceMode::Acceleration; // hardcoded!
+		if (AngularDriveConstraint.AngularDriveMode == EAngularDriveMode::SLERP)
+		{
+			JointSettings.AngularDriveStiffness = FVec3(
+				ConstraintSettings::AngularDriveStiffnessScale() * AngularDriveConstraint.SlerpDrive.Stiffness);
+			JointSettings.AngularDriveDamping = FVec3(
+				ConstraintSettings::AngularDriveDampingScale() * AngularDriveConstraint.SlerpDrive.Damping);
+			JointSettings.AngularDriveMaxTorque = FVec3(
+				AngularDriveConstraint.SlerpDrive.MaxForce);
+			JointSettings.bAngularSLerpPositionDriveEnabled = AngularDriveConstraint.SlerpDrive.bEnablePositionDrive;
+			JointSettings.bAngularSLerpVelocityDriveEnabled = AngularDriveConstraint.SlerpDrive.bEnableVelocityDrive;
+			JointSettings.bAngularTwistPositionDriveEnabled = false;
+			JointSettings.bAngularTwistVelocityDriveEnabled = false;
+			JointSettings.bAngularSwingPositionDriveEnabled = false;
+			JointSettings.bAngularSwingVelocityDriveEnabled = false;
+		}
+		else
+		{
+			JointSettings.AngularDriveStiffness = ConstraintSettings::AngularDriveStiffnessScale() * FVec3(
+				AngularDriveConstraint.TwistDrive.Stiffness,
+				AngularDriveConstraint.SwingDrive.Stiffness,
+				AngularDriveConstraint.SwingDrive.Stiffness);
+			JointSettings.AngularDriveDamping = ConstraintSettings::AngularDriveDampingScale() * FVec3(
+				AngularDriveConstraint.TwistDrive.Damping,
+				AngularDriveConstraint.SwingDrive.Damping,
+				AngularDriveConstraint.SwingDrive.Damping);
+			JointSettings.AngularDriveMaxTorque[0] = AngularDriveConstraint.TwistDrive.MaxForce;
+			JointSettings.AngularDriveMaxTorque[1] = AngularDriveConstraint.SwingDrive.MaxForce;
+			JointSettings.AngularDriveMaxTorque[2] = AngularDriveConstraint.SwingDrive.MaxForce;
+			JointSettings.bAngularSLerpPositionDriveEnabled = false;
+			JointSettings.bAngularSLerpVelocityDriveEnabled = false;
+			JointSettings.bAngularTwistPositionDriveEnabled = AngularDriveConstraint.TwistDrive.bEnablePositionDrive;
+			JointSettings.bAngularTwistVelocityDriveEnabled = AngularDriveConstraint.TwistDrive.bEnableVelocityDrive;
+			JointSettings.bAngularSwingPositionDriveEnabled = AngularDriveConstraint.SwingDrive.bEnablePositionDrive;
+			JointSettings.bAngularSwingVelocityDriveEnabled = AngularDriveConstraint.SwingDrive.bEnableVelocityDrive;
+		}
+	}
+
 	// Copies data from a profile into settings
 	void UpdateJointSettingsFromConstraintProfile(const FConstraintProfileProperties& Profile, Chaos::FPBDJointSettings& JointSettings)
 	{
@@ -390,67 +463,8 @@ namespace ImmediatePhysics
 		JointSettings.AngularSoftForceMode = (ConstraintSettings::SoftAngularForceMode() == 0) ?
 			EJointForceMode::Acceleration : EJointForceMode::Force;
 
-		JointSettings.LinearDrivePositionTarget = Profile.LinearDrive.PositionTarget;
-		JointSettings.LinearDriveVelocityTarget = Profile.LinearDrive.VelocityTarget;
-		JointSettings.bLinearPositionDriveEnabled[0] = Profile.LinearDrive.XDrive.bEnablePositionDrive;
-		JointSettings.bLinearPositionDriveEnabled[1] = Profile.LinearDrive.YDrive.bEnablePositionDrive;
-		JointSettings.bLinearPositionDriveEnabled[2] = Profile.LinearDrive.ZDrive.bEnablePositionDrive;
-		JointSettings.bLinearVelocityDriveEnabled[0] = Profile.LinearDrive.XDrive.bEnableVelocityDrive;
-		JointSettings.bLinearVelocityDriveEnabled[1] = Profile.LinearDrive.YDrive.bEnableVelocityDrive;
-		JointSettings.bLinearVelocityDriveEnabled[2] = Profile.LinearDrive.ZDrive.bEnableVelocityDrive;
-
-		JointSettings.LinearDriveForceMode = EJointForceMode::Acceleration; // hardcoded!
-		JointSettings.LinearDriveStiffness = ConstraintSettings::LinearDriveStiffnessScale() * FVec3(
-			Profile.LinearDrive.XDrive.Stiffness,
-			Profile.LinearDrive.YDrive.Stiffness,
-			Profile.LinearDrive.ZDrive.Stiffness);
-		JointSettings.LinearDriveDamping = ConstraintSettings::LinearDriveDampingScale() * FVec3(
-			Profile.LinearDrive.XDrive.Damping,
-			Profile.LinearDrive.YDrive.Damping,
-			Profile.LinearDrive.ZDrive.Damping);
-		JointSettings.LinearDriveMaxForce[0] = Profile.LinearDrive.XDrive.MaxForce;
-		JointSettings.LinearDriveMaxForce[1] = Profile.LinearDrive.YDrive.MaxForce;
-		JointSettings.LinearDriveMaxForce[2] = Profile.LinearDrive.ZDrive.MaxForce;
-
-		JointSettings.AngularDrivePositionTarget = FQuat(Profile.AngularDrive.OrientationTarget);
-		JointSettings.AngularDriveVelocityTarget = Profile.AngularDrive.AngularVelocityTarget * UE_TWO_PI; // rev/s to rad/s
-
-		JointSettings.AngularDriveForceMode = EJointForceMode::Acceleration; // hardcoded!
-		if (Profile.AngularDrive.AngularDriveMode == EAngularDriveMode::SLERP)
-		{
-			JointSettings.AngularDriveStiffness = FVec3(
-				ConstraintSettings::AngularDriveStiffnessScale() * Profile.AngularDrive.SlerpDrive.Stiffness);
-			JointSettings.AngularDriveDamping = FVec3(
-				ConstraintSettings::AngularDriveDampingScale() * Profile.AngularDrive.SlerpDrive.Damping);
-			JointSettings.AngularDriveMaxTorque = FVec3(
-				Profile.AngularDrive.SlerpDrive.MaxForce);
-			JointSettings.bAngularSLerpPositionDriveEnabled = Profile.AngularDrive.SlerpDrive.bEnablePositionDrive;
-			JointSettings.bAngularSLerpVelocityDriveEnabled = Profile.AngularDrive.SlerpDrive.bEnableVelocityDrive;
-			JointSettings.bAngularTwistPositionDriveEnabled = false;
-			JointSettings.bAngularTwistVelocityDriveEnabled = false;
-			JointSettings.bAngularSwingPositionDriveEnabled = false;
-			JointSettings.bAngularSwingVelocityDriveEnabled = false;
-		}
-		else
-		{
-			JointSettings.AngularDriveStiffness = ConstraintSettings::AngularDriveStiffnessScale() * FVec3(
-				Profile.AngularDrive.TwistDrive.Stiffness,
-				Profile.AngularDrive.SwingDrive.Stiffness,
-				Profile.AngularDrive.SwingDrive.Stiffness);
-			JointSettings.AngularDriveDamping = ConstraintSettings::AngularDriveDampingScale() * FVec3(
-				Profile.AngularDrive.TwistDrive.Damping,
-				Profile.AngularDrive.SwingDrive.Damping,
-				Profile.AngularDrive.SwingDrive.Damping);
-			JointSettings.AngularDriveMaxTorque[0] = Profile.AngularDrive.TwistDrive.MaxForce;
-			JointSettings.AngularDriveMaxTorque[1] = Profile.AngularDrive.SwingDrive.MaxForce;
-			JointSettings.AngularDriveMaxTorque[2] = Profile.AngularDrive.SwingDrive.MaxForce;
-			JointSettings.bAngularSLerpPositionDriveEnabled = false;
-			JointSettings.bAngularSLerpVelocityDriveEnabled = false;
-			JointSettings.bAngularTwistPositionDriveEnabled = Profile.AngularDrive.TwistDrive.bEnablePositionDrive;
-			JointSettings.bAngularTwistVelocityDriveEnabled = Profile.AngularDrive.TwistDrive.bEnableVelocityDrive;
-			JointSettings.bAngularSwingPositionDriveEnabled = Profile.AngularDrive.SwingDrive.bEnablePositionDrive;
-			JointSettings.bAngularSwingVelocityDriveEnabled = Profile.AngularDrive.SwingDrive.bEnableVelocityDrive;
-		}
+		UpdateJointSettingsFromLinearDriveConstraint(Profile.LinearDrive, JointSettings);
+		UpdateJointSettingsFromAngularDriveConstraint(Profile.AngularDrive, JointSettings);
 
 		JointSettings.LinearBreakForce = Profile.bLinearBreakable ?
 			ConstraintSettings::LinearBreakScale() * Profile.LinearBreakThreshold : FLT_MAX;
diff --git a/Engine/Source/Runtime/Engine/Public/Physics/ImmediatePhysics/ImmediatePhysicsAdapters.h b/Engine/Source/Runtime/Engine/Public/Physics/ImmediatePhysics/ImmediatePhysicsAdapters.h
--- a/Engine/Source/Runtime/Engine/Public/Physics/ImmediatePhysics/ImmediatePhysicsAdapters.h
+++ b/Engine/Source/Runtime/Engine/Public/Physics/ImmediatePhysics/ImmediatePhysicsAdapters.h
@@ -11,6 +11,8 @@ struct FConeConstraint;
 struct FConstraintProfileProperties;
 struct FLinearConstraint;
 struct FTwistConstraint;
+struct FLinearDriveConstraint;
+struct FAngularDriveConstraint;
 
 namespace ImmediatePhysics
 {
@@ -40,4 +42,10 @@ namespace ImmediatePhysics
 	
 	//Extracts parameters from a constraint constraint into the joint settings
 	ENGINE_API void UpdateJointSettingsFromTwistConstraint(const FTwistConstraint& Constraint, Chaos::FPBDJointSettings& JointSettings);
+
+	ENGINE_API void UpdateJointSettingsFromLinearDriveConstraint(
+		const FLinearDriveConstraint& Constraint, Chaos::FPBDJointSettings& JointSettings);
+	ENGINE_API void UpdateJointSettingsFromAngularDriveConstraint(
+		const FAngularDriveConstraint& Constraint, Chaos::FPBDJointSettings& JointSettings);
+
 }
