Add a "FindLocalPlayerFromDeviceId" function to the game instance. This is just a small helper function which can be useful when writing viewport client code

diff --git a/Engine/Source/Runtime/Engine/Classes/Engine/GameInstance.h b/Engine/Source/Runtime/Engine/Classes/Engine/GameInstance.h
--- a/Engine/Source/Runtime/Engine/Classes/Engine/GameInstance.h
+++ b/Engine/Source/Runtime/Engine/Classes/Engine/GameInstance.h
@@ -385,6 +385,9 @@ public:
 	/** Returns the local player assigned to this platform user id, or null if not found */
 	ENGINE_API ULocalPlayer* FindLocalPlayerFromPlatformUserId(const FPlatformUserId UserId) const;
 
+	/** Returns the local player who the given input device is assigned to. Null if not found */
+	ENGINE_API ULocalPlayer* FindLocalPlayerFromDeviceId(const FInputDeviceId DeviceID) const;
+
 	/** Returns the local player that has been assigned the specific unique net id */
 	ENGINE_API ULocalPlayer*			FindLocalPlayerFromUniqueNetId(FUniqueNetIdPtr UniqueNetId) const;
 	ENGINE_API ULocalPlayer*			FindLocalPlayerFromUniqueNetId(const FUniqueNetId& UniqueNetId) const;
diff --git a/Engine/Source/Runtime/Engine/Private/GameInstance.cpp b/Engine/Source/Runtime/Engine/Private/GameInstance.cpp
--- a/Engine/Source/Runtime/Engine/Private/GameInstance.cpp
+++ b/Engine/Source/Runtime/Engine/Private/GameInstance.cpp
@@ -1183,6 +1183,12 @@ ULocalPlayer* UGameInstance::FindLocalPlayerFromPlatformUserId(const FPlatformUs
 	return nullptr;
 }
 
+ULocalPlayer* UGameInstance::FindLocalPlayerFromDeviceId(const FInputDeviceId DeviceID) const
+{
+	const FPlatformUserId UserId = IPlatformInputDeviceMapper::Get().GetUserForInputDevice(DeviceID);
+	return FindLocalPlayerFromPlatformUserId(UserId);
+}
+
 ULocalPlayer* UGameInstance::FindLocalPlayerFromUniqueNetId(const FUniqueNetId& UniqueNetId) const
 {
 	return FindLocalPlayerFromUniqueNetId(FUniqueNetIdRepl(UniqueNetId.AsShared()));
