[Backout] - CL40368939
#rnx
Original CL Desc
-----------------------------------------------------------------
World Partition - HLOD: Introduce an IWorldPartitionHLODObject interface
* HLODRuntime subsystem cell handling & warmup logic now rely solely on this interface and not actors

diff --git a/Engine/Source/Runtime/Engine/Private/UnrealEngine.cpp b/Engine/Source/Runtime/Engine/Private/UnrealEngine.cpp
--- a/Engine/Source/Runtime/Engine/Private/UnrealEngine.cpp
+++ b/Engine/Source/Runtime/Engine/Private/UnrealEngine.cpp
@@ -12799,11 +12799,11 @@ float DrawMapWarnings(UWorld* World, FViewport* Viewport, FCanvas* Canvas, UCanv
 
 	if (UWorldPartitionHLODRuntimeSubsystem* HLODSubsystem = World->GetSubsystem<UWorldPartitionHLODRuntimeSubsystem>())
 	{
-		uint32 NumOutdatedHLODObjects = HLODSubsystem->GetNumOutdatedHLODObjects();
-		if (NumOutdatedHLODObjects)
+		uint32 NumOutdatedHLODActors = HLODSubsystem->GetNumOutdatedHLODActors();
+		if (NumOutdatedHLODActors)
 		{
 			SmallTextItem.SetColor(FLinearColor::Yellow);
-			SmallTextItem.Text = FText::Format(LOCTEXT("OUTDATED_HLOD_ACTORS_FMT", "Outdated HLODs detected, HLODs need to be rebuilt ({0} outdated)"), NumOutdatedHLODObjects);
+			SmallTextItem.Text = FText::Format(LOCTEXT("OUTDATED_HLOD_ACTORS_FMT", "Outdated HLOD actors detected, HLODs need to be rebuilt ({0} outdated {0}|plural(one=actor,other=actors))"), NumOutdatedHLODActors);
 			Canvas->DrawItem(SmallTextItem, FVector2D(MessageX, MessageY));
 			MessageY += FontSizeY;
 		}
diff --git a/Engine/Source/Runtime/Engine/Private/WorldPartition/HLOD/HLODActor.cpp b/Engine/Source/Runtime/Engine/Private/WorldPartition/HLOD/HLODActor.cpp
--- a/Engine/Source/Runtime/Engine/Private/WorldPartition/HLOD/HLODActor.cpp
+++ b/Engine/Source/Runtime/Engine/Private/WorldPartition/HLOD/HLODActor.cpp
@@ -2,10 +2,8 @@
 
 #include "WorldPartition/HLOD/HLODActor.h"
 #include "Engine/World.h"
-#include "Engine/StaticMesh.h"
 #include "WorldPartition/HLOD/HLODRuntimeSubsystem.h"
 #include "Components/PrimitiveComponent.h"
-#include "Components/InstancedStaticMeshComponent.h"
 #include "Misc/PackageName.h"
 #include "UObject/UE5MainStreamObjectVersion.h"
 #include "UObject/UE5PrivateFrostyStreamObjectVersion.h"
@@ -68,16 +66,6 @@ AWorldPartitionHLOD::AWorldPartitionHLOD(const FObjectInitializer& ObjectInitial
 #endif
 }
 
-ULevel* AWorldPartitionHLOD::GetHLODLevel() const
-{
-	return GetLevel();
-}
-
-FString AWorldPartitionHLOD::GetHLODNameOrLabel() const
-{
-	return GetActorNameOrLabel();
-}
-
 const FGuid& AWorldPartitionHLOD::GetSourceCellGuid() const
 {
 	// When no source cell guid was set, try resolving it through its associated world partition runtime cell
@@ -104,44 +92,15 @@ void AWorldPartitionHLOD::SetVisibility(bool bInVisible)
 	});
 }
 
-TSet<UObject*> AWorldPartitionHLOD::GetAssetsToWarmup() const
-{
-	TSet<UObject*> AssetsToWarmup;
-	
-	ForEachComponent<UStaticMeshComponent>(false, [&](UStaticMeshComponent* SMC)
-	{
-		// Assume ISM HLOD don't need warmup, as they are actually found in the source level
-		if (SMC->IsA<UInstancedStaticMeshComponent>())
-		{
-			return;
-		}
-
-		for (int32 iMaterialIndex = 0; iMaterialIndex < SMC->GetNumMaterials(); ++iMaterialIndex)
-		{
-			if (UMaterialInterface* Material = SMC->GetMaterial(iMaterialIndex))
-			{
-				AssetsToWarmup.Add(Material);
-			}
-		}
-
-		if (UStaticMesh* StaticMesh = SMC->GetStaticMesh())
-		{
-			AssetsToWarmup.Add(StaticMesh);
-		}
-	});
-
-	return AssetsToWarmup;
-}
-
 void AWorldPartitionHLOD::BeginPlay()
 {
 	Super::BeginPlay();
-	GetWorld()->GetSubsystem<UWorldPartitionHLODRuntimeSubsystem>()->RegisterHLODObject(this);
+	GetWorld()->GetSubsystem<UWorldPartitionHLODRuntimeSubsystem>()->RegisterHLODActor(this);
 }
 
 void AWorldPartitionHLOD::EndPlay(const EEndPlayReason::Type EndPlayReason)
 {
-	GetWorld()->GetSubsystem<UWorldPartitionHLODRuntimeSubsystem>()->UnregisterHLODObject(this);
+	GetWorld()->GetSubsystem<UWorldPartitionHLODRuntimeSubsystem>()->UnregisterHLODActor(this);
 	Super::EndPlay(EndPlayReason);
 }
 
diff --git a/Engine/Source/Runtime/Engine/Private/WorldPartition/HLOD/HLODDestruction.cpp b/Engine/Source/Runtime/Engine/Private/WorldPartition/HLOD/HLODDestruction.cpp
--- a/Engine/Source/Runtime/Engine/Private/WorldPartition/HLOD/HLODDestruction.cpp
+++ b/Engine/Source/Runtime/Engine/Private/WorldPartition/HLOD/HLODDestruction.cpp
@@ -33,22 +33,18 @@ namespace
 
 		if (RuntimeCell)
 		{
-			for (IWorldPartitionHLODObject* HLODObject : InActor->GetWorld()->GetSubsystem<UWorldPartitionHLODRuntimeSubsystem>()->GetHLODObjectsForCell(RuntimeCell))
+			for (AWorldPartitionHLOD* HLODActor : InActor->GetWorld()->GetSubsystem<UWorldPartitionHLODRuntimeSubsystem>()->GetHLODActorsForCell(RuntimeCell))
 			{
-				// This is the only use of IWorldPartitionHLODObject::GetUObject() for now... Once we solve destruction for FastGeo HLODs, we will be able to remove it.
-				if (AWorldPartitionHLOD* HLODActor = Cast<AWorldPartitionHLOD>(HLODObject->GetUObject()))
+				for (UActorComponent* ActorComponent : HLODActor->GetComponents())
 				{
-					for (UActorComponent* ActorComponent : HLODActor->GetComponents())
+					if (UWorldPartitionDestructibleHLODComponent* DestructibleHLODComponent = Cast<UWorldPartitionDestructibleHLODComponent>(ActorComponent))
 					{
-						if (UWorldPartitionDestructibleHLODComponent* DestructibleHLODComponent = Cast<UWorldPartitionDestructibleHLODComponent>(ActorComponent))
+						int32 ActorIdx = DestructibleHLODComponent->GetDestructibleActors().IndexOfByKey(ActorFName);
+						if (ActorIdx != INDEX_NONE)
 						{
-							int32 ActorIdx = DestructibleHLODComponent->GetDestructibleActors().IndexOfByKey(ActorFName);
-							if (ActorIdx != INDEX_NONE)
-							{
-								HLODDestructionTag.ActorIndex = ActorIdx;
-								HLODDestructionTag.HLODDestructionComponent = DestructibleHLODComponent;
-								break;
-							}
+							HLODDestructionTag.ActorIndex = ActorIdx;
+							HLODDestructionTag.HLODDestructionComponent = DestructibleHLODComponent;
+							break;
 						}
 					}
 				}
diff --git a/Engine/Source/Runtime/Engine/Private/WorldPartition/HLOD/HLODRuntimeSubsystem.cpp b/Engine/Source/Runtime/Engine/Private/WorldPartition/HLOD/HLODRuntimeSubsystem.cpp
--- a/Engine/Source/Runtime/Engine/Private/WorldPartition/HLOD/HLODRuntimeSubsystem.cpp
+++ b/Engine/Source/Runtime/Engine/Private/WorldPartition/HLOD/HLODRuntimeSubsystem.cpp
@@ -82,21 +82,21 @@ static FAutoConsoleVariableSink CVarHLODSink(FConsoleCommandDelegate::CreateStat
 
 namespace FHLODRuntimeSubsystem
 {
-    static const UWorldPartitionRuntimeCell* GetRuntimeCell(IWorldPartitionHLODObject* InWorldPartitionHLOD)
+    static const UWorldPartitionRuntimeCell* GetActorRuntimeCell(AActor* InActor)
     {
-	    const ULevel* Level = InWorldPartitionHLOD->GetHLODLevel();
-	    const ULevelStreaming* LevelStreaming = FLevelUtils::FindStreamingLevel(Level);
+	    const ULevel* ActorLevel = InActor->GetLevel();
+	    const ULevelStreaming* LevelStreaming = FLevelUtils::FindStreamingLevel(ActorLevel);
 	    return LevelStreaming ? Cast<const UWorldPartitionRuntimeCell>(LevelStreaming->GetWorldPartitionCell()) : nullptr;
     }
     
-    static UWorldPartition* GetWorldPartition(IWorldPartitionHLODObject* InWorldPartitionHLOD)
+    static UWorldPartition* GetWorldPartition(AWorldPartitionHLOD* InWorldPartitionHLOD)
     {
 	    // Alwaysloaded Cell level will have a WorldPartition
-	    if (UWorldPartition* WorldPartition = InWorldPartitionHLOD->GetHLODLevel()->GetWorldPartition())
+	    if (UWorldPartition* WorldPartition = InWorldPartitionHLOD->GetLevel()->GetWorldPartition())
 	    {
 		    return WorldPartition;
 	    } // If not find it through the cell
-	    else if (const UWorldPartitionRuntimeCell* RuntimeCell = GetRuntimeCell(InWorldPartitionHLOD))
+	    else if (const UWorldPartitionRuntimeCell* RuntimeCell = GetActorRuntimeCell(InWorldPartitionHLOD))
 	    {
 		    return RuntimeCell->GetOuterWorld()->GetWorldPartition();
 	    }
@@ -150,9 +150,9 @@ FAutoConsoleCommand UWorldPartitionHLODRuntimeSubsystem::EnableHLODCommand(
 						{
 							const FCellData& CellData = CellHLODMapping.Value;
 							bool bIsHLODVisible = UWorldPartitionHLODRuntimeSubsystem::WorldPartitionHLODEnabled && !CellData.bIsCellVisible;
-							for (IWorldPartitionHLODObject* HLODObject : CellData.LoadedHLODs)
+							for (AWorldPartitionHLOD* HLODActor : CellData.LoadedHLODs)
 							{
-								HLODObject->SetVisibility(bIsHLODVisible);
+								HLODActor->SetVisibility(bIsHLODVisible);
 							}
 						}
 					}
@@ -290,7 +290,7 @@ UWorldPartitionHLODRuntimeSubsystem::FCellData* UWorldPartitionHLODRuntimeSubsys
 	return nullptr;
 }
 
-UWorldPartitionHLODRuntimeSubsystem::FCellData* UWorldPartitionHLODRuntimeSubsystem::GetCellData(IWorldPartitionHLODObject* InWorldPartitionHLOD)
+UWorldPartitionHLODRuntimeSubsystem::FCellData* UWorldPartitionHLODRuntimeSubsystem::GetCellData(AWorldPartitionHLOD* InWorldPartitionHLOD)
 {
 	const UWorldPartition* WorldPartition = FHLODRuntimeSubsystem::GetWorldPartition(InWorldPartitionHLOD);
 	if (WorldPartition)
@@ -307,7 +307,7 @@ UWorldPartitionHLODRuntimeSubsystem::FCellData* UWorldPartitionHLODRuntimeSubsys
 	return nullptr;		
 }
 
-const TArray<IWorldPartitionHLODObject*>& UWorldPartitionHLODRuntimeSubsystem::GetHLODObjectsForCell(const UWorldPartitionRuntimeCell* InCell) const
+const TArray<AWorldPartitionHLOD*>& UWorldPartitionHLODRuntimeSubsystem::GetHLODActorsForCell(const UWorldPartitionRuntimeCell* InCell) const
 {
 	if (const FCellData* CellData = GetCellData(InCell))
 	{
@@ -315,55 +315,53 @@ const TArray<IWorldPartitionHLODObject*>& UWorldPartitionHLODRuntimeSubsystem::G
 	}
 
 	// No HLOD found for the given cell, return a dummy array
-	static const TArray<IWorldPartitionHLODObject*> DummyArray;
+	static const TArray<AWorldPartitionHLOD*> DummyArray;
 	return DummyArray;
 }
 
-void UWorldPartitionHLODRuntimeSubsystem::RegisterHLODObject(IWorldPartitionHLODObject* InWorldPartitionHLOD)
+void UWorldPartitionHLODRuntimeSubsystem::RegisterHLODActor(AWorldPartitionHLOD* InWorldPartitionHLOD)
 {
-	TRACE_CPUPROFILER_EVENT_SCOPE(UWorldPartitionHLODRuntimeSubsystem::RegisterHLODObject);
+	TRACE_CPUPROFILER_EVENT_SCOPE(UWorldPartitionHLODRuntimeSubsystem::RegisterHLODActor);
 	
-	if (FCellData* SourceCellData = GetCellData(InWorldPartitionHLOD))
+	if (FCellData* CellData = GetCellData(InWorldPartitionHLOD))
 	{
-		UE_LOG(LogHLODRuntimeSubsystem, Verbose, TEXT("Registering HLOD %s for cell %s"), *InWorldPartitionHLOD->GetHLODNameOrLabel(), *InWorldPartitionHLOD->GetSourceCellGuid().ToString());
+		UE_LOG(LogHLODRuntimeSubsystem, Verbose, TEXT("Registering HLOD %s for cell %s"), *InWorldPartitionHLOD->GetActorNameOrLabel(), *InWorldPartitionHLOD->GetSourceCellGuid().ToString());
 
-		// Register this HLOD as a representation of the source cell - Needed to warmup an HLOD object when it's source cell is unloading
-		SourceCellData->LoadedHLODs.Add(InWorldPartitionHLOD);
-		InWorldPartitionHLOD->SetVisibility(UWorldPartitionHLODRuntimeSubsystem::WorldPartitionHLODEnabled && !SourceCellData->bIsCellVisible);
+		CellData->LoadedHLODs.Add(InWorldPartitionHLOD);
+		InWorldPartitionHLOD->SetVisibility(UWorldPartitionHLODRuntimeSubsystem::WorldPartitionHLODEnabled && !CellData->bIsCellVisible);
 	}
 	else
 	{
-		UE_LOG(LogHLODRuntimeSubsystem, Verbose, TEXT("Found HLOD %s referencing nonexistent cell '%s'"), *InWorldPartitionHLOD->GetHLODNameOrLabel(), *InWorldPartitionHLOD->GetSourceCellGuid().ToString());
+		UE_LOG(LogHLODRuntimeSubsystem, Verbose, TEXT("Found HLOD %s referencing nonexistent cell '%s'"), *InWorldPartitionHLOD->GetActorNameOrLabel(), *InWorldPartitionHLOD->GetSourceCellGuid().ToString());
 		InWorldPartitionHLOD->SetVisibility(false);
 
 #if !UE_BUILD_SHIPPING
-		OutdatedHLODObjects.Add(InWorldPartitionHLOD);
+		OutdatedHLODActors.Add(InWorldPartitionHLOD);
 #endif
 	}
 
-	HLODObjectsToWarmup.Remove(InWorldPartitionHLOD);
-	HLODObjectRegisteredEvent.Broadcast(InWorldPartitionHLOD);
+	HLODActorsToWarmup.Remove(InWorldPartitionHLOD);
+	HLODActorRegisteredEvent.Broadcast(InWorldPartitionHLOD);
 }
 
-void UWorldPartitionHLODRuntimeSubsystem::UnregisterHLODObject(IWorldPartitionHLODObject* InWorldPartitionHLOD)
+void UWorldPartitionHLODRuntimeSubsystem::UnregisterHLODActor(AWorldPartitionHLOD* InWorldPartitionHLOD)
 {
-	TRACE_CPUPROFILER_EVENT_SCOPE(UWorldPartitionHLODRuntimeSubsystem::UnregisterHLODObject);
+	TRACE_CPUPROFILER_EVENT_SCOPE(UWorldPartitionHLODRuntimeSubsystem::UnregisterHLODActor);
 
 	if (FCellData* CellData = GetCellData(InWorldPartitionHLOD))
 	{
-		UE_LOG(LogHLODRuntimeSubsystem, Verbose, TEXT("Unregistering HLOD %s for cell %s"), *InWorldPartitionHLOD->GetHLODNameOrLabel(), *InWorldPartitionHLOD->GetSourceCellGuid().ToString());
+		UE_LOG(LogHLODRuntimeSubsystem, Verbose, TEXT("Unregistering HLOD %s for cell %s"), *InWorldPartitionHLOD->GetActorNameOrLabel(), *InWorldPartitionHLOD->GetSourceCellGuid().ToString());
 
 		verify(CellData->LoadedHLODs.Remove(InWorldPartitionHLOD));
 	}
 	else
 	{
 #if !UE_BUILD_SHIPPING
-		OutdatedHLODObjects.Remove(InWorldPartitionHLOD);
+		OutdatedHLODActors.Remove(InWorldPartitionHLOD);
 #endif
 	}
 
-	HLODObjectsToWarmup.Remove(InWorldPartitionHLOD);
-	HLODObjectUnregisteredEvent.Broadcast(InWorldPartitionHLOD);
+	HLODActorUnregisteredEvent.Broadcast(InWorldPartitionHLOD);
 }
 
 void UWorldPartitionHLODRuntimeSubsystem::OnCellShown(const UWorldPartitionRuntimeCell* InCell)
@@ -374,12 +372,12 @@ void UWorldPartitionHLODRuntimeSubsystem::OnCellShown(const UWorldPartitionRunti
 
 		if (!CellData->LoadedHLODs.IsEmpty())
 		{
-			UE_LOG(LogHLODRuntimeSubsystem, Verbose, TEXT("Cell shown - %s - hiding %d HLOD objects"), *InCell->GetGuid().ToString(), CellData->LoadedHLODs.Num());
+			UE_LOG(LogHLODRuntimeSubsystem, Verbose, TEXT("Cell shown - %s - hiding %d HLOD actors"), *InCell->GetGuid().ToString(), CellData->LoadedHLODs.Num());
 
-			for (IWorldPartitionHLODObject* HLODObject : CellData->LoadedHLODs)
+			for (AWorldPartitionHLOD* HLODActor : CellData->LoadedHLODs)
 			{
-				UE_LOG(LogHLODRuntimeSubsystem, Verbose, TEXT("\t\t* %s"), *HLODObject->GetHLODNameOrLabel());
-				HLODObject->SetVisibility(false);
+				UE_LOG(LogHLODRuntimeSubsystem, Verbose, TEXT("\t\t* %s"), *HLODActor->GetActorNameOrLabel());
+				HLODActor->SetVisibility(false);
 			}
 		}
 	}
@@ -393,22 +391,20 @@ void UWorldPartitionHLODRuntimeSubsystem::OnCellHidden(const UWorldPartitionRunt
 
 		if (!CellData->LoadedHLODs.IsEmpty())
 		{
-			UE_LOG(LogHLODRuntimeSubsystem, Verbose, TEXT("Cell hidden - %s - showing %d HLOD objects"), *InCell->GetGuid().ToString(), CellData->LoadedHLODs.Num());
+			UE_LOG(LogHLODRuntimeSubsystem, Verbose, TEXT("Cell hidden - %s - showing %d HLOD actors"), *InCell->GetGuid().ToString(), CellData->LoadedHLODs.Num());
 
-			for (IWorldPartitionHLODObject* HLODObject : CellData->LoadedHLODs)
+			for (AWorldPartitionHLOD* HLODActor : CellData->LoadedHLODs)
 			{
-				UE_LOG(LogHLODRuntimeSubsystem, Verbose, TEXT("\t\t* %s"), *HLODObject->GetHLODNameOrLabel());
-				HLODObject->SetVisibility(UWorldPartitionHLODRuntimeSubsystem::WorldPartitionHLODEnabled);
-				HLODObjectsToWarmup.Remove(HLODObject);
+				UE_LOG(LogHLODRuntimeSubsystem, Verbose, TEXT("\t\t* %s"), *HLODActor->GetActorNameOrLabel());
+				HLODActor->SetVisibility(UWorldPartitionHLODRuntimeSubsystem::WorldPartitionHLODEnabled);
+				HLODActorsToWarmup.Remove(HLODActor);
 			}
 		}
 	}
 }
 
-static void PrepareVTRequests(TMap<UMaterialInterface*, float>& InOutVTRequests, UMaterialInterface* InMaterialInterface, float InPixelSize)
+static void PrepareVTRequests(TMap<UMaterialInterface*, float>& InOutVTRequests, UStaticMeshComponent* InStaticMeshComponent, float InPixelSize)
 {
-	check(InMaterialInterface);
-
 	float PixelSize = InPixelSize;
 
 	// Assume the texture is wrapped around the object, so the screen size is actually less than the resolution we require.
@@ -417,32 +413,37 @@ static void PrepareVTRequests(TMap<UMaterialInterface*, float>& InOutVTRequests,
 	// Clamp for safety
 	PixelSize = FMath::Min(PixelSize, CVarHLODWarmupVTSizeClamp.GetValueOnAnyThread());
 
-	// We have a VT we'd like to prefetch, add or update a request in our request map.
-	// If the texture was already requested by another component, fetch the highest required resolution only.
-	float& CurrentMaxPixel = InOutVTRequests.FindOrAdd(InMaterialInterface);
-	CurrentMaxPixel = FMath::Max(CurrentMaxPixel, PixelSize);
+	for (UMaterialInterface* MaterialInterface : InStaticMeshComponent->GetMaterials())
+	{
+		if (MaterialInterface)
+		{
+			// We have a VT we'd like to prefetch, add or update a request in our request map.
+			// If the texture was already requested by another component, fetch the highest required resolution only.
+			float& CurrentMaxPixel = InOutVTRequests.FindOrAdd(MaterialInterface);
+			CurrentMaxPixel = FMath::Max(CurrentMaxPixel, PixelSize);
+		}
+	}
 }
 
-static void PrepareNaniteRequests(TMap<Nanite::FResources*, int32>& InOutNaniteRequests, UStaticMesh* InStaticMesh, int32 InNumFramesUntilRender)
+static void PrepareNaniteRequests(TMap<Nanite::FResources*, int32>& InOutNaniteRequests, UStaticMeshComponent* InStaticMeshComponent, int32 InNumFramesUntilRender)
 {
-	check(InStaticMesh);
-
-	if (InStaticMesh->HasValidNaniteData())
+	UStaticMesh* StaticMesh = InStaticMeshComponent->GetStaticMesh();
+	if (StaticMesh && StaticMesh->HasValidNaniteData())
 	{
-		int32& NumFramesUntilRender = InOutNaniteRequests.FindOrAdd(InStaticMesh->GetRenderData()->NaniteResourcesPtr.Get());
+		int32& NumFramesUntilRender = InOutNaniteRequests.FindOrAdd(StaticMesh->GetRenderData()->NaniteResourcesPtr.Get());
 		NumFramesUntilRender = FMath::Max(InNumFramesUntilRender, 1);
 	}
 }
 
-bool UWorldPartitionHLODRuntimeSubsystem::PrepareToWarmup(const UWorldPartitionRuntimeCell* InCell, const IWorldPartitionHLODObject* InHLODObject)
+bool UWorldPartitionHLODRuntimeSubsystem::PrepareToWarmup(const UWorldPartitionRuntimeCell* InCell, AWorldPartitionHLOD* InHLODActor)
 {
 	TRACE_CPUPROFILER_EVENT_SCOPE(UWorldPartitionHLODRuntimeSubsystem::PrepareToWarmup)
 
-	bool bHLODObjectNeedsWarmUp = false;
+	bool bHLODActorNeedsWarmUp = false;
 
-	if (InHLODObject->DoesRequireWarmup())
+	if (InHLODActor->DoesRequireWarmup())
 	{
-		FWorldPartitionHLODWarmupState& WarmupState = HLODObjectsToWarmup.FindOrAdd(InHLODObject);
+		FWorldPartitionHLODWarmupState& WarmupState = HLODActorsToWarmup.FindOrAdd(InHLODActor);
 
 		// Trigger warmup for CVarHLODWarmupNumFrames frames on the first request, or if a warmup wasn't requested in the last frame
 		const bool bResetWarmup =  WarmupState.WarmupLastRequestedFrame == INDEX_NONE ||
@@ -466,11 +467,11 @@ bool UWorldPartitionHLODRuntimeSubsystem::PrepareToWarmup(const UWorldPartitionR
 			WarmupState.WarmupCallsUntilReady--;
 		}
 		
-		bHLODObjectNeedsWarmUp = WarmupState.WarmupCallsUntilReady != 0;
+		bHLODActorNeedsWarmUp = WarmupState.WarmupCallsUntilReady != 0;
 		WarmupState.WarmupLastRequestedFrame = GFrameNumber;
 	}
 
-	return bHLODObjectNeedsWarmUp;
+	return bHLODActorNeedsWarmUp;
 }
 
 void UWorldPartitionHLODRuntimeSubsystem::OnCVarsChanged()
@@ -553,25 +554,17 @@ bool UWorldPartitionHLODRuntimeSubsystem::CanMakeVisible(const UWorldPartitionRu
 
 	bool bCanMakeVisible = true;
 
-	// Prevent cells containing HLODs objects from being made visible until warmup has been performed
+	// Prevent cells containing HLODs actors from being made visible until warmup has been performed
 	if (InCell->GetIsHLOD() && InCell->GetLevel())
 	{
-		TFunction<void(const IWorldPartitionHLODObject*)> PrepareForWarmup =
-			[&bCanMakeVisible, InCell, this](const IWorldPartitionHLODObject* HLODObject)
-			{
-				const bool bHLODObjectNeedsWarmup = PrepareToWarmup(InCell, HLODObject);
-				bCanMakeVisible &= !bHLODObjectNeedsWarmup;
-			};
-		
 		for (AActor* Actor : InCell->GetLevel()->Actors)
 		{
 			if (AWorldPartitionHLOD* HLODActor = Cast<AWorldPartitionHLOD>(Actor))
 			{
-				PrepareForWarmup(HLODActor);
+				const bool bHLODActorNeedsWarmup = PrepareToWarmup(InCell, HLODActor);
+				bCanMakeVisible &= !bHLODActorNeedsWarmup;
 			}
 		}
-
-		ForEachHLODObjectInCellEvent.Broadcast(InCell, PrepareForWarmup);
 	}
 
 	return bCanMakeVisible;
@@ -588,24 +581,19 @@ bool UWorldPartitionHLODRuntimeSubsystem::CanMakeInvisible(const UWorldPartition
 
 	bool bCanMakeInvisible = true;
 
-	// Prevent cells from being hidden until their matching HLOD objects have been warmed up
+	// Prevent cells which have an HLOD representation from being hidden until their matching HLOD actors have been warmed up
 	if (FCellData* CellData = GetCellData(InCell))
 	{
-		for (IWorldPartitionHLODObject* HLODObject : CellData->LoadedHLODs)
+		for (AWorldPartitionHLOD* HLODActor : CellData->LoadedHLODs)
 		{
-			const bool bHLODObjectNeedsWarmup = PrepareToWarmup(InCell, HLODObject);
-			bCanMakeInvisible &= !bHLODObjectNeedsWarmup;
+			const bool bHLODActorNeedsWarmup = PrepareToWarmup(InCell, HLODActor);
+			bCanMakeInvisible &= !bHLODActorNeedsWarmup;
 		}
 	}
 
 	return bCanMakeInvisible;
 }
 
-FWorldPartitionHLODForEachHLODObjectInCellEvent& UWorldPartitionHLODRuntimeSubsystem::GetForEachHLODObjectInCellEvent()
-{
-	return ForEachHLODObjectInCellEvent;
-}
-
 static float IsInView(const FVector& BoundsOrigin, const FVector& BoundsExtent, const FSceneViewFamily& InViewFamily, bool bComputeScreenSize, float& MaxScreenSizePixels)
 {
 	MaxScreenSizePixels = 0;
@@ -613,7 +601,7 @@ static float IsInView(const FVector& BoundsOrigin, const FVector& BoundsExtent,
 	// Estimate the highest screen pixel size of this component in the provided views
 	for (const FSceneView* View : InViewFamily.Views)
 	{
-		// Make sure the HLOD Object we're about to show is actually in the frustum
+		// Make sure the HLOD actor we're about to show is actually in the frustum
 		if (View->ViewFrustum.IntersectBox(BoundsOrigin, BoundsExtent))
 		{
 			if (bComputeScreenSize)
@@ -669,53 +657,63 @@ void UWorldPartitionHLODRuntimeSubsystem::OnBeginRenderViews(const FSceneViewFam
 	const bool bWarmupNanite = CVarHLODWarmupNanite.GetValueOnGameThread() != 0;
 	const bool bWarmupVT = CVarHLODWarmupVT.GetValueOnGameThread() != 0;
 
-	for (FHLODWarmupStateMap::TIterator HLODObjectWarmupStateIt(HLODObjectsToWarmup); HLODObjectWarmupStateIt; ++HLODObjectWarmupStateIt)
+	for (FHLODWarmupStateMap::TIterator HLODActorWarmupStateIt(HLODActorsToWarmup); HLODActorWarmupStateIt; ++HLODActorWarmupStateIt)
 	{
-		const IWorldPartitionHLODObject* HLODObject = HLODObjectWarmupStateIt.Key();
-		FWorldPartitionHLODWarmupState& HLODWarmupState = HLODObjectWarmupStateIt.Value();
-
-		// Retrieve this component's bound - we must support getting the bounds before the component is even registered.
-		FVector BoundsOrigin;
-		FVector BoundsExtent;
-		HLODWarmupState.WarmupBounds.GetCenterAndExtents(BoundsOrigin, BoundsExtent);
-
-		float ScreenSizePixels = 0;
-		if (IsInView(BoundsOrigin, BoundsExtent, InViewFamily, bWarmupVT, ScreenSizePixels))
+		const FObjectKey& HLODActorObjectKey = HLODActorWarmupStateIt.Key();
+		FWorldPartitionHLODWarmupState& HLODWarmupState = HLODActorWarmupStateIt.Value();
+		const AWorldPartitionHLOD* HLODActor = Cast<AWorldPartitionHLOD>(HLODActorObjectKey.ResolveObjectPtr());
+		
+		if (HLODActor)
 		{
-			for (UObject* AssetToWarmup : HLODObject->GetAssetsToWarmup())
+			// Retrieve this component's bound - we must support getting the bounds before the component is even registered.
+			FVector BoundsOrigin;
+			FVector BoundsExtent;
+			HLODWarmupState.WarmupBounds.GetCenterAndExtents(BoundsOrigin, BoundsExtent);
+
+			float ScreenSizePixels = 0;
+			if (IsInView(BoundsOrigin, BoundsExtent, InViewFamily, bWarmupVT, ScreenSizePixels))
 			{
-				if (UMaterialInterface* MaterialInterface = Cast<UMaterialInterface>(AssetToWarmup))
+				HLODActor->ForEachComponent<UStaticMeshComponent>(false, [&](UStaticMeshComponent* SMC)
 				{
-				    if (bWarmupVT)
-				    {
-					    PrepareVTRequests(VTRequests, MaterialInterface, ScreenSizePixels);
-				    }
-				}
+					// Assume ISM HLOD don't need warmup, as they are actually found in the source level
+					if (SMC->IsA<UInstancedStaticMeshComponent>())
+					{
+						return;
+					}
+				
+					if (bWarmupVT)
+					{
+						PrepareVTRequests(VTRequests, SMC, ScreenSizePixels);
+					}
 
-				if (UStaticMesh* StaticMesh = Cast<UStaticMesh>(AssetToWarmup))
+					if (bWarmupNanite)
+					{
+						if (HLODWarmupState.WarmupCallsUntilReady == CVarHLODWarmupNumFrames.GetValueOnGameThread())
+						{
+							// Send a nanite request to prepare for visibility in CVarHLODWarmupNumFrames frames
+							PrepareNaniteRequests(NaniteRequests, SMC, CVarHLODWarmupNumFrames.GetValueOnGameThread());
+						}
+						else if (HLODWarmupState.WarmupCallsUntilReady == 0)
+						{
+							// We expect HLOD to be visible at any moment (likely waiting for server visibility ack)
+							PrepareNaniteRequests(NaniteRequests, SMC, 1);
+						}
+					}
+				});
+
+#if ENABLE_DRAW_DEBUG
+				if (CVarHLODWarmupDebugDraw.GetValueOnAnyThread())
 				{
-				    if (bWarmupNanite)
-				    {
-					    if (HLODWarmupState.WarmupCallsUntilReady == CVarHLODWarmupNumFrames.GetValueOnGameThread())
-					    {
-						    // Send a nanite request to prepare for visibility in CVarHLODWarmupNumFrames frames
-						    PrepareNaniteRequests(NaniteRequests, StaticMesh, CVarHLODWarmupNumFrames.GetValueOnGameThread());
-					    }
-					    else if (HLODWarmupState.WarmupCallsUntilReady == 0)
-					    {
-						    // We expect HLOD to be visible at any moment (likely waiting for server visibility ack)
-						    PrepareNaniteRequests(NaniteRequests, StaticMesh, 1);
-					    }
-				    }
+					DrawDebugBox(HLODActor->GetWorld(), BoundsOrigin, BoundsExtent, FColor::Yellow, /*bPersistentLine*/ false, /*Lifetime*/ 1.0f);
 				}
+#endif
 			}
+		}
 
-#if ENABLE_DRAW_DEBUG
-			if (CVarHLODWarmupDebugDraw.GetValueOnAnyThread())
-			{
-				DrawDebugBox(HLODObject->GetHLODLevel()->GetWorld(), BoundsOrigin, BoundsExtent, HLODWarmupState.WarmupCallsUntilReady > 0 ? FColor::Yellow : FColor::Green, /*bPersistentLine*/ false, /*Lifetime*/ 0.5f);
-			}
-#endif
+		if (!HLODActor)
+		{
+			// HLOD actor has been unloaded, stop tracking it's warmup state
+			HLODActorWarmupStateIt.RemoveCurrent();
 		}
 	}
 
diff --git a/Engine/Source/Runtime/Engine/Public/WorldPartition/HLOD/HLODActor.h b/Engine/Source/Runtime/Engine/Public/WorldPartition/HLOD/HLODActor.h
--- a/Engine/Source/Runtime/Engine/Public/WorldPartition/HLOD/HLODActor.h
+++ b/Engine/Source/Runtime/Engine/Public/WorldPartition/HLOD/HLODActor.h
@@ -9,7 +9,6 @@
 #include "WorldPartition/WorldPartitionActorDesc.h"
 #include "WorldPartition/HLOD/HLODStats.h"
 #include "WorldPartition/HLOD/HLODBuilder.h"
-#include "WorldPartition/HLOD/IWorldPartitionHLODObject.h"
 
 #if WITH_EDITOR
 #include "UObject/ObjectSaveContext.h"
@@ -23,7 +22,7 @@ class UWorldPartitionHLODSourceActors;
 ENGINE_API DECLARE_LOG_CATEGORY_EXTERN(LogHLODHash, Log, All);
 
 UCLASS(NotPlaceable, MinimalAPI, HideCategories=(Rendering, Replication, Collision, Physics, Navigation, Networking, Input, Actor, LevelInstance, Cooking))
-class AWorldPartitionHLOD : public AActor, public IWorldPartitionHLODObject
+class AWorldPartitionHLOD : public AActor
 {
 	GENERATED_UCLASS_BODY()
 
@@ -31,19 +30,14 @@ class AWorldPartitionHLOD : public AActor, public IWorldPartitionHLODObject
 	typedef TMap<FName, int64> FStats;
 
 public:
+	ENGINE_API void SetVisibility(bool bInVisible);
+
+	ENGINE_API const FGuid& GetSourceCellGuid() const;
 	inline uint32 GetLODLevel() const { return LODLevel; }
 
 	virtual bool IsHLODRelevant() const override { return true; }
 
-	// ~Begin IWorldPartitionHLODObject interface
-	ENGINE_API virtual UObject* GetUObject() const override { return const_cast<AWorldPartitionHLOD*>(this); }
-	ENGINE_API virtual ULevel* GetHLODLevel() const override;
-	ENGINE_API virtual FString GetHLODNameOrLabel() const override;
-	ENGINE_API virtual bool DoesRequireWarmup() const override { return bRequireWarmup; }
-	ENGINE_API virtual TSet<UObject*> GetAssetsToWarmup() const override;
-	ENGINE_API virtual void SetVisibility(bool bIsVisible) override;
-	ENGINE_API virtual const FGuid& GetSourceCellGuid() const override;
-	// ~End IWorldPartitionHLODObject interface
+	bool DoesRequireWarmup() const { return bRequireWarmup; }
 
 #if WITH_EDITOR
 	ENGINE_API void SetHLODComponents(const TArray<UActorComponent*>& InHLODComponents);
diff --git a/Engine/Source/Runtime/Engine/Public/WorldPartition/HLOD/HLODRuntimeSubsystem.h b/Engine/Source/Runtime/Engine/Public/WorldPartition/HLOD/HLODRuntimeSubsystem.h
--- a/Engine/Source/Runtime/Engine/Public/WorldPartition/HLOD/HLODRuntimeSubsystem.h
+++ b/Engine/Source/Runtime/Engine/Public/WorldPartition/HLOD/HLODRuntimeSubsystem.h
@@ -6,10 +6,9 @@
 #include "Containers/Map.h"
 #include "UObject/ObjectKey.h"
 #include "Containers/Ticker.h"
-#include "Templates/Function.h"
 #include "HLODRuntimeSubsystem.generated.h"
 
-class IWorldPartitionHLODObject;
+
 class AWorldPartitionHLOD;
 class FSceneViewFamily;
 class FHLODResourcesResidencySceneViewExtension;
@@ -18,15 +17,10 @@ class UWorldPartitionRuntimeCell;
 class URuntimeHashExternalStreamingObjectBase;
 class UWorld;
 
-DECLARE_MULTICAST_DELEGATE_OneParam(FWorldPartitionHLODObjectRegisteredEvent, IWorldPartitionHLODObject* /* InHLODObject */);
-DECLARE_MULTICAST_DELEGATE_OneParam(FWorldPartitionHLODObjectUnregisteredEvent, IWorldPartitionHLODObject* /* InHLODObject */);
-DECLARE_MULTICAST_DELEGATE_TwoParams(FWorldPartitionHLODForEachHLODObjectInCellEvent, const UWorldPartitionRuntimeCell* /* InCell */, TFunction<void(const IWorldPartitionHLODObject*)> /* InFunc */);
-
-UE_DEPRECATED(5.6, "Use FWorldPartitionHLODObjectRegisteredEvent instead")
 DECLARE_MULTICAST_DELEGATE_OneParam(FWorldPartitionHLODActorRegisteredEvent, AWorldPartitionHLOD* /* InHLODActor */);
-UE_DEPRECATED(5.6, "Use FWorldPartitionHLODObjectUnregisteredEvent instead")
 DECLARE_MULTICAST_DELEGATE_OneParam(FWorldPartitionHLODActorUnregisteredEvent, AWorldPartitionHLOD* /* InHLODActor */);
 
+
 /**
  * UWorldPartitionHLODRuntimeSubsystem
  */
@@ -47,8 +41,8 @@ public:
 	ENGINE_API virtual bool DoesSupportWorldType(const EWorldType::Type WorldType) const override;
 	//~ End UWorldSubsystem Interface.
 
-	ENGINE_API void RegisterHLODObject(IWorldPartitionHLODObject* InWorldPartitionHLOD);
-	ENGINE_API void UnregisterHLODObject(IWorldPartitionHLODObject* InWorldPartitionHLOD);
+	ENGINE_API void RegisterHLODActor(AWorldPartitionHLOD* InWorldPartitionHLOD);
+	ENGINE_API void UnregisterHLODActor(AWorldPartitionHLOD* InWorldPartitionHLOD);
 
 	ENGINE_API void OnCellShown(const UWorldPartitionRuntimeCell* InCell);
 	ENGINE_API void OnCellHidden(const UWorldPartitionRuntimeCell* InCell);
@@ -56,13 +50,12 @@ public:
 	ENGINE_API bool CanMakeVisible(const UWorldPartitionRuntimeCell* InCell);
 	ENGINE_API bool CanMakeInvisible(const UWorldPartitionRuntimeCell* InCell);
 
-	ENGINE_API FWorldPartitionHLODForEachHLODObjectInCellEvent& GetForEachHLODObjectInCellEvent();
+	ENGINE_API const TArray<AWorldPartitionHLOD*>& GetHLODActorsForCell(const UWorldPartitionRuntimeCell* InCell) const;
 
 	static ENGINE_API bool IsHLODEnabled();
 
-	ENGINE_API const TArray<IWorldPartitionHLODObject*>& GetHLODObjectsForCell(const UWorldPartitionRuntimeCell* InCell) const;
-	FWorldPartitionHLODObjectRegisteredEvent& OnHLODObjectRegisteredEvent() { return HLODObjectRegisteredEvent; }
-	FWorldPartitionHLODObjectUnregisteredEvent& OnHLODObjectUnregisteredEvent() { return HLODObjectUnregisteredEvent; }
+	FWorldPartitionHLODActorRegisteredEvent& OnHLODActorRegisteredEvent() { return HLODActorRegisteredEvent; }
+	FWorldPartitionHLODActorUnregisteredEvent& OnHLODActorUnregisteredEvent() { return HLODActorUnregisteredEvent; }
 
 	ENGINE_API void OnExternalStreamingObjectInjected(URuntimeHashExternalStreamingObjectBase* ExternalStreamingObject);
 	ENGINE_API void OnExternalStreamingObjectRemoved(URuntimeHashExternalStreamingObjectBase* ExternalStreamingObject);
@@ -70,33 +63,9 @@ public:
 	ENGINE_API void OnCVarsChanged();
 
 #if !UE_BUILD_SHIPPING
-	uint32 GetNumOutdatedHLODObjects() const { return OutdatedHLODObjects.Num(); }
+	uint32 GetNumOutdatedHLODActors() const { return OutdatedHLODActors.Num(); }
 #endif
 
-	UE_DEPRECATED(5.6, "Use RegisterHLODObject instead")
-	ENGINE_API void RegisterHLODActor(AWorldPartitionHLOD* InWorldPartitionHLOD) {}
-	UE_DEPRECATED(5.6, "Use UnregisterHLODObject instead")
-	ENGINE_API void UnregisterHLODActor(AWorldPartitionHLOD* InWorldPartitionHLOD) {}
-
-	UE_DEPRECATED(5.6, "Use GetHLODObjectsForCell instead")
-	ENGINE_API const TArray<AWorldPartitionHLOD*>& GetHLODActorsForCell(const UWorldPartitionRuntimeCell* InCell) const
-	{
-		static TArray<AWorldPartitionHLOD*> Empty;
-		return Empty;
-	}
-
-#if !UE_BUILD_SHIPPING
-	UE_DEPRECATED(5.6, "Use GetNumOutdatedHLODObjects instead")
-	uint32 GetNumOutdatedHLODActors() const { return 0; }
-#endif
-
-	PRAGMA_DISABLE_DEPRECATION_WARNINGS
-	UE_DEPRECATED(5.6, "Use OnHLODObjectRegisteredEvent instead")
-	FWorldPartitionHLODActorRegisteredEvent& OnHLODActorRegisteredEvent() { return HLODActorRegisteredEvent; }
-	UE_DEPRECATED(5.6, "Use OnHLODObjectUnregisteredEvent instead")
-	FWorldPartitionHLODActorUnregisteredEvent& OnHLODActorUnregisteredEvent() {	return HLODActorUnregisteredEvent; }
-	PRAGMA_ENABLE_DEPRECATION_WARNINGS
-
 #if WITH_EDITOR
 	UE_DEPRECATED(5.5, "Use UWorldPartitionHLODEditorSubsystem::WriteHLODStatsCSV()")
 	static ENGINE_API bool WriteHLODStatsCSV(UWorld* InWorld, const FString& InFilename) { return false; }
@@ -106,12 +75,10 @@ public:
 	ENGINE_API void SetHLODAlwaysLoadedCullDistance(int32 InCullDistance) {}
 	
 private:
-	FWorldPartitionHLODForEachHLODObjectInCellEvent ForEachHLODObjectInCellEvent;
-
 	struct FCellData
 	{
 		bool bIsCellVisible = false;
-		TArray<IWorldPartitionHLODObject*> LoadedHLODs;	// HLOD representation of the cell itself
+		TArray<AWorldPartitionHLOD*> LoadedHLODs;
 	};
 
 	struct FWorldPartitionHLODRuntimeData
@@ -122,7 +89,7 @@ private:
 	TMap<TObjectPtr<UWorldPartition>, FWorldPartitionHLODRuntimeData> WorldPartitionsHLODRuntimeData;
 	ENGINE_API const FCellData* GetCellData(const UWorldPartitionRuntimeCell* InCell) const;
 	ENGINE_API FCellData* GetCellData(const UWorldPartitionRuntimeCell* InCell);
-	ENGINE_API FCellData* GetCellData(IWorldPartitionHLODObject* InWorldPartitionHLOD);
+	ENGINE_API FCellData* GetCellData(AWorldPartitionHLOD* InWorldPartitionHLOD);
 
 	struct FWorldPartitionHLODWarmupState
 	{
@@ -131,11 +98,11 @@ private:
 		FBox WarmupBounds;
 	};
 
-	typedef TMap<const IWorldPartitionHLODObject*, FWorldPartitionHLODWarmupState> FHLODWarmupStateMap;
-	FHLODWarmupStateMap HLODObjectsToWarmup;
+	typedef TMap<FObjectKey, FWorldPartitionHLODWarmupState> FHLODWarmupStateMap;
+	FHLODWarmupStateMap HLODActorsToWarmup;
 
 	ENGINE_API void OnBeginRenderViews(const FSceneViewFamily& InViewFamily);
-	ENGINE_API bool PrepareToWarmup(const UWorldPartitionRuntimeCell* InCell, const IWorldPartitionHLODObject* InHLODObject);
+	ENGINE_API bool PrepareToWarmup(const UWorldPartitionRuntimeCell* InCell, AWorldPartitionHLOD* InHLODActor);
 	ENGINE_API bool ShouldPerformWarmup() const;
 	ENGINE_API bool ShouldPerformWarmupForCell(const UWorldPartitionRuntimeCell* InCell) const;
 	bool bCachedShouldPerformWarmup;
@@ -151,15 +118,11 @@ private:
 	ENGINE_API void OnWorldPartitionInitialized(UWorldPartition* InWorldPartition);
 	ENGINE_API void OnWorldPartitionUninitialized(UWorldPartition* InWorldPartition);
 
-	FWorldPartitionHLODObjectRegisteredEvent	HLODObjectRegisteredEvent;
-	FWorldPartitionHLODObjectUnregisteredEvent	HLODObjectUnregisteredEvent;
-
-	PRAGMA_DISABLE_DEPRECATION_WARNINGS
 	FWorldPartitionHLODActorRegisteredEvent		HLODActorRegisteredEvent;
 	FWorldPartitionHLODActorUnregisteredEvent	HLODActorUnregisteredEvent;
-	PRAGMA_ENABLE_DEPRECATION_WARNINGS
 
 #if !UE_BUILD_SHIPPING
-	TSet<IWorldPartitionHLODObject*> OutdatedHLODObjects;
+	TSet<AWorldPartitionHLOD*> OutdatedHLODActors;
 #endif
 };
+
diff --git a/Engine/Source/Runtime/Engine/Public/WorldPartition/HLOD/IWorldPartitionHLODObject.h b/Engine/Source/Runtime/Engine/Public/WorldPartition/HLOD/IWorldPartitionHLODObject.h
deleted file mode 100644
