Updated the comment for MeshDeformer property to explain this can be used to disable the deformer on the mesh asset

diff --git a/Engine/Source/Runtime/Engine/Classes/Components/SkinnedMeshComponent.h b/Engine/Source/Runtime/Engine/Classes/Components/SkinnedMeshComponent.h
--- a/Engine/Source/Runtime/Engine/Classes/Components/SkinnedMeshComponent.h
+++ b/Engine/Source/Runtime/Engine/Classes/Components/SkinnedMeshComponent.h
@@ -301,7 +301,7 @@ protected:
 	UPROPERTY(EditAnywhere, BlueprintReadOnly, Category = "Deformer", meta = (InlineEditConditionToggle))
 	bool bSetMeshDeformer = false;
 
-	/** The mesh deformer to use. If no mesh deformer is set from here or the SkeletalMesh, then we fall back to the fixed function deformation, unless AlwaysUseMeshDeformer is turned on. */
+	/** The mesh deformer to use. Set to None to disable the deformer on the SkeletalMesh. If no deformer is set from here or the SkeletalMesh, we fall back to the fixed function deformation, unless AlwaysUseMeshDeformer is on. */
 	UPROPERTY(EditAnywhere, BlueprintReadOnly, Category = "Deformer", meta = (editcondition = "bSetMeshDeformer"))
 	TObjectPtr<UMeshDeformer> MeshDeformer;
 
