Change RedirectClassPath to use FNames avoiding FString path for performance.
#rnx

diff --git a/Engine/Source/Runtime/Engine/Private/WorldPartition/WorldPartitionClassDescRegistry.cpp b/Engine/Source/Runtime/Engine/Private/WorldPartition/WorldPartitionClassDescRegistry.cpp
--- a/Engine/Source/Runtime/Engine/Private/WorldPartition/WorldPartitionClassDescRegistry.cpp
+++ b/Engine/Source/Runtime/Engine/Private/WorldPartition/WorldPartitionClassDescRegistry.cpp
@@ -733,9 +733,9 @@ void FWorldPartitionClassDescRegistry::ValidateInternalState()
 FTopLevelAssetPath FWorldPartitionClassDescRegistry::RedirectClassPath(const FTopLevelAssetPath& InClassPath) const
 {
 	// Resolve native class redirectors
-	const FCoreRedirectObjectName OldClassName = FCoreRedirectObjectName(InClassPath.GetAssetName(), NAME_None, InClassPath.GetPackageName());
+	const FCoreRedirectObjectName OldClassName(InClassPath.GetAssetName(), NAME_None, InClassPath.GetPackageName());
 	const FCoreRedirectObjectName NewClassName = FCoreRedirects::GetRedirectedName(ECoreRedirectFlags::Type_Class, OldClassName);
-	const FTopLevelAssetPath ResolvedClassPath = FTopLevelAssetPath(NewClassName.ToString());
+	const FTopLevelAssetPath ResolvedClassPath(NewClassName.PackageName, NewClassName.ObjectName);
 
 	// Resolve blueprint object redirectors
 	if (const FTopLevelAssetPath* RedirectedClass = RedirectClassMap.Find(ResolvedClassPath))
