Fix help/comments - no functional changes

diff --git a/Engine/Source/Runtime/Engine/Public/Physics/ImmediatePhysics/ImmediatePhysicsAdapters.h b/Engine/Source/Runtime/Engine/Public/Physics/ImmediatePhysics/ImmediatePhysicsAdapters.h
--- a/Engine/Source/Runtime/Engine/Public/Physics/ImmediatePhysics/ImmediatePhysicsAdapters.h
+++ b/Engine/Source/Runtime/Engine/Public/Physics/ImmediatePhysics/ImmediatePhysicsAdapters.h
@@ -37,15 +37,15 @@ namespace ImmediatePhysics
 	// Extracts parameters from a constraint constraint into the joint settings
 	ENGINE_API void UpdateJointSettingsFromLinearConstraint(const FLinearConstraint& Constraint, Chaos::FPBDJointSettings& JointSettings);
 	
-	//Extracts parameters from a constraint constraint into the joint settings
+	// Extracts parameters from a cone constraint into the joint settings
 	ENGINE_API void UpdateJointSettingsFromConeConstraint(const FConeConstraint& Constraint, Chaos::FPBDJointSettings& JointSettings);
 	
-	//Extracts parameters from a constraint constraint into the joint settings
+	// Extracts parameters from a twist constraint into the joint settings
 	ENGINE_API void UpdateJointSettingsFromTwistConstraint(const FTwistConstraint& Constraint, Chaos::FPBDJointSettings& JointSettings);
 
-	ENGINE_API void UpdateJointSettingsFromLinearDriveConstraint(
-		const FLinearDriveConstraint& Constraint, Chaos::FPBDJointSettings& JointSettings);
-	ENGINE_API void UpdateJointSettingsFromAngularDriveConstraint(
-		const FAngularDriveConstraint& Constraint, Chaos::FPBDJointSettings& JointSettings);
+	// Extracts parameters from a linear drive constraint into the joint settings
+	ENGINE_API void UpdateJointSettingsFromLinearDriveConstraint(const FLinearDriveConstraint& Constraint, Chaos::FPBDJointSettings& JointSettings);
 
+	// Extracts parameters from an angular drive constraint into the joint settings
+	ENGINE_API void UpdateJointSettingsFromAngularDriveConstraint(const FAngularDriveConstraint& Constraint, Chaos::FPBDJointSettings& JointSettings);
 }
