[Navigation] prevent navigation updates on component destroyed during ConsolidatedPostEditChange
Modified UPrimitiveComponent::IsNavigationRelevant() to test validity of the BodyInstance to be relevant (was causing issue after destroying the physic state since component was still considered relevant)
Add CVAR "p.InvalidBodyInstanceIsNeverRelevantToNavigation" to revert to previous behavior if needed.

diff --git a/Engine/Source/Runtime/Engine/Private/Components/ActorComponent.cpp b/Engine/Source/Runtime/Engine/Private/Components/ActorComponent.cpp
--- a/Engine/Source/Runtime/Engine/Private/Components/ActorComponent.cpp
+++ b/Engine/Source/Runtime/Engine/Private/Components/ActorComponent.cpp
@@ -1183,19 +1183,23 @@ void UActorComponent::ConsolidatedPostEditChange(const FPropertyChangedEvent& Pr
 		}
 	}
 
-	if (PropertyChangedEvent.Property != nullptr && PropertyChangedEvent.Property->GetFName() == NAME_CanEverAffectNavigation)
-	{
-		HandleCanEverAffectNavigationChange(/*bForce=*/true);
-	}
-
 	// The component or its outer could be pending kill when calling PostEditChange when applying a transaction.
-	// Don't do do a full recreate in this situation, and instead simply detach.
-	if(!IsValidChecked(this))
+	// Don't do a full recreate in this situation, and instead simply detach.
+	if (!IsValidChecked(this))
 	{
 		// @todo james should this call UnregisterComponent instead to remove itself from the RegisteredComponents array on the owner?
 		ExecuteUnregisterEvents();
 		WorldPrivate = nullptr;
 	}
+	else
+	{
+		// Navigation specific fix until UE-252220 is fixed
+		// Only process navigation status changes on valid components (i.e., RerunConstructionScripts will destroy this component)
+		if (PropertyChangedEvent.Property != nullptr && PropertyChangedEvent.Property->GetFName() == NAME_CanEverAffectNavigation)
+		{
+			HandleCanEverAffectNavigationChange(/*bForce=*/true);
+		}
+	}
 
 	for (UAssetUserData* Datum : AssetUserData)
 	{
diff --git a/Engine/Source/Runtime/Engine/Private/Components/PrimitiveComponent.cpp b/Engine/Source/Runtime/Engine/Private/Components/PrimitiveComponent.cpp
--- a/Engine/Source/Runtime/Engine/Private/Components/PrimitiveComponent.cpp
+++ b/Engine/Source/Runtime/Engine/Private/Components/PrimitiveComponent.cpp
@@ -125,7 +125,17 @@ namespace PrimitiveComponentCVars
 		ECVF_Default);
 
 	int32 bEnableFastOverlapCheck = 1;
-	static FAutoConsoleVariableRef CVarEnableFastOverlapCheck(TEXT("p.EnableFastOverlapCheck"), bEnableFastOverlapCheck, TEXT("Enable fast overlap check against sweep hits, avoiding UpdateOverlaps (for the swept component)."));
+	static FAutoConsoleVariableRef CVarEnableFastOverlapCheck(
+		TEXT("p.EnableFastOverlapCheck"),
+		bEnableFastOverlapCheck,
+		TEXT("Enable fast overlap check against sweep hits, avoiding UpdateOverlaps (for the swept component)."));
+
+	bool bInvalidBodyInstanceNeverRelevantToNavigation = true;
+	static FAutoConsoleVariableRef InvalidBodyInstanceNeverRelevantToNavigation(
+		TEXT("p.InvalidBodyInstanceIsNeverRelevantToNavigation"),
+		bInvalidBodyInstanceNeverRelevantToNavigation,
+		TEXT("0: IsNavigationRelevant ignores the validity of the body instance (legacy), 1: IsNavigationRelevant returns false for invalid body instance (default)"),
+		ECVF_Default);
 }
 
 namespace PhysicsReplicationCVars
@@ -3385,6 +3395,12 @@ bool UPrimitiveComponent::IsNavigationRelevant() const
 		return false;
 	}
 
+	if (PrimitiveComponentCVars::bInvalidBodyInstanceNeverRelevantToNavigation
+		&& !BodyInstance.IsValidBodyInstance())
+	{
+		return false;
+	}
+
 	if (HasCustomNavigableGeometry() >= EHasCustomNavigableGeometry::EvenIfNotCollidable)
 	{
 		return true;
diff --git a/Engine/Source/Runtime/Engine/Private/Components/StaticMeshComponent.cpp b/Engine/Source/Runtime/Engine/Private/Components/StaticMeshComponent.cpp
--- a/Engine/Source/Runtime/Engine/Private/Components/StaticMeshComponent.cpp
+++ b/Engine/Source/Runtime/Engine/Private/Components/StaticMeshComponent.cpp
@@ -3209,7 +3209,7 @@ bool UStaticMeshComponent::IsNavigationRelevant() const
 {
 	if (const UStaticMesh* Mesh = GetStaticMesh())
 	{
-		// Pending compilation, update to the the navigation system will be done once compilation finishes.
+		// Pending compilation, update to the navigation system will be done once compilation finishes.
 		return !Mesh->IsCompiling() && Mesh->IsNavigationRelevant() && Super::IsNavigationRelevant();
 	}
 
diff --git a/Engine/Source/Runtime/NavigationSystem/Private/NavigationSystem.cpp b/Engine/Source/Runtime/NavigationSystem/Private/NavigationSystem.cpp
--- a/Engine/Source/Runtime/NavigationSystem/Private/NavigationSystem.cpp
+++ b/Engine/Source/Runtime/NavigationSystem/Private/NavigationSystem.cpp
@@ -3739,6 +3739,14 @@ void UNavigationSystemV1::UpdateComponentInNavOctree(UActorComponent& Comp)
 		return;
 	}
 
+	// Due to an issue with PostEditChangeProperty and AActor::RerunConstructionScripts()
+	// we need to make sure that we are not processing an invalid component.
+	// Could be converted to an ensure once UE-252220 is fixed
+	if (!IsValid(&Comp))
+	{
+		return;
+	}
+
 	if (UE::Navigation::Private::ShouldComponentWaitForActorToRegister(&Comp))
 	{
 		return;
