Deprecate the unused function WantsPollingMouseMovement on the game viewport client. Nothing uses this in the engine, so if there are callsites in your project simply subclass the game viewport client and add it there.

#rnx

diff --git a/Engine/Source/Runtime/Engine/Public/ViewportClient.h b/Engine/Source/Runtime/Engine/Public/ViewportClient.h
--- a/Engine/Source/Runtime/Engine/Public/ViewportClient.h
+++ b/Engine/Source/Runtime/Engine/Public/ViewportClient.h
@@ -131,7 +131,7 @@ public:
 
 	virtual void SetIsSimulateInEditorViewport(bool bInIsSimulateInEditorViewport) { };
 
-	// TODO: Deprecate? Nothing uses this.
+	UE_DEPRECATED(5.6, "This function will be removed in a future release.")
 	virtual bool WantsPollingMouseMovement(void) const { return true; }
 
 	virtual void MouseEnter( FViewport* Viewport,int32 x, int32 y ) {}
