New material translator changes:
- Introduced MIR::EGraphProperties system to give values the ability to "talk" to each other even indirectly. A value can turn on one of these flags, and they are automatically propagated to all values using it. This effectively allows for a value to know whether upstream, some indirect dependency turned some knob on.
- Factored various "flags" pieces of data inside a single EValueFlags enum for efficiency and conciseveness.
- Renamed MIR::FExternalCode to MIR::FInlineHLSL for a more consistent naming.
- Wrote more comments and inline documentation.

#rnx

diff --git a/Engine/Source/Runtime/Engine/Private/Materials/MaterialExpressionsIR.cpp b/Engine/Source/Runtime/Engine/Private/Materials/MaterialExpressionsIR.cpp
--- a/Engine/Source/Runtime/Engine/Private/Materials/MaterialExpressionsIR.cpp
+++ b/Engine/Source/Runtime/Engine/Private/Materials/MaterialExpressionsIR.cpp
@@ -296,7 +296,7 @@ static MIR::FValue* BuildMaterialExpressionParameter(UMaterialExpressionParamete
 	if (!ParameterExpr->GetParameterValue(Metadata))
 	{
 		Emitter.Error(TEXT("Could not get parameter value."));
-		return nullptr;
+		return MIR::FPoison::Get();
 	}
 
 	return Emitter.EmitParameter(ParameterExpr->GetParameterName(), Metadata);
@@ -320,18 +320,10 @@ void UMaterialExpressionVectorParameter::Build(MIR::FEmitter& Emitter)
 void UMaterialExpressionChannelMaskParameter::Build(MIR::FEmitter& Emitter)
 {
 	MIR::FValue* InputValue = Emitter.Input(&Input);
-
-	UE_MIR_CHECKPOINT();
-
 	InputValue = Emitter.EmitConstruct(MIR::FPrimitiveType::GetVector(MIR::ScalarKind_Float, 4), InputValue);
 	MIR::FValue* ParameterValue = BuildMaterialExpressionParameter(this, Emitter);
-
-	UE_MIR_CHECKPOINT();
-
 	MIR::FValue* DotResult = Emitter.EmitOperator(MIR::BO_Dot, InputValue, ParameterValue);
 
-	UE_MIR_CHECKPOINT();
-
 	Emitter.Output(0, DotResult);
 	Emitter.Output(1, Emitter.EmitSubscript(DotResult, 0));
 	Emitter.Output(2, Emitter.EmitSubscript(DotResult, 1));
@@ -467,9 +459,6 @@ static void BuildUnaryArithmeticOperatorFromInputWithDefault(
 {
 	// Default input to its relative constants if disconnected, then get it after checking it has primitive type.
 	MIR::FValue* Value = Emitter.DefaultInputToFloatScalar(Input, Const).Input(Input);
-
-	UE_MIR_CHECKPOINT();
-
 	BuildUnaryArithmeticOperatorFromValue(Emitter, Op, Value);
 }
 
@@ -481,9 +470,6 @@ static void BuildUnaryArithmeticOperatorFromInput(
 {
 	// Get input after checking it has primitive type.
 	MIR::FValue* Value = Emitter.Input(Input);
-
-	UE_MIR_CHECKPOINT();
-
 	BuildUnaryArithmeticOperatorFromValue(Emitter, Op, Value);
 }
 
@@ -512,8 +498,6 @@ void UMaterialExpressionFrac::Build(MIR::FEmitter& Emitter)
 void UMaterialExpressionLength::Build(MIR::FEmitter& Emitter)
 {
 	MIR::FValue* Value = Emitter.Input(&Input);
-	UE_MIR_CHECKPOINT();
-
 	BuildUnaryArithmeticOperatorFromValue(Emitter, MIR::UO_Length, Value);
 }
 
@@ -604,7 +588,7 @@ void UMaterialExpressionExternalCodeBase::Build(MIR::FEmitter& Emitter)
 	{
 		if (const FMaterialExternalCodeDeclaration* ExternalCodeDeclaration = MaterialExternalCodeRegistry::Get().FindExternalCode(ExternalCodeIdentifiers[OutputIndex]))
 		{
-			MIR::FValue* Value = Emitter.EmitExternalCode(ExternalCodeDeclaration);
+			MIR::FValue* Value = Emitter.EmitInlineHLSL(ExternalCodeDeclaration);
 			Emitter.Output(OutputIndex, Value);
 		}
 		else
@@ -1126,14 +1110,20 @@ void UMaterialExpressionViewProperty::Build(MIR::FEmitter& Emitter)
 			Emitter.Output(OutputIndex,
 				Emitter.EmitOperator(MIR::BO_Divide,
 					Emitter.EmitConstantFloatScalar(1.0f),
-					Emitter.EmitExternalCode(HLSLCodeType, MoveTemp(HLSLCode))));
+					Emitter.EmitInlineHLSL(HLSLCodeType, MoveTemp(HLSLCode), {}, {})));
 		}
 		else
 		{
 			// @todo-laura.hermanns - old translator used CastToNonLWCIfDisabled(), but LWC not supported in MIR yet
-			Emitter.Output(OutputIndex, Emitter.EmitExternalCode(HLSLCodeType, MoveTemp(HLSLCode)));
+			Emitter.Output(OutputIndex, Emitter.EmitInlineHLSL(HLSLCodeType, MoveTemp(HLSLCode), {}, {}));
 		}
 	}
 }
 
+void UMaterialExpressionPixelNormalWS::Build(MIR::FEmitter& Emitter)
+{
+	MIR::FValue* Output = Emitter.EmitInlineHLSL(MIR::FPrimitiveType::GetFloat3(), "Parameters.WorldNormal", {}, MIR::EGraphProperties::ReadsPixelNormal);
+	Emitter.Output(0, Output);
+}
+
 #endif // WITH_EDITOR
diff --git a/Engine/Source/Runtime/Engine/Private/Materials/MaterialIR.cpp b/Engine/Source/Runtime/Engine/Private/Materials/MaterialIR.cpp
--- a/Engine/Source/Runtime/Engine/Private/Materials/MaterialIR.cpp
+++ b/Engine/Source/Runtime/Engine/Private/Materials/MaterialIR.cpp
@@ -20,7 +20,7 @@ const TCHAR* ValueKindToString(EValueKind Kind)
 	{
 		case VK_Constant: return TEXT("Constant");
 		case VK_ExternalInput: return TEXT("ExternalInput");
-		case VK_ExternalCode: return TEXT("ExternalCode");
+		case VK_InlineHLSL: return TEXT("ExternalCode");
 		case VK_TextureObject: return TEXT("TextureObject");
 		case VK_UniformParameter: return TEXT("UniformParameter");
 		case VK_Dimensional: return TEXT("Dimensional");
@@ -36,7 +36,32 @@ const TCHAR* ValueKindToString(EValueKind Kind)
 
 bool FValue::IsAnalyzed(EStage Stage) const
 {
-	return (AnalyzedInStage & (1 << (uint32)Stage)) != 0;
+	return (Flags & EValueFlags(1 << Stage)) != EValueFlags::None;
+}
+
+bool FValue::HasFlags(EValueFlags InFlags) const
+{
+	return (Flags & InFlags) == InFlags;
+}
+
+void FValue::SetFlags(EValueFlags InFlags)
+{
+	Flags |= InFlags;
+}
+
+void FValue::ClearFlags(EValueFlags InFlags)
+{
+	Flags &= ~InFlags;
+}
+
+bool FValue::HasSubgraphProperties(EGraphProperties Properties) const
+{
+	return (GraphProperties & Properties) == Properties;
+}
+
+void FValue::UseSubgraphProperties(EGraphProperties Properties)
+{
+	GraphProperties |= Properties;
 }
 
 uint32 FValue::GetSizeInBytes() const
@@ -45,7 +70,7 @@ uint32 FValue::GetSizeInBytes() const
 	{
 		case VK_Constant: return sizeof(FConstant);
 		case VK_ExternalInput: return sizeof(FExternalInput);
-		case VK_ExternalCode: return sizeof(FExternalCode);
+		case VK_InlineHLSL: return sizeof(FInlineHLSL);
 		case VK_TextureObject: return sizeof(FTextureObject);
 		case VK_UniformParameter: return sizeof(FUniformParameter);
 		case VK_SetMaterialOutput: return sizeof(FSetMaterialOutput);
@@ -139,7 +164,7 @@ TArrayView<FValue*> FValue::GetUses()
 			return { &This->TextureObject, 4 };
 		}
 
-		case VK_ExternalCode:
+		case VK_InlineHLSL:
 		{
 			return {}; // @todo-laura.hermanns - Add input arguments as soon as external code expressions support parameters
 		}
@@ -318,16 +343,9 @@ const TCHAR* ExternalInputToString(EExternalInput Input)
 	}
 }
 
-MIR::EExternalInput TexCoordIndexToExternalInput(int TexCoordIndex)
-{
-	check(TexCoordIndex < TexCoordMaxNum);
-	return (MIR::EExternalInput)((int)MIR::EExternalInput::TexCoord0 + TexCoordIndex);
-}
-
 FTypePtr GetExternalInputType(EExternalInput Id)
 {
-	int IdIndex = (int)Id;
-	if (IdIndex >= (int)EExternalInput::TexCoord0 && IdIndex <= (int)EExternalInput::TexCoord7_Ddy)
+	if (Id >= EExternalInput::TexCoord0 && Id <= EExternalInput::TexCoord7_Ddy)
 	{
 		return FPrimitiveType::GetFloat2();
 	}
@@ -335,6 +353,12 @@ FTypePtr GetExternalInputType(EExternalInput Id)
 	UE_MIR_UNREACHABLE();
 }
 
+MIR::EExternalInput TexCoordIndexToExternalInput(int TexCoordIndex)
+{
+	check(TexCoordIndex < TexCoordMaxNum);
+	return (MIR::EExternalInput)((int)MIR::EExternalInput::TexCoord0 + TexCoordIndex);
+}
+
 bool IsExternalInputTexCoord(EExternalInput Id)
 {
 	return Id >= EExternalInput::TexCoord0 && Id <= EExternalInput::TexCoord7;
diff --git a/Engine/Source/Runtime/Engine/Private/Materials/MaterialIRDebug.cpp b/Engine/Source/Runtime/Engine/Private/Materials/MaterialIRDebug.cpp
--- a/Engine/Source/Runtime/Engine/Private/Materials/MaterialIRDebug.cpp
+++ b/Engine/Source/Runtime/Engine/Private/Materials/MaterialIRDebug.cpp
@@ -308,9 +308,9 @@ struct FDebugDumpIRState
 			Out.Appendf(TEXT(" MipValueMode=\"%s\""), *StaticEnum<ETextureMipValueMode>()->GetDisplayNameTextByValue(TextureSample->MipValueMode).ToString());
 			Out.Appendf(TEXT(" SamplerType=\"%s\""), *StaticEnum<EMaterialSamplerType>()->GetDisplayNameTextByValue(TextureSample->SamplerType).ToString());
 		}
-		else if (const FExternalCode* ExternalCode = Instr->As<FExternalCode>())
+		else if (const FInlineHLSL* ExternalCode = Instr->As<FInlineHLSL>())
 		{
-			Out.Appendf(TEXT(" \"%s\""), ExternalCode->HLSLCode);
+			Out.Appendf(TEXT(" \"%s\""), ExternalCode->Code);
 		}
 	}
 
diff --git a/Engine/Source/Runtime/Engine/Private/Materials/MaterialIREmitter.cpp b/Engine/Source/Runtime/Engine/Private/Materials/MaterialIREmitter.cpp
--- a/Engine/Source/Runtime/Engine/Private/Materials/MaterialIREmitter.cpp
+++ b/Engine/Source/Runtime/Engine/Private/Materials/MaterialIREmitter.cpp
@@ -313,7 +313,7 @@ bool FEmitter::ToConstantBoolScalar(MIR::FValue* Value)
 {
 	CheckValueIsScalar(Value, ScalarKind_Bool);
 
-	if (IsInvalid())
+	if (CurrentExpressionHasErrors())
 	{
 		return false;
 	}
@@ -592,6 +592,12 @@ FValue* FEmitter::EmitSubscript(FValue* Value, int32 Index)
 		return Value;
 	}
 
+	if (Index >= PrimitiveType->GetNumComponents())
+	{
+		Errorf(TEXT("Value of type `%s` has fewer dimensions than subscript index `%d`."), Value->Type->GetSpelling().GetData(), Index);
+		return MIR::FPoison::Get();
+	}
+
 	// Getting first component and Value is already a scalar, just return itself.
 	if (FDimensional* DimensionalValue = As<FDimensional>(Value))
 	{
@@ -1835,19 +1841,20 @@ FValue* FEmitter::EmitTextureSample(FValue* Texture, FValue* TexCoord, ESamplerS
 	return EmitPrototype(this, Prototype);
 }
 
-FValue* FEmitter::EmitExternalCode(FTypePtr Type, FString InHLSLCode, bool bNoTokenSubstitution)
+FValue* FEmitter::EmitInlineHLSL(FTypePtr Type, FString Code, EValueFlags ValueFlags, EGraphProperties UsedGraphProperties)
 {
-	MIR::FExternalCode Prototype = MakePrototype<MIR::FExternalCode>(Type);
-	Prototype.HLSLCode = AllocateModuleString(MoveTemp(InHLSLCode));
-	Prototype.bNoTokenSubstitution = bNoTokenSubstitution;
+	MIR::FInlineHLSL Prototype = MakePrototype<MIR::FInlineHLSL>(Type);
+	Prototype.Flags = ValueFlags;
+	Prototype.GraphProperties = UsedGraphProperties;
+	Prototype.Code = Module->PushUserString(MoveTemp(Code));
 	return EmitPrototype(this, Prototype);
 }
 
-FValue* FEmitter::EmitExternalCode(const FMaterialExternalCodeDeclaration* InExternalCodeDeclaration)
+FValue* FEmitter::EmitInlineHLSL(const FMaterialExternalCodeDeclaration* InExternalCodeDeclaration)
 {
 	check(InExternalCodeDeclaration != nullptr);
 	FTypePtr ReturnType = FType::FromMaterialValueType(InExternalCodeDeclaration->GetReturnTypeValue());
-	return  EmitExternalCode(ReturnType, InExternalCodeDeclaration->Definition);
+	return EmitInlineHLSL(ReturnType, InExternalCodeDeclaration->Definition, {}, {});
 }
 
 FTypePtr FEmitter::TryGetCommonType(FTypePtr A, FTypePtr B)
@@ -1901,7 +1908,7 @@ void FEmitter::Error(FStringView Message)
 	Error.Message.Appendf(TEXT("(Node %s) %s"), *ErrorClassName + ChopCount, Message.GetData());
 	
 	Module->Errors.Push(Error);
-	bExpressionHasError = true;
+	bCurrentExpressionHasErrors = true;
 }
 
 void FEmitter::Initialize()
@@ -1916,12 +1923,6 @@ void FEmitter::Initialize()
 	ConstantFalse = EmitPrototype(this, Temp);
 }
 
-const TCHAR* FEmitter::AllocateModuleString(FString InString)
-{
-	Module->Strings.Add(MoveTemp(InString));
-	return GetData(Module->Strings.Last());
-}
-
 bool FEmitter::FValueKeyFuncs::Matches(KeyInitType A, KeyInitType B)
 {
 	return A->Equals(B);
diff --git a/Engine/Source/Runtime/Engine/Private/Materials/MaterialIRModule.cpp b/Engine/Source/Runtime/Engine/Private/Materials/MaterialIRModule.cpp
--- a/Engine/Source/Runtime/Engine/Private/Materials/MaterialIRModule.cpp
+++ b/Engine/Source/Runtime/Engine/Private/Materials/MaterialIRModule.cpp
@@ -51,4 +51,15 @@ void FMaterialIRModule::Empty()
 	Allocator.Flush();
 }
 
+const TCHAR* FMaterialIRModule::PushUserString(FString InString)
+{
+	UserStrings.Add(MoveTemp(InString));
+	return GetData(UserStrings.Last());
+}
+
+void FMaterialIRModule::AddError(UMaterialExpression* Expression, FString Message)
+{
+	Errors.Emplace(Expression, MoveTemp(Message));
+}
+
 #endif // #if WITH_EDITOR
diff --git a/Engine/Source/Runtime/Engine/Private/Materials/MaterialIRModuleBuilder.cpp b/Engine/Source/Runtime/Engine/Private/Materials/MaterialIRModuleBuilder.cpp
--- a/Engine/Source/Runtime/Engine/Private/Materials/MaterialIRModuleBuilder.cpp
+++ b/Engine/Source/Runtime/Engine/Private/Materials/MaterialIRModuleBuilder.cpp
@@ -233,7 +233,7 @@ struct FMaterialIRModuleBuilderImpl
 		}
 
 		// If some error occurred (e.g. some function input wasn't linked in) early out.
-		if (Emitter->IsInvalid())
+		if (Emitter->CurrentExpressionHasErrors())
 		{
 			return;
 		}
@@ -393,13 +393,13 @@ struct FMaterialIRModuleBuilderImpl
 
 				// If this is the first time this value is analyzed, let the analyzer process it.
 				// Note that individual value processing is independent from the stage it runs on so we can perform it only once.
-				if (Value->AnalyzedInStage == 0)
+				if ((Value->Flags & MIR::EValueFlags::AnalyzedInAnyStageMask) == MIR::EValueFlags::None)
 				{
 					ValueAnalyzer.Analyze(Value);
 				}
 
 				// Mark the used instruction as analyzed for this stage.
-				Value->AnalyzedInStage |= 1 << StageIndex;
+				Value->Flags |= MIR::EValueFlags(1 << StageIndex);
 			}
 		}
 	}
@@ -546,7 +546,7 @@ bool FMaterialIRModuleBuilder::Build(FMaterialIRModule* TargetModule)
 	Impl.Step_GenerateOutputInstructions();
 	Impl.Step_BuildMaterialExpressionsToIRGraph();
 
-	if (Impl.Emitter->IsInvalid())
+	if (!TargetModule->IsValid())
 	{
 		return false;
 	}
@@ -572,7 +572,7 @@ bool FMaterialIRModuleBuilder::Build(FMaterialIRModule* TargetModule)
 		}
 	}
 
-	return true;
+	return TargetModule->IsValid();
 }
 
 namespace MIR::Internal {
diff --git a/Engine/Source/Runtime/Engine/Private/Materials/MaterialIRToHLSLTranslator.cpp b/Engine/Source/Runtime/Engine/Private/Materials/MaterialIRToHLSLTranslator.cpp
--- a/Engine/Source/Runtime/Engine/Private/Materials/MaterialIRToHLSLTranslator.cpp
+++ b/Engine/Source/Runtime/Engine/Private/Materials/MaterialIRToHLSLTranslator.cpp
@@ -51,7 +51,7 @@ struct FHLSLPrinter
 		return *this;
 	}
 
-	FHLSLPrinter& operator<<(int Value)
+	FHLSLPrinter& operator<<(int32 Value)
 	{
 		Buffer.Appendf(TEXT("%d"), Value);
 		return *this;
@@ -71,7 +71,7 @@ struct FHLSLPrinter
 
 	FHLSLPrinter& operator<<(EIndentation)
 	{
-		for (int i = 0; i < Tabs; ++i)
+		for (int32 i = 0; i < Tabs; ++i)
 		{
 			Buffer.AppendChar('\t');
 		}
@@ -262,18 +262,18 @@ struct FTranslator : FMaterialIRToHLSLTranslation
 
 		WorldPositionOffsetHLSL = MoveTemp(Printer.Buffer);
 	}
-	
+
 	void GenerateFragmentStageHLSL()
-	{
+		{
 		BeginStage(MIR::Stage_Fragment);
 
 		LowerBlock(Module->GetRootBlock(CurrentStage));
 
-		Printer << TEXT("PixelMaterialInputs.FrontMaterial = GetInitialisedSubstrateData()") << EndOfStatement;
-		Printer << TEXT("PixelMaterialInputs.Subsurface = 0") << EndOfStatement;
+			Printer << TEXT("PixelMaterialInputs.FrontMaterial = GetInitialisedSubstrateData()") << EndOfStatement;
+			Printer << TEXT("PixelMaterialInputs.Subsurface = 0") << EndOfStatement;
 
 		EvaluateOtherMaterialAttributesHLSL = MoveTemp(Printer.Buffer);
-	}
+		}
 
 	void BeginStage(MIR::EStage Stage)
 	{
@@ -285,7 +285,7 @@ struct FTranslator : FMaterialIRToHLSLTranslation
 	
 	ENoOp LowerBlock(const MIR::FBlock& Block)
 	{
-		int OldNumLocals = NumLocals;
+		int32 OldNumLocals = NumLocals;
 		for (MIR::FInstruction* Instr = Block.Instructions; Instr; Instr = Instr->Next[CurrentStage])
 		{
             if (Instr->NumUsers[CurrentStage] == 1 && IsFoldable(Instr, CurrentStage))
@@ -370,7 +370,7 @@ struct FTranslator : FMaterialIRToHLSLTranslation
 			case MIR::VK_Branch: LowerBranch(static_cast<MIR::FBranch*>(Instr)); break;
 			case MIR::VK_Subscript: LowerSubscript(static_cast<MIR::FSubscript*>(Instr)); break;
 			case MIR::VK_TextureSample: LowerTextureSample(static_cast<MIR::FTextureSample*>(Instr)); break;
-			case MIR::VK_ExternalCode: LowerExternalCode(static_cast<const MIR::FExternalCode*>(Instr)); break;
+			case MIR::VK_InlineHLSL: LowerInlineHLSL(static_cast<const MIR::FInlineHLSL*>(Instr)); break;
 
 			default:
 				UE_MIR_UNREACHABLE();
@@ -394,21 +394,20 @@ struct FTranslator : FMaterialIRToHLSLTranslation
 
 	void LowerExternalInput(const MIR::FExternalInput* ExternalInput)
 	{
-		int ExternalInputIndex = (int)ExternalInput->Id;
-
+		int32 ExternalInputIndex = (int32)ExternalInput->Id;
 		if (MIR::IsExternalInputTexCoord(ExternalInput->Id))
 		{
-			int Index = ExternalInputIndex - (int)MIR::EExternalInput::TexCoord0;
+			int32 Index = ExternalInputIndex - (int32)MIR::EExternalInput::TexCoord0;
 			Printer.Buffer.Appendf(TEXT("Parameters.TexCoords[%d]"), Index);
 		}
 		else if (MIR::IsExternalInputTexCoordDdx(ExternalInput->Id))
 		{
-			int Index = ExternalInputIndex - (int)MIR::EExternalInput::TexCoord0_Ddx;
+			int32 Index = ExternalInputIndex - (int32)MIR::EExternalInput::TexCoord0_Ddx;
 			Printer.Buffer.Appendf(TEXT("Parameters.TexCoords_DDX[%d]"), Index);
 		}
 		else if (MIR::IsExternalInputTexCoordDdy(ExternalInput->Id))
 		{
-			int Index = ExternalInputIndex - (int)MIR::EExternalInput::TexCoord0_Ddy;
+			int32 Index = ExternalInputIndex - (int32)MIR::EExternalInput::TexCoord0_Ddy;
 			Printer.Buffer.Appendf(TEXT("Parameters.TexCoords_DDY[%d]"), Index);
 		}
 		else
@@ -417,24 +416,24 @@ struct FTranslator : FMaterialIRToHLSLTranslation
 		}
 	}
 
-	void LowerExternalCode(const MIR::FExternalCode* ExternalCode)
+	void LowerInlineHLSL(const MIR::FInlineHLSL* ExternalCode)
 	{
-		check(ExternalCode->HLSLCode != nullptr);
+		check(ExternalCode->Code != nullptr);
 
 		// Substitute placeholder tokens now unless disabled for custom nodes
-		if (ExternalCode->bNoTokenSubstitution)
+		if (ExternalCode->HasFlags(MIR::EValueFlags::SubstituteTagsInInlineHLSL))
 		{
-			Printer.Buffer.Append(ExternalCode->HLSLCode);
-		}
-		else
-		{
-			FString HLSLCode = ExternalCode->HLSLCode;
+			FString Code = ExternalCode->Code;
 
 			// @todo-laura.hermanns - Must be replaced as soon as this state is available in MIR::FExternalCode
 			constexpr bool bCompilingPreviousFrame = false;
-			HLSLCode.ReplaceInline(TEXT("<PREV>"), bCompilingPreviousFrame ? TEXT("Prev") : TEXT(""));
+			Code.ReplaceInline(TEXT("<PREV>"), bCompilingPreviousFrame ? TEXT("Prev") : TEXT(""));
 
-			Printer.Buffer.Append(MoveTemp(HLSLCode));
+			Printer.Buffer.Append(MoveTemp(Code));
+		}
+		else
+		{
+			Printer.Buffer.Append(ExternalCode->Code);
 		}
 	}
 
@@ -767,7 +766,7 @@ struct FTranslator : FMaterialIRToHLSLTranslation
 		return NoOp;
 	}
 
-	ENoOp LowerTextureReference(EMaterialValueType TextureType, int TextureParameterIndex)
+	ENoOp LowerTextureReference(EMaterialValueType TextureType, int32 TextureParameterIndex)
 	{
 		Printer << TEXT("Material.");
 
@@ -788,7 +787,7 @@ struct FTranslator : FMaterialIRToHLSLTranslation
 	{
 		const FMaterialIRModule::FStatistics ModuleStatistics = Module->GetStatistics();
 
-		auto SetParamInt = [&] (const TCHAR* InParamName, int InValue)
+		auto SetParamInt = [&] (const TCHAR* InParamName, int32 InValue)
 		{
 			Params.Add(InParamName, FString::Printf(TEXT("%d"), InValue));
 		};
@@ -829,7 +828,7 @@ struct FTranslator : FMaterialIRToHLSLTranslation
 		SetParamInt(TEXT("num_tex_coord_interpolators"), ModuleStatistics.NumPixelTexCoords);
 
 		FString GetMaterialCustomizedUVS;
-		for (int CustomUVIndex = 0; CustomUVIndex < ModuleStatistics.NumPixelTexCoords; CustomUVIndex++)
+		for (int32 CustomUVIndex = 0; CustomUVIndex < ModuleStatistics.NumPixelTexCoords; CustomUVIndex++)
 		{
 			const FString AttributeName = FMaterialAttributeDefinitionMap::GetAttributeName((EMaterialProperty)(MP_CustomizedUVs0 + CustomUVIndex));
 			GetMaterialCustomizedUVS.Appendf(TEXT(TAB "OutTexCoords[%u] = Parameters.MaterialAttributes.%s;\n"), CustomUVIndex, *AttributeName);
@@ -858,7 +857,7 @@ struct FTranslator : FMaterialIRToHLSLTranslation
 
 		FString EvaluateMaterialDeclaration;
 		EvaluateMaterialDeclaration.Append(TEXT("void EvaluateVertexMaterialAttributes(in out FMaterialVertexParameters Parameters)\n{\n"));
-		for (int CustomUVIndex = 0; CustomUVIndex < ModuleStatistics.NumPixelTexCoords; CustomUVIndex++)
+		for (int32 CustomUVIndex = 0; CustomUVIndex < ModuleStatistics.NumPixelTexCoords; CustomUVIndex++)
 		{
 			EvaluateMaterialDeclaration.Appendf(TEXT(TAB "Parameters.MaterialAttributes.CustomizedUV%d = Parameters.TexCoords[%d].xy;\n"), CustomUVIndex, CustomUVIndex);
 		}
@@ -873,7 +872,7 @@ struct FTranslator : FMaterialIRToHLSLTranslation
 			switch (ArithmeticType->ScalarKind)
 			{
 				case MIR::ScalarKind_Bool:	Printer << TEXT("bool"); break;
-				case MIR::ScalarKind_Int: 	Printer << TEXT("int"); break;
+				case MIR::ScalarKind_Int: 	Printer << TEXT("int32"); break;
 				case MIR::ScalarKind_Float:	Printer << TEXT("MaterialFloat"); break;
 			}
 
@@ -935,7 +934,7 @@ struct FTranslator : FMaterialIRToHLSLTranslation
 			// This is to have platforms use the simple single layer water shading similar to mobile: no dynamic lights, only sun and sky, no distortion, no colored transmittance on background, no custom depth read.
 			const bool bSingleLayerWaterUsesSimpleShading = FDataDrivenShaderPlatformInfo::GetWaterUsesSimpleForwardShading(ShaderPlatform) && IsForwardShadingEnabled(ShaderPlatform);
 
-			for (int i = 0; i < MSM_NUM; ++i)
+			for (int32 i = 0; i < MSM_NUM; ++i)
 			{
 				EMaterialShadingModel Model = (EMaterialShadingModel)i;
 				if (Model == MSM_Strata || !ShadingModels.HasShadingModel(Model))
diff --git a/Engine/Source/Runtime/Engine/Private/Materials/MaterialIRValueAnalyzer.cpp b/Engine/Source/Runtime/Engine/Private/Materials/MaterialIRValueAnalyzer.cpp
--- a/Engine/Source/Runtime/Engine/Private/Materials/MaterialIRValueAnalyzer.cpp
+++ b/Engine/Source/Runtime/Engine/Private/Materials/MaterialIRValueAnalyzer.cpp
@@ -10,10 +10,9 @@
 #include "Materials/MaterialIRTypes.h"
 #include "Materials/MaterialIRInternal.h"
 #include "Materials/MaterialInsights.h"
+#include "Materials/Material.h"
 #include "MaterialShared.h"
-#include "Materials/MaterialInsights.h"
 #include "Engine/Texture.h"
-#include "Materials/Material.h"
 
 static void AnalyzeExternalInput(FMaterialIRValueAnalyzer& Analyzer, MIR::FExternalInput* ExternalInput)
 {
@@ -168,6 +167,18 @@ static void AnalyzeUniformParameter(FMaterialIRValueAnalyzer& Analyzer, MIR::FUn
 	}
 }
 
+static void AnalyzeSetMaterialOutput(FMaterialIRValueAnalyzer& Analyzer, MIR::FSetMaterialOutput* SetMaterialOutput)
+{
+	if (SetMaterialOutput->Property == EMaterialProperty::MP_Normal)
+	{
+		if (SetMaterialOutput->HasSubgraphProperties(MIR::EGraphProperties::ReadsPixelNormal))
+		{
+			Analyzer.Module->AddError(nullptr, TEXT("Cannot set material attribute Normal to a value that depends on reading the pixel normal, as that would create a circular dependency."));
+		}
+	}
+}
+
+
 void FMaterialIRValueAnalyzer::Setup(UMaterial* InMaterial, FMaterialIRModule* InModule, FMaterialCompilationOutput* InCompilationOutput, FMaterialInsights* InInsights)
 {
 	Material = InMaterial;
@@ -178,6 +189,11 @@ void FMaterialIRValueAnalyzer::Setup(UMaterial* InMaterial, FMaterialIRModule* I
 
 void FMaterialIRValueAnalyzer::Analyze(MIR::FValue* Value)
 {
+	for (MIR::FValue* Use : Value->GetUses())
+	{
+		Value->GraphProperties |= Use->GraphProperties;
+	}
+
 	#define EXPAND_CASE(ValueType) \
 	case MIR::VK_##ValueType: \
 	{ \
@@ -190,6 +206,7 @@ void FMaterialIRValueAnalyzer::Analyze(MIR::FValue* Value)
 		EXPAND_CASE(ExternalInput);
 		EXPAND_CASE(TextureObject);
 		EXPAND_CASE(UniformParameter);
+		EXPAND_CASE(SetMaterialOutput);
 		default: break;
 	}
 
diff --git a/Engine/Source/Runtime/Engine/Private/Materials/MaterialIRValueAnalyzer.h b/Engine/Source/Runtime/Engine/Private/Materials/MaterialIRValueAnalyzer.h
--- a/Engine/Source/Runtime/Engine/Private/Materials/MaterialIRValueAnalyzer.h
+++ b/Engine/Source/Runtime/Engine/Private/Materials/MaterialIRValueAnalyzer.h
@@ -2,8 +2,6 @@
 
 #pragma once
 
-#include "Materials/MaterialIRModuleBuilder.h"
-
 #if WITH_EDITOR
 
 #include "Materials/MaterialIRModule.h"
@@ -11,8 +9,8 @@
 #include "Materials/MaterialIR.h"
 #include "Materials/MaterialIRTypes.h"
 #include "Materials/MaterialIRInternal.h"
-#include "MaterialShared.h"
 #include "Materials/MaterialInsights.h"
+#include "MaterialShared.h"
 #include "Engine/Texture.h"
 
 struct FMaterialIRValueAnalyzer
@@ -40,6 +38,9 @@ struct FMaterialIRValueAnalyzer
 		
 	//
 	FMaterialInsights* Insights{}; // Optional
+
+	//
+	FMaterialIRStateTracker* StateTracker{};
 };
 
 #endif // #if WITH_EDITOR
diff --git a/Engine/Source/Runtime/Engine/Public/Materials/MaterialExpressionPixelNormalWS.h b/Engine/Source/Runtime/Engine/Public/Materials/MaterialExpressionPixelNormalWS.h
--- a/Engine/Source/Runtime/Engine/Public/Materials/MaterialExpressionPixelNormalWS.h
+++ b/Engine/Source/Runtime/Engine/Public/Materials/MaterialExpressionPixelNormalWS.h
@@ -17,6 +17,7 @@ public:
 
 	//~ Begin UMaterialExpression Interface
 #if WITH_EDITOR
+	virtual void Build(MIR::FEmitter& Emitter) override;
 	virtual int32 Compile(class FMaterialCompiler* Compiler, int32 OutputIndex) override;
 	virtual void GetCaption(TArray<FString>& OutCaptions) const override;
 #endif
diff --git a/Engine/Source/Runtime/Engine/Public/Materials/MaterialIR.h b/Engine/Source/Runtime/Engine/Public/Materials/MaterialIR.h
--- a/Engine/Source/Runtime/Engine/Public/Materials/MaterialIR.h
+++ b/Engine/Source/Runtime/Engine/Public/Materials/MaterialIR.h
@@ -11,6 +11,7 @@ struct FMaterialExternalCodeDeclaration;
 namespace MIR {
 
 // Identifies the block of execution in which an instruction runs.
+// Note: if you introduce a new stage, make sure you update EValueFlags accordingly.
 enum EStage
 {
 	Stage_Vertex,
@@ -21,6 +22,46 @@ enum EStage
 // Returns the stage in which specified material property shall be evaluated.
 EStage GetMaterialOutputStage(EMaterialProperty Property);
 
+// A collection of bit flags for a specific Value instance.
+enum class EValueFlags : uint8
+{
+	None = 0,
+
+	// Value has been analyzed for the vertex stage.
+	AnalyzedForStageVertex = 1 << 0,
+
+	// Value has been analyzed for the fragment stage.
+	AnalyzedForStageFragment = 1 << 1,
+
+	// Mask that includes all "analyzed in stage" flags.
+	AnalyzedInAnyStageMask = AnalyzedForStageVertex | AnalyzedForStageFragment,
+
+	// Whether to skip substution of abstract tokens like "<PREV>". Only true for custom nodes.
+	SubstituteTagsInInlineHLSL = 1 << 2,
+};
+
+ENUM_CLASS_FLAGS(EValueFlags);
+
+static_assert(uint32(1 << NumStages) - 1 == uint32(EValueFlags::AnalyzedInAnyStageMask), "Discrepancy between the number of stages in NumStages and the flags in EValueFlags.");
+
+// A collection of graph properties used by a value. If a graph property is true, it means
+// that either the value itself makes has of that property, or one of its dependencies
+// (direct or indirect) has it. This entails these flags are automatically propagated to
+// all dependant values (values that depend on a given one).
+// As an example, if "ReadsPixelNormal" is true on a specific value it means that either
+// that value itself or some other upstream value that value is dependent on reads
+// the pixel normal.
+enum class EGraphProperties : uint8
+{
+	None = 0,
+
+	// Some value reads the pixel normal.
+	ReadsPixelNormal = 1 << 0,
+};
+
+ENUM_CLASS_FLAGS(EGraphProperties);
+
+// Enumeration of all different structs deriving from FValue. Used for dynamic casting.
 enum EValueKind
 {
 	/* Values */
@@ -42,7 +83,7 @@ enum EValueKind
 	VK_Subscript,
 	VK_Cast,
 	VK_TextureSample,
-	VK_ExternalCode,
+	VK_InlineHLSL,
 
 	VK_InstructionEnd,
 };
@@ -76,15 +117,35 @@ struct FValue
 	// Used to discern the concrete C++ type of this value (e.g. Subscript)
 	EValueKind Kind : 8;
 
-	// Whether this value was analyzed, which occurs after its dependencies have been analyzed.
-	uint32 AnalyzedInStage : MIR::NumStages;
+	// Set of fundamental flags true for this value.
+	EValueFlags Flags;
+
+	// The set of properties that are true for this value. Some flags might have been
+	// set but some upstream dependency leading to this value and not this value directly.
+	// See `EGraphProperties` for more information.
+	EGraphProperties GraphProperties;
 
 	// The runtime type this value has.
-	FTypePtr  Type;
+	FTypePtr Type;
 
-	// Utility function to tell whether this value has been analyzed for specified stage.
+	// Returns whether this value has been analyzed for specified stage.
 	bool IsAnalyzed(EStage State) const;
 
+	// Returns whether specified flags are true for this value.
+	bool HasFlags(EValueFlags InFlags) const;
+
+	// Enables the specified flags without affecting others.
+	void SetFlags(EValueFlags InFlags);
+
+	// Disables the specified flags without affecting others.
+	void ClearFlags(EValueFlags InFlags);
+
+	// Returns whether specified properties are true for this value.
+	bool HasSubgraphProperties(EGraphProperties Properties) const;
+
+	// Enables specified properties for this value.
+	void UseSubgraphProperties(EGraphProperties Properties);
+
 	// Returns the size in bytes of this value instance.
 	uint32 GetSizeInBytes() const;
 	
@@ -255,9 +316,11 @@ static constexpr int TexCoordMaxNum = 8;
 const TCHAR* ExternalInputToString(EExternalInput Input);
 
 //
-MIR::EExternalInput TexCoordIndexToExternalInput(int TexCoordIndex);
 FTypePtr GetExternalInputType(EExternalInput Id);
 
+//
+MIR::EExternalInput TexCoordIndexToExternalInput(int TexCoordIndex);
+
 //
 bool IsExternalInputTexCoord(EExternalInput Id);
 
@@ -357,10 +420,9 @@ struct FSetMaterialOutput : TInstruction<VK_SetMaterialOutput>
 };
 
 // MIR instruction for external code declarations; See BaseMaterialExpressions.ini file.
-struct FExternalCode : TInstruction<VK_ExternalCode>
+struct FInlineHLSL : TInstruction<VK_InlineHLSL>
 {
-	const TCHAR* HLSLCode;
-	bool bNoTokenSubstitution; // Whether to skip substution of abstract tokens like "<PREV>". Only true for custom nodes.
+	const TCHAR* Code;
 };
 
 // Note: if you modify this enum, update the implementations of the
diff --git a/Engine/Source/Runtime/Engine/Public/Materials/MaterialIRCommon.h b/Engine/Source/Runtime/Engine/Public/Materials/MaterialIRCommon.h
--- a/Engine/Source/Runtime/Engine/Public/Materials/MaterialIRCommon.h
+++ b/Engine/Source/Runtime/Engine/Public/Materials/MaterialIRCommon.h
@@ -10,23 +10,23 @@
 
 /* Forward declarations */
 
-enum EMaterialProperty : int;
 class FMaterial;
 class FMaterialIRModule;
-class FMaterialIRModule;
-struct FMaterialIRModuleBuilder;
 class ITargetPlatform;
 class UMaterial;
 class UMaterialExpression;
+class UMaterialExpressionMaterialFunctionCall;
+class UTexture;
+enum EMaterialProperty : int;
 struct FExpressionInput;
 struct FExpressionOutput;
-struct FShaderCompilerEnvironment;
-struct FStaticParameterSet;
 struct FMaterialInputDescription;
 struct FMaterialInsights;
-class UTexture;
-class UMaterialExpressionMaterialFunctionCall;
+struct FMaterialIRModuleBuilder;
 struct FMaterialIRModuleBuilderImpl;
+struct FMaterialIRStateTracker;
+struct FShaderCompilerEnvironment;
+struct FStaticParameterSet;
 
 namespace UE::Shader
 {
diff --git a/Engine/Source/Runtime/Engine/Public/Materials/MaterialIREmitter.h b/Engine/Source/Runtime/Engine/Public/Materials/MaterialIREmitter.h
--- a/Engine/Source/Runtime/Engine/Public/Materials/MaterialIREmitter.h
+++ b/Engine/Source/Runtime/Engine/Public/Materials/MaterialIREmitter.h
@@ -7,7 +7,7 @@
 
 #if WITH_EDITOR
 
-#define UE_MIR_CHECKPOINT() if (Emitter.IsInvalid()) return
+#define UE_MIR_CHECKPOINT() if (Emitter.CurrentExpressionHasErrors()) return
 
 namespace MIR {
 
@@ -161,8 +161,8 @@ public:
 	FValue* EmitBranch(FValue* Condition, FValue* True, FValue* False);
 	FValue* EmitConstruct(FTypePtr Type, FValue* Initializer);
 	FValue* EmitTextureSample(FValue* Texture, FValue* TexCoord, ESamplerSourceMode SamplerSourceMode, ETextureMipValueMode MipValueMode, EMaterialSamplerType SamplerType);
-	FValue* EmitExternalCode(FTypePtr Type, FString InHLSLCode, bool bNoTokenSubstitution = false);
-	FValue* EmitExternalCode(const FMaterialExternalCodeDeclaration* InExternalCodeDeclaration);
+	FValue* EmitInlineHLSL(FTypePtr Type, FString Code, EValueFlags ValueFlags, EGraphProperties UsedGraphProperties);
+	FValue* EmitInlineHLSL(const FMaterialExternalCodeDeclaration* InExternalCodeDeclaration);
 
 	/* Types */
 
@@ -171,7 +171,7 @@ public:
 
 	/* Error reporting */
 
-	bool IsInvalid() const { return bExpressionHasError; }
+	bool CurrentExpressionHasErrors() const { return bCurrentExpressionHasErrors; }
 
 	template <typename... TArgs>
 	void Errorf(UE::Core::TCheckedFormatString<FString::FmtCharType, TArgs...> Format, TArgs... Args)
@@ -188,8 +188,6 @@ private:
 
 	void Initialize();
 
-	const TCHAR* AllocateModuleString(FString InString);
-
 	struct FValueKeyFuncs : DefaultKeyFuncs<FValue*>
 	{
 		static bool Matches(KeyInitType A, KeyInitType B);
@@ -201,7 +199,7 @@ private:
 	FMaterialIRModule* Module{};
 	const FStaticParameterSet* StaticParameterSet{};
 	UMaterialExpression* Expression{};
-	bool bExpressionHasError = false;
+	bool bCurrentExpressionHasErrors = false;
 	FValue* ConstantTrue;
 	FValue* ConstantFalse;
 	TSet<FValue*, FValueKeyFuncs> ValueSet;
diff --git a/Engine/Source/Runtime/Engine/Public/Materials/MaterialIRModule.h b/Engine/Source/Runtime/Engine/Public/Materials/MaterialIRModule.h
--- a/Engine/Source/Runtime/Engine/Public/Materials/MaterialIRModule.h
+++ b/Engine/Source/Runtime/Engine/Public/Materials/MaterialIRModule.h
@@ -7,61 +7,121 @@
 #if WITH_EDITOR
 
 // This class represents the intermediate representation (IR) of a material build.
-// The `IRModule` includes an IR value graph, produced through expression analysis,
+// The IRModule includes an IR value graph, produced through expression analysis,
 // as well as metadata on resource usage and reflection. The IR graph serves as an
 // abstract representation of the material and must be translated into a target backend
 // such as HLSL or specific Preshader opcodes for execution.
 //
 // This class is designed to be backend-agnostic, meaning it does not contain any
-// HLSL code nor does it configure a `MaterialCompilationOutput` instance. The data
+// HLSL code nor does it configure a MaterialCompilationOutput instance. The data
 // stored within this class should be sufficient to enable translation to any supported
 // backend without requiring additional processing or validation.
 class FMaterialIRModule
 {
 public:
-	struct FError
-	{
-		UMaterialExpression* Expression;
-		FString Message;
-	};
-	
-	struct FStatistics
-	{
-		TBitArray<> ExternalInputUsedMask[SF_NumFrequencies];
-		int NumVertexTexCoords;
-		int NumPixelTexCoords;
-	};
+    // Represents an error encountered during material processing.
+    struct FError
+    {
+		// The expression that caused the error.
+        UMaterialExpression* Expression;
+
+		// Description of the error.
+        FString Message;
+    };
+    
+    // Stores information about the resources used by the translated material.
+    struct FStatistics
+    {
+		// Tracks external inputs used per frequency.
+        TBitArray<> ExternalInputUsedMask[SF_NumFrequencies]; 
+
+		// Number of vertex texture coordinates used.
+        int NumVertexTexCoords;
+		
+		// Number of pixel texture coordinates used.
+        int NumPixelTexCoords;
+    };
 
 public:
-	FMaterialIRModule();
-	~FMaterialIRModule();
-	void Empty();
-	EShaderPlatform GetShaderPlatform() const { return ShaderPlatform; }
-	const FMaterialCompilationOutput& GetCompilationOutput() const { return CompilationOutput; }
-	TArrayView<const MIR::FSetMaterialOutput* const> GetOutputs(MIR::EStage Stage) const { return Outputs[Stage]; }
-	const MIR::FBlock& GetRootBlock(MIR::EStage Stage) const { return *RootBlock[Stage]; }
-	TArrayView<const FError> GetErrors() const { return Errors; }
-	const FStatistics& GetStatistics() const { return Statistics; }
-	FStatistics& GetStatistics() { return Statistics; }
-	const FMaterialParameterInfo& GetParameterInfo(uint32 ParameterId) const { return ParameterIdToData[ParameterId].Key; }
-	const FMaterialParameterMetadata& GetParameterMetadata(uint32 ParameterId) const { return ParameterIdToData[ParameterId].Value; }
+    FMaterialIRModule();
+    ~FMaterialIRModule();
+
+    // Clears the module, releasing all stored data.
+    void Empty();
+    
+    // Returns the shader platform associated with this module.
+    EShaderPlatform GetShaderPlatform() const { return ShaderPlatform; }
+    
+    // Retrieves the material compilation output.
+    const FMaterialCompilationOutput& GetCompilationOutput() const { return CompilationOutput; }
+    
+    // Returns the material outputs for a given stage.
+    TArrayView<const MIR::FSetMaterialOutput* const> GetOutputs(MIR::EStage Stage) const { return Outputs[Stage]; }
+    
+    // Retrieves the root block (i.e. the "main" scope) for a specific shader stage.
+    const MIR::FBlock& GetRootBlock(MIR::EStage Stage) const { return *RootBlock[Stage]; }
+
+    // Provides access to the translated material statistics.
+    const FStatistics& GetStatistics() const { return Statistics; }
+    
+    // Provides mutable access to the translated material statistics.
+    FStatistics& GetStatistics() { return Statistics; }
+    
+    // Retrieves parameter info for a given parameter ID.
+    const FMaterialParameterInfo& GetParameterInfo(uint32 ParameterId) const { return ParameterIdToData[ParameterId].Key; }
+    
+    // Retrieves parameter metadata for a given parameter ID.
+    const FMaterialParameterMetadata& GetParameterMetadata(uint32 ParameterId) const { return ParameterIdToData[ParameterId].Value; }
+    
+    // Stores a user-defined string and returns a pointer to it.
+    const TCHAR* PushUserString(FString InString);
 
+    // Checks if the module is valid (i.e., contains no errors).
+    bool IsValid() const { return Errors.IsEmpty(); }
+
+    // Returns a list of errors encountered during processing.
+    TArrayView<const FError> GetErrors() const { return Errors; }
+
+    // Reports a translation error.
+    void AddError(UMaterialExpression* Expression, FString Message);
+    
 private:
-	EShaderPlatform ShaderPlatform;
-	FMaterialCompilationOutput CompilationOutput;
-	FMemStackBase Allocator{};
-	TArray<MIR::FValue*> Values;
-	TArray<MIR::FSetMaterialOutput*> Outputs[MIR::NumStages];
-	TArray<FError> Errors;
-	TArray<FString> Strings;
-	MIR::FBlock* RootBlock[MIR::NumStages];
-	FStatistics Statistics;
-	TMap<FMaterialParameterInfo, uint32> ParameterInfoToId;
-	TArray<TPair<FMaterialParameterInfo, FMaterialParameterMetadata>> ParameterIdToData;
+	// Target shader platform.
+    EShaderPlatform ShaderPlatform;
+    
+	// Compilation output data.
+    FMaterialCompilationOutput CompilationOutput;
+    
+	// Memory allocator used to allocate IR data (values, payloads, etc).
+    FMemStackBase Allocator{};
+    
+	// List of all the IR values contained in this module.
+    TArray<MIR::FValue*> Values;
+    
+	// Output nodes per stage.
+    TArray<MIR::FSetMaterialOutput*> Outputs[MIR::NumStages];
+    
+	// Root blocks per stage.
+    MIR::FBlock* RootBlock[MIR::NumStages];
+    
+	// Compilation statistics.
+    FStatistics Statistics;
+    
+	// Maps parameter info to IDs.
+    TMap<FMaterialParameterInfo, uint32> ParameterInfoToId;
+    
+	// Parameter metadata.
+    TArray<TPair<FMaterialParameterInfo, FMaterialParameterMetadata>> ParameterIdToData;
+
+	// Stores user-defined strings.    
+    TArray<FString> UserStrings;
 
-	friend MIR::FEmitter;
-	friend FMaterialIRModuleBuilder;
-	friend FMaterialIRModuleBuilderImpl;
+	// List of compilation errors.
+    TArray<FError> Errors;
+    
+    friend MIR::FEmitter;
+    friend FMaterialIRModuleBuilder;
+    friend FMaterialIRModuleBuilderImpl;
 };
 
 #endif // #if WITH_EDITOR
