[shaders] guard access to inflightjobs map to only populate/use in cook commandlet to prevent threading issues with ODSC. this is currently only used to track cook hangs so there's no need to touch it in other contexts.

diff --git a/Engine/Source/Runtime/Engine/Private/Materials/MaterialShader.cpp b/Engine/Source/Runtime/Engine/Private/Materials/MaterialShader.cpp
--- a/Engine/Source/Runtime/Engine/Private/Materials/MaterialShader.cpp
+++ b/Engine/Source/Runtime/Engine/Private/Materials/MaterialShader.cpp
@@ -2230,7 +2230,10 @@ void FMaterialShaderMap::ReleaseCompilingId()
 		check(CompilingMaterialDependencies.Num() == 0);
 		check(CompilingMaterialNumExternalDependencies == 0);
 #if ALLOW_SHADERMAP_DEBUG_DATA
-		InFlightJobs.Empty();
+		if (IsRunningCookCommandlet())
+		{
+			InFlightJobs.Empty();
+		}
 #endif
 		verify(GetCompilingShaderMapLookup().Remove(CompilingId) == 1);
 		CompilingId = 0;
@@ -2503,11 +2506,14 @@ int32 FMaterialShaderMap::SubmitCompileJobs(uint32 CompilingShaderMapId,
 	UE_LOG(LogShaders, Verbose, TEXT("		%u Shaders among %u VertexFactories"), NumShaders, NumVertexFactories);
 
 #if ALLOW_SHADERMAP_DEBUG_DATA
-	for (const FShaderCommonCompileJobPtr& Job : CompileJobs)
+	if (IsRunningCookCommandlet())
 	{
-		TStringBuilder<2048> JobName;
-		Job->AppendDebugName(JobName);
-		InFlightJobs.Add(Job->Hash, { JobName.ToString(), Job->JobStatusPtr });
+		for (const FShaderCommonCompileJobPtr& Job : CompileJobs)
+		{
+			TStringBuilder<2048> JobName;
+			Job->AppendDebugName(JobName);
+			InFlightJobs.Add(Job->Hash, { JobName.ToString(), Job->JobStatusPtr });
+		}
 	}
 #endif
 
@@ -2784,6 +2790,7 @@ FShader* FMaterialShaderMap::ProcessCompilationResultsForSingleJob(FShaderCompil
 void FMaterialShaderMap::AppendCompileStateDebugInfo(FStringBuilderBase& OutDebugInfo) const
 {
 #if ALLOW_SHADERMAP_DEBUG_DATA
+	check(IsRunningCookCommandlet()); 
 	for (const TPair<uint32, TPair<FString, TSharedPtr<FShaderCompileJobStatus>>>& StatusPair : InFlightJobs)
 	{
 		TSharedPtr<FShaderCompileJobStatus> Status = StatusPair.Value.Value;
@@ -2810,7 +2817,10 @@ void FMaterialShaderMap::ProcessCompilationResults(const TArray<FShaderCommonCom
 	{
 		auto& BaseJob = InCompilationResults[InOutJobIndex++];
 #if ALLOW_SHADERMAP_DEBUG_DATA
-		InFlightJobs.Remove(BaseJob->Hash);
+		if (IsRunningCookCommandlet())
+		{
+			InFlightJobs.Remove(BaseJob->Hash);
+		}
 #endif
 
 		FShaderCompileJob* SingleJob = BaseJob->GetSingleShaderJob();
diff --git a/Engine/Source/Runtime/Engine/Public/MaterialShared.h b/Engine/Source/Runtime/Engine/Public/MaterialShared.h
--- a/Engine/Source/Runtime/Engine/Public/MaterialShared.h
+++ b/Engine/Source/Runtime/Engine/Public/MaterialShared.h
@@ -1797,6 +1797,9 @@ private:
 #endif
 
 #if WITH_EDITOR && ALLOW_SHADERMAP_DEBUG_DATA
+	// InFlightJobs is only populated during cooking as its current use is to track cook hangs.
+	// if we need to generalize this later threadsafety will need to be added to make this work with ODSC
+	// where the same shadermap can be accessed simultaneously both on GT and RT
 	mutable TMap<uint32, TPair<FString, TSharedPtr<FShaderCompileJobStatus>>> InFlightJobs;
 #endif
 
