[Shaders] Remove check of virtual shader file paths to include "/Generated/".

This restriction on generated shader source is no longer desired to allow transient code to act as a proxy for global shaders.
These shaders need to use the same base path as global shaders, which do not include "/Generated/", so that the preprocessor can find its include dependencies.

diff --git a/Engine/Source/Runtime/Engine/Private/ShaderCompiler/ShaderCompiler.cpp b/Engine/Source/Runtime/Engine/Private/ShaderCompiler/ShaderCompiler.cpp
--- a/Engine/Source/Runtime/Engine/Private/ShaderCompiler/ShaderCompiler.cpp
+++ b/Engine/Source/Runtime/Engine/Private/ShaderCompiler/ShaderCompiler.cpp
@@ -2997,15 +2997,14 @@ void GenerateInstancedStereoCode(FString& Result, EShaderPlatform ShaderPlatform
 	Result += "#endif\n";
 }
 
+/**
+ * Basic validation of virtual shader file paths. This used to also require 'VirtualShaderFilePath' to include "/Generated/",
+ * which is no longer desired to allow compiling transient code that acts as a proxy for any other shader.
+ */
 void ValidateShaderFilePath(const FString& VirtualShaderFilePath, const FString& VirtualSourceFilePath)
 {
 	check(CheckVirtualShaderFilePath(VirtualShaderFilePath));
 
-	checkf(VirtualShaderFilePath.Contains(TEXT("/Generated/")),
-		TEXT("Incorrect virtual shader path for generated file '%s': Generated files must be located under an "
-				"non existing 'Generated' directory, for instance: /Engine/Generated/ or /Plugin/FooBar/Generated/."),
-		*VirtualShaderFilePath);
-
 	checkf(VirtualShaderFilePath == VirtualSourceFilePath || FPaths::GetExtension(VirtualShaderFilePath) == TEXT("ush"),
 		TEXT("Incorrect virtual shader path extension for generated file '%s': Generated file must either be the "
 				"USF to compile, or a USH file to be included."),
