Fix thumbnail rendering and live parameter updates for Material Function Instances and Material Layer Instances.

- `UMaterialFunctionInstance::GetPreviewMaterial()` is called when going to render a thumbnail preview and the material is made on demand.
- We would never re-create this preview material ever, and we would never update the parameter values of the preview material.  So as you tweaked parameters in your function instances the thumbnail preview would only ever render with the original state when the preview material was created.
- Now we always copy over new parameters on the preview material when asking for the preview material.
- Also, to support live thumbnail preview when tweaking parameters (useful when in the Material Layer Tool), we dirty the material function instance as well as apply the changes to the source function to see live updates on thumbnails.

[RN] minor | Rendering - Materials and Shaders

diff --git a/Engine/Source/Editor/UnrealEd/Private/PreviewMaterial.cpp b/Engine/Source/Editor/UnrealEd/Private/PreviewMaterial.cpp
--- a/Engine/Source/Editor/UnrealEd/Private/PreviewMaterial.cpp
+++ b/Engine/Source/Editor/UnrealEd/Private/PreviewMaterial.cpp
@@ -589,6 +589,14 @@ void UMaterialEditorInstanceConstant::PostEditChangeProperty(FPropertyChangedEve
 				RegenerateArrays();
 			}
 		}
+		else if (bIsFunctionPreviewMaterial)
+		{
+			// If a property changed and we are editing a material function instance
+			// we should dirty the function and apply the parameter changes to the source material
+			// This way you get live previews and updates of the thumbnails.
+			bIsFunctionInstanceDirty = true;
+			ApplySourceFunctionChanges();
+		}
 
 		CopyToSourceInstance(bLayersParameterChanged);
 
diff --git a/Engine/Source/Runtime/Engine/Private/Materials/MaterialExpressions.cpp b/Engine/Source/Runtime/Engine/Private/Materials/MaterialExpressions.cpp
--- a/Engine/Source/Runtime/Engine/Private/Materials/MaterialExpressions.cpp
+++ b/Engine/Source/Runtime/Engine/Private/Materials/MaterialExpressions.cpp
@@ -15291,12 +15291,16 @@ UMaterialInterface* UMaterialFunctionInstance::GetPreviewMaterial()
 	if (nullptr == PreviewMaterial)
 	{
 		PreviewMaterial = NewObject<UMaterialInstanceConstant>((UObject*)GetTransientPackage(), FName(TEXT("None")), RF_Transient);
-		PreviewMaterial->SetParentEditorOnly(Parent ? Parent->GetPreviewMaterial() : nullptr);
-		OverrideMaterialInstanceParameterValues(PreviewMaterial);
-		PreviewMaterial->PreEditChange(nullptr);
-		PreviewMaterial->PostEditChange();
-
 	}
+
+	// Update parameters in case they've changed so we get a live preview.
+	// Once the preview material is created for a material function instance we never re-create it again
+	// so it is important to take parameter updates so thumbnail and preview renders are correct.
+	PreviewMaterial->SetParentEditorOnly(Parent ? Parent->GetPreviewMaterial() : nullptr);
+	OverrideMaterialInstanceParameterValues(PreviewMaterial);
+	PreviewMaterial->PreEditChange(nullptr);
+	PreviewMaterial->PostEditChange();
+
 	return PreviewMaterial;
 }
 
