Replication:
* Allow actors to opt-in to full init state replication by settting FOutBunch::bOutWantsFullInitState in their OnSerializeNewActor
* This causes ALL replicated properties to be sent even if they are identical to the Archetype/CDO of the replicated object
* Can be used to fix issues where the authority sets the replicated value of a loaded (stable) actor back to the CDO value.
* In non-iris replication this prevents the client from receiving the replicated value and causes a divergence between each.

EngineNetTest:
* Added UnitTest Net.Engine.InitState showing the issue and the fix.

diff --git a/Engine/Source/Runtime/Engine/Private/DataBunch.cpp b/Engine/Source/Runtime/Engine/Private/DataBunch.cpp
--- a/Engine/Source/Runtime/Engine/Private/DataBunch.cpp
+++ b/Engine/Source/Runtime/Engine/Private/DataBunch.cpp
@@ -138,6 +138,7 @@ FOutBunch::FOutBunch( UChannel* InChannel, bool bInClose )
 ,	bHasPackageMapExports	( 0 )
 ,	bHasMustBeMappedGUIDs	( 0 )
 ,	bPartialCustomExportsFinal	( 0 )
+,	bOutWantsFullInitState	( 0 )
 ,	CloseReason( EChannelCloseReason::Destroyed )
 {
 	checkSlow(!Channel->Closing);
@@ -173,6 +174,7 @@ FOutBunch::FOutBunch( UPackageMap *InPackageMap, int64 MaxBits )
 ,	bHasPackageMapExports	( 0 )
 ,	bHasMustBeMappedGUIDs	( 0 )
 ,	bPartialCustomExportsFinal	( 0 )
+,	bOutWantsFullInitState	( 0 )
 ,	CloseReason( EChannelCloseReason::Destroyed )
 {
 }
diff --git a/Engine/Source/Runtime/Engine/Private/DataChannel.cpp b/Engine/Source/Runtime/Engine/Private/DataChannel.cpp
--- a/Engine/Source/Runtime/Engine/Private/DataChannel.cpp
+++ b/Engine/Source/Runtime/Engine/Private/DataChannel.cpp
@@ -3768,6 +3768,8 @@ int64 UActorChannel::ReplicateActor()
 		bWroteSomethingImportant = true;
 
 		Actor->OnSerializeNewActor(Bunch);
+
+		RepFlags.bForceInitialDirty = Bunch.bOutWantsFullInitState;
 	}
 
 	// Possibly downgrade role of actor if this connection doesn't own it
@@ -3787,7 +3789,7 @@ int64 UActorChannel::ReplicateActor()
 	RepFlags.bRepPhysics	= Actor->GetReplicatedMovement().bRepPhysics;
 	RepFlags.bReplay		= bReplay;
 	RepFlags.bClientReplay	= ActorWorld->IsRecordingClientReplay();
-	RepFlags.bForceInitialDirty = Connection->IsForceInitialDirty();
+	RepFlags.bForceInitialDirty |= Connection->IsForceInitialDirty();
 	if (EnumHasAnyFlags(ActorReplicator->RepLayout->GetFlags(), ERepLayoutFlags::HasDynamicConditionProperties))
 	{
 		if (const FRepState* RepState = ActorReplicator->RepState.Get())
diff --git a/Engine/Source/Runtime/Engine/Public/Net/DataBunch.h b/Engine/Source/Runtime/Engine/Public/Net/DataBunch.h
--- a/Engine/Source/Runtime/Engine/Public/Net/DataBunch.h
+++ b/Engine/Source/Runtime/Engine/Public/Net/DataBunch.h
@@ -43,6 +43,7 @@ public:
 	uint8					bHasPackageMapExports:1;	// This bunch has networkGUID name/id pairs
 	uint8					bHasMustBeMappedGUIDs:1;	// This bunch has guids that must be mapped before we can process this bunch
 	uint8					bPartialCustomExportsFinal:1;	// This bunch marks the end of the CustomExports data that needs to be processed immediately (not queued)
+	uint8					bOutWantsFullInitState:1;	// Set this true to force all replicated properties to be serialized in the initial bunch even if they do not differ from Archetype baseline
 
 	EChannelCloseReason		CloseReason;
 
