[HWRT] Changes to FSkeletalMeshLODRenderData -> StaticRayTracingGeometry initialization.
- moved initalization logic out of FSkeletalMeshObjectStatic to a new function FSkeletalMeshRenderData::InitStaticRayTracingGeometry(...).
- only initialize static ray tracing geometry if ray tracing is enabled.

diff --git a/Engine/Source/Runtime/Engine/Private/SkeletalMeshRenderData.cpp b/Engine/Source/Runtime/Engine/Private/SkeletalMeshRenderData.cpp
--- a/Engine/Source/Runtime/Engine/Private/SkeletalMeshRenderData.cpp
+++ b/Engine/Source/Runtime/Engine/Private/SkeletalMeshRenderData.cpp
@@ -690,6 +690,101 @@ void FSkeletalMeshRenderData::ReleaseResources()
 	}
 }
 
+void FSkeletalMeshRenderData::InitStaticRayTracingGeometry(int32 LODIndex)
+{
+#if RHI_RAYTRACING
+	FSkeletalMeshLODRenderData& LODData = LODRenderData[LODIndex];
+
+	if (LODData.NumReferencingStaticSkeletalMeshObjects == 0)
+	{
+		FPositionVertexBuffer* PositionVertexBufferPtr = &LODData.StaticVertexBuffers.PositionVertexBuffer;
+		FRawStaticIndexBuffer16or32Interface* IndexBufferPtr = LODData.MultiSizeIndexContainer.GetIndexBuffer();
+
+		uint32 TrianglesCount = 0;
+		for (int32 SectionIndex = 0; SectionIndex < LODData.RenderSections.Num(); SectionIndex++)
+		{
+			const FSkelMeshRenderSection& Section = LODData.RenderSections[SectionIndex];
+			TrianglesCount += Section.NumTriangles;
+		}
+
+		TArray<FSkelMeshRenderSection>* RenderSections = &LODData.RenderSections;
+		ENQUEUE_RENDER_COMMAND(InitSkeletalRenderStaticRayTracingGeometry)(UE::RenderCommandPipe::SkeletalMesh,
+			[PositionVertexBufferPtr, IndexBufferPtr, TrianglesCount, RenderSections,
+			LODIndex = LODIndex,
+			SkelMeshRenderData = this,
+			&RayTracingGeometry = LODData.StaticRayTracingGeometry,
+			&bReferencedByStaticSkeletalMeshObjects_RenderThread = LODData.bReferencedByStaticSkeletalMeshObjects_RenderThread](FRHICommandList& RHICmdList)
+			{
+				FRayTracingGeometryInitializer Initializer;
+				static const FName DebugName("FSkeletalMeshLODRenderData_Static");
+				static int32 DebugNumber = 0;
+				Initializer.DebugName = FDebugName(DebugName, DebugNumber++);
+				Initializer.IndexBuffer = IndexBufferPtr->IndexBufferRHI;
+				Initializer.TotalPrimitiveCount = TrianglesCount;
+				Initializer.GeometryType = RTGT_Triangles;
+				Initializer.bFastBuild = false;
+
+				TArray<FRayTracingGeometrySegment> GeometrySections;
+				GeometrySections.Reserve(RenderSections->Num());
+
+				for (const FSkelMeshRenderSection& Section : *RenderSections)
+				{
+					FRayTracingGeometrySegment Segment;
+					Segment.VertexBuffer = PositionVertexBufferPtr->VertexBufferRHI;
+					Segment.VertexBufferElementType = VET_Float3;
+					Segment.VertexBufferOffset = 0;
+					Segment.VertexBufferStride = PositionVertexBufferPtr->GetStride();
+					Segment.MaxVertices = PositionVertexBufferPtr->GetNumVertices();
+					Segment.FirstPrimitive = Section.BaseIndex / 3;
+					Segment.NumPrimitives = Section.NumTriangles;
+					Segment.bEnabled = !Section.bDisabled && Section.bVisibleInRayTracing;
+					GeometrySections.Add(Segment);
+				}
+				Initializer.Segments = GeometrySections;
+
+				if (LODIndex < SkelMeshRenderData->CurrentFirstLODIdx) // According to GetMeshElementsConditionallySelectable(), non-resident LODs should just be skipped
+				{
+					Initializer.Type = ERayTracingGeometryInitializerType::StreamingDestination;
+				}
+
+				RayTracingGeometry.GroupHandle = SkelMeshRenderData->RayTracingGeometryGroupHandle;
+				RayTracingGeometry.LODIndex = LODIndex;
+
+				RayTracingGeometry.SetInitializer(Initializer);
+				RayTracingGeometry.InitResource(RHICmdList);
+
+				bReferencedByStaticSkeletalMeshObjects_RenderThread = true;
+			}
+		);
+	}
+
+	LODData.NumReferencingStaticSkeletalMeshObjects++;
+#endif // RHI_RAYTRACING
+}
+
+void FSkeletalMeshRenderData::ReleaseStaticRayTracingGeometry(int32 LODIndex)
+{
+#if RHI_RAYTRACING
+	FSkeletalMeshLODRenderData& LODData = LODRenderData[LODIndex];
+
+	if (LODData.NumReferencingStaticSkeletalMeshObjects > 0)
+	{
+		LODData.NumReferencingStaticSkeletalMeshObjects--;
+
+		if (LODData.NumReferencingStaticSkeletalMeshObjects == 0)
+		{
+			ENQUEUE_RENDER_COMMAND(ResetStaticRayTracingGeometryFlag)(UE::RenderCommandPipe::SkeletalMesh,
+				[&bReferencedByStaticSkeletalMeshObjects_RenderThread = LODData.bReferencedByStaticSkeletalMeshObjects_RenderThread]
+				{
+					bReferencedByStaticSkeletalMeshObjects_RenderThread = false;
+				});
+
+			BeginReleaseResource(&LODData.StaticRayTracingGeometry, &UE::RenderCommandPipe::SkeletalMesh);
+		}
+	}
+#endif // RHI_RAYTRACING
+}
+
 uint32 FSkeletalMeshRenderData::GetNumBoneInfluences(int32 MinLODIndex) const
 {
 	uint32 NumBoneInfluences = 0;
diff --git a/Engine/Source/Runtime/Engine/Private/SkeletalMeshSceneProxy.cpp b/Engine/Source/Runtime/Engine/Private/SkeletalMeshSceneProxy.cpp
--- a/Engine/Source/Runtime/Engine/Private/SkeletalMeshSceneProxy.cpp
+++ b/Engine/Source/Runtime/Engine/Private/SkeletalMeshSceneProxy.cpp
@@ -905,7 +905,7 @@ RayTracing::FGeometryGroupHandle FSkeletalMeshSceneProxy::GetRayTracingGeometryG
 
 TArray<FRayTracingGeometry*> FSkeletalMeshSceneProxy::GetStaticRayTracingGeometries() const
 {
-	if (IsRayTracingAllowed() && bRenderStatic)
+	if (IsRayTracingEnabled() && bRenderStatic)
 	{
 		TArray<FRayTracingGeometry*> RayTracingGeometries;
 		RayTracingGeometries.AddDefaulted(SkeletalMeshRenderData->LODRenderData.Num());
diff --git a/Engine/Source/Runtime/Engine/Private/SkeletalRenderStatic.cpp b/Engine/Source/Runtime/Engine/Private/SkeletalRenderStatic.cpp
--- a/Engine/Source/Runtime/Engine/Private/SkeletalRenderStatic.cpp
+++ b/Engine/Source/Runtime/Engine/Private/SkeletalRenderStatic.cpp
@@ -42,84 +42,12 @@ void FSkeletalMeshObjectStatic::InitResources(const FSkinnedMeshSceneProxyDesc&
 
 		check(SkelLOD.SkelMeshRenderData);
 		check(SkelLOD.SkelMeshRenderData->LODRenderData.IsValidIndex(LODIndex));
-
-		FSkeletalMeshLODRenderData& LODData = SkelLOD.SkelMeshRenderData->LODRenderData[LODIndex];
 		
 		// Skip LODs that have their render data stripped
-		if (LODData.GetNumVertices() > 0)
+		if (SkelLOD.SkelMeshRenderData->LODRenderData[LODIndex].GetNumVertices() > 0)
 		{
 			const FSkelMeshComponentLODInfo* CompLODInfo = InMeshDesc.LODInfo.IsValidIndex(LODIndex) ? &InMeshDesc.LODInfo[LODIndex] : nullptr;
 			SkelLOD.InitResources(CompLODInfo);
-
-#if RHI_RAYTRACING
-			if (IsRayTracingAllowed() && SkelLOD.SkelMeshRenderData->bSupportRayTracing)
-			{
-				if (LODData.NumReferencingStaticSkeletalMeshObjects == 0)
-				{
-					FPositionVertexBuffer* PositionVertexBufferPtr = &LODData.StaticVertexBuffers.PositionVertexBuffer;
-					FRawStaticIndexBuffer16or32Interface* IndexBufferPtr = LODData.MultiSizeIndexContainer.GetIndexBuffer();
-
-					uint32 TrianglesCount = 0;
-					for (int32 SectionIndex = 0; SectionIndex < LODData.RenderSections.Num(); SectionIndex++)
-					{
-						const FSkelMeshRenderSection& Section = LODData.RenderSections[SectionIndex];
-						TrianglesCount += Section.NumTriangles;
-					}
-
-					TArray<FSkelMeshRenderSection>* RenderSections = &LODData.RenderSections;
-					ENQUEUE_RENDER_COMMAND(InitSkeletalRenderStaticRayTracingGeometry)(UE::RenderCommandPipe::SkeletalMesh,
-						[PositionVertexBufferPtr, IndexBufferPtr, TrianglesCount, RenderSections,
-						LODIndex = LODIndex, 
-						SkelMeshRenderData = SkelLOD.SkelMeshRenderData, 
-						&RayTracingGeometry = LODData.StaticRayTracingGeometry,
-						&bReferencedByStaticSkeletalMeshObjects_RenderThread = LODData.bReferencedByStaticSkeletalMeshObjects_RenderThread](FRHICommandList& RHICmdList)
-						{
-							FRayTracingGeometryInitializer Initializer;
-							static const FName DebugName("FSkeletalMeshObjectLOD");
-							static int32 DebugNumber = 0;
-							Initializer.DebugName = FDebugName(DebugName, DebugNumber++);
-							Initializer.IndexBuffer = IndexBufferPtr->IndexBufferRHI;
-							Initializer.TotalPrimitiveCount = TrianglesCount;
-							Initializer.GeometryType = RTGT_Triangles;
-							Initializer.bFastBuild = false;
-
-							TArray<FRayTracingGeometrySegment> GeometrySections;
-							GeometrySections.Reserve(RenderSections->Num());
-
-							for (const FSkelMeshRenderSection& Section : *RenderSections)
-							{
-								FRayTracingGeometrySegment Segment;
-								Segment.VertexBuffer = PositionVertexBufferPtr->VertexBufferRHI;
-								Segment.VertexBufferElementType = VET_Float3;
-								Segment.VertexBufferOffset = 0;
-								Segment.VertexBufferStride = PositionVertexBufferPtr->GetStride();
-								Segment.MaxVertices = PositionVertexBufferPtr->GetNumVertices();
-								Segment.FirstPrimitive = Section.BaseIndex / 3;
-								Segment.NumPrimitives = Section.NumTriangles;
-								Segment.bEnabled = !Section.bDisabled && Section.bVisibleInRayTracing;
-								GeometrySections.Add(Segment);
-							}
-							Initializer.Segments = GeometrySections;
-
-							if (LODIndex < SkelMeshRenderData->CurrentFirstLODIdx) // According to GetMeshElementsConditionallySelectable(), non-resident LODs should just be skipped
-							{
-								Initializer.Type = ERayTracingGeometryInitializerType::StreamingDestination;
-							}
-
-							RayTracingGeometry.GroupHandle = SkelMeshRenderData->RayTracingGeometryGroupHandle;
-							RayTracingGeometry.LODIndex = LODIndex;
-
-							RayTracingGeometry.SetInitializer(Initializer);
-							RayTracingGeometry.InitResource(RHICmdList);
-
-							bReferencedByStaticSkeletalMeshObjects_RenderThread = true;
-						}
-					);
-				}
-
-				LODData.NumReferencingStaticSkeletalMeshObjects++;
-			}
-#endif
 		}
 	}
 }
@@ -129,31 +57,13 @@ void FSkeletalMeshObjectStatic::ReleaseResources()
 	for( int32 LODIndex=0;LODIndex < LODs.Num();LODIndex++ )
 	{
 		FSkeletalMeshObjectLOD& SkelLOD = LODs[LODIndex];
+
+		check(SkelLOD.SkelMeshRenderData);
+		check(SkelLOD.SkelMeshRenderData->LODRenderData.IsValidIndex(LODIndex));
 		
 		// Skip LODs that have their render data stripped
 		if (SkelLOD.SkelMeshRenderData->LODRenderData[LODIndex].GetNumVertices() > 0)
 		{
-#if RHI_RAYTRACING
-			if (IsRayTracingAllowed())
-			{
-				if (SkelLOD.SkelMeshRenderData->LODRenderData[LODIndex].NumReferencingStaticSkeletalMeshObjects > 0)
-				{
-					SkelLOD.SkelMeshRenderData->LODRenderData[LODIndex].NumReferencingStaticSkeletalMeshObjects--;
-
-					if (SkelLOD.SkelMeshRenderData->LODRenderData[LODIndex].NumReferencingStaticSkeletalMeshObjects == 0)
-					{
-						ENQUEUE_RENDER_COMMAND(ResetStaticRayTracingGeometryFlag)(UE::RenderCommandPipe::SkeletalMesh,
-							[&bReferencedByStaticSkeletalMeshObjects_RenderThread = SkelLOD.SkelMeshRenderData->LODRenderData[LODIndex].bReferencedByStaticSkeletalMeshObjects_RenderThread]
-						{
-							bReferencedByStaticSkeletalMeshObjects_RenderThread = false;
-						});
-
-						BeginReleaseResource(&SkelLOD.SkelMeshRenderData->LODRenderData[LODIndex].StaticRayTracingGeometry, &UE::RenderCommandPipe::SkeletalMesh);
-					}
-				}
-			}
-#endif
-
 			SkelLOD.ReleaseResources();
 		}
 	}
@@ -229,6 +139,14 @@ void FSkeletalMeshObjectStatic::FSkeletalMeshObjectLOD::InitResources(const FSke
 			VertexFactoryPtr->InitResource(RHICmdList);
 		});
 
+#if RHI_RAYTRACING
+	if (IsRayTracingEnabled() && SkelMeshRenderData->bSupportRayTracing)
+	{
+		SkelMeshRenderData->InitStaticRayTracingGeometry(LODIndex);
+		bStaticRayTracingGeometryInitialized = true;
+	}
+#endif
+
 	bResourcesInitialized = true;
 }
 
@@ -237,11 +155,15 @@ void FSkeletalMeshObjectStatic::FSkeletalMeshObjectLOD::InitResources(const FSke
  */
 void FSkeletalMeshObjectStatic::FSkeletalMeshObjectLOD::ReleaseResources()
 {	
-	BeginReleaseResource(&VertexFactory, &UE::RenderCommandPipe::SkeletalMesh);
-
 #if RHI_RAYTRACING
-	BeginReleaseResource(&RayTracingGeometry, &UE::RenderCommandPipe::SkeletalMesh);
-#endif // RHI_RAYTRACING
+	if (bStaticRayTracingGeometryInitialized)
+	{
+		SkelMeshRenderData->ReleaseStaticRayTracingGeometry(LODIndex);
+		bStaticRayTracingGeometryInitialized = false;
+	}
+#endif
+
+	BeginReleaseResource(&VertexFactory, &UE::RenderCommandPipe::SkeletalMesh);
 
 	bResourcesInitialized = false;
 }
diff --git a/Engine/Source/Runtime/Engine/Private/SkeletalRenderStatic.h b/Engine/Source/Runtime/Engine/Private/SkeletalRenderStatic.h
--- a/Engine/Source/Runtime/Engine/Private/SkeletalRenderStatic.h
+++ b/Engine/Source/Runtime/Engine/Private/SkeletalRenderStatic.h
@@ -72,10 +72,11 @@ private:
 		/** Color buffer to user, could be from asset or component override */
 		FColorVertexBuffer* ColorVertexBuffer;
 
-#if RHI_RAYTRACING
-		/** Geometry for ray tracing. */
-		FRayTracingGeometry RayTracingGeometry;
-#endif // RHI_RAYTRACING
+		/**
+		* Whether InitStaticRayTracingGeometry(...) was called during initialization,
+		* so ReleaseStaticRayTracingGeometry(...) must be called when releasing resources.
+		*/
+		bool bStaticRayTracingGeometryInitialized;
 
 		/** true if resources for this LOD have already been initialized. */
 		bool bResourcesInitialized;
@@ -85,6 +86,7 @@ private:
 			,	LODIndex(InLOD)
 			,	VertexFactory(InFeatureLevel, "FSkeletalMeshObjectStatic::FSkeletalMeshObjectLOD")
 			,	ColorVertexBuffer(nullptr)
+			,	bStaticRayTracingGeometryInitialized(false)
 			,	bResourcesInitialized(false)
 		{
 		}
diff --git a/Engine/Source/Runtime/Engine/Public/Rendering/SkeletalMeshRenderData.h b/Engine/Source/Runtime/Engine/Public/Rendering/SkeletalMeshRenderData.h
--- a/Engine/Source/Runtime/Engine/Public/Rendering/SkeletalMeshRenderData.h
+++ b/Engine/Source/Runtime/Engine/Public/Rendering/SkeletalMeshRenderData.h
@@ -81,6 +81,12 @@ public:
 	/** Return the resource size */
 	ENGINE_API void GetResourceSizeEx(FResourceSizeEx& CumulativeResourceSize);
 
+	/** Initializes the static ray tracing geometry for the specified LOD (if necessary). */
+	ENGINE_API void InitStaticRayTracingGeometry(int32 LODIndex);
+
+	/** Releases the static ray tracing geometry for the specified LOD (if it's the last reference). */
+	ENGINE_API void ReleaseStaticRayTracingGeometry(int32 LODIndex);
+
 	/** Get the estimated memory overhead of buffers marked as NeedsCPUAccess. */
 	ENGINE_API SIZE_T GetCPUAccessMemoryOverhead() const;
 
