[Visual Logger] Added coordinate system macro defines to the section for !ENABLE_VISUAL_LOG

diff --git a/Engine/Source/Runtime/Engine/Public/VisualLogger/VisualLogger.h b/Engine/Source/Runtime/Engine/Public/VisualLogger/VisualLogger.h
--- a/Engine/Source/Runtime/Engine/Public/VisualLogger/VisualLogger.h
+++ b/Engine/Source/Runtime/Engine/Public/VisualLogger/VisualLogger.h
@@ -183,6 +183,8 @@
 #define UE_CVLOG_WIRECIRCLE(Condition, LogOwner, CategoryName, Verbosity, Center, UpAxis, Radius, Color, Format, ...)
 #define UE_VLOG_WIRECIRCLE_THICK(LogOwner, CategoryName, Verbosity, Center, UpAxis, Radius, Color, Thickness, Format, ...)
 #define UE_CVLOG_WIRECIRCLE_THICK(Condition, LogOwner, CategoryName, Verbosity, Center, UpAxis, Radius, Color, Thickness, Format, ...)
+#define UE_VLOG_COORDINATESYSTEM(LogOwner, CategoryName, Verbosity, AxisLoc, AxisRot, Scale, Color, Thickness, Format, ...)
+#define UE_CVLOG_COORDINATESYSTEM(Condition, LogOwner, CategoryName, Verbosity, AxisLoc, AxisRot, Scale, Color, Thickness, Format, ...)
 
 #define DECLARE_VLOG_EVENT(EventName)
 #define DEFINE_VLOG_EVENT(EventName, Verbosity, UserFriendlyDesc)
