Skin Cache use packed tangents when inputs are packed. Fixed RHI validation error due to async compute being disabled for small RDG builders. Instead, the GPU skin cache system explicitly disables async compute for Niagara.

diff --git a/Engine/Plugins/FX/Niagara/Source/Niagara/Private/NiagaraGpuComputeDispatch.cpp b/Engine/Plugins/FX/Niagara/Source/Niagara/Private/NiagaraGpuComputeDispatch.cpp
--- a/Engine/Plugins/FX/Niagara/Source/Niagara/Private/NiagaraGpuComputeDispatch.cpp
+++ b/Engine/Plugins/FX/Niagara/Source/Niagara/Private/NiagaraGpuComputeDispatch.cpp
@@ -445,14 +445,16 @@ void FNiagaraGpuComputeDispatch::ProcessPendingTicksFlush(FRHICommandListImmedia
 
 			if (SkeletalMeshUpdater || GPUSkinCache)
 			{
+				const ERHIPipeline GPUSkinCachePipeline = ERHIPipeline::Graphics;
+
 				FRDGBuilder GraphBuilder(RHICmdList);
 				if (SkeletalMeshUpdater)
 				{
-					SkeletalMeshUpdater->Submit(GraphBuilder);
+					SkeletalMeshUpdater->Submit(GraphBuilder, GPUSkinCachePipeline);
 				}
 				if (GPUSkinCache)
 				{
-					GPUSkinCache->Dispatch(GraphBuilder, {});
+					GPUSkinCache->Dispatch(GraphBuilder, {}, GPUSkinCachePipeline);
 				}
 				GraphBuilder.Execute();
 			}
diff --git a/Engine/Source/Runtime/Engine/Internal/SkeletalMeshUpdater.h b/Engine/Source/Runtime/Engine/Internal/SkeletalMeshUpdater.h
--- a/Engine/Source/Runtime/Engine/Internal/SkeletalMeshUpdater.h
+++ b/Engine/Source/Runtime/Engine/Internal/SkeletalMeshUpdater.h
@@ -269,7 +269,7 @@ public:
 	};
 
 	// Issues setup tasks to process commands pushed from the game side. Use the provided tasks to sync stages as needed. The builder automatically syncs otherwise.
-	ENGINE_API FSubmitTasks Submit(FRDGBuilder& GraphBuilder);
+	ENGINE_API FSubmitTasks Submit(FRDGBuilder& GraphBuilder, ERHIPipeline GPUSkinCachePipeline);
 
 	// Waits for tasks associated with the provided update stage.
 	ENGINE_API static void WaitForStage(FRDGBuilder& GraphBuilder, ESkeletalMeshUpdateStage Stage);
diff --git a/Engine/Source/Runtime/Engine/Private/GPUSkinCache.cpp b/Engine/Source/Runtime/Engine/Private/GPUSkinCache.cpp
--- a/Engine/Source/Runtime/Engine/Private/GPUSkinCache.cpp
+++ b/Engine/Source/Runtime/Engine/Private/GPUSkinCache.cpp
@@ -287,6 +287,259 @@ enum class EGPUSkinCacheBufferBits : uint8
 };
 ENUM_CLASS_FLAGS(EGPUSkinCacheBufferBits);
 
+struct FGPUSkinCache::FSkinCacheRWBuffer
+{
+	FRWBuffer	Buffer;
+	ERHIAccess	AccessState = ERHIAccess::Unknown;	// Keep track of current access state
+	mutable uint32	LastTransitionFence = 0;
+
+	void Release()
+	{
+		Buffer.Release();
+		AccessState = ERHIAccess::Unknown;
+	}
+
+	bool UpdateFence(uint32 NextTransitionFence)
+	{
+		const bool bUpdateRequired = LastTransitionFence != NextTransitionFence;
+		LastTransitionFence = NextTransitionFence;
+		return bUpdateRequired;
+	}
+
+	// Update the access state and return transition info
+	FRHITransitionInfo UpdateAccessState(ERHIAccess NewState)
+	{
+		ERHIAccess OldState = AccessState;
+		AccessState = NewState;
+		return FRHITransitionInfo(Buffer.UAV.GetReference(), OldState, AccessState);
+	}
+};
+
+struct FGPUSkinCache::FRWBuffersAllocationInitializer
+{
+	static const uint32 PositionStride = 4;
+
+	EGPUSkinCacheBufferBits BufferBits = EGPUSkinCacheBufferBits::None;
+	uint32 NumVertices = 0;
+	uint32 IntermediateAccumulatedTangentsSize = 0;
+	EPixelFormat TangentFormat = PF_Unknown;
+
+	static uint32 GetPositionStride()
+	{
+		return PositionStride;
+	}
+
+	uint32 GetTangentStride() const
+	{
+		return GPixelFormats[TangentFormat].BlockBytes;
+	}
+
+	uint32 GetBufferSize() const
+	{
+		const bool bIntermediateTangents = EnumHasAnyFlags(BufferBits, EGPUSkinCacheBufferBits::IntermediateTangents);
+		const bool bPositionPrevious = EnumHasAnyFlags(BufferBits, EGPUSkinCacheBufferBits::PositionPrevious);
+
+		const uint32 TangentStride = GetTangentStride();
+
+		const uint32 PositionBufferSize = PositionStride * NumVertices * 3 * (bPositionPrevious ? NUM_BUFFERS : 1);
+		const uint32 TangentBufferSize =  TangentStride * NumVertices * 2;
+		const uint32 IntermediateTangentBufferSize = bIntermediateTangents ? TangentStride * NumVertices * 2 : 0;
+		const uint32 AccumulatedTangentBufferSize = IntermediateAccumulatedTangentsSize * FGPUSkinCache::IntermediateAccumBufferNumInts * sizeof(int32);
+
+		return TangentBufferSize + IntermediateTangentBufferSize + PositionBufferSize + AccumulatedTangentBufferSize;
+	}
+};
+
+struct FGPUSkinCache::FRWBuffersAllocation
+{
+	friend FRWBufferTracker;
+
+	FRWBuffersAllocation(FRHICommandList& RHICmdList, const FRWBuffersAllocationInitializer& InInitializer, const FName& OwnerName)
+		: Initializer(InInitializer)
+	{
+		const static FLazyName PositionsName(TEXT("SkinCachePositions"));
+		const static FLazyName TangentsName(TEXT("SkinCacheTangents"));
+		const static FLazyName IntermediateTangentsName(TEXT("SkinCacheIntermediateTangents"));
+		const static FLazyName IntermediateAccumulatedTangentsName(TEXT("SkinCacheIntermediateAccumulatedTangents"));
+
+		const int32 NumBuffers = EnumHasAnyFlags(Initializer.BufferBits, EGPUSkinCacheBufferBits::PositionPrevious) ? NUM_BUFFERS : 1;
+		
+		const uint32 PositionStride = InInitializer.GetPositionStride();
+		const uint32 TangentStride = InInitializer.GetTangentStride();
+
+		for (int32 Index = 0; Index < NumBuffers; ++Index)
+		{
+			PositionBuffers[Index].Buffer.ClassName = PositionsName;
+			PositionBuffers[Index].Buffer.OwnerName = OwnerName;
+			PositionBuffers[Index].Buffer.Initialize(RHICmdList, TEXT("SkinCachePositions"), PositionStride, Initializer.NumVertices * 3, PF_R32_FLOAT, ERHIAccess::SRVMask, BUF_Static);
+			PositionBuffers[Index].Buffer.Buffer->SetOwnerName(OwnerName);
+			PositionBuffers[Index].AccessState = ERHIAccess::Unknown;
+		}
+
+		// Tangents are skinned inside the main skinning compute shader and are always allocated, even if the recompute tangents pass doesn't run.
+		Tangents.Buffer.ClassName = TangentsName;
+		Tangents.Buffer.OwnerName = OwnerName;
+		Tangents.Buffer.Initialize(RHICmdList, TEXT("SkinCacheTangents"), TangentStride, Initializer.NumVertices * 2, Initializer.TangentFormat, BUF_Static);
+		Tangents.Buffer.Buffer->SetOwnerName(OwnerName);
+		Tangents.AccessState = ERHIAccess::Unknown;
+
+		if (EnumHasAnyFlags(Initializer.BufferBits, EGPUSkinCacheBufferBits::IntermediateTangents))
+		{
+			IntermediateTangents.Buffer.ClassName = IntermediateTangentsName;
+			IntermediateTangents.Buffer.OwnerName = OwnerName;
+			IntermediateTangents.Buffer.Initialize(RHICmdList, TEXT("SkinCacheIntermediateTangents"), TangentStride, Initializer.NumVertices * 2, Initializer.TangentFormat, BUF_Static);
+			IntermediateTangents.Buffer.Buffer->SetOwnerName(OwnerName);
+			IntermediateTangents.AccessState = ERHIAccess::Unknown;
+		}
+
+		if (Initializer.IntermediateAccumulatedTangentsSize > 0)
+		{
+			IntermediateAccumulatedTangents.Buffer.ClassName = IntermediateAccumulatedTangentsName;
+			IntermediateAccumulatedTangents.Buffer.OwnerName = OwnerName;
+			IntermediateAccumulatedTangents.Buffer.Initialize(RHICmdList, TEXT("SkinCacheIntermediateAccumulatedTangents"), sizeof(int32), Initializer.IntermediateAccumulatedTangentsSize * FGPUSkinCache::IntermediateAccumBufferNumInts, PF_R32_SINT, BUF_UnorderedAccess);
+			IntermediateAccumulatedTangents.Buffer.Buffer->SetOwnerName(OwnerName);
+			IntermediateAccumulatedTangents.AccessState = ERHIAccess::Unknown;
+
+			// The UAV must be zero-filled. We leave it zeroed after each round (see RecomputeTangentsPerVertexPass.usf), so this is only needed on when the buffer is first created.
+			RHICmdList.ClearUAVUint(IntermediateAccumulatedTangents.Buffer.UAV, FUintVector4(0, 0, 0, 0));
+		}
+	}
+
+	~FRWBuffersAllocation()
+	{
+		for (int32 Index = 0; Index < NUM_BUFFERS; ++Index)
+		{
+			PositionBuffers[Index].Release();
+		}
+
+		Tangents.Release();
+		IntermediateTangents.Release();
+		IntermediateAccumulatedTangents.Release();
+	}
+
+	uint64 GetBufferSize() const
+	{
+		return Initializer.GetBufferSize();
+	}
+
+	FSkinCacheRWBuffer* GetTangentBuffer()
+	{
+		return &Tangents;
+	}
+
+	FSkinCacheRWBuffer* GetIntermediateTangentBuffer()
+	{
+		return EnumHasAllFlags(Initializer.BufferBits, EGPUSkinCacheBufferBits::IntermediateTangents) ? &IntermediateTangents : nullptr;
+	}
+
+	FSkinCacheRWBuffer* GetIntermediateAccumulatedTangentBuffer()
+	{
+		return Initializer.IntermediateAccumulatedTangentsSize > 0 ? &IntermediateAccumulatedTangents : nullptr;
+	}
+
+	bool HasPreviousBuffer() const
+	{
+		return EnumHasAllFlags(Initializer.BufferBits, EGPUSkinCacheBufferBits::PositionPrevious);
+	}
+
+	FSkinCacheRWBuffer& GetPositionBuffer()
+	{
+		return PositionBuffers[0];
+	}
+
+private:
+	// Output of the GPU skinning (ie Pos, Normals)
+	FSkinCacheRWBuffer PositionBuffers[NUM_BUFFERS];
+
+	FSkinCacheRWBuffer Tangents;
+	FSkinCacheRWBuffer IntermediateTangents;
+	FSkinCacheRWBuffer IntermediateAccumulatedTangents;	// Intermediate buffer used to accumulate results of triangle pass to be passed onto vertex pass
+
+	const FRWBuffersAllocationInitializer Initializer;
+};
+
+struct FGPUSkinCache::FRWBufferTracker
+{
+	FRWBuffersAllocation* Allocation;
+
+	FRWBufferTracker()
+		: Allocation(nullptr)
+	{
+		Reset();
+	}
+
+	void Reset()
+	{
+		for (int32 Index = 0; Index < NUM_BUFFERS; ++Index)
+		{
+			Revisions[Index] = 0;
+			BoneBuffers[Index] = nullptr;
+		}
+	}
+
+	uint32 GetBufferSize() const
+	{
+		return Allocation->GetBufferSize();
+	}
+
+	FSkinCacheRWBuffer* Find(const FVertexBufferAndSRV& BoneBuffer, uint32 Revision)
+	{
+		for (int32 Index = 0; Index < NUM_BUFFERS; ++Index)
+		{
+			if (Revisions[Index] == Revision && BoneBuffers[Index] == &BoneBuffer)
+			{
+				return &Allocation->PositionBuffers[Index];
+			}
+		}
+
+		return nullptr;
+	}
+
+	FSkinCacheRWBuffer* GetTangentBuffer()
+	{
+		return Allocation ? Allocation->GetTangentBuffer() : nullptr;
+	}
+
+	FSkinCacheRWBuffer* GetIntermediateTangentBuffer()
+	{
+		return Allocation ? Allocation->GetIntermediateTangentBuffer() : nullptr;
+	}
+
+	FSkinCacheRWBuffer* GetIntermediateAccumulatedTangentBuffer()
+	{
+		return Allocation ? Allocation->GetIntermediateAccumulatedTangentBuffer() : nullptr;
+	}
+
+	// Allocates an element that's not the "Used" element passed in (or if Used is NULL, allocates any element).
+	FSkinCacheRWBuffer* AllocateUnused(const FVertexBufferAndSRV& BoneBuffer, uint32 Revision, const FSkinCacheRWBuffer* Used)
+	{
+		int32 UnusedIndex = Used == &Allocation->PositionBuffers[0] ? 1 : 0;
+		Revisions[UnusedIndex] = Revision;
+		BoneBuffers[UnusedIndex] = &BoneBuffer;
+
+		return &Allocation->PositionBuffers[UnusedIndex];
+	}
+
+	// On recreate of the render state where the GPU skin cache entry is preserved, the bone buffer will have been reallocated,
+	// even though the transforms didn't change.  We need to force the Find() call above to treat the data as up-to-date, which
+	// can be accomplished by updating the BoneBuffer pointer for the previous Revision, so it matches again.
+	void UpdatePreviousBoneBuffer(const FVertexBufferAndSRV& PreviousBoneBuffer, uint32 PreviousRevision)
+	{
+		for (int32 Index = 0; Index < NUM_BUFFERS; ++Index)
+		{
+			if (Revisions[Index] == PreviousRevision)
+			{
+				BoneBuffers[Index] = &PreviousBoneBuffer;
+				break;
+			}
+		}
+	}
+
+private:
+	uint32 Revisions[NUM_BUFFERS];
+	const FVertexBufferAndSRV* BoneBuffers[NUM_BUFFERS];
+};
+
 class FGPUSkinCacheEntry
 {
 public:
@@ -1230,27 +1483,6 @@ void FGPUSkinCache::DispatchUpdateSkinTangentsTrianglePass(FRHICommandList& RHIC
 	UnsetShaderParametersLegacyCS(RHICmdList, Shader);
 }
 
-FGPUSkinCache::FRWBuffersAllocation* FGPUSkinCache::TryAllocBuffer(FRHICommandList& RHICmdList, EGPUSkinCacheBufferBits BufferBits, uint32 NumVertices, uint32 InIntermediateAccumulatedTangentsSize, const FName& OwnerName)
-{
-	uint64 MaxSizeInBytes = (uint64)(GSkinCacheSceneMemoryLimitInMB * 1024.0f * 1024.0f);
-	uint64 RequiredMemInBytes = FRWBuffersAllocation::CalculateRequiredMemory(BufferBits, NumVertices, InIntermediateAccumulatedTangentsSize);
-	if (bRequiresMemoryLimit && UsedMemoryInBytes + RequiredMemInBytes >= MaxSizeInBytes)
-	{
-		ExtraRequiredMemory += RequiredMemInBytes;
-
-		// Can't fit
-		return nullptr;
-	}
-
-	FRWBuffersAllocation* NewAllocation = new FRWBuffersAllocation(RHICmdList, BufferBits, NumVertices, InIntermediateAccumulatedTangentsSize, OwnerName);
-	Allocations.Add(NewAllocation);
-
-	UsedMemoryInBytes += RequiredMemInBytes;
-	INC_MEMORY_STAT_BY(STAT_GPUSkinCache_TotalMemUsed, RequiredMemInBytes);
-
-	return NewAllocation;
-}
-
 DECLARE_GPU_STAT(GPUSkinCache);
 
 void FGPUSkinCache::TransitionBuffers(FRHICommandList& RHICmdList, TArray<FSkinCacheRWBuffer*>& Buffers, ERHIAccess ToState)
@@ -1319,8 +1551,8 @@ ERHIPipeline FGPUSkinCache::GetDispatchPipeline(FRDGBuilder& GraphBuilder)
 
 struct FGPUSkinCache::FTaskData
 {
-	FTaskData(FRDGBuilder& GraphBuilder)
-		: Pipeline(FGPUSkinCache::GetDispatchPipeline(GraphBuilder))
+	FTaskData(FRDGBuilder& GraphBuilder, ERHIPipeline InPipeline)
+		: Pipeline(InPipeline)
 		, bAsyncCommandList(GraphBuilder.IsParallelSetupEnabled() && !GNumDispatchesToCapture)
 	{
 		if (bAsyncCommandList)
@@ -1386,9 +1618,9 @@ struct FGPUSkinCache::FTaskData
 	bool bSignalPassAdded = false;
 };
 
-UE::Tasks::FTask FGPUSkinCache::Dispatch(FRDGBuilder& GraphBuilder, const UE::Tasks::FTask& PrerequisitesTask)
+UE::Tasks::FTask FGPUSkinCache::Dispatch(FRDGBuilder& GraphBuilder, const UE::Tasks::FTask& PrerequisitesTask, ERHIPipeline InPipeline)
 {
-	FTaskData* TaskData = &GraphBuilder.Blackboard.Create<FTaskData>(GraphBuilder);
+	FTaskData* TaskData = &GraphBuilder.Blackboard.Create<FTaskData>(GraphBuilder, InPipeline);
 
 	TaskData->SetupTask = GraphBuilder.AddSetupTask([this, TaskData]
 	{
@@ -1974,18 +2206,39 @@ bool FGPUSkinCache::ProcessEntry(
 		// If something caused the existing entry to be invalid, disable recreate logic for the rest of the function
 		bRecreating = false;
 
-		// IntermediateTangents buffer is needed if mesh has at least one section using vertex color as recompute tangents blending mask
-		int32 TotalNumVertices = VertexFactory->GetNumVertices();
+		FRWBuffersAllocationInitializer Initializer
+		{
+			  .BufferBits = BufferBits
+			, .NumVertices = VertexFactory->GetNumVertices()
+			, .IntermediateAccumulatedTangentsSize = IntermediateAccumulatedTangentBufferSize
+		};
+
+		// OpenGL ES does not support writing to RGBA16_SNORM images, so use the packed format instead.
+		if (IsOpenGLPlatform(GMaxRHIShaderPlatform))
+		{
+			Initializer.TangentFormat = PF_R16G16B16A16_SINT;
+		}
+		else if (GPixelFormats[TargetVertexFactory->GetTangentFormat()].BlockBytes == 4)
+		{
+			Initializer.TangentFormat = PF_R8G8B8A8_SNORM;
+		}
+		else
+		{
+			Initializer.TangentFormat = PF_R16G16B16A16_SNORM;
+		}
+
+		const uint32 BufferSize = Initializer.GetBufferSize();
+		const uint32 MaxSizeInBytes = (GSkinCacheSceneMemoryLimitInMB * 1024.0f * 1024.0f);
 
-		FRWBuffersAllocation* NewPositionAllocation = TryAllocBuffer(RHICmdList, BufferBits, TotalNumVertices, IntermediateAccumulatedTangentBufferSize, Skin->GetAssetPathName(LODIndex));
-		if (!NewPositionAllocation)
+		if (bRequiresMemoryLimit && UsedMemoryInBytes + BufferSize >= MaxSizeInBytes)
 		{
+			ExtraRequiredMemory += BufferSize;
+
 			if (GSkinCachePrintMemorySummary > 0)
 			{
-				uint64 RequiredMemInBytes = FRWBuffersAllocation::CalculateRequiredMemory(BufferBits, TotalNumVertices, IntermediateAccumulatedTangentBufferSize);
 				UE_LOG(LogSkinCache, Warning, TEXT("FGPUSkinCache::ProcessEntry%s failed to allocate %.3fMB for mesh %s LOD%d, extra required memory increased to %.3fMB")
 					, Mode == EGPUSkinCacheEntryMode::RayTracing ? TEXT("[RT]") : TEXT("")
-					, RequiredMemInBytes / MBSize
+					, BufferSize / MBSize
 					, *GetSkeletalMeshObjectName(Skin)
 					, LODIndex
 					, ExtraRequiredMemory / MBSize
@@ -1996,12 +2249,18 @@ bool FGPUSkinCache::ProcessEntry(
 			return false;
 		}
 
-		InOutEntry = new FGPUSkinCacheEntry(this, Skin, NewPositionAllocation, LODIndex, Mode);
+		FRWBuffersAllocation* BufferAllocation = new FRWBuffersAllocation(RHICmdList, Initializer, Skin->GetAssetPathName(LODIndex));
+		Allocations.Add(BufferAllocation);
+
+		UsedMemoryInBytes += BufferSize;
+		INC_MEMORY_STAT_BY(STAT_GPUSkinCache_TotalMemUsed, BufferSize);
+
+		InOutEntry = new FGPUSkinCacheEntry(this, Skin, BufferAllocation, LODIndex, Mode);
 		InOutEntry->GPUSkin = Skin;
 
 		InOutEntry->SetupSection(
 			Section,
-			NewPositionAllocation,
+			BufferAllocation,
 			&LodData.RenderSections[Section],
 			MorphVertexBuffer,
 			ClothVertexBuffer,
@@ -2221,24 +2480,6 @@ void FGPUSkinCache::Release(FGPUSkinCacheEntry*& SkinCacheEntry)
 	}
 }
 
-void FGPUSkinCache::FRWBuffersAllocation::RemoveAllFromTransitionArray(TSet<FSkinCacheRWBuffer*>& InBuffersToTransition)
-{
-	for (uint32 i = 0; i < NUM_BUFFERS; i++)
-	{
-		FSkinCacheRWBuffer& RWBuffer = PositionBuffers[i];
-		InBuffersToTransition.Remove(&RWBuffer);
-		
-		if (auto TangentBuffer = GetTangentBuffer())
-		{
-			InBuffersToTransition.Remove(TangentBuffer);
-		}
-		if (auto IntermediateTangentBuffer = GetIntermediateTangentBuffer())
-		{
-			InBuffersToTransition.Remove(IntermediateTangentBuffer);
-		}
-	}
-}
-
 void FGPUSkinCache::ReleaseSkinCacheEntry(FGPUSkinCacheEntry* SkinCacheEntry)
 {
 	FGPUSkinCache* SkinCache = SkinCacheEntry->SkinCache;
@@ -2246,9 +2487,9 @@ void FGPUSkinCache::ReleaseSkinCacheEntry(FGPUSkinCacheEntry* SkinCacheEntry)
 	FRWBuffersAllocation* PositionAllocation = SkinCacheEntry->PositionAllocation;
 	if (PositionAllocation)
 	{
-		uint64 RequiredMemInBytes = PositionAllocation->GetNumBytes();
-		SkinCache->UsedMemoryInBytes -= RequiredMemInBytes;
-		DEC_MEMORY_STAT_BY(STAT_GPUSkinCache_TotalMemUsed, RequiredMemInBytes);
+		uint32 BufferSize = PositionAllocation->GetBufferSize();
+		SkinCache->UsedMemoryInBytes -= BufferSize;
+		DEC_MEMORY_STAT_BY(STAT_GPUSkinCache_TotalMemUsed, BufferSize);
 
 		SkinCache->Allocations.Remove(PositionAllocation);
 
@@ -2450,7 +2691,7 @@ void FGPUSkinCache::PrintMemorySummary() const
 				RecomputeTangentSections = TEXT("Off");
 			}
 
-			uint64 MemInBytes = Entry->PositionAllocation ? Entry->PositionAllocation->GetNumBytes() : 0;
+			uint64 MemInBytes = Entry->PositionAllocation ? Entry->PositionAllocation->GetBufferSize() : 0;
 			uint64 TangentsInBytes = (Entry->PositionAllocation && Entry->PositionAllocation->GetTangentBuffer()) ? Entry->PositionAllocation->GetTangentBuffer()->Buffer.NumBytes : 0;
 			uint64 IntermediateTangentsInBytes = (Entry->PositionAllocation && Entry->PositionAllocation->GetIntermediateTangentBuffer()) ? Entry->PositionAllocation->GetIntermediateTangentBuffer()->Buffer.NumBytes : 0;
 			uint64 IntermediateAccumulatedTangentsInBytes = (Entry->PositionAllocation && Entry->PositionAllocation->GetIntermediateAccumulatedTangentBuffer()) ? Entry->PositionAllocation->GetIntermediateAccumulatedTangentBuffer()->Buffer.NumBytes : 0;
@@ -2519,12 +2760,12 @@ FColor FGPUSkinCache::GetVisualizationDebugColor(const FName& GPUSkinCacheVisual
 		}
 		else if (ModeType == FGPUSkinCacheVisualizationData::FModeType::Memory)
 		{
-			uint64 MemoryInBytes = (Entry && Entry->PositionAllocation) ? Entry->PositionAllocation->GetNumBytes() : 0;
+			uint64 MemoryInBytes = (Entry && Entry->PositionAllocation) ? Entry->PositionAllocation->GetBufferSize() : 0;
 #if RHI_RAYTRACING
 			if (RayTracingEntry && RayTracingEntry != Entry)
 			{
 				// Separate ray tracing entry
-				MemoryInBytes += RayTracingEntry->PositionAllocation ? RayTracingEntry->PositionAllocation->GetNumBytes() : 0;
+				MemoryInBytes += RayTracingEntry->PositionAllocation ? RayTracingEntry->PositionAllocation->GetBufferSize() : 0;
 			}
 #endif
 			float MemoryInMB = MemoryInBytes / MBSize;
@@ -2599,98 +2840,6 @@ void FGPUSkinCache::DrawVisualizationInfoText(const FName& GPUSkinCacheVisualiza
 	}
 }
 
-FGPUSkinCache::FRWBuffersAllocation::FRWBuffersAllocation(FRHICommandList& RHICmdList, EGPUSkinCacheBufferBits InBufferBits, uint32 InNumVertices, uint32 InIntermediateAccumulatedTangentsSize, const FName& OwnerName)
-	: NumVertices(InNumVertices)
-	, IntermediateAccumulatedTangentsSize(InIntermediateAccumulatedTangentsSize)
-	, BufferBits(InBufferBits)
-{
-	const static FLazyName PositionsName(TEXT("SkinCachePositions"));
-	const static FLazyName TangentsName(TEXT("SkinCacheTangents"));
-	const static FLazyName IntermediateTangentsName(TEXT("SkinCacheIntermediateTangents"));
-	const static FLazyName IntermediateAccumulatedTangentsName(TEXT("SkinCacheIntermediateAccumulatedTangents"));
-
-	const int32 NumBuffers = EnumHasAnyFlags(InBufferBits, EGPUSkinCacheBufferBits::PositionPrevious) ? NUM_BUFFERS : 1;
-
-	for (int32 Index = 0; Index < NumBuffers; ++Index)
-	{
-		PositionBuffers[Index].Buffer.ClassName = PositionsName;
-		PositionBuffers[Index].Buffer.OwnerName = OwnerName;
-		PositionBuffers[Index].Buffer.Initialize(RHICmdList, TEXT("SkinCachePositions"), PosBufferBytesPerElement, NumVertices * 3, PF_R32_FLOAT, ERHIAccess::SRVMask, BUF_Static);
-		PositionBuffers[Index].Buffer.Buffer->SetOwnerName(OwnerName);
-		PositionBuffers[Index].AccessState = ERHIAccess::Unknown;
-	}
-	
-	// OpenGL ES does not support writing to RGBA16_SNORM images, instead pack data into SINT in the shader
-	const EPixelFormat TangentsFormat = IsOpenGLPlatform(GMaxRHIShaderPlatform) ? PF_R16G16B16A16_SINT : PF_R16G16B16A16_SNORM;
-
-	// Tangents are skinned inside the main skinning compute shader and are always allocated, even if the recompute tangents pass doesn't run.
-	Tangents.Buffer.ClassName = TangentsName;
-	Tangents.Buffer.OwnerName = OwnerName;
-	Tangents.Buffer.Initialize(RHICmdList, TEXT("SkinCacheTangents"), TangentBufferBytesPerElement, NumVertices * 2, TangentsFormat, BUF_Static);
-	Tangents.Buffer.Buffer->SetOwnerName(OwnerName);
-	Tangents.AccessState = ERHIAccess::Unknown;
-
-	if (EnumHasAnyFlags(InBufferBits, EGPUSkinCacheBufferBits::IntermediateTangents))
-	{
-		IntermediateTangents.Buffer.ClassName = IntermediateTangentsName;
-		IntermediateTangents.Buffer.OwnerName = OwnerName;
-		IntermediateTangents.Buffer.Initialize(RHICmdList, TEXT("SkinCacheIntermediateTangents"), TangentBufferBytesPerElement, NumVertices * 2, TangentsFormat, BUF_Static);
-		IntermediateTangents.Buffer.Buffer->SetOwnerName(OwnerName);
-		IntermediateTangents.AccessState = ERHIAccess::Unknown;
-	}
-
-	if (IntermediateAccumulatedTangentsSize > 0)
-	{
-		IntermediateAccumulatedTangents.Buffer.ClassName = IntermediateAccumulatedTangentsName;
-		IntermediateAccumulatedTangents.Buffer.OwnerName = OwnerName;
-		IntermediateAccumulatedTangents.Buffer.Initialize(RHICmdList, TEXT("SkinCacheIntermediateAccumulatedTangents"), sizeof(int32), IntermediateAccumulatedTangentsSize * FGPUSkinCache::IntermediateAccumBufferNumInts, PF_R32_SINT, BUF_UnorderedAccess);
-		IntermediateAccumulatedTangents.Buffer.Buffer->SetOwnerName(OwnerName);
-		IntermediateAccumulatedTangents.AccessState = ERHIAccess::Unknown;
-
-		// The UAV must be zero-filled. We leave it zeroed after each round (see RecomputeTangentsPerVertexPass.usf), so this is only needed on when the buffer is first created.
-		RHICmdList.ClearUAVUint(IntermediateAccumulatedTangents.Buffer.UAV, FUintVector4(0, 0, 0, 0));
-	}
-}
-
-FGPUSkinCache::FRWBuffersAllocation::~FRWBuffersAllocation()
-{
-	for (int32 Index = 0; Index < NUM_BUFFERS; ++Index)
-	{
-		PositionBuffers[Index].Release();
-	}
-
-	Tangents.Release();
-	IntermediateTangents.Release();
-	IntermediateAccumulatedTangents.Release();
-}
-
-FGPUSkinCache::FSkinCacheRWBuffer* FGPUSkinCache::FRWBuffersAllocation::GetTangentBuffer()
-{
-	return &Tangents;
-}
-
-FGPUSkinCache::FSkinCacheRWBuffer* FGPUSkinCache::FRWBuffersAllocation::GetIntermediateTangentBuffer()
-{
-	return EnumHasAllFlags(BufferBits, EGPUSkinCacheBufferBits::IntermediateTangents) ? &IntermediateTangents : nullptr;
-}
-
-bool FGPUSkinCache::FRWBuffersAllocation::HasPreviousBuffer() const
-{
-	return EnumHasAllFlags(BufferBits, EGPUSkinCacheBufferBits::PositionPrevious);
-}
-
-uint64 FGPUSkinCache::FRWBuffersAllocation::CalculateRequiredMemory(EGPUSkinCacheBufferBits BufferBits, uint32 NumVertices, uint32 IntermediateAccumulatedTangentsSize)
-{
-	const bool bIntermediateTangents = EnumHasAnyFlags(BufferBits, EGPUSkinCacheBufferBits::IntermediateTangents);
-	const bool bPositionPrevious = EnumHasAnyFlags(BufferBits, EGPUSkinCacheBufferBits::PositionPrevious);
-
-	uint64 PositionBufferSize = PosBufferBytesPerElement * NumVertices * 3 * (bPositionPrevious ? NUM_BUFFERS : 1);
-	uint64 TangentBufferSize =  TangentBufferBytesPerElement * NumVertices * 2;
-	uint64 IntermediateTangentBufferSize = bIntermediateTangents ? TangentBufferBytesPerElement * NumVertices * 2 : 0;
-	uint64 AccumulatedTangentBufferSize = IntermediateAccumulatedTangentsSize * FGPUSkinCache::IntermediateAccumBufferNumInts * sizeof(int32);
-	return TangentBufferSize + IntermediateTangentBufferSize + PositionBufferSize + AccumulatedTangentBufferSize;
-}
-
 #undef IMPLEMENT_SKIN_CACHE_SHADER_CLOTH
 #undef IMPLEMENT_SKIN_CACHE_SHADER_ALL_SKIN_TYPES
 #undef IMPLEMENT_SKIN_CACHE_SHADER
diff --git a/Engine/Source/Runtime/Engine/Private/GPUSkinVertexFactory.cpp b/Engine/Source/Runtime/Engine/Private/GPUSkinVertexFactory.cpp
--- a/Engine/Source/Runtime/Engine/Private/GPUSkinVertexFactory.cpp
+++ b/Engine/Source/Runtime/Engine/Private/GPUSkinVertexFactory.cpp
@@ -1415,6 +1415,11 @@ void FGPUSkinPassthroughVertexFactory::InitRHI(FRHICommandListBase& RHICmdList)
 	LooseParameters.PreSkinnedTangentBuffer = GNullVertexBuffer.VertexBufferSRV;
 	LooseParametersUniformBuffer = TUniformBufferRef<FGPUSkinPassThroughFactoryLooseParameters>::CreateUniformBufferImmediate(LooseParameters, UniformBuffer_MultiFrame);
 
+	if (Data.TangentsSRV)
+	{
+		TangentFormat = Data.TangentsSRV->GetDesc().Common.Format;
+	}
+
 	FLocalVertexFactory::InitRHI(RHICmdList);
 }
 
diff --git a/Engine/Source/Runtime/Engine/Private/SkeletalMeshUpdater.cpp b/Engine/Source/Runtime/Engine/Private/SkeletalMeshUpdater.cpp
--- a/Engine/Source/Runtime/Engine/Private/SkeletalMeshUpdater.cpp
+++ b/Engine/Source/Runtime/Engine/Private/SkeletalMeshUpdater.cpp
@@ -263,8 +263,8 @@ RDG_REGISTER_BLACKBOARD_STRUCT(FSkeletalMeshUpdater::FSubmitTasks);
 
 struct FSkeletalMeshUpdater::FTaskData
 {
-	FTaskData(FRDGBuilder& GraphBuilder)
-		: GPUSkinCachePipeline(FGPUSkinCache::GetDispatchPipeline(GraphBuilder))
+	FTaskData(FRDGBuilder& GraphBuilder, ERHIPipeline InGPUSkinCachePipeline)
+		: GPUSkinCachePipeline(InGPUSkinCachePipeline)
 		, bAsyncCommandList(GraphBuilder.IsParallelSetupEnabled())
 	{
 		if (bAsyncCommandList)
@@ -352,7 +352,7 @@ struct FSkeletalMeshUpdater::FTaskData
 	} Tasks;
 };
 
-FSkeletalMeshUpdater::FSubmitTasks FSkeletalMeshUpdater::Submit(FRDGBuilder& GraphBuilder)
+FSkeletalMeshUpdater::FSubmitTasks FSkeletalMeshUpdater::Submit(FRDGBuilder& GraphBuilder, ERHIPipeline GPUSkinCachePipeline)
 {
 	bool bAllOpsEmpty = true;
 
@@ -373,7 +373,7 @@ FSkeletalMeshUpdater::FSubmitTasks FSkeletalMeshUpdater::Submit(FRDGBuilder& Gra
 	checkf(!bSubmitting, TEXT("Submit was called twice on the same RDG builder. This is not allowed."));
 	bSubmitting = true;
 
-	FTaskData& TaskData = *GraphBuilder.AllocObject<FTaskData>(GraphBuilder);
+	FTaskData& TaskData = *GraphBuilder.AllocObject<FTaskData>(GraphBuilder, GPUSkinCachePipeline);
 
 	UE::Tasks::FTask ForegroundTask = GraphBuilder.AddSetupTask([this, &TaskData]
 	{
diff --git a/Engine/Source/Runtime/Engine/Public/GPUSkinCache.h b/Engine/Source/Runtime/Engine/Public/GPUSkinCache.h
--- a/Engine/Source/Runtime/Engine/Public/GPUSkinCache.h
+++ b/Engine/Source/Runtime/Engine/Public/GPUSkinCache.h
@@ -166,7 +166,7 @@ public:
 	
 	ENGINE_API static ERHIPipeline GetDispatchPipeline(FRDGBuilder& GraphBuilder);
 
-	ENGINE_API UE::Tasks::FTask Dispatch(FRDGBuilder& GraphBuilder, const UE::Tasks::FTask& PrerequisitesTask);
+	ENGINE_API UE::Tasks::FTask Dispatch(FRDGBuilder& GraphBuilder, const UE::Tasks::FTask& PrerequisitesTask, ERHIPipeline Pipeline);
 
 	ENGINE_API void AddAsyncComputeSignal(FRDGBuilder& GraphBuilder);
 	ENGINE_API void AddAsyncComputeWait(FRDGBuilder& GraphBuilder);
@@ -189,7 +189,6 @@ public:
 
 private:
 	friend FGPUSkinCacheEntry;
-	struct FRWBufferTracker;
 
 	struct FDispatchEntry
 	{
@@ -208,163 +207,10 @@ private:
 		NUM_BUFFERS = 2,
 	};
 
-	struct FSkinCacheRWBuffer
-	{
-		FRWBuffer	Buffer;
-		ERHIAccess	AccessState = ERHIAccess::Unknown;	// Keep track of current access state
-		mutable uint32	LastTransitionFence = 0;
-
-		void Release()
-		{
-			Buffer.Release();
-			AccessState = ERHIAccess::Unknown;
-		}
-
-		bool UpdateFence(uint32 NextTransitionFence)
-		{
-			const bool bUpdateRequired = LastTransitionFence != NextTransitionFence;
-			LastTransitionFence = NextTransitionFence;
-			return bUpdateRequired;
-		}
-
-		// Update the access state and return transition info
-		FRHITransitionInfo UpdateAccessState(ERHIAccess NewState)
-		{
-			ERHIAccess OldState = AccessState;
-			AccessState = NewState;
-			return FRHITransitionInfo(Buffer.UAV.GetReference(), OldState, AccessState);
-		}
-	};
-
-	struct FRWBuffersAllocation
-	{
-		friend FRWBufferTracker;
-
-		static uint64 CalculateRequiredMemory(EGPUSkinCacheBufferBits BufferBits, uint32 NumVertices, uint32 IntermediateAccumulatedTangentsSize);
-
-		FRWBuffersAllocation(FRHICommandList& RHICmdList, EGPUSkinCacheBufferBits BufferBits, uint32 NumVertices, uint32 IntermediateAccumulatedTangentsSize, const FName& OwnerName);
-		~FRWBuffersAllocation();
-
-		uint64 GetNumBytes() const
-		{
-			return CalculateRequiredMemory(BufferBits, NumVertices, IntermediateAccumulatedTangentsSize);
-		}
-
-		FSkinCacheRWBuffer* GetTangentBuffer();
-
-		FSkinCacheRWBuffer* GetIntermediateTangentBuffer();
-
-		FSkinCacheRWBuffer* GetIntermediateAccumulatedTangentBuffer()
-		{
-			return IntermediateAccumulatedTangentsSize > 0 ? &IntermediateAccumulatedTangents : nullptr;
-		}
-
-		bool HasPreviousBuffer() const;
-
-		FSkinCacheRWBuffer& GetPositionBuffer()
-		{
-			return PositionBuffers[0];
-		}
-
-		void RemoveAllFromTransitionArray(TSet<FSkinCacheRWBuffer*>& BuffersToTransition);
-
-	private:
-		// Output of the GPU skinning (ie Pos, Normals)
-		FSkinCacheRWBuffer PositionBuffers[NUM_BUFFERS];
-
-		FSkinCacheRWBuffer Tangents;
-		FSkinCacheRWBuffer IntermediateTangents;
-		FSkinCacheRWBuffer IntermediateAccumulatedTangents;	// Intermediate buffer used to accumulate results of triangle pass to be passed onto vertex pass
-
-		const uint32 NumVertices;
-		const uint32 IntermediateAccumulatedTangentsSize;
-		EGPUSkinCacheBufferBits BufferBits;
-
-		static const uint32 PosBufferBytesPerElement = 4;
-		static const uint32 TangentBufferBytesPerElement = 8;
-	};
-
-	struct FRWBufferTracker
-	{
-		FRWBuffersAllocation* Allocation;
-
-		FRWBufferTracker()
-			: Allocation(nullptr)
-		{
-			Reset();
-		}
-
-		void Reset()
-		{
-			for (int32 Index = 0; Index < NUM_BUFFERS; ++Index)
-			{
-				Revisions[Index] = 0;
-				BoneBuffers[Index] = nullptr;
-			}
-		}
-
-		inline uint32 GetNumBytes() const
-		{
-			return IntCastChecked<uint32>(Allocation->GetNumBytes());
-		}
-
-		FSkinCacheRWBuffer* Find(const FVertexBufferAndSRV& BoneBuffer, uint32 Revision)
-		{
-			for (int32 Index = 0; Index < NUM_BUFFERS; ++Index)
-			{
-				if (Revisions[Index] == Revision && BoneBuffers[Index] == &BoneBuffer)
-				{
-					return &Allocation->PositionBuffers[Index];
-				}
-			}
-
-			return nullptr;
-		}
-
-		FSkinCacheRWBuffer* GetTangentBuffer()
-		{
-			return Allocation ? Allocation->GetTangentBuffer() : nullptr;
-		}
-
-		FSkinCacheRWBuffer* GetIntermediateTangentBuffer()
-		{
-			return Allocation ? Allocation->GetIntermediateTangentBuffer() : nullptr;
-		}
-
-		FSkinCacheRWBuffer* GetIntermediateAccumulatedTangentBuffer()
-		{
-			return Allocation ? Allocation->GetIntermediateAccumulatedTangentBuffer() : nullptr;
-		}
-
-		// Allocates an element that's not the "Used" element passed in (or if Used is NULL, allocates any element).
-		FSkinCacheRWBuffer* AllocateUnused(const FVertexBufferAndSRV& BoneBuffer, uint32 Revision, const FSkinCacheRWBuffer* Used)
-		{
-			int32 UnusedIndex = Used == &Allocation->PositionBuffers[0] ? 1 : 0;
-			Revisions[UnusedIndex] = Revision;
-			BoneBuffers[UnusedIndex] = &BoneBuffer;
-
-			return &Allocation->PositionBuffers[UnusedIndex];
-		}
-
-		// On recreate of the render state where the GPU skin cache entry is preserved, the bone buffer will have been reallocated,
-		// even though the transforms didn't change.  We need to force the Find() call above to treat the data as up-to-date, which
-		// can be accomplished by updating the BoneBuffer pointer for the previous Revision, so it matches again.
-		void UpdatePreviousBoneBuffer(const FVertexBufferAndSRV& PreviousBoneBuffer, uint32 PreviousRevision)
-		{
-			for (int32 Index = 0; Index < NUM_BUFFERS; ++Index)
-			{
-				if (Revisions[Index] == PreviousRevision)
-				{
-					BoneBuffers[Index] = &PreviousBoneBuffer;
-					break;
-				}
-			}
-		}
-
-	private:
-		uint32 Revisions[NUM_BUFFERS];
-		const FVertexBufferAndSRV* BoneBuffers[NUM_BUFFERS];
-	};
+	struct FSkinCacheRWBuffer;
+	struct FRWBuffersAllocationInitializer;
+	struct FRWBuffersAllocation;
+	struct FRWBufferTracker;
 
 	struct FTaskData;
 	void DispatchPassSetup(FTaskData& TaskData);
@@ -378,7 +224,6 @@ private:
 	TSet<FGPUSkinCacheEntry*> PendingProcessRTGeometryEntries;
 	TArray<FDispatchEntry> BatchDispatches;
 
-	FRWBuffersAllocation* TryAllocBuffer(FRHICommandList& RHICmdList, EGPUSkinCacheBufferBits BufferBits, uint32 NumVertices, uint32 InIntermediateAccumulatedTangentsSize, const FName& OwnerName);
 	void DispatchUpdateSkinTangentsVertexPass(FRHICommandList& RHICmdList, FGPUSkinCacheEntry* Entry, int32 SectionIndex, FSkinCacheRWBuffer*& StagingBuffer);
 	void DispatchUpdateSkinTangentsTrianglePass(FRHICommandList& RHICmdList, FGPUSkinCacheEntry* Entry, int32 SectionIndex, FSkinCacheRWBuffer*& StagingBuffer);
 
diff --git a/Engine/Source/Runtime/RenderCore/Private/RenderGraphBuilder.cpp b/Engine/Source/Runtime/RenderCore/Private/RenderGraphBuilder.cpp
--- a/Engine/Source/Runtime/RenderCore/Private/RenderGraphBuilder.cpp
+++ b/Engine/Source/Runtime/RenderCore/Private/RenderGraphBuilder.cpp
@@ -582,8 +582,7 @@ FRDGBuilder::FRDGBuilder(FRHICommandListImmediate& InRHICmdList, FRDGEventName I
 {
 	ProloguePass = SetupEmptyPass(Passes.Allocate<FRDGSentinelPass>(Allocators.Root, RDG_EVENT_NAME("Graph Prologue (Graphics)")));
 
-	// Only enable async compute for large builders that are set to run in parallel.
-	bSupportsAsyncCompute = IsAsyncComputeSupported(ShaderPlatform) && EnumHasAnyFlags(InFlags, ERDGBuilderFlags::ParallelExecute);
+	bSupportsAsyncCompute = IsAsyncComputeSupported(ShaderPlatform);
 	bSupportsRenderPassMerge = IsRenderPassMergeEnabled(ShaderPlatform);
 
 	const bool bParallelExecuteFlag = EnumHasAnyFlags(InFlags, ERDGBuilderFlags::ParallelExecute);
diff --git a/Engine/Source/Runtime/Renderer/Private/RendererScene.cpp b/Engine/Source/Runtime/Renderer/Private/RendererScene.cpp
--- a/Engine/Source/Runtime/Renderer/Private/RendererScene.cpp
+++ b/Engine/Source/Runtime/Renderer/Private/RendererScene.cpp
@@ -5569,15 +5569,17 @@ void FScene::Update(FRDGBuilder& GraphBuilder, const FUpdateParameters& Paramete
 
 	FSkeletalMeshUpdater::FSubmitTasks GPUSkinTasks;
 
+	const ERHIPipeline GPUSkinCachePipeline = GPUSkinCache ? FGPUSkinCache::GetDispatchPipeline(GraphBuilder) : ERHIPipeline::Graphics;
+
 	if (SkeletalMeshUpdater)
 	{
-		GPUSkinTasks = SkeletalMeshUpdater->Submit(GraphBuilder);
+		GPUSkinTasks = SkeletalMeshUpdater->Submit(GraphBuilder, GPUSkinCachePipeline);
 		GPUSkinUpdateTask = GPUSkinTasks.SkinCache;
 	}
 
 	if (GPUSkinCache)
 	{
-		GPUSkinCacheTask = GPUSkinCache->Dispatch(GraphBuilder, GPUSkinTasks.SkinCache);
+		GPUSkinCacheTask = GPUSkinCache->Dispatch(GraphBuilder, GPUSkinTasks.SkinCache, GPUSkinCachePipeline);
 	}
 
 	// Add interface to blackboard to receive scene updates from compute.
