Redo:
Fast-path for editor access of current target platform compressed data

diff --git a/Engine/Source/Runtime/Engine/Classes/Animation/AnimSequence.h b/Engine/Source/Runtime/Engine/Classes/Animation/AnimSequence.h
--- a/Engine/Source/Runtime/Engine/Classes/Animation/AnimSequence.h
+++ b/Engine/Source/Runtime/Engine/Classes/Animation/AnimSequence.h
@@ -806,6 +806,7 @@ protected:
 	ENGINE_API bool PollCacheDerivedData(const FIoHash& KeyHash) const;
 	ENGINE_API void EndCacheDerivedData(const FIoHash& KeyHash);
 
+	mutable FCompressedAnimSequence* CurrentPlatformData = nullptr; 
 	TMap<FIoHash, TUniquePtr<FCompressedAnimSequence>> DataByPlatformKeyHash;
 	TMap<FIoHash, TPimplPtr<UE::Anim::FAnimationSequenceAsyncCacheTask>> CacheTasksByKeyHash;
 	mutable FRWLock CompressedDataLock;
diff --git a/Engine/Source/Runtime/Engine/Private/Animation/AnimSequence.cpp b/Engine/Source/Runtime/Engine/Private/Animation/AnimSequence.cpp
--- a/Engine/Source/Runtime/Engine/Private/Animation/AnimSequence.cpp
+++ b/Engine/Source/Runtime/Engine/Private/Animation/AnimSequence.cpp
@@ -793,11 +793,21 @@ void UAnimSequence::Serialize(FArchive& Ar)
 						{
 							uint32 NumPlatforms = DataByPlatformKeyHash.Num();
 							Ar << NumPlatforms;
+
+							FIoHash CurrentPlatformHash = FIoHash::Zero;
 							for (TPair<FIoHash, TUniquePtr<FCompressedAnimSequence>>& CompressedDataPair : DataByPlatformKeyHash)
 							{
 								Ar << CompressedDataPair.Key;
 								SerializeCompressedData(Ar, *CompressedDataPair.Value.Get());
+
+								if (CurrentPlatformData && CurrentPlatformData == CompressedDataPair.Value.Get())
+								{
+									CurrentPlatformHash = CompressedDataPair.Key;
+								}
 							}
+							
+							// Store off the expected hash for CurrentPlatformData
+							Ar << CurrentPlatformHash;
 						}
 						else
 						{
@@ -810,6 +820,19 @@ void UAnimSequence::Serialize(FArchive& Ar)
 								FCompressedAnimSequence& TargetData = *DataByPlatformKeyHash.Emplace(PlatformHash, MakeUnique<FCompressedAnimSequence>()).Get();
 								SerializeCompressedData(Ar, TargetData);
 							}
+
+							CurrentPlatformData = nullptr;
+							
+							// Try and restore short-circuit ptr for CurrentPlatformData according to its previously serialized hash
+							FIoHash CurrentPlatformHash = FIoHash::Zero;
+							Ar << CurrentPlatformHash;
+							if (CurrentPlatformHash != FIoHash::Zero)
+							{
+								if (TUniquePtr<FCompressedAnimSequence>* Entry = DataByPlatformKeyHash.Find(CurrentPlatformHash))
+								{
+									CurrentPlatformData = Entry->Get();
+								}
+							}	
 						}
 					}
 
@@ -1944,6 +1967,12 @@ const FCompressedAnimSequence& UAnimSequence::GetCompressedData_Internal() const
 	if(!GetOutermost()->HasAnyPackageFlags(PKG_Cooked))
 	{
 		static thread_local FCompressedAnimSequence DefaultCompressedAnimSequence;
+
+		if (CurrentPlatformData)
+		{
+			return *CurrentPlatformData;
+		}
+		
 		const ITargetPlatform* RunningPlatform = GetTargetPlatformManagerRef().GetRunningTargetPlatform();
 
 		if(DataByPlatformKeyHash.Num())
@@ -1994,6 +2023,11 @@ const FCompressedAnimSequence& UAnimSequence::GetPlatformCompressedData(const FA
 		}
 		else
 		{
+			if (CurrentPlatformData)
+			{
+				return *CurrentPlatformData;
+			}
+			
 			const ITargetPlatform* RunningPlatform = GetTargetPlatformManagerRef().GetRunningTargetPlatform();
 			Hash = GetDerivedDataKeyHash(RunningPlatform);
 		}
@@ -2041,6 +2075,11 @@ FCompressedAnimSequence& UAnimSequence::GetPlatformCompressedData(const FAnimExt
 		}
 		else
 		{
+			if (CurrentPlatformData)
+			{
+				return *CurrentPlatformData;
+			}
+			
 			const ITargetPlatform* RunningPlatform = GetTargetPlatformManagerRef().GetRunningTargetPlatform();
 			Hash = GetDerivedDataKeyHash(RunningPlatform);
 		}
@@ -2343,15 +2382,16 @@ void UAnimSequence::SerializeCompressedData(FArchive& Ar, bool bDDCData)
 		const ITargetPlatform* Platform = nullptr;
 #endif // WITH_EDITOR
 		
-		FCompressedAnimSequence& CompressedDataToSerialize = GetPlatformCompressedData(Platform);
 		if (Ar.IsLoading())
 		{
 			UE_COMPRESSED_DATA_WRITE_SCOPE(this);
+			FCompressedAnimSequence& CompressedDataToSerialize = GetPlatformCompressedData(Platform);
 			CompressedDataToSerialize.SerializeCompressedData(Ar, bDDCData, this, this->GetSkeleton(), BoneCompressionSettings, CurveCompressionSettings);
 		}
 		else
 		{
 			UE_COMPRESSED_DATA_READ_SCOPE(this);
+			FCompressedAnimSequence& CompressedDataToSerialize = GetPlatformCompressedData(Platform);
 			check(!Ar.IsCooking() || CompressedDataToSerialize.IsValid(this, true));
 			CompressedDataToSerialize.SerializeCompressedData(Ar, bDDCData, this, this->GetSkeleton(), BoneCompressionSettings, CurveCompressionSettings);
 		}
@@ -5054,6 +5094,12 @@ FIoHash UAnimSequence::BeginCacheDerivedData(const ITargetPlatform* TargetPlatfo
 		TargetData = DataByPlatformKeyHash.Emplace(KeyHash, MakeUnique<FCompressedAnimSequence>()).Get();
 	}
 
+	const ITargetPlatform* RunningPlatform = GetTargetPlatformManagerRef().GetRunningTargetPlatform();	
+	if (TargetPlatform == RunningPlatform)
+	{
+		CurrentPlatformData = nullptr;
+	}
+
 	// Reset the target compressed data, to mark it invalid
 	check(TargetData);
 	TargetData->Reset();
@@ -5196,6 +5242,8 @@ void UAnimSequence::ClearAllCompressionData()
 				It.RemoveCurrent();
 			}
 		}
+		
+		CurrentPlatformData = nullptr;
 	}
 }
 
@@ -5209,6 +5257,12 @@ void UAnimSequence::ClearCompressionData(const FIoHash& InKeyHash)
 		UE_COMPRESSED_DATA_WRITE_SCOPE(this);
 		check(DataByPlatformKeyHash.Contains(InKeyHash));
 		DataByPlatformKeyHash.Remove(InKeyHash);
+		
+		const ITargetPlatform* RunningPlatform = GetTargetPlatformManagerRef().GetRunningTargetPlatform();
+		if (GetDerivedDataKeyHash(RunningPlatform) == InKeyHash)
+		{
+			CurrentPlatformData = nullptr;
+		}
 	}
 }
 
@@ -5339,7 +5393,10 @@ void UAnimSequence::FinishAsyncTasks()
 			const bool bCompressedRunningPlatform = It->Key == CurrentPlatformHash;
 			auto ResetData = [this, TaskData, It]()
 			{
-				TaskData->Reset();
+				if(TaskData != nullptr)
+				{
+					TaskData->Reset();
+				}
 				DataByPlatformKeyHash.Remove(It->Key);
 			};
 
@@ -5349,19 +5406,21 @@ void UAnimSequence::FinishAsyncTasks()
 			}
 			else
 			{
-				if (TaskData->IsValid(this, true))
+				if (TaskData && TaskData->IsValid(this, true))
 				{
-#if WITH_EDITOR
 					//This is only safe during sync anim compression
 					if (GetSkeleton())
 					{
 						SetSkeletonVirtualBoneGuid(GetSkeleton()->GetVirtualBoneGuid());
 					}
-#endif
 					if (bCompressedRunningPlatform)
 					{
 						FAssetRegistryModule& AssetRegistryModule = FModuleManager::LoadModuleChecked<FAssetRegistryModule>(TEXT("AssetRegistry"));
 						AssetRegistryModule.Get().AssetTagsFinalized(*this);
+
+						// Cache compressed data ptr for current platform to by-pass TMap lookup requiring data-hash retrieval/generation and locking
+						ensure(CurrentPlatformData == nullptr);
+						CurrentPlatformData = TaskData;
 					}
 				
 					check(TaskData->IsValid(this, true));
