Register actors in TEDS when they are loaded. This is a follow up to a previous change to unregister actors from TEDS when they are unloaded, to ensure they are re-registered if they are loaded back in without the actor actually being GC'd and reconstructed

diff --git a/Engine/Source/Runtime/Engine/Private/Level.cpp b/Engine/Source/Runtime/Engine/Private/Level.cpp
--- a/Engine/Source/Runtime/Engine/Private/Level.cpp
+++ b/Engine/Source/Runtime/Engine/Private/Level.cpp
@@ -1059,6 +1059,10 @@ void ULevel::AddLoadedActors(const TArray<AActor*>& ActorList, const FTransform*
 		SlowTask.EnterProgressFrame(1);
 	}
 
+	using namespace UE::Editor::DataStorage;
+	IEditorDataStorageCompatibilityProvider* TedsCompat =
+					GetMutableDataStorageFeature<IEditorDataStorageCompatibilityProvider>(CompatibilityFeatureName);
+
 	// Finalize actors
 	for (AActor* Actor : ActorsQueue)
 	{
@@ -1079,6 +1083,14 @@ void ULevel::AddLoadedActors(const TArray<AActor*>& ActorList, const FTransform*
 			Actor->OnLoadedActorAddedToLevel();
 
 			OnLoadedActorAddedToLevelEvent.Broadcast(*Actor);
+
+			
+			// NOTE: We have to add the actor to TEDS explicitly here because we remove it on unloading, and an unload+load flow does not always go
+			// through the regular actor creation process which registers it in TEDS via UEngineElementsLibrary::CreateEditorActorElement
+			if (TedsCompat)
+			{
+				TedsCompat->AddCompatibleObject(Actor);
+			}
 		}
 
 		SlowTask.EnterProgressFrame(1);
@@ -1170,6 +1182,10 @@ void ULevel::RemoveLoadedActors(const TArray<AActor*>& ActorList, const FTransfo
 	SlowTask.MakeDialogDelayed(1.0f);
 
 	OnLoadedActorRemovedFromLevelPreEvent.Broadcast(ActorsQueue);
+	
+	using namespace UE::Editor::DataStorage;
+	IEditorDataStorageCompatibilityProvider* TedsCompat =
+					GetMutableDataStorageFeature<IEditorDataStorageCompatibilityProvider>(CompatibilityFeatureName);
 
 	for (AActor* Actor : ActorsQueue)
 	{
@@ -1185,8 +1201,7 @@ void ULevel::RemoveLoadedActors(const TArray<AActor*>& ActorList, const FTransfo
 			using namespace UE::Editor::DataStorage;
 			// NOTE: we do have to remove the actor from TEDS explicitly here as an actor unloading doesn't always trigger GC immediately,
 			// but for our purposes it is considered dead already
-			if (IEditorDataStorageCompatibilityProvider* TedsCompat =
-					GetMutableDataStorageFeature<IEditorDataStorageCompatibilityProvider>(CompatibilityFeatureName))
+			if (TedsCompat)
 			{
 				TedsCompat->RemoveCompatibleObject(Actor);
 			}
