Submitting on behalf of Jon.Cain:

Initial MPC setup for epic_internal Verse. En lieu of a singleton, implemented an "Enabled" variable with Verse setter and getter that tracks which Verse object is currently "live" using a static map. Also refactored the material replication foundation code out to a separate class to make it reusable in for MPCs.

##okforversepublic

diff --git a/Engine/Source/Runtime/Engine/Private/Materials/ParameterCollection.cpp b/Engine/Source/Runtime/Engine/Private/Materials/ParameterCollection.cpp
--- a/Engine/Source/Runtime/Engine/Private/Materials/ParameterCollection.cpp
+++ b/Engine/Source/Runtime/Engine/Private/Materials/ParameterCollection.cpp
@@ -794,6 +794,14 @@ void UMaterialParameterCollectionInstance::DeferredUpdateRenderState(bool bRecre
 	bNeedsRenderStateUpdate = false;
 }
 
+void UMaterialParameterCollectionInstance::ForceReturnToDefaultValues()
+{
+	ScalarParameterValues.Empty();
+	VectorParameterValues.Empty();
+	bool bNeedsUniformBuffer = !(Resource && Resource->GetUniformBuffer() && Resource->GetUniformBuffer()->IsValid());
+	UpdateRenderState(bNeedsUniformBuffer);
+}
+
 void UMaterialParameterCollectionInstance::GetParameterData(TArray<FVector4f>& ParameterData) const
 {
 	// The memory layout created here must match the index assignment in UMaterialParameterCollection::GetParameterIndex
diff --git a/Engine/Source/Runtime/Engine/Public/Materials/MaterialParameterCollection.h b/Engine/Source/Runtime/Engine/Public/Materials/MaterialParameterCollection.h
--- a/Engine/Source/Runtime/Engine/Public/Materials/MaterialParameterCollection.h
+++ b/Engine/Source/Runtime/Engine/Public/Materials/MaterialParameterCollection.h
@@ -111,11 +111,11 @@ class UMaterialParameterCollection : public UObject
 
 	/** Returns an array of the names of all the scalar parameters in this Material Parameter Collection **/
 	UFUNCTION(BlueprintCallable, Category="Rendering|Material", meta=(Keywords="GetScalarParameterNames"))
-	TArray<FName> GetScalarParameterNames() const;
+	ENGINE_API TArray<FName> GetScalarParameterNames() const;
 
 	/** Returns an array of the names of all the vector parameters in this Material Parameter Collection **/
 	UFUNCTION(BlueprintCallable, Category="Rendering|Material", meta=(Keywords="GetVectorParameterNames"))
-	TArray<FName> GetVectorParameterNames() const;
+	ENGINE_API TArray<FName> GetVectorParameterNames() const;
 
 	/** Gets the default value of a scalar parameter from a material collection.
 	 * @param ParameterName - The name of the value to get the value of
@@ -123,7 +123,7 @@ class UMaterialParameterCollection : public UObject
 	 * @returns the value of the parameter
 	 **/
 	UFUNCTION(BlueprintCallable, Category="Rendering|Material", meta=(Keywords="GetFloatParameterDefaultValue"))
-	float GetScalarParameterDefaultValue(FName ParameterName, bool& bParameterFound) const;
+	ENGINE_API float GetScalarParameterDefaultValue(FName ParameterName, bool& bParameterFound) const;
 
 	/** Gets the default value of a scalar parameter from a material collection.
 	 * @param ParameterName - The name of the value to get the value of
@@ -131,7 +131,7 @@ class UMaterialParameterCollection : public UObject
 	 * @returns the value of the parameter
 	 **/
 	UFUNCTION(BlueprintCallable, Category="Rendering|Material", meta=(Keywords="GetFloatParameterDefaultValue"))
-	FLinearColor GetVectorParameterDefaultValue(FName ParameterName, bool& bParameterFound) const;
+	ENGINE_API FLinearColor GetVectorParameterDefaultValue(FName ParameterName, bool& bParameterFound) const;
 	
 	//~ Begin UObject Interface.
 #if WITH_EDITOR
diff --git a/Engine/Source/Runtime/Engine/Public/Materials/MaterialParameterCollectionInstance.h b/Engine/Source/Runtime/Engine/Public/Materials/MaterialParameterCollectionInstance.h
--- a/Engine/Source/Runtime/Engine/Public/Materials/MaterialParameterCollectionInstance.h
+++ b/Engine/Source/Runtime/Engine/Public/Materials/MaterialParameterCollectionInstance.h
@@ -77,6 +77,8 @@ class UMaterialParameterCollectionInstance : public UObject
 
 	ENGINE_API void DeferredUpdateRenderState(bool bRecreateUniformBuffer);
 
+	ENGINE_API void ForceReturnToDefaultValues();
+
 	/** Tracks whether this instance has ever issued a missing parameter warning, to reduce log spam. */
 	bool bLoggedMissingParameterWarning;
 
