Fix Duplicate localization keys in UMaterialExpressionConvert::RegisterAdditionalMenuActions

diff --git a/Engine/Source/Runtime/Engine/Private/Materials/MaterialExpressions.cpp b/Engine/Source/Runtime/Engine/Private/Materials/MaterialExpressions.cpp
--- a/Engine/Source/Runtime/Engine/Private/Materials/MaterialExpressions.cpp
+++ b/Engine/Source/Runtime/Engine/Private/Materials/MaterialExpressions.cpp
@@ -23632,8 +23632,8 @@ void UMaterialExpressionConvert::RegisterAdditionalMenuActions(FGraphActionMenuB
 	auto CreateBreakNode = [this, &ActionMenuBuilder, &CategoryName](EMaterialExpressionConvertType ConvertType)
 	{
 		const int32 NumComponents = MaterialExpressionConvertType::GetComponentCount(ConvertType);
-		const FText ActionName = FText::Format(LOCTEXT("Convert_MakeNode_Name", "Break Float{0}"), NumComponents);
-		const FText ActionTooltip = FText::Format(LOCTEXT("Convert_MakeNode_Tooltip", "Adds a Break Float{0} Node"), NumComponents);
+		const FText ActionName = FText::Format(LOCTEXT("Convert_BreakNode_Name", "Break Float{0}"), NumComponents);
+		const FText ActionTooltip = FText::Format(LOCTEXT("Convert_BreakNode_Tooltip", "Adds a Break Float{0} Node"), NumComponents);
 
 		TSharedPtr<FMaterialGraphSchemaAction_NewNode> Action_BreakNode(new FMaterialGraphSchemaAction_NewNode(
 			CategoryName,
