Replay
* Changed TActorIterator to TActorIterator<PlayerStart> to reduce iteration cost

PR #12614: Improve FReplayHelper performance
#rnx

diff --git a/Engine/Source/Runtime/Engine/Private/ReplayHelper.cpp b/Engine/Source/Runtime/Engine/Private/ReplayHelper.cpp
--- a/Engine/Source/Runtime/Engine/Private/ReplayHelper.cpp
+++ b/Engine/Source/Runtime/Engine/Private/ReplayHelper.cpp
@@ -389,13 +389,10 @@ APlayerController* FReplayHelper::CreateSpectatorController(UNetConnection* Conn
 		NewDemoController->PlayerState->SetIsOnlyASpectator(true);
 	}
 
-	for (FActorIterator It(World.Get()); It; ++It)
+	for (TActorIterator<APlayerStart> It(World.Get()); It; ++It)
 	{
-		if (It->IsA(APlayerStart::StaticClass()))
-		{
-			NewDemoController->SetInitialLocationAndRotation(It->GetActorLocation(), It->GetActorRotation());
-			break;
-		}
+		NewDemoController->SetInitialLocationAndRotation(It->GetActorLocation(), It->GetActorRotation());
+		break;
 	}
 
 	NewDemoController->SetReplicates(true);
