New material translator changes:
- Added FEmitter::DefaultInputToFloat4() API.
- Fixed UMaterialExpressionConvert::Build() implementation to emit constant if pin is disconnected (instead of zero).

diff --git a/Engine/Source/Runtime/Engine/Private/Materials/MaterialExpressionsIR.cpp b/Engine/Source/Runtime/Engine/Private/Materials/MaterialExpressionsIR.cpp
--- a/Engine/Source/Runtime/Engine/Private/Materials/MaterialExpressionsIR.cpp
+++ b/Engine/Source/Runtime/Engine/Private/Materials/MaterialExpressionsIR.cpp
@@ -1036,8 +1036,7 @@ void UMaterialExpressionConvert::Build(MIR::FEmitter& Emitter)
 			if (!InputValues[InputIndex])
 			{
 				// Read the input's value (or read float zero if disconnected).
-				FExpressionInput* Input = &ConvertInputs[InputIndex].ExpressionInput;
-				InputValues[InputIndex] = Emitter.DefaultInputToFloatScalarZero(Input).Input(Input);
+				InputValues[InputIndex] = Emitter.DefaultInputToFloat4(&ConvertInput.ExpressionInput, ConvertInput.DefaultValue).Input(&ConvertInput.ExpressionInput);
 
 				// Expect type to be primitive.
 				MIR::FPrimitiveTypePtr InputPrimitiveType = InputValues[InputIndex]->Type->AsPrimitive();
diff --git a/Engine/Source/Runtime/Engine/Private/Materials/MaterialIREmitter.cpp b/Engine/Source/Runtime/Engine/Private/Materials/MaterialIREmitter.cpp
--- a/Engine/Source/Runtime/Engine/Private/Materials/MaterialIREmitter.cpp
+++ b/Engine/Source/Runtime/Engine/Private/Materials/MaterialIREmitter.cpp
@@ -200,6 +200,16 @@ FEmitter& FEmitter::DefaultInputToFloatScalar(const FExpressionInput* InInput, T
 	return *this;
 }
 
+FEmitter& FEmitter::DefaultInputToFloat4(const FExpressionInput* InInput, FVector4f Value)
+{
+	if (!InInput->IsConnected())
+	{
+		Internal::SetInputValue(BuilderImpl, InInput, EmitConstantFloat4(Value));
+	}
+
+	return *this;
+}
+
 FValue* FEmitter::TryInput(const FExpressionInput* InInput)
 {
 	FValue* Value = MIR::Internal::GetInputValue(BuilderImpl, InInput);
diff --git a/Engine/Source/Runtime/Engine/Public/Materials/MaterialIREmitter.h b/Engine/Source/Runtime/Engine/Public/Materials/MaterialIREmitter.h
--- a/Engine/Source/Runtime/Engine/Public/Materials/MaterialIREmitter.h
+++ b/Engine/Source/Runtime/Engine/Public/Materials/MaterialIREmitter.h
@@ -52,6 +52,9 @@ public:
 
 	//
 	FEmitter& DefaultInputToFloatScalar(const FExpressionInput* Input, TFloat Value);
+	
+	//
+	FEmitter& DefaultInputToFloat4(const FExpressionInput* Input, FVector4f Value);
 
 	//
 
