[OptimusDeformer] added validation for optimus deformer to detect when a skinnedmesh does not have half edge buffers

diff --git a/Engine/Plugins/Animation/DeformerGraph/Source/OptimusCore/Private/DataInterfaces/OptimusDataInterfaceHalfEdge.cpp b/Engine/Plugins/Animation/DeformerGraph/Source/OptimusCore/Private/DataInterfaces/OptimusDataInterfaceHalfEdge.cpp
--- a/Engine/Plugins/Animation/DeformerGraph/Source/OptimusCore/Private/DataInterfaces/OptimusDataInterfaceHalfEdge.cpp
+++ b/Engine/Plugins/Animation/DeformerGraph/Source/OptimusCore/Private/DataInterfaces/OptimusDataInterfaceHalfEdge.cpp
@@ -6,6 +6,8 @@
 #include "ComponentSources/OptimusSkinnedMeshComponentSource.h"
 #include "ComputeFramework/ShaderParamTypeDefinition.h"
 #include "OptimusDataDomain.h"
+#include "OptimusSettings.h"
+#include "OptimusCoreModule.h"
 #include "RenderGraphBuilder.h"
 #include "Rendering/SkeletalMeshLODRenderData.h"
 #include "Rendering/SkeletalMeshRenderData.h"
@@ -188,9 +190,59 @@ bool UOptimusHalfEdgeDataProvider::IsReadyForFinishDestroy()
 
 FComputeDataProviderRenderProxy* UOptimusHalfEdgeDataProvider::GetRenderProxy()
 {
+#if WITH_EDITOR
+	if (Optimus::IsAssetValidationEnabled())
+	{
+		FFunctionGraphTask::CreateAndDispatchWhenReady([WeakThis = TWeakObjectPtr(this)]()
+		{
+			if (UOptimusHalfEdgeDataProvider* DataProvider = Cast<UOptimusHalfEdgeDataProvider>(WeakThis))
+			{
+				DataProvider->ValidateSkinnedAsset();
+			}
+		}, TStatId(), NULL, ENamedThreads::GameThread);
+	}
+#endif
+		
 	return new FOptimusHalfEdgeDataProviderProxy(SkinnedMesh, OnDemandHalfEdgeBuffers);
 }
 
+void UOptimusHalfEdgeDataProvider::SetOptimusDeformerAssetPath(const FTopLevelAssetPath& InPath)
+{
+#if WITH_EDITOR
+	DeformerAssetPath = InPath;
+#endif
+}
+
+#if WITH_EDITOR
+void UOptimusHalfEdgeDataProvider::ValidateSkinnedAsset()
+{
+	check(IsInGameThread());
+	if (USkinnedAsset* Asset = SkinnedMesh->GetSkinnedAsset())
+	{
+		if (!ValidatedAssets.Contains(Asset))
+		{
+			ValidatedAssets.Add(Asset);
+					
+			if (FSkeletalMeshObject* MeshObject = SkinnedMesh->MeshObject)
+			{
+				if (!Asset->HasHalfEdgeBuffer(MeshObject->GetLOD()))
+				{
+					UE_LOG(LogOptimusCore, Warning,
+						   TEXT("Please add Deformer Graph %s to TargetMeshDeformers of Skeletal Mesh Asset %s to "
+							   "make sure extra data required by the Deformer is cooked for non-editor builds"
+							   "This warning can be turned off with %s 0"
+						   ),
+						   *DeformerAssetPath.GetPackageName().ToString(),
+						   *Asset->GetPackage()->GetName(),
+						   *Optimus::GetEnableAssetValidationCVarName().ToString()
+						   );
+				}
+			}
+		}
+	}		
+}
+
+#endif
 
 FOptimusHalfEdgeDataProviderProxy::FOptimusHalfEdgeDataProviderProxy(
 	USkinnedMeshComponent* InSkinnedMeshComponent, 
diff --git a/Engine/Plugins/Animation/DeformerGraph/Source/OptimusCore/Private/DataInterfaces/OptimusDataInterfaceHalfEdge.h b/Engine/Plugins/Animation/DeformerGraph/Source/OptimusCore/Private/DataInterfaces/OptimusDataInterfaceHalfEdge.h
--- a/Engine/Plugins/Animation/DeformerGraph/Source/OptimusCore/Private/DataInterfaces/OptimusDataInterfaceHalfEdge.h
+++ b/Engine/Plugins/Animation/DeformerGraph/Source/OptimusCore/Private/DataInterfaces/OptimusDataInterfaceHalfEdge.h
@@ -2,6 +2,7 @@
 
 #pragma once
 
+#include "IOptimusDeformerAssetPathAccessor.h"
 #include "RenderGraphDefinitions.h"
 #include "OptimusComputeDataInterface.h"
 #include "ComputeFramework/ComputeDataProvider.h"
@@ -14,6 +15,7 @@ class FRDGBuffer;
 class FRDGBufferSRV;
 class FSkeletalMeshObject;
 class USkinnedMeshComponent;
+class USkinnedAsset;
 
 
 /** Compute Framework Data Interface for reading mesh half edge data. */
@@ -49,7 +51,10 @@ private:
 
 /** Compute Framework Data Provider for reading skeletal mesh. */
 UCLASS(BlueprintType, editinlinenew, Category = ComputeFramework)
-class UOptimusHalfEdgeDataProvider : public UComputeDataProvider
+class UOptimusHalfEdgeDataProvider :
+	public UComputeDataProvider,
+	public IOptimusDeformerAssetPathAccessor
+	
 {
 	GENERATED_BODY()
 
@@ -66,8 +71,18 @@ public:
 	FComputeDataProviderRenderProxy* GetRenderProxy() override;
 	//~ End UComputeDataProvider Interface
 
+	// IOptimusDeformerAssetPathAccessor
+	void SetOptimusDeformerAssetPath(const FTopLevelAssetPath& InPath) override;
+
 	TArray<FHalfEdgeBuffers> OnDemandHalfEdgeBuffers;
 	FRenderCommandFence DestroyFence;
+	
+protected:
+#if WITH_EDITOR
+	void ValidateSkinnedAsset();
+	FTopLevelAssetPath DeformerAssetPath;
+	TSet<TWeakObjectPtr<USkinnedAsset>> ValidatedAssets;
+#endif
 };
 
 class FOptimusHalfEdgeDataProviderProxy : public FComputeDataProviderRenderProxy
diff --git a/Engine/Plugins/Animation/DeformerGraph/Source/OptimusCore/Private/OptimusComputeGraph.cpp b/Engine/Plugins/Animation/DeformerGraph/Source/OptimusCore/Private/OptimusComputeGraph.cpp
--- a/Engine/Plugins/Animation/DeformerGraph/Source/OptimusCore/Private/OptimusComputeGraph.cpp
+++ b/Engine/Plugins/Animation/DeformerGraph/Source/OptimusCore/Private/OptimusComputeGraph.cpp
@@ -17,6 +17,7 @@
 #include "OptimusObjectVersion.h"
 #include "Animation/MeshDeformerInstance.h"
 #include "DataInterfaces/OptimusDataInterfaceGraph.h"
+#include "DataInterfaces/OptimusDataInterfaceHalfEdge.h"
 #include "Misc/UObjectToken.h"
 
 #include UE_INLINE_GENERATED_CPP_BY_NAME(OptimusComputeGraph)
@@ -167,4 +168,17 @@ UOptimusGraphDataInterface* UOptimusComputeGraph::GetGraphDataInterfaceForPostLo
 	return nullptr;
 }
 
+bool UOptimusComputeGraph::HasHalfEdgeDataInterface() const
+{
+	for (UComputeDataInterface* DataInterface : DataInterfaces)
+	{
+		if (DataInterface->IsA<UOptimusHalfEdgeDataInterface>())
+		{
+			return true;
+		}
+	}
+	
+	return false;
+}
+
 #undef LOCTEXT_NAMESPACE
diff --git a/Engine/Plugins/Animation/DeformerGraph/Source/OptimusCore/Private/OptimusComputeGraph.h b/Engine/Plugins/Animation/DeformerGraph/Source/OptimusCore/Private/OptimusComputeGraph.h
--- a/Engine/Plugins/Animation/DeformerGraph/Source/OptimusCore/Private/OptimusComputeGraph.h
+++ b/Engine/Plugins/Animation/DeformerGraph/Source/OptimusCore/Private/OptimusComputeGraph.h
@@ -33,6 +33,7 @@ protected:
 	TArray<TWeakObjectPtr<const UOptimusNode>> KernelToNode;
 
 	UOptimusGraphDataInterface* GetGraphDataInterfaceForPostLoadFixUp();
+	bool HasHalfEdgeDataInterface() const;
 	
 	friend class UOptimusDeformer;
 };
diff --git a/Engine/Plugins/Animation/DeformerGraph/Source/OptimusCore/Private/OptimusDeformer.cpp b/Engine/Plugins/Animation/DeformerGraph/Source/OptimusCore/Private/OptimusDeformer.cpp
--- a/Engine/Plugins/Animation/DeformerGraph/Source/OptimusCore/Private/OptimusDeformer.cpp
+++ b/Engine/Plugins/Animation/DeformerGraph/Source/OptimusCore/Private/OptimusDeformer.cpp
@@ -3830,6 +3830,11 @@ void UOptimusDeformer::GetAssetRegistryTags(FAssetRegistryTagsContext Context) c
 
 		Context.AddTag(FAssetRegistryTag(FunctionReferencesAssetTagName, *FunctionReferenceString, FAssetRegistryTag::TT_Hidden));	
 	}
+
+	{
+		const FName& TagName = SkeletalMeshHalfEdgeBufferAccessor::GetHalfEdgeRequirementAssetTagName(); 
+		Context.AddTag(FAssetRegistryTag(TagName, IsSkeletalMeshHalfEdgeBufferRequired() ? TEXT("True") : TEXT("False"), FAssetRegistryTag::TT_Hidden));
+	}
 }
 
 #if WITH_EDITOR
@@ -3902,6 +3907,19 @@ UOptimusDeformerInstance* UOptimusDeformer::CreateOptimusInstance(UMeshComponent
 	return Instance;
 }
 
+bool UOptimusDeformer::IsSkeletalMeshHalfEdgeBufferRequired() const
+{
+	for (const FOptimusComputeGraphInfo& Info : ComputeGraphs)
+	{
+		if (Info.ComputeGraph->HasHalfEdgeDataInterface())
+		{
+			return true;
+		}
+	}
+
+	return false;
+}
+
 
 void UOptimusDeformer::SetPreviewMesh(USkeletalMesh* PreviewMesh, bool bMarkAsDirty)
 {
diff --git a/Engine/Plugins/Animation/DeformerGraph/Source/OptimusCore/Private/OptimusDeformerInstance.cpp b/Engine/Plugins/Animation/DeformerGraph/Source/OptimusCore/Private/OptimusDeformerInstance.cpp
--- a/Engine/Plugins/Animation/DeformerGraph/Source/OptimusCore/Private/OptimusDeformerInstance.cpp
+++ b/Engine/Plugins/Animation/DeformerGraph/Source/OptimusCore/Private/OptimusDeformerInstance.cpp
@@ -5,6 +5,7 @@
 #include "Components/MeshComponent.h"
 #include "ComputeFramework/ComputeFramework.h"
 #include "ComputeWorkerInterface.h"
+#include "IOptimusDeformerAssetPathAccessor.h"
 #include "IOptimusDeformerGeometryReadbackProvider.h"
 #include "IOptimusPersistentBufferProvider.h"
 #include "IOptimusValueProvider.h"
@@ -433,6 +434,12 @@ void UOptimusDeformerInstance::SetupFromDeformer(UOptimusDeformer* InDeformer)
 			{
 				InstanceAccessor->SetDeformerInstance(this);
 			}
+
+			if (IOptimusDeformerAssetPathAccessor* AssetPathAccessor = Cast<IOptimusDeformerAssetPathAccessor>(DataProvider))
+			{
+				AssetPathAccessor->SetOptimusDeformerAssetPath(FTopLevelAssetPath(InDeformer));
+			}
+			
 #if WITH_EDITORONLY_DATA
 			if (Info.GraphType == EOptimusNodeGraphType::Update)
 			{
diff --git a/Engine/Plugins/Animation/DeformerGraph/Source/OptimusCore/Public/IOptimusDeformerAssetPathAccessor.h b/Engine/Plugins/Animation/DeformerGraph/Source/OptimusCore/Public/IOptimusDeformerAssetPathAccessor.h
new file mode 100644
--- /dev/null
+++ b/Engine/Plugins/Animation/DeformerGraph/Source/OptimusCore/Public/IOptimusDeformerAssetPathAccessor.h
@@ -0,0 +1,27 @@
+﻿// Copyright Epic Games, Inc. All Rights Reserved.
+
+#pragma once
+
+#include "UObject/Interface.h"
+
+#include "IOptimusDeformerAssetPathAccessor.generated.h"
+
+
+class UOptimusDeformerInstance;
+struct FTopLevelAssetPath;
+
+UINTERFACE()
+class OPTIMUSCORE_API UOptimusDeformerAssetPathAccessor :
+	public UInterface
+{
+	GENERATED_BODY()
+};
+
+
+class IOptimusDeformerAssetPathAccessor
+{
+	GENERATED_BODY()
+
+public:
+	virtual void SetOptimusDeformerAssetPath(const FTopLevelAssetPath& InPath) = 0;
+};
diff --git a/Engine/Plugins/Animation/DeformerGraph/Source/OptimusCore/Public/OptimusDeformer.h b/Engine/Plugins/Animation/DeformerGraph/Source/OptimusCore/Public/OptimusDeformer.h
--- a/Engine/Plugins/Animation/DeformerGraph/Source/OptimusCore/Public/OptimusDeformer.h
+++ b/Engine/Plugins/Animation/DeformerGraph/Source/OptimusCore/Public/OptimusDeformer.h
@@ -18,6 +18,7 @@
 #include "Interfaces/Interface_PreviewMeshProvider.h"
 #include "Logging/TokenizedMessage.h"
 #include "Misc/TVariant.h"
+#include "Rendering/SkeletalMeshHalfEdgeBufferAccessor.h"
 
 #include "OptimusDeformer.generated.h"
 
@@ -130,7 +131,8 @@ class OPTIMUSCORE_API UOptimusDeformer :
 	public IInterface_PreviewMeshProvider,
 	public IOptimusPathResolver,
 	public IOptimusNodeGraphCollectionOwner,
-	public IOptimusNodeFunctionLibraryOwner 
+	public IOptimusNodeFunctionLibraryOwner,
+	public ISkeletalMeshHalfEdgeBufferAccessor
 {
 	GENERATED_BODY()
 
@@ -339,6 +341,9 @@ public:
 		UMeshDeformerInstanceSettings* InSettings
 		);
 	
+	// ISkeletalMeshHalfEdgeBufferAccessor overrides
+	bool IsSkeletalMeshHalfEdgeBufferRequired() const override;
+	
 	// IInterface_PreviewMeshProvider overrides
 	void SetPreviewMesh(USkeletalMesh* PreviewMesh, bool bMarkAsDirty = true) override;
 	USkeletalMesh* GetPreviewMesh() const override;
diff --git a/Engine/Plugins/Animation/DeformerGraph/Source/OptimusSettings/Private/OptimusSettingsModule.cpp b/Engine/Plugins/Animation/DeformerGraph/Source/OptimusSettings/Private/OptimusSettingsModule.cpp
--- a/Engine/Plugins/Animation/DeformerGraph/Source/OptimusSettings/Private/OptimusSettingsModule.cpp
+++ b/Engine/Plugins/Animation/DeformerGraph/Source/OptimusSettings/Private/OptimusSettingsModule.cpp
@@ -30,6 +30,23 @@ namespace Optimus
 	{
 		return ComputeFramework::IsEnabled();
 	}
+	
+	static const TCHAR* DeformerGraphEnableAssetValidationCVarName = TEXT("r.DeformerGraph.EnableAssetValidation");
+	static TAutoConsoleVariable<int32> CVarDeformerGraphEnableAssetValidation(
+	DeformerGraphEnableAssetValidationCVarName,
+	1,
+	TEXT("Set to 0 to disable Deformer Graph validation that may generate "
+	  "extra warnings when Deformer Graph cannot run optimally on certain assets"));
+
+	bool IsAssetValidationEnabled()
+	{
+		return CVarDeformerGraphEnableAssetValidation.GetValueOnAnyThread() != 0;
+	}
+
+	FName GetEnableAssetValidationCVarName()
+	{
+		return DeformerGraphEnableAssetValidationCVarName;
+	}
 }
 
 
diff --git a/Engine/Plugins/Animation/DeformerGraph/Source/OptimusSettings/Public/OptimusSettings.h b/Engine/Plugins/Animation/DeformerGraph/Source/OptimusSettings/Public/OptimusSettings.h
--- a/Engine/Plugins/Animation/DeformerGraph/Source/OptimusSettings/Public/OptimusSettings.h
+++ b/Engine/Plugins/Animation/DeformerGraph/Source/OptimusSettings/Public/OptimusSettings.h
@@ -52,4 +52,10 @@ namespace Optimus
 
 	/** Returns true if DeformerGraph is currently enabled. */
 	OPTIMUSSETTINGS_API bool IsEnabled();
+	
+	/** Returns true if DeformerGraph should enable asset validation. */
+	OPTIMUSSETTINGS_API bool IsAssetValidationEnabled();
+	
+	/** Returns the name of the CVar that toggles asset validation, used by warning messages. */
+	OPTIMUSSETTINGS_API FName GetEnableAssetValidationCVarName();
 }
diff --git a/Engine/Source/Editor/UnrealEd/Classes/Factories/MeshDeformerCollectionFactory.h b/Engine/Source/Editor/UnrealEd/Classes/Factories/MeshDeformerCollectionFactory.h
new file mode 100644
--- /dev/null
+++ b/Engine/Source/Editor/UnrealEd/Classes/Factories/MeshDeformerCollectionFactory.h
@@ -0,0 +1,24 @@
+﻿// Copyright Epic Games, Inc. All Rights Reserved.
+
+#pragma once
+
+#include "CoreMinimal.h"
+#include "Factories/Factory.h"
+#include "MeshDeformerCollectionFactory.generated.h"
+
+
+UCLASS(hidecategories = Object, MinimalAPI)
+class UMeshDeformerCollectionFactory : public UFactory
+{
+public:
+	GENERATED_BODY()
+	UMeshDeformerCollectionFactory();
+
+	// UFactory interface
+	virtual UObject* FactoryCreateNew(UClass* InClass, UObject* InParent, FName InName, EObjectFlags InFlags, UObject* InContext, FFeedbackContext* InWarn) override;
+	virtual FText GetDisplayName() const override;
+	virtual FText GetToolTip() const override;
+	virtual uint32 GetMenuCategories() const override;
+	virtual const TArray<FText>& GetMenuCategorySubMenus() const override;
+	// End of UFactory interface
+};
\ No newline at end of file
diff --git a/Engine/Source/Editor/UnrealEd/Private/Factories/MeshDeformerCollectionFactory.cpp b/Engine/Source/Editor/UnrealEd/Private/Factories/MeshDeformerCollectionFactory.cpp
new file mode 100644
--- /dev/null
+++ b/Engine/Source/Editor/UnrealEd/Private/Factories/MeshDeformerCollectionFactory.cpp
@@ -0,0 +1,47 @@
+﻿// Copyright Epic Games, Inc. All Rights Reserved.
+
+#include "Factories/MeshDeformerCollectionFactory.h"
+#include "AssetTypeCategories.h"
+#include "Animation/MeshDeformerCollection.h"
+
+
+#define LOCTEXT_NAMESPACE "MeshDeformerCollectionFactory"
+
+UMeshDeformerCollectionFactory::UMeshDeformerCollectionFactory()
+{
+	SupportedClass = UMeshDeformerCollection::StaticClass();
+	bEditAfterNew = true;
+	bCreateNew = true;
+}
+
+
+
+UObject* UMeshDeformerCollectionFactory::FactoryCreateNew(UClass* InClass, UObject* InParent, FName InName, EObjectFlags InFlags, UObject* Context,
+                                                          FFeedbackContext* Warn)
+{
+	UMeshDeformerCollection* MeshDeformerCollection = NewObject<UMeshDeformerCollection>(InParent, InClass, InName, InFlags);
+	return MeshDeformerCollection;
+}
+
+FText UMeshDeformerCollectionFactory::GetDisplayName() const
+{
+	return LOCTEXT("DisplayName", "Mesh Deformer Collection");
+}
+
+FText UMeshDeformerCollectionFactory::GetToolTip() const
+{
+	return LOCTEXT("ToolTip", "A simple collection of Mesh Deformers primarily used by Skeletal Mesh Asset to determined if extra deformer specific data should be built");
+}
+
+uint32 UMeshDeformerCollectionFactory::GetMenuCategories() const
+{
+	return EAssetTypeCategories::Animation;
+}
+
+const TArray<FText>& UMeshDeformerCollectionFactory::GetMenuCategorySubMenus() const
+{
+	static TArray<FText> SubMenus { LOCTEXT("SubMenuDeformers", "Deformers") };
+	return SubMenus;
+}
+
+#undef LOCTEXT_NAMESPACE
diff --git a/Engine/Source/Editor/UnrealEd/Private/ImportUtils/SkeletalMeshImportUtils.cpp b/Engine/Source/Editor/UnrealEd/Private/ImportUtils/SkeletalMeshImportUtils.cpp
--- a/Engine/Source/Editor/UnrealEd/Private/ImportUtils/SkeletalMeshImportUtils.cpp
+++ b/Engine/Source/Editor/UnrealEd/Private/ImportUtils/SkeletalMeshImportUtils.cpp
@@ -402,6 +402,7 @@ TSharedPtr<FExistingSkelMeshData> SkeletalMeshImportUtils::SaveExistingSkelMeshD
 	ExistingMeshDataPtr->ExistingSamplingInfo = SourceSkeletalMesh->GetSamplingInfo();
 	ExistingMeshDataPtr->ExistingDefaultAnimatingRig = SourceSkeletalMesh->GetDefaultAnimatingRig();
 	ExistingMeshDataPtr->ExistingDefaultMeshDeformer = SourceSkeletalMesh->GetDefaultMeshDeformer();
+	ExistingMeshDataPtr->ExistingTargetMeshDeformers = SourceSkeletalMesh->GetTargetMeshDeformers();
 
 	if (ExistingMeshDataPtr->UseMaterialNameSlotWorkflow)
 	{
@@ -904,6 +905,7 @@ void SkeletalMeshImportUtils::RestoreExistingSkelMeshData(const TSharedPtr<const
 		SkeletalMesh->SetPostProcessAnimBlueprint(MeshData->ExistingPostProcessAnimBlueprint);
 		SkeletalMesh->SetDefaultAnimatingRig(MeshData->ExistingDefaultAnimatingRig);
 		SkeletalMesh->SetDefaultMeshDeformer(MeshData->ExistingDefaultMeshDeformer);
+		SkeletalMesh->SetTargetMeshDeformers(MeshData->ExistingTargetMeshDeformers);
 
 		SkeletalMesh->GetMorphTargets().Empty(MeshData->ExistingMorphTargets.Num());
 		SkeletalMesh->GetMorphTargets().Append(MeshData->ExistingMorphTargets);
diff --git a/Engine/Source/Editor/UnrealEd/Public/ImportUtils/SkelImport.h b/Engine/Source/Editor/UnrealEd/Public/ImportUtils/SkelImport.h
--- a/Engine/Source/Editor/UnrealEd/Public/ImportUtils/SkelImport.h
+++ b/Engine/Source/Editor/UnrealEd/Public/ImportUtils/SkelImport.h
@@ -85,6 +85,7 @@ struct FExistingSkelMeshData
 	TSubclassOf<UAnimInstance>				ExistingPostProcessAnimBlueprint;
 	TSoftObjectPtr<UObject>					ExistingDefaultAnimatingRig;
 	UMeshDeformer*							ExistingDefaultMeshDeformer;
+	UMeshDeformerCollection*				ExistingTargetMeshDeformers;
 	//////////////////////////////////////////////////////////////////////////
 
 	bool									bExistingUseFullPrecisionUVs;
diff --git a/Engine/Source/Runtime/Engine/Classes/Components/SkinnedMeshComponent.h b/Engine/Source/Runtime/Engine/Classes/Components/SkinnedMeshComponent.h
--- a/Engine/Source/Runtime/Engine/Classes/Components/SkinnedMeshComponent.h
+++ b/Engine/Source/Runtime/Engine/Classes/Components/SkinnedMeshComponent.h
@@ -1226,6 +1226,7 @@ public:
 #if WITH_EDITOR
 	ENGINE_API virtual bool CanEditChange(const FProperty* InProperty) const override;
 	ENGINE_API virtual void PostEditChangeProperty(FPropertyChangedEvent& PropertyChangedEvent) override;
+	ENGINE_API virtual EDataValidationResult IsDataValid(class FDataValidationContext& Context) const override;
 #endif // WITH_EDITOR
 
 protected:
diff --git a/Engine/Source/Runtime/Engine/Classes/Engine/SkeletalMesh.h b/Engine/Source/Runtime/Engine/Classes/Engine/SkeletalMesh.h
--- a/Engine/Source/Runtime/Engine/Classes/Engine/SkeletalMesh.h
+++ b/Engine/Source/Runtime/Engine/Classes/Engine/SkeletalMesh.h
@@ -131,6 +131,7 @@ enum class ESkeletalMeshAsyncProperties : uint64
 	DefaultMeshDeformer = 1llu << 56,
 	OverlayMaterial = 1llu << 57,
 	OverlayMaterialMaxDrawDistance = 1llu << 58,
+	TargetMeshDeformers = 1llu << 59,
 	All = MAX_uint64
 };
 
@@ -2987,6 +2988,10 @@ protected:
 	UE_DEPRECATED(5.1, "This must be protected for async build, always use the accessors even internally.")
 	UPROPERTY(EditAnywhere, BlueprintReadOnly, Category = "Deformer", BlueprintGetter=GetDefaultMeshDeformer)
 	TObjectPtr<UMeshDeformer> DefaultMeshDeformer;
+	
+	/** Skeletal Mesh needs this collection of deformers to make sure it cooks any extra data required by these deformers */
+	UPROPERTY(EditAnywhere, BlueprintReadOnly, Category = "Deformer", BlueprintGetter=GetTargetMeshDeformers)
+	TObjectPtr<UMeshDeformerCollection> TargetMeshDeformers;
 
 	/** Default translucent material to blend on top of this mesh. Mesh will be rendered twice - once with a base material and once with overlay material */
 	UE_DEPRECATED(5.3, "This must be protected for async build, always use the accessors even internally.")
@@ -3018,7 +3023,24 @@ public:
 		return DefaultMeshDeformer;
 		PRAGMA_ENABLE_DEPRECATION_WARNINGS
 	}
+	
+	/** Set the collection of mesh deformers that may be applied to this mesh. Skeletal Mesh use these deformers to determined if any extra data needs to be cooked */
+	void SetTargetMeshDeformers(UMeshDeformerCollection* InMeshDeformerCollection)
+	{
+		WaitUntilAsyncPropertyReleased(ESkeletalMeshAsyncProperties::TargetMeshDeformers, ESkinnedAssetAsyncPropertyLockType::WriteOnly);
+		TargetMeshDeformers = InMeshDeformerCollection;
+	}
+	
+	/** Set the collection of mesh deformers that may be applied to this mesh. Skeletal Mesh use these deformers to determined if any extra data needs to be cooked*/
+	UFUNCTION(BlueprintCallable, Category="Rendering|Deformer")
+	virtual UMeshDeformerCollection* GetTargetMeshDeformers() const override
+	{
+		WaitUntilAsyncPropertyReleased(ESkeletalMeshAsyncProperties::TargetMeshDeformers, ESkinnedAssetAsyncPropertyLockType::ReadOnly);
+		return TargetMeshDeformers;
+	}
 
+	/** True if this mesh LOD builds the half edge buffer. USkinnedAsset interface.	*/
+	ENGINE_API virtual bool HasHalfEdgeBuffer(int32 LODIndex)const override;
 	
 	/** Get the default overlay material used by this mesh */
 	UFUNCTION(BlueprintCallable, Category="Rendering|Material")
diff --git a/Engine/Source/Runtime/Engine/Classes/Engine/SkinnedAsset.h b/Engine/Source/Runtime/Engine/Classes/Engine/SkinnedAsset.h
--- a/Engine/Source/Runtime/Engine/Classes/Engine/SkinnedAsset.h
+++ b/Engine/Source/Runtime/Engine/Classes/Engine/SkinnedAsset.h
@@ -22,6 +22,7 @@ class FSkinnedAsyncTaskContext;
 class FVertexFactoryType;
 class ITargetPlatform;
 class UMeshDeformer;
+class UMeshDeformerCollection;
 class UMorphTarget;
 class UPhysicsAsset;
 class USkeleton;
@@ -174,7 +175,13 @@ public:
 
 	ENGINE_API virtual UMeshDeformer* GetDefaultMeshDeformer() const
 	PURE_VIRTUAL(USkinnedAsset::GetDefaultMeshDeformer, return nullptr;);
-
+	
+	ENGINE_API virtual UMeshDeformerCollection* GetTargetMeshDeformers() const
+	PURE_VIRTUAL(USkinnedAsset::GetTargetMeshDeformers, return nullptr;);
+	
+	ENGINE_API virtual bool HasHalfEdgeBuffer(int32 LODIndex) const
+	PURE_VIRTUAL(USkinnedAsset::HasHalfEdgeBuffer, return false;);
+	
 	ENGINE_API virtual class UMaterialInterface* GetOverlayMaterial() const
 	PURE_VIRTUAL(USkinnedAsset::GetOverlayMaterial, return nullptr;);
 	ENGINE_API virtual float GetOverlayMaterialMaxDrawDistance() const
diff --git a/Engine/Source/Runtime/Engine/Classes/Engine/SkinnedAssetCommon.h b/Engine/Source/Runtime/Engine/Classes/Engine/SkinnedAssetCommon.h
--- a/Engine/Source/Runtime/Engine/Classes/Engine/SkinnedAssetCommon.h
+++ b/Engine/Source/Runtime/Engine/Classes/Engine/SkinnedAssetCommon.h
@@ -228,7 +228,7 @@ struct FSkeletalMeshLODInfo
  	 * may be useful for other systems like Mesh Deformer. Please note that Half Edge buffers are always built if the Skeletal Mesh Asset
  	 * has Default Mesh Deformer assigned, and the LOD has AllowMeshDeformer enabled, 
 	 */
-	UPROPERTY(EditAnywhere, Category = SkeletalMeshLODInfo, meta=(DisplayName="Force Build Half Edge Buffers"))
+	UPROPERTY(EditAnywhere, Category = SkeletalMeshLODInfo, meta=(DisplayName="Force Build Half Edge Buffers", EditCondition=bAllowMeshDeformer))
 	uint8 bBuildHalfEdgeBuffers : 1;
 
 	/** Whether a Mesh Deformer applied to the mesh asset or Skinned Mesh Component should be used on this LOD or not */
diff --git a/Engine/Source/Runtime/Engine/Private/Animation/MeshDeformerCollection.cpp b/Engine/Source/Runtime/Engine/Private/Animation/MeshDeformerCollection.cpp
new file mode 100644
--- /dev/null
+++ b/Engine/Source/Runtime/Engine/Private/Animation/MeshDeformerCollection.cpp
@@ -0,0 +1,37 @@
+﻿// Copyright Epic Games, Inc. All Rights Reserved.
+
+#include "Animation/MeshDeformerCollection.h"
+#include "Animation/MeshDeformer.h"
+
+TArray<TSoftObjectPtr<UMeshDeformer>> UMeshDeformerCollection::GetMeshDeformers() const
+{
+	TSet<TObjectPtr<const UMeshDeformerCollection>> VisitedCollection;
+	TArray<TSoftObjectPtr<UMeshDeformer>> Deformers;
+	
+	GetMeshDeformers_Internal(VisitedCollection, Deformers);
+	
+	return Deformers;
+}
+
+void UMeshDeformerCollection::GetMeshDeformers_Internal(
+	TSet<TObjectPtr<const UMeshDeformerCollection>>& InOutVisitedCollection,
+	TArray<TSoftObjectPtr<UMeshDeformer>>& InOutMeshDeformers
+	) const
+{
+	if (InOutVisitedCollection.Contains(this))
+	{
+		return;
+	}
+	
+	for (const TSoftObjectPtr<UMeshDeformer>& Deformer : MeshDeformers)
+	{
+		InOutMeshDeformers.AddUnique(Deformer);
+	}
+
+	InOutVisitedCollection.Add(this);
+
+	for (const TObjectPtr<UMeshDeformerCollection>& MeshDeformerCollection :  MeshDeformerCollections)
+	{
+		MeshDeformerCollection->GetMeshDeformers_Internal(InOutVisitedCollection, InOutMeshDeformers);
+	}
+}
diff --git a/Engine/Source/Runtime/Engine/Private/Components/SkinnedMeshComponent.cpp b/Engine/Source/Runtime/Engine/Private/Components/SkinnedMeshComponent.cpp
--- a/Engine/Source/Runtime/Engine/Private/Components/SkinnedMeshComponent.cpp
+++ b/Engine/Source/Runtime/Engine/Private/Components/SkinnedMeshComponent.cpp
@@ -51,8 +51,11 @@
 #include "NaniteSceneProxy.h"
 #include "NaniteVertexFactory.h"
 #include "Animation/MeshDeformerGeometryReadback.h"
+#include "Misc/DataValidation.h"
+#include "Rendering/SkeletalMeshHalfEdgeBufferAccessor.h"
 
 DEFINE_LOG_CATEGORY_STATIC(LogSkinnedMeshComp, Log, All);
+#define LOCTEXT_NAMESPACE "SkinnedMeshComponent"
 
 int32 GSkeletalMeshLODBias = 0;
 FAutoConsoleVariableRef CVarSkeletalMeshLODBias(
@@ -1424,6 +1427,46 @@ void USkinnedMeshComponent::PostEditChangeProperty(FPropertyChangedEvent& Proper
 	}
 }
 
+EDataValidationResult USkinnedMeshComponent::IsDataValid(class FDataValidationContext& Context) const
+{
+	EDataValidationResult Result = EDataValidationResult::NotValidated;
+
+	if (USkinnedAsset* Asset = GetSkinnedAsset())
+	{
+		FMeshDeformerSet MeshDeformers = GetActiveMeshDeformers();
+		
+		for (int32 LODIndex = 0; LODIndex < MeshDeformers.DeformerIndexForLOD.Num(); LODIndex++)
+		{
+			int32 DeformerIndex = MeshDeformers.DeformerIndexForLOD[LODIndex];
+			UMeshDeformer* Deformer = MeshDeformers.Deformers[DeformerIndex];
+			if (Cast<ISkeletalMeshHalfEdgeBufferAccessor>(Deformer))
+			{
+				if (SkeletalMeshHalfEdgeBufferAccessor::IsHalfEdgeRequired(Deformer))
+				{
+					if (!Asset->HasHalfEdgeBuffer(LODIndex))
+					{
+						// Deformer requires the half edge buffer to work,
+						// but it is missing on the skinned asset,
+						// so notify the user to manual fix it
+						Context.AddWarning(
+							FText::Format(
+								LOCTEXT("SkeletalMeshMissingHalfEdgeBuffer",
+								        "Please add Deformer Graph {0} to TargetMeshDeformers of Skeletal Mesh Asset {1} to "
+								        "make sure extra data required by the Deformer is cooked for non-editor builds."),
+								FText::FromName(Deformer->GetPackage()->GetFName()),
+								FText::FromName(Asset->GetPackage()->GetFName())
+							));
+						Result = EDataValidationResult::Invalid;
+						break;
+					}
+				}
+			}
+		}
+	}
+	
+	return CombineDataValidationResults(Result, Super::IsDataValid(Context));
+}
+
 #endif // WITH_EDITOR
 
 void USkinnedMeshComponent::InitLODInfos()
@@ -2450,8 +2493,7 @@ FMeshDeformerSet USkinnedMeshComponent::GetActiveMeshDeformers() const
 	const FSkeletalMeshRenderData* RenderData = GetSkeletalMeshRenderData();
 
 	if (!RenderData
-		|| !GetSkinnedAsset()
-		|| !GetScene())
+		|| !GetSkinnedAsset())
 	{
 		// In order to select the right deformer, there needs to be a valid mesh and valid render data.
 		return FMeshDeformerSet();
@@ -2471,7 +2513,13 @@ FMeshDeformerSet USkinnedMeshComponent::GetActiveMeshDeformers() const
 
 	// If there's no user-specified deformer, find out if we need to set a default deformer for
 	// Unlimited Bone Influences.
-	const bool bIsDeformerRequiredForUBI = FGPUBaseSkinVertexFactory::GetAlwaysUseDeformerForUnlimitedBoneInfluences(GetScene()->GetShaderPlatform());
+	EShaderPlatform ShaderPlatformToCheck = GShaderPlatformForFeatureLevel[GMaxRHIFeatureLevel];
+	if (GetScene())
+	{
+		ShaderPlatformToCheck = GetScene()->GetShaderPlatform();
+	}
+	
+	const bool bIsDeformerRequiredForUBI = FGPUBaseSkinVertexFactory::GetAlwaysUseDeformerForUnlimitedBoneInfluences(ShaderPlatformToCheck);
 
 	if (!bIsDeformerRequestedByUser)
 	{
@@ -5745,4 +5793,6 @@ template FVector3f GetTypedSkinnedVertexPosition<true>(const USkinnedMeshCompone
 	const FSkinWeightVertexBuffer& SkinWeightVertexBuffer, const int32 VertIndex, const TArray<FMatrix44f> & RefToLocals);
 
 template FVector3f GetTypedSkinnedVertexPosition<false>(const USkinnedMeshComponent* SkinnedComp, const FSkelMeshRenderSection& Section, const FPositionVertexBuffer& PositionVertexBuffer,
-	const FSkinWeightVertexBuffer& SkinWeightVertexBuffer, const int32 VertIndex, const TArray<FMatrix44f> & RefToLocals);
\ No newline at end of file
+	const FSkinWeightVertexBuffer& SkinWeightVertexBuffer, const int32 VertIndex, const TArray<FMatrix44f> & RefToLocals);
+
+#undef LOCTEXT_NAMESPACE
\ No newline at end of file
diff --git a/Engine/Source/Runtime/Engine/Private/Rendering/SkeletalMeshHalfEdgeBufferAccessor.cpp b/Engine/Source/Runtime/Engine/Private/Rendering/SkeletalMeshHalfEdgeBufferAccessor.cpp
new file mode 100644
--- /dev/null
+++ b/Engine/Source/Runtime/Engine/Private/Rendering/SkeletalMeshHalfEdgeBufferAccessor.cpp
@@ -0,0 +1,52 @@
+﻿// Copyright Epic Games, Inc. All Rights Reserved.
+
+#include "Rendering/SkeletalMeshHalfEdgeBufferAccessor.h"
+
+#include "AssetRegistry/AssetRegistryModule.h"
+#include "AssetRegistry/AssetData.h"
+
+FName SkeletalMeshHalfEdgeBufferAccessor::GetHalfEdgeRequirementAssetTagName()
+{
+	static FName Tag("bRequiresSkeletalMeshHalfEdgeBuffer");
+	return Tag;
+}
+
+bool SkeletalMeshHalfEdgeBufferAccessor::IsHalfEdgeRequired(TSoftObjectPtr<UObject> InAssetSoftPtr)
+{
+	if (InAssetSoftPtr.IsNull())
+	{
+		return false;
+	}
+	
+	const FAssetRegistryModule& AssetRegistryModule = FModuleManager::LoadModuleChecked<FAssetRegistryModule>("AssetRegistry");	
+
+	FAssetData AssetData;
+	AssetRegistryModule.TryGetAssetByObjectPath(InAssetSoftPtr.ToSoftObjectPath(), AssetData);
+
+	if (AssetData.GetClass()->ImplementsInterface(USkeletalMeshHalfEdgeBufferAccessor::StaticClass()))
+	{
+		bool bIsHalfEdgeRequired = false;
+		if (AssetData.GetTagValue<bool>(GetHalfEdgeRequirementAssetTagName(), bIsHalfEdgeRequired))
+		{
+			if (bIsHalfEdgeRequired)
+			{
+				return true;
+			}
+		}
+		else
+		{
+			// Fallback path
+			// Older assets may not have this tag serialized yet, lets trigger a load and query it directly
+			// However this isn't bulletproof since loading new assets isn't allowed when this function is called by certain loaders (ZenLoader)
+			if (ISkeletalMeshHalfEdgeBufferAccessor* HalfEdgeBufferAccessor = Cast<ISkeletalMeshHalfEdgeBufferAccessor>(InAssetSoftPtr.LoadSynchronous()))
+			{
+				if (HalfEdgeBufferAccessor->IsSkeletalMeshHalfEdgeBufferRequired())
+				{
+					return true;
+				}
+			}
+		}
+	}
+
+	return false;
+}
diff --git a/Engine/Source/Runtime/Engine/Private/SkeletalMesh.cpp b/Engine/Source/Runtime/Engine/Private/SkeletalMesh.cpp
--- a/Engine/Source/Runtime/Engine/Private/SkeletalMesh.cpp
+++ b/Engine/Source/Runtime/Engine/Private/SkeletalMesh.cpp
@@ -103,6 +103,9 @@
 #include "SkeletalMeshAttributes.h"
 #include "Algo/AnyOf.h"
 #include "Logging/StructuredLog.h"
+#include "Rendering/SkeletalMeshHalfEdgeBufferAccessor.h"
+#include "Animation/MeshDeformerCollection.h"
+#include "Animation/MeshDeformer.h"
 
 #include UE_INLINE_GENERATED_CPP_BY_NAME(SkeletalMesh)
 
@@ -3508,6 +3511,43 @@ bool USkeletalMesh::IsPostLoadThreadSafe() const
 	return false;	// PostLoad is not thread safe
 }
 
+bool USkeletalMesh::HasHalfEdgeBuffer(int32 LODIndex) const
+{
+	const FSkeletalMeshLODInfo* Info = GetLODInfo(LODIndex);
+	if (!Info)
+	{
+		return false;
+	}
+	
+	if (!Info->bAllowMeshDeformer)
+	{
+		return false;
+	}
+	
+	if (Info->bBuildHalfEdgeBuffers)
+	{
+		return true;
+	}
+	
+	if (SkeletalMeshHalfEdgeBufferAccessor::IsHalfEdgeRequired(GetDefaultMeshDeformer()))
+	{
+		return true;
+	}
+
+	if (TargetMeshDeformers)
+	{
+		for (const TSoftObjectPtr<UMeshDeformer>& MeshDeformer : TargetMeshDeformers->GetMeshDeformers())
+		{
+			if (SkeletalMeshHalfEdgeBufferAccessor::IsHalfEdgeRequired(MeshDeformer))
+			{
+				return true;
+			}
+		}	
+	}
+	
+	return false;
+}
+
 void USkeletalMesh::BeginPostLoadInternal(FSkinnedAssetPostLoadContext& Context)
 {
 #if WITH_EDITOR
@@ -5310,6 +5350,8 @@ static void SerializeLODInfoForDDC(USkeletalMesh* SkeletalMesh, FString& KeySuff
 		const FSkeletalMeshLODGroupSettings* SkeletalMeshLODGroupSettings = bValidLODSettings ? &SkeletalMesh->GetLODSettings()->GetSettingsForLODLevel(LODIndex) : nullptr;
 		LODInfo.BuildGUID = LODInfo.ComputeDeriveDataCacheKey(SkeletalMeshLODGroupSettings);
 		KeySuffix += LODInfo.BuildGUID.ToString(EGuidFormats::Digits);
+
+		KeySuffix += SkeletalMesh->HasHalfEdgeBuffer(LODIndex) ? "1" : "0";
 	}
 }
 
diff --git a/Engine/Source/Runtime/Engine/Private/SkeletalMeshRenderData.cpp b/Engine/Source/Runtime/Engine/Private/SkeletalMeshRenderData.cpp
--- a/Engine/Source/Runtime/Engine/Private/SkeletalMeshRenderData.cpp
+++ b/Engine/Source/Runtime/Engine/Private/SkeletalMeshRenderData.cpp
@@ -348,15 +348,8 @@ void FSkeletalMeshRenderData::Cache(const ITargetPlatform* TargetPlatform, USkin
 				}
 				FSkeletalMeshLODRenderData::FBuildSettings BuildSettings;
 				BuildSettings.BuildFlags = VertexBufferBuildFlags;
-				BuildSettings.bBuildHalfEdgeBuffers = LODInfo->bBuildHalfEdgeBuffers;
+				BuildSettings.bBuildHalfEdgeBuffers = Owner->HasHalfEdgeBuffer(LODIndex); 
 
-				// The presence of a deformer on the skinned asset almost guarantees dependency on HalfEdgeBuffer,
-				// so let's build half edge buffer for them automatically
-				if (LODInfo->bAllowMeshDeformer && Owner->GetDefaultMeshDeformer())
-				{
-					BuildSettings.bBuildHalfEdgeBuffers = true;
-				}
-				
 				LODData->BuildFromLODModel(LODModel, LODInfo->VertexAttributes, BuildSettings);
 			}
 
diff --git a/Engine/Source/Runtime/Engine/Public/Animation/MeshDeformerCollection.h b/Engine/Source/Runtime/Engine/Public/Animation/MeshDeformerCollection.h
new file mode 100644
--- /dev/null
+++ b/Engine/Source/Runtime/Engine/Public/Animation/MeshDeformerCollection.h
@@ -0,0 +1,36 @@
+﻿// Copyright Epic Games, Inc. All Rights Reserved.
+
+#pragma once
+
+#include "CoreMinimal.h"
+#include "Engine/DataAsset.h"
+#include "MeshDeformerCollection.generated.h"
+
+class UMeshDeformer;
+
+/** A simple collection of Mesh Deformers */
+UCLASS(EditInlineNew, MinimalAPI, BlueprintType)
+class UMeshDeformerCollection : public UDataAsset
+{
+	GENERATED_BODY()
+public:
+	
+#if WITH_EDITORONLY_DATA
+	UPROPERTY(EditAnywhere, Category = "Mesh Deformer Collections")
+	FString Description;
+#endif
+	
+	UPROPERTY(EditAnywhere, Category = "Mesh Deformer Collections")
+	TArray<TSoftObjectPtr<UMeshDeformer>> MeshDeformers;
+
+	UPROPERTY(EditAnywhere, Category = "Mesh Deformer Collections")
+	TArray<TObjectPtr<class UMeshDeformerCollection>> MeshDeformerCollections;
+
+	TArray<TSoftObjectPtr<UMeshDeformer>> GetMeshDeformers() const;
+
+private:
+	void GetMeshDeformers_Internal(
+		TSet<TObjectPtr<const UMeshDeformerCollection>>& InOutVisitedCollection, 
+		TArray<TSoftObjectPtr<UMeshDeformer>>& InOutMeshDeformers
+		) const;
+};
\ No newline at end of file
diff --git a/Engine/Source/Runtime/Engine/Public/Rendering/SkeletalMeshHalfEdgeBufferAccessor.h b/Engine/Source/Runtime/Engine/Public/Rendering/SkeletalMeshHalfEdgeBufferAccessor.h
new file mode 100644
--- /dev/null
+++ b/Engine/Source/Runtime/Engine/Public/Rendering/SkeletalMeshHalfEdgeBufferAccessor.h
@@ -0,0 +1,32 @@
+// Copyright Epic Games, Inc. All Rights Reserved.
+
+#pragma once
+
+#include "UObject/Interface.h"
+
+#include "SkeletalMeshHalfEdgeBufferAccessor.generated.h"
+
+namespace SkeletalMeshHalfEdgeBufferAccessor
+{
+	// Asset implementing this interface should add this tag to the asset tags
+	// if half edge buffer is required
+	ENGINE_API	FName GetHalfEdgeRequirementAssetTagName();
+	bool IsHalfEdgeRequired(TSoftObjectPtr<UObject> InAssetSoftPtr);	
+}
+
+UINTERFACE()
+class USkeletalMeshHalfEdgeBufferAccessor : public UInterface
+{
+	GENERATED_BODY()
+};
+
+class ISkeletalMeshHalfEdgeBufferAccessor
+{
+	GENERATED_BODY()
+
+public:
+	
+	virtual bool IsSkeletalMeshHalfEdgeBufferRequired() const = 0;
+};
+
+
