New material translator changes:
- Applied the a round of global review changes suggested by Laura.
- Written inline documentation to all symbols in MaterialIR.h.
- Fixed a couple of bugs the emitter regarding dimensional value emission and the analyzer.

diff --git a/Engine/Source/Runtime/Engine/Private/Materials/MaterialExpressionsIR.cpp b/Engine/Source/Runtime/Engine/Private/Materials/MaterialExpressionsIR.cpp
--- a/Engine/Source/Runtime/Engine/Private/Materials/MaterialExpressionsIR.cpp
+++ b/Engine/Source/Runtime/Engine/Private/Materials/MaterialExpressionsIR.cpp
@@ -374,14 +374,14 @@ void UMaterialExpressionAppendVector::Build(MIR::FEmitter& Emitter)
 
 	UE_MIR_CHECKPOINT();
 
-	MIR::FPrimitiveTypePtr AType = AValue->Type->AsPrimitive();
+	const MIR::FPrimitiveType* AType = AValue->Type->AsPrimitive();
 	if (AType->IsMatrix())
 	{
 		Emitter.Error(TEXT("Input A does not have scalar or vector type."));
 		return;
 	}
 
-	MIR::FPrimitiveTypePtr BType = nullptr;
+	const MIR::FPrimitiveType* BType = nullptr;
 	if (BValue)
 	{
 		BType = BValue->Type->AsPrimitive();
@@ -402,7 +402,7 @@ void UMaterialExpressionAppendVector::Build(MIR::FEmitter& Emitter)
 	check(Dimensions >= 2 && Dimensions <= 4);
 
 	// Construct the output vector type.
-	MIR::FPrimitiveTypePtr ResultType = MIR::FPrimitiveType::GetVector(MIR::ScalarKind_Float, Dimensions);
+	const MIR::FPrimitiveType* ResultType = MIR::FPrimitiveType::GetVector(MIR::ScalarKind_Float, Dimensions);
 
 	// Set up each output vector component 
 	MIR::FValue* Components[4] = { nullptr, nullptr, nullptr, nullptr };
@@ -1039,7 +1039,7 @@ void UMaterialExpressionConvert::Build(MIR::FEmitter& Emitter)
 				InputValues[InputIndex] = Emitter.DefaultInputToFloat4(&ConvertInput.ExpressionInput, ConvertInput.DefaultValue).Input(&ConvertInput.ExpressionInput);
 
 				// Expect type to be primitive.
-				MIR::FPrimitiveTypePtr InputPrimitiveType = InputValues[InputIndex]->Type->AsPrimitive();
+				const MIR::FPrimitiveType* InputPrimitiveType = InputValues[InputIndex]->Type->AsPrimitive();
 				if (!InputPrimitiveType)
 				{
 					Emitter.Errorf(TEXT("Input `%d` of type `%s` is not primitive."), InputComponentIndex, InputValues[InputIndex]->Type->GetSpelling().GetData());
@@ -1047,7 +1047,7 @@ void UMaterialExpressionConvert::Build(MIR::FEmitter& Emitter)
 				}
 
 				// Determine the target type.
-				MIR::FTypePtr InputType = MIR::FPrimitiveType::GetVector(InputPrimitiveType->ScalarKind, MaterialExpressionConvertType::GetComponentCount(ConvertInput.Type));
+				const MIR::FType* InputType = MIR::FPrimitiveType::GetVector(InputPrimitiveType->ScalarKind, MaterialExpressionConvertType::GetComponentCount(ConvertInput.Type));
 
 				// Cast the input value to the target type.
 				InputValues[InputIndex] = Emitter.EmitConstruct(InputType, InputValues[InputIndex]);
@@ -1101,7 +1101,7 @@ void UMaterialExpressionViewProperty::Build(MIR::FEmitter& Emitter)
 		const bool bInvProperty = OutputIndex == 1;
 
 		FString HLSLCode = bInvProperty && bHasCustomInverseCode ? PropertyMeta.InvPropertyCode : PropertyMeta.PropertyCode;
-		MIR::FTypePtr HLSLCodeType = MIR::FType::FromMaterialValueType(PropertyMeta.Type);
+		const MIR::FType* HLSLCodeType = MIR::FType::FromMaterialValueType(PropertyMeta.Type);
 
 		if (bInvProperty && !bHasCustomInverseCode)
 		{
diff --git a/Engine/Source/Runtime/Engine/Private/Materials/MaterialIR.cpp b/Engine/Source/Runtime/Engine/Private/Materials/MaterialIR.cpp
--- a/Engine/Source/Runtime/Engine/Private/Materials/MaterialIR.cpp
+++ b/Engine/Source/Runtime/Engine/Private/Materials/MaterialIR.cpp
@@ -11,10 +11,10 @@ namespace MIR
 
 EStage GetMaterialOutputStage(EMaterialProperty Property)
 {
-	return Property == MP_WorldPositionOffset ? MIR::Stage_Vertex : MIR::Stage_Fragment;
+	return Property == MP_WorldPositionOffset ? MIR::Stage_Vertex : MIR::Stage_Pixel;
 }
 
-const TCHAR* ValueKindToString(EValueKind Kind)
+const TCHAR* LexToString(EValueKind Kind)
 {
 	switch (Kind)
 	{
@@ -112,13 +112,7 @@ bool FValue::Equals(const FValue* Other) const
 	return FMemory::Memcmp(this, Other, SizeInBytes) == 0;
 }
 
-TArrayView<const FValue*> FValue::GetUses() const
-{
-	TArrayView<FValue*> Uses = const_cast<FValue*>(this)->GetUses();
-	return { Uses.Num() ? (const FValue**)&Uses[0] : nullptr , Uses.Num() };
-}
-
-TArrayView<FValue*> FValue::GetUses()
+TConstArrayView<FValue*> FValue::GetUses() const
 {
 	// Values have no uses by definition.
 	if (Kind < VK_InstructionBegin)
@@ -130,37 +124,37 @@ TArrayView<FValue*> FValue::GetUses()
 	{
 		case VK_Dimensional:
 		{
-			auto This = static_cast<FDimensional*>(this);
+			auto This = static_cast<const FDimensional*>(this);
 			return This->GetComponents();
 		}
 
 		case VK_SetMaterialOutput:
 		{
-			auto This = static_cast<FSetMaterialOutput*>(this);
+			auto This = static_cast<const FSetMaterialOutput*>(this);
 			return { &This->Arg, 1 };
 		}
 
 		case VK_Operator:
 		{
-			auto This = static_cast<FOperator*>(this);
+			auto This = static_cast<const FOperator*>(this);
 			return { &This->AArg, 3 };
 		}
 
 		case VK_Branch:
 		{
-			auto This = static_cast<FBranch*>(this);
+			auto This = static_cast<const FBranch*>(this);
 			return { &This->ConditionArg, 3 };
 		}
 
 		case VK_Subscript:
 		{
-			auto This = static_cast<FSubscript*>(this);
+			auto This = static_cast<const FSubscript*>(this);
 			return { &This->Arg, 1 };
 		}
 
 		case VK_TextureSample:
 		{
-			auto This = static_cast<FTextureSample*>(this);
+			auto This = static_cast<const FTextureSample*>(this);
 			return { &This->TextureObject, 4 };
 		}
 
@@ -311,7 +305,7 @@ bool FConstant::IsFloat() const
 	return Type == FPrimitiveType::GetFloat1();
 }
 
-const TCHAR* ExternalInputToString(EExternalInput Input)
+const TCHAR* LexToString(EExternalInput Input)
 {
 	switch (Input)
 	{
@@ -343,7 +337,7 @@ const TCHAR* ExternalInputToString(EExternalInput Input)
 	}
 }
 
-FTypePtr GetExternalInputType(EExternalInput Id)
+const FType* GetExternalInputType(EExternalInput Id)
 {
 	if (Id >= EExternalInput::TexCoord0 && Id <= EExternalInput::TexCoord7_Ddy)
 	{
@@ -403,21 +397,22 @@ FBlock* FBlock::FindCommonParentWith(MIR::FBlock* Other)
 	return A;
 }
 
-TArrayView<FValue* const> FDimensional::GetComponents() const
+TConstArrayView<FValue*> FDimensional::GetComponents() const
 {
-	TArrayView<FValue*> Components = const_cast<FDimensional*>(this)->GetComponents();
-	return { Components.GetData(), Components.Num() };
-}
-
-TArrayView<FValue*> FDimensional::GetComponents()
-{
-	FPrimitiveTypePtr PrimitiveType = Type->AsPrimitive();
+	const FPrimitiveType* PrimitiveType = Type->AsPrimitive();
 	check(PrimitiveType);
 
-	FValue** Ptr = (FValue**)static_cast<TDimensional<1>*>(this)->Components;
+	auto Ptr = static_cast<const TDimensional<1>*>(this)->Components;
 	return { Ptr, PrimitiveType->NumRows };
 }
 
+TArrayView<FValue*> FDimensional::GetMutableComponents()
+{
+	// const_cast is okay here as it's used only to get the array of components
+	TConstArrayView<FValue*> Components = static_cast<const FDimensional*>(this)->GetComponents();
+	return { const_cast<FValue**>(Components.GetData()), Components.Num() };
+}
+
 bool FDimensional::AreComponentsConstant() const
 {
 	for (FValue const* Component : GetComponents())
@@ -472,7 +467,7 @@ bool IsTernaryOperator(EOperator Op)
 	return Op >= TO_FirstTernaryOperator;
 }
 
-const TCHAR* OperatorToString(EOperator Op)
+const TCHAR* LexToString(EOperator Op)
 {
 	switch (Op)
 	{
diff --git a/Engine/Source/Runtime/Engine/Private/Materials/MaterialIRDebug.cpp b/Engine/Source/Runtime/Engine/Private/Materials/MaterialIRDebug.cpp
--- a/Engine/Source/Runtime/Engine/Private/Materials/MaterialIRDebug.cpp
+++ b/Engine/Source/Runtime/Engine/Private/Materials/MaterialIRDebug.cpp
@@ -60,7 +60,7 @@ struct FDebugDumpIRUseGraphState
 		// Begin the node declaration
 		Out.Appendf(TEXT("\"%p\" [label=< <b>%s</b>  (%s) <br/> "),
 						Value,
-						ValueKindToString(Value->Kind),
+						LexToString(Value->Kind),
 						Value->Type ? Value->Type->GetSpelling().GetData() : TEXT("???"));
 
 		DumpValueInfo(Value);
@@ -121,7 +121,7 @@ struct FDebugDumpIRUseGraphState
 		}
 		else if (const FExternalInput* ExternalInput = Value->As<FExternalInput>())
 		{
-			Out.Append(ExternalInputToString(ExternalInput->Id));
+			Out.Append(LexToString(ExternalInput->Id));
 		}
 		else if (const FSetMaterialOutput* SetMaterailOutput = Value->As<FSetMaterialOutput>())
 		{
@@ -146,7 +146,7 @@ struct FDebugDumpIRUseGraphState
 		}
 		else if (const FOperator* Operator = Value->As<FOperator>())
 		{
-			Out.Append(OperatorToString(Operator->Op));
+			Out.Append(LexToString(Operator->Op));
 		}
 	}
 
@@ -229,13 +229,13 @@ struct FDebugDumpIRState
 			AppendLeftColumn(Indentation, Temp);
 
 			// Print the kind of the instruction (the opcode, e.g. "Operator")
-			Out.Append(ValueKindToString(Instr->Kind));
+			Out.Append(LexToString(Instr->Kind));
 
 			// Begin printing the arguments (used values)
 			Out.Append(TEXT(" ("));
 
 			bool bAddComma = false;
-			TArrayView<FValue*> Uses = Instr->GetUses();
+			TConstArrayView<FValue*> Uses = Instr->GetUses();
 			for (int32 UseIndex = 0; UseIndex < Uses.Num(); ++UseIndex)
 			{
 				if (bAddComma)
@@ -288,7 +288,7 @@ struct FDebugDumpIRState
 		}
 		else if (const FOperator* Operator = Instr->As<FOperator>())
 		{
-			Out.Appendf(TEXT(" \"%s\""), OperatorToString(Operator->Op));
+			Out.Appendf(TEXT(" \"%s\""), LexToString(Operator->Op));
 		}
 		else if (const FSubscript* Subscript = Instr->As<FSubscript>())
 		{
@@ -349,7 +349,7 @@ struct FDebugDumpIRState
 		}
 		else if (const FExternalInput* ExternalInput = Value->As<FExternalInput>())
 		{
-			Out.Appendf(TEXT("[ExternalInput \"%s\"]"), ExternalInputToString(Value->As<FExternalInput>()->Id));
+			Out.Appendf(TEXT("[ExternalInput \"%s\"]"), LexToString(Value->As<FExternalInput>()->Id));
 		}
 		else if (const FTextureObject* TextureObject = Value->As<FTextureObject>())
 		{
@@ -364,7 +364,7 @@ struct FDebugDumpIRState
 		}
 		else
 		{
-			Out.Appendf(TEXT("[%s]"), ValueKindToString(Value->Kind));
+			Out.Appendf(TEXT("[%s]"), LexToString(Value->Kind));
 		}
 	}
 
diff --git a/Engine/Source/Runtime/Engine/Private/Materials/MaterialIREmitter.cpp b/Engine/Source/Runtime/Engine/Private/Materials/MaterialIREmitter.cpp
--- a/Engine/Source/Runtime/Engine/Private/Materials/MaterialIREmitter.cpp
+++ b/Engine/Source/Runtime/Engine/Private/Materials/MaterialIREmitter.cpp
@@ -90,7 +90,7 @@ struct FEmitter::FPrivate
 };
 
 // Creates a new `FDimensional` value of specified `Type` and returns it.
-static FDimensional* NewDimensionalValue(FEmitter* Emitter, FPrimitiveTypePtr Type)
+static FDimensional* NewDimensionalValue(FEmitter* Emitter, const FPrimitiveType* Type)
 {
 	check(!Type->IsScalar());
 
@@ -121,7 +121,7 @@ static FValue* EmitNew(FEmitter* Emitter, FValue* Value)
 }
 
 template <typename T>
-static T MakePrototype(FTypePtr InType)
+static T MakePrototype(const FType* InType)
 {
 	static_assert(std::is_trivially_constructible_v<T> && std::is_trivially_destructible_v<T> &&  std::is_trivially_copy_constructible_v<T> && std::is_trivially_copy_assignable_v<T>,
 		"FValues are expected to be trivial types.");
@@ -279,7 +279,7 @@ FValue* FEmitter::InputScalar(const FExpressionInput* InInput)
 		return Value;
 	}
 
-	FPrimitiveTypePtr ScalarType = Value->Type->AsPrimitive()->ToScalar();
+	const FPrimitiveType* ScalarType = Value->Type->AsPrimitive()->ToScalar();
 	return EmitConstruct(ScalarType, Value);
 }
 
@@ -344,7 +344,7 @@ void FEmitter::CheckInputIsScalar(const FExpressionInput* InInput, FValue* Input
 		return;
 	}
 
-	FPrimitiveTypePtr PrimitiveType = InputValue->Type->AsPrimitive();
+	const FPrimitiveType* PrimitiveType = InputValue->Type->AsPrimitive();
 	if (!PrimitiveType || !PrimitiveType->IsScalar())
 	{
 		Errorf(TEXT("Input '%s' expected to be a scalar. It is a %s instead."), *SlowFindInputName(Expression, InInput).ToString(), InputValue->Type->GetSpelling().GetData());
@@ -358,7 +358,7 @@ void FEmitter::CheckInputIsScalar(const FExpressionInput* InInput, FValue* Input
 		return;
 	}
 
-	FPrimitiveTypePtr PrimitiveType = InputValue->Type->AsPrimitive();
+	const FPrimitiveType* PrimitiveType = InputValue->Type->AsPrimitive();
 	if (!PrimitiveType || !PrimitiveType->IsScalar() || PrimitiveType->ScalarKind != Kind)
 	{
 		Errorf(TEXT("Input '%s' expected to be a %s scalar. It is %s instead."), *SlowFindInputName(Expression, InInput).ToString(), ScalarKindToString(Kind), InputValue->Type->GetSpelling().GetData());
@@ -385,7 +385,7 @@ void FEmitter::CheckValueIsScalar(const FValue* Value, EScalarKind Kind)
 		return;
 	}
 
-	FPrimitiveTypePtr PrimitiveType = Value->Type->AsPrimitive();
+	const FPrimitiveType* PrimitiveType = Value->Type->AsPrimitive();
 	if (!PrimitiveType || !PrimitiveType->IsScalar() || PrimitiveType->ScalarKind != Kind)
 	{
 		Errorf(TEXT("Value expected to be a %s scalar. It is %s instead."), ScalarKindToString(Kind), Value->Type->GetSpelling().GetData());
@@ -526,7 +526,7 @@ FValue* FEmitter::EmitVector2(FValue* InX, FValue* InY)
 	check(InX->Type == InY->Type);
 
 	TDimensional<2> Vector = MakePrototype<TDimensional<2>>(FPrimitiveType::GetVector(InX->Type->AsPrimitive()->ScalarKind, 2));
-	TArrayView<FValue*> Components = Vector.GetComponents();
+	TArrayView<FValue*> Components = Vector.GetMutableComponents();
 	Components[0] = InX;
 	Components[1] = InY;
 
@@ -545,7 +545,7 @@ FValue* FEmitter::EmitVector3(FValue* InX, FValue* InY, FValue* InZ)
 	check(InY->Type == InZ->Type);
 
 	TDimensional<3> Vector = MakePrototype<TDimensional<3>>(FPrimitiveType::GetVector(InX->Type->AsPrimitive()->ScalarKind, 3));
-	TArrayView<FValue*> Components = Vector.GetComponents();
+	TArrayView<FValue*> Components = Vector.GetMutableComponents();
 	Components[0] = InX;
 	Components[1] = InY;
 	Components[2] = InZ;
@@ -566,7 +566,7 @@ FValue* FEmitter::EmitVector4(FValue* InX, FValue* InY, FValue* InZ, FValue* InW
 	check(InZ->Type == InW->Type);
 
 	TDimensional<4> Vector = MakePrototype<TDimensional<4>>(FPrimitiveType::GetVector(InX->Type->AsPrimitive()->ScalarKind, 4));
-	TArrayView<FValue*> Components = Vector.GetComponents();
+	TArrayView<FValue*> Components = Vector.GetMutableComponents();
 	Components[0] = InX;
 	Components[1] = InY;
 	Components[2] = InZ;
@@ -589,7 +589,7 @@ FValue* FEmitter::EmitSubscript(FValue* Value, int32 Index)
 		return Value;
 	}
 
-	FPrimitiveTypePtr PrimitiveType = Value->Type->AsPrimitive();
+	const FPrimitiveType* PrimitiveType = Value->Type->AsPrimitive();
 	if (!PrimitiveType)
 	{
 		Errorf(TEXT("Value of type `%s` cannot be subscripted."), Value->Type->GetSpelling().GetData());
@@ -640,7 +640,7 @@ FValue* FEmitter::EmitSwizzle(FValue* Value, FSwizzleMask Mask)
 	check(Mask.NumComponents > 0);
 
 	// We can only swizzle on non-matrix primitive types.
-	FPrimitiveTypePtr PrimitiveType = Value->Type->AsVector();
+	const FPrimitiveType* PrimitiveType = Value->Type->AsVector();
 	if (!PrimitiveType || PrimitiveType->IsMatrix())
 	{
 		Errorf(TEXT("Cannot swizzle a `%s` value."), Value->Type->GetSpelling().GetData());
@@ -684,15 +684,16 @@ FValue* FEmitter::EmitSwizzle(FValue* Value, FSwizzleMask Mask)
 	}
 
 	// Make the result vector type.
-	FPrimitiveTypePtr ResultType = FPrimitiveType::GetVector(PrimitiveType->ScalarKind, Mask.NumComponents);
+	const FPrimitiveType* ResultType = FPrimitiveType::GetVector(PrimitiveType->ScalarKind, Mask.NumComponents);
 	FDimensional* Result = NewDimensionalValue(this, ResultType);
+	TArrayView<FValue*> ResultComponents = Result->GetMutableComponents();
 
 	for (int32 i = 0; i < Mask.NumComponents; ++i)
 	{
-		Result->GetComponents()[i] = EmitSubscript(Value, (int32)Mask.Components[i]);
+		ResultComponents[i] = EmitSubscript(Value, (int32)Mask.Components[i]);
 	}
 
-	return Result;
+	return EmitNew(this, Result);
 }
 
 FValue* FEmitter::EmitTextureObject(UTexture* Texture, EMaterialSamplerType SamplerType)
@@ -1043,7 +1044,7 @@ static FValue* TrySimplifyOperator(FEmitter* Emitter, EOperator& Op, FValue*& A,
 // It returns either the result of the operator or null if it could not be folded.
 static FValue* TryFoldOperatorScalar(FEmitter* Emitter, EOperator Op, FValue* A, FValue* B, FValue* C)
 {
-	FPrimitiveTypePtr PrimitiveType = A->Type->AsPrimitive();
+	const FPrimitiveType* PrimitiveType = A->Type->AsPrimitive();
 
 	// Try to simplify the operator. This could potentially change Op, A, B and C.
 	if (FValue* Simplified = TrySimplifyOperator(Emitter, Op, A, B, C))
@@ -1219,7 +1220,7 @@ static const FOperatorSignature* GetOperatorSignature(EOperator Op)
  
 /// Validates that the types of the arguments are valid for specified operator.
 /// If valid, it returns the type of the result. Otherwise if it is not valid, it returns nullptr.
-static FPrimitiveTypePtr ValidateOperatorAndGetResultType(FEmitter* Emitter, EOperator Op, FValue*& A, FValue*& B, FValue*& C)
+static const FPrimitiveType* ValidateOperatorAndGetResultType(FEmitter* Emitter, EOperator Op, FValue*& A, FValue*& B, FValue*& C)
 {
 	// Argument A must have always been provided.
 	check(A);
@@ -1235,21 +1236,21 @@ static FPrimitiveTypePtr ValidateOperatorAndGetResultType(FEmitter* Emitter, EOp
 
 	// Verify that the first argument type is primitive.
 	FValue* Arguments[] = { A, B, C, nullptr };
-	FPrimitiveTypePtr FirstArgumentPrimitiveType = A->Type->AsPrimitive();
+	const FPrimitiveType* FirstArgumentPrimitiveType = A->Type->AsPrimitive();
 	static const TCHAR* ArgumentsStr[] = { TEXT("first"), TEXT("second"), TEXT("third") };
 	const FOperatorSignature* Signature = GetOperatorSignature(Op);
-	FTypePtr ArgumentsCommonType = FirstArgumentPrimitiveType;
+	const FType* ArgumentsCommonType = FirstArgumentPrimitiveType;
 
 	bool bValid = true;
 	for (int32 i = 0; Arguments[i]; ++i)
 	{
 		// Cast this argument type to primitive.
-		FPrimitiveTypePtr ArgumentPrimitiveType = Arguments[i]->Type->AsPrimitive();
+		const FPrimitiveType* ArgumentPrimitiveType = Arguments[i]->Type->AsPrimitive();
 
 		// All operators have primitive parameter types.
 		if (!ArgumentPrimitiveType)
 		{
-			Emitter->Errorf(TEXT("Operator %s's %s argument does not have primitive type."), OperatorToString(Op), ArgumentsStr[i]);
+			Emitter->Errorf(TEXT("Operator %s's %s argument does not have primitive type."), LexToString(Op), ArgumentsStr[i]);
 			return nullptr;
 		}
 		
@@ -1269,7 +1270,7 @@ static FPrimitiveTypePtr ValidateOperatorAndGetResultType(FEmitter* Emitter, EOp
 		{
 			if (ArgumentPrimitiveType->ScalarKind != ScalarKind_Bool)
 			{
-				Emitter->Errorf(TEXT("Operator %s's %s argument component types are not boolean."), OperatorToString(Op), ArgumentsStr[i]);
+				Emitter->Errorf(TEXT("Operator %s's %s argument component types are not boolean."), LexToString(Op), ArgumentsStr[i]);
 				bValid = false;
 			}
 		}
@@ -1278,7 +1279,7 @@ static FPrimitiveTypePtr ValidateOperatorAndGetResultType(FEmitter* Emitter, EOp
 		{
 			if (!ScalarKindIsArithmetic(ArgumentPrimitiveType->ScalarKind))
 			{
-				Emitter->Errorf(TEXT("Operator %s's %s argument component types expected to be arithmetic (int, float, double, etc)."), OperatorToString(Op), ArgumentsStr[i]);
+				Emitter->Errorf(TEXT("Operator %s's %s argument component types expected to be arithmetic (int, float, double, etc)."), LexToString(Op), ArgumentsStr[i]);
 				bValid = false;
 			}
 		}
@@ -1287,7 +1288,7 @@ static FPrimitiveTypePtr ValidateOperatorAndGetResultType(FEmitter* Emitter, EOp
 		{
 			if (!ScalarKindIsAnyFloat(ArgumentPrimitiveType->ScalarKind))
 			{
-				Emitter->Errorf(TEXT("Operator %s's %s argument component types expected to be floating point (float, double, etc)."), OperatorToString(Op), ArgumentsStr[i]);
+				Emitter->Errorf(TEXT("Operator %s's %s argument component types expected to be floating point (float, double, etc)."), LexToString(Op), ArgumentsStr[i]);
 				bValid = false;
 			}
 		}
@@ -1296,7 +1297,7 @@ static FPrimitiveTypePtr ValidateOperatorAndGetResultType(FEmitter* Emitter, EOp
 		{
 			if (ArgumentPrimitiveType->IsMatrix())
 			{
-				Emitter->Errorf(TEXT("Operator %s's %s argument type expected to be a scalar or vector."), OperatorToString(Op), ArgumentsStr[i]);
+				Emitter->Errorf(TEXT("Operator %s's %s argument type expected to be a scalar or vector."), LexToString(Op), ArgumentsStr[i]);
 				bValid = false;
 			}
 		}
@@ -1305,7 +1306,7 @@ static FPrimitiveTypePtr ValidateOperatorAndGetResultType(FEmitter* Emitter, EOp
 		{
 			if (!ArgumentPrimitiveType->IsVector() || ArgumentPrimitiveType->GetNumComponents() != 3)
 			{
-				Emitter->Errorf(TEXT("Operator %s's %s argument type expected to be a 3D vector."), OperatorToString(Op), ArgumentsStr[i]);
+				Emitter->Errorf(TEXT("Operator %s's %s argument type expected to be a 3D vector."), LexToString(Op), ArgumentsStr[i]);
 				bValid = false;
 			}
 		}
@@ -1372,7 +1373,7 @@ static bool IsComponentwiseOperator(EOperator Op)
 /// If a value is returned, it will be a dimensional with some component folded to a constant. If some argument
 /// isn't a dimensional, or all arguments components are non-constant, the folding will not be carried out.
 /// If no folding is carried out, this function simply returns nullptr.
-static FValue* TryFoldComponentwiseOperator(FEmitter* Emitter, EOperator Op, FValue* A, FValue* B, FValue* C, FPrimitiveTypePtr ResultType)
+static FValue* TryFoldComponentwiseOperator(FEmitter* Emitter, EOperator Op, FValue* A, FValue* B, FValue* C, const FPrimitiveType* ResultType)
 {
 	// Check that at least one component of the resulting dimensional value would folded.
 	// If all components of resulting dimensional value are not folded, then instead of emitting
@@ -1440,10 +1441,10 @@ static FValue* TryFoldComponentwiseOperator(FEmitter* Emitter, EOperator Op, FVa
 		FDimensional* Result = NewDimensionalValue(Emitter, ResultType);
 
 		// Fetch the components array from the result dimensional
-		TArrayView<FValue*> ResultComponents = Result->GetComponents();
+		TArrayView<FValue*> ResultComponents = Result->GetMutableComponents();
 
 		// Also cache the type of a single component
-		FPrimitiveTypePtr ComponentType = ResultType->ToScalar();
+		const FPrimitiveType* ComponentType = ResultType->ToScalar();
 
 		// Create the operator instruction for each component pair
 		for (int32 i = 0; i < ResultType->GetNumComponents(); ++i)
@@ -1462,8 +1463,8 @@ static FValue* TryFoldComponentwiseOperator(FEmitter* Emitter, EOperator Op, FVa
 				ResultComponents[i] = EmitPrototype(Emitter, Proto);
 			}
 		}
-
-		return Result;
+		
+		return EmitNew(Emitter, Result);
 	}
 
 	return nullptr;
@@ -1480,7 +1481,7 @@ static bool TryUnpackConstantScalarOrVector(FValue* V, TArrayView<FConstant*> Ou
 		return false;
 	}
 
-	TArrayView<FValue*> Components = Dimensional->GetComponents();
+	TConstArrayView<FValue*> Components = Dimensional->GetComponents();
 	for (int32 i = 0; i < Components.Num(); ++i)
 	{
 		OutComponents[i] = As<FConstant>(Components[i]);
@@ -1507,7 +1508,7 @@ static MIR::TFloat ConstantDotFloat(TArrayView<FConstant*> AComponents, TArrayVi
 
 /// Tries to fold the operator, that is to evaluate its result now at translation time if its arguments are constant.
 /// If the operator could not be folded in any way, it returns nullptr.
-static FValue* TryFoldOperator(FEmitter* Emitter, EOperator Op, FValue* A, FValue* B, FValue* C, FPrimitiveTypePtr ResultType)
+static FValue* TryFoldOperator(FEmitter* Emitter, EOperator Op, FValue* A, FValue* B, FValue* C, const FPrimitiveType* ResultType)
 {
 	FConstant* AComponents[4];
 	int32      ANumComponents;
@@ -1583,7 +1584,7 @@ FValue* FEmitter::EmitOperator(EOperator Op, FValue* A, FValue* B, FValue* C)
 	}
 
 	// Validate the operation and retrieve the result type.
-	FPrimitiveTypePtr ResultType = ValidateOperatorAndGetResultType(this, Op, A, B, C);
+	const FPrimitiveType* ResultType = ValidateOperatorAndGetResultType(this, Op, A, B, C);
 	if (!ResultType)
 	{
 		return FPoison::Get();
@@ -1631,7 +1632,7 @@ FValue* FEmitter::EmitBranch(FValue* Condition, FValue* True, FValue* False)
 
 	// If the condition is not static, make both true and false arguments have the same type,
 	// by casting false argument into the true's type.
-	FTypePtr CommonType = GetCommonType(True->Type, False->Type);
+	const FType* CommonType = GetCommonType(True->Type, False->Type);
 	if (!CommonType)
 	{
 		return FPoison::Get();
@@ -1690,9 +1691,9 @@ static FValue* CastConstant(FEmitter* Emitter, FConstant* Constant, EScalarKind
 	UE_MIR_UNREACHABLE();
 }
 
-static FValue* ConstructPrimitiveValue(FEmitter* Emitter, FPrimitiveTypePtr TargetPrimitiveType, FValue* Initializer)
+static FValue* ConstructPrimitiveValue(FEmitter* Emitter, const FPrimitiveType* TargetPrimitiveType, FValue* Initializer)
 {
-	FPrimitiveTypePtr InitializerPrimitiveType = Initializer->Type->AsPrimitive();
+	const FPrimitiveType* InitializerPrimitiveType = Initializer->Type->AsPrimitive();
 	if (!InitializerPrimitiveType)
 	{
 		Emitter->Errorf(TEXT("Cannot construct a '%s' from non primitive type '%s'."), TargetPrimitiveType->Spelling.GetData(), Initializer->Type->GetSpelling().GetData());
@@ -1734,10 +1735,13 @@ static FValue* ConstructPrimitiveValue(FEmitter* Emitter, FPrimitiveTypePtr Targ
 		// of initializer value to the single component type.
 		FValue* Component = Emitter->EmitConstruct(TargetPrimitiveType->ToScalar(), Initializer);
 
+		// Get the mutable array of components.
+		TArrayView<FValue*> ResultComponents = Result->GetMutableComponents();
+
 		// Initialize all result components to the same scalar.
 		for (int32 i = 0; i < TargetPrimitiveType->GetNumComponents(); ++i)
 		{
-			Result->GetComponents()[i] = Component;
+			ResultComponents[i] = Component;
 		}
 		
 		return EmitNew(Emitter, Result);
@@ -1753,24 +1757,25 @@ static FValue* ConstructPrimitiveValue(FEmitter* Emitter, FPrimitiveTypePtr Targ
 
 		// Create the result dimensional value.
 		FDimensional* Result = NewDimensionalValue(Emitter, TargetPrimitiveType);
+		TArrayView<FValue*> ResultComponents = Result->GetMutableComponents();
 
 		// Determine the result component type (scalar).
-		FPrimitiveTypePtr ResultComponentType = TargetPrimitiveType->ToScalar();
+		const FPrimitiveType* ResultComponentType = TargetPrimitiveType->ToScalar();
 
 		// For iterating over the components of the result dimensional value.
 		int32 Index = 0;
 
 		// Convert components from the initializer vector.
-		int32 MinNumComponents = FMath::Min(TargetNumComponents, InitializerNumComponents);
+		const int32 MinNumComponents = FMath::Min(TargetNumComponents, InitializerNumComponents);
 		for (; Index < MinNumComponents; ++Index)
 		{
-			Result->GetComponents()[Index] = Emitter->EmitConstruct(ResultComponentType, Emitter->EmitSubscript(Initializer, Index));
+			ResultComponents[Index] = Emitter->EmitConstruct(ResultComponentType, Emitter->EmitSubscript(Initializer, Index));
 		}
 
 		// Initialize remaining result dimensional components to zero.
 		for (; Index < TargetNumComponents; ++Index)
 		{
-			Result->GetComponents()[Index] = Emitter->EmitConstantScalarZero(ResultComponentType->ScalarKind);
+			ResultComponents[Index] = Emitter->EmitConstantScalarZero(ResultComponentType->ScalarKind);
 		}
 
 		return EmitNew(Emitter, Result);
@@ -1787,14 +1792,15 @@ static FValue* ConstructPrimitiveValue(FEmitter* Emitter, FPrimitiveTypePtr Targ
 		{
 			// Create the result dimensional value.
 			FDimensional* Result = NewDimensionalValue(Emitter, TargetPrimitiveType);
+			TArrayView<FValue*> ResultComponents = Result->GetMutableComponents();
 			
 			// Determine the result component type (scalar).
-			FPrimitiveTypePtr ResultComponentType = TargetPrimitiveType->ToScalar();
+			const FPrimitiveType* ResultComponentType = TargetPrimitiveType->ToScalar();
 
 			// Convert components from the initializer vector.
 			for (int32 Index = 0, Num = TargetPrimitiveType->GetNumComponents(); Index < Num; ++Index)
 			{
-				Result->GetComponents()[Index] = Emitter->EmitConstruct(ResultComponentType, DimensionalInitializer->GetComponents()[Index]);
+				ResultComponents[Index] = Emitter->EmitConstruct(ResultComponentType, DimensionalInitializer->GetComponents()[Index]);
 			}
 
 			return EmitNew(Emitter, Result);
@@ -1812,17 +1818,17 @@ static FValue* ConstructPrimitiveValue(FEmitter* Emitter, FPrimitiveTypePtr Targ
 	return MIR::FPoison::Get();
 }
 
-FValue* FEmitter::EmitConstruct(FTypePtr Type, FValue* Initializer)
+FValue* FEmitter::EmitConstruct(const FType* Type, FValue* Initializer)
 {
 	// If target type matches initializer's, simply return the same value.
-	FTypePtr InitializerType = Initializer->Type;
+	const FType* InitializerType = Initializer->Type;
 	if (InitializerType == Type)
 	{
 		return Initializer;
 	}
 	
 	FValue* Result{};
-	if (FPrimitiveTypePtr PrimitiveType = Type->AsPrimitive())
+	if (const FPrimitiveType* PrimitiveType = Type->AsPrimitive())
 	{
 		Result = ConstructPrimitiveValue(this, PrimitiveType, Initializer);
 	}
@@ -1851,7 +1857,7 @@ FValue* FEmitter::EmitTextureSample(FValue* Texture, FValue* TexCoord, ESamplerS
 	return EmitPrototype(this, Prototype);
 }
 
-FValue* FEmitter::EmitInlineHLSL(FTypePtr Type, FString Code, EValueFlags ValueFlags, EGraphProperties UsedGraphProperties)
+FValue* FEmitter::EmitInlineHLSL(const FType* Type, FString Code, EValueFlags ValueFlags, EGraphProperties UsedGraphProperties)
 {
 	MIR::FInlineHLSL Prototype = MakePrototype<MIR::FInlineHLSL>(Type);
 	Prototype.Flags = ValueFlags;
@@ -1863,11 +1869,11 @@ FValue* FEmitter::EmitInlineHLSL(FTypePtr Type, FString Code, EValueFlags ValueF
 FValue* FEmitter::EmitInlineHLSL(const FMaterialExternalCodeDeclaration* InExternalCodeDeclaration)
 {
 	check(InExternalCodeDeclaration != nullptr);
-	FTypePtr ReturnType = FType::FromMaterialValueType(InExternalCodeDeclaration->GetReturnTypeValue());
+	const FType* ReturnType = FType::FromMaterialValueType(InExternalCodeDeclaration->GetReturnTypeValue());
 	return EmitInlineHLSL(ReturnType, InExternalCodeDeclaration->Definition, {}, {});
 }
 
-FTypePtr FEmitter::TryGetCommonType(FTypePtr A, FTypePtr B)
+const FType* FEmitter::TryGetCommonType(const FType* A, const FType* B)
 {
 	// Trivial case: types are equal
 	if (A == B)
@@ -1875,8 +1881,8 @@ FTypePtr FEmitter::TryGetCommonType(FTypePtr A, FTypePtr B)
 		return A;
 	}
 
-	FPrimitiveTypePtr PrimitiveA = A->AsPrimitive();
-	FPrimitiveTypePtr PrimitiveB = B->AsPrimitive();
+	const FPrimitiveType* PrimitiveA = A->AsPrimitive();
+	const FPrimitiveType* PrimitiveB = B->AsPrimitive();
 	if (!PrimitiveA || !PrimitiveB)
 	{
 		return nullptr;
@@ -1895,9 +1901,9 @@ FTypePtr FEmitter::TryGetCommonType(FTypePtr A, FTypePtr B)
 	return FPrimitiveType::Get(ScalarKind, NumRows, 1);
 }
 
-FTypePtr FEmitter::GetCommonType(FTypePtr A, FTypePtr B)
+const FType* FEmitter::GetCommonType(const FType* A, const FType* B)
 {
-	if (FTypePtr CommonType = TryGetCommonType(A, B))
+	if (const FType* CommonType = TryGetCommonType(A, B))
 	{
 		return CommonType;
 	}
diff --git a/Engine/Source/Runtime/Engine/Private/Materials/MaterialIRModule.cpp b/Engine/Source/Runtime/Engine/Private/Materials/MaterialIRModule.cpp
--- a/Engine/Source/Runtime/Engine/Private/Materials/MaterialIRModule.cpp
+++ b/Engine/Source/Runtime/Engine/Private/Materials/MaterialIRModule.cpp
@@ -40,7 +40,7 @@ void FMaterialIRModule::Empty()
 	Values.Empty();
 
 	// Reset module statistics.
-	for (int i = 0; i < (int)SF_NumFrequencies; ++i)
+	for (int i = 0; i < MIR::NumStages; ++i)
 	{
 		Statistics.ExternalInputUsedMask[i].Init(false, (int)MIR::EExternalInput::Count);
 	}
diff --git a/Engine/Source/Runtime/Engine/Private/Materials/MaterialIRModuleBuilder.cpp b/Engine/Source/Runtime/Engine/Private/Materials/MaterialIRModuleBuilder.cpp
--- a/Engine/Source/Runtime/Engine/Private/Materials/MaterialIRModuleBuilder.cpp
+++ b/Engine/Source/Runtime/Engine/Private/Materials/MaterialIRModuleBuilder.cpp
@@ -227,7 +227,7 @@ struct FMaterialIRModuleBuilderImpl
 			MIR::FValue* Value = Emitter->Input(Call->GetInput(i));
 			if (Value)
 			{
-				MIR::FTypePtr Type = MIR::FType::FromMaterialValueType((EMaterialValueType)FunctionInput.ExpressionInput->GetInputType(0));
+				const MIR::FType* Type = MIR::FType::FromMaterialValueType((EMaterialValueType)FunctionInput.ExpressionInput->GetInputType(0));
 				CallInputValues[i] = Emitter->EmitConstruct(Type, Value);
 			}
 		}
@@ -284,7 +284,7 @@ struct FMaterialIRModuleBuilderImpl
 			MIR::FValue* Value = Emitter->Input(FunctionOutput.ExpressionOutput->GetInput(0));
 
 			// Get the function output type.
-			MIR::FTypePtr OutputType = MIR::FType::FromMaterialValueType((EMaterialValueType)FunctionOutput.ExpressionOutput->GetOutputType(0));
+			const MIR::FType* OutputType = MIR::FType::FromMaterialValueType((EMaterialValueType)FunctionOutput.ExpressionOutput->GetOutputType(0));
 
 			// Cast the value to the expected output type. This may fail (value will be poison). 
 			Value = Emitter->EmitConstruct(OutputType, Value);
@@ -318,7 +318,7 @@ struct FMaterialIRModuleBuilderImpl
 
 					MIR::Internal::SetInputValue(this, Input.Input, *ValuePtr);
 
-					MIR::FTypePtr OutputArgType = MIR::FType::FromShaderType(Input.Type);
+					const MIR::FType* OutputArgType = MIR::FType::FromShaderType(Input.Type);
 					Output->Arg = Emitter->EmitConstruct(OutputArgType, *ValuePtr);
 				}
 
@@ -429,7 +429,7 @@ struct FMaterialIRModuleBuilderImpl
 				Instr->Next[StageIndex] = Instr->Block[StageIndex] ->Instructions;
 				Instr->Block[StageIndex]->Instructions = Instr;
 
-				TArrayView<MIR::FValue*> Uses = Instr->GetUses();
+				TConstArrayView<MIR::FValue*> Uses = Instr->GetUses();
 				for (int32 UseIndex = 0; UseIndex < Uses.Num(); ++UseIndex)
 				{
 					MIR::FValue* Use = Uses[UseIndex];
@@ -475,11 +475,11 @@ struct FMaterialIRModuleBuilderImpl
 		for (int TexCoordIndex = 0; TexCoordIndex < MIR::TexCoordMaxNum; ++TexCoordIndex)
 		{
 			MIR::EExternalInput TexCoordInput = MIR::TexCoordIndexToExternalInput(TexCoordIndex);
-			if (Module->Statistics.ExternalInputUsedMask[SF_Vertex][(int)TexCoordInput])
+			if (Module->Statistics.ExternalInputUsedMask[MIR::Stage_Vertex][(int)TexCoordInput])
 			{
 				Module->Statistics.NumVertexTexCoords = TexCoordIndex + 1;
 			}
-			if (Module->Statistics.ExternalInputUsedMask[SF_Pixel][(int)TexCoordInput])
+			if (Module->Statistics.ExternalInputUsedMask[MIR::Stage_Pixel][(int)TexCoordInput])
 			{
 				Module->Statistics.NumPixelTexCoords = TexCoordIndex + 1;
 			}
@@ -514,7 +514,7 @@ struct FMaterialIRModuleBuilderImpl
 	}
 	
 	//
-	void PushConnectionInsight(const UObject* InputObject, int InputIndex, const UMaterialExpression* OutputExpression, int OutputIndex, MIR::FTypePtr Type)
+	void PushConnectionInsight(const UObject* InputObject, int InputIndex, const UMaterialExpression* OutputExpression, int OutputIndex, const MIR::FType* Type)
 	{
 		FMaterialInsights::FConnectionInsight Insight;
 		Insight.InputObject = InputObject,
diff --git a/Engine/Source/Runtime/Engine/Private/Materials/MaterialIRToHLSLTranslator.cpp b/Engine/Source/Runtime/Engine/Private/Materials/MaterialIRToHLSLTranslator.cpp
--- a/Engine/Source/Runtime/Engine/Private/Materials/MaterialIRToHLSLTranslator.cpp
+++ b/Engine/Source/Runtime/Engine/Private/Materials/MaterialIRToHLSLTranslator.cpp
@@ -238,7 +238,7 @@ struct FTranslator : FMaterialIRToHLSLTranslation
 		for (int32 PropertyIndex = 0; PropertyIndex < MP_MAX; ++PropertyIndex)
 		{
 			EMaterialProperty Property = (EMaterialProperty)PropertyIndex;
-			if (!MIR::Internal::IsMaterialPropertyEnabled(Property) || MIR::GetMaterialOutputStage(Property) != MIR::Stage_Fragment)
+			if (!MIR::Internal::IsMaterialPropertyEnabled(Property) || MIR::GetMaterialOutputStage(Property) != MIR::Stage_Pixel)
 			{
 				continue;
 			}
@@ -264,8 +264,8 @@ struct FTranslator : FMaterialIRToHLSLTranslation
 	}
 
 	void GenerateFragmentStageHLSL()
-		{
-		BeginStage(MIR::Stage_Fragment);
+	{
+		BeginStage(MIR::Stage_Pixel);
 
 		LowerBlock(Module->GetRootBlock(CurrentStage));
 
@@ -273,7 +273,7 @@ struct FTranslator : FMaterialIRToHLSLTranslation
 			Printer << TEXT("PixelMaterialInputs.Subsurface = 0") << EndOfStatement;
 
 		EvaluateOtherMaterialAttributesHLSL = MoveTemp(Printer.Buffer);
-		}
+	}
 
 	void BeginStage(MIR::EStage Stage)
 	{
@@ -381,7 +381,7 @@ struct FTranslator : FMaterialIRToHLSLTranslation
 
 	void LowerConstant(const MIR::FConstant* Constant)
 	{
-		MIR::FPrimitiveTypePtr checkNoEntry = Constant->Type->AsPrimitive();
+		const MIR::FPrimitiveType* checkNoEntry = Constant->Type->AsPrimitive();
 		check(checkNoEntry && checkNoEntry->IsScalar());
 
 		switch (checkNoEntry->ScalarKind)
@@ -457,7 +457,7 @@ struct FTranslator : FMaterialIRToHLSLTranslation
 
 	void LowerPrimitiveUniformParameter(const MIR::FUniformParameter* UniformParameter)
 	{
-		MIR::FPrimitiveTypePtr PrimitiveType = UniformParameter->Type->AsPrimitive();
+		const MIR::FPrimitiveType* PrimitiveType = UniformParameter->Type->AsPrimitive();
 		check(PrimitiveType);
 		check(PrimitiveType->IsScalar() || PrimitiveType->IsVector()); // no matrices yet
 
@@ -500,7 +500,7 @@ struct FTranslator : FMaterialIRToHLSLTranslation
 	void LowerDimensional(const MIR::FDimensional* Dimensional) 
 	{
 		TArrayView<MIR::FValue* const> Components = Dimensional->GetComponents();
-		MIR::FPrimitiveTypePtr ArithmeticType = Dimensional->Type->AsPrimitive();
+		const MIR::FPrimitiveType* ArithmeticType = Dimensional->Type->AsPrimitive();
 		check(ArithmeticType && !ArithmeticType->IsScalar());
 		
 		// In order to generate smaller and tidier HLSL, first check whether all components
@@ -656,7 +656,7 @@ struct FTranslator : FMaterialIRToHLSLTranslation
 	{
 		LowerValue(Subscript->Arg);
 
-		if (MIR::FPrimitiveTypePtr ArgArithmeticType = Subscript->Arg->Type->AsVector())
+		if (const MIR::FPrimitiveType* ArgArithmeticType = Subscript->Arg->Type->AsVector())
 		{
 			const TCHAR* ComponentsStr[] = { TEXT(".x"), TEXT(".y"), TEXT(".z"), TEXT(".w") };
 			check(Subscript->Index <= ArgArithmeticType->GetNumComponents());
diff --git a/Engine/Source/Runtime/Engine/Private/Materials/MaterialIRTypes.cpp b/Engine/Source/Runtime/Engine/Private/Materials/MaterialIRTypes.cpp
--- a/Engine/Source/Runtime/Engine/Private/Materials/MaterialIRTypes.cpp
+++ b/Engine/Source/Runtime/Engine/Private/Materials/MaterialIRTypes.cpp
@@ -17,7 +17,7 @@ const TCHAR* TypeKindToString(ETypeKind Kind)
 	}
 }
 
-FTypePtr FType::FromShaderType(const UE::Shader::FType& InShaderType)
+const FType* FType::FromShaderType(const UE::Shader::FType& InShaderType)
 {
 	check(!InShaderType.IsStruct());
 	check(!InShaderType.IsObject());
@@ -50,7 +50,7 @@ FTypePtr FType::FromShaderType(const UE::Shader::FType& InShaderType)
 	}
 }
 
-FTypePtr FType::FromMaterialValueType(EMaterialValueType Type)
+const FType* FType::FromMaterialValueType(EMaterialValueType Type)
 {
 	switch (Type)
 	{
@@ -86,7 +86,7 @@ FTypePtr FType::FromMaterialValueType(EMaterialValueType Type)
 	}
 }
 
-FTypePtr FType::FromMaterialParameterType(EMaterialParameterType Type)
+const FType* FType::FromMaterialParameterType(EMaterialParameterType Type)
 {
 	switch (Type)
 	{
@@ -103,7 +103,7 @@ FTypePtr FType::FromMaterialParameterType(EMaterialParameterType Type)
 	}
 }
 
-FTypePtr FType::GetVoid()
+const FType* FType::GetVoid()
 {
 	static FType Type{ TypeKind_Void };
 	return &Type;
@@ -133,7 +133,7 @@ UE::Shader::EValueType FType::ToValueType() const
 {
 	using namespace UE::Shader;
 
-	if (FPrimitiveTypePtr PrimitiveType = AsPrimitive())
+	if (const FPrimitiveType* PrimitiveType = AsPrimitive())
 	{
 		if (PrimitiveType->IsMatrix())
 		{
@@ -162,7 +162,7 @@ UE::Shader::EValueType FType::ToValueType() const
 			default: UE_MIR_UNREACHABLE();
 		}
 	}
-	else if (FObjectTypePtr ObjectType = AsObject())
+	else if (const FObjectType* ObjectType = AsObject())
 	{
 		return EValueType::Object;
 	}
@@ -170,7 +170,7 @@ UE::Shader::EValueType FType::ToValueType() const
 	UE_MIR_UNREACHABLE();
 }
 
-FTypePtr FType::GetPoison()
+const FType* FType::GetPoison()
 {
 	static FType Type{ TypeKind_Poison };
 	return &Type;
@@ -183,13 +183,13 @@ bool FType::IsPoison() const
 
 bool FType::IsBoolean() const
 {
-	FPrimitiveTypePtr Ptr = AsPrimitive();
+	const FPrimitiveType* Ptr = AsPrimitive();
 	return Ptr && Ptr->ScalarKind == ScalarKind_Bool;
 }
 
 bool FType::IsInteger() const
 {
-	FPrimitiveTypePtr Ptr = AsPrimitive();
+	const FPrimitiveType* Ptr = AsPrimitive();
 	return Ptr && Ptr->ScalarKind == ScalarKind_Int;
 }
 
@@ -200,7 +200,7 @@ bool FType::IsBoolScalar() const
 
 bool FType::IsObjectOfKind(EObjectKind ObjectKind) const
 {
-	if (FObjectTypePtr ObjectPtr = AsObject())
+	if (const FObjectType* ObjectPtr = AsObject())
 	{
 		return ObjectPtr->ObjectKind == ObjectKind;
 	}
@@ -209,38 +209,38 @@ bool FType::IsObjectOfKind(EObjectKind ObjectKind) const
 
 bool FType::IsTexture() const
 {
-	if (FObjectTypePtr ObjectPtr = AsObject())
+	if (const FObjectType* ObjectPtr = AsObject())
 	{
 		return true;
 	}
 	return false; 
 }
 
-FPrimitiveTypePtr FType::AsPrimitive() const
+const FPrimitiveType* FType::AsPrimitive() const
 {
-	return Kind == TypeKind_Primitive ? static_cast<FPrimitiveTypePtr>(this) : nullptr; 
+	return Kind == TypeKind_Primitive ? static_cast<const FPrimitiveType*>(this) : nullptr; 
 }
 
-FObjectTypePtr FType::AsObject() const
+const FObjectType* FType::AsObject() const
 {
-	return Kind == TypeKind_Object ? static_cast<FObjectTypePtr>(this) : nullptr; 
+	return Kind == TypeKind_Object ? static_cast<const FObjectType*>(this) : nullptr; 
 }
 
-FPrimitiveTypePtr FType::AsScalar() const
+const FPrimitiveType* FType::AsScalar() const
 {
-	FPrimitiveTypePtr Type = AsPrimitive();
+	const FPrimitiveType* Type = AsPrimitive();
 	return Type->IsScalar() ? Type : nullptr;
 }
 
-FPrimitiveTypePtr FType::AsVector() const
+const FPrimitiveType* FType::AsVector() const
 {
-	FPrimitiveTypePtr Type = AsPrimitive();
+	const FPrimitiveType* Type = AsPrimitive();
 	return Type->IsVector() ? Type : nullptr;
 }
 
-FPrimitiveTypePtr FType::AsMatrix() const
+const FPrimitiveType* FType::AsMatrix() const
 {
-	FPrimitiveTypePtr Type = AsPrimitive();
+	const FPrimitiveType* Type = AsPrimitive();
 	return Type->IsMatrix() ? Type : nullptr;
 }
 
@@ -265,32 +265,32 @@ const TCHAR* ScalarKindToString(EScalarKind Kind)
 	}
 }
 
-FPrimitiveTypePtr FPrimitiveType::GetBool1()
+const FPrimitiveType* FPrimitiveType::GetBool1()
 {
 	return GetScalar(ScalarKind_Bool);
 }
 
-FPrimitiveTypePtr FPrimitiveType::GetInt1()
+const FPrimitiveType* FPrimitiveType::GetInt1()
 {
 	return GetScalar(ScalarKind_Int);
 }
 
-FPrimitiveTypePtr FPrimitiveType::GetFloat1()
+const FPrimitiveType* FPrimitiveType::GetFloat1()
 {
 	return GetScalar(ScalarKind_Float);
 }
 
-FPrimitiveTypePtr FPrimitiveType::GetFloat2()
+const FPrimitiveType* FPrimitiveType::GetFloat2()
 {
 	return GetVector(ScalarKind_Float, 2);
 }
 
-FPrimitiveTypePtr FPrimitiveType::GetFloat3()
+const FPrimitiveType* FPrimitiveType::GetFloat3()
 {
 	return GetVector(ScalarKind_Float, 3);
 }
 
-FPrimitiveTypePtr FPrimitiveType::GetFloat4()
+const FPrimitiveType* FPrimitiveType::GetFloat4()
 {
 	return GetVector(ScalarKind_Float, 4);
 }
@@ -313,7 +313,7 @@ const FPrimitiveType* FPrimitiveType::GetMatrix(EScalarKind InScalarKind, int Nu
 	return Get(InScalarKind, NumRows, NumColumns);
 }
 
-FPrimitiveTypePtr FPrimitiveType::Get(EScalarKind InScalarKind, int NumRows, int NumColumns)
+const FPrimitiveType* FPrimitiveType::Get(EScalarKind InScalarKind, int NumRows, int NumColumns)
 {
 	check(InScalarKind >= 0 && InScalarKind <= ScalarKind_Float);
 	
@@ -375,12 +375,12 @@ FPrimitiveTypePtr FPrimitiveType::Get(EScalarKind InScalarKind, int NumRows, int
 	return &Types[Index];
 }
 
-FPrimitiveTypePtr FPrimitiveType::ToScalar() const
+const FPrimitiveType* FPrimitiveType::ToScalar() const
 {
 	return FPrimitiveType::GetScalar(ScalarKind);
 }
 
-FTypePtr FObjectType::GetTexture2D()
+const FType* FObjectType::GetTexture2D()
 {
 	static const FObjectType Instance { { TypeKind_Object }, ObjectKind_Texture2D };
 	return &Instance;
diff --git a/Engine/Source/Runtime/Engine/Private/Materials/MaterialIRValueAnalyzer.cpp b/Engine/Source/Runtime/Engine/Private/Materials/MaterialIRValueAnalyzer.cpp
--- a/Engine/Source/Runtime/Engine/Private/Materials/MaterialIRValueAnalyzer.cpp
+++ b/Engine/Source/Runtime/Engine/Private/Materials/MaterialIRValueAnalyzer.cpp
@@ -16,8 +16,8 @@
 
 static void AnalyzeExternalInput(FMaterialIRValueAnalyzer& Analyzer, MIR::FExternalInput* ExternalInput)
 {
-	 Analyzer.Module->GetStatistics().ExternalInputUsedMask[SF_Vertex][(int)ExternalInput->Id] = true;
-	 Analyzer.Module->GetStatistics().ExternalInputUsedMask[SF_Pixel][(int)ExternalInput->Id] = true;
+	 Analyzer.Module->GetStatistics().ExternalInputUsedMask[MIR::Stage_Vertex][(int)ExternalInput->Id] = true;
+	 Analyzer.Module->GetStatistics().ExternalInputUsedMask[MIR::Stage_Pixel][(int)ExternalInput->Id] = true;
 }
 
 static void AnalyzeTextureObject(FMaterialIRValueAnalyzer& Analyzer, MIR::FTextureObject* TextureObject)
@@ -77,7 +77,7 @@ static void AnalyzePrimitiveUniformParameter(FMaterialIRValueAnalyzer& Analyzer,
 	Parameter->Analysis_UniformParameterIndex = UniformExpressionSet.FindOrAddNumericParameter(ParameterMetadata.Value.Type, ParameterInfo, DefaultValueOffset);
 
 	// Get the parameter type as primitive
-	MIR::FPrimitiveTypePtr ParameterType = Parameter->Type->AsPrimitive();
+	const MIR::FPrimitiveType* ParameterType = Parameter->Type->AsPrimitive();
 
 	// Only int and float parameters supported for now (TODO)
 	check(ParameterType->ScalarKind == MIR::ScalarKind_Int || ParameterType->ScalarKind == MIR::ScalarKind_Float);
@@ -191,7 +191,10 @@ void FMaterialIRValueAnalyzer::Analyze(MIR::FValue* Value)
 {
 	for (MIR::FValue* Use : Value->GetUses())
 	{
-		Value->GraphProperties |= Use->GraphProperties;
+		if (Use)
+		{
+			Value->GraphProperties |= Use->GraphProperties;
+		}
 	}
 
 	#define EXPAND_CASE(ValueType) \
diff --git a/Engine/Source/Runtime/Engine/Public/Materials/MaterialIR.h b/Engine/Source/Runtime/Engine/Public/Materials/MaterialIR.h
--- a/Engine/Source/Runtime/Engine/Public/Materials/MaterialIR.h
+++ b/Engine/Source/Runtime/Engine/Public/Materials/MaterialIR.h
@@ -8,14 +8,15 @@
 
 struct FMaterialExternalCodeDeclaration;
 
-namespace MIR {
+namespace MIR
+{
 
 // Identifies the block of execution in which an instruction runs.
 // Note: if you introduce a new stage, make sure you update EValueFlags accordingly.
 enum EStage
 {
 	Stage_Vertex,
-	Stage_Fragment,
+	Stage_Pixel,
 	NumStages
 };
 
@@ -30,11 +31,11 @@ enum class EValueFlags : uint8
 	// Value has been analyzed for the vertex stage.
 	AnalyzedForStageVertex = 1 << 0,
 
-	// Value has been analyzed for the fragment stage.
-	AnalyzedForStageFragment = 1 << 1,
+	// Value has been analyzed for the pixel stage.
+	AnalyzedForStagePixel = 1 << 1,
 
 	// Mask that includes all "analyzed in stage" flags.
-	AnalyzedInAnyStageMask = AnalyzedForStageVertex | AnalyzedForStageFragment,
+	AnalyzedInAnyStageMask = AnalyzedForStageVertex | AnalyzedForStagePixel,
 
 	// Whether to skip substution of abstract tokens like "<PREV>". Only true for custom nodes.
 	SubstituteTagsInInlineHLSL = 1 << 2,
@@ -88,7 +89,8 @@ enum EValueKind
 	VK_InstructionEnd,
 };
 
-const TCHAR* ValueKindToString(EValueKind Kind);
+// Returns the string representation of given value kind.
+const TCHAR* LexToString(EValueKind Kind);
 
 /* Values */
 
@@ -126,7 +128,7 @@ struct FValue
 	EGraphProperties GraphProperties;
 
 	// The runtime type this value has.
-	FTypePtr Type;
+	const FType* Type;
 
 	// Returns whether this value has been analyzed for specified stage.
 	bool IsAnalyzed(EStage State) const;
@@ -154,10 +156,7 @@ struct FValue
 
 	// Returns the array of this value's uses. An use is another value referenced by this 
 	// one (e.g. the operands of a binary expression).
-	TArrayView<const FValue*> GetUses() const;
-
-	// Returns the array of this value's uses. An use is another value referenced by this one (e.g. the operands of a binary expression).
-	TArrayView<FValue*> GetUses();
+	TConstArrayView<FValue*> GetUses() const;
 
 	// Returns whether this value is poison.
 	bool IsPoison() const;
@@ -240,6 +239,9 @@ struct FPoison : TValue<VK_Poison>
 using TInteger = int64_t;
 using TFloat = float;
 
+// A constant represents a translation-time known scalar primitive value. Operations on
+// constant values can be folded by the builder, that is they can be evaluated statically
+// while the builder constructs the IR graph of an input material.
 struct FConstant : TValue<VK_Constant>
 {
 	union
@@ -249,10 +251,16 @@ struct FConstant : TValue<VK_Constant>
 		TFloat 		Float;
 	};
 
+	// Returns whether this constant is a boolean.
 	bool IsBool() const;
+
+	// Returns whether this constant is an integer.
 	bool IsInteger() const;
+	
+	// Returns whether this constant is a float.
 	bool IsFloat() const;
 
+	// Returns the constant value of given type T. The type must be bool, integral or floating point.
 	template <typename T>
 	T Get() const
 	{
@@ -309,45 +317,48 @@ enum class EExternalInput
 	Count,
 };
 
-//
+// Maximum number of supported texture coordinates.
 static constexpr int TexCoordMaxNum = 8;
 
-//
-const TCHAR* ExternalInputToString(EExternalInput Input);
+// Converts the given external input to its string representation.
+const TCHAR* LexToString(EExternalInput Input);
 
-//
-FTypePtr GetExternalInputType(EExternalInput Id);
+// Returns the given external input type.
+const FType* GetExternalInputType(EExternalInput Id);
 
-//
+// Converts the given texture coordinate index to its corresponding external input mapping.
 MIR::EExternalInput TexCoordIndexToExternalInput(int TexCoordIndex);
 
-//
+// Returns whether the given external is a texture coordinate.
 bool IsExternalInputTexCoord(EExternalInput Id);
 
-//
+// Returns whether the given external is a texture coordinate ddx.
 bool IsExternalInputTexCoordDdx(EExternalInput Id);
 
-//
+// Returns whether the given external is a texture coordinate ddy.
 bool IsExternalInputTexCoordDdy(EExternalInput Id);
 
-//
+// It represents an external input (like texture coordinates) value.
 struct FExternalInput : TValue<VK_ExternalInput>
 {
 	EExternalInput Id;
 };
 
-//
+// It represents a texture object value.
 struct FTextureObject : TValue<VK_TextureObject>
 {
+	// The texture object.
 	UTexture* Texture;
 
+	// The sampler type associated to this texture object.
 	EMaterialSamplerType SamplerType;
 
-	// Index of this parameter in the uniform expression set
+	// Index of this parameter in the uniform expression set.
+	// Note: This field  is automatically set by the builder.
 	uint32 Analysis_UniformParameterIndex;
 };
 
-//
+// It represents a material uniform parameter.
 struct FUniformParameter: TValue<VK_UniformParameter>
 {
 	// Index of the parameter as registered in the module.
@@ -357,32 +368,59 @@ struct FUniformParameter: TValue<VK_UniformParameter>
 	EMaterialSamplerType SamplerType;
 
 	// Index of this parameter in the uniform expression set
+	// Note: This field is automatically set by the builder.
 	uint32 Analysis_UniformParameterIndex;
 };
 
 
 /* Instructions */
 
-//
+// A block is a sequence of instructions in order of execution. Blocks are organized in a
+// tree-like structure, used to model blocks nested inside other blocks.
 struct FBlock
 {
+	// This block's parent block, if any. If this is null, this is a *root block*.
 	FBlock* Parent;
+
+	// The linked-list head of the instructions contained in this block. Links are contained
+	// in the FInstruction::Next field.
 	FInstruction* Instructions;
+
+	// Depth of this block in the tree structure (root blocks have level zero).
 	int32 Level;
 
+	// Finds and returns the common block between this and Other, if any. It returns null
+	// if no common block was found.
+	// Note: this has O(n) complexity, where n is the maximum depth of the tree structure.
 	FBlock* FindCommonParentWith(MIR::FBlock* Other);
 };
 
-//
+// An instruction is a value in a well defined order of evaluation.
+// Instructions have a parent block and are organized in a linked list: the Next field
+// indicates the instruction that will execute immediately after this one.
+// Since the material shader has multiple stages, the same instruction can belong to two
+// different graphs of execution, which explains why all fields in this structure have
+// a different possible value per stage.
+// Note: all fields in this struct are not expected to be set by the user when emitting
+// 		an instruction, and are instead automatically populated by the builder.
 struct FInstruction : FValue
 {
+	// The next instruction executing after this one in all stages of execution.
 	FInstruction* Next[NumStages];
+
+	// This instruction parent block in all stages of execution.
 	FBlock* Block[NumStages];
+
+	// How many users (i.e., dependencies) this instruction has in each stage of execution.
 	uint32 NumUsers[NumStages];
+
+	// The number of users processed during instruction graph linking in each stage of execution.
+	// Note: This information combined with NumUsers is used by the builder to push instructions
+	//       in the appropriate block automatically.
 	uint32 NumProcessedUsers[NumStages];
 
-	// Returns the block in which use with specified index should execute.
-	FBlock* GetDesiredBlockForUse(EStage Stage, int32 UseIndex); 
+	// Returns the block in which the dependency with specified index should execute.
+	FBlock* GetDesiredBlockForUse(EStage Stage, int32 UseIndex);
 };
 
 template <EValueKind TTypeKind>
@@ -391,18 +429,17 @@ struct TInstruction : FInstruction
 	static constexpr EValueKind TypeKind = TTypeKind;
 };
 
-//
+// A dimensional is a fixed array of other values. This value is used to model vectors
+// and matrices.
 struct FDimensional : TInstruction<VK_Dimensional>
 {
-	static constexpr int MaxNumComponents = 16;
-
 	// Returns the constant array of component values. 
-	TArrayView<FValue* const> GetComponents() const;
+	TConstArrayView<FValue*> GetComponents() const;
 
 	// Returns the mutable array of component values. 
-	TArrayView<FValue*> GetComponents();
+	TArrayView<FValue*> GetMutableComponents();
 
-	// Returns whether all components are constant.
+	// Returns whether all components are constant (i.e., they're instances of FConstant).
 	bool AreComponentsConstant() const;
 };
 
@@ -412,16 +449,20 @@ struct TDimensional : FDimensional
 	FValue* Components[TDimension];
 };
 
-//
+// This instruction sets material attribute (i.e., "BaseColor") to its value.
 struct FSetMaterialOutput : TInstruction<VK_SetMaterialOutput>
 {
+	// The material attribute to set.
 	EMaterialProperty Property;
+
+	// The value this material attribute should be set to.
 	FValue* Arg;
 };
 
 // MIR instruction for external code declarations; See BaseMaterialExpressions.ini file.
 struct FInlineHLSL : TInstruction<VK_InlineHLSL>
 {
+	// The inlined HLSL code string.
 	const TCHAR* Code;
 };
 
@@ -497,66 +538,105 @@ enum EOperator
 	OperatorCount,
 };
 
-//
+// Whether the given operator identifies a comparison operation (e.g., ">=", "==").
 bool IsComparisonOperator(EOperator Op);
 
-//
+// Whether the given operator identifies a bit operation (e.g., "&", "|").
 bool IsBitOperator(EOperator Op);
 
-//
+// Whether the given operator identifies a unary operator.
 bool IsUnaryOperator(EOperator Op);
 
-//
+// Whether the given operator identifies a binary operator.
 bool IsBinaryOperator(EOperator Op);
 
-//
+// Whether the given operator identifies a ternary operator.
 bool IsTernaryOperator(EOperator Op);
 
-//
-const TCHAR* OperatorToString(EOperator Op);
+// Returns the string representation of the given operator.
+const TCHAR* LexToString(EOperator Op);
 
-//
+// This instruction identifies a built-in operation on one, two or three argument values.
 struct FOperator : TInstruction<VK_Operator>
 {
+	// It identifies which supported operation to carry.
 	EOperator Op;
+
+	// The first argument of the operation. This value is never null.
 	FValue* AArg;
+
+	// The second argument of the operation. This value is null for unary operators.
 	FValue* BArg;
+
+	// The third argument of the operation. This value is null for unary and binary operators.
 	FValue* CArg;
 };
 
-//
+// This instruction evaluates to one or another argument based on whether a third boolean
+// condition argument is true or false. The builder will automatically place as many instructions
+// as possible in the true/false inner blocks whilst respecting dependency requirements.
+// This is done in an effort to avoid the unnecessary computation of the input value that was
+// not selected by the condition.
 struct FBranch : TInstruction<VK_Branch>
 {
+	// The boolean condition argument used to forward the "true" or "false" argument.
 	FValue* ConditionArg;
+
+	// Value this branch evaluates to when the condition is true.
 	FValue* TrueArg;
+
+	// Value this branch evaluates to when the condition is false.
 	FValue* FalseArg;
+
+	// The inner block (in each execution stage) the subgraph evaluating `TrueArg` should be placed in.
 	FBlock TrueBlock[NumStages];
+
+	// The inner block (in each execution stage) the subgraph evaluating `FalseArg` should be placed in.
 	FBlock FalseBlock[NumStages];
 };
 
-//
+// This instruction is used to pull the inner value making a compound one. For example,
+// this instruction is used to extract an individual component of a vector value.
 struct FSubscript : TInstruction<VK_Subscript>
 {
+	// The argument to subscript.
 	FValue* Arg;
+
+	// The subscript index, i.e. the index of the component to extract.
 	int Index;
 };
 
-//
+// This instruction casts an argument value to a different type.
 struct FCast : TInstruction<VK_Cast>
 {
+	// The argument value to cast.
 	FValue* Arg;
 };
 
-//
+// This instruction performs texture sampling.
 struct FTextureSample : TInstruction<VK_TextureSample>
 {
+	// The texture object to sample.
 	FValue* TextureObject;
+
+	// The texture coordinate at which to sample.
 	FValue* TexCoordArg;
+
+	// The mip index to sample, if any provided.
 	FValue* MipValueArg;
+
+	// The automatic mip bias to use, if any provided.
 	FValue* AutomaticMipBiasArg;
+
+	// The sampler source mode to use for sampling.
 	ESamplerSourceMode SamplerSourceMode;
+
+	// The mip value mode to use for sampling.
 	ETextureMipValueMode MipValueMode;
+
+	// The sampler type to use for sampling.
 	EMaterialSamplerType SamplerType;
+
 	// TODO: TexCoordDerivatives;
 };
 
diff --git a/Engine/Source/Runtime/Engine/Public/Materials/MaterialIRCommon.h b/Engine/Source/Runtime/Engine/Public/Materials/MaterialIRCommon.h
--- a/Engine/Source/Runtime/Engine/Public/Materials/MaterialIRCommon.h
+++ b/Engine/Source/Runtime/Engine/Public/Materials/MaterialIRCommon.h
@@ -43,10 +43,6 @@ struct FType;
 struct FPrimitiveType;
 struct FObjectType;
 
-using FTypePtr = const FType*;
-using FPrimitiveTypePtr = const FPrimitiveType*;
-using FObjectTypePtr = const FObjectType*;
-
 /* IR */
 struct FValue;
 struct FInstruction;
diff --git a/Engine/Source/Runtime/Engine/Public/Materials/MaterialIREmitter.h b/Engine/Source/Runtime/Engine/Public/Materials/MaterialIREmitter.h
--- a/Engine/Source/Runtime/Engine/Public/Materials/MaterialIREmitter.h
+++ b/Engine/Source/Runtime/Engine/Public/Materials/MaterialIREmitter.h
@@ -162,15 +162,15 @@ public:
 	FSetMaterialOutput* EmitSetMaterialOutput(EMaterialProperty InProperty, FValue* Arg);
 	FValue* EmitOperator(EOperator Operator, FValue* A, FValue* B = nullptr, FValue* C = nullptr);
 	FValue* EmitBranch(FValue* Condition, FValue* True, FValue* False);
-	FValue* EmitConstruct(FTypePtr Type, FValue* Initializer);
+	FValue* EmitConstruct(const FType* Type, FValue* Initializer);
 	FValue* EmitTextureSample(FValue* Texture, FValue* TexCoord, ESamplerSourceMode SamplerSourceMode, ETextureMipValueMode MipValueMode, EMaterialSamplerType SamplerType);
-	FValue* EmitInlineHLSL(FTypePtr Type, FString Code, EValueFlags ValueFlags, EGraphProperties UsedGraphProperties);
+	FValue* EmitInlineHLSL(const FType* Type, FString Code, EValueFlags ValueFlags, EGraphProperties UsedGraphProperties);
 	FValue* EmitInlineHLSL(const FMaterialExternalCodeDeclaration* InExternalCodeDeclaration);
 
 	/* Types */
 
-	FTypePtr TryGetCommonType(FTypePtr A, FTypePtr B);
-	FTypePtr GetCommonType(FTypePtr A, FTypePtr B);
+	const FType* TryGetCommonType(const FType* A, const FType* B);
+	const FType* GetCommonType(const FType* A, const FType* B);
 
 	/* Error reporting */
 
diff --git a/Engine/Source/Runtime/Engine/Public/Materials/MaterialIRModule.h b/Engine/Source/Runtime/Engine/Public/Materials/MaterialIRModule.h
--- a/Engine/Source/Runtime/Engine/Public/Materials/MaterialIRModule.h
+++ b/Engine/Source/Runtime/Engine/Public/Materials/MaterialIRModule.h
@@ -33,7 +33,7 @@ public:
 	struct FStatistics
 	{
 		// Tracks external inputs used per frequency.
-		TBitArray<> ExternalInputUsedMask[SF_NumFrequencies]; 
+		TBitArray<> ExternalInputUsedMask[MIR::NumStages]; 
 
 		// Number of vertex texture coordinates used.
 		int NumVertexTexCoords;
diff --git a/Engine/Source/Runtime/Engine/Public/Materials/MaterialIRTypes.h b/Engine/Source/Runtime/Engine/Public/Materials/MaterialIRTypes.h
--- a/Engine/Source/Runtime/Engine/Public/Materials/MaterialIRTypes.h
+++ b/Engine/Source/Runtime/Engine/Public/Materials/MaterialIRTypes.h
@@ -32,19 +32,19 @@ struct FType
 	ETypeKind Kind;
 
 	/// Returns the type matching specified UE::Shader::FType.
-	static FTypePtr FromShaderType(const UE::Shader::FType& InShaderType);
+	static const FType* FromShaderType(const UE::Shader::FType& InShaderType);
 	
 	/// Returns the type matching specified EMaterialValueType.
-	static FTypePtr FromMaterialValueType(EMaterialValueType Type);
+	static const FType* FromMaterialValueType(EMaterialValueType Type);
 
 	/// Returns the type matching specified EMaterialParameterType.
-	static FTypePtr FromMaterialParameterType(EMaterialParameterType Type);
+	static const FType* FromMaterialParameterType(EMaterialParameterType Type);
 
 	/// Returns the `void` type.
-	static FTypePtr GetVoid();
+	static const FType* GetVoid();
 
 	/// Returns the `void` type.
-	static FTypePtr GetPoison();
+	static const FType* GetPoison();
 
 	/// Returns whether this is the poison type.
 	bool IsPoison() const;
@@ -65,22 +65,22 @@ struct FType
 	bool IsTexture() const;
 
 	/// Returns this type upcast this type to PrimitiveType if it is one, otherwise nullptr.
-	FPrimitiveTypePtr AsPrimitive() const;
+	const FPrimitiveType* AsPrimitive() const;
 
 	/// Returns this type upcast this type to PrimitiveType if it's a scalar, otherwise nullptr.
-	FPrimitiveTypePtr AsScalar() const;
+	const FPrimitiveType* AsScalar() const;
 
 	/// Returns this type upcast this type to PrimitiveType if it's a vector, otherwise nullptr.
-	FPrimitiveTypePtr AsVector() const;
+	const FPrimitiveType* AsVector() const;
 
 	/// Returns this type upcast this type to PrimitiveType if it's a matrix, otherwise nullptr.
-	FPrimitiveTypePtr AsMatrix() const;
+	const FPrimitiveType* AsMatrix() const;
 
 	/// Returns the this type name spelling (e.g. float4x4).
 	FStringView GetSpelling() const;
 	
 	/// Returns this type cast as a FObjectPtr if it is one, otherwise nullptr.
-	FObjectTypePtr AsObject() const;
+	const FObjectType* AsObject() const;
 
 	/// Converts this type to a UE::Shader::EValueType.
 	UE::Shader::EValueType ToValueType() const;
@@ -120,30 +120,30 @@ struct FPrimitiveType : FType
 	/// Number of columns (rows > 1 and columns == 1, this is a vector).
 	int NumColumns;
 
-	static FPrimitiveTypePtr GetBool1();
-	static FPrimitiveTypePtr GetInt1();
-	static FPrimitiveTypePtr GetFloat1();
-	static FPrimitiveTypePtr GetFloat2();
-	static FPrimitiveTypePtr GetFloat3();
-	static FPrimitiveTypePtr GetFloat4();
+	static const FPrimitiveType* GetBool1();
+	static const FPrimitiveType* GetInt1();
+	static const FPrimitiveType* GetFloat1();
+	static const FPrimitiveType* GetFloat2();
+	static const FPrimitiveType* GetFloat3();
+	static const FPrimitiveType* GetFloat4();
 
-	static FPrimitiveTypePtr GetScalar(EScalarKind InScalarKind);
-	static FPrimitiveTypePtr GetVector(EScalarKind InScalarKind, int NumRows);
-	static FPrimitiveTypePtr GetMatrix(EScalarKind InScalarKind, int NumColumns, int NumRows);
-	static FPrimitiveTypePtr Get(EScalarKind InScalarKind, int NumRows, int NumColumns);
+	static const FPrimitiveType* GetScalar(EScalarKind InScalarKind);
+	static const FPrimitiveType* GetVector(EScalarKind InScalarKind, int NumRows);
+	static const FPrimitiveType* GetMatrix(EScalarKind InScalarKind, int NumColumns, int NumRows);
+	static const FPrimitiveType* Get(EScalarKind InScalarKind, int NumRows, int NumColumns);
 
 	int  GetNumComponents() const { return NumRows * NumColumns; }
 	bool IsScalar() const { return GetNumComponents() == 1; }
 	bool IsVector() const { return NumRows > 1 && NumColumns == 1; }
 	bool IsMatrix() const { return NumRows > 1 && NumColumns > 1; }
-	FPrimitiveTypePtr ToScalar() const;
+	const FPrimitiveType* ToScalar() const;
 };
 
 struct FObjectType : FType
 {
 	EObjectKind ObjectKind;
 
-	static FTypePtr GetTexture2D();
+	static const FType* GetTexture2D();
 };
 
 } // namespace MIR
