New layer tool needs to pull in the MaterialAttributes input on the layer node as the base layer.

- The old layer code does this and was just missing from the new code.

#rnx

diff --git a/Engine/Source/Runtime/Engine/Private/Materials/MaterialExpressions.cpp b/Engine/Source/Runtime/Engine/Private/Materials/MaterialExpressions.cpp
--- a/Engine/Source/Runtime/Engine/Private/Materials/MaterialExpressions.cpp
+++ b/Engine/Source/Runtime/Engine/Private/Materials/MaterialExpressions.cpp
@@ -7098,6 +7098,15 @@ void UMaterialExpressionMaterialAttributeLayers::RebuildLayerGraph(bool bReportE
 						LayersTreeCache->LayerCallers[LayerIndex]->Outputs.Add(FunctionOutput.Output);
 					}
 
+					// Optional: Single material attributes input, the base input to the stack
+					if (LayersTreeCache->LayerCallers[LayerIndex]->FunctionInputs.Num() > 0)
+					{
+						if (Input.GetTracedInput().Expression)
+						{
+							LayersTreeCache->LayerCallers[LayerIndex]->FunctionInputs[0].Input = Input;
+						}
+					}
+
 					// Recursively run through internal functions to allow connection of inputs/outputs
 					LayersTreeCache->LayerCallers[LayerIndex]->UpdateFromFunctionResource();
 
