[NetPhysics] Don't use PhysicsReplicationCache when recording network replay.
- Replay not fully supported for networked physics and there is a potential crash when caching data for physics particles that gets destroyed.

diff --git a/Engine/Source/Runtime/Engine/Private/ActorReplication.cpp b/Engine/Source/Runtime/Engine/Private/ActorReplication.cpp
--- a/Engine/Source/Runtime/Engine/Private/ActorReplication.cpp
+++ b/Engine/Source/Runtime/Engine/Private/ActorReplication.cpp
@@ -442,7 +442,7 @@ void AActor::GatherCurrentMovement()
 
 			UWorld* World = GetWorld();
 			int ServerFrame = 0; 
-			if (bShouldUsePhysicsReplicationCache)
+			if (bShouldUsePhysicsReplicationCache && World && !World->IsRecordingClientReplay())
 			{
 				if (FPhysScene_Chaos* Scene = static_cast<FPhysScene_Chaos*>(World->GetPhysicsScene()))
 				{
