[MeshDeformerCollection] removed unnecessary uproperties tags, added nullchecks

diff --git a/Engine/Source/Runtime/Engine/Classes/Engine/SkeletalMesh.h b/Engine/Source/Runtime/Engine/Classes/Engine/SkeletalMesh.h
--- a/Engine/Source/Runtime/Engine/Classes/Engine/SkeletalMesh.h
+++ b/Engine/Source/Runtime/Engine/Classes/Engine/SkeletalMesh.h
@@ -2990,7 +2990,7 @@ protected:
 	TObjectPtr<UMeshDeformer> DefaultMeshDeformer;
 	
 	/** Skeletal Mesh needs this collection of deformers to make sure it cooks any extra data required by these deformers */
-	UPROPERTY(EditAnywhere, BlueprintReadOnly, Instanced ,Category = "Deformer", BlueprintGetter=GetTargetMeshDeformers)
+	UPROPERTY(EditAnywhere, BlueprintReadOnly, Category = "Deformer", BlueprintGetter=GetTargetMeshDeformers)
 	TObjectPtr<UMeshDeformerCollection> TargetMeshDeformers;
 
 	/** Default translucent material to blend on top of this mesh. Mesh will be rendered twice - once with a base material and once with overlay material */
diff --git a/Engine/Source/Runtime/Engine/Private/Animation/MeshDeformerCollection.cpp b/Engine/Source/Runtime/Engine/Private/Animation/MeshDeformerCollection.cpp
--- a/Engine/Source/Runtime/Engine/Private/Animation/MeshDeformerCollection.cpp
+++ b/Engine/Source/Runtime/Engine/Private/Animation/MeshDeformerCollection.cpp
@@ -25,13 +25,19 @@ void UMeshDeformerCollection::GetMeshDeformers_Internal(
 	
 	for (const TSoftObjectPtr<UMeshDeformer>& Deformer : MeshDeformers)
 	{
-		InOutMeshDeformers.AddUnique(Deformer);
+		if (!Deformer.IsNull())
+		{
+			InOutMeshDeformers.AddUnique(Deformer);
+		}
 	}
 
 	InOutVisitedCollection.Add(this);
 
 	for (const TObjectPtr<UMeshDeformerCollection>& MeshDeformerCollection :  MeshDeformerCollections)
 	{
-		MeshDeformerCollection->GetMeshDeformers_Internal(InOutVisitedCollection, InOutMeshDeformers);
+		if (MeshDeformerCollection)
+		{
+			MeshDeformerCollection->GetMeshDeformers_Internal(InOutVisitedCollection, InOutMeshDeformers);
+		}
 	}
 }
diff --git a/Engine/Source/Runtime/Engine/Public/Animation/MeshDeformerCollection.h b/Engine/Source/Runtime/Engine/Public/Animation/MeshDeformerCollection.h
--- a/Engine/Source/Runtime/Engine/Public/Animation/MeshDeformerCollection.h
+++ b/Engine/Source/Runtime/Engine/Public/Animation/MeshDeformerCollection.h
@@ -9,7 +9,7 @@
 class UMeshDeformer;
 
 /** A simple collection of Mesh Deformers */
-UCLASS(EditInlineNew, MinimalAPI, BlueprintType)
+UCLASS(MinimalAPI, BlueprintType)
 class UMeshDeformerCollection : public UDataAsset
 {
 	GENERATED_BODY()
