Remove FAssetRootPackagePatcher as it has been replaced by the FAssetHeaderPatcher

#rnx

diff --git a/Engine/Source/Runtime/Engine/Public/WorldPartition/WorldPartitionActorDesc.h b/Engine/Source/Runtime/Engine/Public/WorldPartition/WorldPartitionActorDesc.h
--- a/Engine/Source/Runtime/Engine/Public/WorldPartition/WorldPartitionActorDesc.h
+++ b/Engine/Source/Runtime/Engine/Public/WorldPartition/WorldPartitionActorDesc.h
@@ -242,7 +242,6 @@ class FWorldPartitionActorDesc
 	friend struct FWorldPartitionReferenceImpl;
 	friend struct FWorldPartitionActorDescUtils;
 	friend struct FWorldPartitionActorDescUnitTestAcccessor;
-	friend class FAssetRootPackagePatcher;
 	friend class FAssetHeaderPatcherInner;
 	friend class FActorDescArchive;
 	friend class FWorldPartitionActorDescInstance;
