[AutoRTFM] Fix build issue with UnixSocketPath.

UnixSocketPath should only be used when UE_HTTP_SUPPORT_UNIX_SOCKET is defined.

diff --git a/Engine/Source/Runtime/Engine/Private/Tests/AutoRTFM/AutoRTFMHttpModuleTests.cpp b/Engine/Source/Runtime/Engine/Private/Tests/AutoRTFM/AutoRTFMHttpModuleTests.cpp
--- a/Engine/Source/Runtime/Engine/Private/Tests/AutoRTFM/AutoRTFMHttpModuleTests.cpp
+++ b/Engine/Source/Runtime/Engine/Private/Tests/AutoRTFM/AutoRTFMHttpModuleTests.cpp
@@ -228,11 +228,15 @@ bool FAutoRTFMHttpRequestPreservesSettings::RunTest(const FString& Parameters)
 	DoCheck(
 		[&](IHttpRequest& Request)
 		{
+#if UE_HTTP_SUPPORT_UNIX_SOCKET
 			Request.SetOption(HttpRequestOptions::UnixSocketPath, "MyUnixSocketPath");
+#endif
 		},
 		[&](IHttpRequest& Request)
 		{
+#if UE_HTTP_SUPPORT_UNIX_SOCKET
 			TestTrueExpr(Request.GetOption(HttpRequestOptions::UnixSocketPath) == FString("MyUnixSocketPath"));
+#endif
 			TestTrueExpr(Request.GetOption(FName()) == FString());
 		});
 
@@ -403,7 +407,9 @@ bool FAutoRTFMHttpRequestPreservesSettings::RunTest(const FString& Parameters)
 			TestTrueExpr(Request.GetEffectiveURL() == FString());
 			TestTrueExpr(Request.GetVerb() == FString(TEXT("GET")));
 			TestTrueExpr(Request.GetHeader(TEXT("Mystery")) == FString());
+#if UE_HTTP_SUPPORT_UNIX_SOCKET
 			TestTrueExpr(Request.GetOption(HttpRequestOptions::UnixSocketPath) == FString());
+#endif
 			TestTrueExpr(Request.GetContent() == TArray<uint8>{});
 			TestTrueExpr(Request.GetStatus() == EHttpRequestStatus::NotStarted);
 			TestTrueExpr(!Request.GetTimeout().IsSet());
