Replace invalid/stale assertion on mmap'd data to just assert on data type alignment for FScriptContainerElement. Assertion predates application bundles and as such, files may be addressed by offset which will not necessarily be platform page size aligned.

diff --git a/Engine/Source/Runtime/Engine/Public/Animation/AnimCompressionTypes.h b/Engine/Source/Runtime/Engine/Public/Animation/AnimCompressionTypes.h
--- a/Engine/Source/Runtime/Engine/Public/Animation/AnimCompressionTypes.h
+++ b/Engine/Source/Runtime/Engine/Public/Animation/AnimCompressionTypes.h
@@ -655,12 +655,12 @@ public:
 				MappedHandle = InMappedHandle;
 				MappedRegion = InMappedRegion;
 				Data = (FScriptContainerElement*)MappedRegion->GetMappedPtr(); //@todo mapped files should probably be const-correct
-				check(IsAligned(Data, FPlatformProperties::GetMemoryMappingAlignment()));
 			}
 			else
 			{
 				Data = (FScriptContainerElement*)MallocPtr;
 			}
+			check(IsAligned(Data, alignof(FScriptContainerElement)));
 		}
 
 		bool IsMapped() const
