Change default value of r.vt.pagefreethreshold from 60 to 15.
This has been seen to be a better default setting on a range of internal projects. It should reduce pool residency on fast camera movement which can slow down page production and cause pop in.

diff --git a/Engine/Source/Runtime/Engine/Private/VT/VirtualTextureScalability.cpp b/Engine/Source/Runtime/Engine/Private/VT/VirtualTextureScalability.cpp
--- a/Engine/Source/Runtime/Engine/Private/VT/VirtualTextureScalability.cpp
+++ b/Engine/Source/Runtime/Engine/Private/VT/VirtualTextureScalability.cpp
@@ -73,7 +73,7 @@ namespace VirtualTextureScalability
 
 	static TAutoConsoleVariable<int32> CVarVTPageFreeThreshold(
 		TEXT("r.VT.PageFreeThreshold"),
-		60,
+		15,
 		TEXT("Number of frames since the last time a VT page was used, before it's considered free.\n")
 		TEXT("VT pages are not necesarily marked as used on the CPU every time they're accessed by the GPU.\n")
 		TEXT("Increasing this threshold reduces the chances that an in-use frame is considered free."),
