[Brush Navmesh] Fixed BSP Navmesh not updated when undo/redo a move

diff --git a/Engine/Source/Runtime/Engine/Classes/Engine/Level.h b/Engine/Source/Runtime/Engine/Classes/Engine/Level.h
--- a/Engine/Source/Runtime/Engine/Classes/Engine/Level.h
+++ b/Engine/Source/Runtime/Engine/Classes/Engine/Level.h
@@ -1464,9 +1464,6 @@ public:
 	int32 GetEstimatedAddToWorldWorkUnitsTotal() const;
 
 #if WITH_EDITOR
-	/** meant to be called only from editor, calculating and storing static geometry to be used with off-line and/or on-line navigation building */
-	ENGINE_API void RebuildStaticNavigableGeometry();
-
 	DECLARE_DELEGATE_ThreeParams(FLevelExternalActorsPathsProviderDelegate, const FString&, const FString&, TArray<FString>&);
 	/** Registers a level external actor paths provider */
 	static ENGINE_API FDelegateHandle RegisterLevelExternalActorsPathsProvider(const FLevelExternalActorsPathsProviderDelegate& Provider);
diff --git a/Engine/Source/Runtime/Engine/Private/Level.cpp b/Engine/Source/Runtime/Engine/Private/Level.cpp
--- a/Engine/Source/Runtime/Engine/Private/Level.cpp
+++ b/Engine/Source/Runtime/Engine/Private/Level.cpp
@@ -3149,6 +3149,17 @@ void ULevel::CommitModelSurfaces()
 		}
 
 		Model->bOnlyRebuildMaterialIndexBuffers = false;
+
+#if WITH_EDITOR
+		// The current method (i.e., CommitModelSurface) is called explicitly when ending transactions (e.g., Undo)
+		// and takes care of refreshing the render and physic state by calling 'CommitSurfaces' on each model component.
+		// When this happens we also need to rebuild the static navigable geometry of that level since
+		// RebuildAlteredBSP will early out based on 'if (!GIsTransacting)'
+		if (GEditor)
+		{
+			GEditor->RebuildStaticNavigableGeometry(this);
+		}
+#endif // WITH_EDITOR
 	}
 }
 
