Invalid BlueprintPure functions

Some functions are marked as `BlueprintPure` even though they do not return any values and do not have out parameters.
After discussing with the team, these functions will now be marked `BlueprintCallable` or be removed.
This change is part of a bigger UHT change that would now correctly flag these functions as being invalid during compilation, but would then fail the compilation with this error:
`BlueprintPure specifier is not allowed for functions with no return value and no output parameters.`

The first step is to have those functions fixed before submitting the UHT change to avoid compilation errors in other streams.

diff --git a/Engine/Source/Runtime/Engine/Public/UniversalObjectLocators/UniversalObjectLocatorScriptingExtensions.h b/Engine/Source/Runtime/Engine/Public/UniversalObjectLocators/UniversalObjectLocatorScriptingExtensions.h
--- a/Engine/Source/Runtime/Engine/Public/UniversalObjectLocators/UniversalObjectLocatorScriptingExtensions.h
+++ b/Engine/Source/Runtime/Engine/Public/UniversalObjectLocators/UniversalObjectLocatorScriptingExtensions.h
@@ -67,6 +67,6 @@ public:
 	 * @param Context    (Optional) Context object to use for resolving the object. This should usually be the object that owns or created the locator.
 	 * @return The resolve object pointer, or null if it was not found.
 	 */
-	UFUNCTION(BlueprintPure, Category="Universal Object Locators", meta=(ScriptMethod))
+	UFUNCTION(BlueprintCallable, Category="Universal Object Locators", meta=(ScriptMethod))
 	static void SyncUnload(const FUniversalObjectLocator& Locator, UObject* Context = nullptr);
 };
\ No newline at end of file
