Substrate - material hidden conversion, with FrontMaterial pin visible if substrate mode needs to be selected. So that Substrate is OptIn.

#tests pc editor dx12
#rnx

diff --git a/Engine/Shaders/Private/Substrate/Substrate.ush b/Engine/Shaders/Private/Substrate/Substrate.ush
--- a/Engine/Shaders/Private/Substrate/Substrate.ush
+++ b/Engine/Shaders/Private/Substrate/Substrate.ush
@@ -1253,6 +1253,9 @@ FSubstrateTree GetInitialisedSubstrateTree()
 	FSubstrateTree Out;
 	Out.BSDFCount = 0;
 	Out.OperatorCount = 0;
+	// Also set first BSDF and operator to make sure unlit still works and all data is isnitialised with empty graph.
+	Out.BSDFs[0] = (FSubstrateBSDF)0;
+	Out.Operators[0] = (FSubstrateOperator)0;
 	return Out;
 #else
 	return (FSubstrateTree)0;
diff --git a/Engine/Source/Editor/GraphEditor/Private/SGraphEditorImpl.cpp b/Engine/Source/Editor/GraphEditor/Private/SGraphEditorImpl.cpp
--- a/Engine/Source/Editor/GraphEditor/Private/SGraphEditorImpl.cpp
+++ b/Engine/Source/Editor/GraphEditor/Private/SGraphEditorImpl.cpp
@@ -8,6 +8,7 @@
 #include "EdGraph/EdGraph.h"
 #include "Editor.h"
 #include "Editor/EditorEngine.h"
+#include "MaterialGraph/MaterialGraphNode_Root.h"
 #include "Framework/Application/SlateApplication.h"
 #include "Framework/Commands/UIAction.h"
 #include "Framework/Commands/UICommandInfo.h"
@@ -54,6 +55,7 @@
 #include "UObject/UnrealType.h"
 #include "Widgets/Input/SEditableTextBox.h"
 #include "Widgets/Input/SMultiLineEditableTextBox.h"
+#include "Widgets/Input/SButton.h"
 #include "Widgets/Layout/SBorder.h"
 #include "Widgets/Notifications/SNotificationList.h"
 #include "Widgets/SBoxPanel.h"
@@ -61,6 +63,8 @@
 #include "Widgets/SNullWidget.h"
 #include "Widgets/SWidget.h"
 #include "Widgets/Text/STextBlock.h"
+#include "HAL/IConsoleManager.h"
+#include "RenderUtils.h"
 
 class FActiveTimerHandle;
 struct FGeometry;
@@ -901,7 +905,7 @@ void SGraphEditorImpl::OnClosedActionMenu()
 	GraphPanel->OnStopMakingConnection(/*bForceStop=*/ true);
 }
 
-void SGraphEditorImpl::AddContextMenuCommentSection(UToolMenu* InMenu)
+void SGraphEditorImpl::AddContextMenuCommentSection(UToolMenu* InMenu) 
 {
 		UGraphNodeContextMenuContext* Context = InMenu->FindContext<UGraphNodeContextMenuContext>();
 		if (!Context)
@@ -1162,7 +1166,41 @@ void SGraphEditorImpl::RegisterContextMenu(const UEdGraphSchema* Schema, FToolMe
 			}
 		}));
 
-		Menu->AddDynamicSection("EdGraphSchema", FNewToolMenuDelegate::CreateStatic(&SGraphEditorImpl::AddContextMenuCommentSection));
+		Menu->AddDynamicSection("EdGraphSchema", FNewToolMenuDelegate::CreateStatic(&SGraphEditorImpl::AddContextMenuCommentSection)); 
+
+		if(Substrate::IsSubstrateEnabled())
+		{
+			Menu->AddDynamicSection("SubstrateRootActions", FNewToolMenuDelegate::CreateLambda([](UToolMenu* InMenu)
+			{
+				if (UGraphNodeContextMenuContext * ContextObject = InMenu->FindContext<UGraphNodeContextMenuContext>())
+				{
+					if (const UEdGraphSchema* GraphSchema = ContextObject->Graph->GetSchema())
+					{
+						if (ContextObject->Node->IsA<UMaterialGraphNode_Root>() && !ContextObject->Pin)
+						{
+							const UMaterialGraphNode_Root* RootNode = Cast<UMaterialGraphNode_Root>(ContextObject->Node);
+
+							TSharedRef<SButton> SubstrateConvertButton = SNew(SButton)
+								.Text(LOCTEXT("SubstrateConvertButtonLabel", "Convert To Substrate"))
+								.ToolTipText(LOCTEXT("SubstrateConvertButtonToolTip", "Convert the root node to a simple Substrate graph for further authoring."))
+								.ButtonStyle(FAppStyle::Get(), "Button")
+								.OnClicked_Lambda([RootNode]()
+								{
+									if (RootNode && RootNode->Material)
+									{
+										RootNode->Material->UserRequestsConvertMaterialToSubstrateMaterial();
+									}
+									return FReply::Handled();
+								});
+						
+							FToolMenuSection& SubstrateSection = InMenu->AddSection("Substrate", LOCTEXT("SubstrateMenuHeader", "Substrate"));
+							SubstrateSection.AddEntry(FToolMenuEntry::InitWidget("NodeCommentBox", SubstrateConvertButton, FText::GetEmpty()));
+						}
+					}
+				}
+			}));
+		}
+
 	}
 
 	if (bDidRegisterNodeMenu)
diff --git a/Engine/Source/Editor/UnrealEd/Classes/MaterialGraph/MaterialGraph.h b/Engine/Source/Editor/UnrealEd/Classes/MaterialGraph/MaterialGraph.h
--- a/Engine/Source/Editor/UnrealEd/Classes/MaterialGraph/MaterialGraph.h
+++ b/Engine/Source/Editor/UnrealEd/Classes/MaterialGraph/MaterialGraph.h
@@ -59,6 +59,11 @@ struct FMaterialInputInfo
 			return false;
 		}
 
+		const bool bSubstrateEnabled = Substrate::IsSubstrateEnabled();
+		const bool bSubstrateHiddenMaterialConversion = Substrate::IsHiddenMaterialAssetConversionEnabled();
+		const bool bFrontMaterialConnected = Material->HasSubstrateFrontMaterialConnected();
+		const bool bShowAsSubstrate = bSubstrateEnabled && (!bSubstrateHiddenMaterialConversion || bFrontMaterialConnected);
+
 		if (Material->bUseMaterialAttributes && !bIgnoreMaterialAttributes)
 		{
 			// On the root node both MaterialAttributes and FrontMaterial are visible when Substrate is enabled. 
@@ -67,7 +72,11 @@ struct FMaterialInputInfo
 		}
 		else if( Material->IsUIMaterial() )
 		{
-			if (Substrate::IsSubstrateEnabled())
+			if (bSubstrateEnabled && bSubstrateHiddenMaterialConversion && Property == MP_FrontMaterial)
+			{
+				return true; // If hidden conversion, always enable FrontMaterial in case we want to plug a Substrate material in there.
+			}
+			if (bShowAsSubstrate)
 			{
 				if (Property == MP_FrontMaterial || Property == MP_Opacity || Property == MP_OpacityMask || Property == MP_WorldPositionOffset)
 				{
@@ -102,6 +111,21 @@ struct FMaterialInputInfo
 				return (Property - MP_CustomizedUVs0) < Material->NumCustomizedUVs;
 			}
 
+			if (bShowAsSubstrate)
+			{
+				// The only properties visible when Substrate is enable and FrontMaterial pin is plugged.
+				return Property == MP_Opacity
+					|| Property == MP_OpacityMask
+					|| Property == MP_WorldPositionOffset
+					|| Property == MP_AmbientOcclusion
+					|| Property == MP_Refraction
+					|| Property == MP_PixelDepthOffset
+					|| Property == MP_FrontMaterial
+					|| Property == MP_SurfaceThickness
+					|| Property == MP_Displacement;
+			}
+
+
 			return true;
 		}
 	}
diff --git a/Engine/Source/Runtime/Core/Private/UObject/FortniteMainBranchObjectVersion.cpp b/Engine/Source/Runtime/Core/Private/UObject/FortniteMainBranchObjectVersion.cpp
--- a/Engine/Source/Runtime/Core/Private/UObject/FortniteMainBranchObjectVersion.cpp
+++ b/Engine/Source/Runtime/Core/Private/UObject/FortniteMainBranchObjectVersion.cpp
@@ -14,6 +14,6 @@ TMap<FGuid, FGuid> FFortniteMainBranchObjectVersion::GetSystemGuids()
 	SystemGuids.Add(DevGuids.Niagara_LatestScriptCompileVersion, FGuid("670A2EBDEE1A4794BB77BBFF331EA7EF"));
 	SystemGuids.Add(DevGuids.SkeletalMeshDerivedDataVersion, FGuid("5E93192FCF3647EEB7F9D4FFF1750050"));
 	SystemGuids.Add(DevGuids.STATICMESH_DERIVEDDATA_VER, FGuid("3ABF32149F9D4D838750368AB95573BE"));
-	SystemGuids.Add(DevGuids.MaterialTranslationDDCVersion, FGuid("BB6334B9845A4A9E8E5DE2F7768127DB"));
+	SystemGuids.Add(DevGuids.MaterialTranslationDDCVersion, FGuid("90F19CAEDD684FA59D8A3DCF8DA5C83C"));
 	return SystemGuids;
 }
diff --git a/Engine/Source/Runtime/Engine/Private/Materials/HLSLMaterialTranslator.cpp b/Engine/Source/Runtime/Engine/Private/Materials/HLSLMaterialTranslator.cpp
--- a/Engine/Source/Runtime/Engine/Private/Materials/HLSLMaterialTranslator.cpp
+++ b/Engine/Source/Runtime/Engine/Private/Materials/HLSLMaterialTranslator.cpp
@@ -65,6 +65,8 @@
 #include "MaterialCachedData.h"
 #include "PostProcess/PostProcessMaterialInputs.h"
 #include "Engine/SubsurfaceProfile.h"
+#include "Materials/MaterialExpressionClearCoatNormalCustomOutput.h"
+#include "Materials/MaterialExpressionTangentOutput.h"
 
 #if WITH_EDITORONLY_DATA
 #include "Materials/MaterialExpressionSubstrate.h"
@@ -1492,6 +1494,7 @@ void FHLSLMaterialTranslator::TranslateMaterial()
 	// Process the Substrate tree representing the material topology.
 	//
 	const bool bSubstrateEnabled = Substrate::IsSubstrateEnabled();
+	const bool bSubstrateHiddenMaterialAssetConversion = Substrate::IsHiddenMaterialAssetConversionEnabled();
 	UMaterialExpression* FrontMaterialExpr = nullptr;
 	int32 FrontMaterialOutputIndex = INDEX_NONE;
 
@@ -1522,7 +1525,161 @@ void FHLSLMaterialTranslator::TranslateMaterial()
 		FrontMaterialOutputIndex = FrontMaterialInput ? FrontMaterialInput->OutputIndex : INDEX_NONE;
 	}
 
-	if (bSubstrateEnabled && FrontMaterialExpr)
+
+	FExpressionInput NullInput;
+	FExpressionInput* ThinTranslucentTransmittanceColor = &NullInput;
+	FExpressionInput* ThinTranslucentSurfaceCoverage = &NullInput;
+	FExpressionInput* WaterScatteringCoefficients = &NullInput;
+	FExpressionInput* WaterAbsorptionCoefficients = &NullInput;
+	FExpressionInput* WaterPhaseG = &NullInput;
+	FExpressionInput* ColorScaleBehindWater = &NullInput;
+	FExpressionInput* ClearCoatNormal = &NullInput;
+	FExpressionInput* CustomTangent = &NullInput;
+
+	const bool bUseRootNodeToSubstrateHiddenConversion = bSubstrateEnabled && bSubstrateHiddenMaterialAssetConversion && !FrontMaterialExpr;
+	if (bUseRootNodeToSubstrateHiddenConversion)
+	{
+		UMaterialInterface* MaterialInterface = Material->GetMaterialInterface();
+		UMaterial* BaseMaterial = MaterialInterface ? MaterialInterface->GetMaterial() : nullptr;
+		UMaterialEditorOnlyData* EditorOnlyData = BaseMaterial ? BaseMaterial->GetEditorOnlyData() : nullptr;
+
+		{
+			// Gather input for custom output
+			UMaterialExpressionThinTranslucentMaterialOutput* ThinTranslucentOutput = nullptr;
+			UMaterialExpressionSingleLayerWaterMaterialOutput* SingleLayerWaterOutput = nullptr;
+			UMaterialExpressionClearCoatNormalCustomOutput* ClearCoatBottomNormalOutput = nullptr;
+			UMaterialExpressionTangentOutput* TangentOutput = nullptr;
+			TArray<class UMaterialExpressionCustomOutput*> CustomOutputExpressions;
+			BaseMaterial->GetAllCustomOutputExpressions(CustomOutputExpressions);
+			for (UMaterialExpressionCustomOutput* Expression : CustomOutputExpressions)
+			{
+				// Gather custom output for thin translucency
+				if (ThinTranslucentOutput == nullptr && Cast<UMaterialExpressionThinTranslucentMaterialOutput>(Expression))
+				{
+					ThinTranslucentOutput = Cast<UMaterialExpressionThinTranslucentMaterialOutput>(Expression);
+
+					ThinTranslucentTransmittanceColor = &ThinTranslucentOutput->TransmittanceColor;
+					ThinTranslucentSurfaceCoverage    = &ThinTranslucentOutput->SurfaceCoverage;
+				}
+
+				// Gather custom output for single layer water
+				if (SingleLayerWaterOutput == nullptr && Cast<UMaterialExpressionSingleLayerWaterMaterialOutput>(Expression))
+				{
+					SingleLayerWaterOutput = Cast<UMaterialExpressionSingleLayerWaterMaterialOutput>(Expression);
+
+					WaterScatteringCoefficients = &SingleLayerWaterOutput->ScatteringCoefficients;
+					WaterAbsorptionCoefficients = &SingleLayerWaterOutput->AbsorptionCoefficients;
+					WaterPhaseG = &SingleLayerWaterOutput->PhaseG;
+					ColorScaleBehindWater = &SingleLayerWaterOutput->ColorScaleBehindWater;
+				}
+
+				// Gather custom output for clear coat
+				if (ClearCoatBottomNormalOutput == nullptr && Cast<UMaterialExpressionClearCoatNormalCustomOutput>(Expression))
+				{
+					ClearCoatBottomNormalOutput = Cast<UMaterialExpressionClearCoatNormalCustomOutput>(Expression);
+					ClearCoatNormal = &ClearCoatBottomNormalOutput->Input;
+				}
+
+				// Gather custom output for tangent (unused atm)
+				if (TangentOutput == nullptr && Cast<UMaterialExpressionTangentOutput>(Expression))
+				{
+					TangentOutput = Cast<UMaterialExpressionTangentOutput>(Expression);
+					CustomTangent = &TangentOutput->Input;
+				}
+
+				if (ThinTranslucentOutput && SingleLayerWaterOutput && ClearCoatBottomNormalOutput && TangentOutput)
+				{
+					break;
+				}
+			}
+		}
+
+		for (uint32 SubstrateCompilationContextIndex = 0; SubstrateCompilationContextIndex < ESubstrateCompilationContext::SCC_MAX; ++SubstrateCompilationContextIndex)
+		{
+			// This is needed because SubstrateGenerateMaterialTopologyTree will call some compiler context though material expressions.
+			CurrentSubstrateCompilationContext = ESubstrateCompilationContext(SubstrateCompilationContextIndex);
+
+			FSubstrateCompilationContext& SubstrateCtx = SubstrateCompilationContext[SubstrateCompilationContextIndex];
+			SubstrateCtx = FSubstrateCompilationContext(ESubstrateCompilationContext(SubstrateCompilationContextIndex));
+
+			if (EditorOnlyData)
+			{
+				FGuid RootNodeConversionGUID = FGuid::FGuid(0, 1, 2, 3); // This dummy GUID is the root node.
+				bSubstrateWritesEmissive |= EditorOnlyData->EmissiveColor.IsConnected();
+				bSubstrateWritesAmbientOcclusion |= EditorOnlyData->AmbientOcclusion.IsConnected();
+
+				FExpressionInput* SurfaceThickness = nullptr;
+				SubstrateThicknessStackPush(nullptr, SurfaceThickness);
+
+				if (BaseMaterial && BaseMaterial->bUseMaterialAttributes)
+				{
+					UMaterialExpressionSubstrateConvertMaterialAttributes::SubstrateGenerateMaterialTopologyTreeCommon(
+						this, RootNodeConversionGUID, nullptr /*Parent*/, 0,
+						BaseMaterial ? FMaterialAttributeDefinitionMap::GetConnectedMaterialAttributesBitmask(BaseMaterial->GetExpressions()) : 0,
+						EditorOnlyData->ShadingModelFromMaterialExpression.IsConnected(),
+						EditorOnlyData->MaterialAttributes.IsConnected(MP_EmissiveColor));
+				}
+				else if ((Material->GetMaterialDomain() == MD_Surface || Material->GetMaterialDomain() == MD_RuntimeVirtualTexture || Material->GetMaterialDomain() == MD_DeferredDecal) && EditorOnlyData)
+				{
+					FSubstrateOperator* Op = UMaterialExpressionSubstrateShadingModels::SubstrateGenerateMaterialTopologyTreeCommon(
+						this, RootNodeConversionGUID, nullptr /*Parent*/, 0,
+						EditorOnlyData->EmissiveColor,
+						EditorOnlyData->Anisotropy,
+						*ClearCoatNormal,
+						*CustomTangent,
+						EditorOnlyData->ShadingModelFromMaterialExpression);
+
+					if (Material->GetMaterialDomain() == MD_DeferredDecal)
+					{
+						// We notify that this is for decal using this usage flag.
+						// Otherwise, we do not need to do anything else, the ConvertToDecal node is only here to enforce parameter blending.
+						// This in order to get a signle BSDF as input to convert as a decal.
+						Op->SubUsage |= SUBSTRATE_OPERATOR_SUBUSAGE_DECAL;
+					}
+				}
+				else if (Material->GetMaterialDomain() == MD_Volume && EditorOnlyData)
+				{
+					UMaterialExpressionSubstrateVolumetricFogCloudBSDF::SubstrateGenerateMaterialTopologyTreeCommon(
+						this, RootNodeConversionGUID, nullptr /*Parent*/, 0,
+						EditorOnlyData->EmissiveColor,
+						EditorOnlyData->AmbientOcclusion);
+				}
+				else if (Material->GetMaterialDomain() == MD_LightFunction && EditorOnlyData)
+				{
+					UMaterialExpressionSubstrateLightFunction::SubstrateGenerateMaterialTopologyTreeCommon(
+						this, RootNodeConversionGUID, nullptr /*Parent*/, 0);
+				}
+				else if (Material->GetMaterialDomain() == MD_PostProcess && EditorOnlyData)
+				{
+					UMaterialExpressionSubstratePostProcess::SubstrateGenerateMaterialTopologyTreeCommon(
+						this, RootNodeConversionGUID, nullptr /*Parent*/, 0);
+				}
+				else if (Material->GetMaterialDomain() == MD_UI && EditorOnlyData)
+				{
+					UMaterialExpressionSubstrateUI::SubstrateGenerateMaterialTopologyTreeCommon(
+						this, RootNodeConversionGUID, nullptr /*Parent*/, 0);
+				}
+
+				SubstrateThicknessStackPop();
+
+				check(SubstrateCtx.SubstrateThicknessStack.Num() == 0);
+
+				if (SubstrateCtx.bSubstrateTreeOutOfStackDepthOccurred)
+				{
+					Errorf(TEXT(" %s [%s]: Substrate - Cyclic graph detected when we only support acyclic graph."), *Material->GetDebugName(), *Material->GetAssetPath().ToString());
+					return;
+				}
+				if (!SubstrateCtx.SubstrateGenerateDerivedMaterialOperatorData(this))
+				{
+					Errorf(TEXT("Substrate material errors encountered."));
+					return;
+				}
+			}
+
+		}
+		CurrentSubstrateCompilationContext = ESubstrateCompilationContext::SCC_Default;
+	}
+	else if (bSubstrateEnabled && FrontMaterialExpr)
 	{
 		// Temp code chunk scope (e.g.needed for the creation of static booleans from static switch parameter node, see UMaterialExpressionStaticSwitch::GetEffectiveInput).
 		TArray<FShaderCodeChunk> TempChunks;
@@ -1697,19 +1854,24 @@ void FHLSLMaterialTranslator::TranslateMaterial()
 	UE_LOG(LogMaterial, Display, TEXT(" SubstrateTreeStack: Material->CompilePropertyAndSetMaterialProperty(MP_FrontMaterial)"));
 #endif
 	Chunk[MP_SurfaceThickness] = Material->CompilePropertyAndSetMaterialProperty(MP_SurfaceThickness		,this);
-	Chunk[MP_FrontMaterial] = Material->CompilePropertyAndSetMaterialProperty(MP_FrontMaterial			,this);
 
-	CHECK_DDC_QUERY_FINISHED_ELSE_RETURN();
+	// This is what is run when hidden conversion is not executed
+	if (!bUseRootNodeToSubstrateHiddenConversion)
+	{
+		Chunk[MP_FrontMaterial] = Material->CompilePropertyAndSetMaterialProperty(MP_FrontMaterial, this);
 
-	// Now generate the code for the fully simplified MP_FrontMaterial right after MP_FrontMaterial
-	FullySimplifiedSubstrateFrontMaterialCodeChunk = SubstrateCreateAndRegisterNullMaterial();
-	{																							// This causes issues, material look different
-		// Then generate the code
-		FullySimplifiedFrontMaterialCodeChunkStart = SharedPropertyCodeChunks[FrontMaterialShaderFrequency].Num();
-		CurrentSubstrateCompilationContext = ESubstrateCompilationContext::SCC_FullySimplified;
-		FullySimplifiedSubstrateFrontMaterialCodeChunk = Material->CompilePropertyAndSetMaterialProperty(MP_FrontMaterial, this);
-		CurrentSubstrateCompilationContext = ESubstrateCompilationContext::SCC_Default;
-		FullySimplifiedFrontMaterialCodeChunkEnd = SharedPropertyCodeChunks[FrontMaterialShaderFrequency].Num();
+		CHECK_DDC_QUERY_FINISHED_ELSE_RETURN();
+
+		// Now generate the code for the fully simplified MP_FrontMaterial right after MP_FrontMaterial
+		FullySimplifiedSubstrateFrontMaterialCodeChunk = SubstrateCreateAndRegisterNullMaterial();
+		{																							// This causes issues, material look different
+			// Then generate the code
+			FullySimplifiedFrontMaterialCodeChunkStart = SharedPropertyCodeChunks[FrontMaterialShaderFrequency].Num();
+			CurrentSubstrateCompilationContext = ESubstrateCompilationContext::SCC_FullySimplified;
+			FullySimplifiedSubstrateFrontMaterialCodeChunk = Material->CompilePropertyAndSetMaterialProperty(MP_FrontMaterial, this);
+			CurrentSubstrateCompilationContext = ESubstrateCompilationContext::SCC_Default;
+			FullySimplifiedFrontMaterialCodeChunkEnd = SharedPropertyCodeChunks[FrontMaterialShaderFrequency].Num();
+		}
 	}
 
 	CHECK_DDC_QUERY_FINISHED_ELSE_RETURN();
@@ -1769,6 +1931,74 @@ void FHLSLMaterialTranslator::TranslateMaterial()
 
 	Chunk[MP_Displacement] = Material->CompilePropertyAndSetMaterialProperty(MP_Displacement, this);
 
+
+	if (bUseRootNodeToSubstrateHiddenConversion)
+	{
+		UMaterialInterface* MaterialInterface = Material->GetMaterialInterface();
+		UMaterial* BaseMaterial = MaterialInterface ? MaterialInterface->GetMaterial() : nullptr;
+		UMaterialEditorOnlyData* EditorOnlyData = BaseMaterial ? BaseMaterial->GetEditorOnlyData() : nullptr;
+
+		if (BaseMaterial && BaseMaterial->bUseMaterialAttributes)
+		{
+			const bool bHasSSS = MaterialInterface && MaterialInterface->SubsurfaceProfile;
+
+			Chunk[MP_FrontMaterial] = UMaterialExpressionSubstrateConvertMaterialAttributes::CompileCommon(this, 0,
+				FMaterialAttributeDefinitionMap::GetConnectedMaterialAttributesBitmask(BaseMaterial->GetExpressions()), EditorOnlyData->MaterialAttributes, MaterialInterface->GetShadingModels().GetFirstShadingModel(),
+				*WaterScatteringCoefficients, *WaterAbsorptionCoefficients, *WaterPhaseG, *ColorScaleBehindWater,
+				bHasSSS, bHasSSS ? MaterialInterface->SubsurfaceProfile : nullptr);
+
+		}
+		else if ((Material->GetMaterialDomain() == MD_Surface || Material->GetMaterialDomain() == MD_RuntimeVirtualTexture || Material->GetMaterialDomain() == MD_DeferredDecal) && EditorOnlyData)
+		{
+			const bool bHasSSS = MaterialInterface && MaterialInterface->SubsurfaceProfile;
+			const bool bHasAnisotropy = EditorOnlyData->Anisotropy.IsConnected();
+
+			Chunk[MP_FrontMaterial] = UMaterialExpressionSubstrateShadingModels::CompileCommon(this,
+				EditorOnlyData->BaseColor, EditorOnlyData->Specular, EditorOnlyData->Metallic, EditorOnlyData->Roughness, EditorOnlyData->EmissiveColor,
+				EditorOnlyData->Opacity, EditorOnlyData->SubsurfaceColor, EditorOnlyData->ClearCoat, EditorOnlyData->ClearCoatRoughness,
+				EditorOnlyData->ShadingModelFromMaterialExpression, MaterialInterface->GetShadingModels().GetFirstShadingModel(),
+				*ThinTranslucentTransmittanceColor, *ThinTranslucentSurfaceCoverage,
+				*WaterScatteringCoefficients, *WaterAbsorptionCoefficients, *WaterPhaseG, *ColorScaleBehindWater,
+				bHasAnisotropy, EditorOnlyData->Anisotropy,
+				EditorOnlyData->Normal, EditorOnlyData->Tangent,
+				*ClearCoatNormal, *CustomTangent,
+				bHasSSS, MaterialInterface->SubsurfaceProfile,
+				EditorOnlyData);
+		}
+		else if (Material->GetMaterialDomain() == MD_Volume && EditorOnlyData)
+		{
+			Chunk[MP_FrontMaterial] = UMaterialExpressionSubstrateVolumetricFogCloudBSDF::CompileCommon(
+				this,
+				EditorOnlyData->BaseColor,
+				EditorOnlyData->SubsurfaceColor,
+				EditorOnlyData->EmissiveColor,
+				EditorOnlyData->AmbientOcclusion,
+				MaterialInterface->GetShadingModels().HasOnlyShadingModel(MSM_Unlit),
+				EditorOnlyData);
+		}
+		else if (Material->GetMaterialDomain() == MD_LightFunction && EditorOnlyData)
+		{
+			Chunk[MP_FrontMaterial] = UMaterialExpressionSubstrateLightFunction::CompileCommon(this, EditorOnlyData->EmissiveColor, EditorOnlyData);
+		}
+		else if (Material->GetMaterialDomain() == MD_PostProcess && EditorOnlyData)
+		{
+			Chunk[MP_FrontMaterial] = UMaterialExpressionSubstratePostProcess::CompileCommon(this, EditorOnlyData->EmissiveColor, EditorOnlyData->Opacity, EditorOnlyData);
+		}
+		else if (Material->GetMaterialDomain() == MD_UI && EditorOnlyData)
+		{
+			Chunk[MP_FrontMaterial] = UMaterialExpressionSubstrateUI::CompileCommon(this, EditorOnlyData->EmissiveColor, EditorOnlyData->Opacity, EditorOnlyData);
+		}
+		else
+		{
+			Chunk[MP_FrontMaterial] = SubstrateCreateAndRegisterNullMaterial();
+		}
+
+		// Make the fully simplified code be the same as the FrontMaterial that is a single BSDF when using hidden conversion.
+		FullySimplifiedFrontMaterialCodeChunkStart = SharedPropertyCodeChunks[FrontMaterialShaderFrequency].Num();
+		FullySimplifiedSubstrateFrontMaterialCodeChunk = Chunk[MP_FrontMaterial];
+		FullySimplifiedFrontMaterialCodeChunkEnd = SharedPropertyCodeChunks[FrontMaterialShaderFrequency].Num();
+	}
+
 	CHECK_DDC_QUERY_FINISHED_ELSE_RETURN();
 	ResourcesString = TEXT("");
 
@@ -2329,7 +2559,49 @@ void FHLSLMaterialTranslator::TranslateMaterial()
 				}
 			}
 
-			if (bSubstrateFrontMaterialIsValid)
+			if (bUseRootNodeToSubstrateHiddenConversion)
+			{
+				// Hardcoded substrate tree node visit when using hidden conversion.
+				bMaterialIsSubstrate = true;
+
+				if (SubstrateCompilationContextIndex == ESubstrateCompilationContext::SCC_Default)
+				{
+					MaterialCompilationOutput.SubstrateMaterialCompilationOutput.SubstrateMaterialDescription = "";
+					ResourcesString += "// Substrate: HiddenMaterialAssetConversion\n";
+
+					// Adde default Substrate functions
+					ResourcesString += "#if TEMPLATE_USES_SUBSTRATE\n";
+					ResourcesString += "void FSubstratePixelHeader::PreUpdateAllBSDFWithBottomUpOperatorVisit(float3 V) {}\n";
+					ResourcesString += "void FSubstratePixelHeader::UpdateAllBSDFsOperatorCoverageTransmittance(FSubstrateIntegrationSettings Settings, float3 V)\n"
+										"{\n"
+										"	#if SUBSTRATE_COMPILER_SUPPORTS_STRUCT_FORWARD_DECLARATION\n"
+										"						SubstrateTree.UpdateSingleBSDFOperatorCoverageTransmittance(this, 0, Settings, V);\n"
+										"	#else\n"
+										"						UpdateSingleBSDFOperatorCoverageTransmittance(SubstrateTree, this, 0, Settings, V);\n"
+										"	#endif\n"
+										"}\n";
+					ResourcesString += "void FSubstratePixelHeader::UpdateAllOperatorsCoverageTransmittance() {}\n";
+					ResourcesString += "void FSubstratePixelHeader::UpdateAllBSDFWithBottomUpOperatorVisit() {}\n";
+					ResourcesString += "#endif // TEMPLATE_USES_SUBSTRATE\n";
+				}
+				else
+				{
+					ResourcesString += "#if TEMPLATE_USES_SUBSTRATE\n";
+					ResourcesString += "void FSubstratePixelHeader::PreUpdateAllBSDFWithBottomUpOperatorVisit_FullySimplified(float3 V) {}\n";
+					ResourcesString += "void FSubstratePixelHeader::UpdateAllBSDFsOperatorCoverageTransmittance_FullySimplified(FSubstrateIntegrationSettings Settings, float3 V)\n"
+										"{\n"
+										"	#if SUBSTRATE_COMPILER_SUPPORTS_STRUCT_FORWARD_DECLARATION\n"
+										"						SubstrateTree.UpdateSingleBSDFOperatorCoverageTransmittance(this, 0, Settings, V);\n"
+										"	#else\n"
+										"						UpdateSingleBSDFOperatorCoverageTransmittance(SubstrateTree, this, 0, Settings, V);\n"
+										"	#endif\n"
+										"}\n";
+					ResourcesString += "void FSubstratePixelHeader::UpdateAllOperatorsCoverageTransmittance_FullySimplified() {}\n";
+					ResourcesString += "void FSubstratePixelHeader::UpdateAllBSDFWithBottomUpOperatorVisit_FullySimplified() {}\n";
+					ResourcesString += "#endif // TEMPLATE_USES_SUBSTRATE\n";
+					}
+			}
+			else if (bSubstrateFrontMaterialIsValid)
 			{
 				bMaterialIsSubstrate = true;
 
@@ -2523,20 +2795,28 @@ void FHLSLMaterialTranslator::TranslateMaterial()
 			}
 			else
 			{
-				MaterialCompilationOutput.SubstrateMaterialCompilationOutput.SubstrateMaterialDescription = "";
-				ResourcesString += "// No Substrate material provided\n";
-
-				// Adde default Substrate functions
-				ResourcesString += "#if TEMPLATE_USES_SUBSTRATE\n";
-				ResourcesString += "void PreUpdateAllBSDFWithBottomUpOperatorVisit(float3 V) {}\n";
-				ResourcesString += "void UpdateAllBSDFsOperatorCoverageTransmittance(FSubstrateIntegrationSettings Settings, float3 V) {}\n";
-				ResourcesString += "void UpdateAllOperatorsCoverageTransmittance(inout FSubstrateTree SubstrateTree) {}\n";
-				ResourcesString += "void UpdateAllBSDFWithBottomUpOperatorVisit(inout FSubstrateTree SubstrateTree) {}\n";
-				ResourcesString += "void PreUpdateAllBSDFWithBottomUpOperatorVisit_FullySimplified(float3 V) {}\n";
-				ResourcesString += "void UpdateAllBSDFsOperatorCoverageTransmittance_FullySimplified(FSubstrateIntegrationSettings Settings, float3 V) {}\n";
-				ResourcesString += "void UpdateAllOperatorsCoverageTransmittance_FullySimplified(inout FSubstrateTree SubstrateTree) {}\n";
-				ResourcesString += "void UpdateAllBSDFWithBottomUpOperatorVisit_FullySimplified(inout FSubstrateTree SubstrateTree) {}\n";
-				ResourcesString += "#endif\n";
+				if (SubstrateCompilationContextIndex == ESubstrateCompilationContext::SCC_Default)
+				{
+					MaterialCompilationOutput.SubstrateMaterialCompilationOutput.SubstrateMaterialDescription = "";
+					ResourcesString += "// No Substrate material provided\n";
+
+					// Adde default Substrate functions
+					ResourcesString += "#if TEMPLATE_USES_SUBSTRATE\n";
+					ResourcesString += "void FSubstratePixelHeader::PreUpdateAllBSDFWithBottomUpOperatorVisit(float3 V) {}\n";
+					ResourcesString += "void FSubstratePixelHeader::UpdateAllBSDFsOperatorCoverageTransmittance(FSubstrateIntegrationSettings Settings, float3 V) {}\n";
+					ResourcesString += "void FSubstratePixelHeader::UpdateAllOperatorsCoverageTransmittance() {}\n";
+					ResourcesString += "void FSubstratePixelHeader::UpdateAllBSDFWithBottomUpOperatorVisit() {}\n";
+					ResourcesString += "#endif\n";
+				}
+				else
+				{
+					ResourcesString += "#if TEMPLATE_USES_SUBSTRATE\n";
+					ResourcesString += "void FSubstratePixelHeader::PreUpdateAllBSDFWithBottomUpOperatorVisit_FullySimplified(float3 V) {}\n";
+					ResourcesString += "void FSubstratePixelHeader::UpdateAllBSDFsOperatorCoverageTransmittance_FullySimplified(FSubstrateIntegrationSettings Settings, float3 V) {}\n";
+					ResourcesString += "void FSubstratePixelHeader::UpdateAllOperatorsCoverageTransmittance_FullySimplified() {}\n";
+					ResourcesString += "void FSubstratePixelHeader::UpdateAllBSDFWithBottomUpOperatorVisit_FullySimplified() {}\n";
+					ResourcesString += "#endif\n";
+				}
 			}
 		}
 	}
@@ -12618,7 +12898,7 @@ int32 FHLSLMaterialTranslator::NeuralOutput(int32 ViewportUV, uint32 NeuralIndex
 	return INDEX_NONE;
 }
 
-FSubstrateOperator& FHLSLMaterialTranslator::SubstrateCompilationRegisterOperator(int32 OperatorType, FGuid SubstrateExpressionGuid, UMaterialExpression* Child, UMaterialExpression* Parent, FGuid SubstrateParentExpressionGuid, bool bUseParameterBlending)
+FSubstrateOperator& FHLSLMaterialTranslator::SubstrateCompilationRegisterOperator(int32 OperatorType, FGuid SubstrateExpressionGuid, FGuid ChildMaterialExpressionGuid, UMaterialExpression* Parent, FGuid SubstrateParentExpressionGuid, bool bUseParameterBlending)
 {
 	FSubstrateCompilationContext& SubstrateCtx = SubstrateCompilationContext[CurrentSubstrateCompilationContext];
 
@@ -12669,7 +12949,7 @@ FSubstrateOperator& FHLSLMaterialTranslator::SubstrateCompilationRegisterOperato
 	NewOperator.MaxDistanceFromLeaves = 0;
 	NewOperator.bIsBottom = false;
 	NewOperator.bIsTop = false;
-	NewOperator.MaterialExpressionGuid = Child->MaterialExpressionGuid;
+	NewOperator.MaterialExpressionGuid = ChildMaterialExpressionGuid;
 
 	return NewOperator;
 }
diff --git a/Engine/Source/Runtime/Engine/Private/Materials/HLSLMaterialTranslator.h b/Engine/Source/Runtime/Engine/Private/Materials/HLSLMaterialTranslator.h
--- a/Engine/Source/Runtime/Engine/Private/Materials/HLSLMaterialTranslator.h
+++ b/Engine/Source/Runtime/Engine/Private/Materials/HLSLMaterialTranslator.h
@@ -1390,7 +1390,7 @@ protected:
 	virtual int32 SubstrateThicknessStackPush(UMaterialExpression* Expression, FExpressionInput* Input) override;
 	virtual void SubstrateThicknessStackPop() override;
 
-	virtual FSubstrateOperator& SubstrateCompilationRegisterOperator(int32 OperatorType, FGuid SubstrateExpressionGuid, UMaterialExpression* Child, UMaterialExpression* Parent, FGuid SubstrateParentExpressionGuid, bool bUseParameterBlending = false) override;
+	virtual FSubstrateOperator& SubstrateCompilationRegisterOperator(int32 OperatorType, FGuid SubstrateExpressionGuid, FGuid ChildMaterialExpressionGuid, UMaterialExpression* Parent, FGuid SubstrateParentExpressionGuid, bool bUseParameterBlending = false) override;
 	virtual FSubstrateOperator& SubstrateCompilationGetOperator(FGuid SubstrateExpressionGuid) override;
 	virtual FSubstrateOperator* SubstrateCompilationGetOperatorFromIndex(int32 OperatorIndex) override;
 
diff --git a/Engine/Source/Runtime/Engine/Private/Materials/Material.cpp b/Engine/Source/Runtime/Engine/Private/Materials/Material.cpp
--- a/Engine/Source/Runtime/Engine/Private/Materials/Material.cpp
+++ b/Engine/Source/Runtime/Engine/Private/Materials/Material.cpp
@@ -106,6 +106,7 @@
 #include "Widgets/Notifications/SNotificationList.h"
 #include "ThumbnailRendering/SceneThumbnailInfoWithPrimitive.h"
 #include "Editor.h"
+#include "ScopedTransaction.h"
 #endif
 #include "ShaderCodeLibrary.h"
 #include "String/ParseTokens.h"
@@ -3380,7 +3381,46 @@ static int32 GetSubstrateVersion()				{ return 0; }
 
 #if WITH_EDITOR
 
-bool UMaterial::ConvertMaterialToSubstrateMaterial(bool bAllowEmptyMaterialUpdate)
+
+bool UMaterial::UserRequestsConvertMaterialToSubstrateMaterial()
+{
+	bool bUpdated = false;
+	UMaterialEditorOnlyData* EditorOnly = GetEditorOnlyData();
+	if (Substrate::IsSubstrateEnabled() && EditorOnly && !EditorOnly->FrontMaterial.IsConnected())
+	{
+		const FScopedTransaction Transaction(LOCTEXT("SubstrateConversion", "Substrate Conversion"));
+
+		GetTransientPackage()->Modify();
+		EditorOnly->Modify();
+		MaterialGraph->Modify();
+		Modify(true);
+
+		bUpdated = ConvertMaterialToSubstrateMaterial(true /*bAllowEmptyMaterialUpdate*/, true /*bRequestedByUser*/);
+		if (bUpdated)
+		{
+			if (MaterialGraph)
+			{
+				MaterialGraph->RebuildGraph();
+			}
+
+			GetTransientPackage()->MarkPackageDirty();
+			EditorOnly->MarkPackageDirty();
+			MaterialGraph->MarkPackageDirty();
+			MarkPackageDirty();
+
+			UpdateCachedExpressionData();
+
+			PropagateDataToMaterialProxy();
+
+			PostEditChange();
+
+			ForceRecompileForRendering();
+		}
+	}
+	return bUpdated;
+}
+
+bool UMaterial::ConvertMaterialToSubstrateMaterial(bool bAllowEmptyMaterialUpdate, bool bRequestedByUser)
 {
 	/*
 	* The data flow for legacy material conversion node that can be used in isolation is as such:
@@ -3409,10 +3449,18 @@ bool UMaterial::ConvertMaterialToSubstrateMaterial(bool bAllowEmptyMaterialUpdat
 	// * Store the current version of Subtrate's auto-conversion
 	//   This allows to version the conversion, and safely update auto-converted materials if they have been saved.
 	// * If no conversion was needed, the conversion version is set to 0, to avoid rerunning this function
-	if ((EditorOnly->SubstrateConversionVersion == GetSubstrateConversionVersion() || EditorOnly->SubstrateConversionVersion == GetSubstrateNoConversionVersion()) && EditorOnly->SubstrateVersion == GetSubstrateVersion())
+	if ((EditorOnly->SubstrateConversionVersion == GetSubstrateConversionVersion() || EditorOnly->SubstrateConversionVersion == GetSubstrateNoConversionVersion()) && EditorOnly->SubstrateVersion == GetSubstrateVersion() && !bRequestedByUser)
 	{
 		return false;
 	}
+	else if (bRequestedByUser && EditorOnly->FrontMaterial.IsConnected())
+	{
+		return false; // We do not convert if the front material pin has is already connected.
+	}
+	else if(Substrate::IsHiddenMaterialAssetConversionEnabled() && !bRequestedByUser)
+	{
+		return false; // This is to work with ShouldAllowSubstrateReset. So that a material authored with IsHiddenMaterialAssetConversionEnabled using the root node inputs can then be converted to substrate.
+	}
 
 	// Store current node post from the root node.
 	int32 CurrentNodePosX = EditorX;
@@ -3463,9 +3511,11 @@ bool UMaterial::ConvertMaterialToSubstrateMaterial(bool bAllowEmptyMaterialUpdat
 		}
 		else if (OldNodeInput.UseConstant && DefaultPinValue.X != OldNodeInput.Constant)
 		{
-			UMaterialExpressionConstant* ExpressionConstantScalar = NewObject<UMaterialExpressionConstant>(this);
-			ExpressionConstantScalar->R = OldNodeInput.Constant;
-			NewNode->GetInput(NewInputIndex)->Connect(0, ExpressionConstantScalar);
+			UMaterialExpressionConstant* ExpressionConstant = NewObject<UMaterialExpressionConstant>(this);
+			ExpressionConstant->Material = this;
+			EditorOnly->ExpressionCollection.AddExpression(ExpressionConstant);
+			ExpressionConstant->R = OldNodeInput.Constant;
+			NewNode->GetInput(NewInputIndex)->Connect(0, ExpressionConstant);
 			if (OperationType == SUBSTRATE_MOVE_CONNECTION)
 			{
 				OldNodeInput.Expression = nullptr;
@@ -3482,9 +3532,11 @@ bool UMaterial::ConvertMaterialToSubstrateMaterial(bool bAllowEmptyMaterialUpdat
 		}
 		else if (OldNodeInput.UseConstant && DefaultPinValue.X != OldNodeInput.Constant)
 		{
-			UMaterialExpressionConstant* ExpressionConstantScalar = NewObject<UMaterialExpressionConstant>(this);
-			ExpressionConstantScalar->R = OldNodeInput.Constant;
-			NewNodeInput.Connect(0, ExpressionConstantScalar);
+			UMaterialExpressionConstant* ExpressionConstant = NewObject<UMaterialExpressionConstant>(this);
+			ExpressionConstant->Material = this;
+			EditorOnly->ExpressionCollection.AddExpression(ExpressionConstant);
+			ExpressionConstant->R = OldNodeInput.Constant;
+			NewNodeInput.Connect(0, ExpressionConstant);
 			if (OperationType == SUBSTRATE_MOVE_CONNECTION)
 			{
 				OldNodeInput.Expression = nullptr;
@@ -3503,9 +3555,11 @@ bool UMaterial::ConvertMaterialToSubstrateMaterial(bool bAllowEmptyMaterialUpdat
 		}
 		else if (OldNodeInput.UseConstant && DefaultPinLinearColor != CurrentPinLinearColor)
 		{
-			UMaterialExpressionConstant3Vector* ExpressionConstantScalar = NewObject<UMaterialExpressionConstant3Vector>(this);
-			ExpressionConstantScalar->Constant = CurrentPinLinearColor;
-			NewNode->GetInput(NewInputIndex)->Connect(0, ExpressionConstantScalar);
+			UMaterialExpressionConstant3Vector* ExpressionConstant = NewObject<UMaterialExpressionConstant3Vector>(this);
+			ExpressionConstant->Material = this;
+			EditorOnly->ExpressionCollection.AddExpression(ExpressionConstant);
+			ExpressionConstant->Constant = CurrentPinLinearColor;
+			NewNode->GetInput(NewInputIndex)->Connect(0, ExpressionConstant);
 			if (OperationType == SUBSTRATE_MOVE_CONNECTION)
 			{
 				OldNodeInput.Expression = nullptr;
@@ -3524,9 +3578,11 @@ bool UMaterial::ConvertMaterialToSubstrateMaterial(bool bAllowEmptyMaterialUpdat
 		}
 		else if (OldNodeInput.UseConstant && DefaultPinLinearColor != CurrentPinLinearColor)
 		{
-			UMaterialExpressionConstant3Vector* ExpressionConstantScalar = NewObject<UMaterialExpressionConstant3Vector>(this);
-			ExpressionConstantScalar->Constant = CurrentPinLinearColor;
-			NewNodeInput.Connect(0, ExpressionConstantScalar);
+			UMaterialExpressionConstant3Vector* ExpressionConstant = NewObject<UMaterialExpressionConstant3Vector>(this);
+			ExpressionConstant->Material = this;
+			EditorOnly->ExpressionCollection.AddExpression(ExpressionConstant);
+			ExpressionConstant->Constant = CurrentPinLinearColor;
+			NewNodeInput.Connect(0, ExpressionConstant);
 			if (OperationType == SUBSTRATE_MOVE_CONNECTION)
 			{
 				OldNodeInput.Expression = nullptr;
@@ -3545,9 +3601,11 @@ bool UMaterial::ConvertMaterialToSubstrateMaterial(bool bAllowEmptyMaterialUpdat
 		}
 		else if (OldNodeInput.UseConstant && DefaultPinVec3f != CurrentPinVec3f)
 		{
-			UMaterialExpressionConstant3Vector* ExpressionConstantScalar = NewObject<UMaterialExpressionConstant3Vector>(this);
-			ExpressionConstantScalar->Constant = CurrentPinVec3f;
-			NewNode->GetInput(NewInputIndex)->Connect(0, ExpressionConstantScalar);
+			UMaterialExpressionConstant3Vector* ExpressionConstant = NewObject<UMaterialExpressionConstant3Vector>(this);
+			ExpressionConstant->Material = this;
+			EditorOnly->ExpressionCollection.AddExpression(ExpressionConstant);
+			ExpressionConstant->Constant = CurrentPinVec3f;
+			NewNode->GetInput(NewInputIndex)->Connect(0, ExpressionConstant);
 			if (OperationType == SUBSTRATE_MOVE_CONNECTION)
 			{
 				OldNodeInput.Expression = nullptr;
@@ -3566,9 +3624,11 @@ bool UMaterial::ConvertMaterialToSubstrateMaterial(bool bAllowEmptyMaterialUpdat
 		}
 		else if (OldNodeInput.UseConstant && DefaultPinVec3f != CurrentPinVec3f)
 		{
-			UMaterialExpressionConstant3Vector* ExpressionConstantScalar = NewObject<UMaterialExpressionConstant3Vector>(this);
-			ExpressionConstantScalar->Constant = CurrentPinVec3f;
-			NewNodeInput.Connect(0, ExpressionConstantScalar);
+			UMaterialExpressionConstant3Vector* ExpressionConstant = NewObject<UMaterialExpressionConstant3Vector>(this);
+			ExpressionConstant->Material = this;
+			EditorOnly->ExpressionCollection.AddExpression(ExpressionConstant);
+			ExpressionConstant->Constant = CurrentPinVec3f;
+			NewNodeInput.Connect(0, ExpressionConstant);
 			if (OperationType == SUBSTRATE_MOVE_CONNECTION)
 			{
 				OldNodeInput.Expression = nullptr;
@@ -3622,426 +3682,431 @@ bool UMaterial::ConvertMaterialToSubstrateMaterial(bool bAllowEmptyMaterialUpdat
 		}
 	};
 
-	// SSS Profile
-	const bool bHasShadingModelMixture		= ShadingModels.CountShadingModels() > 1;
-	const bool bRequireSubsurfacePasses		= ShadingModels.HasShadingModel(MSM_SubsurfaceProfile) || ShadingModels.HasShadingModel(MSM_Subsurface) || ShadingModels.HasShadingModel(MSM_PreintegratedSkin) || ShadingModels.HasShadingModel(MSM_Eye);
-	const bool bRequireNoSubsurfaceProfile	= !bHasShadingModelMixture && (ShadingModel == MSM_Subsurface || ShadingModel == MSM_PreintegratedSkin); // Insure there is no profile, as this would take priority otherwise
-
-	// Empty material (i.e., material without any expression) can be 'patched' with UE4Legacy's MF in order to have correct visuals by default
-	const bool bUpdateEmptyMaterial = bAllowEmptyMaterialUpdate && !bUseMaterialAttributes && !EditorOnly->FrontMaterial.IsConnected() && GetExpressions().IsEmpty();
-
 	bool bInvalidateShader = false;
 	bool bEmptyShader = false;
-	bool bRelinkCustomOutputNodes = false;
-	UMaterialExpressionSubstrateShadingModels* ConvertNode = nullptr;
-	// Connect all the legacy pin into the conversion node
-	//First check if this material contains a layer node connection, as that requires different upgrade behavior to convert to the Substrate material layer system
-	bool bIsSubstrateLayersMaterial = false;
-	if (Substrate::IsMaterialLayeringSupportEnabled() && bUseMaterialAttributes && EditorOnly->MaterialAttributes.Expression && !EditorOnly->FrontMaterial.IsConnected())
-	{
-		UMaterialExpression* FinalNode = nullptr;
-		UMaterialExpressionMaterialAttributeLayers* LayersNode = nullptr;
-		for(UMaterialExpression* Expression : GetExpressions())
-		{
-			if(Expression->IsA<UMaterialExpressionMaterialAttributeLayers>())
-			{
-				LayersNode = Cast<UMaterialExpressionMaterialAttributeLayers>(Expression);
-				FinalNode = EditorOnly->MaterialAttributes.Expression;
-				bIsSubstrateLayersMaterial = true;
-				break;
+	// We only convert materials to substrate when IsHiddenMaterialAssetConversionEnabled is false.
+	// However, we do run the Substrate conversion/update code to make sure material with FrontMaterial plugged in are still converted.
+	if (Substrate::IsSubstrateEnabled() && (!Substrate::IsHiddenMaterialAssetConversionEnabled() || bRequestedByUser))
+	{
+		// SSS Profile
+		const bool bHasShadingModelMixture		= ShadingModels.CountShadingModels() > 1;
+		const bool bRequireSubsurfacePasses		= ShadingModels.HasShadingModel(MSM_SubsurfaceProfile) || ShadingModels.HasShadingModel(MSM_Subsurface) || ShadingModels.HasShadingModel(MSM_PreintegratedSkin) || ShadingModels.HasShadingModel(MSM_Eye);
+		const bool bRequireNoSubsurfaceProfile	= !bHasShadingModelMixture && (ShadingModel == MSM_Subsurface || ShadingModel == MSM_PreintegratedSkin); // Insure there is no profile, as this would take priority otherwise
+
+		// Empty material (i.e., material without any expression) can be 'patched' with UE4Legacy's MF in order to have correct visuals by default
+		const bool bUpdateEmptyMaterial = bAllowEmptyMaterialUpdate && !bUseMaterialAttributes && !EditorOnly->FrontMaterial.IsConnected() && GetExpressions().IsEmpty();
+
+		bool bRelinkCustomOutputNodes = false;
+		UMaterialExpressionSubstrateShadingModels* ConvertNode = nullptr;
+		// Connect all the legacy pin into the conversion node
+		//First check if this material contains a layer node connection, as that requires different upgrade behavior to convert to the Substrate material layer system
+		bool bIsSubstrateLayersMaterial = false;
+		if (Substrate::IsMaterialLayeringSupportEnabled() && bUseMaterialAttributes && EditorOnly->MaterialAttributes.Expression && !EditorOnly->FrontMaterial.IsConnected())
+		{
+			UMaterialExpression* FinalNode = nullptr;
+			UMaterialExpressionMaterialAttributeLayers* LayersNode = nullptr;
+			for(UMaterialExpression* Expression : GetExpressions())
+			{
+				if(Expression->IsA<UMaterialExpressionMaterialAttributeLayers>())
+				{
+					LayersNode = Cast<UMaterialExpressionMaterialAttributeLayers>(Expression);
+					FinalNode = EditorOnly->MaterialAttributes.Expression;
+					bIsSubstrateLayersMaterial = true;
+					break;
+				}
 			}
-		}
-
-		if(FinalNode && LayersNode)
-		{
-			//Separate the layers/blends MAs to Substrate/Non-Substrate respectively and connect to their relevant material output
-			UMaterialExpressionGetMaterialAttributes* GetAttributesNode = NewObject<UMaterialExpressionGetMaterialAttributes>(this);
-			GetAttributesNode->Material = this;
-			GetAttributesNode->MaterialAttributes.Connect(0, FinalNode);
-			EditorOnly->FrontMaterial.Connect(GetAttributesNode->CreateOrGetOutputAttribute(MP_FrontMaterial), GetAttributesNode);
-			EditorOnly->MaterialAttributes.Connect(GetAttributesNode->CreateOrGetOutputAttribute(MP_MaterialAttributes), GetAttributesNode);
-			SetPosXAndMoveReferenceToTheRight(GetAttributesNode);
 
-			//If a non-Substrate default input is connected to the MLA node, convert it.
-			if (LayersNode->Input.IsConnected() && !LayersNode->Input.Expression->IsResultSubstrateMaterial(LayersNode->Input.OutputIndex))
+			if(FinalNode && LayersNode)
 			{
-				UMaterialExpressionSubstrateConvertMaterialAttributes* ConvertAttributeNode = NewObject<UMaterialExpressionSubstrateConvertMaterialAttributes>(this);
-				ConvertAttributeNode->Material = this;
-				ConvertAttributeNode->MaterialAttributes.Connect(LayersNode->Input.OutputIndex, LayersNode->Input.Expression);
-				ConvertAttributeNode->ShadingModelOverride = MSM_DefaultLit;
-
-				UMaterialExpressionSetMaterialAttributes* SetAttributesNode = NewObject<UMaterialExpressionSetMaterialAttributes>(this);
-				SetAttributesNode->Material = this;
-				SetAttributesNode->ConnectInputAttribute(MP_FrontMaterial, ConvertAttributeNode);
-				SetAttributesNode->ConnectInputAttribute(MP_MaterialAttributes, LayersNode->Input.Expression, LayersNode->Input.OutputIndex);
+				//Separate the layers/blends MAs to Substrate/Non-Substrate respectively and connect to their relevant material output
+				UMaterialExpressionGetMaterialAttributes* GetAttributesNode = NewObject<UMaterialExpressionGetMaterialAttributes>(this);
+				GetAttributesNode->Material = this;
+				GetAttributesNode->MaterialAttributes.Connect(0, FinalNode);
+				EditorOnly->FrontMaterial.Connect(GetAttributesNode->CreateOrGetOutputAttribute(MP_FrontMaterial), GetAttributesNode);
+				EditorOnly->MaterialAttributes.Connect(GetAttributesNode->CreateOrGetOutputAttribute(MP_MaterialAttributes), GetAttributesNode);
+				SetPosXAndMoveReferenceToTheRight(GetAttributesNode);
 
-				LayersNode->Input.Connect(0, SetAttributesNode);
+				//If a non-Substrate default input is connected to the MLA node, convert it.
+				if (LayersNode->Input.IsConnected() && !LayersNode->Input.Expression->IsResultSubstrateMaterial(LayersNode->Input.OutputIndex))
+				{
+					UMaterialExpressionSubstrateConvertMaterialAttributes* ConvertAttributeNode = NewObject<UMaterialExpressionSubstrateConvertMaterialAttributes>(this);
+					ConvertAttributeNode->Material = this;
+					ConvertAttributeNode->MaterialAttributes.Connect(LayersNode->Input.OutputIndex, LayersNode->Input.Expression);
+					ConvertAttributeNode->ShadingModelOverride = MSM_DefaultLit;
+
+					UMaterialExpressionSetMaterialAttributes* SetAttributesNode = NewObject<UMaterialExpressionSetMaterialAttributes>(this);
+					SetAttributesNode->Material = this;
+					SetAttributesNode->ConnectInputAttribute(MP_FrontMaterial, ConvertAttributeNode);
+					SetAttributesNode->ConnectInputAttribute(MP_MaterialAttributes, LayersNode->Input.Expression, LayersNode->Input.OutputIndex);
+
+					LayersNode->Input.Connect(0, SetAttributesNode);
+
+					SetPosXAndMoveReferenceToTheRight(GetAttributesNode);
+					ReplaceNodeAndMoveToTheRight(LayersNode, ConvertAttributeNode);
+					SetPosXAndMoveReferenceToTheRight(GetAttributesNode);
+					ReplaceNodeAndMoveToTheRight(LayersNode, SetAttributesNode);			
+				}
 
-				SetPosXAndMoveReferenceToTheRight(GetAttributesNode);
-				ReplaceNodeAndMoveToTheRight(LayersNode, ConvertAttributeNode);
-				SetPosXAndMoveReferenceToTheRight(GetAttributesNode);
-				ReplaceNodeAndMoveToTheRight(LayersNode, SetAttributesNode);			
+				BlendMode = ConvertLegacyBlendMode(BlendMode, ShadingModels);
+				RefractionCoverageMode = RCM_CoverageIgnored;
+				bInvalidateShader = true;
 			}
-
-			BlendMode = ConvertLegacyBlendMode(BlendMode, ShadingModels);
-			RefractionCoverageMode = RCM_CoverageIgnored;
-			bInvalidateShader = true;
 		}
-	}
 	
-	if(!bIsSubstrateLayersMaterial)
-	{
-		const bool bHasAnySubstrateNodes = HasAnyExpressionsInMaterialAndFunctionsOfType<UMaterialExpressionSubstrateBSDF>();
-		if (bUseMaterialAttributes && EditorOnly->MaterialAttributes.Expression && !EditorOnly->FrontMaterial.IsConnected() && !EditorOnly->MaterialAttributes.Expression->IsResultSubstrateMaterial(EditorOnly->MaterialAttributes.OutputIndex) && !bHasAnySubstrateNodes) // M_Rifle cause issues there
+		if(!bIsSubstrateLayersMaterial)
 		{
-			UMaterialExpressionSubstrateConvertMaterialAttributes* ConvertAttributeNode = NewObject<UMaterialExpressionSubstrateConvertMaterialAttributes>(this);
-			EditorOnly->ExpressionCollection.AddExpression(ConvertAttributeNode);
-			ConvertAttributeNode->Material = this;
-			SetPosXAndMoveReferenceToTheRight(ConvertAttributeNode);
-			ConvertAttributeNode->SubsurfaceProfile = bRequireNoSubsurfaceProfile ? nullptr : SubsurfaceProfile;
-
-			// * Copy the material attribute connection to the conversion node.
-			// * Leave the material attribute existing connection plugged to the root node, 
-			//   so that other input (PixelDepthOffset, WorldPositionOffset, ...) get pull 
-			//   from the material attributes node
-			ConnectionTo(EditorOnly->MaterialAttributes, ConvertAttributeNode->MaterialAttributes, SUBSTRATE_COPY_CONNECTION);
-
-			// Reconnect custom output to material attribute conversion node
+			const bool bHasAnySubstrateNodes = HasAnyExpressionsInMaterialAndFunctionsOfType<UMaterialExpressionSubstrateBSDF>();
+			if (bUseMaterialAttributes && EditorOnly->MaterialAttributes.Expression && !EditorOnly->FrontMaterial.IsConnected() && !EditorOnly->MaterialAttributes.Expression->IsResultSubstrateMaterial(EditorOnly->MaterialAttributes.OutputIndex) && !bHasAnySubstrateNodes) // M_Rifle cause issues there
 			{
-				check(ConvertAttributeNode);
-				GatherCustomNodes();
+				UMaterialExpressionSubstrateConvertMaterialAttributes* ConvertAttributeNode = NewObject<UMaterialExpressionSubstrateConvertMaterialAttributes>(this);
+				EditorOnly->ExpressionCollection.AddExpression(ConvertAttributeNode);
+				ConvertAttributeNode->Material = this;
+				SetPosXAndMoveReferenceToTheRight(ConvertAttributeNode);
+				ConvertAttributeNode->SubsurfaceProfile = bRequireNoSubsurfaceProfile ? nullptr : SubsurfaceProfile;
+
+				// * Copy the material attribute connection to the conversion node.
+				// * Leave the material attribute existing connection plugged to the root node, 
+				//   so that other input (PixelDepthOffset, WorldPositionOffset, ...) get pull 
+				//   from the material attributes node
+				ConnectionTo(EditorOnly->MaterialAttributes, ConvertAttributeNode->MaterialAttributes, SUBSTRATE_COPY_CONNECTION);
 
-				if (SingleLayerWaterOutput)
+				// Reconnect custom output to material attribute conversion node
 				{
-					ConnectionTo(SingleLayerWaterOutput->ScatteringCoefficients,	ConvertAttributeNode->WaterScatteringCoefficients);
-					ConnectionTo(SingleLayerWaterOutput->AbsorptionCoefficients,	ConvertAttributeNode->WaterAbsorptionCoefficients);
-					ConnectionTo(SingleLayerWaterOutput->PhaseG,					ConvertAttributeNode->WaterPhaseG);
-					ConnectionTo(SingleLayerWaterOutput->ColorScaleBehindWater,		ConvertAttributeNode->ColorScaleBehindWater);
-				}
-			}
+					check(ConvertAttributeNode);
+					GatherCustomNodes();
 
-			// If the graph contains a 'tangent output' node and uses 'eye' shading model, we add the custom tangent to the material attribute list
-			const bool bRequireCustomTangent = ShadingModels.HasShadingModel(MSM_Eye) && TangentOutput;
-			if (bRequireCustomTangent)
-			{
-				const FGuid MaterialAttributeGuid = FMaterialAttributeDefinitionMap::GetID(MP_MaterialAttributes);
-				const FGuid CustomEyeTangentGuid = FMaterialAttributeDefinitionMap::GetCustomAttributeID(TEXT("CustomEyeTangent"));
+					if (SingleLayerWaterOutput)
+					{
+						ConnectionTo(SingleLayerWaterOutput->ScatteringCoefficients,	ConvertAttributeNode->WaterScatteringCoefficients);
+						ConnectionTo(SingleLayerWaterOutput->AbsorptionCoefficients,	ConvertAttributeNode->WaterAbsorptionCoefficients);
+						ConnectionTo(SingleLayerWaterOutput->PhaseG,					ConvertAttributeNode->WaterPhaseG);
+						ConnectionTo(SingleLayerWaterOutput->ColorScaleBehindWater,		ConvertAttributeNode->ColorScaleBehindWater);
+					}
+				}
 
-				UMaterialExpressionSetMaterialAttributes* SetAttributesNode = NewObject<UMaterialExpressionSetMaterialAttributes>(this);
-				EditorOnly->ExpressionCollection.AddExpression(SetAttributesNode); // Add the SetAttributesNode to ensure the MaterialAttribute input are correctly tracked & cached
-				SetAttributesNode->Material = this;
-				SetAttributesNode->Inputs[0].Connect(0, EditorOnly->MaterialAttributes.Expression);
-				SetAttributesNode->AttributeSetTypes.Add(CustomEyeTangentGuid);
-				SetAttributesNode->Inputs.Add(TangentOutput->Input);
+				// If the graph contains a 'tangent output' node and uses 'eye' shading model, we add the custom tangent to the material attribute list
+				const bool bRequireCustomTangent = ShadingModels.HasShadingModel(MSM_Eye) && TangentOutput;
+				if (bRequireCustomTangent)
+				{
+					const FGuid MaterialAttributeGuid = FMaterialAttributeDefinitionMap::GetID(MP_MaterialAttributes);
+					const FGuid CustomEyeTangentGuid = FMaterialAttributeDefinitionMap::GetCustomAttributeID(TEXT("CustomEyeTangent"));
 
-				// Connect to conversion node
-				ConvertAttributeNode->MaterialAttributes.Connect(0, SetAttributesNode);
+					UMaterialExpressionSetMaterialAttributes* SetAttributesNode = NewObject<UMaterialExpressionSetMaterialAttributes>(this);
+					EditorOnly->ExpressionCollection.AddExpression(SetAttributesNode); // Add the SetAttributesNode to ensure the MaterialAttribute input are correctly tracked & cached
+					SetAttributesNode->Material = this;
+					SetAttributesNode->Inputs[0].Connect(0, EditorOnly->MaterialAttributes.Expression);
+					SetAttributesNode->AttributeSetTypes.Add(CustomEyeTangentGuid);
+					SetAttributesNode->Inputs.Add(TangentOutput->Input);
 
-				// Connect to root node
-				EditorOnly->MaterialAttributes.Connect(0, SetAttributesNode);
-			}
+					// Connect to conversion node
+					ConvertAttributeNode->MaterialAttributes.Connect(0, SetAttributesNode);
 
-			// Connect converted Substrate data to root node
-			EditorOnly->FrontMaterial.Connect(0, ConvertAttributeNode);
+					// Connect to root node
+					EditorOnly->MaterialAttributes.Connect(0, SetAttributesNode);
+				}
 
-			// Shading Model
-			// * either use the shader graph expression 
-			// * or add a constant shading model
-			if (ShadingModel == MSM_FromMaterialExpression)
-			{
-				ConvertAttributeNode->ShadingModelOverride = MSM_FromMaterialExpression;
-			}
-			else
-			{
-				// Store Substrate shading model of the converted material. 
-				check(ShadingModels.CountShadingModels() == 1);
-				ConvertAttributeNode->ShadingModelOverride = ShadingModel;
+				// Connect converted Substrate data to root node
+				EditorOnly->FrontMaterial.Connect(0, ConvertAttributeNode);
 
-				// Apply is thin surface if needed
-				if (ShadingModel == MSM_TwoSidedFoliage)
+				// Shading Model
+				// * either use the shader graph expression 
+				// * or add a constant shading model
+				if (ShadingModel == MSM_FromMaterialExpression)
 				{
-					bIsThinSurface = true;
+					ConvertAttributeNode->ShadingModelOverride = MSM_FromMaterialExpression;
 				}
-			}
+				else
+				{
+					// Store Substrate shading model of the converted material. 
+					check(ShadingModels.CountShadingModels() == 1);
+					ConvertAttributeNode->ShadingModelOverride = ShadingModel;
 
-			if (MaterialDomain == MD_DeferredDecal)
-			{
-				// For now we don't enforce shading model since it could be driven by expression and we don't have much 
-				// control on this, but only DefaultLit should be supported.
+					// Apply is thin surface if needed
+					if (ShadingModel == MSM_TwoSidedFoliage)
+					{
+						bIsThinSurface = true;
+					}
+				}
 
-				// Now pass through the convert to decal node, which flag the material as SSM_Decal, which will set the domain to Decal.
-				UMaterialExpressionSubstrateConvertToDecal* ConvertToDecalNode = NewObject<UMaterialExpressionSubstrateConvertToDecal>(this);
-				EditorOnly->ExpressionCollection.AddExpression(ConvertToDecalNode);
-				ConvertAttributeNode->Material = this;
-				ReplaceNodeAndMoveToTheRight(ConvertAttributeNode, ConvertToDecalNode);
-				ConvertToDecalNode->DecalMaterial.Connect(0, ConvertAttributeNode);
+				if (MaterialDomain == MD_DeferredDecal)
+				{
+					// For now we don't enforce shading model since it could be driven by expression and we don't have much 
+					// control on this, but only DefaultLit should be supported.
 
-				EditorOnly->FrontMaterial.Connect(0, ConvertToDecalNode);
-			}
+					// Now pass through the convert to decal node, which flag the material as SSM_Decal, which will set the domain to Decal.
+					UMaterialExpressionSubstrateConvertToDecal* ConvertToDecalNode = NewObject<UMaterialExpressionSubstrateConvertToDecal>(this);
+					EditorOnly->ExpressionCollection.AddExpression(ConvertToDecalNode);
+					ConvertAttributeNode->Material = this;
+					ReplaceNodeAndMoveToTheRight(ConvertAttributeNode, ConvertToDecalNode);
+					ConvertToDecalNode->DecalMaterial.Connect(0, ConvertAttributeNode);
 
-			BlendMode = ConvertLegacyBlendMode(BlendMode, ShadingModels);
-			RefractionCoverageMode = RCM_CoverageIgnored;
-			bInvalidateShader = true;
-		}
-		else if (bUpdateEmptyMaterial)
-		{
-			// Empty material: Create by default a slab node
-			UMaterialFunction* DefaultMF = LoadObject<UMaterialFunction>(nullptr, TEXT("/Engine/Functions/Substrate/SMF_UE4Legacy.SMF_UE4Legacy")); 
-			if (DefaultMF)
+					EditorOnly->FrontMaterial.Connect(0, ConvertToDecalNode);
+				}
+
+				BlendMode = ConvertLegacyBlendMode(BlendMode, ShadingModels);
+				RefractionCoverageMode = RCM_CoverageIgnored;
+				bInvalidateShader = true;
+			}
+			else if (bUpdateEmptyMaterial)
 			{
-				DefaultMF->UpdateFromFunctionResource();
-				DefaultMF->PostEditChange();
-				DefaultMF->ConditionalPostLoad();
+				// Empty material: Create by default a slab node
+				UMaterialFunction* DefaultMF = LoadObject<UMaterialFunction>(nullptr, TEXT("/Engine/Functions/Substrate/SMF_UE4Legacy.SMF_UE4Legacy")); 
+				if (DefaultMF)
+				{
+					DefaultMF->UpdateFromFunctionResource();
+					DefaultMF->PostEditChange();
+					DefaultMF->ConditionalPostLoad();
 
-				UMaterialExpressionMaterialFunctionCall* MFCallNode = NewObject<UMaterialExpressionMaterialFunctionCall>(this);
-				if (MFCallNode->SetMaterialFunction(DefaultMF))
+					UMaterialExpressionMaterialFunctionCall* MFCallNode = NewObject<UMaterialExpressionMaterialFunctionCall>(this);
+					if (MFCallNode->SetMaterialFunction(DefaultMF))
+					{
+						// This is needed for input/output expressions to be set correctly, otherwise compilation will fail.
+						EditorOnly->ExpressionCollection.AddExpression(MFCallNode);
+
+						SetPosXAndMoveReferenceToTheRight(MFCallNode);
+						EditorOnly->FrontMaterial.Connect(0, MFCallNode);
+
+						MFCallNode->UpdateFromFunctionResource();
+						MFCallNode->PostEditChange();
+						MFCallNode->ConditionalPostLoad();
+
+						ColorMatInputConnectionToIndexed(EditorOnly->BaseColor,		MFCallNode, 0, MP_BaseColor);
+						ScalarMatInputConnectionToIndexed(EditorOnly->Metallic,		MFCallNode, 1, MP_Metallic);
+						ScalarMatInputConnectionToIndexed(EditorOnly->Specular,		MFCallNode, 2, MP_Specular);
+						ScalarMatInputConnectionToIndexed(EditorOnly->Roughness,	MFCallNode, 3, MP_Roughness);
+						Vector3MatInputConnectionToIndexed(EditorOnly->Normal,		MFCallNode, 4, MP_Normal);
+						ColorMatInputConnectionToIndexed(EditorOnly->EmissiveColor,	MFCallNode, 5, MP_EmissiveColor);
+						ScalarMatInputConnectionToIndexed(EditorOnly->Opacity,		MFCallNode, 6, MP_Opacity);
+					}
+				}
+				else
 				{
-					// This is needed for input/output expressions to be set correctly, otherwise compilation will fail.
-					EditorOnly->ExpressionCollection.AddExpression(MFCallNode);
-
-					SetPosXAndMoveReferenceToTheRight(MFCallNode);
-					EditorOnly->FrontMaterial.Connect(0, MFCallNode);
-
-					MFCallNode->UpdateFromFunctionResource();
-					MFCallNode->PostEditChange();
-					MFCallNode->ConditionalPostLoad();
-
-					ColorMatInputConnectionToIndexed(EditorOnly->BaseColor,		MFCallNode, 0, MP_BaseColor);
-					ScalarMatInputConnectionToIndexed(EditorOnly->Metallic,		MFCallNode, 1, MP_Metallic);
-					ScalarMatInputConnectionToIndexed(EditorOnly->Specular,		MFCallNode, 2, MP_Specular);
-					ScalarMatInputConnectionToIndexed(EditorOnly->Roughness,	MFCallNode, 3, MP_Roughness);
-					Vector3MatInputConnectionToIndexed(EditorOnly->Normal,		MFCallNode, 4, MP_Normal);
-					ColorMatInputConnectionToIndexed(EditorOnly->EmissiveColor,	MFCallNode, 5, MP_EmissiveColor);
-					ScalarMatInputConnectionToIndexed(EditorOnly->Opacity,		MFCallNode, 6, MP_Opacity);
+					// Or if it cannot be found, a slab node
+					UMaterialExpressionSubstrateSlabBSDF* SlabNode = NewObject<UMaterialExpressionSubstrateSlabBSDF>(this);
+					EditorOnly->ExpressionCollection.AddExpression(SlabNode);
+					SlabNode->Material = this;
+					SetPosXAndMoveReferenceToTheRight(SlabNode);
+					EditorOnly->FrontMaterial.Connect(0, SlabNode);
 				}
+				bRelinkCustomOutputNodes = false;
+				bInvalidateShader = true;
+				bEmptyShader = true;
 			}
-			else
+			else if (!bUseMaterialAttributes && !EditorOnly->FrontMaterial.IsConnected() && !bHasAnySubstrateNodes)
 			{
-				// Or if it cannot be found, a slab node
-				UMaterialExpressionSubstrateSlabBSDF* SlabNode = NewObject<UMaterialExpressionSubstrateSlabBSDF>(this);
-				EditorOnly->ExpressionCollection.AddExpression(SlabNode);
-				SlabNode->Material = this;
-				SetPosXAndMoveReferenceToTheRight(SlabNode);
-				EditorOnly->FrontMaterial.Connect(0, SlabNode);
-			}
-			bRelinkCustomOutputNodes = false;
-			bInvalidateShader = true;
-			bEmptyShader = true;
-		}
-		else if (!bUseMaterialAttributes && !EditorOnly->FrontMaterial.IsConnected() && !bHasAnySubstrateNodes)
-		{
-			if (MaterialDomain == MD_Surface)
-			{
-				ConvertNode = NewObject<UMaterialExpressionSubstrateShadingModels>(this);
-				EditorOnly->ExpressionCollection.AddExpression(ConvertNode); // Add the SetAttributesNode to ensure the MaterialAttribute input are correctly tracked & cached
-				ConvertNode->Material = this;
-				SetPosXAndMoveReferenceToTheRight(ConvertNode);
-				ConvertNode->SubsurfaceProfile = bRequireNoSubsurfaceProfile ? nullptr : SubsurfaceProfile;
-				ColorMatInputConnectionTo(EditorOnly->BaseColor,			ConvertNode->BaseColor,			MP_BaseColor);
-				ScalarMatInputConnectionTo(EditorOnly->Metallic,			ConvertNode->Metallic,			MP_Metallic);
-				ScalarMatInputConnectionTo(EditorOnly->Specular,			ConvertNode->Specular,			MP_Specular);
-				ScalarMatInputConnectionTo(EditorOnly->Roughness,			ConvertNode->Roughness,			MP_Roughness);
-				ScalarMatInputConnectionTo(EditorOnly->Anisotropy,			ConvertNode->Anisotropy,		MP_Anisotropy);
-				ColorMatInputConnectionTo(EditorOnly->EmissiveColor,		ConvertNode->EmissiveColor,		MP_EmissiveColor);
-				Vector3MatInputConnectionTo(EditorOnly->Normal,				ConvertNode->Normal,			MP_Normal, SUBSTRATE_COPY_CONNECTION);
-				Vector3MatInputConnectionTo(EditorOnly->Tangent,			ConvertNode->Tangent,			MP_Tangent);
-				ColorMatInputConnectionTo(EditorOnly->SubsurfaceColor,		ConvertNode->SubSurfaceColor,	MP_SubsurfaceColor);
-				ScalarMatInputConnectionTo(EditorOnly->ClearCoat,			ConvertNode->ClearCoat,			MP_CustomData0);
-				ScalarMatInputConnectionTo(EditorOnly->ClearCoatRoughness,	ConvertNode->ClearCoatRoughness,MP_CustomData1);
-				ScalarMatInputConnectionTo(EditorOnly->Opacity,				ConvertNode->Opacity,			MP_Opacity, SUBSTRATE_COPY_CONNECTION);	// We only copy, to keep Opacity on the root node in case BLEND_AlphaComposite is selected.
-				bRelinkCustomOutputNodes = true;
-			
-				// Shading Model
-				// * either use the shader graph expression 
-				// * or add a constant shading model
-				if (ShadingModel == MSM_FromMaterialExpression)
+				if (MaterialDomain == MD_Surface)
 				{
-					if (!EditorOnly->ShadingModelFromMaterialExpression.IsConnected())
+					ConvertNode = NewObject<UMaterialExpressionSubstrateShadingModels>(this);
+					EditorOnly->ExpressionCollection.AddExpression(ConvertNode); // Add the SetAttributesNode to ensure the MaterialAttribute input are correctly tracked & cached
+					ConvertNode->Material = this;
+					SetPosXAndMoveReferenceToTheRight(ConvertNode);
+					ConvertNode->SubsurfaceProfile = bRequireNoSubsurfaceProfile ? nullptr : SubsurfaceProfile;
+					ColorMatInputConnectionTo(EditorOnly->BaseColor,			ConvertNode->BaseColor,			MP_BaseColor);
+					ScalarMatInputConnectionTo(EditorOnly->Metallic,			ConvertNode->Metallic,			MP_Metallic);
+					ScalarMatInputConnectionTo(EditorOnly->Specular,			ConvertNode->Specular,			MP_Specular);
+					ScalarMatInputConnectionTo(EditorOnly->Roughness,			ConvertNode->Roughness,			MP_Roughness);
+					ScalarMatInputConnectionTo(EditorOnly->Anisotropy,			ConvertNode->Anisotropy,		MP_Anisotropy);
+					ColorMatInputConnectionTo(EditorOnly->EmissiveColor,		ConvertNode->EmissiveColor,		MP_EmissiveColor);
+					Vector3MatInputConnectionTo(EditorOnly->Normal,				ConvertNode->Normal,			MP_Normal, SUBSTRATE_COPY_CONNECTION);
+					Vector3MatInputConnectionTo(EditorOnly->Tangent,			ConvertNode->Tangent,			MP_Tangent);
+					ColorMatInputConnectionTo(EditorOnly->SubsurfaceColor,		ConvertNode->SubSurfaceColor,	MP_SubsurfaceColor);
+					ScalarMatInputConnectionTo(EditorOnly->ClearCoat,			ConvertNode->ClearCoat,			MP_CustomData0);
+					ScalarMatInputConnectionTo(EditorOnly->ClearCoatRoughness,	ConvertNode->ClearCoatRoughness,MP_CustomData1);
+					ScalarMatInputConnectionTo(EditorOnly->Opacity,				ConvertNode->Opacity,			MP_Opacity, SUBSTRATE_COPY_CONNECTION);	// We only copy, to keep Opacity on the root node in case BLEND_AlphaComposite is selected.
+					bRelinkCustomOutputNodes = true;
+			
+					// Shading Model
+					// * either use the shader graph expression 
+					// * or add a constant shading model
+					if (ShadingModel == MSM_FromMaterialExpression)
 					{
-						ConvertNode->ShadingModelOverride = MSM_DefaultLit;
+						if (!EditorOnly->ShadingModelFromMaterialExpression.IsConnected())
+						{
+							ConvertNode->ShadingModelOverride = MSM_DefaultLit;
+						}
+						else
+						{
+							// Reconnect the shading model expression. 
+							// Note: assign the expression directly, as using ConvertNode->GetInput(19)->Connect(..) causes the expression to not be assigned
+							ConvertNode->ShadingModel.Connect(EditorOnly->ShadingModelFromMaterialExpression.OutputIndex, EditorOnly->ShadingModelFromMaterialExpression.Expression);
+						}
+
+						// Store Substrate shading model of the converted material. 
+						GatherCustomNodes();
+						if (SingleLayerWaterOutput)
+						{
+							ShadingModels.AddShadingModel(MSM_SingleLayerWater);
+						}
+
+						check(ShadingModels.CountShadingModels() >= 1);
 					}
 					else
 					{
-						// Reconnect the shading model expression. 
-						// Note: assign the expression directly, as using ConvertNode->GetInput(19)->Connect(..) causes the expression to not be assigned
-						ConvertNode->ShadingModel.Connect(EditorOnly->ShadingModelFromMaterialExpression.OutputIndex, EditorOnly->ShadingModelFromMaterialExpression.Expression);
-					}
+						ConvertNode->ShadingModelOverride = ShadingModel;
+						check(ShadingModels.CountShadingModels() == 1);
 
-					// Store Substrate shading model of the converted material. 
-					GatherCustomNodes();
-					if (SingleLayerWaterOutput)
-					{
-						ShadingModels.AddShadingModel(MSM_SingleLayerWater);
+						if (ShadingModel == MSM_TwoSidedFoliage)
+						{
+							bIsThinSurface = true;
+						}
 					}
 
-					check(ShadingModels.CountShadingModels() >= 1);
+					EditorOnly->FrontMaterial.Connect(0, ConvertNode);
+
+					bInvalidateShader = true;
 				}
-				else
+				else if (MaterialDomain == MD_Volume)
 				{
-					ConvertNode->ShadingModelOverride = ShadingModel;
-					check(ShadingModels.CountShadingModels() == 1);
+					UMaterialExpressionSubstrateVolumetricFogCloudBSDF* VolBSDF = NewObject<UMaterialExpressionSubstrateVolumetricFogCloudBSDF>(this);
+					EditorOnly->ExpressionCollection.AddExpression(VolBSDF);
+					VolBSDF->Material = this;
+					SetPosXAndMoveReferenceToTheRight(VolBSDF);
+					ColorMatInputConnectionTo(EditorOnly->BaseColor,		VolBSDF->Albedo,			MP_BaseColor);	
+					ColorMatInputConnectionTo(EditorOnly->SubsurfaceColor,	VolBSDF->Extinction,		MP_SubsurfaceColor);
+					ColorMatInputConnectionTo(EditorOnly->EmissiveColor,	VolBSDF->EmissiveColor,		MP_EmissiveColor);	
+					ScalarMatInputConnectionTo(EditorOnly->AmbientOcclusion,VolBSDF->AmbientOcclusion,	MP_AmbientOcclusion);
+
+					VolBSDF->bEmissiveOnly = ShadingModel == MSM_Unlit;
+
+					// SUBSTRATE_TODO remove the VolumetricAdvancedOutput node and add the input onto FogCloudBSDF even if only used by the cloud renderer?
+					EditorOnly->FrontMaterial.Connect(0, VolBSDF);
+					bInvalidateShader = true;
+				}
+				else if (MaterialDomain == MD_LightFunction)
+				{
+					// Some materials don't have their shading mode set correctly to Unlit. Since only Unlit is supported, forcing it here.
+					ShadingModel = MSM_Unlit;
+					ShadingModels.ClearShadingModels();
+					ShadingModels.AddShadingModel(MSM_Unlit);
+
+					// Only Emissive & Opacity are valid input for PostProcess material
+					UMaterialExpressionSubstrateLightFunction* LightFunctionNode = NewObject<UMaterialExpressionSubstrateLightFunction>(this);
+					EditorOnly->ExpressionCollection.AddExpression(LightFunctionNode);
+					LightFunctionNode->Material = this;
+					SetPosXAndMoveReferenceToTheRight(LightFunctionNode);
+					ColorMatInputConnectionTo(EditorOnly->EmissiveColor, LightFunctionNode->Color, MP_EmissiveColor);
+
+					EditorOnly->FrontMaterial.Connect(0, LightFunctionNode);
+					bInvalidateShader = true;
+				}
+				else if (MaterialDomain == MD_PostProcess)
+				{
+					// Some materials don't have their shading mode set correctly to Unlit. Since only Unlit is supported, forcing it here.
+					ShadingModel = MSM_Unlit;
+					ShadingModels.ClearShadingModels();
+					ShadingModels.AddShadingModel(MSM_Unlit);
 
-					if (ShadingModel == MSM_TwoSidedFoliage)
+					if (MaterialDomain == MD_PostProcess && !IsPostProcessMaterialOutputingAlpha())
 					{
-						bIsThinSurface = true;
+						BlendMode = BLEND_Opaque;
 					}
-				}
 
-				EditorOnly->FrontMaterial.Connect(0, ConvertNode);
-
-				bInvalidateShader = true;
-			}
-			else if (MaterialDomain == MD_Volume)
-			{
-				UMaterialExpressionSubstrateVolumetricFogCloudBSDF* VolBSDF = NewObject<UMaterialExpressionSubstrateVolumetricFogCloudBSDF>(this);
-				EditorOnly->ExpressionCollection.AddExpression(VolBSDF);
-				VolBSDF->Material = this;
-				SetPosXAndMoveReferenceToTheRight(VolBSDF);
-				ColorMatInputConnectionTo(EditorOnly->BaseColor,		VolBSDF->Albedo,			MP_BaseColor);	
-				ColorMatInputConnectionTo(EditorOnly->SubsurfaceColor,	VolBSDF->Extinction,		MP_SubsurfaceColor);
-				ColorMatInputConnectionTo(EditorOnly->EmissiveColor,	VolBSDF->EmissiveColor,		MP_EmissiveColor);	
-				ScalarMatInputConnectionTo(EditorOnly->AmbientOcclusion,VolBSDF->AmbientOcclusion,	MP_AmbientOcclusion);
+					UMaterialExpressionSubstratePostProcess* PostProcNode = NewObject<UMaterialExpressionSubstratePostProcess>(this);
+					EditorOnly->ExpressionCollection.AddExpression(PostProcNode);
+					PostProcNode->Material = this;
+					SetPosXAndMoveReferenceToTheRight(PostProcNode);
 
-				VolBSDF->bEmissiveOnly = ShadingModel == MSM_Unlit;
+					ColorMatInputConnectionTo(EditorOnly->EmissiveColor, PostProcNode->Color, MP_EmissiveColor);
+					ScalarMatInputConnectionTo(EditorOnly->Opacity, PostProcNode->Opacity, MP_Opacity, SUBSTRATE_COPY_CONNECTION);	// We only copy, to keep Opacity on the root node in case BLEND_AlphaComposite is selected.
 
-				// SUBSTRATE_TODO remove the VolumetricAdvancedOutput node and add the input onto FogCloudBSDF even if only used by the cloud renderer?
-				EditorOnly->FrontMaterial.Connect(0, VolBSDF);
-				bInvalidateShader = true;
-			}
-			else if (MaterialDomain == MD_LightFunction)
-			{
-				// Some materials don't have their shading mode set correctly to Unlit. Since only Unlit is supported, forcing it here.
-				ShadingModel = MSM_Unlit;
-				ShadingModels.ClearShadingModels();
-				ShadingModels.AddShadingModel(MSM_Unlit);
+					EditorOnly->FrontMaterial.Connect(0, PostProcNode);
+					bInvalidateShader = true;
+				}
+				else if (MaterialDomain == MD_DeferredDecal)
+				{
+					// Some decal materials don't have their shading mode set correctly to DefaultLit. Since only DefaultLit is supported, forcing it here.
+					ShadingModel = MSM_DefaultLit;
+					ShadingModels.ClearShadingModels();
+					ShadingModels.AddShadingModel(MSM_DefaultLit);
 
-				// Only Emissive & Opacity are valid input for PostProcess material
-				UMaterialExpressionSubstrateLightFunction* LightFunctionNode = NewObject<UMaterialExpressionSubstrateLightFunction>(this);
-				EditorOnly->ExpressionCollection.AddExpression(LightFunctionNode);
-				LightFunctionNode->Material = this;
-				SetPosXAndMoveReferenceToTheRight(LightFunctionNode);
-				ColorMatInputConnectionTo(EditorOnly->EmissiveColor, LightFunctionNode->Color, MP_EmissiveColor);
+					ConvertNode = NewObject<UMaterialExpressionSubstrateShadingModels>(this);
+					EditorOnly->ExpressionCollection.AddExpression(ConvertNode);
+					ConvertNode->Material = this;
+					SetPosXAndMoveReferenceToTheRight(ConvertNode);
+					ColorMatInputConnectionTo(EditorOnly->BaseColor,			ConvertNode->BaseColor,			MP_BaseColor);
+					ScalarMatInputConnectionTo(EditorOnly->Metallic,			ConvertNode->Metallic,			MP_Metallic);
+					ScalarMatInputConnectionTo(EditorOnly->Specular,			ConvertNode->Specular,			MP_Specular);
+					ScalarMatInputConnectionTo(EditorOnly->Roughness,			ConvertNode->Roughness,			MP_Roughness);
+					ScalarMatInputConnectionTo(EditorOnly->Anisotropy,			ConvertNode->Anisotropy,		MP_Anisotropy);
+					ColorMatInputConnectionTo(EditorOnly->EmissiveColor,		ConvertNode->EmissiveColor,		MP_EmissiveColor);
+					Vector3MatInputConnectionTo(EditorOnly->Normal,				ConvertNode->Normal,			MP_Normal, SUBSTRATE_COPY_CONNECTION);
+					Vector3MatInputConnectionTo(EditorOnly->Tangent,			ConvertNode->Tangent,			MP_Tangent);
+					ColorMatInputConnectionTo(EditorOnly->SubsurfaceColor,		ConvertNode->SubSurfaceColor,	MP_SubsurfaceColor);
+					ScalarMatInputConnectionTo(EditorOnly->ClearCoat,			ConvertNode->ClearCoat,			MP_CustomData0);
+					ScalarMatInputConnectionTo(EditorOnly->ClearCoatRoughness,	ConvertNode->ClearCoatRoughness,MP_CustomData1);
+					ScalarMatInputConnectionTo(EditorOnly->Opacity,				ConvertNode->Opacity,			MP_Opacity, SUBSTRATE_COPY_CONNECTION);	// We only copy, to keep Opacity on the root node in case BLEND_AlphaComposite is selected.
+
+					// Add constant for the Unlit shading model
+					ConvertNode->ShadingModelOverride = ShadingModel;
+					check(ShadingModels.CountShadingModels() == 1);
 
-				EditorOnly->FrontMaterial.Connect(0, LightFunctionNode);
-				bInvalidateShader = true;
-			}
-			else if (MaterialDomain == MD_PostProcess)
-			{
-				// Some materials don't have their shading mode set correctly to Unlit. Since only Unlit is supported, forcing it here.
-				ShadingModel = MSM_Unlit;
-				ShadingModels.ClearShadingModels();
-				ShadingModels.AddShadingModel(MSM_Unlit);
+					// Now pass through the convert to decal node, which flag the material as SSM_Decal, which will set the domain to Decal.
+					UMaterialExpressionSubstrateConvertToDecal* ConvertToDecalNode= NewObject<UMaterialExpressionSubstrateConvertToDecal>(this);
+					EditorOnly->ExpressionCollection.AddExpression(ConvertToDecalNode);
+					ConvertToDecalNode->Material = this;
+					ReplaceNodeAndMoveToTheRight(ConvertNode, ConvertToDecalNode);
+					ConvertToDecalNode->DecalMaterial.Connect(0, ConvertNode);
 
-				if (MaterialDomain == MD_PostProcess && !IsPostProcessMaterialOutputingAlpha())
+					EditorOnly->FrontMaterial.Connect(0, ConvertToDecalNode);
+					bInvalidateShader = true;
+				}
+				else if (MaterialDomain == MD_UI)
 				{
-					BlendMode = BLEND_Opaque;
+					// Some materials don't have their shading mode set correctly to Unlit. Since only Unlit is supported, forcing it here.
+					ShadingModel = MSM_Unlit;
+					ShadingModels.ClearShadingModels();
+					ShadingModels.AddShadingModel(MSM_Unlit);
+
+					UMaterialExpressionSubstrateUI* UINode = NewObject<UMaterialExpressionSubstrateUI>(this);
+					EditorOnly->ExpressionCollection.AddExpression(UINode);
+					UINode->Material = this;
+					SetPosXAndMoveReferenceToTheRight(UINode);
+					ColorMatInputConnectionTo(EditorOnly->EmissiveColor,UINode->Color,	MP_EmissiveColor);
+					ScalarMatInputConnectionTo(EditorOnly->Opacity,		UINode->Opacity,MP_Opacity, SUBSTRATE_COPY_CONNECTION);	// We only copy, to keep Opacity on the root node in case BLEND_AlphaComposite is selected.
+
+					EditorOnly->FrontMaterial.Connect(0, UINode);
+					bInvalidateShader = true;
 				}
 
-				UMaterialExpressionSubstratePostProcess* PostProcNode = NewObject<UMaterialExpressionSubstratePostProcess>(this);
-				EditorOnly->ExpressionCollection.AddExpression(PostProcNode);
-				PostProcNode->Material = this;
-				SetPosXAndMoveReferenceToTheRight(PostProcNode);
+				BlendMode = ConvertLegacyBlendMode(BlendMode, ShadingModels);
+				RefractionCoverageMode = RCM_CoverageIgnored;
+			}
+		}
 
-				ColorMatInputConnectionTo(EditorOnly->EmissiveColor, PostProcNode->Color, MP_EmissiveColor);
-				ScalarMatInputConnectionTo(EditorOnly->Opacity, PostProcNode->Opacity, MP_Opacity, SUBSTRATE_COPY_CONNECTION);	// We only copy, to keep Opacity on the root node in case BLEND_AlphaComposite is selected.
+		if (bRelinkCustomOutputNodes)
+		{
+			check(ConvertNode);
+			GatherCustomNodes();
 
-				EditorOnly->FrontMaterial.Connect(0, PostProcNode);
-				bInvalidateShader = true;
-			}
-			else if (MaterialDomain == MD_DeferredDecal)
+			if (ThinTranslucentOutput)
 			{
-				// Some decal materials don't have their shading mode set correctly to DefaultLit. Since only DefaultLit is supported, forcing it here.
-				ShadingModel = MSM_DefaultLit;
-				ShadingModels.ClearShadingModels();
-				ShadingModels.AddShadingModel(MSM_DefaultLit);
-
-				ConvertNode = NewObject<UMaterialExpressionSubstrateShadingModels>(this);
-				EditorOnly->ExpressionCollection.AddExpression(ConvertNode);
-				ConvertNode->Material = this;
-				SetPosXAndMoveReferenceToTheRight(ConvertNode);
-				ColorMatInputConnectionTo(EditorOnly->BaseColor,			ConvertNode->BaseColor,			MP_BaseColor);
-				ScalarMatInputConnectionTo(EditorOnly->Metallic,			ConvertNode->Metallic,			MP_Metallic);
-				ScalarMatInputConnectionTo(EditorOnly->Specular,			ConvertNode->Specular,			MP_Specular);
-				ScalarMatInputConnectionTo(EditorOnly->Roughness,			ConvertNode->Roughness,			MP_Roughness);
-				ScalarMatInputConnectionTo(EditorOnly->Anisotropy,			ConvertNode->Anisotropy,		MP_Anisotropy);
-				ColorMatInputConnectionTo(EditorOnly->EmissiveColor,		ConvertNode->EmissiveColor,		MP_EmissiveColor);
-				Vector3MatInputConnectionTo(EditorOnly->Normal,				ConvertNode->Normal,			MP_Normal, SUBSTRATE_COPY_CONNECTION);
-				Vector3MatInputConnectionTo(EditorOnly->Tangent,			ConvertNode->Tangent,			MP_Tangent);
-				ColorMatInputConnectionTo(EditorOnly->SubsurfaceColor,		ConvertNode->SubSurfaceColor,	MP_SubsurfaceColor);
-				ScalarMatInputConnectionTo(EditorOnly->ClearCoat,			ConvertNode->ClearCoat,			MP_CustomData0);
-				ScalarMatInputConnectionTo(EditorOnly->ClearCoatRoughness,	ConvertNode->ClearCoatRoughness,MP_CustomData1);
-				ScalarMatInputConnectionTo(EditorOnly->Opacity,				ConvertNode->Opacity,			MP_Opacity, SUBSTRATE_COPY_CONNECTION);	// We only copy, to keep Opacity on the root node in case BLEND_AlphaComposite is selected.
-
-				// Add constant for the Unlit shading model
-				ConvertNode->ShadingModelOverride = ShadingModel;
-				check(ShadingModels.CountShadingModels() == 1);
-
-				// Now pass through the convert to decal node, which flag the material as SSM_Decal, which will set the domain to Decal.
-				UMaterialExpressionSubstrateConvertToDecal* ConvertToDecalNode= NewObject<UMaterialExpressionSubstrateConvertToDecal>(this);
-				EditorOnly->ExpressionCollection.AddExpression(ConvertToDecalNode);
-				ConvertToDecalNode->Material = this;
-				ReplaceNodeAndMoveToTheRight(ConvertNode, ConvertToDecalNode);
-				ConvertToDecalNode->DecalMaterial.Connect(0, ConvertNode);
-
-				EditorOnly->FrontMaterial.Connect(0, ConvertToDecalNode);
-				bInvalidateShader = true;
+				FExpressionInput* TransmittanceColorInput = &ThinTranslucentOutput->TransmittanceColor;
+				FExpressionInput* ThinTranslucentSurfaceCoverageInput = &ThinTranslucentOutput->SurfaceCoverage;
+				if (TransmittanceColorInput && TransmittanceColorInput->IsConnected())
+				{
+					ConvertNode->TransmittanceColor.Connect(TransmittanceColorInput->OutputIndex, TransmittanceColorInput->Expression);
+				}
+				if (ThinTranslucentSurfaceCoverageInput && ThinTranslucentSurfaceCoverageInput->IsConnected())
+				{
+					ConvertNode->ThinTranslucentSurfaceCoverage.Connect(ThinTranslucentSurfaceCoverageInput->OutputIndex, ThinTranslucentSurfaceCoverageInput->Expression);
+				}
 			}
-			else if (MaterialDomain == MD_UI)
+			if (SingleLayerWaterOutput)
 			{
-				// Some materials don't have their shading mode set correctly to Unlit. Since only Unlit is supported, forcing it here.
-				ShadingModel = MSM_Unlit;
-				ShadingModels.ClearShadingModels();
-				ShadingModels.AddShadingModel(MSM_Unlit);
-
-				UMaterialExpressionSubstrateUI* UINode = NewObject<UMaterialExpressionSubstrateUI>(this);
-				EditorOnly->ExpressionCollection.AddExpression(UINode);
-				UINode->Material = this;
-				SetPosXAndMoveReferenceToTheRight(UINode);
-				ColorMatInputConnectionTo(EditorOnly->EmissiveColor,UINode->Color,	MP_EmissiveColor);
-				ScalarMatInputConnectionTo(EditorOnly->Opacity,		UINode->Opacity,MP_Opacity, SUBSTRATE_COPY_CONNECTION);	// We only copy, to keep Opacity on the root node in case BLEND_AlphaComposite is selected.
-
-				EditorOnly->FrontMaterial.Connect(0, UINode);
-				bInvalidateShader = true;
+				ConnectionTo(SingleLayerWaterOutput->ScatteringCoefficients,	ConvertNode->WaterScatteringCoefficients);
+				ConnectionTo(SingleLayerWaterOutput->AbsorptionCoefficients,	ConvertNode->WaterAbsorptionCoefficients);
+				ConnectionTo(SingleLayerWaterOutput->PhaseG,					ConvertNode->WaterPhaseG);
+				ConnectionTo(SingleLayerWaterOutput->ColorScaleBehindWater,		ConvertNode->ColorScaleBehindWater);
 			}
-
-			BlendMode = ConvertLegacyBlendMode(BlendMode, ShadingModels);
-			RefractionCoverageMode = RCM_CoverageIgnored;
-		}
-	}
-
-	if (bRelinkCustomOutputNodes)
-	{
-		check(ConvertNode);
-		GatherCustomNodes();
-
-		if (ThinTranslucentOutput)
-		{
-			FExpressionInput* TransmittanceColorInput = &ThinTranslucentOutput->TransmittanceColor;
-			FExpressionInput* ThinTranslucentSurfaceCoverageInput = &ThinTranslucentOutput->SurfaceCoverage;
-			if (TransmittanceColorInput && TransmittanceColorInput->IsConnected())
+			if (ClearCoatBottomNormalOutput)
 			{
-				ConvertNode->TransmittanceColor.Connect(TransmittanceColorInput->OutputIndex, TransmittanceColorInput->Expression);
+				ConnectionTo(ClearCoatBottomNormalOutput->Input, ConvertNode->ClearCoatNormal, SUBSTRATE_COPY_CONNECTION);
 			}
-			if (ThinTranslucentSurfaceCoverageInput && ThinTranslucentSurfaceCoverageInput->IsConnected())
+			if (TangentOutput)
 			{
-				ConvertNode->ThinTranslucentSurfaceCoverage.Connect(ThinTranslucentSurfaceCoverageInput->OutputIndex, ThinTranslucentSurfaceCoverageInput->Expression);
+				ConnectionTo(TangentOutput->Input, ConvertNode->Tangent, SUBSTRATE_COPY_CONNECTION);
 			}
 		}
-		if (SingleLayerWaterOutput)
-		{
-			ConnectionTo(SingleLayerWaterOutput->ScatteringCoefficients,	ConvertNode->WaterScatteringCoefficients);
-			ConnectionTo(SingleLayerWaterOutput->AbsorptionCoefficients,	ConvertNode->WaterAbsorptionCoefficients);
-			ConnectionTo(SingleLayerWaterOutput->PhaseG,					ConvertNode->WaterPhaseG);
-			ConnectionTo(SingleLayerWaterOutput->ColorScaleBehindWater,		ConvertNode->ColorScaleBehindWater);
-		}
-		if (ClearCoatBottomNormalOutput)
-		{
-			ConnectionTo(ClearCoatBottomNormalOutput->Input, ConvertNode->ClearCoatNormal, SUBSTRATE_COPY_CONNECTION);
-		}
-		if (TangentOutput)
-		{
-			ConnectionTo(TangentOutput->Input, ConvertNode->Tangent, SUBSTRATE_COPY_CONNECTION);
-		}
 	}
 
 	// Glint conversion parametrization
@@ -4911,10 +4976,12 @@ bool UMaterial::CanEditChange(const FProperty* InProperty) const
 		const UMaterialEditorOnlyData* EditorOnly = GetEditorOnlyData();
 		FString PropertyName = InProperty->GetName();
 		const bool bSubstrateEnabled = Substrate::IsSubstrateEnabled();
+		const bool bSubstrateHiddenMaterialAssetConversion = bSubstrateEnabled && Substrate::IsHiddenMaterialAssetConversionEnabled() && !EditorOnly->FrontMaterial.IsConnected();
 
 		if (PropertyName == GET_MEMBER_NAME_STRING_CHECKED(UMaterial, MaterialDomain))
 		{
-			return !bSubstrateEnabled; // Material domain is no longer tweakable with Substrate. It is instead derived from the graph.
+			// Material domain is no longer tweakable with Substrate. It is instead derived from the graph.
+			return !bSubstrateEnabled || bSubstrateHiddenMaterialAssetConversion;
 		}
 
 		if (PropertyName == GET_MEMBER_NAME_STRING_CHECKED(UMaterial, PhysMaterial) || PropertyName == GET_MEMBER_NAME_STRING_CHECKED(UMaterial, PhysMaterialMask))
@@ -5008,7 +5075,7 @@ bool UMaterial::CanEditChange(const FProperty* InProperty) const
 		if (PropertyName == GET_MEMBER_NAME_STRING_CHECKED(UMaterial, ShadingModel))
 		{
 			// Volume can also use the unlit mode for emissive only volumetric effect (also overridable on material instance)
-			return !bSubstrateEnabled && (MaterialDomain == MD_Surface || MaterialDomain == MD_Volume);
+			return (!bSubstrateEnabled || bSubstrateHiddenMaterialAssetConversion) && (MaterialDomain == MD_Surface || MaterialDomain == MD_Volume);
 		}
 
 		if (PropertyName == GET_MEMBER_NAME_STRING_CHECKED(UMaterial, bIsThinSurface))
@@ -5119,6 +5186,12 @@ static bool IsMaterialAlreadyConvertedToSubstrate(const UMaterial* InMaterial, c
 		InMaterial->HasAnyExpressionsInMaterialAndFunctionsOfType<UMaterialExpressionSubstrateBSDF>();
 }
 
+static bool ShouldAllowSubstrateReset(const UMaterialEditorOnlyData* EditorOnly)
+{
+	// We do reset the conversion to as "has been run" when Substrate is authored using the root node instead of FrontMaterial.
+	return !Substrate::IsHiddenMaterialAssetConversionEnabled() || (EditorOnly && EditorOnly->FrontMaterial.IsConnected());
+}
+
 void UMaterial::PostEditChangePropertyInternal(FPropertyChangedEvent& PropertyChangedEvent, const EPostEditChangeEffectOnShaders EffectOnShaders)
 {
 	// PreEditChange is not enforced to be called before PostEditChange.
@@ -5135,7 +5208,7 @@ void UMaterial::PostEditChangePropertyInternal(FPropertyChangedEvent& PropertyCh
 
 	const UMaterialEditorOnlyData* EditorOnly = GetEditorOnlyData();
 
-	if (Substrate::IsSubstrateEnabled())
+	if (Substrate::IsSubstrateEnabled() && ShouldAllowSubstrateReset(EditorOnly))
 	{
 		// If a modification is made to the material, revert the auto-conversion version, to mark the material as 'edited'
 		const bool bResetConversionVersion = PropertyThatChanged != nullptr || PropertyChangedEvent.ChangeType != EPropertyChangeType::Unspecified;
@@ -7381,7 +7454,7 @@ bool UMaterial::IsPropertySupported(EMaterialProperty InProperty) const
 {
 	bool bSupported = true;
 
-	if (Substrate::IsSubstrateEnabled())
+	if (Substrate::IsSubstrateEnabled() && !Substrate::IsHiddenMaterialAssetConversionEnabled())
 	{
 		bSupported = false;
 		switch (InProperty)
@@ -7436,14 +7509,23 @@ static bool IsPropertyActive_Internal(EMaterialProperty InProperty,
 	bool bUsesShadingModelFromMaterialExpression,
 	bool bIsTranslucencyWritingVelocity,
 	bool bIsThinSurface,
-	bool bIsSupported)
+	bool bIsSupported,
+	bool bFrontMaterialIsConnected)
 {
 	const bool bSubstrateEnabled = Substrate::IsSubstrateEnabled();
+	const bool bSubstrateHiddenConversion = Substrate::IsHiddenMaterialAssetConversionEnabled();
+	const bool bShowAsSubstrate = bSubstrateEnabled && (!bSubstrateHiddenConversion || bFrontMaterialIsConnected);
+
+	if (bSubstrateEnabled && bSubstrateHiddenConversion && InProperty == MP_FrontMaterial)
+	{
+		return true; // If hidden conversion, always enable FrontMaterial in case we want to plug a Substrate material in there.
+	}
+
 	const bool bSubstrateOpacityOverrideAllowed = BlendMode == BLEND_AlphaComposite; // Should we always have it enabled to be able to be plugged in an fed when blend mode is toggled later on a material instance?
 
 	if (Domain == MD_PostProcess)
 	{
-		if (bSubstrateEnabled)
+		if (bShowAsSubstrate)
 		{
 			return InProperty == MP_FrontMaterial || (InProperty == MP_Opacity && bSubstrateOpacityOverrideAllowed);
 		}
@@ -7455,7 +7537,7 @@ static bool IsPropertyActive_Internal(EMaterialProperty InProperty,
 	else if (Domain == MD_LightFunction)
 	{
 		// light functions should already use MSM_Unlit but we also we don't want WorldPosOffset
-		if (bSubstrateEnabled)
+		if (bShowAsSubstrate)
 		{
 			return InProperty == MP_FrontMaterial;
 		}
@@ -7466,7 +7548,7 @@ static bool IsPropertyActive_Internal(EMaterialProperty InProperty,
 	}
 	else if (Domain == MD_DeferredDecal)
 	{
-		if (bSubstrateEnabled)
+		if (bShowAsSubstrate)
 		{
 			return InProperty == MP_FrontMaterial
 				|| InProperty == MP_AmbientOcclusion
@@ -7525,7 +7607,7 @@ static bool IsPropertyActive_Internal(EMaterialProperty InProperty,
 	}
 	else if (Domain == MD_Volume)
 	{
-		if (bSubstrateEnabled)
+		if (bShowAsSubstrate)
 		{
 			return InProperty == MP_FrontMaterial;
 		}
@@ -7536,7 +7618,7 @@ static bool IsPropertyActive_Internal(EMaterialProperty InProperty,
 	}
 	else if (Domain == MD_UI)
 	{
-		if (bSubstrateEnabled)
+		if (bShowAsSubstrate)
 		{
 			return InProperty == MP_FrontMaterial
 				|| (InProperty == MP_WorldPositionOffset)
@@ -7546,12 +7628,12 @@ static bool IsPropertyActive_Internal(EMaterialProperty InProperty,
 		}
 		else
 		{
-		return InProperty == MP_EmissiveColor
-			|| (InProperty == MP_WorldPositionOffset)
+			return InProperty == MP_EmissiveColor
+				|| (InProperty == MP_WorldPositionOffset)
 				|| (InProperty == MP_OpacityMask && IsMaskedBlendMode(BlendMode))
-			|| (InProperty == MP_Opacity && IsTranslucentBlendMode(BlendMode) && BlendMode != BLEND_Modulate)
-			|| (InProperty >= MP_CustomizedUVs0 && InProperty <= MP_CustomizedUVs7);
-	}
+				|| (InProperty == MP_Opacity && IsTranslucentBlendMode(BlendMode) && BlendMode != BLEND_Modulate)
+				|| (InProperty >= MP_CustomizedUVs0 && InProperty <= MP_CustomizedUVs7);
+		}
 	}
 
 	// Now processing MD_Surface
@@ -7565,7 +7647,7 @@ static bool IsPropertyActive_Internal(EMaterialProperty InProperty,
 	
 	bool Active = true;
 
-	if (bSubstrateEnabled)
+	if (bShowAsSubstrate)
 	{
 		Active = false;
 		if (bIsSupported)
@@ -7682,7 +7764,7 @@ static bool IsPropertyActive_Internal(EMaterialProperty InProperty,
 		case MP_SurfaceThickness:
 		case MP_FrontMaterial:
 			{
-				Active = false;
+				Active = bSubstrateHiddenConversion;
 				break;
 			}
 		case MP_MaterialAttributes:
@@ -7715,12 +7797,14 @@ bool UMaterial::IsPropertyActiveInEditor(EMaterialProperty InProperty) const
 		IsShadingModelFromMaterialExpression(),
 		IsTranslucencyWritingVelocity(),
 		IsThinSurface(),
-		IsPropertySupported(InProperty));
+		IsPropertySupported(InProperty),
+		GetEditorOnlyData()->FrontMaterial.IsConnected());
 }
 #endif // WITH_EDITOR
 
 bool UMaterial::IsPropertyActiveInDerived(EMaterialProperty InProperty, const UMaterialInterface* DerivedMaterial) const
 {
+	bool bFrontMaterialConnected = false;
 #if WITH_EDITOR
 	if (!bUseMaterialAttributes)
 	{
@@ -7729,6 +7813,7 @@ bool UMaterial::IsPropertyActiveInDerived(EMaterialProperty InProperty, const UM
 		ensureMsgf(!(GetEditorOnlyData()->Refraction.IsConnected() && !GetCachedExpressionData().IsPropertyConnected(MP_Refraction)),
 			TEXT("GetCachedExpressionData() says refraction isn't connected, but GetEditorOnlyData() says it is"));
 	}
+	bFrontMaterialConnected = GetEditorOnlyData()->FrontMaterial.IsConnected();
 #endif
 	return IsPropertyActive_Internal(InProperty,
 		MaterialDomain,
@@ -7741,7 +7826,8 @@ bool UMaterial::IsPropertyActiveInDerived(EMaterialProperty InProperty, const UM
 		DerivedMaterial->IsShadingModelFromMaterialExpression(),
 		IsTranslucencyWritingVelocity(),
 		IsThinSurface(),
-		IsPropertySupported(InProperty));
+		IsPropertySupported(InProperty),
+		bFrontMaterialConnected);
 }
 
 #if WITH_EDITORONLY_DATA
@@ -7933,7 +8019,12 @@ UMaterialEditorOnlyData::UMaterialEditorOnlyData()
 
 void UMaterialEditorOnlyData::ResetSubstrateConversionVersion()
 {
-	SubstrateConversionVersion = GetSubstrateNoConversionVersion();
+#if WITH_EDITOR
+	if (ShouldAllowSubstrateReset(this))
+#endif
+	{
+		SubstrateConversionVersion = GetSubstrateNoConversionVersion();
+	}
 }
 
 bool UMaterialEditorOnlyData::IsSubstrateAutoConverted() const
diff --git a/Engine/Source/Runtime/Engine/Private/Materials/MaterialExpressionSubstrate.cpp b/Engine/Source/Runtime/Engine/Private/Materials/MaterialExpressionSubstrate.cpp
--- a/Engine/Source/Runtime/Engine/Private/Materials/MaterialExpressionSubstrate.cpp
+++ b/Engine/Source/Runtime/Engine/Private/Materials/MaterialExpressionSubstrate.cpp
@@ -11,6 +11,8 @@
 #include "Materials/MaterialExpressionMakeMaterialAttributes.h"
 #include "Engine/SubsurfaceProfile.h"
 #include "Engine/SpecularProfile.h"
+#include "MaterialExpressionIO.h"
+
 #if WITH_EDITOR
 #include "MaterialGraph/MaterialGraphSchema.h"
 #include "SubstrateMaterial.h"
@@ -79,73 +81,38 @@ static int32 CastToNonLWCType(class FMaterialCompiler* Compiler, int32 CodeChunk
 // or when the node is inside a material function, linked to an input pin of the material function and that input is not plugged in to anything.
 // But for normals or tangents, Substrate absolutely need a valid code chunk to de-duplicate when stored in memory. 
 // Also, we want all our nodes to have default, as that is needed when creating BSDF, when registering code chunk representing material topology.
-static int32 CompileWithDefaultCodeChunk(class FMaterialCompiler* Compiler, FExpressionInput& Input, int DefaultCodeChunk, bool* bDefaultIsUsed = nullptr)
+static int32 CompileWithDefaultFloat1(class FMaterialCompiler* Compiler, FExpressionInput& Input, float X, const FScalarMaterialInput* RootNodeInput = nullptr)
 {
-	if (bDefaultIsUsed)
+	int32 DefaultCodeChunk = Compiler->Constant(X);
+	if (RootNodeInput && RootNodeInput->UseConstant)
 	{
-		*bDefaultIsUsed = Input.GetTracedInput().Expression == nullptr;
+		DefaultCodeChunk = Compiler->Constant(RootNodeInput->Constant);
 	}
 	int32 CodeChunk = Input.GetTracedInput().Expression ? Input.Compile(Compiler) : DefaultCodeChunk;
-	if (bDefaultIsUsed)
-	{
-		*bDefaultIsUsed |= CodeChunk == INDEX_NONE;
-	}
-	else
-	{
-		CodeChunk = CastToNonLWCType(Compiler, CodeChunk);
-	}
+	CodeChunk = CastToNonLWCType(Compiler, CodeChunk);
 	return CodeChunk == INDEX_NONE ? DefaultCodeChunk : CodeChunk;
 }
-static int32 CompileWithDefaultFloat1(class FMaterialCompiler* Compiler, FExpressionInput& Input, float X, bool* bDefaultIsUsed = nullptr)
-{
-	if (bDefaultIsUsed)
-	{
-		*bDefaultIsUsed = Input.GetTracedInput().Expression == nullptr;
-	}
-	int32 CodeChunk = Input.GetTracedInput().Expression ? Input.Compile(Compiler) : Compiler->Constant(X);
-	if (bDefaultIsUsed)
-	{
-		*bDefaultIsUsed |= CodeChunk == INDEX_NONE;
-	}
-	else
-	{
-		CodeChunk = CastToNonLWCType(Compiler, CodeChunk);
-	}
-	return CodeChunk == INDEX_NONE ? Compiler->Constant(X) : CodeChunk;
-}
-static int32 CompileWithDefaultFloat2(class FMaterialCompiler* Compiler, FExpressionInput& Input, float X, float Y, bool* bDefaultIsUsed = nullptr)
+static int32 CompileWithDefaultFloat2(class FMaterialCompiler* Compiler, FExpressionInput& Input, float X, float Y, const FVector2MaterialInput* RootNodeInput = nullptr)
 {
-	if (bDefaultIsUsed)
-	{
-		*bDefaultIsUsed = Input.GetTracedInput().Expression == nullptr;
-	}
-	int32 CodeChunk = Input.GetTracedInput().Expression ? Input.Compile(Compiler) : Compiler->Constant2(X, Y);
-	if (bDefaultIsUsed)
-	{
-		*bDefaultIsUsed |= CodeChunk == INDEX_NONE;
-	}
-	else
+	int32 DefaultCodeChunk = Compiler->Constant2(X, Y);
+	if (RootNodeInput && RootNodeInput->UseConstant)
 	{
-		CodeChunk = CastToNonLWCType(Compiler, CodeChunk);
+		DefaultCodeChunk = Compiler->Constant2(RootNodeInput->Constant.X, RootNodeInput->Constant.Y);
 	}
-	return CodeChunk == INDEX_NONE ? Compiler->Constant2(X, Y) : CodeChunk;
+	int32 CodeChunk = Input.GetTracedInput().Expression ? Input.Compile(Compiler) : DefaultCodeChunk;
+	CodeChunk = CastToNonLWCType(Compiler, CodeChunk);
+	return CodeChunk == INDEX_NONE ? DefaultCodeChunk : CodeChunk;
 }
-static int32 CompileWithDefaultFloat3(class FMaterialCompiler* Compiler, FExpressionInput& Input, float X, float Y, float Z, bool* bDefaultIsUsed = nullptr)
+static int32 CompileWithDefaultFloat3(class FMaterialCompiler* Compiler, FExpressionInput& Input, float X, float Y, float Z, const FColorMaterialInput* RootNodeInput = nullptr)
 {
-	if (bDefaultIsUsed)
-	{
-		*bDefaultIsUsed = Input.GetTracedInput().Expression == nullptr;
-	}
-	int32 CodeChunk = Input.GetTracedInput().Expression ? Input.Compile(Compiler) : Compiler->Constant3(X, Y, Z);
-	if (bDefaultIsUsed)
-	{
-		*bDefaultIsUsed |= CodeChunk == INDEX_NONE;
-	}
-	else
+	int32 DefaultCodeChunk = Compiler->Constant3(X, Y, Z);
+	if (RootNodeInput && RootNodeInput->UseConstant)
 	{
-		CodeChunk = CastToNonLWCType(Compiler, CodeChunk);
+		DefaultCodeChunk = Compiler->Constant3(RootNodeInput->Constant.R, RootNodeInput->Constant.G, RootNodeInput->Constant.B);
 	}
-	return CodeChunk == INDEX_NONE ? Compiler->Constant3(X, Y, Z) : CodeChunk;
+	int32 CodeChunk = Input.GetTracedInput().Expression ? Input.Compile(Compiler) : DefaultCodeChunk;
+	CodeChunk = CastToNonLWCType(Compiler, CodeChunk);
+	return CodeChunk == INDEX_NONE ? DefaultCodeChunk : CodeChunk;
 }
 static int32 CompileWithDefaultNormalWS(class FMaterialCompiler* Compiler, FExpressionInput& Input, bool bConvertToRequestedSpace = true)
 {
@@ -255,14 +222,21 @@ FExpressionInput* UMaterialExpressionSubstrateShadingModels::GetInput(int32 Inpu
 																					int32 CodeChunkResult = Code;\
 																					Compiler->PopMaterialAttribute();
 
-int32 UMaterialExpressionSubstrateShadingModels::Compile(class FMaterialCompiler* Compiler, int32 OutputIndex)
-{
-	LEGACY_DIRECT_ATTRIBUTE_MAPPING(RoughnessCodeChunk, MP_Roughness, CompileWithDefaultFloat1(Compiler, Roughness, 0.5f));
-	LEGACY_DIRECT_ATTRIBUTE_MAPPING(AnisotropyCodeChunk, MP_Anisotropy, CompileWithDefaultFloat1(Compiler, Anisotropy, 0.0f));
 
-	// We also cannot ignore the tangent when using the default Tangent because GetTangentBasis
-	// used in SubstrateGetBSDFSharedBasis cannot be relied on for smooth tangent used for lighting on any mesh.
-	const bool bHasAnisotropy = HasAnisotropy();
+int32 UMaterialExpressionSubstrateShadingModels::CompileCommon(class FMaterialCompiler* Compiler,
+	FExpressionInput& BaseColor, FExpressionInput& Specular, FExpressionInput& Metallic, FExpressionInput& Roughness, FExpressionInput& EmissiveColor,
+	FExpressionInput& Opacity, FExpressionInput& SubSurfaceColor, FExpressionInput& ClearCoat, FExpressionInput& ClearCoatRoughness,
+	FExpressionInput& ShadingModel, TEnumAsByte<enum EMaterialShadingModel> ShadingModelOverride,
+	FExpressionInput& TransmittanceColor, FExpressionInput& ThinTranslucentSurfaceCoverage,
+	FExpressionInput& WaterScatteringCoefficients, FExpressionInput& WaterAbsorptionCoefficients, FExpressionInput& WaterPhaseG, FExpressionInput& ColorScaleBehindWater,
+	const bool bHasAnisotropy, FExpressionInput& Anisotropy,
+	FExpressionInput& Normal, FExpressionInput& Tangent,
+	FExpressionInput& ClearCoatNormal, FExpressionInput& CustomTangent,
+	const bool bHasSSS, USubsurfaceProfile* SSSProfile,
+	const UMaterialEditorOnlyData* EditorOnlyData)
+{
+	LEGACY_DIRECT_ATTRIBUTE_MAPPING(RoughnessCodeChunk, MP_Roughness, CompileWithDefaultFloat1(Compiler, Roughness, 0.5f, EditorOnlyData ? &EditorOnlyData->Roughness : nullptr));
+	LEGACY_DIRECT_ATTRIBUTE_MAPPING(AnisotropyCodeChunk, MP_Anisotropy, CompileWithDefaultFloat1(Compiler, Anisotropy, 0.0f, EditorOnlyData ? &EditorOnlyData->Anisotropy : nullptr));
 
 	// Regular normal basis
 	LEGACY_DIRECT_ATTRIBUTE_MAPPING(NormalCodeChunk, MP_Normal, CompileWithDefaultNormalWS(Compiler, Normal));
@@ -317,10 +291,9 @@ int32 UMaterialExpressionSubstrateShadingModels::Compile(class FMaterialCompiler
 	}
 
 	int32 SSSProfileCodeChunk = INDEX_NONE;
-	const bool bHasSSS = HasSSS();
-	if (bHasSSS && SubsurfaceProfile)
+	if (bHasSSS && SSSProfile)
 	{
-		SSSProfileCodeChunk = CreateSubsurfaceProfileParameter(Compiler, SubsurfaceProfile);
+		SSSProfileCodeChunk = CreateSubsurfaceProfileParameter(Compiler, SSSProfile);
 	}
 	else
 	{
@@ -335,7 +308,7 @@ int32 UMaterialExpressionSubstrateShadingModels::Compile(class FMaterialCompiler
 	{
 		// We evaluate opacity only for shading models and blending mode requiring it.
 		// For instance, a translucent shader reading depth for soft fading should no evaluate opacity when an instance forces an opaque mode.
-		LEGACY_DIRECT_ATTRIBUTE_MAPPING(OpacityCodeChunkTmp, MP_Opacity, CompileWithDefaultFloat1(Compiler, Opacity, 1.0f));
+		LEGACY_DIRECT_ATTRIBUTE_MAPPING(OpacityCodeChunkTmp, MP_Opacity, CompileWithDefaultFloat1(Compiler, Opacity, 1.0f, EditorOnlyData ? &EditorOnlyData->Opacity : nullptr));
 		OpacityCodeChunk = OpacityCodeChunkTmp;
 	}
 	else
@@ -343,16 +316,16 @@ int32 UMaterialExpressionSubstrateShadingModels::Compile(class FMaterialCompiler
 		OpacityCodeChunk = Compiler->Constant(1.0f);
 	}
 
-	LEGACY_DIRECT_ATTRIBUTE_MAPPING(EmissiveCodeChunk, MP_EmissiveColor, CompileWithDefaultFloat3(Compiler, EmissiveColor, 0.0f, 0.0f, 0.0f));
+	LEGACY_DIRECT_ATTRIBUTE_MAPPING(EmissiveCodeChunk, MP_EmissiveColor, CompileWithDefaultFloat3(Compiler, EmissiveColor, 0.0f, 0.0f, 0.0f, EditorOnlyData ? &EditorOnlyData->EmissiveColor : nullptr));
 
-	LEGACY_DIRECT_ATTRIBUTE_MAPPING(BaseColorCodeChunk, MP_BaseColor, CompileWithDefaultFloat3(Compiler, BaseColor, 0.0f, 0.0f, 0.0f));
-	LEGACY_DIRECT_ATTRIBUTE_MAPPING(SpecularChunk, MP_Specular, CompileWithDefaultFloat1(Compiler, Specular, 0.5f));
-	LEGACY_DIRECT_ATTRIBUTE_MAPPING(MetallicCodeChunk, MP_Metallic, CompileWithDefaultFloat1(Compiler, Metallic, 0.0f));
+	LEGACY_DIRECT_ATTRIBUTE_MAPPING(BaseColorCodeChunk, MP_BaseColor, CompileWithDefaultFloat3(Compiler, BaseColor, 0.0f, 0.0f, 0.0f, EditorOnlyData ? &EditorOnlyData->BaseColor : nullptr));
+	LEGACY_DIRECT_ATTRIBUTE_MAPPING(SpecularChunk, MP_Specular, CompileWithDefaultFloat1(Compiler, Specular, 0.5f, EditorOnlyData ? &EditorOnlyData->Specular: nullptr));
+	LEGACY_DIRECT_ATTRIBUTE_MAPPING(MetallicCodeChunk, MP_Metallic, CompileWithDefaultFloat1(Compiler, Metallic, 0.0f, EditorOnlyData ? &EditorOnlyData->Metallic : nullptr));
 
-	LEGACY_DIRECT_ATTRIBUTE_MAPPING(SubSurfaceColorCodeChunk, MP_SubsurfaceColor, CompileWithDefaultFloat3(Compiler, SubSurfaceColor, 1.0f, 1.0f, 1.0f));
+	LEGACY_DIRECT_ATTRIBUTE_MAPPING(SubSurfaceColorCodeChunk, MP_SubsurfaceColor, CompileWithDefaultFloat3(Compiler, SubSurfaceColor, 1.0f, 1.0f, 1.0f, EditorOnlyData ? &EditorOnlyData->SubsurfaceColor : nullptr));
 
-	LEGACY_DIRECT_ATTRIBUTE_MAPPING(ClearCoatCodeChunk, MP_CustomData0, CompileWithDefaultFloat1(Compiler, ClearCoat, 1.0f));
-	LEGACY_DIRECT_ATTRIBUTE_MAPPING(ClearCoatRoughnessCodeChunk, MP_CustomData1, CompileWithDefaultFloat1(Compiler, ClearCoatRoughness, 0.1f));
+	LEGACY_DIRECT_ATTRIBUTE_MAPPING(ClearCoatCodeChunk, MP_CustomData0, CompileWithDefaultFloat1(Compiler, ClearCoat, 1.0f, EditorOnlyData ? &EditorOnlyData->ClearCoat : nullptr));
+	LEGACY_DIRECT_ATTRIBUTE_MAPPING(ClearCoatRoughnessCodeChunk, MP_CustomData1, CompileWithDefaultFloat1(Compiler, ClearCoatRoughness, 0.1f, EditorOnlyData ? &EditorOnlyData->ClearCoatRoughness: nullptr));
 
 	int32 ShadingModelCodeChunk = ShadingModel.IsConnected() ? CompileWithDefaultFloat1(Compiler, ShadingModel, float(MSM_DefaultLit)) : Compiler->Constant(float(ShadingModelOverride));
 	int32 ShadingModelCount = Compiler->GetMaterialShadingModels().CountShadingModels();
@@ -396,6 +369,20 @@ int32 UMaterialExpressionSubstrateShadingModels::Compile(class FMaterialCompiler
 	return OutputCodeChunk;
 }
 
+int32 UMaterialExpressionSubstrateShadingModels::Compile(class FMaterialCompiler* Compiler, int32 OutputIndex)
+{
+	return UMaterialExpressionSubstrateShadingModels::CompileCommon(Compiler,
+		BaseColor, Specular, Metallic, Roughness, EmissiveColor,
+		Opacity, SubSurfaceColor, ClearCoat, ClearCoatRoughness,
+		ShadingModel, ShadingModelOverride,
+		TransmittanceColor, ThinTranslucentSurfaceCoverage,
+		WaterScatteringCoefficients, WaterAbsorptionCoefficients, WaterPhaseG, ColorScaleBehindWater,
+		HasAnisotropy(), Anisotropy,
+		Normal, Tangent,
+		ClearCoatNormal, CustomTangent,
+		HasSSS(), SubsurfaceProfile);
+}
+
 
 void UMaterialExpressionSubstrateShadingModels::PostEditChangeProperty(FPropertyChangedEvent& PropertyChangedEvent)
 {
@@ -598,14 +585,20 @@ void UMaterialExpressionSubstrateShadingModels::GatherSubstrateMaterialInfo(FSub
 	}
 }
 
-FSubstrateOperator* UMaterialExpressionSubstrateShadingModels::SubstrateGenerateMaterialTopologyTree(class FMaterialCompiler* Compiler, class UMaterialExpression* Parent, int32 OutputIndex)
+FSubstrateOperator* UMaterialExpressionSubstrateShadingModels::SubstrateGenerateMaterialTopologyTreeCommon(
+	class FMaterialCompiler* Compiler, struct FGuid ThisExpressionGuid, class UMaterialExpression* Parent, int32 OutputIndex,
+	const FExpressionInput& EmissiveColor,
+	const FExpressionInput& Anisotropy,
+	const FExpressionInput& ClearCoatNormal,
+	const FExpressionInput& CustomTangent,
+	const FExpressionInput& ShadingModel)
 {
 	// Note Thickness has no meaning/usage in the context of SubstrateLegacyConversionNode
 	int32 ThicknessIndex = Compiler->SubstrateThicknessStackGetThicknessIndex();
 
 	auto AddDefaultWorstCase = [&](uint32 In)
 		{
-			FSubstrateOperator& SlabOperator = Compiler->SubstrateCompilationRegisterOperator(SUBSTRATE_OPERATOR_BSDF_LEGACY, Compiler->SubstrateTreeStackGetPathUniqueId(), this, Parent, Compiler->SubstrateTreeStackGetParentPathUniqueId());
+			FSubstrateOperator& SlabOperator = Compiler->SubstrateCompilationRegisterOperator(SUBSTRATE_OPERATOR_BSDF_LEGACY, Compiler->SubstrateTreeStackGetPathUniqueId(), ThisExpressionGuid, Parent, Compiler->SubstrateTreeStackGetParentPathUniqueId());
 			SlabOperator.BSDFType = SUBSTRATE_BSDF_TYPE_SLAB;
 			SlabOperator.BSDFFeatures = In | (Anisotropy.IsConnected() ? ESubstrateBsdfFeature_Anisotropy : ESubstrateBsdfFeature_None);
 			SlabOperator.SubSurfaceType = uint8(EMaterialSubSurfaceType::MSS_None);
@@ -633,7 +626,7 @@ FSubstrateOperator* UMaterialExpressionSubstrateShadingModels::SubstrateGenerate
 		// Special case for unlit only material to get fast path
 		if (ShadingModels.HasOnlyShadingModel(MSM_Unlit))
 		{
-			FSubstrateOperator& Operator = Compiler->SubstrateCompilationRegisterOperator(SUBSTRATE_OPERATOR_BSDF_LEGACY, Compiler->SubstrateTreeStackGetPathUniqueId(), this, Parent, Compiler->SubstrateTreeStackGetParentPathUniqueId());
+			FSubstrateOperator& Operator = Compiler->SubstrateCompilationRegisterOperator(SUBSTRATE_OPERATOR_BSDF_LEGACY, Compiler->SubstrateTreeStackGetPathUniqueId(), ThisExpressionGuid, Parent, Compiler->SubstrateTreeStackGetParentPathUniqueId());
 			Operator.BSDFType = SUBSTRATE_BSDF_TYPE_UNLIT;
 			Operator.ThicknessIndex = ThicknessIndex;
 			return &Operator;
@@ -657,7 +650,7 @@ FSubstrateOperator* UMaterialExpressionSubstrateShadingModels::SubstrateGenerate
 
 		if (ShadingModels.HasShadingModel(MSM_Unlit))
 		{
-			FSubstrateOperator& Operator = Compiler->SubstrateCompilationRegisterOperator(SUBSTRATE_OPERATOR_BSDF_LEGACY, Compiler->SubstrateTreeStackGetPathUniqueId(), this, Parent, Compiler->SubstrateTreeStackGetParentPathUniqueId());
+			FSubstrateOperator& Operator = Compiler->SubstrateCompilationRegisterOperator(SUBSTRATE_OPERATOR_BSDF_LEGACY, Compiler->SubstrateTreeStackGetPathUniqueId(), ThisExpressionGuid, Parent, Compiler->SubstrateTreeStackGetParentPathUniqueId());
 			Operator.BSDFType = SUBSTRATE_BSDF_TYPE_UNLIT;
 			Operator.ThicknessIndex = ThicknessIndex;
 			Operator.bBSDFWritesEmissive = true;
@@ -693,7 +686,7 @@ FSubstrateOperator* UMaterialExpressionSubstrateShadingModels::SubstrateGenerate
 		}
 		else if (ShadingModels.HasShadingModel(MSM_ClearCoat))
 		{
-			FSubstrateOperator& Operator = Compiler->SubstrateCompilationRegisterOperator(SUBSTRATE_OPERATOR_BSDF_LEGACY, Compiler->SubstrateTreeStackGetPathUniqueId(), this, Parent, Compiler->SubstrateTreeStackGetParentPathUniqueId());
+			FSubstrateOperator& Operator = Compiler->SubstrateCompilationRegisterOperator(SUBSTRATE_OPERATOR_BSDF_LEGACY, Compiler->SubstrateTreeStackGetPathUniqueId(), ThisExpressionGuid, Parent, Compiler->SubstrateTreeStackGetParentPathUniqueId());
 			Operator.BSDFType = SUBSTRATE_BSDF_TYPE_SLAB;
 			Operator.ThicknessIndex = ThicknessIndex;
 			Operator.BSDFFeatures = ESubstrateBsdfFeature_SecondRoughnessOrSimpleClearCoat | (Anisotropy.IsConnected() ? ESubstrateBsdfFeature_Anisotropy : ESubstrateBsdfFeature_None);
@@ -702,7 +695,7 @@ FSubstrateOperator* UMaterialExpressionSubstrateShadingModels::SubstrateGenerate
 		}
 		else if (ShadingModels.HasShadingModel(MSM_Hair))
 		{
-			FSubstrateOperator& Operator = Compiler->SubstrateCompilationRegisterOperator(SUBSTRATE_OPERATOR_BSDF_LEGACY, Compiler->SubstrateTreeStackGetPathUniqueId(), this, Parent, Compiler->SubstrateTreeStackGetParentPathUniqueId());
+			FSubstrateOperator& Operator = Compiler->SubstrateCompilationRegisterOperator(SUBSTRATE_OPERATOR_BSDF_LEGACY, Compiler->SubstrateTreeStackGetPathUniqueId(), ThisExpressionGuid, Parent, Compiler->SubstrateTreeStackGetParentPathUniqueId());
 			Operator.BSDFType = SUBSTRATE_BSDF_TYPE_HAIR;
 			Operator.BSDFFeatures = ESubstrateBsdfFeature_Hair;
 			Operator.ThicknessIndex = ThicknessIndex;
@@ -711,7 +704,7 @@ FSubstrateOperator* UMaterialExpressionSubstrateShadingModels::SubstrateGenerate
 		}
 		else if (ShadingModels.HasShadingModel(MSM_Eye))
 		{
-			FSubstrateOperator& Operator = Compiler->SubstrateCompilationRegisterOperator(SUBSTRATE_OPERATOR_BSDF_LEGACY, Compiler->SubstrateTreeStackGetPathUniqueId(), this, Parent, Compiler->SubstrateTreeStackGetParentPathUniqueId());
+			FSubstrateOperator& Operator = Compiler->SubstrateCompilationRegisterOperator(SUBSTRATE_OPERATOR_BSDF_LEGACY, Compiler->SubstrateTreeStackGetPathUniqueId(), ThisExpressionGuid, Parent, Compiler->SubstrateTreeStackGetParentPathUniqueId());
 			Operator.BSDFType = SUBSTRATE_BSDF_TYPE_EYE;
 			Operator.BSDFFeatures = ESubstrateBsdfFeature_Eye;
 			ApplyEyeIrisUsedFeatures(&Operator);
@@ -722,7 +715,7 @@ FSubstrateOperator* UMaterialExpressionSubstrateShadingModels::SubstrateGenerate
 		}
 		else if (ShadingModels.HasShadingModel(MSM_SingleLayerWater))
 		{
-			FSubstrateOperator& Operator = Compiler->SubstrateCompilationRegisterOperator(SUBSTRATE_OPERATOR_BSDF_LEGACY, Compiler->SubstrateTreeStackGetPathUniqueId(), this, Parent, Compiler->SubstrateTreeStackGetParentPathUniqueId());
+			FSubstrateOperator& Operator = Compiler->SubstrateCompilationRegisterOperator(SUBSTRATE_OPERATOR_BSDF_LEGACY, Compiler->SubstrateTreeStackGetPathUniqueId(), ThisExpressionGuid, Parent, Compiler->SubstrateTreeStackGetParentPathUniqueId());
 			Operator.BSDFType = SUBSTRATE_BSDF_TYPE_SINGLELAYERWATER;
 			Operator.ThicknessIndex = ThicknessIndex;
 			Operator.bBSDFWritesEmissive = EmissiveColor.IsConnected();
@@ -735,6 +728,17 @@ FSubstrateOperator* UMaterialExpressionSubstrateShadingModels::SubstrateGenerate
 	}
 }
 
+FSubstrateOperator* UMaterialExpressionSubstrateShadingModels::SubstrateGenerateMaterialTopologyTree(class FMaterialCompiler* Compiler, class UMaterialExpression* Parent, int32 OutputIndex)
+{
+	return SubstrateGenerateMaterialTopologyTreeCommon(
+		Compiler, this->MaterialExpressionGuid, Parent, OutputIndex,
+		EmissiveColor,
+		Anisotropy,
+		ClearCoatNormal,
+		CustomTangent,
+		ShadingModel);
+}
+
 bool UMaterialExpressionSubstrateShadingModels::HasSSS() const
 {
 	return SubsurfaceProfile != nullptr;
@@ -1250,7 +1254,7 @@ void UMaterialExpressionSubstrateSlabBSDF::GatherSubstrateMaterialInfo(FSubstrat
 
 FSubstrateOperator* UMaterialExpressionSubstrateSlabBSDF::SubstrateGenerateMaterialTopologyTree(class FMaterialCompiler* Compiler, class UMaterialExpression* Parent, int32 OutputIndex)
 {
-	FSubstrateOperator& SubstrateOperator = Compiler->SubstrateCompilationRegisterOperator(SUBSTRATE_OPERATOR_BSDF, Compiler->SubstrateTreeStackGetPathUniqueId(), this, Parent, Compiler->SubstrateTreeStackGetParentPathUniqueId());
+	FSubstrateOperator& SubstrateOperator = Compiler->SubstrateCompilationRegisterOperator(SUBSTRATE_OPERATOR_BSDF, Compiler->SubstrateTreeStackGetPathUniqueId(), this->MaterialExpressionGuid, Parent, Compiler->SubstrateTreeStackGetParentPathUniqueId());
 	SubstrateOperator.BSDFType = SUBSTRATE_BSDF_TYPE_SLAB;
 
 	if (HasEdgeColor())			{ SubstrateOperator.BSDFFeatures |= ESubstrateBsdfFeature_EdgeColor; }
@@ -1515,7 +1519,7 @@ void UMaterialExpressionSubstrateSimpleClearCoatBSDF::GatherSubstrateMaterialInf
 
 FSubstrateOperator* UMaterialExpressionSubstrateSimpleClearCoatBSDF::SubstrateGenerateMaterialTopologyTree(class FMaterialCompiler* Compiler, class UMaterialExpression* Parent, int32 OutputIndex)
 {
-	FSubstrateOperator& SubstrateOperator = Compiler->SubstrateCompilationRegisterOperator(SUBSTRATE_OPERATOR_BSDF, Compiler->SubstrateTreeStackGetPathUniqueId(), this, Parent, Compiler->SubstrateTreeStackGetParentPathUniqueId());
+	FSubstrateOperator& SubstrateOperator = Compiler->SubstrateCompilationRegisterOperator(SUBSTRATE_OPERATOR_BSDF, Compiler->SubstrateTreeStackGetPathUniqueId(), this->MaterialExpressionGuid, Parent, Compiler->SubstrateTreeStackGetParentPathUniqueId());
 	SubstrateOperator.BSDFType = SUBSTRATE_BSDF_TYPE_SLAB;
 	SubstrateOperator.BSDFFeatures = ESubstrateBsdfFeature_SecondRoughnessOrSimpleClearCoat; // This node explicitly requires simple clear coat
 	SubstrateOperator.ThicknessIndex = Compiler->SubstrateThicknessStackGetThicknessIndex();
@@ -1541,18 +1545,29 @@ UMaterialExpressionSubstrateVolumetricFogCloudBSDF::UMaterialExpressionSubstrate
 }
 
 #if WITH_EDITOR
-int32 UMaterialExpressionSubstrateVolumetricFogCloudBSDF::Compile(class FMaterialCompiler* Compiler, int32 OutputIndex)
+
+int32 UMaterialExpressionSubstrateVolumetricFogCloudBSDF::CompileCommon(class FMaterialCompiler* Compiler,
+	FExpressionInput& Albedo, FExpressionInput& Extinction, FExpressionInput& EmissiveColor, FExpressionInput& AmbientOcclusion,
+	bool bEmissiveOnly,	
+	const UMaterialEditorOnlyData* EditorOnlyData)
 {
 	int32 OutputCodeChunk = Compiler->SubstrateVolumetricFogCloudBSDF(
-		CompileWithDefaultFloat3(Compiler, Albedo, 0.0f, 0.0f, 0.0f),
-		CompileWithDefaultFloat3(Compiler, Extinction, 0.0f, 0.0f, 0.0f),
-		CompileWithDefaultFloat3(Compiler, EmissiveColor, 0.0f, 0.0f, 0.0f),
-		CompileWithDefaultFloat1(Compiler, AmbientOcclusion, 1.0f),
+		CompileWithDefaultFloat3(Compiler, Albedo,				0.0f, 0.0f, 0.0f,	EditorOnlyData ? &EditorOnlyData->BaseColor : nullptr),
+		CompileWithDefaultFloat3(Compiler, Extinction,			0.0f, 0.0f, 0.0f,	EditorOnlyData ? &EditorOnlyData->SubsurfaceColor : nullptr),
+		CompileWithDefaultFloat3(Compiler, EmissiveColor,		0.0f, 0.0f, 0.0f,	EditorOnlyData ? &EditorOnlyData->EmissiveColor : nullptr),
+		CompileWithDefaultFloat1(Compiler, AmbientOcclusion,	1.0f,				EditorOnlyData ? &EditorOnlyData->AmbientOcclusion : nullptr),
 		bEmissiveOnly);
 
 	return OutputCodeChunk;
 }
 
+int32 UMaterialExpressionSubstrateVolumetricFogCloudBSDF::Compile(class FMaterialCompiler* Compiler, int32 OutputIndex)
+{
+	return CompileCommon(Compiler,
+		Albedo, Extinction, EmissiveColor, AmbientOcclusion,
+		bEmissiveOnly);
+}
+
 void UMaterialExpressionSubstrateVolumetricFogCloudBSDF::GetCaption(TArray<FString>& OutCaptions) const
 {
 	FString Caption = TEXT("Substrate Volumetric-Fog-Cloud BSDF");
@@ -1601,15 +1616,26 @@ void UMaterialExpressionSubstrateVolumetricFogCloudBSDF::GatherSubstrateMaterial
 	SubstrateMaterialInfo.AddGuid(MaterialExpressionGuid);
 }
 
-FSubstrateOperator* UMaterialExpressionSubstrateVolumetricFogCloudBSDF::SubstrateGenerateMaterialTopologyTree(class FMaterialCompiler* Compiler, class UMaterialExpression* Parent, int32 OutputIndex)
+FSubstrateOperator* UMaterialExpressionSubstrateVolumetricFogCloudBSDF::SubstrateGenerateMaterialTopologyTreeCommon(
+	class FMaterialCompiler* Compiler, struct FGuid ThisExpressionGuid, class UMaterialExpression* Parent, int32 OutputIndex,
+	const FExpressionInput& EmissiveColor,
+	const FExpressionInput& AmbientOcclusion)
 {
-	FSubstrateOperator& SubstrateOperator = Compiler->SubstrateCompilationRegisterOperator(SUBSTRATE_OPERATOR_BSDF, Compiler->SubstrateTreeStackGetPathUniqueId(), this, Parent, Compiler->SubstrateTreeStackGetParentPathUniqueId());
+	FSubstrateOperator& SubstrateOperator = Compiler->SubstrateCompilationRegisterOperator(SUBSTRATE_OPERATOR_BSDF, Compiler->SubstrateTreeStackGetPathUniqueId(), ThisExpressionGuid, Parent, Compiler->SubstrateTreeStackGetParentPathUniqueId());
 	SubstrateOperator.BSDFType = SUBSTRATE_BSDF_TYPE_VOLUMETRICFOGCLOUD;
 	SubstrateOperator.ThicknessIndex = Compiler->SubstrateThicknessStackGetThicknessIndex();
 	SubstrateOperator.bBSDFWritesEmissive = EmissiveColor.IsConnected();
 	SubstrateOperator.bBSDFWritesAmbientOcclusion = AmbientOcclusion.IsConnected();
 	return &SubstrateOperator;
 }
+
+FSubstrateOperator* UMaterialExpressionSubstrateVolumetricFogCloudBSDF::SubstrateGenerateMaterialTopologyTree(class FMaterialCompiler* Compiler, class UMaterialExpression* Parent, int32 OutputIndex)
+{
+	return SubstrateGenerateMaterialTopologyTreeCommon(
+		Compiler, MaterialExpressionGuid, Parent, OutputIndex,
+		EmissiveColor,
+		AmbientOcclusion);
+}
 #endif // WITH_EDITOR
 
 
@@ -1629,17 +1655,24 @@ UMaterialExpressionSubstrateLightFunction::UMaterialExpressionSubstrateLightFunc
 }
 
 #if WITH_EDITOR
-int32 UMaterialExpressionSubstrateLightFunction::Compile(class FMaterialCompiler* Compiler, int32 OutputIndex)
+int32 UMaterialExpressionSubstrateLightFunction::CompileCommon(class FMaterialCompiler* Compiler,
+	FExpressionInput& Color, 
+	const UMaterialEditorOnlyData* EditorOnlyData)
 {
 	FSubstrateOperator& SubstrateOperator = Compiler->SubstrateCompilationGetOperator(Compiler->SubstrateTreeStackGetPathUniqueId());
 	int32 OutputCodeChunk = Compiler->SubstrateUnlitBSDF(
-		CompileWithDefaultFloat3(Compiler, Color, 0.0f, 0.0f, 0.0f),
+		CompileWithDefaultFloat3(Compiler, Color, 0.0f, 0.0f, 0.0f, EditorOnlyData ? &EditorOnlyData->EmissiveColor : nullptr),
 		Compiler->Constant(1.0f),				// Opacity / Transmittance is ignored by light functions.
 		Compiler->Constant3(0.0f, 0.0f, 1.0f),	// place holder normal
 		&SubstrateOperator);
 	return OutputCodeChunk;
 }
 
+int32 UMaterialExpressionSubstrateLightFunction::Compile(class FMaterialCompiler* Compiler, int32 OutputIndex)
+{
+	return CompileCommon(Compiler, Color);
+}
+
 void UMaterialExpressionSubstrateLightFunction::GetCaption(TArray<FString>& OutCaptions) const
 {
 	OutCaptions.Add(TEXT("Substrate Light Function"));
@@ -1674,14 +1707,20 @@ void UMaterialExpressionSubstrateLightFunction::GatherSubstrateMaterialInfo(FSub
 	SubstrateMaterialInfo.AddGuid(MaterialExpressionGuid);
 }
 
-FSubstrateOperator* UMaterialExpressionSubstrateLightFunction::SubstrateGenerateMaterialTopologyTree(class FMaterialCompiler* Compiler, class UMaterialExpression* Parent, int32 OutputIndex)
+FSubstrateOperator* UMaterialExpressionSubstrateLightFunction::SubstrateGenerateMaterialTopologyTreeCommon(
+	class FMaterialCompiler* Compiler, struct FGuid ThisExpressionGuid, class UMaterialExpression* Parent, int32 OutputIndex)
 {
-	FSubstrateOperator& SubstrateOperator = Compiler->SubstrateCompilationRegisterOperator(SUBSTRATE_OPERATOR_BSDF, Compiler->SubstrateTreeStackGetPathUniqueId(), this, Parent, Compiler->SubstrateTreeStackGetParentPathUniqueId());
+	FSubstrateOperator& SubstrateOperator = Compiler->SubstrateCompilationRegisterOperator(SUBSTRATE_OPERATOR_BSDF, Compiler->SubstrateTreeStackGetPathUniqueId(), ThisExpressionGuid, Parent, Compiler->SubstrateTreeStackGetParentPathUniqueId());
 	SubstrateOperator.BSDFType = SUBSTRATE_BSDF_TYPE_UNLIT;
 	SubstrateOperator.SubUsage = SUBSTRATE_OPERATOR_SUBUSAGE_LIGHTFUNCTION;
 	SubstrateOperator.ThicknessIndex = Compiler->SubstrateThicknessStackGetThicknessIndex();
 	return &SubstrateOperator;
 }
+
+FSubstrateOperator* UMaterialExpressionSubstrateLightFunction::SubstrateGenerateMaterialTopologyTree(class FMaterialCompiler* Compiler, class UMaterialExpression* Parent, int32 OutputIndex)
+{
+	return SubstrateGenerateMaterialTopologyTreeCommon(Compiler, MaterialExpressionGuid, Parent, OutputIndex);
+}
 #endif // WITH_EDITOR
 
 
@@ -1701,20 +1740,27 @@ UMaterialExpressionSubstratePostProcess::UMaterialExpressionSubstratePostProcess
 }
 
 #if WITH_EDITOR
-int32 UMaterialExpressionSubstratePostProcess::Compile(class FMaterialCompiler* Compiler, int32 OutputIndex)
+int32 UMaterialExpressionSubstratePostProcess::CompileCommon(class FMaterialCompiler* Compiler,
+	FExpressionInput& Color, FExpressionInput& Opacity,
+	const UMaterialEditorOnlyData* EditorOnlyData)
 {
-	int OpacityCodeChunk = CompileWithDefaultFloat1(Compiler, Opacity, 1.0f);
+	int OpacityCodeChunk = CompileWithDefaultFloat1(Compiler, Opacity, 1.0f, EditorOnlyData ? &EditorOnlyData->Opacity : nullptr);
 	int TransmittanceCodeChunk = Compiler->Saturate(Compiler->Sub(Compiler->Constant(1.0f), OpacityCodeChunk));
 
 	FSubstrateOperator& SubstrateOperator = Compiler->SubstrateCompilationGetOperator(Compiler->SubstrateTreeStackGetPathUniqueId());
 	int32 OutputCodeChunk = Compiler->SubstrateUnlitBSDF(
-		CompileWithDefaultFloat3(Compiler, Color, 0.0f, 0.0f, 0.0f),
+		CompileWithDefaultFloat3(Compiler, Color, 0.0f, 0.0f, 0.0f, EditorOnlyData ? &EditorOnlyData->EmissiveColor : nullptr),
 		TransmittanceCodeChunk,
 		Compiler->Constant3(0.0f, 0.0f, 1.0f),	// place holder normal
 		&SubstrateOperator);
 	return OutputCodeChunk;
 }
 
+int32 UMaterialExpressionSubstratePostProcess::Compile(class FMaterialCompiler* Compiler, int32 OutputIndex)
+{
+	return CompileCommon(Compiler, Color, Opacity);
+}
+
 void UMaterialExpressionSubstratePostProcess::GetCaption(TArray<FString>& OutCaptions) const
 {
 	OutCaptions.Add(TEXT("Substrate Post Process"));
@@ -1752,14 +1798,20 @@ void UMaterialExpressionSubstratePostProcess::GatherSubstrateMaterialInfo(FSubst
 	SubstrateMaterialInfo.AddGuid(MaterialExpressionGuid);
 }
 
-FSubstrateOperator* UMaterialExpressionSubstratePostProcess::SubstrateGenerateMaterialTopologyTree(class FMaterialCompiler* Compiler, class UMaterialExpression* Parent, int32 OutputIndex)
+FSubstrateOperator* UMaterialExpressionSubstratePostProcess::SubstrateGenerateMaterialTopologyTreeCommon(
+	class FMaterialCompiler* Compiler, struct FGuid ThisExpressionGuid, class UMaterialExpression* Parent, int32 OutputIndex)
 {
-	FSubstrateOperator& SubstrateOperator = Compiler->SubstrateCompilationRegisterOperator(SUBSTRATE_OPERATOR_BSDF, Compiler->SubstrateTreeStackGetPathUniqueId(), this, Parent, Compiler->SubstrateTreeStackGetParentPathUniqueId());
+	FSubstrateOperator& SubstrateOperator = Compiler->SubstrateCompilationRegisterOperator(SUBSTRATE_OPERATOR_BSDF, Compiler->SubstrateTreeStackGetPathUniqueId(), ThisExpressionGuid, Parent, Compiler->SubstrateTreeStackGetParentPathUniqueId());
 	SubstrateOperator.BSDFType = SUBSTRATE_BSDF_TYPE_UNLIT;
 	SubstrateOperator.SubUsage = SUBSTRATE_OPERATOR_SUBUSAGE_POSTPROCESS;
 	SubstrateOperator.ThicknessIndex = Compiler->SubstrateThicknessStackGetThicknessIndex();
 	return &SubstrateOperator;
 }
+
+FSubstrateOperator* UMaterialExpressionSubstratePostProcess::SubstrateGenerateMaterialTopologyTree(class FMaterialCompiler* Compiler, class UMaterialExpression* Parent, int32 OutputIndex)
+{
+	return SubstrateGenerateMaterialTopologyTreeCommon(Compiler, MaterialExpressionGuid, Parent, OutputIndex);
+}
 #endif // WITH_EDITOR
 
 
@@ -1779,18 +1831,25 @@ UMaterialExpressionSubstrateUI::UMaterialExpressionSubstrateUI(const FObjectInit
 }
 
 #if WITH_EDITOR
-int32 UMaterialExpressionSubstrateUI::Compile(class FMaterialCompiler* Compiler, int32 OutputIndex)
+int32 UMaterialExpressionSubstrateUI::CompileCommon(class FMaterialCompiler* Compiler,
+	FExpressionInput& Color, FExpressionInput& Opacity,
+	const UMaterialEditorOnlyData* EditorOnlyData)
 {
-	int OpacityCodeChunk = CompileWithDefaultFloat1(Compiler, Opacity, 1.0f);
+	int OpacityCodeChunk = CompileWithDefaultFloat1(Compiler, Opacity, 1.0f, EditorOnlyData ? &EditorOnlyData->Opacity : nullptr);
 
 	FSubstrateOperator& SubstrateOperator = Compiler->SubstrateCompilationGetOperator(Compiler->SubstrateTreeStackGetPathUniqueId());
 	int32 OutputCodeChunk = Compiler->SubstrateUIBSDF(
-		CompileWithDefaultFloat3(Compiler, Color, 0.0f, 0.0f, 0.0f),
+		CompileWithDefaultFloat3(Compiler, Color, 0.0f, 0.0f, 0.0f, EditorOnlyData ? &EditorOnlyData->EmissiveColor : nullptr),
 		OpacityCodeChunk,
 		&SubstrateOperator);
 	return OutputCodeChunk;
 }
 
+int32 UMaterialExpressionSubstrateUI::Compile(class FMaterialCompiler* Compiler, int32 OutputIndex)
+{
+	return CompileCommon(Compiler, Color, Opacity);
+}
+
 void UMaterialExpressionSubstrateUI::GetCaption(TArray<FString>& OutCaptions) const
 {
 	OutCaptions.Add(TEXT("Substrate UI"));
@@ -1828,14 +1887,20 @@ void UMaterialExpressionSubstrateUI::GatherSubstrateMaterialInfo(FSubstrateMater
 	SubstrateMaterialInfo.AddGuid(MaterialExpressionGuid);
 }
 
-FSubstrateOperator* UMaterialExpressionSubstrateUI::SubstrateGenerateMaterialTopologyTree(class FMaterialCompiler* Compiler, class UMaterialExpression* Parent, int32 OutputIndex)
+FSubstrateOperator* UMaterialExpressionSubstrateUI::SubstrateGenerateMaterialTopologyTreeCommon(
+	class FMaterialCompiler* Compiler, struct FGuid ThisExpressionGuid, class UMaterialExpression* Parent, int32 OutputIndex)
 {
-	FSubstrateOperator& SubstrateOperator = Compiler->SubstrateCompilationRegisterOperator(SUBSTRATE_OPERATOR_BSDF, Compiler->SubstrateTreeStackGetPathUniqueId(), this, Parent, Compiler->SubstrateTreeStackGetParentPathUniqueId());
+	FSubstrateOperator& SubstrateOperator = Compiler->SubstrateCompilationRegisterOperator(SUBSTRATE_OPERATOR_BSDF, Compiler->SubstrateTreeStackGetPathUniqueId(), ThisExpressionGuid, Parent, Compiler->SubstrateTreeStackGetParentPathUniqueId());
 	SubstrateOperator.BSDFType = SUBSTRATE_BSDF_TYPE_UNLIT;
 	SubstrateOperator.SubUsage = SUBSTRATE_OPERATOR_SUBUSAGE_UI;
 	SubstrateOperator.ThicknessIndex = Compiler->SubstrateThicknessStackGetThicknessIndex();
 	return &SubstrateOperator;
 }
+
+FSubstrateOperator* UMaterialExpressionSubstrateUI::SubstrateGenerateMaterialTopologyTree(class FMaterialCompiler* Compiler, class UMaterialExpression* Parent, int32 OutputIndex)
+{
+	return SubstrateGenerateMaterialTopologyTreeCommon(Compiler, MaterialExpressionGuid, Parent, OutputIndex);
+}
 #endif // WITH_EDITOR
 
 
@@ -1941,7 +2006,7 @@ void UMaterialExpressionSubstrateConvertToDecal::GatherSubstrateMaterialInfo(FSu
 FSubstrateOperator* UMaterialExpressionSubstrateConvertToDecal::SubstrateGenerateMaterialTopologyTree(class FMaterialCompiler* Compiler, class UMaterialExpression* Parent, int32 OutputIndex)
 {
 	const bool bUseParameterBlending = true;
-	FSubstrateOperator& SubstrateOperator = Compiler->SubstrateCompilationRegisterOperator(SUBSTRATE_OPERATOR_WEIGHT, Compiler->SubstrateTreeStackGetPathUniqueId(), this, Parent, Compiler->SubstrateTreeStackGetParentPathUniqueId(), bUseParameterBlending);
+	FSubstrateOperator& SubstrateOperator = Compiler->SubstrateCompilationRegisterOperator(SUBSTRATE_OPERATOR_WEIGHT, Compiler->SubstrateTreeStackGetPathUniqueId(), this->MaterialExpressionGuid, Parent, Compiler->SubstrateTreeStackGetParentPathUniqueId(), bUseParameterBlending);
 	SubstrateOperator.SubUsage = SUBSTRATE_OPERATOR_SUBUSAGE_DECAL;
 	if (Compiler->GetSubstrateTreeOutOfStackDepthOccurred())
 	{
@@ -2035,7 +2100,7 @@ void UMaterialExpressionSubstrateUnlitBSDF::GatherSubstrateMaterialInfo(FSubstra
 
 FSubstrateOperator* UMaterialExpressionSubstrateUnlitBSDF::SubstrateGenerateMaterialTopologyTree(class FMaterialCompiler* Compiler, class UMaterialExpression* Parent, int32 OutputIndex)
 {
-	FSubstrateOperator& SubstrateOperator = Compiler->SubstrateCompilationRegisterOperator(SUBSTRATE_OPERATOR_BSDF, Compiler->SubstrateTreeStackGetPathUniqueId(), this, Parent, Compiler->SubstrateTreeStackGetParentPathUniqueId());
+	FSubstrateOperator& SubstrateOperator = Compiler->SubstrateCompilationRegisterOperator(SUBSTRATE_OPERATOR_BSDF, Compiler->SubstrateTreeStackGetPathUniqueId(), this->MaterialExpressionGuid, Parent, Compiler->SubstrateTreeStackGetParentPathUniqueId());
 	SubstrateOperator.BSDFType = SUBSTRATE_BSDF_TYPE_UNLIT;
 	SubstrateOperator.ThicknessIndex = Compiler->SubstrateThicknessStackGetThicknessIndex();
 	SubstrateOperator.bBSDFWritesEmissive = EmissiveColor.IsConnected();
@@ -2153,7 +2218,7 @@ void UMaterialExpressionSubstrateHairBSDF::GatherSubstrateMaterialInfo(FSubstrat
 
 FSubstrateOperator* UMaterialExpressionSubstrateHairBSDF::SubstrateGenerateMaterialTopologyTree(class FMaterialCompiler* Compiler, class UMaterialExpression* Parent, int32 OutputIndex)
 {
-	FSubstrateOperator& SubstrateOperator = Compiler->SubstrateCompilationRegisterOperator(SUBSTRATE_OPERATOR_BSDF, Compiler->SubstrateTreeStackGetPathUniqueId(), this, Parent, Compiler->SubstrateTreeStackGetParentPathUniqueId());
+	FSubstrateOperator& SubstrateOperator = Compiler->SubstrateCompilationRegisterOperator(SUBSTRATE_OPERATOR_BSDF, Compiler->SubstrateTreeStackGetPathUniqueId(), this->MaterialExpressionGuid, Parent, Compiler->SubstrateTreeStackGetParentPathUniqueId());
 	SubstrateOperator.BSDFType = SUBSTRATE_BSDF_TYPE_HAIR;
 	SubstrateOperator.BSDFFeatures = ESubstrateBsdfFeature_Hair;
 	SubstrateOperator.ThicknessIndex = Compiler->SubstrateThicknessStackGetThicknessIndex();
@@ -2274,7 +2339,7 @@ void UMaterialExpressionSubstrateEyeBSDF::GatherSubstrateMaterialInfo(FSubstrate
 
 FSubstrateOperator* UMaterialExpressionSubstrateEyeBSDF::SubstrateGenerateMaterialTopologyTree(class FMaterialCompiler* Compiler, class UMaterialExpression* Parent, int32 OutputIndex)
 {
-	FSubstrateOperator& SubstrateOperator = Compiler->SubstrateCompilationRegisterOperator(SUBSTRATE_OPERATOR_BSDF, Compiler->SubstrateTreeStackGetPathUniqueId(), this, Parent, Compiler->SubstrateTreeStackGetParentPathUniqueId());
+	FSubstrateOperator& SubstrateOperator = Compiler->SubstrateCompilationRegisterOperator(SUBSTRATE_OPERATOR_BSDF, Compiler->SubstrateTreeStackGetPathUniqueId(), this->MaterialExpressionGuid, Parent, Compiler->SubstrateTreeStackGetParentPathUniqueId());
 	SubstrateOperator.BSDFType = SUBSTRATE_BSDF_TYPE_EYE;
 	SubstrateOperator.BSDFFeatures = ESubstrateBsdfFeature_Eye;
 	if (IrisNormal.IsConnected())
@@ -2417,7 +2482,7 @@ void UMaterialExpressionSubstrateSingleLayerWaterBSDF::GatherSubstrateMaterialIn
 
 FSubstrateOperator* UMaterialExpressionSubstrateSingleLayerWaterBSDF::SubstrateGenerateMaterialTopologyTree(class FMaterialCompiler* Compiler, class UMaterialExpression* Parent, int32 OutputIndex)
 {
-	FSubstrateOperator& SubstrateOperator = Compiler->SubstrateCompilationRegisterOperator(SUBSTRATE_OPERATOR_BSDF, Compiler->SubstrateTreeStackGetPathUniqueId(), this, Parent, Compiler->SubstrateTreeStackGetParentPathUniqueId());
+	FSubstrateOperator& SubstrateOperator = Compiler->SubstrateCompilationRegisterOperator(SUBSTRATE_OPERATOR_BSDF, Compiler->SubstrateTreeStackGetPathUniqueId(), this->MaterialExpressionGuid, Parent, Compiler->SubstrateTreeStackGetParentPathUniqueId());
 	SubstrateOperator.BSDFType = SUBSTRATE_BSDF_TYPE_SINGLELAYERWATER;
 	SubstrateOperator.ThicknessIndex = Compiler->SubstrateThicknessStackGetThicknessIndex();
 	SubstrateOperator.bBSDFWritesEmissive = EmissiveColor.IsConnected();
@@ -2574,7 +2639,7 @@ void UMaterialExpressionSubstrateHorizontalMixing::GatherSubstrateMaterialInfo(F
 
 FSubstrateOperator* UMaterialExpressionSubstrateHorizontalMixing::SubstrateGenerateMaterialTopologyTree(class FMaterialCompiler* Compiler, class UMaterialExpression* Parent, int32 OutputIndex)
 {
-	FSubstrateOperator& SubstrateOperator = Compiler->SubstrateCompilationRegisterOperator(SUBSTRATE_OPERATOR_HORIZONTAL, Compiler->SubstrateTreeStackGetPathUniqueId(), this, Parent, Compiler->SubstrateTreeStackGetParentPathUniqueId(), bUseParameterBlending);
+	FSubstrateOperator& SubstrateOperator = Compiler->SubstrateCompilationRegisterOperator(SUBSTRATE_OPERATOR_HORIZONTAL, Compiler->SubstrateTreeStackGetPathUniqueId(), this->MaterialExpressionGuid, Parent, Compiler->SubstrateTreeStackGetParentPathUniqueId(), bUseParameterBlending);
 	if (Compiler->GetSubstrateTreeOutOfStackDepthOccurred())
 	{
 		return &SubstrateOperator; // Out ot stack space, return now to fail the compilation
@@ -2774,7 +2839,7 @@ void UMaterialExpressionSubstrateVerticalLayering::GatherSubstrateMaterialInfo(F
 
 FSubstrateOperator* UMaterialExpressionSubstrateVerticalLayering::SubstrateGenerateMaterialTopologyTree(class FMaterialCompiler* Compiler, class UMaterialExpression* Parent, int32 OutputIndex)
 {
-	FSubstrateOperator& SubstrateOperator = Compiler->SubstrateCompilationRegisterOperator(SUBSTRATE_OPERATOR_VERTICAL, Compiler->SubstrateTreeStackGetPathUniqueId(), this, Parent, Compiler->SubstrateTreeStackGetParentPathUniqueId(), bUseParameterBlending);
+	FSubstrateOperator& SubstrateOperator = Compiler->SubstrateCompilationRegisterOperator(SUBSTRATE_OPERATOR_VERTICAL, Compiler->SubstrateTreeStackGetPathUniqueId(), this->MaterialExpressionGuid, Parent, Compiler->SubstrateTreeStackGetParentPathUniqueId(), bUseParameterBlending);
 	if (Compiler->GetSubstrateTreeOutOfStackDepthOccurred())
 	{
 		return &SubstrateOperator; // Out ot stack space, return now to fail the compilation
@@ -2953,7 +3018,7 @@ void UMaterialExpressionSubstrateAdd::GatherSubstrateMaterialInfo(FSubstrateMate
 
 FSubstrateOperator* UMaterialExpressionSubstrateAdd::SubstrateGenerateMaterialTopologyTree(class FMaterialCompiler* Compiler, class UMaterialExpression* Parent, int32 OutputIndex)
 {
-	FSubstrateOperator& SubstrateOperator = Compiler->SubstrateCompilationRegisterOperator(SUBSTRATE_OPERATOR_ADD, Compiler->SubstrateTreeStackGetPathUniqueId(), this, Parent, Compiler->SubstrateTreeStackGetParentPathUniqueId(), bUseParameterBlending);
+	FSubstrateOperator& SubstrateOperator = Compiler->SubstrateCompilationRegisterOperator(SUBSTRATE_OPERATOR_ADD, Compiler->SubstrateTreeStackGetPathUniqueId(), this->MaterialExpressionGuid, Parent, Compiler->SubstrateTreeStackGetParentPathUniqueId(), bUseParameterBlending);
 	if (Compiler->GetSubstrateTreeOutOfStackDepthOccurred())
 	{
 		return &SubstrateOperator; // Out ot stack space, return now to fail the compilation
@@ -3091,7 +3156,7 @@ void UMaterialExpressionSubstrateWeight::GatherSubstrateMaterialInfo(FSubstrateM
 
 FSubstrateOperator* UMaterialExpressionSubstrateWeight::SubstrateGenerateMaterialTopologyTree(class FMaterialCompiler* Compiler, class UMaterialExpression* Parent, int32 OutputIndex)
 {
-	FSubstrateOperator& SubstrateOperator = Compiler->SubstrateCompilationRegisterOperator(SUBSTRATE_OPERATOR_WEIGHT, Compiler->SubstrateTreeStackGetPathUniqueId(), this, Parent, Compiler->SubstrateTreeStackGetParentPathUniqueId());
+	FSubstrateOperator& SubstrateOperator = Compiler->SubstrateCompilationRegisterOperator(SUBSTRATE_OPERATOR_WEIGHT, Compiler->SubstrateTreeStackGetPathUniqueId(), this->MaterialExpressionGuid, Parent, Compiler->SubstrateTreeStackGetParentPathUniqueId());
 	if (Compiler->GetSubstrateTreeOutOfStackDepthOccurred())
 	{
 		return &SubstrateOperator; // Out ot stack space, return now to fail the compilation
@@ -3258,7 +3323,7 @@ void UMaterialExpressionSubstrateSelect::GatherSubstrateMaterialInfo(FSubstrateM
 
 FSubstrateOperator* UMaterialExpressionSubstrateSelect::SubstrateGenerateMaterialTopologyTree(class FMaterialCompiler* Compiler, class UMaterialExpression* Parent, int32 OutputIndex)
 {
-	FSubstrateOperator& SubstrateOperator = Compiler->SubstrateCompilationRegisterOperator(SUBSTRATE_OPERATOR_SELECT, Compiler->SubstrateTreeStackGetPathUniqueId(), this, Parent, Compiler->SubstrateTreeStackGetParentPathUniqueId(), bUseParameterBlending);
+	FSubstrateOperator& SubstrateOperator = Compiler->SubstrateCompilationRegisterOperator(SUBSTRATE_OPERATOR_SELECT, Compiler->SubstrateTreeStackGetPathUniqueId(), this->MaterialExpressionGuid, Parent, Compiler->SubstrateTreeStackGetParentPathUniqueId(), bUseParameterBlending);
 	if (Compiler->GetSubstrateTreeOutOfStackDepthOccurred())
 	{
 		return &SubstrateOperator; // Out ot stack space, return now to fail the compilation
@@ -3771,7 +3836,10 @@ FExpressionInput* UMaterialExpressionSubstrateConvertMaterialAttributes::GetInpu
 	}
 }
 
-int32 UMaterialExpressionSubstrateConvertMaterialAttributes::Compile(class FMaterialCompiler* Compiler, int32 OutputIndex)
+int32 UMaterialExpressionSubstrateConvertMaterialAttributes::CompileCommon(class FMaterialCompiler* Compiler, int32 OutputIndex,
+	const uint64 CachedConnectedMaterialAttributesInputs, FMaterialAttributesInput& MaterialAttributes, enum EMaterialShadingModel ShadingModelOverride,
+	FExpressionInput& WaterScatteringCoefficients, FExpressionInput& WaterAbsorptionCoefficients, FExpressionInput& WaterPhaseG, FExpressionInput& ColorScaleBehindWater,
+	const bool bHasSSS, USubsurfaceProfile* SSSProfile)
 {
 	static const FGuid ClearCoatBottomNormalGuid = FMaterialAttributeDefinitionMap::GetCustomAttributeID(TEXT("ClearCoatBottomNormal"));
 	static const FGuid CustomEyeTangentGuid = FMaterialAttributeDefinitionMap::GetCustomAttributeID(TEXT("CustomEyeTangent"));
@@ -3786,8 +3854,7 @@ int32 UMaterialExpressionSubstrateConvertMaterialAttributes::Compile(class FMate
 	// We also cannot ignore the tangent when using the default Tangent because GetTangentBasis
 	// used in SubstrateGetBSDFSharedBasis cannot be relied on for smooth tangent used for lighting on any mesh.
 
-	const uint64 Cached = GetConnectedMaterialAttributesInputs(Material);
-	const bool bHasAnisotropy = FMaterialAttributeDefinitionMap::IsAttributeInBitmask(Cached, MP_Anisotropy);
+	const bool bHasAnisotropy = FMaterialAttributeDefinitionMap::IsAttributeInBitmask(CachedConnectedMaterialAttributesInputs, MP_Anisotropy);
 
 	// Regular normal basis
 	int32 NormalCodeChunk = MaterialAttributes.CompileWithDefault(Compiler, FMaterialAttributeDefinitionMap::GetID(MP_Normal));
@@ -3800,7 +3867,7 @@ int32 UMaterialExpressionSubstrateConvertMaterialAttributes::Compile(class FMate
 	// to the user, but the compilation will succeed.
 	if (NormalCodeChunk == INDEX_NONE) { NormalCodeChunk = Compiler->VertexNormal(); }
 
-	int32 TangentCodeChunk = INDEX_NONE; 
+	int32 TangentCodeChunk = INDEX_NONE;
 	if (bHasAnisotropy)
 	{
 		TangentCodeChunk = MaterialAttributes.CompileWithDefault(Compiler, FMaterialAttributeDefinitionMap::GetID(MP_Tangent));
@@ -3810,7 +3877,7 @@ int32 UMaterialExpressionSubstrateConvertMaterialAttributes::Compile(class FMate
 	const FSubstrateRegisteredSharedLocalBasis NewRegisteredSharedLocalBasis = SubstrateCompilationInfoCreateSharedLocalBasis(Compiler, NormalCodeChunk, TangentCodeChunk);
 	const FString BasisIndexMacro = Compiler->GetSubstrateSharedLocalBasisIndexMacro(NewRegisteredSharedLocalBasis);
 
-	const bool bHasCoatNormal = FMaterialAttributeDefinitionMap::IsAttributeInBitmask(Cached, ClearCoatBottomNormalGuid);
+	const bool bHasCoatNormal = FMaterialAttributeDefinitionMap::IsAttributeInBitmask(CachedConnectedMaterialAttributesInputs, ClearCoatBottomNormalGuid);
 	// Clear coat normal basis
 	int32 ClearCoat_NormalCodeChunk = INDEX_NONE;
 	int32 ClearCoat_TangentCodeChunk = INDEX_NONE;
@@ -3833,7 +3900,7 @@ int32 UMaterialExpressionSubstrateConvertMaterialAttributes::Compile(class FMate
 
 	// Custom tangent. No need to register it as a local basis, as it is only used for eye shading internal conversion
 	int32 CustomTangent_TangentCodeChunk = INDEX_NONE;
-	const bool bHasCustomTangent = FMaterialAttributeDefinitionMap::IsAttributeInBitmask(Cached, CustomEyeTangentGuid);
+	const bool bHasCustomTangent = FMaterialAttributeDefinitionMap::IsAttributeInBitmask(CachedConnectedMaterialAttributesInputs, CustomEyeTangentGuid);
 	if (bHasCustomTangent)
 	{
 		// Legacy code doesn't do tangent <-> world basis conversion on tangent output, when provided.
@@ -3852,10 +3919,9 @@ int32 UMaterialExpressionSubstrateConvertMaterialAttributes::Compile(class FMate
 	// SSS profile
 	// Need to handle this by looking at the material instead of the node?
 	int32 SSSProfileCodeChunk = INDEX_NONE;
-	const bool bHasSSS = HasSSS();
-	if (bHasSSS && SubsurfaceProfile)
+	if (bHasSSS && SSSProfile)
 	{
-		SSSProfileCodeChunk = CreateSubsurfaceProfileParameter(Compiler, SubsurfaceProfile);
+		SSSProfileCodeChunk = CreateSubsurfaceProfileParameter(Compiler, SSSProfile);
 	}
 	else
 	{
@@ -3880,7 +3946,7 @@ int32 UMaterialExpressionSubstrateConvertMaterialAttributes::Compile(class FMate
 	}
 
 	// Transmittance Color
-	const bool bHasTransmittanceColor = FMaterialAttributeDefinitionMap::IsAttributeInBitmask(Cached, TransmittanceColorGuid);
+	const bool bHasTransmittanceColor = FMaterialAttributeDefinitionMap::IsAttributeInBitmask(CachedConnectedMaterialAttributesInputs, TransmittanceColorGuid);
 	int32 TransmittanceColorChunk = INDEX_NONE;
 	if (bHasTransmittanceColor)
 	{
@@ -3891,7 +3957,7 @@ int32 UMaterialExpressionSubstrateConvertMaterialAttributes::Compile(class FMate
 		TransmittanceColorChunk = Compiler->Constant3(0.5f, 0.5f, 0.5f);
 	}
 	// Thin Translucent Surface Coverage
-	const bool bHasThinTranslucentSurfaceCoverage = FMaterialAttributeDefinitionMap::IsAttributeInBitmask(Cached, ThinTranslucentSurfaceCoverageGuid);
+	const bool bHasThinTranslucentSurfaceCoverage = FMaterialAttributeDefinitionMap::IsAttributeInBitmask(CachedConnectedMaterialAttributesInputs, ThinTranslucentSurfaceCoverageGuid);
 	int32 ThinTranslucentSurfaceCoverageChunk = INDEX_NONE;
 	if (bHasThinTranslucentSurfaceCoverage)
 	{
@@ -3903,7 +3969,7 @@ int32 UMaterialExpressionSubstrateConvertMaterialAttributes::Compile(class FMate
 	}
 
 	int32 ShadingModelCodeChunk = MaterialAttributes.CompileWithDefault(Compiler, FMaterialAttributeDefinitionMap::GetID(MP_ShadingModel));
-	const bool bHasShadingModelExpression = FMaterialAttributeDefinitionMap::IsAttributeInBitmask(Cached, MP_ShadingModel)
+	const bool bHasShadingModelExpression = FMaterialAttributeDefinitionMap::IsAttributeInBitmask(CachedConnectedMaterialAttributesInputs, MP_ShadingModel)
 		|| ShadingModelOverride == MSM_FromMaterialExpression; // In this case, rely on the default compilation to return DefaultLit.
 	if (!bHasShadingModelExpression)
 	{
@@ -3950,6 +4016,13 @@ int32 UMaterialExpressionSubstrateConvertMaterialAttributes::Compile(class FMate
 	return OutputCodeChunk;
 }
 
+int32 UMaterialExpressionSubstrateConvertMaterialAttributes::Compile(class FMaterialCompiler* Compiler, int32 OutputIndex)
+{
+	return CompileCommon(Compiler, OutputIndex, GetConnectedMaterialAttributesInputs(Material), MaterialAttributes, ShadingModelOverride,
+		WaterScatteringCoefficients, WaterAbsorptionCoefficients, WaterPhaseG, ColorScaleBehindWater,
+		HasSSS(), SubsurfaceProfile);
+}
+
 void UMaterialExpressionSubstrateConvertMaterialAttributes::PostEditChangeProperty(FPropertyChangedEvent& PropertyChangedEvent)
 {
 	if (GraphNode && PropertyChangedEvent.Property != nullptr)
@@ -4065,20 +4138,20 @@ void UMaterialExpressionSubstrateConvertMaterialAttributes::GatherSubstrateMater
 	SubstrateMaterialInfo.AddGuid(MaterialExpressionGuid);
 }
 
-FSubstrateOperator* UMaterialExpressionSubstrateConvertMaterialAttributes::SubstrateGenerateMaterialTopologyTree(class FMaterialCompiler* Compiler, class UMaterialExpression* Parent, int32 OutputIndex)
+FSubstrateOperator* UMaterialExpressionSubstrateConvertMaterialAttributes::SubstrateGenerateMaterialTopologyTreeCommon(
+	class FMaterialCompiler* Compiler, struct FGuid ThisExpressionGuid, class UMaterialExpression* Parent, int32 OutputIndex,
+	const uint64 CachedConnectedMaterialAttributesInputs, const bool bShadingModelFromMaterialExpression, const bool bIsEmissiveConnected)
 {
-	const uint64 Cached = GetConnectedMaterialAttributesInputs(Material);
-
 	// Note Thickness has no meaning/usage in the context of SubstrateLegacyConversionNode
 	int32 ThicknessIndex = Compiler->SubstrateThicknessStackGetThicknessIndex();
 
-	const bool bHasAnisotropy = FMaterialAttributeDefinitionMap::IsAttributeInBitmask(Cached, MP_Anisotropy);
+	const bool bHasAnisotropy = FMaterialAttributeDefinitionMap::IsAttributeInBitmask(CachedConnectedMaterialAttributesInputs, MP_Anisotropy);
 
 	auto AddDefaultWorstCase = [&](uint32 In)
 		{
-			FSubstrateOperator& SlabOperator = Compiler->SubstrateCompilationRegisterOperator(SUBSTRATE_OPERATOR_BSDF_LEGACY, Compiler->SubstrateTreeStackGetPathUniqueId(), this, Parent, Compiler->SubstrateTreeStackGetParentPathUniqueId());
+			FSubstrateOperator& SlabOperator = Compiler->SubstrateCompilationRegisterOperator(SUBSTRATE_OPERATOR_BSDF_LEGACY, Compiler->SubstrateTreeStackGetPathUniqueId(), ThisExpressionGuid, Parent, Compiler->SubstrateTreeStackGetParentPathUniqueId());
 			SlabOperator.BSDFType = SUBSTRATE_BSDF_TYPE_SLAB;
-			SlabOperator.BSDFFeatures  = In | (bHasAnisotropy ? ESubstrateBsdfFeature_Anisotropy : ESubstrateBsdfFeature_None);
+			SlabOperator.BSDFFeatures = In | (bHasAnisotropy ? ESubstrateBsdfFeature_Anisotropy : ESubstrateBsdfFeature_None);
 			SlabOperator.ThicknessIndex = ThicknessIndex;
 
 			return &SlabOperator;
@@ -4088,11 +4161,11 @@ FSubstrateOperator* UMaterialExpressionSubstrateConvertMaterialAttributes::Subst
 	FMaterialShadingModelField ShadingModels = Compiler->GetMaterialShadingModels();
 
 	// Logic about shading models and complexity should match UMaterialExpressionSubstrateConvertMaterialAttributes::Compile.
-	const bool bHasShadingModelFromExpression = FMaterialAttributeDefinitionMap::IsAttributeInBitmask(Cached, MP_ShadingModel) || ShadingModelOverride == MSM_FromMaterialExpression; // We keep HasShadingModelFromExpression in case all shading models cannot be safely recovered from material functions.
+	const bool bHasShadingModelFromExpression = FMaterialAttributeDefinitionMap::IsAttributeInBitmask(CachedConnectedMaterialAttributesInputs, MP_ShadingModel) || bShadingModelFromMaterialExpression; // We keep HasShadingModelFromExpression in case all shading models cannot be safely recovered from material functions.
 	if ((ShadingModels.CountShadingModels() > 1) || bHasShadingModelFromExpression)
 	{
 		uint32 Features = ESubstrateBsdfFeature_SSS | ESubstrateBsdfFeature_Fuzz;
-		if (ShadingModels.HasShadingModel(MSM_Eye))  { Features |= ESubstrateBsdfFeature_Eye; }
+		if (ShadingModels.HasShadingModel(MSM_Eye)) { Features |= ESubstrateBsdfFeature_Eye; }
 		if (ShadingModels.HasShadingModel(MSM_Hair)) { Features |= ESubstrateBsdfFeature_Hair; }
 		return AddDefaultWorstCase(Features);
 	}
@@ -4102,7 +4175,7 @@ FSubstrateOperator* UMaterialExpressionSubstrateConvertMaterialAttributes::Subst
 
 		if (ShadingModels.HasShadingModel(MSM_Unlit))
 		{
-			FSubstrateOperator& Operator = Compiler->SubstrateCompilationRegisterOperator(SUBSTRATE_OPERATOR_BSDF_LEGACY, Compiler->SubstrateTreeStackGetPathUniqueId(), this, Parent, Compiler->SubstrateTreeStackGetParentPathUniqueId());
+			FSubstrateOperator& Operator = Compiler->SubstrateCompilationRegisterOperator(SUBSTRATE_OPERATOR_BSDF_LEGACY, Compiler->SubstrateTreeStackGetPathUniqueId(), ThisExpressionGuid, Parent, Compiler->SubstrateTreeStackGetParentPathUniqueId());
 			Operator.BSDFType = SUBSTRATE_BSDF_TYPE_UNLIT;
 			Operator.ThicknessIndex = ThicknessIndex;
 			return &Operator;
@@ -4137,7 +4210,7 @@ FSubstrateOperator* UMaterialExpressionSubstrateConvertMaterialAttributes::Subst
 		}
 		else if (ShadingModels.HasShadingModel(MSM_ClearCoat))
 		{
-			FSubstrateOperator& Operator = Compiler->SubstrateCompilationRegisterOperator(SUBSTRATE_OPERATOR_BSDF_LEGACY, Compiler->SubstrateTreeStackGetPathUniqueId(), this, Parent, Compiler->SubstrateTreeStackGetParentPathUniqueId());
+			FSubstrateOperator& Operator = Compiler->SubstrateCompilationRegisterOperator(SUBSTRATE_OPERATOR_BSDF_LEGACY, Compiler->SubstrateTreeStackGetPathUniqueId(), ThisExpressionGuid, Parent, Compiler->SubstrateTreeStackGetParentPathUniqueId());
 			Operator.BSDFType = SUBSTRATE_BSDF_TYPE_SLAB;
 			Operator.ThicknessIndex = ThicknessIndex;
 			Operator.BSDFFeatures = ESubstrateBsdfFeature_SecondRoughnessOrSimpleClearCoat | (bHasAnisotropy ? ESubstrateBsdfFeature_Anisotropy : ESubstrateBsdfFeature_None);
@@ -4145,7 +4218,7 @@ FSubstrateOperator* UMaterialExpressionSubstrateConvertMaterialAttributes::Subst
 		}
 		else if (ShadingModels.HasShadingModel(MSM_Hair))
 		{
-			FSubstrateOperator& Operator = Compiler->SubstrateCompilationRegisterOperator(SUBSTRATE_OPERATOR_BSDF_LEGACY, Compiler->SubstrateTreeStackGetPathUniqueId(), this, Parent, Compiler->SubstrateTreeStackGetParentPathUniqueId());
+			FSubstrateOperator& Operator = Compiler->SubstrateCompilationRegisterOperator(SUBSTRATE_OPERATOR_BSDF_LEGACY, Compiler->SubstrateTreeStackGetPathUniqueId(), ThisExpressionGuid, Parent, Compiler->SubstrateTreeStackGetParentPathUniqueId());
 			Operator.BSDFType = SUBSTRATE_BSDF_TYPE_HAIR;
 			Operator.BSDFFeatures = ESubstrateBsdfFeature_Hair;
 			Operator.ThicknessIndex = ThicknessIndex;
@@ -4153,7 +4226,7 @@ FSubstrateOperator* UMaterialExpressionSubstrateConvertMaterialAttributes::Subst
 		}
 		else if (ShadingModels.HasShadingModel(MSM_Eye))
 		{
-			FSubstrateOperator& Operator = Compiler->SubstrateCompilationRegisterOperator(SUBSTRATE_OPERATOR_BSDF_LEGACY, Compiler->SubstrateTreeStackGetPathUniqueId(), this, Parent, Compiler->SubstrateTreeStackGetParentPathUniqueId());
+			FSubstrateOperator& Operator = Compiler->SubstrateCompilationRegisterOperator(SUBSTRATE_OPERATOR_BSDF_LEGACY, Compiler->SubstrateTreeStackGetPathUniqueId(), ThisExpressionGuid, Parent, Compiler->SubstrateTreeStackGetParentPathUniqueId());
 			Operator.BSDFType = SUBSTRATE_BSDF_TYPE_EYE;
 			Operator.BSDFFeatures = ESubstrateBsdfFeature_Eye;
 			Operator.ThicknessIndex = ThicknessIndex;
@@ -4161,10 +4234,10 @@ FSubstrateOperator* UMaterialExpressionSubstrateConvertMaterialAttributes::Subst
 		}
 		else if (ShadingModels.HasShadingModel(MSM_SingleLayerWater))
 		{
-			FSubstrateOperator& Operator = Compiler->SubstrateCompilationRegisterOperator(SUBSTRATE_OPERATOR_BSDF_LEGACY, Compiler->SubstrateTreeStackGetPathUniqueId(), this, Parent, Compiler->SubstrateTreeStackGetParentPathUniqueId());
+			FSubstrateOperator& Operator = Compiler->SubstrateCompilationRegisterOperator(SUBSTRATE_OPERATOR_BSDF_LEGACY, Compiler->SubstrateTreeStackGetPathUniqueId(), ThisExpressionGuid, Parent, Compiler->SubstrateTreeStackGetParentPathUniqueId());
 			Operator.BSDFType = SUBSTRATE_BSDF_TYPE_SINGLELAYERWATER;
 			Operator.ThicknessIndex = ThicknessIndex;
-			Operator.bBSDFWritesEmissive = MaterialAttributes.IsConnected(MP_EmissiveColor);
+			Operator.bBSDFWritesEmissive = bIsEmissiveConnected;
 			return &Operator;
 		}
 
@@ -4174,6 +4247,15 @@ FSubstrateOperator* UMaterialExpressionSubstrateConvertMaterialAttributes::Subst
 	}
 }
 
+FSubstrateOperator* UMaterialExpressionSubstrateConvertMaterialAttributes::SubstrateGenerateMaterialTopologyTree(class FMaterialCompiler* Compiler, class UMaterialExpression* Parent, int32 OutputIndex)
+{
+	const uint64 Cached = GetConnectedMaterialAttributesInputs(Material);
+
+	return SubstrateGenerateMaterialTopologyTreeCommon(
+		Compiler, MaterialExpressionGuid, Parent, OutputIndex,
+		Cached, ShadingModelOverride == MSM_FromMaterialExpression, MaterialAttributes.IsConnected(MP_EmissiveColor));
+}
+
 bool UMaterialExpressionSubstrateConvertMaterialAttributes::HasSSS() const
 {
 	return SubsurfaceProfile != nullptr;
diff --git a/Engine/Source/Runtime/Engine/Private/Materials/MaterialExpressions.cpp b/Engine/Source/Runtime/Engine/Private/Materials/MaterialExpressions.cpp
--- a/Engine/Source/Runtime/Engine/Private/Materials/MaterialExpressions.cpp
+++ b/Engine/Source/Runtime/Engine/Private/Materials/MaterialExpressions.cpp
@@ -6785,7 +6785,7 @@ FSubstrateOperator* UMaterialExpressionBlendMaterialAttributes::SubstrateGenerat
 {
 	// SUBSTRATE_TODO: this likely no longer work (or only works with the new layering system). We would need to stop parsing and always do parameter blending at this stage.
 	const bool bUseParameterBlending = false;
-	FSubstrateOperator& SubstrateOperator = Compiler->SubstrateCompilationRegisterOperator(SUBSTRATE_OPERATOR_HORIZONTAL, Compiler->SubstrateTreeStackGetPathUniqueId(), this, Parent, Compiler->SubstrateTreeStackGetParentPathUniqueId(), bUseParameterBlending);
+	FSubstrateOperator& SubstrateOperator = Compiler->SubstrateCompilationRegisterOperator(SUBSTRATE_OPERATOR_HORIZONTAL, Compiler->SubstrateTreeStackGetPathUniqueId(), this->MaterialExpressionGuid, Parent, Compiler->SubstrateTreeStackGetParentPathUniqueId(), bUseParameterBlending);
 
 	UMaterialExpression* ChildAExpression = A.GetTracedInput().Expression;
 	UMaterialExpression* ChildBExpression = B.GetTracedInput().Expression;
@@ -7658,7 +7658,7 @@ FSubstrateOperator* UMaterialExpressionMaterialAttributeLayers::SubstrateGenerat
 			if (!OutOperator)
 			{
 				//If this is reached, compile a default operator to avoid a crash
-				FSubstrateOperator& DefaultOperator = Compiler->SubstrateCompilationRegisterOperator(SUBSTRATE_OPERATOR_BSDF, Compiler->SubstrateTreeStackGetPathUniqueId(), this, Parent, Compiler->SubstrateTreeStackGetParentPathUniqueId());
+				FSubstrateOperator& DefaultOperator = Compiler->SubstrateCompilationRegisterOperator(SUBSTRATE_OPERATOR_BSDF, Compiler->SubstrateTreeStackGetPathUniqueId(), this->MaterialExpressionGuid, Parent, Compiler->SubstrateTreeStackGetParentPathUniqueId());
 				DefaultOperator.BSDFType = SUBSTRATE_BSDF_TYPE_SLAB;
 				DefaultOperator.ThicknessIndex = Compiler->SubstrateThicknessStackGetThicknessIndex();
 				OutOperator = &DefaultOperator;
@@ -18087,7 +18087,7 @@ FSubstrateOperator* UMaterialExpressionFunctionInput::SubstrateGenerateMaterialT
 	// If this is not the case however, we need to compile a default material to compile the shader, or to preview the material function.
 	if(!SlabOperator)
 	{
-		SlabOperator = &Compiler->SubstrateCompilationRegisterOperator(SUBSTRATE_OPERATOR_BSDF, Compiler->SubstrateTreeStackGetPathUniqueId(), this, Parent, Compiler->SubstrateTreeStackGetParentPathUniqueId());
+		SlabOperator = &Compiler->SubstrateCompilationRegisterOperator(SUBSTRATE_OPERATOR_BSDF, Compiler->SubstrateTreeStackGetPathUniqueId(), this->MaterialExpressionGuid, Parent, Compiler->SubstrateTreeStackGetParentPathUniqueId());
 		SlabOperator->BSDFType = SUBSTRATE_BSDF_TYPE_SLAB;
 		SlabOperator->ThicknessIndex = Compiler->SubstrateThicknessStackGetThicknessIndex();
 	}
diff --git a/Engine/Source/Runtime/Engine/Public/MaterialCompiler.h b/Engine/Source/Runtime/Engine/Public/MaterialCompiler.h
--- a/Engine/Source/Runtime/Engine/Public/MaterialCompiler.h
+++ b/Engine/Source/Runtime/Engine/Public/MaterialCompiler.h
@@ -819,7 +819,7 @@ public:
 	/**
 	 * Register an operator of the tree representation the Substrate material and its topology.
 	 */
-	virtual FSubstrateOperator& SubstrateCompilationRegisterOperator(int32 OperatorType, FGuid SubstrateExpressionGuid, UMaterialExpression* Child, UMaterialExpression* Parent, FGuid SubstrateParentExpressionGuid, bool bUseParameterBlending = false) = 0;
+	virtual FSubstrateOperator& SubstrateCompilationRegisterOperator(int32 OperatorType, FGuid SubstrateExpressionGuid, FGuid ChildMaterialExpressionGuid, UMaterialExpression* Parent, FGuid SubstrateParentExpressionGuid, bool bUseParameterBlending = false) = 0;
 	/**
 	 * Return the operator information for a given expression.
 	 */
@@ -1635,9 +1635,9 @@ public:
 		Compiler->SubstrateThicknessStackPop();
 	}
 
-	virtual FSubstrateOperator& SubstrateCompilationRegisterOperator(int32 OperatorType, FGuid SubstrateExpressionGuid, UMaterialExpression* Child, UMaterialExpression* Parent, FGuid SubstrateParentExpressionGuid, bool bUseParameterBlending = false) override
+	virtual FSubstrateOperator& SubstrateCompilationRegisterOperator(int32 OperatorType, FGuid SubstrateExpressionGuid, FGuid ChildMaterialExpressionGuid, UMaterialExpression* Parent, FGuid SubstrateParentExpressionGuid, bool bUseParameterBlending = false) override
 	{
-		return Compiler->SubstrateCompilationRegisterOperator(OperatorType, SubstrateExpressionGuid, Child, Parent, SubstrateParentExpressionGuid, bUseParameterBlending);
+		return Compiler->SubstrateCompilationRegisterOperator(OperatorType, SubstrateExpressionGuid, ChildMaterialExpressionGuid, Parent, SubstrateParentExpressionGuid, bUseParameterBlending);
 	}
 
 	virtual FSubstrateOperator& SubstrateCompilationGetOperator(FGuid SubstrateExpressionGuid) override
diff --git a/Engine/Source/Runtime/Engine/Public/Materials/Material.h b/Engine/Source/Runtime/Engine/Public/Materials/Material.h
--- a/Engine/Source/Runtime/Engine/Public/Materials/Material.h
+++ b/Engine/Source/Runtime/Engine/Public/Materials/Material.h
@@ -1449,7 +1449,7 @@ private:
 	void BackwardsCompatibilityInputConversion();
 	void BackwardsCompatibilityVirtualTextureOutputConversion();
 	void BackwardsCompatibilityDecalConversion();
-	bool ConvertMaterialToSubstrateMaterial(bool bAllowEmptyMaterialUpdate);
+	bool ConvertMaterialToSubstrateMaterial(bool bAllowEmptyMaterialUpdate, bool bRequestedByUser = false);
 
 	/** Marks the material's package dirty in order to make a material usage change set during map load persistent. 
 	  * This couldn't be done during map load as loading cannot mark packages dirty. Invoked manually by the user 
@@ -1460,6 +1460,10 @@ private:
 
 public:
 
+#if WITH_EDITOR
+	ENGINE_API bool UserRequestsConvertMaterialToSubstrateMaterial();
+#endif
+
 	/** @return the name of the given usage flag. */
 	FString GetUsageName(const EMaterialUsage Usage) const;
 
diff --git a/Engine/Source/Runtime/Engine/Public/Materials/MaterialExpressionSubstrate.h b/Engine/Source/Runtime/Engine/Public/Materials/MaterialExpressionSubstrate.h
--- a/Engine/Source/Runtime/Engine/Public/Materials/MaterialExpressionSubstrate.h
+++ b/Engine/Source/Runtime/Engine/Public/Materials/MaterialExpressionSubstrate.h
@@ -226,6 +226,25 @@ class UMaterialExpressionSubstrateShadingModels : public UMaterialExpressionSubs
 	bool HasSSS() const;
 	bool HasAnisotropy() const;
 
+	static FSubstrateOperator* SubstrateGenerateMaterialTopologyTreeCommon(
+		class FMaterialCompiler* Compiler, struct FGuid ThisExpressionGuid, class UMaterialExpression* Parent, int32 OutputIndex,
+		const FExpressionInput& EmissiveColor,
+		const FExpressionInput& Anisotropy,
+		const FExpressionInput& ClearCoatNormal,
+		const FExpressionInput& CustomTangent,
+		const FExpressionInput& ShadingModel);
+
+	static int32 CompileCommon(class FMaterialCompiler* Compiler,
+		FExpressionInput& BaseColor, FExpressionInput& Specular, FExpressionInput& Metallic, FExpressionInput& Roughness, FExpressionInput& EmissiveColor,
+		FExpressionInput& Opacity, FExpressionInput& SubSurfaceColor, FExpressionInput& ClearCoat, FExpressionInput& ClearCoatRoughness,
+		FExpressionInput& ShadingModel, TEnumAsByte<enum EMaterialShadingModel> ShadingModelOverride,
+		FExpressionInput& TransmittanceColor, FExpressionInput& ThinTranslucentSurfaceCoverage,
+		FExpressionInput& WaterScatteringCoefficients, FExpressionInput& WaterAbsorptionCoefficients, FExpressionInput& WaterPhaseG, FExpressionInput& ColorScaleBehindWater,
+		const bool bHasAnisotropy, FExpressionInput& Anisotropy,
+		FExpressionInput& Normal, FExpressionInput& Tangent,
+		FExpressionInput& ClearCoatNormal, FExpressionInput& CustomTangent,
+		const bool bHasSSS, USubsurfaceProfile* SSSProfile,
+		const class UMaterialEditorOnlyData* EditorOnlyData = nullptr /*Used to provide default values from the root node pins*/);
 #endif
 	//~ End UMaterialExpression Interface
 };
@@ -504,6 +523,16 @@ class UMaterialExpressionSubstrateVolumetricFogCloudBSDF : public UMaterialExpre
 	virtual bool IsResultSubstrateMaterial(int32 OutputIndex) override;
 	virtual void GatherSubstrateMaterialInfo(FSubstrateMaterialInfo& SubstrateMaterialInfo, int32 OutputIndex) override;
 	virtual FSubstrateOperator* SubstrateGenerateMaterialTopologyTree(class FMaterialCompiler* Compiler, class UMaterialExpression* Parent, int32 OutputIndex) override;
+
+	static FSubstrateOperator* SubstrateGenerateMaterialTopologyTreeCommon(
+		class FMaterialCompiler* Compiler, struct FGuid ThisExpressionGuid, class UMaterialExpression* Parent, int32 OutputIndex,
+		const FExpressionInput& EmissiveColor,
+		const FExpressionInput& AmbientOcclusion);
+
+	static int32 CompileCommon(class FMaterialCompiler* Compiler,
+		FExpressionInput& Albedo, FExpressionInput& Extinction, FExpressionInput& EmissiveColor, FExpressionInput& AmbientOcclusion,
+		bool bEmissiveOnly,
+		const class UMaterialEditorOnlyData* EditorOnlyData = nullptr /*Used to provide default values from the root node pins*/);
 #endif
 	//~ End UMaterialExpression Interface
 };
@@ -778,6 +807,14 @@ class UMaterialExpressionSubstrateLightFunction : public UMaterialExpressionSubs
 	virtual bool IsResultSubstrateMaterial(int32 OutputIndex) override;
 	virtual void GatherSubstrateMaterialInfo(FSubstrateMaterialInfo& SubstrateMaterialInfo, int32 OutputIndex) override;
 	virtual FSubstrateOperator* SubstrateGenerateMaterialTopologyTree(class FMaterialCompiler* Compiler, class UMaterialExpression* Parent, int32 OutputIndex) override;
+
+	static FSubstrateOperator* SubstrateGenerateMaterialTopologyTreeCommon(
+		class FMaterialCompiler* Compiler, struct FGuid ThisExpressionGuid, class UMaterialExpression* Parent, int32 OutputIndex);
+
+	static int32 CompileCommon(class FMaterialCompiler* Compiler,
+		FExpressionInput& Color,
+		const UMaterialEditorOnlyData* EditorOnlyData = nullptr);
+
 #endif
 	//~ End UMaterialExpression Interface
 };
@@ -808,6 +845,13 @@ class UMaterialExpressionSubstratePostProcess : public UMaterialExpressionSubstr
 	virtual bool IsResultSubstrateMaterial(int32 OutputIndex) override;
 	virtual void GatherSubstrateMaterialInfo(FSubstrateMaterialInfo& SubstrateMaterialInfo, int32 OutputIndex) override;
 	virtual FSubstrateOperator* SubstrateGenerateMaterialTopologyTree(class FMaterialCompiler* Compiler, class UMaterialExpression* Parent, int32 OutputIndex) override;
+
+	static FSubstrateOperator* SubstrateGenerateMaterialTopologyTreeCommon(
+		class FMaterialCompiler* Compiler, struct FGuid ThisExpressionGuid, class UMaterialExpression* Parent, int32 OutputIndex);
+
+	static int32 CompileCommon(class FMaterialCompiler* Compiler,
+		FExpressionInput& Color, FExpressionInput& Opacity,
+		const UMaterialEditorOnlyData* EditorOnlyData = nullptr);
 #endif
 	//~ End UMaterialExpression Interface
 };
@@ -838,6 +882,13 @@ class UMaterialExpressionSubstrateUI : public UMaterialExpressionSubstrateBSDF
 	virtual bool IsResultSubstrateMaterial(int32 OutputIndex) override;
 	virtual void GatherSubstrateMaterialInfo(FSubstrateMaterialInfo& SubstrateMaterialInfo, int32 OutputIndex) override;
 	virtual FSubstrateOperator* SubstrateGenerateMaterialTopologyTree(class FMaterialCompiler* Compiler, class UMaterialExpression* Parent, int32 OutputIndex) override;
+
+	static FSubstrateOperator* SubstrateGenerateMaterialTopologyTreeCommon(
+		class FMaterialCompiler* Compiler, struct FGuid ThisExpressionGuid, class UMaterialExpression* Parent, int32 OutputIndex);
+
+	static int32 CompileCommon(class FMaterialCompiler* Compiler,
+		FExpressionInput& Color, FExpressionInput& Opacity,
+		const UMaterialEditorOnlyData* EditorOnlyData = nullptr);
 #endif
 	//~ End UMaterialExpression Interface
 };
@@ -929,6 +980,15 @@ class UMaterialExpressionSubstrateConvertMaterialAttributes : public UMaterialEx
 	
 	bool HasSSS() const;
 
+	static FSubstrateOperator* SubstrateGenerateMaterialTopologyTreeCommon(
+		class FMaterialCompiler* Compiler, struct FGuid ThisExpressionGuid, class UMaterialExpression* Parent, int32 OutputIndex,
+		const uint64 CachedConnectedMaterialAttributesInputs, const bool bShadingModelFromMaterialExpression, const bool bIsEmissiveConnected);
+
+	static int32 CompileCommon(class FMaterialCompiler* Compiler, int32 OutputIndex,
+		const uint64 CachedConnectedMaterialAttributesInputs, FMaterialAttributesInput& MaterialAttributes, enum EMaterialShadingModel ShadingModelOverride,
+		FExpressionInput& WaterScatteringCoefficients, FExpressionInput& WaterAbsorptionCoefficients, FExpressionInput& WaterPhaseG, FExpressionInput& ColorScaleBehindWater,
+		const bool bHasSSS, USubsurfaceProfile* SSSProfile);
+
 #endif // WITH_EDITOR
 	//~ End UMaterialExpression Interface
 };
diff --git a/Engine/Source/Runtime/RenderCore/Private/RenderUtils.cpp b/Engine/Source/Runtime/RenderCore/Private/RenderUtils.cpp
--- a/Engine/Source/Runtime/RenderCore/Private/RenderUtils.cpp
+++ b/Engine/Source/Runtime/RenderCore/Private/RenderUtils.cpp
@@ -2143,6 +2143,14 @@ namespace Substrate
 	{
 		return IsSubstrateEnabled() && CVarSubstrateEnableLayerSupport.GetValueOnAnyThread() > 0;
 	}
+
+	bool IsHiddenMaterialAssetConversionEnabled()
+	{
+		// Set this bool to false when we want to force all material to be converted to substrate.
+		const bool bSubstrateHiddenMaterialAssetConversion = true;
+
+		return IsSubstrateEnabled() && bSubstrateHiddenMaterialAssetConversion;
+	}
 }
 
 bool HasFirstPersonGBufferBit(EShaderPlatform ShaderPlatform)
diff --git a/Engine/Source/Runtime/RenderCore/Private/Shader.cpp b/Engine/Source/Runtime/RenderCore/Private/Shader.cpp
--- a/Engine/Source/Runtime/RenderCore/Private/Shader.cpp
+++ b/Engine/Source/Runtime/RenderCore/Private/Shader.cpp
@@ -2382,6 +2382,12 @@ void ShaderMapAppendKey(EShaderPlatform Platform, FShaderKeyGenerator& KeyGen)
 			KeyGen.Append(TEXT("DBUFFERPASS"));
 		}
 
+		if (Substrate::IsHiddenMaterialAssetConversionEnabled())
+		{
+			KeyGen.AppendSeparator();
+			KeyGen.Append(TEXT("HIDDENCONV"));
+		}
+
 		if (Substrate::IsBackCompatibilityEnabled())
 		{
 			KeyGen.AppendSeparator();
diff --git a/Engine/Source/Runtime/RenderCore/Public/RenderUtils.h b/Engine/Source/Runtime/RenderCore/Public/RenderUtils.h
--- a/Engine/Source/Runtime/RenderCore/Public/RenderUtils.h
+++ b/Engine/Source/Runtime/RenderCore/Public/RenderUtils.h
@@ -603,6 +603,8 @@ namespace Substrate
 	RENDERCORE_API bool IsAdvancedVisualizationEnabled();
 	RENDERCORE_API bool IsMaterialLayeringSupportEnabled();
 
+	RENDERCORE_API bool IsHiddenMaterialAssetConversionEnabled();
+
 	RENDERCORE_API uint32 GetRayTracingMaterialPayloadSizeInBytes(bool bFullySimplifiedMaterial);
 
 	RENDERCORE_API uint32 GetBytePerPixel();
