Modified FMaterialResource::IsDitherMasked() to also check IsMasked().
- fixes issues with codepaths checking IsDitherMasked() without explictly checking IsMasked() (eg: BlendModeToRayTracingInstanceMask(...))

diff --git a/Engine/Source/Runtime/Engine/Private/Materials/MaterialShared.cpp b/Engine/Source/Runtime/Engine/Private/Materials/MaterialShared.cpp
--- a/Engine/Source/Runtime/Engine/Private/Materials/MaterialShared.cpp
+++ b/Engine/Source/Runtime/Engine/Private/Materials/MaterialShared.cpp
@@ -1967,7 +1967,7 @@ bool FMaterialResource::IsMasked() const
 
 bool FMaterialResource::IsDitherMasked() const 
 {
-	return Material->DitherOpacityMask;
+	return Material->DitherOpacityMask && IsMasked();
 }
 
 bool FMaterialResource::AllowNegativeEmissiveColor() const 
