Fixes for compilation issues following 40645535

#tests PC editor dx12
#rnx

diff --git a/Engine/Source/Editor/GraphEditor/Private/SGraphEditorImpl.cpp b/Engine/Source/Editor/GraphEditor/Private/SGraphEditorImpl.cpp
--- a/Engine/Source/Editor/GraphEditor/Private/SGraphEditorImpl.cpp
+++ b/Engine/Source/Editor/GraphEditor/Private/SGraphEditorImpl.cpp
@@ -65,6 +65,7 @@
 #include "Widgets/Text/STextBlock.h"
 #include "HAL/IConsoleManager.h"
 #include "RenderUtils.h"
+#include "Materials/Material.h"
 
 class FActiveTimerHandle;
 struct FGeometry;
diff --git a/Engine/Source/Runtime/Engine/Private/Materials/HLSLMaterialTranslator.cpp b/Engine/Source/Runtime/Engine/Private/Materials/HLSLMaterialTranslator.cpp
--- a/Engine/Source/Runtime/Engine/Private/Materials/HLSLMaterialTranslator.cpp
+++ b/Engine/Source/Runtime/Engine/Private/Materials/HLSLMaterialTranslator.cpp
@@ -1604,7 +1604,7 @@ void FHLSLMaterialTranslator::TranslateMaterial()
 
 			if (EditorOnlyData)
 			{
-				FGuid RootNodeConversionGUID = FGuid::FGuid(0, 1, 2, 3); // This dummy GUID is the root node.
+				FGuid RootNodeConversionGUID = FGuid(0, 1, 2, 3); // This dummy GUID is the root node.
 				bSubstrateWritesEmissive |= EditorOnlyData->EmissiveColor.IsConnected();
 				bSubstrateWritesAmbientOcclusion |= EditorOnlyData->AmbientOcclusion.IsConnected();
 
