Use integers for IAD priority

diff --git a/Engine/Plugins/Runtime/GameFeatures/Source/GameFeatures/Public/GameFeaturesSubsystem.h b/Engine/Plugins/Runtime/GameFeatures/Source/GameFeatures/Public/GameFeaturesSubsystem.h
--- a/Engine/Plugins/Runtime/GameFeatures/Source/GameFeatures/Public/GameFeaturesSubsystem.h
+++ b/Engine/Plugins/Runtime/GameFeatures/Source/GameFeatures/Public/GameFeaturesSubsystem.h
@@ -734,7 +734,6 @@ private:
 	friend struct FGameFeaturePluginState_Unloading;
 
 	void OnGameFeaturePauseChange(const FGameFeaturePluginIdentifier& PluginIdentifier, const FString& PluginName, FGameFeaturePauseStateChangeContext& Context);
-	friend struct FBaseDownloadGameFeaturePluginState;
 	friend struct FGameFeaturePluginState_Deactivating;
 
 	void OnAssetManagerCreated();
diff --git a/Engine/Source/Runtime/Engine/Classes/Engine/AssetManager.h b/Engine/Source/Runtime/Engine/Classes/Engine/AssetManager.h
--- a/Engine/Source/Runtime/Engine/Classes/Engine/AssetManager.h
+++ b/Engine/Source/Runtime/Engine/Classes/Engine/AssetManager.h
@@ -40,7 +40,7 @@ DECLARE_MULTICAST_DELEGATE_OneParam(FOnAddedAssetSearchRoot, const FString&);
 struct FAssetManagerDownloadParams
 {
 	/** Download priority */
-	EStreamableDownloadPriority Priority = EStreamableDownloadPriority::Normal;
+	FStreamableDownloadPriority Priority = UE::StreamableManager::Private::DownloadDefaultPriority;
 	/** If set, any soft references will also be installed */
 	bool bInstallSoftReferences = false;
 };
diff --git a/Engine/Source/Runtime/Engine/Classes/Engine/StreamableManager.h b/Engine/Source/Runtime/Engine/Classes/Engine/StreamableManager.h
--- a/Engine/Source/Runtime/Engine/Classes/Engine/StreamableManager.h
+++ b/Engine/Source/Runtime/Engine/Classes/Engine/StreamableManager.h
@@ -57,12 +57,21 @@ using FStreamableDelegate = TDelegate<void()>;
 using FStreamableDelegateWithHandle = TDelegate<void(TSharedPtr<struct FStreamableHandle>)>;
 using FStreamableUpdateDelegate = TDelegate<void(TSharedRef<struct FStreamableHandle>)>;
 
+/** EXPERIMENTAL - download priority */
+using FStreamableDownloadPriority = int32;
+
 namespace UE::StreamableManager::Private
 {
 	// Default priority for all async loads
-	static constexpr TAsyncLoadPriority DefaultAsyncLoadPriority = 0;
+	static constexpr TAsyncLoadPriority DefaultAsyncLoadPriority	= 0;
 	// Priority to try and load immediately
-	static constexpr TAsyncLoadPriority AsyncLoadHighPriority = 100;
+	static constexpr TAsyncLoadPriority AsyncLoadHighPriority		= 100;
+
+	static constexpr FStreamableDownloadPriority DownloadVeryLowPriority	= -200;
+	static constexpr FStreamableDownloadPriority DownloadLowPriority		= -100;
+	static constexpr FStreamableDownloadPriority DownloadDefaultPriority	= 0;
+	static constexpr FStreamableDownloadPriority DownloadHighPriority		= 100;
+	static constexpr FStreamableDownloadPriority DownloadVeryHighPriority	= 200;
 
 	inline FStreamableDelegateWithHandle WrapDelegate(const FStreamableDelegate& InDelegate)
 	{
@@ -201,17 +210,6 @@ struct TStreamableHandleContextData : public FStreamableHandleContextDataBase
 	}
 };
 
-/** EXPERIMENTAL - download priority */
-enum class EStreamableDownloadPriority : uint8
-{
-	/** Below normal priority. */
-	Low,
-	/** Normal priority. */
-	Normal,
-	/** Above normal priority. */
-	High
-};
-
 /** A handle to a synchronous or async load. As long as the handle is Active, loaded assets will stay in memory */
 struct FStreamableHandle : public TSharedFromThis<FStreamableHandle>
 {
@@ -312,13 +310,13 @@ struct FStreamableHandle : public TSharedFromThis<FStreamableHandle>
 	}
 
 	/** Returns download priority. Only valid if a download has been requested and is not yet complete. */
-	ENGINE_API TOptional<EStreamableDownloadPriority> GetDownloadPriority() const;
+	ENGINE_API TOptional<FStreamableDownloadPriority> GetDownloadPriority() const;
 
 	/** 
 	 * Set download priority. Only valid if a download has been requested and is not yet complete. 
 	 * @return true if the priority could be set.
 	 */
-	ENGINE_API bool SetDownloadPriority(EStreamableDownloadPriority Priority);
+	ENGINE_API bool SetDownloadPriority(FStreamableDownloadPriority Priority);
 
 	/**
 	 * Release this handle. This can be called from normal gameplay code to indicate that the loaded assets are no longer needed
@@ -687,21 +685,13 @@ enum class EStreamableManagerCombinedHandleOptions : uint8
 
 ENUM_CLASS_FLAGS(EStreamableManagerCombinedHandleOptions);
 
-namespace UE::StreamableManager
-{
-	// Default priority for all async loads
-	static constexpr TAsyncLoadPriority DefaultAsyncLoadPriority = 0;
-	// Priority to try and load immediately
-	static constexpr TAsyncLoadPriority AsyncLoadHighPriority = 100;
-}
-
 /** EXPERIMENTAL -  Additional parameters for downloading packages */
 struct FStreamableDownloadParams
 {
 	/** Handle to keep packages in the local cache. If no cache pin is provided, a new one will be created. */
 	FStreamableDownloadCachePin CachePin;
 	/** Download priority */
-	EStreamableDownloadPriority Priority = EStreamableDownloadPriority::Normal;
+	FStreamableDownloadPriority Priority = UE::StreamableManager::Private::DownloadDefaultPriority;
 	/** If set, any soft references will also be installed */
 	bool bInstallSoftReferences = false;
 	/** If set, request will complete after download, no assets will be loaded */
@@ -733,9 +723,15 @@ struct FStreamableAsyncLoadParams
 struct FStreamableManager : public FGCObject
 {
 	// Default priority for all async loads
-	static constexpr TAsyncLoadPriority DefaultAsyncLoadPriority = UE::StreamableManager::Private::DefaultAsyncLoadPriority;
+	static constexpr TAsyncLoadPriority DefaultAsyncLoadPriority	= UE::StreamableManager::Private::DefaultAsyncLoadPriority;
 	// Priority to try and load immediately
-	static constexpr TAsyncLoadPriority AsyncLoadHighPriority = UE::StreamableManager::Private::AsyncLoadHighPriority;
+	static constexpr TAsyncLoadPriority AsyncLoadHighPriority		= UE::StreamableManager::Private::AsyncLoadHighPriority;
+
+	static constexpr FStreamableDownloadPriority DownloadVeryLowPriority	= UE::StreamableManager::Private::DownloadVeryLowPriority;
+	static constexpr FStreamableDownloadPriority DownloadLowPriority		= UE::StreamableManager::Private::DownloadLowPriority;
+	static constexpr FStreamableDownloadPriority DownloadDefaultPriority	= UE::StreamableManager::Private::DownloadDefaultPriority;
+	static constexpr FStreamableDownloadPriority DownloadHighPriority		= UE::StreamableManager::Private::DownloadHighPriority;
+	static constexpr FStreamableDownloadPriority DownloadVeryHighPriority	= UE::StreamableManager::Private::DownloadVeryHighPriority;
 
 	/** 
 	 * This is the primary streamable operation. Requests streaming of one or more target objects. When complete, a delegate function is called. Returns a Streamable Handle.
diff --git a/Engine/Source/Runtime/Engine/Private/StreamableManager.cpp b/Engine/Source/Runtime/Engine/Private/StreamableManager.cpp
--- a/Engine/Source/Runtime/Engine/Private/StreamableManager.cpp
+++ b/Engine/Source/Runtime/Engine/Private/StreamableManager.cpp
@@ -380,28 +380,12 @@ TStreamableHandleContextDataTypeID FStreamableHandleContextDataBase::AllocateCla
 
 namespace UE::StreamableManager::Private
 {
-	static UE::IoStore::EOnDemandInstallPriority GetOnDemandInstallPriority(EStreamableDownloadPriority Priority)
-	{
-		UE::IoStore::EOnDemandInstallPriority OutPriority = UE::IoStore::EOnDemandInstallPriority::Normal;
-		switch (Priority)
-		{
-		case EStreamableDownloadPriority::Low:
-			OutPriority = UE::IoStore::EOnDemandInstallPriority::Low;
-			break;
-		case EStreamableDownloadPriority::High:
-			OutPriority = UE::IoStore::EOnDemandInstallPriority::High;
-			break;
-		}
-
-		return OutPriority;
-	}
-
 	// self-deleting context for managing package downloads for a load request
 	struct FDownloadContext
 	{
 		UE::IoStore::FOnDemandInstallRequest InstallRequest;
 		UE::IoStore::FOnDemandInstallProgress InstallProgress;
-		EStreamableDownloadPriority Priority = EStreamableDownloadPriority::Normal;
+		FStreamableDownloadPriority Priority = FStreamableManager::DownloadDefaultPriority;
 		bool bInstallSoftReferences = false;
 		bool bDownloadOnly = false;
 
@@ -476,7 +460,7 @@ namespace UE::StreamableManager::Private
 				InstallArgs.Options |= UE::IoStore::EOnDemandInstallOptions::InstallSoftReferences;
 			}
 
-			InstallArgs.Priority = GetOnDemandInstallPriority(Priority);
+			InstallArgs.Priority = Priority;
 
 			InstallRequest = IoStore->Install(MoveTemp(InstallArgs),
 				[WeakOwningHandle = OwningHandle.ToWeakPtr()](const UE::IoStore::FOnDemandInstallResult& OnDemandInstallResult)
@@ -765,7 +749,7 @@ void FStreamableHandle::SetDebugNameIfEmptyOrDefault(const FString& NewName)
 	}
 }
 
-TOptional<EStreamableDownloadPriority> FStreamableHandle::GetDownloadPriority() const
+TOptional<FStreamableDownloadPriority> FStreamableHandle::GetDownloadPriority() const
 {
 	if (IsActive() && DownloadContext)
 	{
@@ -775,14 +759,12 @@ TOptional<EStreamableDownloadPriority> FStreamableHandle::GetDownloadPriority()
 	return {};
 }
 
-bool FStreamableHandle::SetDownloadPriority(EStreamableDownloadPriority InPriority)
+bool FStreamableHandle::SetDownloadPriority(FStreamableDownloadPriority InPriority)
 {
 	if (IsActive() && DownloadContext)
 	{
 		DownloadContext->Priority = InPriority;
-		DownloadContext->InstallRequest.UpdatePriority(
-			UE::StreamableManager::Private::GetOnDemandInstallPriority(InPriority));
-
+		DownloadContext->InstallRequest.UpdatePriority(DownloadContext->Priority);
 		return true;
 	}
 
diff --git a/Engine/Source/Runtime/Experimental/IoStore/OnDemand/Private/OnDemandIoStore.cpp b/Engine/Source/Runtime/Experimental/IoStore/OnDemand/Private/OnDemandIoStore.cpp
--- a/Engine/Source/Runtime/Experimental/IoStore/OnDemand/Private/OnDemandIoStore.cpp
+++ b/Engine/Source/Runtime/Experimental/IoStore/OnDemand/Private/OnDemandIoStore.cpp
@@ -2005,7 +2005,7 @@ void FOnDemandIoStore::CancelInstallRequest(FSharedInternalInstallRequest Instal
 	UE_LOG(LogIoStoreOnDemand, Log, TEXT("Cancelling install request, %s"), *LexToString(*InstallRequest));
 }
 
-void FOnDemandIoStore::UpdateInstallRequestPriority(FSharedInternalInstallRequest InstallRequest, EOnDemandInstallPriority NewPriority)
+void FOnDemandIoStore::UpdateInstallRequestPriority(FSharedInternalInstallRequest InstallRequest, int32 NewPriority)
 {
 	{
 		UE::TUniqueLock Lock(RequestMutex);
@@ -2165,7 +2165,7 @@ void FOnDemandInstallRequest::Cancel()
 	}
 }
 
-void FOnDemandInstallRequest::UpdatePriority(EPriority NewPriority)
+void FOnDemandInstallRequest::UpdatePriority(int32 NewPriority)
 {
 	if (IsPending())
 	{
diff --git a/Engine/Source/Runtime/Experimental/IoStore/OnDemand/Private/OnDemandIoStore.h b/Engine/Source/Runtime/Experimental/IoStore/OnDemand/Private/OnDemandIoStore.h
--- a/Engine/Source/Runtime/Experimental/IoStore/OnDemand/Private/OnDemandIoStore.h
+++ b/Engine/Source/Runtime/Experimental/IoStore/OnDemand/Private/OnDemandIoStore.h
@@ -339,7 +339,7 @@ public:
 	void								ReleaseContent(FOnDemandInternalContentHandle& ContentHandle);
 	void								GetReferencedContent(TArray<FSharedOnDemandContainer>& OutContainers, TArray<TBitArray<>>& OutChunkEntryIndices);
 	void								CancelInstallRequest(FSharedInternalInstallRequest InstallRequest);
-	void								UpdateInstallRequestPriority(FSharedInternalInstallRequest InstallRequest, EOnDemandInstallPriority NewPriority);
+	void								UpdateInstallRequestPriority(FSharedInternalInstallRequest InstallRequest, int32 NewPriority);
 
 	// IOnDemandIoStore interface
 	virtual void						Mount(FOnDemandMountArgs&& Args, FOnDemandMountCompleted&& OnCompleted) override;
diff --git a/Engine/Source/Runtime/Experimental/IoStore/OnDemand/Public/IO/IoStoreOnDemand.h b/Engine/Source/Runtime/Experimental/IoStore/OnDemand/Public/IO/IoStoreOnDemand.h
--- a/Engine/Source/Runtime/Experimental/IoStore/OnDemand/Public/IO/IoStoreOnDemand.h
+++ b/Engine/Source/Runtime/Experimental/IoStore/OnDemand/Public/IO/IoStoreOnDemand.h
@@ -54,16 +54,6 @@ public:
 		/** The request failed. */
 		Error
 	};
-	/** Request priority. */
-	enum EPriority : uint8
-	{
-		/** Below normal priority. */
-		Low,
-		/** Normal priority. */
-		Normal,
-		/** Above normal priority. */
-		High
-	};
 	/** Destructor. */
 	virtual ~FOnDemandRequest() = default;
 	/** Returns the current status. */
@@ -187,9 +177,6 @@ enum class EOnDemandInstallOptions
 };
 ENUM_CLASS_FLAGS(EOnDemandInstallOptions);
 
-/** Install priority. */
-using EOnDemandInstallPriority = FOnDemandRequest::EPriority;
-
 /** Arguments for installing/downloading on-demand content. */
 struct FOnDemandInstallArgs 
 {
@@ -204,7 +191,7 @@ struct FOnDemandInstallArgs
 	/** Install options. */
 	EOnDemandInstallOptions Options = EOnDemandInstallOptions::None;
 	/** Priority. */
-	EOnDemandInstallPriority Priority = EOnDemandInstallPriority::Low;
+	int32 Priority = 0;
 	/** Optional debug name . */
 	FSharedString DebugName;
 };
@@ -302,7 +289,7 @@ public:
 	/** Cancel the install request. */
 	UE_API void Cancel();
 	/** Update priority of the instlal request. */
-	UE_API void UpdatePriority(EPriority NewPriority);
+	UE_API void UpdatePriority(int32 NewPriority);
 
 private:
 	FOnDemandInstallRequest(FWeakOnDemandIoStore IoStore, FSharedInternalInstallRequest InternalRequest);
