Modulation Concurrency Interaction on One-Shots only

diff --git a/Engine/Source/Runtime/Engine/Private/ActiveSound.cpp b/Engine/Source/Runtime/Engine/Private/ActiveSound.cpp
--- a/Engine/Source/Runtime/Engine/Private/ActiveSound.cpp
+++ b/Engine/Source/Runtime/Engine/Private/ActiveSound.cpp
@@ -61,11 +61,11 @@ FAutoConsoleVariableRef CVarFocusDependsOnSpatialization(
 	TEXT("0: Focus and Spatialization are orthogonal, 1: Focus will be disabled if spatialization is not also enabled"),
 	ECVF_Default);
 
-static bool bIncludeModulationInConcurrencyEvaluationCVar = false;
+static int32 IncludeModulationInConcurrencyEvaluationCVar = 0;
 FAutoConsoleVariableRef CVarIncludeModulationInConcurrencyEvaluation(
 	TEXT("au.Concurrency.IncludeVolumeModulation"),
-	bIncludeModulationInConcurrencyEvaluationCVar,
-	TEXT("0: Ignore Audio Modulation Volume when evaluating concurrency (Default), 1: Include a sound's Volume Modulation Destination when determining whether to virtualize a sound based on volume."),
+	IncludeModulationInConcurrencyEvaluationCVar,
+	TEXT("0: Ignore Audio Modulation Volume when evaluating concurrency (Default), 1: Apply Modulation Destination to logical volume for one-shots only, 2: Apply Modulation Destination to logical volume on all sounds."),
 	ECVF_Default);
 
 FTraceDelegate FActiveSound::ActiveSoundTraceDelegate;
@@ -1061,7 +1061,8 @@ void FActiveSound::UpdateWaveInstances(TArray<FWaveInstance*> &InWaveInstances,
 				}
 			}
 
-			if (bIncludeModulationInConcurrencyEvaluationCVar)
+			const bool bUseModulationInConcurrency = (IncludeModulationInConcurrencyEvaluationCVar == 2) || (IncludeModulationInConcurrencyEvaluationCVar == 1 && IsOneShot());
+			if (bUseModulationInConcurrency)
 			{
 				// If there are no modulators set on any wave instance, then MaxSourceModulationValue may still be negative (see line below),
 				// so we take the absolute value in that case (1.0).
