Fix welds waking up physics components with StartAwake=false before actors initialize.

Welding calls UpdatePhysicsFilterData which forces the body to wake up. As a result, any setup that welds or autowelds will start awake after play begins
regardless of the value of the StartAwake flag. This change will sleep the body based on the StartAwake flag prior to play beginning to undo the wakeups caused
by welding in init.

diff --git a/Engine/Source/Runtime/Engine/Private/Components/PrimitiveComponent.cpp b/Engine/Source/Runtime/Engine/Private/Components/PrimitiveComponent.cpp
--- a/Engine/Source/Runtime/Engine/Private/Components/PrimitiveComponent.cpp
+++ b/Engine/Source/Runtime/Engine/Private/Components/PrimitiveComponent.cpp
@@ -933,6 +933,16 @@ void UPrimitiveComponent::OnCreatePhysicsState()
 			}
 #endif // WITH_EDITOR
 		}
+
+		if (GetOwner() && !GetOwner()->HasActorBegunPlay())
+		{
+			FBodyInstance* RootBodyInstance = GetBodyInstance();
+			if (RootBodyInstance && !RootBodyInstance->bStartAwake)
+			{
+				RootBodyInstance->PutInstanceToSleep();
+			}
+		}
+
 	}
 
 	if (!IsAsyncCreatePhysicsStateRunning())
