[AutoRTFMTests] Add 403 to the list of accepted HTTP GET responses for FAutoRTFMHttpModuleNetworking

It would appear that certain network environments can upset the expected response from this test.

diff --git a/Engine/Source/Runtime/Engine/Private/Tests/AutoRTFM/AutoRTFMHttpModuleTests.cpp b/Engine/Source/Runtime/Engine/Private/Tests/AutoRTFM/AutoRTFMHttpModuleTests.cpp
--- a/Engine/Source/Runtime/Engine/Private/Tests/AutoRTFM/AutoRTFMHttpModuleTests.cpp
+++ b/Engine/Source/Runtime/Engine/Private/Tests/AutoRTFM/AutoRTFMHttpModuleTests.cpp
@@ -97,6 +97,7 @@ bool FAutoRTFMHttpModuleNetworking::RunTest(const FString& Parameters)
 			TestTrueExpr(HttpResponse->GetResponseCode() == 200 ||
 						 HttpResponse->GetResponseCode() == 301 ||
 						 HttpResponse->GetResponseCode() == 307 ||
+						 HttpResponse->GetResponseCode() == 403 ||
 						 HttpResponse->GetResponseCode() == 404);
 		});
 		HttpRequest->ProcessRequest();
