[AutoRTFM] Enable AutoRTFM for Slate.

This is being called from the closed.

Enabling this will add additional "closed" variants of the module's functions, which can be called from an AutoRTFM transaction.

diff --git a/Engine/Source/Runtime/Engine/Private/Tests/AutoRTFM/AutoRTFMSlateTests.cpp b/Engine/Source/Runtime/Engine/Private/Tests/AutoRTFM/AutoRTFMSlateTests.cpp
new file mode 100644
--- /dev/null
+++ b/Engine/Source/Runtime/Engine/Private/Tests/AutoRTFM/AutoRTFMSlateTests.cpp
@@ -0,0 +1,33 @@
+// Copyright Epic Games, Inc. All Rights Reserved.
+
+#include "AutoRTFM/AutoRTFMTesting.h"
+#include "Misc/AutomationTest.h"
+#include "AutoRTFM/AutoRTFM.h"
+
+#include "Framework/Text/CharRangeList.h"
+
+#if WITH_DEV_AUTOMATION_TESTS
+
+IMPLEMENT_SIMPLE_AUTOMATION_TEST(FAutoRTFMSlateTests, "AutoRTFM + Slate", EAutomationTestFlags::EngineFilter | EAutomationTestFlags::ClientContext | EAutomationTestFlags::ServerContext | EAutomationTestFlags::CommandletContext)
+
+bool FAutoRTFMSlateTests::RunTest(const FString & Parameters)
+{
+	if (!AutoRTFM::ForTheRuntime::IsAutoRTFMRuntimeEnabled())
+	{
+		ExecutionInfo.AddEvent(FAutomationEvent(EAutomationEventType::Info, TEXT("SKIPPED 'FAutoRTFMSlateTests' test. AutoRTFM disabled.")));
+		return true;
+	}
+
+	// Test for SOL-7842
+	AutoRTFM::Testing::Commit([&]
+	{
+		FCharRangeList CharRangeList;
+		TestTrueExpr(CharRangeList.IsEmpty());
+		CharRangeList.InitializeFromString(TEXT("a-zA-Z0-9._"));
+		TestFalseExpr(CharRangeList.IsEmpty());
+	});
+	
+	return true;
+}
+
+#endif //WITH_DEV_AUTOMATION_TESTS
diff --git a/Engine/Source/Runtime/Slate/Slate.Build.cs b/Engine/Source/Runtime/Slate/Slate.Build.cs
--- a/Engine/Source/Runtime/Slate/Slate.Build.cs
+++ b/Engine/Source/Runtime/Slate/Slate.Build.cs
@@ -60,7 +60,5 @@ public class Slate : ModuleRules
         }
 
 		PrivateIncludePaths.Add(System.IO.Path.Combine(EngineDirectory,"Source/ThirdParty/AHEasing/AHEasing-1.3.2"));
-
-		bDisableAutoRTFMInstrumentation = true;
 	}
 }
