Substrate - improve constant layout in material graph when converting to Substrate.

#tests PC editor
#rnx

diff --git a/Engine/Source/Runtime/Engine/Private/Materials/Material.cpp b/Engine/Source/Runtime/Engine/Private/Materials/Material.cpp
--- a/Engine/Source/Runtime/Engine/Private/Materials/Material.cpp
+++ b/Engine/Source/Runtime/Engine/Private/Materials/Material.cpp
@@ -3464,6 +3464,8 @@ bool UMaterial::ConvertMaterialToSubstrateMaterial(bool bAllowEmptyMaterialUpdat
 
 	// Store current node post from the root node.
 	int32 CurrentNodePosX = EditorX;
+	int32 ConstantNodePosX = EditorX;
+	int32 ConstantNodePosY = EditorY;
 	const int32 TranslationOffsetX = 350.0f;
 	auto SetPosXAndMoveReferenceToTheRight = [&](UMaterialExpression* Node)
 	{
@@ -3477,6 +3479,18 @@ bool UMaterial::ConvertMaterialToSubstrateMaterial(bool bAllowEmptyMaterialUpdat
 		NodeToReplace->MaterialExpressionEditorX = NewNode->MaterialExpressionEditorX + TranslationOffsetX;
 		CurrentNodePosX = FMath::Max(NodeToReplace->MaterialExpressionEditorX + TranslationOffsetX, CurrentNodePosX);
 	};
+	auto BeginConstantNodeLayout = [&]()
+	{
+		// Leave some space for constants
+		CurrentNodePosX += TranslationOffsetX;
+	};
+	auto LaydownConstant = [&](UMaterialExpression* Node)
+	{
+		// Position the constant on the constant column, each one on top of each other
+		Node->MaterialExpressionEditorX = ConstantNodePosX;
+		Node->MaterialExpressionEditorY = ConstantNodePosY;
+		ConstantNodePosY += Node->IsA<UMaterialExpressionConstant>() ? 80 : 200;
+	};
 
 	// Recommendation: use the non indexed version whenever possible. However, indexed version are mandatory for material function which do not have explicit input.
 	auto ConnectionToIndexed = [](auto& OldNodeInput, UMaterialExpression* NewNode, uint32 NewInputIndex, uint32 OperationType = SUBSTRATE_MOVE_CONNECTION)
@@ -3520,6 +3534,7 @@ bool UMaterial::ConvertMaterialToSubstrateMaterial(bool bAllowEmptyMaterialUpdat
 			{
 				OldNodeInput.Expression = nullptr;
 			}
+			LaydownConstant(ExpressionConstant);
 		}
 	};
 	auto ScalarMatInputConnectionTo = [&](FScalarMaterialInput& OldNodeInput, FExpressionInput& NewNodeInput, EMaterialProperty Property, uint32 OperationType = SUBSTRATE_MOVE_CONNECTION)
@@ -3541,6 +3556,7 @@ bool UMaterial::ConvertMaterialToSubstrateMaterial(bool bAllowEmptyMaterialUpdat
 			{
 				OldNodeInput.Expression = nullptr;
 			}
+			LaydownConstant(ExpressionConstant);
 		}
 	};
 	auto ColorMatInputConnectionToIndexed = [&](FColorMaterialInput& OldNodeInput, UMaterialExpression* NewNode, uint32 NewInputIndex, EMaterialProperty Property, uint32 OperationType = SUBSTRATE_MOVE_CONNECTION)
@@ -3564,6 +3580,7 @@ bool UMaterial::ConvertMaterialToSubstrateMaterial(bool bAllowEmptyMaterialUpdat
 			{
 				OldNodeInput.Expression = nullptr;
 			}
+			LaydownConstant(ExpressionConstant);
 		}
 	};
 	auto ColorMatInputConnectionTo = [&](FColorMaterialInput& OldNodeInput, FExpressionInput& NewNodeInput, EMaterialProperty Property, uint32 OperationType = SUBSTRATE_MOVE_CONNECTION)
@@ -3587,6 +3604,7 @@ bool UMaterial::ConvertMaterialToSubstrateMaterial(bool bAllowEmptyMaterialUpdat
 			{
 				OldNodeInput.Expression = nullptr;
 			}
+			LaydownConstant(ExpressionConstant);
 		}
 	};
 	auto Vector3MatInputConnectionToIndexed = [&](FVectorMaterialInput& OldNodeInput, UMaterialExpression* NewNode, uint32 NewInputIndex, EMaterialProperty Property, uint32 OperationType = SUBSTRATE_MOVE_CONNECTION)
@@ -3610,6 +3628,7 @@ bool UMaterial::ConvertMaterialToSubstrateMaterial(bool bAllowEmptyMaterialUpdat
 			{
 				OldNodeInput.Expression = nullptr;
 			}
+			LaydownConstant(ExpressionConstant);
 		}
 	};
 	auto Vector3MatInputConnectionTo = [&](FVectorMaterialInput& OldNodeInput, FExpressionInput& NewNodeInput, EMaterialProperty Property, uint32 OperationType = SUBSTRATE_MOVE_CONNECTION)
@@ -3633,6 +3652,7 @@ bool UMaterial::ConvertMaterialToSubstrateMaterial(bool bAllowEmptyMaterialUpdat
 			{
 				OldNodeInput.Expression = nullptr;
 			}
+			LaydownConstant(ExpressionConstant);
 		}
 	};
 
@@ -3688,6 +3708,8 @@ bool UMaterial::ConvertMaterialToSubstrateMaterial(bool bAllowEmptyMaterialUpdat
 	// However, we do run the Substrate conversion/update code to make sure material with FrontMaterial plugged in are still converted.
 	if (Substrate::IsSubstrateEnabled() && (!Substrate::IsHiddenMaterialAssetConversionEnabled() || bRequestedByUser))
 	{
+		BeginConstantNodeLayout();
+
 		// SSS Profile
 		const bool bHasShadingModelMixture		= ShadingModels.CountShadingModels() > 1;
 		const bool bRequireSubsurfacePasses		= ShadingModels.HasShadingModel(MSM_SubsurfaceProfile) || ShadingModels.HasShadingModel(MSM_Subsurface) || ShadingModels.HasShadingModel(MSM_PreintegratedSkin) || ShadingModels.HasShadingModel(MSM_Eye);
