Fix WPO no longer being visible and active for deferred decal with Substrate.

diff --git a/Engine/Source/Runtime/Engine/Private/Materials/Material.cpp b/Engine/Source/Runtime/Engine/Private/Materials/Material.cpp
--- a/Engine/Source/Runtime/Engine/Private/Materials/Material.cpp
+++ b/Engine/Source/Runtime/Engine/Private/Materials/Material.cpp
@@ -7573,6 +7573,7 @@ static bool IsPropertyActive_Internal(EMaterialProperty InProperty,
 		if (bShowAsSubstrate)
 		{
 			return InProperty == MP_FrontMaterial
+				|| InProperty == MP_WorldPositionOffset
 				|| InProperty == MP_AmbientOcclusion
 				|| (InProperty == MP_Opacity && bSubstrateOpacityOverrideAllowed);
 		}
