Fixed skeletal mesh bounds not showing up in game mode.

diff --git a/Engine/Source/Runtime/Engine/Private/SkeletalMeshSceneProxy.cpp b/Engine/Source/Runtime/Engine/Private/SkeletalMeshSceneProxy.cpp
--- a/Engine/Source/Runtime/Engine/Private/SkeletalMeshSceneProxy.cpp
+++ b/Engine/Source/Runtime/Engine/Private/SkeletalMeshSceneProxy.cpp
@@ -1070,10 +1070,11 @@ FPrimitiveViewRelevance FSkeletalMeshSceneProxy::GetViewRelevance(const FSceneVi
 	{
 #if !(UE_BUILD_SHIPPING) || WITH_EDITOR
 		return IsRichView(*View->Family)
+			|| EngineShowFlags.Bounds
 			|| EngineShowFlags.Bones
 			|| EngineShowFlags.Collision
 			|| EngineShowFlags.VisualizeGPUSkinCache
-			|| (IsSelected() && (EngineShowFlags.VertexColors || EngineShowFlags.Bounds))
+			|| (IsSelected() && EngineShowFlags.VertexColors)
 			|| bForceWireframe
 #if WITH_EDITORONLY_DATA
 			|| MeshObject->SelectedEditorMaterial != -1
