Required minor engine changes for upcoming experimental FastGeo plugin
- Export necessary engine functions
- Refactor FStaticMeshComponentHelper & FInstancedStaticMeshComponentHelper to support SceneProxyDesc
- Allow FPrimitiveSceneProxyDesc::InitializeFromPrimitiveComponent & FStaticMeshSceneProxyDesc::InitializeFromStaticMeshComponent to be called even when desc scene is not initialized.

diff --git a/Engine/Source/Runtime/Engine/Private/Components/StaticMeshComponent.cpp b/Engine/Source/Runtime/Engine/Private/Components/StaticMeshComponent.cpp
--- a/Engine/Source/Runtime/Engine/Private/Components/StaticMeshComponent.cpp
+++ b/Engine/Source/Runtime/Engine/Private/Components/StaticMeshComponent.cpp
@@ -1668,12 +1668,12 @@ void UStaticMeshComponent::InitResources()
 
 void UStaticMeshComponent::CollectPSOPrecacheDataImpl(const FVertexFactoryType* VFType, const FPSOPrecacheParams& BasePrecachePSOParams, GetPSOVertexElementsFn GetVertexElements,FMaterialInterfacePSOPrecacheParamsList& OutParams) const
 {
-	FStaticMeshComponentHelper::CollectPSOPrecacheDataImpl(*this, VFType, BasePrecachePSOParams, GetVertexElements, OutParams);
+	FStaticMeshComponentHelper::CollectPSOPrecacheDataImpl(*this, [this]() { return GetRenderMatrix(); }, VFType, BasePrecachePSOParams, GetVertexElements, OutParams);
 }
 
 void UStaticMeshComponent::CollectPSOPrecacheData(const FPSOPrecacheParams& BasePrecachePSOParams, FMaterialInterfacePSOPrecacheParamsList& OutParams)
 {
-	FStaticMeshComponentHelper::CollectPSOPrecacheData(*this, BasePrecachePSOParams, OutParams);
+	FStaticMeshComponentHelper::CollectPSOPrecacheData(*this, [this]() { return GetRenderMatrix(); }, BasePrecachePSOParams, OutParams);
 }
 
 #if WITH_EDITOR
diff --git a/Engine/Source/Runtime/Engine/Private/InstancedStaticMesh.cpp b/Engine/Source/Runtime/Engine/Private/InstancedStaticMesh.cpp
--- a/Engine/Source/Runtime/Engine/Private/InstancedStaticMesh.cpp
+++ b/Engine/Source/Runtime/Engine/Private/InstancedStaticMesh.cpp
@@ -1378,7 +1378,7 @@ void FInstancedStaticMeshSceneProxyDesc::InitializeFromInstancedStaticMeshCompon
 	InitializeFromStaticMeshComponent(InComponent);
 
 	InstanceDataSceneProxy = InComponent->GetInstanceDataSceneProxy();
-	check(InstanceDataSceneProxy)
+	check(!GetScene() || InstanceDataSceneProxy);
 #if WITH_EDITOR
 	bHasSelectedInstances = InComponent->SelectedInstances.Find(true) != INDEX_NONE;
 #endif
@@ -5081,7 +5081,7 @@ void UInstancedStaticMeshComponent::PostLoad()
 
 void UInstancedStaticMeshComponent::CollectPSOPrecacheData(const FPSOPrecacheParams& BasePrecachePSOParams, FMaterialInterfacePSOPrecacheParamsList& OutParams)
 {
-	FInstancedStaticMeshComponentHelper::CollectPSOPrecacheData(*this, BasePrecachePSOParams, OutParams);
+	FInstancedStaticMeshComponentHelper::CollectPSOPrecacheData(*this, [this]() { return GetRenderMatrix(); }, BasePrecachePSOParams, OutParams);
 }
 
 void UInstancedStaticMeshComponent::OnPostLoadPerInstanceData()
diff --git a/Engine/Source/Runtime/Engine/Private/PrimitiveSceneProxy.cpp b/Engine/Source/Runtime/Engine/Private/PrimitiveSceneProxy.cpp
--- a/Engine/Source/Runtime/Engine/Private/PrimitiveSceneProxy.cpp
+++ b/Engine/Source/Runtime/Engine/Private/PrimitiveSceneProxy.cpp
@@ -395,7 +395,10 @@ void FPrimitiveSceneProxyDesc::InitializeFromPrimitiveComponent(const UPrimitive
 	Scene = InComponent->GetScene();
 	PrimitiveComponentInterface = InComponent->GetPrimitiveComponentInterface();
 	
-	FeatureLevel = Scene->GetFeatureLevel();
+	if (Scene)
+	{
+		FeatureLevel = Scene->GetFeatureLevel();
+	}
 
 #if WITH_EDITOR
 	HiddenEditorViews = InComponent->GetHiddenEditorViews();
diff --git a/Engine/Source/Runtime/Engine/Private/StaticMeshSceneProxy.cpp b/Engine/Source/Runtime/Engine/Private/StaticMeshSceneProxy.cpp
--- a/Engine/Source/Runtime/Engine/Private/StaticMeshSceneProxy.cpp
+++ b/Engine/Source/Runtime/Engine/Private/StaticMeshSceneProxy.cpp
@@ -2732,7 +2732,10 @@ void FStaticMeshSceneProxyDesc::InitializeFromStaticMeshComponent(const UStaticM
 	WireframeColor = InComponent->GetWireframeColor();
 	LODParentPrimitive = InComponent->GetLODParentPrimitive();	
 
-	SetMaterialRelevance(InComponent->GetMaterialRelevance(GetScene()->GetFeatureLevel()));
+	if (GetScene())
+	{
+		SetMaterialRelevance(InComponent->GetMaterialRelevance(GetScene()->GetFeatureLevel()));
+	}
 	SetCollisionResponseToChannels(InComponent->GetCollisionResponseToChannels());
 
 	MeshPaintTexture = InComponent->MeshPaintTextureOverride ? InComponent->MeshPaintTextureOverride.Get() : InComponent->GetMeshPaintTexture();
diff --git a/Engine/Source/Runtime/Engine/Public/InstancedStaticMeshComponentHelper.h b/Engine/Source/Runtime/Engine/Public/InstancedStaticMeshComponentHelper.h
--- a/Engine/Source/Runtime/Engine/Public/InstancedStaticMeshComponentHelper.h
+++ b/Engine/Source/Runtime/Engine/Public/InstancedStaticMeshComponentHelper.h
@@ -13,8 +13,10 @@
 class FInstancedStaticMeshComponentHelper
 {
 public:
+	using GetComponentRenderMatrixFn = TFunctionRef<FMatrix()>;
+
 	template<class T>
-	static void CollectPSOPrecacheData(const T& Component, const FPSOPrecacheParams& BasePrecachePSOParams, FMaterialInterfacePSOPrecacheParamsList& OutParams);
+	static void CollectPSOPrecacheData(const T& Component, GetComponentRenderMatrixFn, const FPSOPrecacheParams& BasePrecachePSOParams, FMaterialInterfacePSOPrecacheParamsList& OutParams);
 
 private:
 	// Dummy instance buffer used for PSO pre-caching
@@ -32,7 +34,7 @@ private:
 };
 
 template<class T>
-void FInstancedStaticMeshComponentHelper::CollectPSOPrecacheData(const T& Component, const FPSOPrecacheParams& BasePrecachePSOParams, FMaterialInterfacePSOPrecacheParamsList& OutParams)
+void FInstancedStaticMeshComponentHelper::CollectPSOPrecacheData(const T& Component, GetComponentRenderMatrixFn GetComponentRenderMatrix, const FPSOPrecacheParams& BasePrecachePSOParams, FMaterialInterfacePSOPrecacheParamsList& OutParams)
 {
 	if (Component.GetStaticMesh() == nullptr || Component.GetStaticMesh()->GetRenderData() == nullptr)
 	{
@@ -59,10 +61,10 @@ void FInstancedStaticMeshComponentHelper::CollectPSOPrecacheData(const T& Compon
 	Nanite::FMaterialAudit NaniteMaterials{};
 	if (Nanite::FNaniteResourcesHelper::ShouldCreateNaniteProxy(Component, &NaniteMaterials))
 	{
-		FStaticMeshComponentHelper::CollectPSOPrecacheDataImpl(Component, &FNaniteVertexFactory::StaticType, BasePrecachePSOParams, ISMC_GetElements, OutParams);
+		FStaticMeshComponentHelper::CollectPSOPrecacheDataImpl(Component, GetComponentRenderMatrix, &FNaniteVertexFactory::StaticType, BasePrecachePSOParams, ISMC_GetElements, OutParams);
 	}
 	else
 	{
-		FStaticMeshComponentHelper::CollectPSOPrecacheDataImpl(Component, &FInstancedStaticMeshVertexFactory::StaticType, BasePrecachePSOParams, ISMC_GetElements, OutParams);
+		FStaticMeshComponentHelper::CollectPSOPrecacheDataImpl(Component, GetComponentRenderMatrix, &FInstancedStaticMeshVertexFactory::StaticType, BasePrecachePSOParams, ISMC_GetElements, OutParams);
 	}
 }
\ No newline at end of file
diff --git a/Engine/Source/Runtime/Engine/Public/PrimitiveSceneProxy.h b/Engine/Source/Runtime/Engine/Public/PrimitiveSceneProxy.h
--- a/Engine/Source/Runtime/Engine/Public/PrimitiveSceneProxy.h
+++ b/Engine/Source/Runtime/Engine/Public/PrimitiveSceneProxy.h
@@ -314,7 +314,7 @@ public:
 	/**
 	 * Enqueue and update for the render thread to notify it that the primitive color changed.
 	 */
-	void SetPrimitiveColor_GameThread(const FLinearColor& InPrimitiveColor);
+	ENGINE_API void SetPrimitiveColor_GameThread(const FLinearColor& InPrimitiveColor);
 #endif	
 
 	/** Enqueue and update for the render thread to remove the velocity data for this component from the scene. */
diff --git a/Engine/Source/Runtime/Engine/Public/PrimitiveSceneProxyDesc.h b/Engine/Source/Runtime/Engine/Public/PrimitiveSceneProxyDesc.h
--- a/Engine/Source/Runtime/Engine/Public/PrimitiveSceneProxyDesc.h
+++ b/Engine/Source/Runtime/Engine/Public/PrimitiveSceneProxyDesc.h
@@ -78,7 +78,7 @@ struct FPrimitiveSceneProxyDesc
 
 	ENGINE_API FPrimitiveSceneProxyDesc(const UPrimitiveComponent*);
 
-	void InitializeFromPrimitiveComponent(const UPrimitiveComponent*);
+	ENGINE_API void InitializeFromPrimitiveComponent(const UPrimitiveComponent*);
 	
 	UE_DEPRECATED(5.5, "Use InitializeFromPrimitiveComponent instead.")
 	void InitializeFrom(const UPrimitiveComponent* InComponent) { InitializeFromPrimitiveComponent(InComponent); }
@@ -210,8 +210,9 @@ struct FPrimitiveSceneProxyDesc
 	bool HasValidSettingsForStaticLighting() const { return bHasValidSettingsForStaticLighting; } 
 	bool GetShadowIndirectOnly() const { return bShadowIndirectOnly; }
 	int32 GetRayTracingGroupId() const { return RayTracingGroupId; }
-	FSceneInterface* GetScene()  const { check(Scene); return Scene; }
+	FSceneInterface* GetScene()  const { return Scene; }
 	bool GetLevelInstanceEditingState() const { return bLevelInstanceEditingState; }
+	ENGINE_API bool IsStaticMeshFirstPersonRelevant() const { return bIsFirstPerson || bIsFirstPersonWorldSpaceRepresentation; }
 
 	UObject* GetOwner() const { return Owner; }	
 	template<class T>
diff --git a/Engine/Source/Runtime/Engine/Public/StaticMeshComponentHelper.h b/Engine/Source/Runtime/Engine/Public/StaticMeshComponentHelper.h
--- a/Engine/Source/Runtime/Engine/Public/StaticMeshComponentHelper.h
+++ b/Engine/Source/Runtime/Engine/Public/StaticMeshComponentHelper.h
@@ -16,16 +16,17 @@
 class FStaticMeshComponentHelper
 {
 	using GetPSOVertexElementsFn = TFunctionRef<void(const FStaticMeshLODResources& LODRenderData, int32 LODIndex, bool bSupportsManualVertexFetch, FVertexDeclarationElementList& Elements)>;
+	using GetComponentRenderMatrixFn = TFunctionRef<FMatrix()>;
 
 public:
 	template<class T>
 	static void GetUsedMaterials(const T& Component, TArray<UMaterialInterface*>& OutMaterials, bool bGetDebugMaterials);
 
 	template<class T>
-	static void CollectPSOPrecacheDataImpl(const T& Component, const FVertexFactoryType* VFType, const FPSOPrecacheParams& BasePrecachePSOParams, GetPSOVertexElementsFn GetVertexElements, FMaterialInterfacePSOPrecacheParamsList& OutParams);
+	static void CollectPSOPrecacheDataImpl(const T& Component, GetComponentRenderMatrixFn GetComponentRenderMatrix, const FVertexFactoryType* VFType, const FPSOPrecacheParams& BasePrecachePSOParams, GetPSOVertexElementsFn GetVertexElements, FMaterialInterfacePSOPrecacheParamsList& OutParams);
 
 	template<class T>
-	static void CollectPSOPrecacheData(const T& Component, const FPSOPrecacheParams& BasePrecachePSOParams, FMaterialInterfacePSOPrecacheParamsList& OutParams);
+	static void CollectPSOPrecacheData(const T& Component, GetComponentRenderMatrixFn GetComponentRenderMatrix, const FPSOPrecacheParams& BasePrecachePSOParams, FMaterialInterfacePSOPrecacheParamsList& OutParams);
 };
 
 template<class T>
@@ -53,13 +54,10 @@ void FStaticMeshComponentHelper::GetUsedMaterials(const T& Component, TArray<UMa
 }
 
 template<class T>
-void FStaticMeshComponentHelper::CollectPSOPrecacheDataImpl(const T& Component, const FVertexFactoryType* VFType, const FPSOPrecacheParams& BasePrecachePSOParams, FStaticMeshComponentHelper::GetPSOVertexElementsFn GetVertexElements, FMaterialInterfacePSOPrecacheParamsList& OutParams)
+void FStaticMeshComponentHelper::CollectPSOPrecacheDataImpl(const T& Component, GetComponentRenderMatrixFn GetComponentRenderMatrix, const FVertexFactoryType* VFType, const FPSOPrecacheParams& BasePrecachePSOParams, FStaticMeshComponentHelper::GetPSOVertexElementsFn GetVertexElements, FMaterialInterfacePSOPrecacheParamsList& OutParams)
 {
 	check(Component.GetStaticMesh() != nullptr && Component.GetStaticMesh()->GetRenderData() != nullptr);
 
-	UWorld* World = Component.GetWorld();
-	ERHIFeatureLevel::Type FeatureLevel = World ? World->GetFeatureLevel() : GMaxRHIFeatureLevel;
-
 	bool bSupportsManualVertexFetch = VFType->SupportsManualVertexFetch(GMaxRHIFeatureLevel);
 	bool bAnySectionCastsShadows = false;
 	int32 MeshMinLOD = Component.GetStaticMesh()->GetMinLODIdx();
@@ -99,7 +97,7 @@ void FStaticMeshComponentHelper::CollectPSOPrecacheDataImpl(const T& Component,
 		}
 	}
 
-	bool bIsLocalToWorldDeterminantNegative = Component.GetRenderMatrix().Determinant() < 0;
+	bool bIsLocalToWorldDeterminantNegative = GetComponentRenderMatrix().Determinant() < 0;
 
 	FPSOPrecacheParams PrecachePSOParams = BasePrecachePSOParams;
 	PrecachePSOParams.bCastShadow = bAnySectionCastsShadows;
@@ -134,7 +132,7 @@ void FStaticMeshComponentHelper::CollectPSOPrecacheDataImpl(const T& Component,
 }
 
 template<class T>
-void FStaticMeshComponentHelper::CollectPSOPrecacheData(const T& Component, const FPSOPrecacheParams& BasePrecachePSOParams, FMaterialInterfacePSOPrecacheParamsList& OutParams)
+void FStaticMeshComponentHelper::CollectPSOPrecacheData(const T& Component, GetComponentRenderMatrixFn GetComponentRenderMatrix, const FPSOPrecacheParams& BasePrecachePSOParams, FMaterialInterfacePSOPrecacheParamsList& OutParams)
 {
 	if (Component.GetStaticMesh() == nullptr || Component.GetStaticMesh()->GetRenderData() == nullptr)
 	{
@@ -156,10 +154,10 @@ void FStaticMeshComponentHelper::CollectPSOPrecacheData(const T& Component, cons
 	Nanite::FMaterialAudit NaniteMaterials{};
 	if (Nanite::FNaniteResourcesHelper::ShouldCreateNaniteProxy(Component, &NaniteMaterials))
 	{
-		FStaticMeshComponentHelper::CollectPSOPrecacheDataImpl(Component, &FNaniteVertexFactory::StaticType, BasePrecachePSOParams, SMC_GetElements, OutParams);
+		FStaticMeshComponentHelper::CollectPSOPrecacheDataImpl(Component, GetComponentRenderMatrix, &FNaniteVertexFactory::StaticType, BasePrecachePSOParams, SMC_GetElements, OutParams);
 	}
 	else
 	{
-		FStaticMeshComponentHelper::CollectPSOPrecacheDataImpl(Component, &FLocalVertexFactory::StaticType, BasePrecachePSOParams, SMC_GetElements, OutParams);
+		FStaticMeshComponentHelper::CollectPSOPrecacheDataImpl(Component, GetComponentRenderMatrix, &FLocalVertexFactory::StaticType, BasePrecachePSOParams, SMC_GetElements, OutParams);
 	}
 }
\ No newline at end of file
diff --git a/Engine/Source/Runtime/Engine/Public/StaticMeshSceneProxyDesc.h b/Engine/Source/Runtime/Engine/Public/StaticMeshSceneProxyDesc.h
--- a/Engine/Source/Runtime/Engine/Public/StaticMeshSceneProxyDesc.h
+++ b/Engine/Source/Runtime/Engine/Public/StaticMeshSceneProxyDesc.h
@@ -165,7 +165,6 @@ struct FStaticMeshSceneProxyDesc : public FPrimitiveSceneProxyDesc
 	
 	ENGINE_API bool ShouldCreateNaniteProxy(Nanite::FMaterialAudit* OutNaniteMaterials = nullptr) const;
 	ENGINE_API bool UseNaniteOverrideMaterials(bool bDoingMaterialAudit) const;
-	ENGINE_API bool IsStaticMeshFirstPersonRelevant() const { return bIsFirstPerson || bIsFirstPersonWorldSpaceRepresentation; }
 	
 	UMaterialInterface* GetNaniteAuditMaterial(int32 MaterialIndex) const
 	{
