New material translator changes:
- Added support for enabling the new translator while editing a material function.
- Some minor refactoring and renaming.

diff --git a/Engine/Source/Runtime/Engine/Private/Materials/MaterialExpressionsIR.cpp b/Engine/Source/Runtime/Engine/Private/Materials/MaterialExpressionsIR.cpp
--- a/Engine/Source/Runtime/Engine/Private/Materials/MaterialExpressionsIR.cpp
+++ b/Engine/Source/Runtime/Engine/Private/Materials/MaterialExpressionsIR.cpp
@@ -249,6 +249,59 @@ void UMaterialExpression::Build(MIR::FEmitter& Emitter)
 	Emitter.Error(TEXT("Unsupported material expression."));
 } 
 
+void UMaterialExpressionFunctionInput::Build(MIR::FEmitter& Emitter)
+{
+	MIR::FValue* OutputValue = Emitter.TryInput(&Preview);
+	if (OutputValue)
+	{
+		Emitter.Output(0, OutputValue);
+		return;
+	}
+
+	switch (InputType)
+	{
+		case FunctionInput_Scalar:
+			OutputValue = Emitter.EmitConstantFloatScalar(PreviewValue.X);
+			break;
+
+		case FunctionInput_Vector2:
+			OutputValue = Emitter.EmitConstantFloat2({ PreviewValue.X, PreviewValue.Y });
+			break;
+
+		case FunctionInput_Vector3:
+			OutputValue = Emitter.EmitConstantFloat3({ PreviewValue.X, PreviewValue.Y, PreviewValue.Z });
+			break;
+
+		case FunctionInput_Vector4:
+			OutputValue = Emitter.EmitConstantFloat4(PreviewValue);
+			break;
+
+		case FunctionInput_Bool:
+		case FunctionInput_StaticBool:
+			OutputValue = Emitter.EmitConstantBoolScalar(PreviewValue.X != 0.0f);
+
+		case FunctionInput_Texture2D:
+		case FunctionInput_TextureCube:
+		case FunctionInput_Texture2DArray:
+		case FunctionInput_VolumeTexture:
+		case FunctionInput_MaterialAttributes:
+		case FunctionInput_TextureExternal:
+		case FunctionInput_Substrate:
+			Emitter.Error(TEXT("Function input of object type requires preview input to be provided."));
+			return;
+
+		default:
+			UE_MIR_UNREACHABLE();
+	}
+
+	Emitter.Output(0, OutputValue);
+}
+
+void UMaterialExpressionFunctionOutput::Build(MIR::FEmitter& Emitter)
+{
+	Emitter.Output(0, Emitter.Input(&A));
+}
+
 void UMaterialExpressionConstant::Build(MIR::FEmitter& Emitter)
 {
 	MIR::FValue* Value = Emitter.EmitConstantFloatScalar(R);
diff --git a/Engine/Source/Runtime/Engine/Private/Materials/MaterialIREmitter.cpp b/Engine/Source/Runtime/Engine/Private/Materials/MaterialIREmitter.cpp
--- a/Engine/Source/Runtime/Engine/Private/Materials/MaterialIREmitter.cpp
+++ b/Engine/Source/Runtime/Engine/Private/Materials/MaterialIREmitter.cpp
@@ -184,7 +184,7 @@ FEmitter& FEmitter::DefaultInputToBoolScalar(const FExpressionInput* InInput, bo
 {
 	if (!InInput->IsConnected())
 	{
-		Internal::SetInputValue(BuilderImpl, InInput, EmitConstantBoolScalar(Value));
+		Internal::BindValueToExpressionInput(BuilderImpl, InInput, EmitConstantBoolScalar(Value));
 	}
 
 	return *this;
@@ -194,7 +194,7 @@ FEmitter& FEmitter::DefaultInputToFloatScalar(const FExpressionInput* InInput, T
 {
 	if (!InInput->IsConnected())
 	{
-		Internal::SetInputValue(BuilderImpl, InInput, EmitConstantFloatScalar(Value));
+		Internal::BindValueToExpressionInput(BuilderImpl, InInput, EmitConstantFloatScalar(Value));
 	}
 
 	return *this;
@@ -204,7 +204,7 @@ FEmitter& FEmitter::DefaultInputToFloat4(const FExpressionInput* InInput, FVecto
 {
 	if (!InInput->IsConnected())
 	{
-		Internal::SetInputValue(BuilderImpl, InInput, EmitConstantFloat4(Value));
+		Internal::BindValueToExpressionInput(BuilderImpl, InInput, EmitConstantFloat4(Value));
 	}
 
 	return *this;
@@ -212,7 +212,7 @@ FEmitter& FEmitter::DefaultInputToFloat4(const FExpressionInput* InInput, FVecto
 
 FValue* FEmitter::TryInput(const FExpressionInput* InInput)
 {
-	FValue* Value = MIR::Internal::GetInputValue(BuilderImpl, InInput);
+	FValue* Value = MIR::Internal::FetchValueFromExpressionInput(BuilderImpl, InInput);
 	return Value ? Value : nullptr;
 }
 
@@ -315,7 +315,7 @@ FEmitter& FEmitter::Output(int32 OutputIndex, FValue* Value)
 
 FEmitter& FEmitter::Output(const FExpressionOutput* Output, FValue* Value)
 {
-	MIR::Internal::SetOutputValue(BuilderImpl, Output, Value);
+	MIR::Internal::BindValueToExpressionOutput(BuilderImpl, Output, Value);
 	return *this;
 }
 
diff --git a/Engine/Source/Runtime/Engine/Private/Materials/MaterialIRInternal.h b/Engine/Source/Runtime/Engine/Private/Materials/MaterialIRInternal.h
--- a/Engine/Source/Runtime/Engine/Private/Materials/MaterialIRInternal.h
+++ b/Engine/Source/Runtime/Engine/Private/Materials/MaterialIRInternal.h
@@ -27,14 +27,14 @@ FValue* CreateMaterialAttributeDefaultValue(FEmitter& Emitter, const UMaterial*
 //
 EMaterialTextureParameterType TextureMaterialValueTypeToParameterType(EMaterialValueType Type);
 
-//
-FValue* GetInputValue(FMaterialIRModuleBuilderImpl* Builder, const FExpressionInput* Input);
+// Returns the value flowing into given expression input (previously set through `BindValueToExpressionInput`).
+FValue* FetchValueFromExpressionInput(FMaterialIRModuleBuilderImpl* Builder, const FExpressionInput* Input);
 
-//
-void SetInputValue(FMaterialIRModuleBuilderImpl* Builder, const FExpressionInput* Input, FValue* Value);
+// Flows a value into given expression input.
+void BindValueToExpressionInput(FMaterialIRModuleBuilderImpl* Builder, const FExpressionInput* Input, FValue* Value);
 
-//
-void SetOutputValue(FMaterialIRModuleBuilderImpl* Builder, const FExpressionOutput* Output, FValue* Value);
+// Flows a value int given expression output.
+void BindValueToExpressionOutput(FMaterialIRModuleBuilderImpl* Builder, const FExpressionOutput* Output, FValue* Value);
 
 //
 uint32 HashBytes(const char* Ptr, uint32 Size);
diff --git a/Engine/Source/Runtime/Engine/Private/Materials/MaterialIRModuleBuilder.cpp b/Engine/Source/Runtime/Engine/Private/Materials/MaterialIRModuleBuilder.cpp
--- a/Engine/Source/Runtime/Engine/Private/Materials/MaterialIRModuleBuilder.cpp
+++ b/Engine/Source/Runtime/Engine/Private/Materials/MaterialIRModuleBuilder.cpp
@@ -202,10 +202,8 @@ struct FMaterialIRModuleBuilderImpl
 			// Function calls are handled internally as they manipulate the analysis context stack.
 			PushFunctionCall(Call);
 		}
-		else if (!Cast<UMaterialExpressionFunctionOutput>(Emitter->Expression))
+		else
 		{
-			// Ignore UMaterialExpressionFunctionOutputs as they're handled by PopFunctionCall().
-
 			// Invoke the expression build function. This will perform semantic analysis, error reporting and
 			// emit IR values for its outputs (which will flow into connected expressions inputs).
 			Emitter->Expression->Build(*Emitter);
@@ -317,17 +315,15 @@ struct FMaterialIRModuleBuilderImpl
 					MIR::FValue** ValuePtr = Context.OutputValues.Find(Input.Input->GetConnectedOutput());
 					check(ValuePtr && *ValuePtr);
 
-					MIR::Internal::SetInputValue(this, Input.Input, *ValuePtr);
+					MIR::Internal::BindValueToExpressionInput(this, Input.Input, *ValuePtr);
 
 					const MIR::FType* OutputArgType = MIR::FType::FromShaderType(Input.Type);
 					Output->Arg = Emitter->EmitConstruct(OutputArgType, *ValuePtr);
 				}
 
-				if (Builder->TargetInsights)
-				{
-					check(Output->Arg);
-					PushConnectionInsight(Builder->Material, (int)Output->Property, Input.Input->Expression, Input.Input->OutputIndex, Output->Arg->Type);
-				}
+				// Push this connection insight
+				check(Output->Arg);
+				PushConnectionInsight(Builder->Material, (int)Output->Property, Input.Input->Expression, Input.Input->OutputIndex, Output->Arg->Type);
 			}
 		}
 	}
@@ -509,19 +505,21 @@ struct FMaterialIRModuleBuilderImpl
 		// Update expression inputs insight.
 		for (FExpressionInputIterator It{ Expression }; It; ++It)
 		{
-			if (!It->IsConnected())
+			if (MIR::FValue* Value = MIR::Internal::FetchValueFromExpressionInput(this, It.Input))
 			{
-				continue;
+				PushConnectionInsight(Expression, It.Index, It->Expression, It->OutputIndex, Value->Type);
 			}
-
-			MIR::FValue* Value = MIR::Internal::GetInputValue(this, It.Input);
-			PushConnectionInsight(Expression, It.Index, It->Expression, It->OutputIndex, Value ? Value->Type : nullptr);
 		}
 	}
 	
 	//
 	void PushConnectionInsight(const UObject* InputObject, int InputIndex, const UMaterialExpression* OutputExpression, int OutputIndex, const MIR::FType* Type)
 	{
+		if (!Builder->TargetInsights || Type->IsPoison())
+		{
+			return;
+		}
+
 		FMaterialInsights::FConnectionInsight Insight;
 		Insight.InputObject = InputObject,
 		Insight.OutputExpression = OutputExpression,
@@ -583,17 +581,17 @@ bool FMaterialIRModuleBuilder::Build(FMaterialIRModule* TargetModule)
 
 namespace MIR::Internal {
 
-FValue* GetInputValue(FMaterialIRModuleBuilderImpl* Builder, const FExpressionInput* Input)
+FValue* FetchValueFromExpressionInput(FMaterialIRModuleBuilderImpl* Builder, const FExpressionInput* Input)
 {
 	return Builder->AnalysisContextStack.Last().GetInputValue(Input);
 }
 
-void SetInputValue(FMaterialIRModuleBuilderImpl* Builder, const FExpressionInput* Input, FValue* Value)
+void BindValueToExpressionInput(FMaterialIRModuleBuilderImpl* Builder, const FExpressionInput* Input, FValue* Value)
 {
 	Builder->AnalysisContextStack.Last().SetInputValue(Input, Value);
 }
 
-void SetOutputValue(FMaterialIRModuleBuilderImpl* Builder, const FExpressionOutput* Output, FValue* Value)
+void BindValueToExpressionOutput(FMaterialIRModuleBuilderImpl* Builder, const FExpressionOutput* Output, FValue* Value)
 {
 	Builder->AnalysisContextStack.Last().SetOutputValue(Output, Value);
 }
diff --git a/Engine/Source/Runtime/Engine/Public/Materials/Material.h b/Engine/Source/Runtime/Engine/Public/Materials/Material.h
--- a/Engine/Source/Runtime/Engine/Public/Materials/Material.h
+++ b/Engine/Source/Runtime/Engine/Public/Materials/Material.h
@@ -972,8 +972,9 @@ public:
 	UPROPERTY(EditAnywhere, Category=Material)
 	uint8 bUseMaterialAttributes : 1;
 
-	// UPROPERTY(EditAnywhere, Category = Material)
-	// uint8 bEnableExecWire : 1; // maxtodo
+	UE_DEPRECATED(5.8, "Experimental bEnableExecWire is no longer supported.")
+	UPROPERTY(EditAnywhere, Category = Material)
+	uint8 bEnableExecWire : 1;
 
 	UPROPERTY(EditAnywhere, Category = Material)
 	uint8 bEnableNewHLSLGenerator : 1;
diff --git a/Engine/Source/Runtime/Engine/Public/Materials/MaterialExpressionFunctionInput.h b/Engine/Source/Runtime/Engine/Public/Materials/MaterialExpressionFunctionInput.h
--- a/Engine/Source/Runtime/Engine/Public/Materials/MaterialExpressionFunctionInput.h
+++ b/Engine/Source/Runtime/Engine/Public/Materials/MaterialExpressionFunctionInput.h
@@ -98,6 +98,7 @@ class UMaterialExpressionFunctionInput : public UMaterialExpression
 
 	//~ Begin UMaterialExpression Interface
 #if WITH_EDITOR
+	virtual void Build(MIR::FEmitter& Emitter) override;
 	virtual void GetExpressionToolTip(TArray<FString>& OutToolTip) override;
 	virtual int32 CompilePreview(class FMaterialCompiler* Compiler, int32 OutputIndex) override;
 	virtual int32 Compile(class FMaterialCompiler* Compiler, int32 OutputIndex) override;
diff --git a/Engine/Source/Runtime/Engine/Public/Materials/MaterialExpressionFunctionOutput.h b/Engine/Source/Runtime/Engine/Public/Materials/MaterialExpressionFunctionOutput.h
--- a/Engine/Source/Runtime/Engine/Public/Materials/MaterialExpressionFunctionOutput.h
+++ b/Engine/Source/Runtime/Engine/Public/Materials/MaterialExpressionFunctionOutput.h
@@ -54,6 +54,7 @@ class UMaterialExpressionFunctionOutput : public UMaterialExpression
 
 	//~ Begin UMaterialExpression Interface
 #if WITH_EDITOR
+	virtual void Build(MIR::FEmitter& Emitter) override;
 	virtual int32 Compile(class FMaterialCompiler* Compiler, int32 OutputIndex) override;
 	virtual int32 CompilePreview(class FMaterialCompiler* Compiler, int32 OutputIndex) override;
 	virtual void GetCaption(TArray<FString>& OutCaptions) const override;
diff --git a/Engine/Source/Runtime/Engine/Public/Materials/MaterialFunction.h b/Engine/Source/Runtime/Engine/Public/Materials/MaterialFunction.h
--- a/Engine/Source/Runtime/Engine/Public/Materials/MaterialFunction.h
+++ b/Engine/Source/Runtime/Engine/Public/Materials/MaterialFunction.h
@@ -71,8 +71,9 @@ class UMaterialFunction : public UMaterialFunctionInterface
 	uint8 bPrefixParameterNames:1;
 #endif // WITH_EDITORONLY_DATA
 
+	UE_DEPRECATED(5.8, "Experimental bEnableExecWire is no longer supported.")
 	UPROPERTY(EditAnywhere, Category = MaterialFunction)
-	uint8 bEnableExecWire : 1; // maxtodo
+	uint8 bEnableExecWire : 1;
 
 	UPROPERTY(EditAnywhere, Category = MaterialFunction)
 	uint8 bEnableNewHLSLGenerator : 1;
