Implement material slot material overlay UI. This allow artists to create meshes by merging multiple meshes that have different overlay material.

SkeletalMesh Editor:
- A CVar SkeletalMesh.FeatureFlags.UI.ShowMaterialSlotOverlayMaterial control if the UI should show the material slot overlay. If a material slot has some specified overlay and the CVar is off the UI will show the overlay for this particular material slot to ensure artist always see the correct data for a section.

StaticMesh Editor:
- A CVar StaticMesh.FeatureFlags.UI.ShowMaterialSlotOverlayMaterial control if the UI should show the material slot overlay. If a material slot has some specified overlay and the CVar is off the UI will show the overlay for this particular material slot to ensure artist always see the correct data for a section.

MeshComponent:
- The mesh component allow the user to add per material slot overlay material, it work similar to the material override array. Since the change is at the mesh component level it works for both static and skeletal meshes.

Mesh proxies:
- Both SM and SK mesh proxies was already driving the overlay material per mesh batch I simply add logic to control which overlay material a section use, it first use the material slot one if not null then use the global mesh overlay if not null.

diff --git a/Engine/Source/Developer/MaterialUtilities/Public/MaterialUtilities.h b/Engine/Source/Developer/MaterialUtilities/Public/MaterialUtilities.h
--- a/Engine/Source/Developer/MaterialUtilities/Public/MaterialUtilities.h
+++ b/Engine/Source/Developer/MaterialUtilities/Public/MaterialUtilities.h
@@ -56,7 +56,7 @@ enum class EFlattenMaterialProperties : uint8
 /** Struct to store per section info used to populate data after (multiple) meshes are merged together */
 struct FSectionInfo
 {
-	FSectionInfo() : Material(nullptr), MaterialSlotName(NAME_None), MaterialIndex(INDEX_NONE), StartIndex(INDEX_NONE), EndIndex(INDEX_NONE), bProcessed(false) {}
+	FSectionInfo() : Material(nullptr), MaterialSlotName(NAME_None), MaterialIndex(INDEX_NONE), StartIndex(INDEX_NONE), EndIndex(INDEX_NONE), bProcessed(false), OverlayMaterial(nullptr) {}
 
 	/** Material used by the section */
 	class UMaterialInterface* Material;
@@ -72,6 +72,8 @@ struct FSectionInfo
 	int32 EndIndex;
 	/** Used while baking out materials, to check which sections are and aren't being baked out */
 	bool bProcessed;
+	/** Overlay material used by the section */
+	class UMaterialInterface* OverlayMaterial;
 
 	bool operator==(const FSectionInfo& Other) const
 	{
diff --git a/Engine/Source/Editor/Persona/Private/PersonaMeshDetails.cpp b/Engine/Source/Editor/Persona/Private/PersonaMeshDetails.cpp
--- a/Engine/Source/Editor/Persona/Private/PersonaMeshDetails.cpp
+++ b/Engine/Source/Editor/Persona/Private/PersonaMeshDetails.cpp
@@ -96,6 +96,16 @@
 DEFINE_LOG_CATEGORY(LogSkeletalMeshPersonaMeshDetail);
 
 #define SUFFIXE_DEFAULT_MATERIAL TEXT(" (Default)")
+
+
+static bool GEnableSKeletalMeshMaterialSlotOverlayMaterialUI = false;
+static FAutoConsoleVariableRef CCvarEnableSKeletalMeshMaterialSlotOverlayMaterialUI(
+	TEXT("SkeletalMesh.FeatureFlags.UI.ShowMaterialSlotOverlayMaterial"),
+	GEnableSKeletalMeshMaterialSlotOverlayMaterialUI,
+	TEXT("Whether skeletal mesh editor show material overlay per material slot UI."),
+	ECVF_Default);
+
+
 /*
 * Custom data key
 */
@@ -2937,6 +2947,7 @@ void FPersonaMeshDetails::AddLODLevelCategories(IDetailLayoutBuilder& DetailLayo
 				MaterialListDelegates.OnMaterialChanged.BindSP(this, &FPersonaMeshDetails::OnMaterialArrayChanged, 0);
 				MaterialListDelegates.OnGenerateCustomNameWidgets.BindSP(this, &FPersonaMeshDetails::OnGenerateCustomNameWidgetsForMaterialArray);
 				MaterialListDelegates.OnGenerateCustomMaterialWidgets.BindSP(this, &FPersonaMeshDetails::OnGenerateCustomMaterialWidgetsForMaterialArray, 0);
+				MaterialListDelegates.OnGenerateExtraBottomWidget.BindSP(this, &FPersonaMeshDetails::OnGenerateMaterialListExtraBottomWidget, 0);
 				MaterialListDelegates.OnMaterialListDirty.BindSP(this, &FPersonaMeshDetails::OnMaterialListDirty);
 
 				MaterialListDelegates.OnCopyMaterialItem.BindSP(this, &FPersonaMeshDetails::OnCopyMaterialItem);
@@ -5260,6 +5271,68 @@ TSharedRef<SWidget> FPersonaMeshDetails::OnGenerateCustomMaterialWidgetsForMater
 		.ToolTipText(this, &FPersonaMeshDetails::GetOriginalImportMaterialNameText, MaterialIndex);
 }
 
+TSharedRef<SWidget> FPersonaMeshDetails::OnGenerateMaterialListExtraBottomWidget(UMaterialInterface* Material, int32 MaterialIndex, int32 LODIndex)
+{
+	if (!SkeletalMeshPtr.IsValid() || !SkeletalMeshPtr->GetMaterials().IsValidIndex(MaterialIndex))
+	{
+		return SNullWidget::NullWidget;
+	}
+	TSharedRef<SVerticalBox> VerticalWidget = SNew(SVerticalBox);
+	const FSkeletalMaterial& SkeletalMaterial = SkeletalMeshPtr->GetMaterials()[MaterialIndex];
+	if(GEnableSKeletalMeshMaterialSlotOverlayMaterialUI || SkeletalMaterial.OverlayMaterialInterface)
+	{
+		VerticalWidget->AddSlot()
+		.AutoHeight()
+		.Padding(0, 2, 0, 0)
+		[
+			SNew(SHorizontalBox)
+			+ SHorizontalBox::Slot()
+			.AutoWidth()
+			.VAlign(VAlign_Center)
+			[
+				SNew(SBox)
+				.HAlign(HAlign_Right)
+				[
+					SNew(STextBlock)
+					.Font(IDetailLayoutBuilder::GetDetailFont())
+					.Text(LOCTEXT("OverlayMaterial", "Overlay Material"))
+				]
+			]
+			+ SHorizontalBox::Slot()
+			.FillWidth(1.0f)
+			.Padding(5, 0, 0, 0)
+			[
+				SNew(SObjectPropertyEntryBox)
+				.ObjectPath_Lambda([this, MaterialIndex]()
+					{
+						UMaterialInterface* OverlayMaterial = nullptr;
+						if (SkeletalMeshPtr.IsValid() && SkeletalMeshPtr->GetMaterials().IsValidIndex(MaterialIndex))
+						{
+							OverlayMaterial = SkeletalMeshPtr->GetMaterials()[MaterialIndex].OverlayMaterialInterface;
+						}
+						if (OverlayMaterial)
+						{
+							return OverlayMaterial->GetPathName();
+						}
+						return FString();
+					})
+				.AllowClear(true)
+				.AllowedClass(UMaterialInterface::StaticClass())
+				.OnObjectChanged_Lambda([this, MaterialIndex](const FAssetData& AssetData)
+					{
+					
+						UMaterialInterface* MaterialInterface = AssetData.IsValid() ? Cast<UMaterialInterface>(AssetData.GetAsset()) : nullptr;
+						OnMaterialSlotOverlayMaterialChanged(MaterialInterface, MaterialIndex);
+					})
+				.ThumbnailPool(UThumbnailManager::Get().GetSharedThumbnailPool())
+				.DisplayCompactSize(true)
+			]
+		];
+	}
+
+	return VerticalWidget;
+}
+
 FText FPersonaMeshDetails::GetFirstMaterialSlotUsedBySection(int32 MaterialIndex) const
 {
 	if (SkeletalMeshPtr.IsValid() && !SkeletalMeshPtr->IsCompiling() && MaterialUsedMap.Contains(MaterialIndex))
@@ -5429,6 +5502,37 @@ bool FPersonaMeshDetails::OnMaterialListDirty()
 	return ForceMaterialListRefresh;
 }
 
+void FPersonaMeshDetails::OnMaterialSlotOverlayMaterialChanged(UMaterialInterface* NewOverlayMaterial, int32 MaterialIndex)
+{
+	USkeletalMesh* Mesh = GetPersonaToolkit()->GetMesh();
+	if (Mesh == nullptr)
+		return;
+
+	if (!ensure(Mesh->GetMaterials().IsValidIndex(MaterialIndex)))
+	{
+		return;
+	}
+
+	FSkeletalMaterial& SkeletalMaterial = Mesh->GetMaterials()[MaterialIndex];
+
+	if (SkeletalMaterial.OverlayMaterialInterface == NewOverlayMaterial)
+	{
+		//Nothing change
+		return;
+	}
+
+	//Make a transactional change
+	{
+		const FScopedTransaction Transaction(LOCTEXT("PersonaSetSectionOverlayMaterial", "Persona editor: Set Overlay Material For Material Slot"));
+		Mesh->Modify();
+		FScopedSuspendAlternateSkinWeightPreview ScopedSuspendAlternateSkinnWeightPreview(Mesh);
+		{
+			FScopedSkeletalMeshPostEditChange ScopedPostEditChange(Mesh);
+			SkeletalMaterial.OverlayMaterialInterface = NewOverlayMaterial;
+		}
+	}
+}
+
 TSharedRef<SWidget> FPersonaMeshDetails::OnGenerateCustomNameWidgetsForSection(int32 LodIndex, int32 SectionIndex)
 {
 	bool IsSectionChunked = false;
diff --git a/Engine/Source/Editor/Persona/Private/PersonaMeshDetails.h b/Engine/Source/Editor/Persona/Private/PersonaMeshDetails.h
--- a/Engine/Source/Editor/Persona/Private/PersonaMeshDetails.h
+++ b/Engine/Source/Editor/Persona/Private/PersonaMeshDetails.h
@@ -475,6 +475,15 @@ private:
 	*/
 	TSharedRef<SWidget> OnGenerateCustomMaterialWidgetsForMaterialArray(UMaterialInterface* Material, int32 MaterialIndex, int32 LODIndex);
 
+	/**
+	* Called by the material list widget on generating each thumbnail widget
+	* Those generated widget will be stack under the material list row
+	*
+	* @param Material		The material that is being displayed
+	* @param MaterialIndex	The index of the material slot
+	*/
+	TSharedRef<SWidget> OnGenerateMaterialListExtraBottomWidget(UMaterialInterface* Material, int32 MaterialIndex, int32 LODIndex);
+
 	/* If the material list is dirty this function will return true */
 	bool OnMaterialListDirty();
 
@@ -484,6 +493,9 @@ private:
 
 	TSharedRef<SWidget> OnGetMaterialSlotUsedByMenuContent(int32 MaterialIndex);
 
+	/* Call when user change the material slot overlay material. */
+	void OnMaterialSlotOverlayMaterialChanged(UMaterialInterface* NewOverlayMaterial, int32 MaterialIndex);
+
 	FText GetFirstMaterialSlotUsedBySection(int32 MaterialIndex) const;
 
 	/**
diff --git a/Engine/Source/Editor/PropertyEditor/Private/MaterialList.cpp b/Engine/Source/Editor/PropertyEditor/Private/MaterialList.cpp
--- a/Engine/Source/Editor/PropertyEditor/Private/MaterialList.cpp
+++ b/Engine/Source/Editor/PropertyEditor/Private/MaterialList.cpp
@@ -110,11 +110,12 @@ TSharedRef<FMaterialItemView> FMaterialItemView::Create(
 	FOnGenerateWidgetsForMaterial InOnGenerateNameWidgetsForMaterial, 
 	FOnGenerateWidgetsForMaterial InOnGenerateWidgetsForMaterial, 
 	FOnResetMaterialToDefaultClicked InOnResetToDefaultClicked,
+	FOnGenerateWidgetsForMaterial InOnGenerateExtraBottomWidget,
 	int32 InMultipleMaterialCount,
 	bool bShowUsedTextures)
 {
 	// FMaterialItemView has private constructor that is why we need to use MakeShareable 
-	return MakeShareable( new FMaterialItemView( Material, InOnMaterialChanged, InOnGenerateNameWidgetsForMaterial, InOnGenerateWidgetsForMaterial, InOnResetToDefaultClicked, InMultipleMaterialCount, bShowUsedTextures) );
+	return MakeShareable( new FMaterialItemView( Material, InOnMaterialChanged, InOnGenerateNameWidgetsForMaterial, InOnGenerateWidgetsForMaterial, InOnResetToDefaultClicked, InOnGenerateExtraBottomWidget, InMultipleMaterialCount, bShowUsedTextures) );
 }
 
 TSharedRef<SWidget> FMaterialItemView::CreateNameContent()
@@ -248,6 +249,11 @@ TSharedRef<SWidget> FMaterialItemView::CreateValueContent(IDetailLayoutBuilder&
 				GetGlobalRowExtensionWidget(InDetailBuilder, InActorComponent).ToSharedRef()
 			]
 		]
+		+ SVerticalBox::Slot()
+		.AutoHeight()
+		[
+			OnGenerateExtraBottomWidget.IsBound() ? OnGenerateExtraBottomWidget.Execute(MaterialItem.Material.Get(), MaterialItem.SlotIndex) : SNullWidget::NullWidget
+		]
 		+SVerticalBox::Slot()
 		.AutoHeight()
 		[
@@ -278,6 +284,7 @@ FMaterialItemView::FMaterialItemView(	const FMaterialListItem& InMaterial,
 					FOnGenerateWidgetsForMaterial& InOnGenerateNameWidgets, 
 					FOnGenerateWidgetsForMaterial& InOnGenerateMaterialWidgets, 
 					FOnResetMaterialToDefaultClicked& InOnResetToDefaultClicked,
+					FOnGenerateWidgetsForMaterial& InOnGenerateExtraBottomWidget,
 					int32 InMultipleMaterialCount,
 					bool bInShowUsedTextures)
 					
@@ -286,6 +293,7 @@ FMaterialItemView::FMaterialItemView(	const FMaterialListItem& InMaterial,
 	, OnGenerateCustomNameWidgets( InOnGenerateNameWidgets )
 	, OnGenerateCustomMaterialWidgets( InOnGenerateMaterialWidgets )
 	, OnResetToDefaultClicked( InOnResetToDefaultClicked )
+	, OnGenerateExtraBottomWidget (InOnGenerateExtraBottomWidget)
 	, MultipleMaterialCount( InMultipleMaterialCount )
 	, bShowUsedTextures( bInShowUsedTextures )
 {
@@ -744,7 +752,14 @@ void FMaterialList::AddMaterialItem( FDetailWidgetRow& Row, int32 CurrentSlot, c
 {
 	uint32 NumMaterials = MaterialListBuilder->GetNumMaterialsInSlot(CurrentSlot);
 
-	TSharedRef<FMaterialItemView> NewView = FMaterialItemView::Create( Item, MaterialListDelegates.OnMaterialChanged, MaterialListDelegates.OnGenerateCustomNameWidgets, MaterialListDelegates.OnGenerateCustomMaterialWidgets, MaterialListDelegates.OnResetMaterialToDefaultClicked, NumMaterials, bShowUsedTextures);
+	TSharedRef<FMaterialItemView> NewView = FMaterialItemView::Create( Item
+		, MaterialListDelegates.OnMaterialChanged
+		, MaterialListDelegates.OnGenerateCustomNameWidgets
+		, MaterialListDelegates.OnGenerateCustomMaterialWidgets
+		, MaterialListDelegates.OnResetMaterialToDefaultClicked
+		, MaterialListDelegates.OnGenerateExtraBottomWidget
+		, NumMaterials
+		, bShowUsedTextures);
 
 	TSharedPtr<SWidget> RightSideContent;
 	if( bDisplayLink )
diff --git a/Engine/Source/Editor/PropertyEditor/Public/MaterialList.h b/Engine/Source/Editor/PropertyEditor/Public/MaterialList.h
--- a/Engine/Source/Editor/PropertyEditor/Public/MaterialList.h
+++ b/Engine/Source/Editor/PropertyEditor/Public/MaterialList.h
@@ -46,6 +46,7 @@ struct FMaterialListDelegates
 		, OnGenerateCustomNameWidgets()
 		, OnGenerateCustomMaterialWidgets()
 		, OnResetMaterialToDefaultClicked()
+		, OnGenerateExtraBottomWidget()
 	{}
 
 	/** Delegate called to populate the list with materials */
@@ -58,6 +59,8 @@ struct FMaterialListDelegates
 	FOnGenerateWidgetsForMaterial OnGenerateCustomMaterialWidgets;
 	/** Delegate called when a material list item should be reset to default */
 	FOnResetMaterialToDefaultClicked OnResetMaterialToDefaultClicked;
+	/** Delegate called to generate extra widget under the material row */
+	FOnGenerateWidgetsForMaterial OnGenerateExtraBottomWidget;
 	/** Delegate called when we tick the material list to know if the list is dirty*/
 	FOnMaterialListDirty OnMaterialListDirty;
 
@@ -164,6 +167,7 @@ public:
 		FOnGenerateWidgetsForMaterial InOnGenerateNameWidgetsForMaterial, 
 		FOnGenerateWidgetsForMaterial InOnGenerateWidgetsForMaterial, 
 		FOnResetMaterialToDefaultClicked InOnResetToDefaultClicked,
+		FOnGenerateWidgetsForMaterial InOnGenerateExtraBottomWidget,
 		int32 InMultipleMaterialCount,
 		bool bShowUsedTextures);
 
@@ -206,6 +210,7 @@ private:
 					FOnGenerateWidgetsForMaterial& InOnGenerateNameWidgets, 
 					FOnGenerateWidgetsForMaterial& InOnGenerateMaterialWidgets, 
 					FOnResetMaterialToDefaultClicked& InOnResetToDefaultClicked,
+					FOnGenerateWidgetsForMaterial& OnGenerateExtraBottomWidget,
 					int32 InMultipleMaterialCount,
 					bool bInShowUsedTextures);
 
@@ -279,6 +284,9 @@ private:
 	/** Reset to default delegate */
 	FOnResetMaterialToDefaultClicked OnResetToDefaultClicked;
 
+	/** Generate custom widget under the material row delegate */
+	FOnGenerateWidgetsForMaterial OnGenerateExtraBottomWidget;
+
 	/** Number of material count */
 	int32 MultipleMaterialCount;
 
diff --git a/Engine/Source/Editor/SkeletalMeshEditor/Private/SkeletalMeshEditorSubsystem.cpp b/Engine/Source/Editor/SkeletalMeshEditor/Private/SkeletalMeshEditorSubsystem.cpp
--- a/Engine/Source/Editor/SkeletalMeshEditor/Private/SkeletalMeshEditorSubsystem.cpp
+++ b/Engine/Source/Editor/SkeletalMeshEditor/Private/SkeletalMeshEditorSubsystem.cpp
@@ -465,6 +465,95 @@ bool USkeletalMeshEditorSubsystem::SetSectionVisibleInRayTracing(USkeletalMesh*
 	return true;
 }
 
+UMaterialInterface* USkeletalMeshEditorSubsystem::GetMaterialSlotOverlayMaterial(const USkeletalMesh* SkeletalMesh, const int32 SlotIndex)
+{
+	TGuardValue<bool> UnattendedScriptGuard(GIsRunningUnattendedScript, true);
+
+	if (!EditorScriptingHelpers::CheckIfInEditorAndPIE())
+	{
+		return nullptr;
+	}
+
+	if (SkeletalMesh == nullptr)
+	{
+		UE_LOG(LogSkeletalMeshEditorSubsystem, Error, TEXT("GetMaterialSlotOverlayMaterial: The SkeletalMesh is null."));
+		return nullptr;
+	}
+
+	const TArray<FSkeletalMaterial>& SkeletalMaterials = SkeletalMesh->GetMaterials();
+	if (!SkeletalMaterials.IsValidIndex(SlotIndex))
+	{
+		UE_LOG(LogSkeletalMeshEditorSubsystem, Error, TEXT("GetMaterialSlotOverlayMaterial: Invalid skeletal mesh material slot index."));
+		return nullptr;
+	}
+	return SkeletalMaterials[SlotIndex].OverlayMaterialInterface;
+}
+
+bool USkeletalMeshEditorSubsystem::SetMaterialSlotOverlayMaterial(USkeletalMesh* SkeletalMesh, const int32 SlotIndex, UMaterialInterface* NewSectionOverlayMaterial)
+{
+	TGuardValue<bool> UnattendedScriptGuard(GIsRunningUnattendedScript, true);
+
+	if (!EditorScriptingHelpers::CheckIfInEditorAndPIE())
+	{
+		return false;
+	}
+
+	if (SkeletalMesh == nullptr)
+	{
+		UE_LOG(LogSkeletalMeshEditorSubsystem, Error, TEXT("GetMaterialSlotOverlayMaterial: The SkeletalMesh is null."));
+		return false;
+	}
+
+	TArray<FSkeletalMaterial>& SkeletalMaterials = SkeletalMesh->GetMaterials();
+	if (!SkeletalMaterials.IsValidIndex(SlotIndex))
+	{
+		UE_LOG(LogSkeletalMeshEditorSubsystem, Error, TEXT("GetMaterialSlotOverlayMaterial: Invalid skeletal mesh material slot index."));
+		return false;
+	}
+
+	//Make the change as a transaction and make sure the render data is rebuilt properly
+	{
+		FScopedTransaction Transaction(LOCTEXT("SetSectionOverlayMaterialTransactionName", "Set skeletal mesh section overlay material"));
+		SkeletalMesh->Modify();
+		SkeletalMaterials[SlotIndex] = NewSectionOverlayMaterial;
+	}
+
+	SkeletalMesh->PostEditChange();
+	//This will make the sk build synchronous
+	UMaterialInterface* MaterialValue = GetMaterialSlotOverlayMaterial(SkeletalMesh, SlotIndex);
+	ensure(MaterialValue == NewSectionOverlayMaterial);
+	return true;
+}
+
+UMaterialInterface* USkeletalMeshEditorSubsystem::GetSkeletalMeshOverlayMaterial(const USkeletalMesh* SkeletalMesh)
+{
+	if (!SkeletalMesh)
+	{
+		return nullptr;
+	}
+	return SkeletalMesh->GetOverlayMaterial();
+}
+
+bool USkeletalMeshEditorSubsystem::SetSkeletalMeshOverlayMaterial(USkeletalMesh* SkeletalMesh, UMaterialInterface* NewOverlayMaterial)
+{
+	if (!SkeletalMesh)
+	{
+		return false;
+	}
+
+	//Transaction scope
+	{
+		FScopedTransaction Transaction(LOCTEXT("SetSkeletalMeshOverlayMaterialTransactionName", "Set skeletal mesh overlay material"));
+		SkeletalMesh->Modify();
+		SkeletalMesh->SetOverlayMaterial(NewOverlayMaterial);
+	}
+	SkeletalMesh->PostEditChange();
+
+	UMaterialInterface* MaterialValue = GetSkeletalMeshOverlayMaterial(SkeletalMesh);
+	ensure(MaterialValue == NewOverlayMaterial);
+	return true;
+}
+
 int32 USkeletalMeshEditorSubsystem::GetLODMaterialSlot( USkeletalMesh* SkeletalMesh, int32 LODIndex, int32 SectionIndex )
 {
 	TGuardValue<bool> UnattendedScriptGuard( GIsRunningUnattendedScript, true );
diff --git a/Engine/Source/Editor/SkeletalMeshEditor/Public/SkeletalMeshEditorSubsystem.h b/Engine/Source/Editor/SkeletalMeshEditor/Public/SkeletalMeshEditorSubsystem.h
--- a/Engine/Source/Editor/SkeletalMeshEditor/Public/SkeletalMeshEditorSubsystem.h
+++ b/Engine/Source/Editor/SkeletalMeshEditor/Public/SkeletalMeshEditorSubsystem.h
@@ -11,7 +11,9 @@
 
 DECLARE_LOG_CATEGORY_EXTERN(LogSkeletalMeshEditorSubsystem, Log, All);
 
+class USkeletalMesh;
 class USkeleton;
+class UMaterialInterface;
 
 /**
  * Blueprint enum to query skeleton curves
@@ -158,6 +160,42 @@ public:
 	UFUNCTION(BlueprintCallable, Category = "Skeletal Mesh Utilities")
 	bool SetSectionVisibleInRayTracing(USkeletalMesh* SkeletalMesh, const int32 LODIndex, const int32 SectionIndex, const bool bVisibleInRayTracing);
 
+	/** Get the overlay material for a material slot in a Skeletal Mesh.
+	 *
+	 * @param SkeletalMesh			Mesh to get material overlay from.
+	 * @param SlotIndex				Index of the skeletal mesh material slot.
+	 * @return the material interface or null if there is no overlay material.
+	 */
+	UFUNCTION(BlueprintPure, Category = "Skeletal Mesh Utilities")
+	UMaterialInterface* GetMaterialSlotOverlayMaterial(const USkeletalMesh* SkeletalMesh, const int32 SlotIndex);
+
+	/** Set the overlay material for a material slot in a Skeletal Mesh.
+	 *
+	 * @param SkeletalMesh			Mesh to set material overlay to.
+	 * @param SlotIndex				Index of the skeletal mesh material slot.
+	 * @param NewSectionOverlayMaterial	The function will set the OverlayMaterial used by the slot
+	 * @return false if invalid mesh or material slot index is invalid.
+	 */
+	UFUNCTION(BlueprintCallable, Category = "Skeletal Mesh Utilities")
+	bool SetMaterialSlotOverlayMaterial(USkeletalMesh* SkeletalMesh, const int32 SlotIndex, UMaterialInterface* NewSectionOverlayMaterial);
+
+	/** Get the overlay material for a Skeletal Mesh.
+	 *
+	 * @param SkeletalMesh			Mesh to get the material overlay from.
+	 * @return the material interface or null if there is no material.
+	 */
+	UFUNCTION(BlueprintPure, Category = "Skeletal Mesh Utilities")
+	UMaterialInterface* GetSkeletalMeshOverlayMaterial(const USkeletalMesh* SkeletalMesh);
+
+	/** Set the overlay material for a Skeletal Mesh.
+	 *
+	 * @param SkeletalMesh					Mesh to set material overlay to.
+	 * @param NewOverlayMaterial		the OverlayMaterial we set for the skeletal mesh
+	 * @return false if invalid mesh.
+	 */
+	UFUNCTION(BlueprintCallable, Category = "Skeletal Mesh Utilities")
+	bool SetSkeletalMeshOverlayMaterial(USkeletalMesh* SkeletalMesh, UMaterialInterface* NewOverlayMaterial);
+
 	/**
 	 * Gets the material slot used for a specific LOD section.
 	 * @param	SkeletalMesh		SkeletalMesh to get the material index from.
diff --git a/Engine/Source/Editor/StaticMeshEditor/Private/StaticMeshEditorTools.cpp b/Engine/Source/Editor/StaticMeshEditor/Private/StaticMeshEditorTools.cpp
--- a/Engine/Source/Editor/StaticMeshEditor/Private/StaticMeshEditorTools.cpp
+++ b/Engine/Source/Editor/StaticMeshEditor/Private/StaticMeshEditorTools.cpp
@@ -45,6 +45,7 @@
 #include "JsonObjectConverter.h"
 #include "Interfaces/IAnalyticsProvider.h"
 #include "ScopedTransaction.h"
+#include "ThumbnailRendering/ThumbnailManager.h"
 #include "UObject/UObjectGlobals.h"
 #include "Widgets/Input/SFilePathPicker.h"
 #include "Widgets/Input/STextComboBox.h"
@@ -80,6 +81,13 @@ static FAutoConsoleVariableRef CVarIsNaniteStaticMeshSettingsInitiallyCollapsed(
 	ECVF_ReadOnly
 );
 
+static bool GEnableStaticMeshMaterialSlotOverlayMaterialUI = false;
+static FAutoConsoleVariableRef CCvarEnableStaticMeshMaterialSlotOverlayMaterialUI(
+	TEXT("StaticMesh.FeatureFlags.UI.ShowMaterialSlotOverlayMaterial"),
+	GEnableStaticMeshMaterialSlotOverlayMaterialUI,
+	TEXT("Whether Static mesh editor show material overlay per material slot UI."),
+	ECVF_Default);
+
 /**
  * Window for Nanite settings.
  */
@@ -2778,6 +2786,7 @@ void FMeshMaterialsLayout::AddToCategory(IDetailCategoryBuilder& CategoryBuilder
 	MaterialListDelegates.OnMaterialChanged.BindSP(this, &FMeshMaterialsLayout::OnMaterialChanged);
 	MaterialListDelegates.OnGenerateCustomMaterialWidgets.BindSP(this, &FMeshMaterialsLayout::OnGenerateWidgetsForMaterial);
 	MaterialListDelegates.OnGenerateCustomNameWidgets.BindSP(this, &FMeshMaterialsLayout::OnGenerateNameWidgetsForMaterial);
+	MaterialListDelegates.OnGenerateExtraBottomWidget.BindSP(this, &FMeshMaterialsLayout::OnGenerateMaterialListExtraBottomWidget);
 	MaterialListDelegates.OnMaterialListDirty.BindSP(this, &FMeshMaterialsLayout::OnMaterialListDirty);
 	MaterialListDelegates.OnResetMaterialToDefaultClicked.BindSP(this, &FMeshMaterialsLayout::OnResetMaterialToDefaultClicked);
 
@@ -3023,6 +3032,100 @@ TSharedRef<SWidget> FMeshMaterialsLayout::OnGenerateNameWidgetsForMaterial(UMate
 		];
 }
 
+TSharedRef<SWidget> FMeshMaterialsLayout::OnGenerateMaterialListExtraBottomWidget(UMaterialInterface* Material, int32 SlotIndex)
+{
+	UStaticMesh& StaticMesh = GetStaticMesh();
+	if (!StaticMesh.GetStaticMaterials().IsValidIndex(SlotIndex))
+	{
+		return SNullWidget::NullWidget;
+	}
+	TSharedRef<SVerticalBox> VerticalWidget = SNew(SVerticalBox);
+	const FStaticMaterial& StaticlMaterial = StaticMesh.GetStaticMaterials()[SlotIndex];
+	if (GEnableStaticMeshMaterialSlotOverlayMaterialUI || StaticlMaterial.OverlayMaterialInterface)
+	{
+		VerticalWidget->AddSlot()
+			.AutoHeight()
+			.Padding(0, 2, 0, 0)
+			[
+				SNew(SHorizontalBox)
+					+ SHorizontalBox::Slot()
+					.AutoWidth()
+					.VAlign(VAlign_Center)
+					[
+						SNew(SBox)
+							.HAlign(HAlign_Right)
+							[
+								SNew(STextBlock)
+									.Font(IDetailLayoutBuilder::GetDetailFont())
+									.Text(LOCTEXT("OverlayMaterial", "Overlay Material"))
+							]
+					]
+					+ SHorizontalBox::Slot()
+					.FillWidth(1.0f)
+					.Padding(5, 0, 0, 0)
+					[
+						SNew(SObjectPropertyEntryBox)
+							.ObjectPath_Lambda([this, SlotIndex]()
+								{
+									UStaticMesh& StaticMesh = GetStaticMesh();
+									UMaterialInterface* OverlayMaterial = nullptr;
+									if (StaticMesh.GetStaticMaterials().IsValidIndex(SlotIndex))
+									{
+										OverlayMaterial = StaticMesh.GetStaticMaterials()[SlotIndex].OverlayMaterialInterface;
+									}
+									if (OverlayMaterial)
+									{
+										return OverlayMaterial->GetPathName();
+									}
+									return FString();
+								})
+							.AllowClear(true)
+							.AllowedClass(UMaterialInterface::StaticClass())
+							.OnObjectChanged_Lambda([this, SlotIndex](const FAssetData& AssetData)
+								{
+									UMaterialInterface* MaterialInterface = AssetData.IsValid() ? Cast<UMaterialInterface>(AssetData.GetAsset()) : nullptr;
+									OnMaterialSlotOverlayMaterialChanged(MaterialInterface, SlotIndex);
+								})
+							.ThumbnailPool(UThumbnailManager::Get().GetSharedThumbnailPool())
+							.DisplayCompactSize(true)
+					]
+			];
+	}
+
+	return VerticalWidget;
+}
+
+void FMeshMaterialsLayout::OnMaterialSlotOverlayMaterialChanged(UMaterialInterface* NewOverlayMaterial, int32 SlotIndex)
+{
+	UStaticMesh& StaticMesh = GetStaticMesh();
+	if (!StaticMesh.GetStaticMaterials().IsValidIndex(SlotIndex))
+	{
+		return;
+	}
+	FStaticMaterial& StaticlMaterial = StaticMesh.GetStaticMaterials()[SlotIndex];
+
+	if (StaticlMaterial.OverlayMaterialInterface == NewOverlayMaterial)
+	{
+		//Nothing change
+		return;
+	}
+
+	//Make a transactional change
+	{
+		FScopedTransaction ScopeTransaction(LOCTEXT("StaticMeshSetMaterialSlotOverlayMaterialTransaction", "Staticmesh editor: Set Overlay Material For Material Slot"));
+
+		FProperty* ChangedProperty = NULL;
+		ChangedProperty = FindFProperty<FProperty>(UStaticMesh::StaticClass(), "StaticMaterials");
+		check(ChangedProperty);
+		StaticMesh.PreEditChange(ChangedProperty);
+
+		StaticlMaterial.OverlayMaterialInterface = NewOverlayMaterial;
+
+		FPropertyChangedEvent PropertyUpdateStruct(ChangedProperty);
+		StaticMesh.PostEditChangeProperty(PropertyUpdateStruct);
+	}
+}
+
 ECheckBoxState FMeshMaterialsLayout::IsMaterialHighlighted(int32 SlotIndex) const
 {
 	ECheckBoxState State = ECheckBoxState::Unchecked;
diff --git a/Engine/Source/Editor/StaticMeshEditor/Private/StaticMeshEditorTools.h b/Engine/Source/Editor/StaticMeshEditor/Private/StaticMeshEditorTools.h
--- a/Engine/Source/Editor/StaticMeshEditor/Private/StaticMeshEditorTools.h
+++ b/Engine/Source/Editor/StaticMeshEditor/Private/StaticMeshEditorTools.h
@@ -476,6 +476,19 @@ private:
 	void OnMaterialChanged(UMaterialInterface* NewMaterial, UMaterialInterface* PrevMaterial, int32 SlotIndex, bool bReplaceAll);
 	TSharedRef<SWidget> OnGenerateWidgetsForMaterial(UMaterialInterface* Material, int32 SlotIndex);
 	TSharedRef<SWidget> OnGenerateNameWidgetsForMaterial(UMaterialInterface* Material, int32 SlotIndex);
+	
+	/**
+	* Called by the material list widget on generating each thumbnail widget
+	* Those generated widget will be stack under the material list row
+	*
+	* @param Material		The material that is being displayed
+	* @param SlotIndex		The index of the material slot
+	*/
+	TSharedRef<SWidget> OnGenerateMaterialListExtraBottomWidget(UMaterialInterface* Material, int32 SlotIndex);
+	
+	/* Call when user change the material slot overlay material. */
+	void OnMaterialSlotOverlayMaterialChanged(UMaterialInterface* NewOverlayMaterial, int32 MaterialIndex);
+
 	void OnResetMaterialToDefaultClicked(UMaterialInterface* Material, int32 SlotIndex);
 
 	ECheckBoxState IsMaterialHighlighted(int32 SlotIndex) const;
diff --git a/Engine/Source/Runtime/Core/Public/UObject/FortniteMainBranchObjectVersions.inl b/Engine/Source/Runtime/Core/Public/UObject/FortniteMainBranchObjectVersions.inl
--- a/Engine/Source/Runtime/Core/Public/UObject/FortniteMainBranchObjectVersions.inl
+++ b/Engine/Source/Runtime/Core/Public/UObject/FortniteMainBranchObjectVersions.inl
@@ -598,5 +598,8 @@ FORTNITE_MAIN_VERSION(WaterBodyPhysicalMaterialPropertyRemoval, 194)
 // PCG fixed attribute set -> point conversion passing through empty point data as-is and violating output pin type.
 FORTNITE_MAIN_VERSION(PCGAttributeSetToPointAlwaysConverts, 195)
 
+//Add per material slot overlay material data
+FORTNITE_MAIN_VERSION(MeshMaterialSlotOverlayMaterialAdded, 196)
+
 // -----<new versions can be added above this line>-------------------------------------------------
 #undef FORTNITE_MAIN_VERSION
\ No newline at end of file
diff --git a/Engine/Source/Runtime/Engine/Classes/Components/MeshComponent.h b/Engine/Source/Runtime/Engine/Classes/Components/MeshComponent.h
--- a/Engine/Source/Runtime/Engine/Classes/Components/MeshComponent.h
+++ b/Engine/Source/Runtime/Engine/Classes/Components/MeshComponent.h
@@ -8,6 +8,7 @@
 #include "Engine/TextureStreamingTypes.h"
 #include "Components/PrimitiveComponent.h"
 #include "Containers/SortedMap.h"
+
 #include "MeshComponent.generated.h"
 
 class UMaterialInterface;
@@ -62,6 +63,23 @@ public:
 	UFUNCTION(BlueprintCallable, Category="Rendering|Material")
 	ENGINE_API void SetOverlayMaterialMaxDrawDistance(float InMaxDrawDistance);
 
+	/**
+	 * Translucent material to blend on top of this mesh. Mesh will be rendered twice - once with a base material and once with overlay material.
+	 * The difference with the  global OverlayMaterial is those are per material slot, if the entry is null or doesn't exist the global
+	 * OverlayMaterial will be use for sections using the material slot.
+	 */
+	UPROPERTY(EditAnywhere, BlueprintReadOnly, AdvancedDisplay, Category=Rendering)
+	TArray<TObjectPtr<class UMaterialInterface>> MaterialSlotsOverlayMaterial;
+
+	/**
+	 * Fill the array with every material slot overlay material use by this instance.
+	 * 
+	 * If this component material slot overlay material will be used if not null.
+	 * If there is no valid component material slot overlay material, the mesh material slot overlay material will be used if not null.
+	 * If there is no valid asset material slot overlay material, a null entry will be set for the material slot overlay material.
+	 */
+	ENGINE_API void GetMaterialSlotsOverlayMaterial(TArray<TObjectPtr<class UMaterialInterface>>& OutMaterialSlotOverlayMaterials) const;
+
 #if WITH_EDITOR
 	/*
 	 * Make sure the Override array is using only the space it should use.
@@ -69,6 +87,12 @@ public:
 	 * 2. The override array must not end with a nullptr UMaterialInterface.
 	 */
 	ENGINE_API void CleanUpOverrideMaterials();
+
+	/*
+	 * Make sure the MaterialSlotsOverlayMaterial is using only the space it should use.
+	 * - The data should be only on existing mesh material slot.
+	 */
+	ENGINE_API void CleanUpMaterialSlotsOverlayMaterial();
 #endif
 
 	/** 
@@ -209,6 +233,11 @@ protected:
 	/** Get the default overlay material max draw distance */
 	virtual float GetDefaultOverlayMaterialMaxDrawDistance() const { return 0.f; };
 
+	/**
+	 * Get all default material slots overlay materials from the mesh.
+	 */
+	virtual void GetDefaultMaterialSlotsOverlayMaterial(TArray<TObjectPtr<UMaterialInterface>>& OutMaterialSlotOverlayMaterials) const { return; }
+
 	/** Retrieves all the (scalar/vector-)parameters from within the used materials on the SkeletalMesh, and stores material index vs parameter names */
 	ENGINE_API void CacheMaterialParameterNameIndices();
 
diff --git a/Engine/Source/Runtime/Engine/Classes/Components/SkinnedMeshComponent.h b/Engine/Source/Runtime/Engine/Classes/Components/SkinnedMeshComponent.h
--- a/Engine/Source/Runtime/Engine/Classes/Components/SkinnedMeshComponent.h
+++ b/Engine/Source/Runtime/Engine/Classes/Components/SkinnedMeshComponent.h
@@ -1764,6 +1764,13 @@ protected:
 
 	/** Get the default overlay material max draw distance */
 	ENGINE_API virtual float GetDefaultOverlayMaterialMaxDrawDistance() const override;
+	
+	/*
+	 * Get all default LOD sections overlay materials used by a mesh.
+	 * The implementer must add all LOD sections and set them to null if there is no value.
+	 * We need all sections in case they are all active we wont add the global mesh MaterialOverlay in GetUseMaterials or GetMaterialRelevance.
+	 */
+	ENGINE_API virtual void GetDefaultMaterialSlotsOverlayMaterial(TArray<TObjectPtr<UMaterialInterface>>& OutMaterialSlotOverlayMaterials) const override;
 
 	void UpdateCachedSocketAttachmentFlag();
 	
diff --git a/Engine/Source/Runtime/Engine/Classes/Components/StaticMeshComponent.h b/Engine/Source/Runtime/Engine/Classes/Components/StaticMeshComponent.h
--- a/Engine/Source/Runtime/Engine/Classes/Components/StaticMeshComponent.h
+++ b/Engine/Source/Runtime/Engine/Classes/Components/StaticMeshComponent.h
@@ -560,6 +560,15 @@ protected:
 	ENGINE_API virtual bool ShouldCreateRenderState() const override;
 	ENGINE_API virtual bool ShouldIncrementalPreRegister(UWorld* WorldContext) const override;
 
+	
+	/**
+	 * Get all default LOD sections overlay materials used by a mesh.
+	 *
+	 * The implementer must add all LOD sections and set them to null if there is no value.
+	 * We need all sections in case they are all active we wont add the global mesh MaterialOverlay in GetUseMaterials or GetMaterialRelevance.
+	 */
+	ENGINE_API virtual void GetDefaultMaterialSlotsOverlayMaterial(TArray<TObjectPtr<UMaterialInterface>>& OutMaterialSlotOverlayMaterials) const override;
+
 public:
 
 	ENGINE_API virtual bool AllowsAsyncPhysicsStateCreation() const override;
diff --git a/Engine/Source/Runtime/Engine/Classes/Engine/SkinnedAssetCommon.h b/Engine/Source/Runtime/Engine/Classes/Engine/SkinnedAssetCommon.h
--- a/Engine/Source/Runtime/Engine/Classes/Engine/SkinnedAssetCommon.h
+++ b/Engine/Source/Runtime/Engine/Classes/Engine/SkinnedAssetCommon.h
@@ -384,6 +384,7 @@ struct FSkeletalMaterial
 		, bRecomputeTangent_DEPRECATED(false)
 		, ImportedMaterialSlotName(NAME_None)
 #endif
+		, OverlayMaterialInterface(nullptr)
 	{
 
 	}
@@ -391,7 +392,8 @@ struct FSkeletalMaterial
 	FSkeletalMaterial(
 		UMaterialInterface* InMaterialInterface,
 		FName InMaterialSlotName,
-		FName InImportedMaterialSlotName = NAME_None)
+		FName InImportedMaterialSlotName = NAME_None,
+		UMaterialInterface* InOverlayMaterialInterface = nullptr)
 		: MaterialInterface( InMaterialInterface )
 		, MaterialSlotName(InMaterialSlotName)
 	#if WITH_EDITORONLY_DATA
@@ -399,6 +401,7 @@ struct FSkeletalMaterial
 		, bRecomputeTangent_DEPRECATED(false)
 		, ImportedMaterialSlotName(InImportedMaterialSlotName)
 	#endif
+		, OverlayMaterialInterface(InOverlayMaterialInterface)
 	{
 
 	}
@@ -407,7 +410,8 @@ struct FSkeletalMaterial
 						, bool bInEnableShadowCasting = true
 						, bool bInRecomputeTangent = false
 						, FName InMaterialSlotName = NAME_None
-						, FName InImportedMaterialSlotName = NAME_None)
+						, FName InImportedMaterialSlotName = NAME_None
+						, UMaterialInterface* InOverlayMaterialInterface = nullptr)
 		: MaterialInterface( InMaterialInterface )
 		, MaterialSlotName(InMaterialSlotName)
 	#if WITH_EDITORONLY_DATA
@@ -415,6 +419,7 @@ struct FSkeletalMaterial
 		, bRecomputeTangent_DEPRECATED(bInRecomputeTangent)
 		, ImportedMaterialSlotName(InImportedMaterialSlotName)
 	#endif
+		, OverlayMaterialInterface(InOverlayMaterialInterface)
 	{
 
 	}
@@ -447,4 +452,7 @@ struct FSkeletalMaterial
 	/** Data used for texture streaming relative to each UV channels. */
 	UPROPERTY(VisibleAnywhere, BlueprintReadOnly, Category = SkeletalMesh)
 	FMeshUVChannelInfo			UVChannelData;
+
+	UPROPERTY(EditAnywhere, BlueprintReadWrite, Category=SkeletalMesh)
+	TObjectPtr<UMaterialInterface> 	OverlayMaterialInterface;
 };
\ No newline at end of file
diff --git a/Engine/Source/Runtime/Engine/Classes/Engine/StaticMesh.h b/Engine/Source/Runtime/Engine/Classes/Engine/StaticMesh.h
--- a/Engine/Source/Runtime/Engine/Classes/Engine/StaticMesh.h
+++ b/Engine/Source/Runtime/Engine/Classes/Engine/StaticMesh.h
@@ -491,6 +491,7 @@ struct FStaticMaterial
 #if WITH_EDITORONLY_DATA
 		, FName InImportedMaterialSlotName = NAME_None
 #endif
+		, class UMaterialInterface* InOverlayMaterialInterface = nullptr
 		);
 
 	friend FArchive& operator<<(FArchive& Ar, FStaticMaterial& Elem);
@@ -513,6 +514,9 @@ struct FStaticMaterial
 	/** Data used for texture streaming relative to each UV channels. */
 	UPROPERTY(VisibleAnywhere, BlueprintReadOnly, Category = StaticMesh)
 	FMeshUVChannelInfo			UVChannelData;
+
+	UPROPERTY(EditAnywhere, BlueprintReadWrite, Category = StaticMesh)
+	TObjectPtr<class UMaterialInterface> OverlayMaterialInterface;
 };
 
 USTRUCT()
diff --git a/Engine/Source/Runtime/Engine/Private/Components/MeshComponent.cpp b/Engine/Source/Runtime/Engine/Private/Components/MeshComponent.cpp
--- a/Engine/Source/Runtime/Engine/Private/Components/MeshComponent.cpp
+++ b/Engine/Source/Runtime/Engine/Private/Components/MeshComponent.cpp
@@ -3,6 +3,7 @@
 #include "Components/MeshComponent.h"
 #include "Materials/Material.h"
 #include "MaterialDomain.h"
+#include "Engine/MaterialOverlayHelper.h"
 #include "Engine/Texture2D.h"
 #include "Materials/MaterialInstanceDynamic.h"
 #include "ContentStreaming.h"
@@ -164,6 +165,10 @@ void UMeshComponent::PostEditChangeChainProperty(FPropertyChangedChainEvent& Pro
 		{
 			CleanUpOverrideMaterials();
 		}
+		else if (PropertyChangedEvent.Property->GetFName() == GET_MEMBER_NAME_STRING_CHECKED(UMeshComponent, MaterialSlotsOverlayMaterial))
+		{
+			CleanUpMaterialSlotsOverlayMaterial();
+		}
 	}
 }
 
@@ -196,10 +201,36 @@ void UMeshComponent::CleanUpOverrideMaterials()
 		MarkRenderStateDirty();
 	}
 }
+
+void UMeshComponent::CleanUpMaterialSlotsOverlayMaterial()
+{
+	//Calling the default will fill all the material slot until we have 
+	const int32 AssetMaterialSlotsOverlayMaterialCount = GetNumMaterials();
+	const int32 ComponentMaterialSlotsOverlayMaterialCount = MaterialSlotsOverlayMaterial.Num();
+	if (ComponentMaterialSlotsOverlayMaterialCount == 0)
+	{
+		return;
+	}
+
+	bool bUpdated = false;
+	if (ComponentMaterialSlotsOverlayMaterialCount > AssetMaterialSlotsOverlayMaterialCount)
+	{
+		int32 RemoveCount = ComponentMaterialSlotsOverlayMaterialCount - AssetMaterialSlotsOverlayMaterialCount;
+		MaterialSlotsOverlayMaterial.RemoveAt(AssetMaterialSlotsOverlayMaterialCount, RemoveCount);
+		bUpdated = true;
+	}
+
+	if (bUpdated)
+	{
+		MarkRenderStateDirty();
+	}
+}
+
 #endif
 
 void UMeshComponent::EmptyOverrideMaterials()
 {
+	bool bRefresh = false;
 	if (OverrideMaterials.Num())
 	{
 		for (int32 MatIndex = 0; MatIndex < OverrideMaterials.Num(); MatIndex++)
@@ -211,6 +242,19 @@ void UMeshComponent::EmptyOverrideMaterials()
 		}
 
 		OverrideMaterials.Reset();
+		bRefresh = true;
+	}
+
+	//Empty the material slot overlay material array
+	//Those are hook to the material list like the override materials
+	if (MaterialSlotsOverlayMaterial.Num())
+	{
+		MaterialSlotsOverlayMaterial.Reset();
+		bRefresh = true;
+	}
+
+	if(bRefresh)
+	{
 		MarkRenderStateDirty();
 
 		// Precache PSOs again
@@ -238,6 +282,12 @@ void UMeshComponent::GetUsedMaterials(TArray<UMaterialInterface*>& OutMaterials,
 		}
 	}
 
+	//Add material slot overlay materials
+	TArray<TObjectPtr<UMaterialInterface>> AssetAndComponentMaterialSlotOverlayMaterials;
+	GetMaterialSlotsOverlayMaterial(AssetAndComponentMaterialSlotOverlayMaterials);
+	FMaterialOverlayHelper::AppendAllOverlayMaterial(AssetAndComponentMaterialSlotOverlayMaterials, OutMaterials);
+
+	//Add global mesh overlay
 	UMaterialInterface* OverlayMaterialInterface = GetOverlayMaterial();
 	if (OverlayMaterialInterface != nullptr)
 	{
@@ -289,6 +339,30 @@ void UMeshComponent::SetOverlayMaterialMaxDrawDistance(float InMaxDrawDistance)
 	}
 }
 
+void UMeshComponent::GetMaterialSlotsOverlayMaterial(TArray<TObjectPtr<UMaterialInterface>>& OutMaterialSlotOverlayMaterials) const
+{
+	//Add all the component override
+	OutMaterialSlotOverlayMaterials = MaterialSlotsOverlayMaterial;
+
+	//Get the asset material slot overlay material
+	TArray<TObjectPtr<UMaterialInterface>> AssetMaterialSlotOverlayMaterials;
+	GetDefaultMaterialSlotsOverlayMaterial(AssetMaterialSlotOverlayMaterials);
+
+	//For each slot not override by the component set the asset slot value
+	for (int32 SlotIndex = 0; SlotIndex < AssetMaterialSlotOverlayMaterials.Num(); ++SlotIndex)
+	{
+		if (!OutMaterialSlotOverlayMaterials.IsValidIndex(SlotIndex))
+		{
+			OutMaterialSlotOverlayMaterials.AddZeroed();
+		}
+		//If the component did not already override this slot index put the asset value
+		if (!OutMaterialSlotOverlayMaterials[SlotIndex])
+		{
+			OutMaterialSlotOverlayMaterials[SlotIndex] = AssetMaterialSlotOverlayMaterials[SlotIndex];
+		}
+	}
+}
+
 void UMeshComponent::PrestreamTextures( float Seconds, bool bPrioritizeCharacterTextures, int32 CinematicTextureGroups )
 {
 	// If requested, tell the streaming system to only process character textures for 30 frames.
diff --git a/Engine/Source/Runtime/Engine/Private/Components/SkeletalMeshComponent.cpp b/Engine/Source/Runtime/Engine/Private/Components/SkeletalMeshComponent.cpp
--- a/Engine/Source/Runtime/Engine/Private/Components/SkeletalMeshComponent.cpp
+++ b/Engine/Source/Runtime/Engine/Private/Components/SkeletalMeshComponent.cpp
@@ -1470,6 +1470,8 @@ void USkeletalMeshComponent::PostEditChangeProperty(FPropertyChangedEvent& Prope
 				CleanUpOverrideMaterials();
 				MarkRenderStateDirty();
 			}
+			//Make sure the mesh section overlay materials is cleanup and match the new skeletal mesh
+			CleanUpMaterialSlotsOverlayMaterial();
 		}
 
 		// when user changes simulate physics, just make sure to update blendphysics together
diff --git a/Engine/Source/Runtime/Engine/Private/Components/SkinnedMeshComponent.cpp b/Engine/Source/Runtime/Engine/Private/Components/SkinnedMeshComponent.cpp
--- a/Engine/Source/Runtime/Engine/Private/Components/SkinnedMeshComponent.cpp
+++ b/Engine/Source/Runtime/Engine/Private/Components/SkinnedMeshComponent.cpp
@@ -19,6 +19,7 @@
 #include "SkeletalRenderStatic.h"
 #include "Animation/AnimStats.h"
 #include "SkeletalMeshDeformerHelpers.h"
+#include "Engine/MaterialOverlayHelper.h"
 #include "Engine/SkeletalMeshSocket.h"
 #include "Engine/SkinnedAssetCommon.h"
 #include "PhysicsEngine/PhysicsAsset.h"
@@ -696,20 +697,38 @@ void USkinnedMeshComponent::CollectPSOPrecacheData(const FPSOPrecacheParams& Bas
 		}
 	}
 
-	UMaterialInterface* OverlayMaterialInterface = GetOverlayMaterial();
-	if (OverlayMaterialInterface && VFsPerMaterials.Num() != 0)
-	{
-		// Overlay is rendered with the same set of VFs
-		FMaterialInterfacePSOPrecacheParams& ComponentParams = OutParams[OutParams.AddDefaulted()];
+	TArray<UMaterialInterface*> AllOverlayMaterials;
+	TArray<TObjectPtr<UMaterialInterface>> AssetAndComponentMaterialSlotsOverlayMaterial;
+	GetMaterialSlotsOverlayMaterial(AssetAndComponentMaterialSlotsOverlayMaterial);
+	bool bAllMaterialSlotOverride = true;
+	FMaterialOverlayHelper::AppendAllOverlayMaterial(AssetAndComponentMaterialSlotsOverlayMaterial, AllOverlayMaterials, bAllMaterialSlotOverride);
 
-		ComponentParams.MaterialInterface = OverlayMaterialInterface;
-		ComponentParams.VertexFactoryDataList = VFsPerMaterials[0].VertexFactoryDataList;
-		ComponentParams.PSOPrecacheParams = PrecachePSOParams;
-		ComponentParams.PSOPrecacheParams.bCastShadow = false;
+	//If all material slot set an overlay material we wont use the global one
+	if (!bAllMaterialSlotOverride)
+	{
+		AllOverlayMaterials.AddUnique(GetOverlayMaterial());
+	}
 
-		if (bHasValidNaniteMaterials)
+	if (AllOverlayMaterials.Num() > 0 && VFsPerMaterials.Num() != 0)
+	{
+		for(UMaterialInterface* OverlayMaterialInterface : AllOverlayMaterials)
 		{
-			ComponentParams.VertexFactoryDataList.Add(&FNaniteVertexFactory::StaticType);
+			if (!OverlayMaterialInterface)
+			{
+				continue;
+			}
+			// Overlay is rendered with the same set of VFs
+			FMaterialInterfacePSOPrecacheParams& ComponentParams = OutParams[OutParams.AddDefaulted()];
+
+			ComponentParams.MaterialInterface = OverlayMaterialInterface;
+			ComponentParams.VertexFactoryDataList = VFsPerMaterials[0].VertexFactoryDataList;
+			ComponentParams.PSOPrecacheParams = PrecachePSOParams;
+			ComponentParams.PSOPrecacheParams.bCastShadow = false;
+
+			if (bHasValidNaniteMaterials)
+			{
+				ComponentParams.VertexFactoryDataList.Add(&FNaniteVertexFactory::StaticType);
+			}
 		}
 	}
 }
@@ -3030,6 +3049,19 @@ float USkinnedMeshComponent::GetDefaultOverlayMaterialMaxDrawDistance() const
 	return 0.f;
 }
 
+void USkinnedMeshComponent::GetDefaultMaterialSlotsOverlayMaterial(TArray<TObjectPtr<UMaterialInterface>>& OutMaterialSlotOverlayMaterials) const
+{
+	OutMaterialSlotOverlayMaterials.Reset();
+	if (GetSkinnedAsset())
+	{
+		const TArray<FSkeletalMaterial>& SkeletalMaterials = GetSkinnedAsset()->GetMaterials();
+		for (const FSkeletalMaterial& SkeletalMaterial : SkeletalMaterials)
+		{
+			OutMaterialSlotOverlayMaterials.Add(SkeletalMaterial.OverlayMaterialInterface);
+		}
+	}
+}
+
 #if WITH_EDITOR
 void USkinnedMeshComponent::SetSectionPreview(int32 InSectionIndexPreview)
 {
@@ -3787,10 +3819,17 @@ void USkinnedMeshComponent::GetUsedMaterials( TArray<UMaterialInterface*>& OutMa
 			OutMaterials.Add( MaterialInterface );
 		}
 
-		UMaterialInterface* OverlayMaterialInterface = GetOverlayMaterial();
-		if (OverlayMaterialInterface != nullptr)
+		TArray<TObjectPtr<UMaterialInterface>> AssetAndComponentMaterialSlotsOverlayMaterial;
+		GetMaterialSlotsOverlayMaterial(AssetAndComponentMaterialSlotsOverlayMaterial);
+		bool bUseGlobalMeshOverlayMaterial = false;
+		FMaterialOverlayHelper::AppendAllOverlayMaterial(AssetAndComponentMaterialSlotsOverlayMaterial, OutMaterials, bUseGlobalMeshOverlayMaterial);
+		if(bUseGlobalMeshOverlayMaterial)
 		{
-			OutMaterials.Add(OverlayMaterialInterface);
+			UMaterialInterface* OverlayMaterialInterface = GetOverlayMaterial();
+			if (OverlayMaterialInterface != nullptr)
+			{
+				OutMaterials.Add(OverlayMaterialInterface);
+			}
 		}
 	}
 
diff --git a/Engine/Source/Runtime/Engine/Private/Components/StaticMeshComponent.cpp b/Engine/Source/Runtime/Engine/Private/Components/StaticMeshComponent.cpp
--- a/Engine/Source/Runtime/Engine/Private/Components/StaticMeshComponent.cpp
+++ b/Engine/Source/Runtime/Engine/Private/Components/StaticMeshComponent.cpp
@@ -1971,6 +1971,9 @@ void UStaticMeshComponent::PostEditChangeProperty(FPropertyChangedEvent& Propert
 				// Static mesh was switched so we should clean up the override materials
 				CleanUpOverrideMaterials();
 			}
+
+			//Make sure the mesh section overlay materials is cleanup and match the new skeletal mesh
+			CleanUpMaterialSlotsOverlayMaterial();
 		}
 
 		if (PropertyThatChanged->GetFName() == GET_MEMBER_NAME_CHECKED(UStaticMeshComponent, OverrideMaterials))
@@ -2245,6 +2248,20 @@ void UStaticMeshComponent::SetStaticMeshInternal(UStaticMesh* NewMesh)
 	NotifyIfStaticMeshChanged();
 }
 
+void UStaticMeshComponent::GetDefaultMaterialSlotsOverlayMaterial(TArray<TObjectPtr<UMaterialInterface>>& OutMaterialSlotOverlayMaterials) const
+{
+	//We add null entry for every section of every LOD, this is a requirement for the MeshComponent this class derived from
+	OutMaterialSlotOverlayMaterials.Reset();
+	if (GetStaticMesh())
+	{
+		const TArray<FStaticMaterial>& StaticMaterials = GetStaticMesh()->GetStaticMaterials();
+		for (const FStaticMaterial& StaticMaterial : StaticMaterials)
+		{
+			OutMaterialSlotOverlayMaterials.Add(StaticMaterial.OverlayMaterialInterface);
+		}
+	}
+}
+
 bool UStaticMeshComponent::SetStaticMesh(UStaticMesh* NewMesh)
 {
 	// Do nothing if we are already using the supplied static mesh
diff --git a/Engine/Source/Runtime/Engine/Private/SkeletalMeshSceneProxy.cpp b/Engine/Source/Runtime/Engine/Private/SkeletalMeshSceneProxy.cpp
--- a/Engine/Source/Runtime/Engine/Private/SkeletalMeshSceneProxy.cpp
+++ b/Engine/Source/Runtime/Engine/Private/SkeletalMeshSceneProxy.cpp
@@ -2,6 +2,7 @@
 
 #include "SkeletalMeshSceneProxy.h"
 #include "Components/SkeletalMeshComponent.h"
+#include "Engine/MaterialOverlayHelper.h"
 #include "Engine/SkeletalMesh.h"
 #include "EngineUtils.h"
 #include "MaterialDomain.h"
@@ -130,6 +131,8 @@ FSkeletalMeshSceneProxy::FSkeletalMeshSceneProxy(const FSkinnedMeshSceneProxyDes
 	// Get the pre-skinned local bounds
 	InMeshDesc.GetPreSkinnedLocalBounds(PreSkinnedLocalBounds);
 
+	InMeshDesc.GetMaterialSlotsOverlayMaterial(MaterialSlotsOverlayMaterial);
+
 	if(InMeshDesc.bPerBoneMotionBlur)
 	{
 		bAlwaysHasVelocity = true;
@@ -205,11 +208,29 @@ FSkeletalMeshSceneProxy::FSkeletalMeshSceneProxy(const FSkinnedMeshSceneProxyDes
 			bAnySegmentUsesWorldPositionOffset |= MaterialRelevance.bUsesWorldPositionOffset;
 #endif
 
+			UMaterialInterface* SectionOverlayMaterial = FMaterialOverlayHelper::GetOverlayMaterial(MaterialSlotsOverlayMaterial, UseMaterialIndex);
+
+			if (SectionOverlayMaterial != nullptr)
+			{
+				if (!SectionOverlayMaterial->CheckMaterialUsage_Concurrent(MATUSAGE_SkeletalMesh))
+				{
+					SectionOverlayMaterial = nullptr;
+					UE_LOG(LogSkeletalMesh, Warning, TEXT("Overlay material per section with missing usage flag was applied to skeletal mesh[%s] LOD %d, section index %d."), *InMeshDesc.GetSkinnedAsset()->GetPathName(), LODIdx, SectionIndex);
+					FMaterialOverlayHelper::ForceMaterial(MaterialSlotsOverlayMaterial, UseMaterialIndex, SectionOverlayMaterial);
+				}
+				else if (bForceDefaultMaterial)
+				{
+					SectionOverlayMaterial = nullptr;
+					FMaterialOverlayHelper::ForceMaterial(MaterialSlotsOverlayMaterial, UseMaterialIndex, SectionOverlayMaterial);
+				}
+			}
+
 			LODSection.SectionElements.Add(
 				FSectionElementInfo(
 					Material,
 					bSectionCastsShadow,
-					UseMaterialIndex
+					UseMaterialIndex,
+					SectionOverlayMaterial
 					));
 			MaterialsInUse_GameThread.Add(Material);
 		}
@@ -219,12 +240,12 @@ FSkeletalMeshSceneProxy::FSkeletalMeshSceneProxy(const FSkinnedMeshSceneProxyDes
 	{
 		if (!OverlayMaterial->CheckMaterialUsage_Concurrent(MATUSAGE_SkeletalMesh))
 		{
-			OverlayMaterial = UMaterial::GetDefaultMaterial(MD_Surface);
+			OverlayMaterial = nullptr;
 			UE_LOG(LogSkeletalMesh, Warning, TEXT("Overlay material with missing usage flag was applied to skeletal mesh %s"),	*InMeshDesc.GetSkinnedAsset()->GetPathName());
 		}
 		else if (bForceDefaultMaterial)
 		{
-			OverlayMaterial = UMaterial::GetDefaultMaterial(MD_Surface);
+			OverlayMaterial = nullptr;
 		}
 	}
 
@@ -510,13 +531,14 @@ void FSkeletalMeshSceneProxy::DrawStaticElements(FStaticPrimitiveDrawInterface*
 
 					PDI->DrawMesh(MeshElement, ScreenSize);
 
-					if (OverlayMaterial != nullptr)
+					UMaterialInterface* SectionOverlayMaterial = (SectionElementInfo.PerSectionOverlayMaterial != nullptr) ? SectionElementInfo.PerSectionOverlayMaterial : OverlayMaterial;
+					if (SectionOverlayMaterial != nullptr)
 					{
 						FMeshBatch OverlayMeshBatch(MeshElement);
 						OverlayMeshBatch.bOverlayMaterial = true;
 						OverlayMeshBatch.CastShadow = false;
 						OverlayMeshBatch.bSelectable = false;
-						OverlayMeshBatch.MaterialRenderProxy = OverlayMaterial->GetRenderProxy();
+						OverlayMeshBatch.MaterialRenderProxy = SectionOverlayMaterial->GetRenderProxy();
 						// make sure overlay is always rendered on top of base mesh
 						OverlayMeshBatch.MeshIdInPrimitive += LODData.RenderSections.Num();
 						// Reuse mesh ScreenSize as cull distance for an overlay. Overlay does not need to compute LOD so we can avoid adding new members into MeshBatch or MeshRelevance
@@ -803,8 +825,9 @@ void FSkeletalMeshSceneProxy::GetDynamicElementsSection(const TArray<const FScen
 			INC_DWORD_STAT_BY(STAT_SkelMeshTriangles,Mesh.GetNumPrimitives());
 			INC_DWORD_STAT(STAT_SkelMeshDrawCalls);
 
+			UMaterialInterface* ActiveOverlayMaterial = SectionElementInfo.PerSectionOverlayMaterial ? SectionElementInfo.PerSectionOverlayMaterial : OverlayMaterial;
 			// negative cull distance disables overlay rendering
-			if (OverlayMaterial != nullptr && OverlayMaterialMaxDrawDistance >= 0.f)
+			if (ActiveOverlayMaterial != nullptr && OverlayMaterialMaxDrawDistance >= 0.f)
 			{
 				const bool bHasOverlayCullDistance = 
 					OverlayMaterialMaxDrawDistance > 0.f && 
@@ -832,7 +855,7 @@ void FSkeletalMeshSceneProxy::GetDynamicElementsSection(const TArray<const FScen
 					OverlayMeshBatch.bOverlayMaterial = true;
 					OverlayMeshBatch.CastShadow = false;
 					OverlayMeshBatch.bSelectable = false;
-					OverlayMeshBatch.MaterialRenderProxy = OverlayMaterial->GetRenderProxy();
+					OverlayMeshBatch.MaterialRenderProxy = ActiveOverlayMaterial->GetRenderProxy();
 					// make sure overlay is always rendered on top of base mesh
 					OverlayMeshBatch.MeshIdInPrimitive += LODData.RenderSections.Num();
 					Collector.AddMesh(ViewIndex, OverlayMeshBatch);
diff --git a/Engine/Source/Runtime/Engine/Private/SkinnedAssetCommon.cpp b/Engine/Source/Runtime/Engine/Private/SkinnedAssetCommon.cpp
--- a/Engine/Source/Runtime/Engine/Private/SkinnedAssetCommon.cpp
+++ b/Engine/Source/Runtime/Engine/Private/SkinnedAssetCommon.cpp
@@ -14,6 +14,7 @@
 #include "SkeletalMeshLegacyCustomVersions.h"
 #include "UObject/CoreObjectVersion.h"
 #include "UObject/EditorObjectVersion.h"
+#include "UObject/FortniteMainBranchObjectVersion.h"
 #include "UObject/FortniteSeasonBranchObjectVersion.h"
 #include "UObject/RenderingObjectVersion.h"
 
@@ -279,6 +280,7 @@ void FSkeletalMaterial::DeclareCustomVersions(FArchive& Ar)
 {
 	Ar.UsingCustomVersion(FEditorObjectVersion::GUID);
 	Ar.UsingCustomVersion(FCoreObjectVersion::GUID);
+	Ar.UsingCustomVersion(FFortniteMainBranchObjectVersion::GUID);
 }
 #endif
 
@@ -286,6 +288,7 @@ FArchive& operator<<(FArchive& Ar, FSkeletalMaterial& Elem)
 {
 	Ar.UsingCustomVersion(FEditorObjectVersion::GUID);
 	Ar.UsingCustomVersion(FCoreObjectVersion::GUID);
+	Ar.UsingCustomVersion(FFortniteMainBranchObjectVersion::GUID);
 
 	Ar << Elem.MaterialInterface;
 
@@ -334,5 +337,13 @@ FArchive& operator<<(FArchive& Ar, FSkeletalMaterial& Elem)
 		Ar << Elem.UVChannelData;
 	}
 
+	if (!Ar.IsLoading() || Ar.CustomVer(FFortniteMainBranchObjectVersion::GUID) >= FFortniteMainBranchObjectVersion::MeshMaterialSlotOverlayMaterialAdded)
+	{
+		Ar << Elem.OverlayMaterialInterface;
+	}
+	else
+	{
+		Elem.OverlayMaterialInterface = nullptr;
+	}
 	return Ar;
 }
diff --git a/Engine/Source/Runtime/Engine/Private/SkinnedMeshSceneProxyDesc.cpp b/Engine/Source/Runtime/Engine/Private/SkinnedMeshSceneProxyDesc.cpp
--- a/Engine/Source/Runtime/Engine/Private/SkinnedMeshSceneProxyDesc.cpp
+++ b/Engine/Source/Runtime/Engine/Private/SkinnedMeshSceneProxyDesc.cpp
@@ -9,6 +9,7 @@
 #include "SkeletalRenderStatic.h"
 #include "Materials/MaterialInstanceDynamic.h"
 #include "Engine/Engine.h"
+#include "Engine/MaterialOverlayHelper.h"
 #include "Engine/SkinnedAssetCommon.h"
 #include "Engine/SkinnedAsset.h"
 #include "PhysicsEngine/PhysicsAsset.h"
@@ -59,6 +60,7 @@ void FSkinnedMeshSceneProxyDesc::InitializeFromSkinnedMeshComponent(const USkinn
 	SkinnedAsset = InComponent->GetSkinnedAsset();
 	PhysicsAsset = InComponent->GetPhysicsAsset();
 	OverlayMaterial = InComponent->GetOverlayMaterial();
+	InComponent->GetMaterialSlotsOverlayMaterial(MaterialSlotsOverlayMaterial);
 	MeshDeformerInstances = &InComponent->GetMeshDeformerInstances();
 	OverrideMaterials = InComponent->OverrideMaterials;
 	SkinCacheUsage = InComponent->SkinCacheUsage;
@@ -169,10 +171,15 @@ void FSkinnedMeshSceneProxyDesc::GetUsedMaterials(TArray<UMaterialInterface*>& O
 			OutMaterials.Add( MaterialInterface );
 		}
 
-		UMaterialInterface* OverlayMaterialInterface = GetOverlayMaterial();
-		if (OverlayMaterialInterface != nullptr)
+		bool bUseGlobalMeshOverlayMaterial = false;
+		FMaterialOverlayHelper::AppendAllOverlayMaterial(MaterialSlotsOverlayMaterial, OutMaterials, bUseGlobalMeshOverlayMaterial);
+		if(bUseGlobalMeshOverlayMaterial)
 		{
-			OutMaterials.Add(OverlayMaterialInterface);
+			UMaterialInterface* OverlayMaterialInterface = GetOverlayMaterial();
+			if (OverlayMaterialInterface != nullptr)
+			{
+				OutMaterials.Add(OverlayMaterialInterface);
+			}
 		}
 	}
 
diff --git a/Engine/Source/Runtime/Engine/Private/StaticMesh.cpp b/Engine/Source/Runtime/Engine/Private/StaticMesh.cpp
--- a/Engine/Source/Runtime/Engine/Private/StaticMesh.cpp
+++ b/Engine/Source/Runtime/Engine/Private/StaticMesh.cpp
@@ -4167,6 +4167,9 @@ void FStaticMeshRenderData::Cache(const ITargetPlatform* TargetPlatform, UStatic
 
 FArchive& operator<<(FArchive& Ar, FStaticMaterial& Elem)
 {
+	Ar.UsingCustomVersion(FFortniteMainBranchObjectVersion::GUID);
+	Ar.UsingCustomVersion(FRenderingObjectVersion::GUID);
+
 	Ar << Elem.MaterialInterface;
 
 	Ar << Elem.MaterialSlotName;
@@ -4181,7 +4184,15 @@ FArchive& operator<<(FArchive& Ar, FStaticMaterial& Elem)
 	{
 		Ar << Elem.UVChannelData;
 	}
-	
+
+	if (!Ar.IsLoading() || Ar.CustomVer(FFortniteMainBranchObjectVersion::GUID) >= FFortniteMainBranchObjectVersion::MeshMaterialSlotOverlayMaterialAdded)
+	{
+		Ar << Elem.OverlayMaterialInterface;
+	}
+	else
+	{
+		Elem.OverlayMaterialInterface = nullptr;
+	}
 	return Ar;
 }
 
@@ -4198,15 +4209,16 @@ FStaticMaterial::FStaticMaterial()
 FStaticMaterial::FStaticMaterial(class UMaterialInterface* InMaterialInterface
 , FName InMaterialSlotName
 #if WITH_EDITORONLY_DATA
-, FName InImportedMaterialSlotName)
-#else
-)
+, FName InImportedMaterialSlotName
 #endif
+, UMaterialInterface* InOverlayMaterialInterface
+)
 : MaterialInterface(InMaterialInterface)
 , MaterialSlotName(InMaterialSlotName)
 #if WITH_EDITORONLY_DATA
 , ImportedMaterialSlotName(InImportedMaterialSlotName)
 #endif //WITH_EDITORONLY_DATA
+, OverlayMaterialInterface(InOverlayMaterialInterface)
 {
 //If not specified add some valid material slot name
 if (MaterialInterface && MaterialSlotName == NAME_None)
diff --git a/Engine/Source/Runtime/Engine/Private/StaticMeshSceneProxy.cpp b/Engine/Source/Runtime/Engine/Private/StaticMeshSceneProxy.cpp
--- a/Engine/Source/Runtime/Engine/Private/StaticMeshSceneProxy.cpp
+++ b/Engine/Source/Runtime/Engine/Private/StaticMeshSceneProxy.cpp
@@ -17,6 +17,7 @@
 #include "Materials/Material.h"
 #include "Materials/MaterialRenderProxy.h"
 #include "MaterialShared.h"
+#include "Engine/MaterialOverlayHelper.h"
 #include "Engine/StaticMesh.h"
 #include "ComponentReregisterContext.h"
 #include "EngineUtils.h"
@@ -288,6 +289,8 @@ FStaticMeshSceneProxy::FStaticMeshSceneProxy(const FStaticMeshSceneProxyDesc& In
 	const int32 SMCurrentMinLOD = InProxyDesc.GetStaticMesh()->GetMinLODIdx();
 	int32 EffectiveMinLOD = InProxyDesc.bOverrideMinLOD ? InProxyDesc.MinLOD : SMCurrentMinLOD;
 
+	InProxyDesc.GetMaterialSlotsOverlayMaterial(MaterialSlotsOverlayMaterial);
+
 #if WITH_EDITOR
 	// If we plan to strip the min LOD during cooking, emulate that behavior in the editor
 	static const auto* CVar = IConsoleManager::Get().FindTConsoleVariableDataInt(TEXT("r.StaticMesh.StripMinLodDataDuringCooking"));
@@ -363,12 +366,12 @@ FStaticMeshSceneProxy::FStaticMeshSceneProxy(const FStaticMeshSceneProxyDesc& In
 		}
 	}
 
-	if (OverlayMaterial != nullptr)
+	if (bForceDefaultMaterial)
 	{
-		if (bForceDefaultMaterial)
-		{
-			OverlayMaterial = UMaterial::GetDefaultMaterial(MD_Surface);
-		}
+		OverlayMaterial = nullptr;
+
+		//Null out the overlay material to avoid the extra drawing with the same result.
+		FMaterialOverlayHelper::ForceMaterial(MaterialSlotsOverlayMaterial, nullptr);
 	}
 
 #if RHI_RAYTRACING
@@ -1314,15 +1317,21 @@ void FStaticMeshSceneProxy::DrawStaticElements(FStaticPrimitiveDrawInterface* PD
 			// Draw the static mesh elements.
 			for(int32 SectionIndex = 0; SectionIndex < LODModel.Sections.Num(); SectionIndex++)
 			{
+				const FLODInfo::FSectionInfo& Section = LODs[LODIndex].Sections[SectionIndex];
 			#if WITH_EDITOR
 				if (GIsEditor)
 				{
-					const FLODInfo::FSectionInfo& Section = LODs[LODIndex].Sections[SectionIndex];
 					bIsMeshElementSelected = Section.bSelected;
 					PDI->SetHitProxy(Section.HitProxy);
 				}
 			#endif
 
+				UMaterialInterface* SpecifiedOverlayMaterial = OverlayMaterial;
+				if (UMaterialInterface* SectionOverlayMaterial = Section.OverlayMaterial)
+				{
+					SpecifiedOverlayMaterial = SectionOverlayMaterial;
+				}
+
 				const int32 NumBatches = GetNumMeshBatches();
 				PDI->ReserveMemoryForMeshes(NumBatches * (1 + NumRuntimeVirtualTextureTypes));
 
@@ -1348,13 +1357,13 @@ void FStaticMeshSceneProxy::DrawStaticElements(FStaticPrimitiveDrawInterface* PD
 							PDI->DrawMesh(BaseMeshBatch, FLT_MAX);
 						}
 
-						if (OverlayMaterial != nullptr)
+						if (SpecifiedOverlayMaterial != nullptr)
 						{
 							FMeshBatch OverlayMeshBatch(BaseMeshBatch);
 							OverlayMeshBatch.bOverlayMaterial = true;
 							OverlayMeshBatch.CastShadow = false;
 							OverlayMeshBatch.bSelectable = false;
-							OverlayMeshBatch.MaterialRenderProxy = OverlayMaterial->GetRenderProxy();
+							OverlayMeshBatch.MaterialRenderProxy = SpecifiedOverlayMaterial->GetRenderProxy();
 							// make sure overlay is always rendered on top of base mesh
 							OverlayMeshBatch.MeshIdInPrimitive += LODModel.Sections.Num();
 							PDI->DrawMesh(OverlayMeshBatch, FLT_MAX);
@@ -1443,16 +1452,21 @@ void FStaticMeshSceneProxy::DrawStaticElements(FStaticPrimitiveDrawInterface* PD
 				// Draw the static mesh elements.
 				for(int32 SectionIndex = 0;SectionIndex < LODModel.Sections.Num();SectionIndex++)
 				{
+					const FLODInfo::FSectionInfo& Section = LODs[LODIndex].Sections[SectionIndex];
 #if WITH_EDITOR
 					if( GIsEditor )
 					{
-						const FLODInfo::FSectionInfo& Section = LODs[LODIndex].Sections[SectionIndex];
-
 						bIsMeshElementSelected = Section.bSelected;
 						PDI->SetHitProxy(Section.HitProxy);
 					}
 #endif // WITH_EDITOR
 
+					UMaterialInterface* SpecifiedOverlayMaterial = OverlayMaterial;
+					if (UMaterialInterface* SectionOverlayMaterial = Section.OverlayMaterial)
+					{
+						SpecifiedOverlayMaterial = SectionOverlayMaterial;
+					}
+
 					const int32 NumBatches = GetNumMeshBatches();
 					PDI->ReserveMemoryForMeshes(NumBatches * (1 + NumRuntimeVirtualTextureTypes));
 
@@ -1485,13 +1499,13 @@ void FStaticMeshSceneProxy::DrawStaticElements(FStaticPrimitiveDrawInterface* PD
 							}
 
 							// negative cull distance disables overlay rendering
-							if (OverlayMaterial != nullptr && OverlayMaterialMaxDrawDistance >= 0.f)
+							if (SpecifiedOverlayMaterial != nullptr && OverlayMaterialMaxDrawDistance >= 0.f)
 							{
 								FMeshBatch OverlayMeshBatch(BaseMeshBatch);
 								OverlayMeshBatch.bOverlayMaterial = true;
 								OverlayMeshBatch.CastShadow = false;
 								OverlayMeshBatch.bSelectable = false;
-								OverlayMeshBatch.MaterialRenderProxy = OverlayMaterial->GetRenderProxy();
+								OverlayMeshBatch.MaterialRenderProxy = SpecifiedOverlayMaterial->GetRenderProxy();
 								// make sure overlay is always rendered on top of base mesh
 								OverlayMeshBatch.MeshIdInPrimitive += LODModel.Sections.Num();
 								// Reuse mesh ScreenSize as cull distance for an overlay. Overlay does not need to compute LOD so we can avoid adding new members into MeshBatch or MeshRelevance
@@ -2408,6 +2422,10 @@ FStaticMeshSceneProxy::FLODInfo::FLODInfo(const FStaticMeshSceneProxyDesc& InPro
 
 	// Gather the materials applied to the LOD.
 	Sections.Empty(MeshRenderData->LODResources[LODIndex].Sections.Num());
+	
+	TArray<TObjectPtr<UMaterialInterface>> ProxyMaterialSlotsOverlayMaterial;
+	InProxyDesc.GetMaterialSlotsOverlayMaterial(ProxyMaterialSlotsOverlayMaterial);
+
 	for(int32 SectionIndex = 0;SectionIndex < LODModel.Sections.Num();SectionIndex++)
 	{
 		const FStaticMeshSection& Section = LODModel.Sections[SectionIndex];
@@ -2415,6 +2433,7 @@ FStaticMeshSceneProxy::FLODInfo::FLODInfo(const FStaticMeshSceneProxyDesc& InPro
 
 		// Determine the material applied to this element of the LOD.
 		SectionInfo.Material = InProxyDesc.GetMaterial(Section.MaterialIndex);
+		SectionInfo.OverlayMaterial = FMaterialOverlayHelper::GetOverlayMaterial(ProxyMaterialSlotsOverlayMaterial, Section.MaterialIndex);
 #if WITH_EDITORONLY_DATA
 		SectionInfo.MaterialIndex = Section.MaterialIndex;
 #endif
@@ -2422,12 +2441,20 @@ FStaticMeshSceneProxy::FLODInfo::FLODInfo(const FStaticMeshSceneProxyDesc& InPro
 		if (bForceDefaultMaterial || (GForceDefaultMaterial && SectionInfo.Material && !IsTranslucentBlendMode(*SectionInfo.Material)))
 		{
 			SectionInfo.Material = UMaterial::GetDefaultMaterial(MD_Surface);
+			if (SectionInfo.OverlayMaterial)
+			{
+				SectionInfo.OverlayMaterial = UMaterial::GetDefaultMaterial(MD_Surface);
+			}
 		}
 
 		// If there isn't an applied material, or if we need static lighting and it doesn't support it, fall back to the default material.
 		if(!SectionInfo.Material || (bHasSurfaceStaticLighting && !SectionInfo.Material->CheckMaterialUsage_Concurrent(MATUSAGE_StaticLighting)))
 		{
 			SectionInfo.Material = UMaterial::GetDefaultMaterial(MD_Surface);
+			if (SectionInfo.OverlayMaterial)
+			{
+				SectionInfo.OverlayMaterial = UMaterial::GetDefaultMaterial(MD_Surface);
+			}
 		}
 
 		// Per-section selection for the editor.
@@ -2686,6 +2713,7 @@ void FStaticMeshSceneProxyDesc::InitializeFromStaticMeshComponent(const UStaticM
 	OverrideMaterials = const_cast<UStaticMeshComponent*>(InComponent)->OverrideMaterials;	
 	OverlayMaterial = InComponent->GetOverlayMaterial();
 	OverlayMaterialMaxDrawDistance = InComponent->GetOverlayMaterialMaxDrawDistance();
+	InComponent->GetMaterialSlotsOverlayMaterial(MaterialSlotsOverlayMaterial);
 
 	ForcedLodModel = InComponent->ForcedLodModel ;
 	MinLOD = InComponent->MinLOD ;
diff --git a/Engine/Source/Runtime/Engine/Public/Engine/MaterialOverlayHelper.h b/Engine/Source/Runtime/Engine/Public/Engine/MaterialOverlayHelper.h
new file mode 100644
--- /dev/null
+++ b/Engine/Source/Runtime/Engine/Public/Engine/MaterialOverlayHelper.h
@@ -0,0 +1,59 @@
+// Copyright Epic Games, Inc. All Rights Reserved.
+
+#pragma once
+
+#include "Containers/Array.h"
+#include "Materials/MaterialInterface.h"
+#include "UObject/Object.h"
+#include "UObject/ObjectMacros.h"
+#include "UObject/ObjectPtr.h"
+#include "UObject/UObjectGlobals.h"
+
+class FMaterialOverlayHelper
+{
+public:
+	static void AppendAllOverlayMaterial(const TArray<TObjectPtr<UMaterialInterface>>& SourceMaterials, TArray<UMaterialInterface*>& OutMaterials, bool& bOutHaveNullEntry)
+	{
+		bOutHaveNullEntry = false;
+		for(const TObjectPtr<UMaterialInterface>& SourceMaterial : SourceMaterials)
+		{
+			if (SourceMaterial)
+			{
+				OutMaterials.Add(SourceMaterial);
+			}
+			else
+			{
+				bOutHaveNullEntry = true;
+			}
+		}
+	}
+
+	static void AppendAllOverlayMaterial(const TArray<TObjectPtr<UMaterialInterface>>& SourceMaterials, TArray<UMaterialInterface*>& OutMaterials)
+	{
+		bool bHaveNullEntry = false;
+		AppendAllOverlayMaterial(SourceMaterials, OutMaterials, bHaveNullEntry);
+	}
+
+	static void ForceMaterial(TArray<TObjectPtr<UMaterialInterface>>& SourceMaterials, UMaterialInterface* ForceMaterial)
+	{
+		for (int32 SlotIndex = 0; SlotIndex < SourceMaterials.Num(); ++SlotIndex)
+		{
+			SourceMaterials[SlotIndex] = ForceMaterial;
+		}
+	}
+
+	static bool ForceMaterial(TArray<TObjectPtr<UMaterialInterface>>& SourceMaterials, const int32 SlotIndex, UMaterialInterface* ForceMaterial)
+	{
+		if (SourceMaterials.IsValidIndex(SlotIndex))
+		{
+			SourceMaterials[SlotIndex] = ForceMaterial;
+			return true;
+		}
+		return false;
+	}
+
+	static UMaterialInterface* GetOverlayMaterial(TArray<TObjectPtr<UMaterialInterface>>& SourceMaterials, const int32 SlotIndex)
+	{
+		return SourceMaterials.IsValidIndex(SlotIndex) ? SourceMaterials[SlotIndex] : nullptr;
+	}
+};
diff --git a/Engine/Source/Runtime/Engine/Public/MeshComponentHelper.h b/Engine/Source/Runtime/Engine/Public/MeshComponentHelper.h
--- a/Engine/Source/Runtime/Engine/Public/MeshComponentHelper.h
+++ b/Engine/Source/Runtime/Engine/Public/MeshComponentHelper.h
@@ -34,10 +34,29 @@ FMaterialRelevance FMeshComponentHelper::GetMaterialRelevance(const T& Component
 		Result |= MaterialInterface->GetRelevance_Concurrent(InFeatureLevel);
 	}
 
-	UMaterialInterface const* OverlayMaterialInterface = Component.GetOverlayMaterial();
-	if (OverlayMaterialInterface != nullptr)
+	TArray<TObjectPtr<UMaterialInterface>> AssetAndComponentMaterialSlotsOverlayMaterial;
+	Component.GetMaterialSlotsOverlayMaterial(AssetAndComponentMaterialSlotsOverlayMaterial);
+	bool bAllMaterialSlotOverride = true;
+	for (const TObjectPtr<UMaterialInterface>& MaterialInterface : AssetAndComponentMaterialSlotsOverlayMaterial)
 	{
-		Result |= OverlayMaterialInterface->GetRelevance_Concurrent(InFeatureLevel);
+		if (MaterialInterface)
+		{
+			Result |= MaterialInterface->GetRelevance_Concurrent(InFeatureLevel);
+		}
+		else
+		{
+			bAllMaterialSlotOverride = false;
+		}
+	}
+
+	//If all material slot set an overlay material we wont use the global one
+	if (!bAllMaterialSlotOverride)
+	{
+		UMaterialInterface const* OverlayMaterialInterface = Component.GetOverlayMaterial();
+		if (OverlayMaterialInterface != nullptr)
+		{
+			Result |= OverlayMaterialInterface->GetRelevance_Concurrent(InFeatureLevel);
+		}
 	}
 
 	return Result;
diff --git a/Engine/Source/Runtime/Engine/Public/SkeletalMeshSceneProxy.h b/Engine/Source/Runtime/Engine/Public/SkeletalMeshSceneProxy.h
--- a/Engine/Source/Runtime/Engine/Public/SkeletalMeshSceneProxy.h
+++ b/Engine/Source/Runtime/Engine/Public/SkeletalMeshSceneProxy.h
@@ -143,6 +143,8 @@ protected:
 	
 	UMaterialInterface* OverlayMaterial;
 	float OverlayMaterialMaxDrawDistance;
+	TArray<TObjectPtr<UMaterialInterface>> MaterialSlotsOverlayMaterial;
+
 public:
 #if RHI_RAYTRACING
 	bool bAnySegmentUsesWorldPositionOffset : 1;
@@ -168,13 +170,14 @@ protected:
 	/** info for section element in an LOD */
 	struct FSectionElementInfo
 	{
-		FSectionElementInfo(UMaterialInterface* InMaterial, bool bInEnableShadowCasting, int32 InUseMaterialIndex)
+		FSectionElementInfo(UMaterialInterface* InMaterial, bool bInEnableShadowCasting, int32 InUseMaterialIndex, UMaterialInterface* InPerSectionOverlayMaterial)
 		: Material( InMaterial )
 		, bEnableShadowCasting( bInEnableShadowCasting )
 		, UseMaterialIndex( InUseMaterialIndex )
 #if WITH_EDITOR
 		, HitProxy(NULL)
 #endif
+		, PerSectionOverlayMaterial(InPerSectionOverlayMaterial)
 		{}
 		
 		UMaterialInterface* Material;
@@ -189,6 +192,8 @@ protected:
 		/** The editor needs to be able to individual sub-mesh hit detection, so we store a hit proxy on each mesh. */
 		HHitProxy* HitProxy;
 #endif
+
+		UMaterialInterface* PerSectionOverlayMaterial;
 	};
 
 	/** Section elements for a particular LOD */
diff --git a/Engine/Source/Runtime/Engine/Public/SkinnedMeshSceneProxyDesc.h b/Engine/Source/Runtime/Engine/Public/SkinnedMeshSceneProxyDesc.h
--- a/Engine/Source/Runtime/Engine/Public/SkinnedMeshSceneProxyDesc.h
+++ b/Engine/Source/Runtime/Engine/Public/SkinnedMeshSceneProxyDesc.h
@@ -36,6 +36,7 @@ struct FSkinnedMeshSceneProxyDesc : public FPrimitiveSceneProxyDesc
 	ENGINE_API FMaterialRelevance GetMaterialRelevance(ERHIFeatureLevel::Type InFeatureLevel) const;
 	ENGINE_API float GetOverlayMaterialMaxDrawDistance() const;
 	ENGINE_API UMaterialInterface* GetOverlayMaterial() const;
+	ENGINE_API void GetMaterialSlotsOverlayMaterial(TArray<TObjectPtr<UMaterialInterface>>& OutMaterialSlotsOverlayMaterial) const { OutMaterialSlotsOverlayMaterial = MaterialSlotsOverlayMaterial; }
 
 	ENGINE_API virtual void GetUsedMaterials(TArray<UMaterialInterface*>& OutMaterials, bool bGetDebugMaterials = false) const override;
 
@@ -66,6 +67,7 @@ struct FSkinnedMeshSceneProxyDesc : public FPrimitiveSceneProxyDesc
 	USkinnedAsset* SkinnedAsset = nullptr;
 	UPhysicsAsset* PhysicsAsset = nullptr;
 	UMaterialInterface* OverlayMaterial = nullptr;
+	TArray<TObjectPtr<UMaterialInterface>> MaterialSlotsOverlayMaterial;
 	const FMeshDeformerInstanceSet* MeshDeformerInstances = nullptr;
 	TArrayView<const TObjectPtr<UMaterialInterface>> OverrideMaterials;
 	TArrayView<const ESkinCacheUsage> SkinCacheUsage{};
diff --git a/Engine/Source/Runtime/Engine/Public/StaticMeshComponentHelper.h b/Engine/Source/Runtime/Engine/Public/StaticMeshComponentHelper.h
--- a/Engine/Source/Runtime/Engine/Public/StaticMeshComponentHelper.h
+++ b/Engine/Source/Runtime/Engine/Public/StaticMeshComponentHelper.h
@@ -10,6 +10,7 @@
 #include "PSOPrecacheFwd.h"
 #include "PSOPrecache.h"
 #include "RenderUtils.h"
+#include "Engine/MaterialOverlayHelper.h"
 #include "Engine/World.h"
 
 /** Helper class used to share implementation for different StaticMeshComponent types */
@@ -44,10 +45,18 @@ void FStaticMeshComponentHelper::GetUsedMaterials(const T& Component, TArray<UMa
 
 		if (OutMaterials.Num() > 0)
 		{
-			UMaterialInterface* OverlayMaterialInterface = Component.GetOverlayMaterial();
-			if (OverlayMaterialInterface != nullptr)
+			TArray<TObjectPtr<UMaterialInterface>> AssetAndComponentMaterialSlotsOverlayMaterial;
+			Component.GetMaterialSlotsOverlayMaterial(AssetAndComponentMaterialSlotsOverlayMaterial);
+			bool bUseGlobalMeshOverlayMaterial = false;
+			FMaterialOverlayHelper::AppendAllOverlayMaterial(AssetAndComponentMaterialSlotsOverlayMaterial, OutMaterials, bUseGlobalMeshOverlayMaterial);
+
+			if (bUseGlobalMeshOverlayMaterial)
 			{
-				OutMaterials.Add(OverlayMaterialInterface);
+				UMaterialInterface* OverlayMaterialInterface = Component.GetOverlayMaterial();
+				if (OverlayMaterialInterface != nullptr)
+				{
+					OutMaterials.Add(OverlayMaterialInterface);
+				}
 			}
 		}
 	}
@@ -118,16 +127,36 @@ void FStaticMeshComponentHelper::CollectPSOPrecacheDataImpl(const T& Component,
 		ComponentParams.PSOPrecacheParams = PrecachePSOParams;
 	}
 
-	UMaterialInterface* OverlayMaterialInterface = Component.GetOverlayMaterial();
-	if (OverlayMaterialInterface && VFTypesPerMaterialIndex.Num() != 0)
+	//Add all overlay materials
+	if (VFTypesPerMaterialIndex.Num() != 0)
 	{
-		// Overlay is rendered with the same set of VFs
-		FMaterialInterfacePSOPrecacheParams& ComponentParams = OutParams[OutParams.AddDefaulted()];
-
-		ComponentParams.MaterialInterface = OverlayMaterialInterface;
-		ComponentParams.VertexFactoryDataList = VFTypesPerMaterialIndex[0].VertexFactoryDataList;
-		ComponentParams.PSOPrecacheParams = PrecachePSOParams;
-		ComponentParams.PSOPrecacheParams.bCastShadow = false;
+		//Add all per section overlay material and add the global mesh only if not all sections has an overlay override
+		TArray<UMaterialInterface*> OverlayMaterials;
+		
+		TArray<TObjectPtr<UMaterialInterface>> AssetAndComponentMaterialSlotsOverlayMaterial;
+		Component.GetMaterialSlotsOverlayMaterial(AssetAndComponentMaterialSlotsOverlayMaterial);
+		bool bUseGlobalMeshOverlayMaterial = false;
+		FMaterialOverlayHelper::AppendAllOverlayMaterial(AssetAndComponentMaterialSlotsOverlayMaterial, OverlayMaterials, bUseGlobalMeshOverlayMaterial);
+		if (bUseGlobalMeshOverlayMaterial)
+		{
+			UMaterialInterface* OverlayMaterialInterface = Component.GetOverlayMaterial();
+			if (OverlayMaterialInterface)
+			{
+				OverlayMaterials.Add(OverlayMaterialInterface);
+			}
+		}
+		for (UMaterialInterface* OverlayMaterialToPreCache : OverlayMaterials)
+		{
+			if (OverlayMaterialToPreCache)
+			{
+				// Overlay is rendered with the same set of VFs
+				FMaterialInterfacePSOPrecacheParams& ComponentParams = OutParams[OutParams.AddDefaulted()];
+				ComponentParams.MaterialInterface = OverlayMaterialToPreCache;
+				ComponentParams.VertexFactoryDataList = VFTypesPerMaterialIndex[0].VertexFactoryDataList;
+				ComponentParams.PSOPrecacheParams = PrecachePSOParams;
+				ComponentParams.PSOPrecacheParams.bCastShadow = false;
+			}
+		}
 	}
 }
 
diff --git a/Engine/Source/Runtime/Engine/Public/StaticMeshSceneProxy.h b/Engine/Source/Runtime/Engine/Public/StaticMeshSceneProxy.h
--- a/Engine/Source/Runtime/Engine/Public/StaticMeshSceneProxy.h
+++ b/Engine/Source/Runtime/Engine/Public/StaticMeshSceneProxy.h
@@ -217,6 +217,8 @@ protected:
 			// The material index from the component. Used by the texture streaming accuracy viewmodes.
 			int32 MaterialIndex = 0;
 		#endif
+			/** The overlay material with which to render a second time this section. */
+			UMaterialInterface* OverlayMaterial = nullptr;
 		};
 
 		/** Per-section information. */
@@ -269,6 +271,7 @@ protected:
 
 	UMaterialInterface* OverlayMaterial;
 	float OverlayMaterialMaxDrawDistance;
+	TArray<TObjectPtr<UMaterialInterface>> MaterialSlotsOverlayMaterial;
 
 #if RHI_RAYTRACING
 	ENGINE_API void CreateDynamicRayTracingGeometries(FRHICommandListBase& RHICmdList);
diff --git a/Engine/Source/Runtime/Engine/Public/StaticMeshSceneProxyDesc.h b/Engine/Source/Runtime/Engine/Public/StaticMeshSceneProxyDesc.h
--- a/Engine/Source/Runtime/Engine/Public/StaticMeshSceneProxyDesc.h
+++ b/Engine/Source/Runtime/Engine/Public/StaticMeshSceneProxyDesc.h
@@ -54,6 +54,8 @@ struct FStaticMeshSceneProxyDesc : public FPrimitiveSceneProxyDesc
 	TArrayView<TObjectPtr<UMaterialInterface>>	OverrideMaterials;
 	TObjectPtr<class UMaterialInterface> OverlayMaterial;
 	float OverlayMaterialMaxDrawDistance = 0.0f;
+	TArray<TObjectPtr<UMaterialInterface>> MaterialSlotsOverlayMaterial;
+
 	int32 ForcedLodModel = 0;
 	int32 MinLOD = 0;
 	int32 WorldPositionOffsetDisableDistance = 0;
@@ -138,6 +140,7 @@ struct FStaticMeshSceneProxyDesc : public FPrimitiveSceneProxyDesc
 
 	UMaterialInterface* GetOverlayMaterial() const { return OverlayMaterial; }
 	float GetOverlayMaterialMaxDrawDistance() const { return OverlayMaterialMaxDrawDistance; }
+	void GetMaterialSlotsOverlayMaterial(TArray<TObjectPtr<UMaterialInterface>>& OutMaterialSlotsOverlayMaterial) const { OutMaterialSlotsOverlayMaterial = MaterialSlotsOverlayMaterial; }
 
 	UMaterialInterface* GetMaterial(int32 MaterialIndex, bool bDoingNaniteMaterialAudit = false, bool bIgnoreNaniteOverrideMaterials = false) const
 	{
