[Input] Add timestamp param to touch, motion, and gesture functions on the player controller and UPlayerInput

#tests ran Lyra on mobile in editor sim and in cooked. Also ran consoles with motion controller support

diff --git a/Engine/Plugins/Runtime/CommonUI/Source/CommonUI/Private/CommonGameViewportClient.cpp b/Engine/Plugins/Runtime/CommonUI/Source/CommonUI/Private/CommonGameViewportClient.cpp
--- a/Engine/Plugins/Runtime/CommonUI/Source/CommonUI/Private/CommonGameViewportClient.cpp
+++ b/Engine/Plugins/Runtime/CommonUI/Source/CommonUI/Private/CommonGameViewportClient.cpp
@@ -22,10 +22,6 @@ UCommonGameViewportClient::UCommonGameViewportClient(FVTableHelper& Helper) : Su
 {
 }
 
-UCommonGameViewportClient::~UCommonGameViewportClient()
-{
-}
-
 bool UCommonGameViewportClient::InputKey(const FInputKeyEventArgs& InEventArgs)
 {
 	FInputKeyEventArgs EventArgs = InEventArgs;
@@ -80,17 +76,26 @@ bool UCommonGameViewportClient::InputAxis(const FInputKeyEventArgs& Args)
 	return Super::InputAxis(Args);
 }
 
-bool UCommonGameViewportClient::InputTouch(FViewport* InViewport, int32 ControllerId, uint32 Handle, ETouchType::Type Type, const FVector2D& TouchLocation, float Force, FDateTime DeviceTimestamp, uint32 TouchpadIndex)
+bool UCommonGameViewportClient::InputTouch(FViewport* InViewport, const FInputDeviceId DeviceId, uint32 Handle, ETouchType::Type Type, const FVector2D& TouchLocation, float Force, uint32 TouchpadIndex, const uint64 Timestamp)
 {
 #if ALLOW_CONSOLE
-	if (ViewportConsole != NULL && (ViewportConsole->ConsoleState != NAME_Typing) && (ViewportConsole->ConsoleState != NAME_Open))
+	if (ViewportConsole != nullptr && (ViewportConsole->ConsoleState != NAME_Typing) && (ViewportConsole->ConsoleState != NAME_Open))
 #endif
 	{
 		FReply Result = FReply::Unhandled();
-		if (!OnRerouteTouch().ExecuteIfBound(ControllerId, Handle, Type, TouchLocation, Result))
+
+		// Remap the newer FInputDeviceId to the old int32 ControllerId for deprecated code. This can be removed
+		// when OnRerouteTouch() is removed.
+		const FPlatformUserId UserId = IPlatformInputDeviceMapper::Get().GetUserForInputDevice(DeviceId);
+		int32 ControllerId = IPlatformInputDeviceMapper::Get().GetUserIndexForPlatformUser(UserId);
+		 
+		PRAGMA_DISABLE_DEPRECATION_WARNINGS
+		if (!OnRerouteTouchInput().ExecuteIfBound(DeviceId, Handle, Type, TouchLocation, Result) &&
+			!OnRerouteTouch().ExecuteIfBound(ControllerId, Handle, Type, TouchLocation, Result))
 		{
-			HandleRerouteTouch(ControllerId, Handle, Type, TouchLocation, Result);
+			HandleRerouteTouch(DeviceId, Handle, Type, TouchLocation, Result);
 		}
+		PRAGMA_ENABLE_DEPRECATION_WARNINGS
 
 		if (Result.IsEventHandled())
 		{
@@ -98,7 +103,7 @@ bool UCommonGameViewportClient::InputTouch(FViewport* InViewport, int32 Controll
 		}
 	}
 
-	return Super::InputTouch(InViewport, ControllerId, Handle, Type, TouchLocation, Force, DeviceTimestamp, TouchpadIndex);
+	return Super::InputTouch(InViewport, DeviceId, Handle, Type, TouchLocation, Force, TouchpadIndex, Timestamp);
 }
 
 void UCommonGameViewportClient::MouseMove(FViewport* InViewport, int32 X, int32 Y)
@@ -144,15 +149,6 @@ void UCommonGameViewportClient::HandleRerouteInput(FInputDeviceId DeviceId, FKey
 	}	
 }
 
-void UCommonGameViewportClient::HandleRerouteInput(int32 ControllerId, FKey Key, EInputEvent EventType, FReply& Reply)
-{
-	// Remap the old int32 ControllerId to the new platform user and input device ID
-	FPlatformUserId UserId = FGenericPlatformMisc::GetPlatformUserForUserIndex(ControllerId);
-	FInputDeviceId DeviceID = INPUTDEVICEID_NONE;
-	IPlatformInputDeviceMapper::Get().RemapControllerIdToPlatformUserAndDevice(ControllerId, UserId, DeviceID);
-	return HandleRerouteInput(DeviceID, Key, EventType, Reply);
-}
-
 void UCommonGameViewportClient::HandleRerouteAxis(FInputDeviceId DeviceId, FKey Key, float Delta, FReply& Reply)
 {
 	// Get the ownign platform user for this input device and their local player
@@ -175,19 +171,9 @@ void UCommonGameViewportClient::HandleRerouteAxis(FInputDeviceId DeviceId, FKey
 	}
 }
 
-void UCommonGameViewportClient::HandleRerouteAxis(int32 ControllerId, FKey Key, float Delta, FReply& Reply)
+void UCommonGameViewportClient::HandleRerouteTouch(FInputDeviceId DeviceId, uint32 TouchId, ETouchType::Type TouchType, const FVector2D& TouchLocation, FReply& Reply)
 {
-	// Remap the old int32 ControllerId to the new platform user and input device ID
-	FPlatformUserId UserId = FGenericPlatformMisc::GetPlatformUserForUserIndex(ControllerId);
-	FInputDeviceId DeviceID = INPUTDEVICEID_NONE;
-	IPlatformInputDeviceMapper::Get().RemapControllerIdToPlatformUserAndDevice(ControllerId, UserId, DeviceID);
-	
-	return HandleRerouteAxis(DeviceID, Key, Delta, Reply);
-}
-
-void UCommonGameViewportClient::HandleRerouteTouch(int32 ControllerId, uint32 TouchId, ETouchType::Type TouchType, const FVector2D& TouchLocation, FReply& Reply)
-{
-	ULocalPlayer* LocalPlayer = GameInstance->FindLocalPlayerFromControllerId(ControllerId);
+	ULocalPlayer* LocalPlayer = GameInstance->FindLocalPlayerFromDeviceId(DeviceId);
 	Reply = FReply::Unhandled();
 
 	if (LocalPlayer && TouchId < EKeys::NUM_TOUCH_KEYS)
@@ -221,6 +207,18 @@ void UCommonGameViewportClient::HandleRerouteTouch(int32 ControllerId, uint32 To
 	}
 }
 
+PRAGMA_DISABLE_DEPRECATION_WARNINGS
+void UCommonGameViewportClient::HandleRerouteTouch(int32 ControllerId, uint32 TouchId, ETouchType::Type TouchType, const FVector2D& TouchLocation, FReply& Reply)
+{
+	// Remap the old int32 ControllerId to the new platform user and input device ID
+	FPlatformUserId UserId = FGenericPlatformMisc::GetPlatformUserForUserIndex(ControllerId);
+	FInputDeviceId DeviceID = INPUTDEVICEID_NONE;
+	IPlatformInputDeviceMapper::Get().RemapControllerIdToPlatformUserAndDevice(ControllerId, UserId, DeviceID);
+	
+	return HandleRerouteTouch(DeviceID, TouchId, TouchType, TouchLocation, Reply);
+}
+PRAGMA_ENABLE_DEPRECATION_WARNINGS
+
 bool UCommonGameViewportClient::IsKeyPriorityAboveUI(const FInputKeyEventArgs& EventArgs)
 {
 #if ALLOW_CONSOLE
diff --git a/Engine/Plugins/Runtime/CommonUI/Source/CommonUI/Public/CommonGameViewportClient.h b/Engine/Plugins/Runtime/CommonUI/Source/CommonUI/Public/CommonGameViewportClient.h
--- a/Engine/Plugins/Runtime/CommonUI/Source/CommonUI/Public/CommonGameViewportClient.h
+++ b/Engine/Plugins/Runtime/CommonUI/Source/CommonUI/Public/CommonGameViewportClient.h
@@ -9,8 +9,12 @@ class FReply;
 
 DECLARE_DELEGATE_FourParams(FOnRerouteInputDelegate, FInputDeviceId /* InputDeviceId */, FKey /* Key */, EInputEvent /* EventType */, FReply& /* Reply */);
 DECLARE_DELEGATE_FourParams(FOnRerouteAxisDelegate, FInputDeviceId /* InputDeviceId */, FKey /* Key */, float /* Delta */, FReply& /* Reply */);
+
+UE_DEPRECATED(5.6, "Use the version which takes an FInputDeviceId instead (FOnRerouteTouchInputDelegate).")
 DECLARE_DELEGATE_FiveParams(FOnRerouteTouchDelegate, int32 /* ControllerId */, uint32 /* TouchId */, ETouchType::Type /* TouchType */, const FVector2D& /* TouchLocation */, FReply& /* Reply */);
 
+DECLARE_DELEGATE_FiveParams(FOnRerouteTouchInputDelegate, FInputDeviceId /* Deviceid */, uint32 /* TouchId */, ETouchType::Type /* TouchType */, const FVector2D& /* TouchLocation */, FReply& /* Reply */);
+
 /**
 * CommonUI Viewport to reroute input to UI first. Needed to allow CommonUI to route / handle inputs.
 */
@@ -22,38 +26,38 @@ class COMMONUI_API UCommonGameViewportClient : public UGameViewportClient
 public:
 
 	UCommonGameViewportClient(FVTableHelper& Helper);
-	virtual ~UCommonGameViewportClient() override;
 
 	// UGameViewportClient interface begin
 	virtual bool InputKey(const FInputKeyEventArgs& EventArgs) override;
 	virtual bool InputAxis(const FInputKeyEventArgs& Args) override;
-	virtual bool InputTouch(FViewport* InViewport, int32 ControllerId, uint32 Handle, ETouchType::Type Type, const FVector2D& TouchLocation, float Force, FDateTime DeviceTimestamp, uint32 TouchpadIndex) override;
+	virtual bool InputTouch(FViewport* Viewport, const FInputDeviceId DeviceId, uint32 Handle, ETouchType::Type Type, const FVector2D& TouchLocation, float Force, uint32 TouchpadIndex, const uint64 Timestamp) override;
 	virtual void MouseMove(FViewport* InViewport, int32 X, int32 Y) override;
 	virtual void CapturedMouseMove(FViewport* InViewport, int32 X, int32 Y) override;
 	// UGameViewportClient interface end
 
 	FOnRerouteInputDelegate& OnRerouteInput() { return RerouteInput; }
 	FOnRerouteAxisDelegate& OnRerouteAxis() { return RerouteAxis; }
+	
+	UE_DEPRECATED(5.6, "Use the version which takes a FInputDeviceId instead (OnRerouteTouchInput).")
+	PRAGMA_DISABLE_DEPRECATION_WARNINGS
 	FOnRerouteTouchDelegate& OnRerouteTouch() { return RerouteTouch; }
+	PRAGMA_ENABLE_DEPRECATION_WARNINGS
+	
+	FOnRerouteTouchInputDelegate& OnRerouteTouchInput() { return RerouteTouchInput; }
 
 	FOnRerouteInputDelegate& OnRerouteBlockedInput() { return RerouteBlockedInput; }
 
-	/** Default Handler for Key input. */
-	UE_DEPRECATED(5.1, "This version of HandleRerouteInput has been deprecated. Please use the version that takes an FInputDeviceId instead")
-	virtual void HandleRerouteInput(int32 ControllerId, FKey Key, EInputEvent EventType, FReply& Reply);
-
 	/** Default Handler for Key input. */
 	virtual void HandleRerouteInput(FInputDeviceId DeviceId, FKey Key, EInputEvent EventType, FReply& Reply);
 
-	/** Default Handler for Axis input. */
-	UE_DEPRECATED(5.1, "This version of HandleRerouteAxis has been deprecated. Please use the version that takes an FInputDeviceId instead")
-	virtual void HandleRerouteAxis(int32 ControllerId, FKey Key, float Delta, FReply& Reply);
-
 	/** Default Handler for Axis input. */
 	virtual void HandleRerouteAxis(FInputDeviceId DeviceId, FKey Key, float Delta, FReply& Reply);
 
 	/** Default Handler for Touch input. */
-	virtual void HandleRerouteTouch(int32 ControllerId, uint32 TouchId, ETouchType::Type TouchType, const FVector2D& TouchLocation, FReply& Reply);
+	UE_DEPRECATED(5.6, "Use the version which takes a FInputDeviceId instead.")
+	virtual void HandleRerouteTouch(int32 ControllerId, uint32 TouchId, ETouchType::Type TouchType, const FVector2D& TouchLocation, FReply& Reply) final;
+	
+	virtual void HandleRerouteTouch(FInputDeviceId DeviceId, uint32 TouchId, ETouchType::Type TouchType, const FVector2D& TouchLocation, FReply& Reply);
 
 protected:
 	
@@ -62,7 +66,13 @@ protected:
 
 	FOnRerouteInputDelegate RerouteInput;
 	FOnRerouteAxisDelegate RerouteAxis;
+	
+	UE_DEPRECATED(5.6, "Use RerouteTouchInput instead.")
+	PRAGMA_DISABLE_DEPRECATION_WARNINGS
 	FOnRerouteTouchDelegate RerouteTouch;
+	PRAGMA_ENABLE_DEPRECATION_WARNINGS
+	
+	FOnRerouteTouchInputDelegate RerouteTouchInput;
 
 	FOnRerouteInputDelegate RerouteBlockedInput;
 };
diff --git a/Engine/Source/Editor/TextureEditor/Private/Models/TextureEditorViewportClient.cpp b/Engine/Source/Editor/TextureEditor/Private/Models/TextureEditorViewportClient.cpp
--- a/Engine/Source/Editor/TextureEditor/Private/Models/TextureEditorViewportClient.cpp
+++ b/Engine/Source/Editor/TextureEditor/Private/Models/TextureEditorViewportClient.cpp
@@ -626,7 +626,7 @@ EMouseCursor::Type FTextureEditorViewportClient::GetCursor(FViewport* Viewport,
 	return ShouldUseMousePanning(Viewport) ? EMouseCursor::GrabHandClosed : EMouseCursor::Default;
 }
 
-bool FTextureEditorViewportClient::InputGesture(FViewport* Viewport, EGestureEvent GestureType, const FVector2D& GestureDelta, bool bIsDirectionInvertedFromDevice)
+bool FTextureEditorViewportClient::InputGesture(FViewport* Viewport, const FInputDeviceId DeviceId, EGestureEvent GestureType, const FVector2D& GestureDelta, bool bIsDirectionInvertedFromDevice, const uint64 Timestamp)
 {
 	const bool LeftMouseButtonDown = Viewport->KeyState(EKeys::LeftMouseButton);
 	const bool RightMouseButtonDown = Viewport->KeyState(EKeys::RightMouseButton);
diff --git a/Engine/Source/Editor/TextureEditor/Private/Models/TextureEditorViewportClient.h b/Engine/Source/Editor/TextureEditor/Private/Models/TextureEditorViewportClient.h
--- a/Engine/Source/Editor/TextureEditor/Private/Models/TextureEditorViewportClient.h
+++ b/Engine/Source/Editor/TextureEditor/Private/Models/TextureEditorViewportClient.h
@@ -26,7 +26,7 @@ public:
 	virtual void Draw(FViewport* Viewport, FCanvas* Canvas) override;
 	virtual bool InputKey(const FInputKeyEventArgs& InEventArgs) override;
 	virtual bool InputAxis(const FInputKeyEventArgs& InEventArgs) override;
-	virtual bool InputGesture(FViewport* Viewport, EGestureEvent GestureType, const FVector2D& GestureDelta, bool bIsDirectionInvertedFromDevice) override;
+	virtual bool InputGesture(FViewport* Viewport, const FInputDeviceId DeviceId, EGestureEvent GestureType, const FVector2D& GestureDelta, bool bIsDirectionInvertedFromDevice, const uint64 Timestamp) override;
 	virtual UWorld* GetWorld() const override { return nullptr; }
 	virtual EMouseCursor::Type GetCursor(FViewport* Viewport, int32 X, int32 Y) override;
 
diff --git a/Engine/Source/Editor/UnrealEd/Private/EditorViewportClient.cpp b/Engine/Source/Editor/UnrealEd/Private/EditorViewportClient.cpp
--- a/Engine/Source/Editor/UnrealEd/Private/EditorViewportClient.cpp
+++ b/Engine/Source/Editor/UnrealEd/Private/EditorViewportClient.cpp
@@ -4846,7 +4846,7 @@ static float AdjustGestureCameraRotation(float Delta, float AdjustLimit, float D
 	return bIsUsingTrackpad ? Delta : FMath::Clamp(Delta, -DeltaCutoff, DeltaCutoff);
 }
 
-bool FEditorViewportClient::InputGesture(FViewport* InViewport, EGestureEvent GestureType, const FVector2D& GestureDelta, bool bIsDirectionInvertedFromDevice)
+bool FEditorViewportClient::InputGesture(FViewport* InViewport, const FInputDeviceId DeviceId, EGestureEvent GestureType, const FVector2D& GestureDelta, bool bIsDirectionInvertedFromDevice, const uint64 Timestamp)
 {
 	if (bDisableInput)
 	{
diff --git a/Engine/Source/Editor/UnrealEd/Public/EditorViewportClient.h b/Engine/Source/Editor/UnrealEd/Public/EditorViewportClient.h
--- a/Engine/Source/Editor/UnrealEd/Public/EditorViewportClient.h
+++ b/Engine/Source/Editor/UnrealEd/Public/EditorViewportClient.h
@@ -599,7 +599,7 @@ public:
 	UNREALED_API virtual void RequestInvalidateHitProxy(FViewport* Viewport) override;
 	UNREALED_API virtual bool InputKey(const FInputKeyEventArgs& EventArgs) override;
 	UNREALED_API virtual bool InputAxis(const FInputKeyEventArgs& Args) override;
-	UNREALED_API virtual bool InputGesture(FViewport* Viewport, EGestureEvent GestureType, const FVector2D& GestureDelta, bool bIsDirectionInvertedFromDevice) override;
+	UNREALED_API virtual bool InputGesture(FViewport* Viewport, const FInputDeviceId DeviceId, EGestureEvent GestureType, const FVector2D& GestureDelta, bool bIsDirectionInvertedFromDevice, const uint64 Timestamp) override;
 	UNREALED_API virtual void ReceivedFocus(FViewport* Viewport) override;
 	UNREALED_API virtual void MouseEnter(FViewport* Viewport,int32 x, int32 y) override;
 	UNREALED_API virtual void MouseMove(FViewport* Viewport,int32 x, int32 y) override;
diff --git a/Engine/Source/Runtime/Engine/Classes/Engine/GameViewportClient.h b/Engine/Source/Runtime/Engine/Classes/Engine/GameViewportClient.h
--- a/Engine/Source/Runtime/Engine/Classes/Engine/GameViewportClient.h
+++ b/Engine/Source/Runtime/Engine/Classes/Engine/GameViewportClient.h
@@ -172,8 +172,8 @@ public:
 	ENGINE_API virtual bool InputKey(const FInputKeyEventArgs& EventArgs) override;
 	ENGINE_API virtual bool InputAxis(const FInputKeyEventArgs& EventArgs) override;
 	ENGINE_API virtual bool InputChar(FViewport* Viewport,int32 ControllerId, TCHAR Character) override;
-	ENGINE_API virtual bool InputTouch(FViewport* Viewport, int32 ControllerId, uint32 Handle, ETouchType::Type Type, const FVector2D& TouchLocation, float Force, FDateTime DeviceTimestamp, uint32 TouchpadIndex) override;
-	ENGINE_API virtual bool InputMotion(FViewport* Viewport, int32 ControllerId, const FVector& Tilt, const FVector& RotationRate, const FVector& Gravity, const FVector& Acceleration) override;
+	ENGINE_API virtual bool InputTouch(FViewport* InViewport, const FInputDeviceId DeviceId, uint32 Handle, ETouchType::Type Type, const FVector2D& TouchLocation, float Force, uint32 TouchpadIndex, const uint64 Timestamp) override;
+	ENGINE_API virtual bool InputMotion(FViewport* Viewport, const FInputDeviceId DeviceId, const FVector& Tilt, const FVector& RotationRate, const FVector& Gravity, const FVector& Acceleration, const uint64 Timestamp) override;
 	ENGINE_API virtual EMouseCursor::Type GetCursor(FViewport* Viewport, int32 X, int32 Y ) override;
 	ENGINE_API virtual TOptional<TSharedRef<SWidget>> MapCursor(FViewport* Viewport, const FCursorReply& CursorReply) override;
 	ENGINE_API virtual void Precache() override;
diff --git a/Engine/Source/Runtime/Engine/Classes/GameFramework/PlayerController.h b/Engine/Source/Runtime/Engine/Classes/GameFramework/PlayerController.h
--- a/Engine/Source/Runtime/Engine/Classes/GameFramework/PlayerController.h
+++ b/Engine/Source/Runtime/Engine/Classes/GameFramework/PlayerController.h
@@ -1757,10 +1757,25 @@ public:
 	ENGINE_API virtual bool InputKey(const FInputKeyEventArgs& Params);
 
 	/** Handles a touch screen action */
-	ENGINE_API virtual bool InputTouch(uint32 Handle, ETouchType::Type Type, const FVector2D& TouchLocation, float Force, FDateTime DeviceTimestamp, uint32 TouchpadIndex);
+	UE_DEPRECATED(5.6, "Use the version which takes in a DeviceId and Timestamp instead.")
+	ENGINE_API virtual bool InputTouch(uint32 Handle, ETouchType::Type Type, const FVector2D& TouchLocation, float Force, FDateTime DeviceTimestamp, uint32 TouchpadIndex) final;
+
+	/** Handles a touch screen action */
+	ENGINE_API virtual bool InputTouch(
+		const FInputDeviceId DeviceId,
+		uint32 Handle,
+		ETouchType::Type Type,
+		const FVector2D& TouchLocation,
+		float Force,		
+		uint32 TouchpadIndex,
+		const uint64 Timestamp);
 	
 	/** Handles motion control */
-	ENGINE_API virtual bool InputMotion(const FVector& Tilt, const FVector& RotationRate, const FVector& Gravity, const FVector& Acceleration);
+	UE_DEPRECATED(5.6, "Use the version which takes in a DeviceId and Timestamp instead.")
+	ENGINE_API virtual bool InputMotion(const FVector& Tilt, const FVector& RotationRate, const FVector& Gravity, const FVector& Acceleration) final;
+
+	/** Handles motion control */
+	ENGINE_API virtual bool InputMotion(const FInputDeviceId DeviceId, const FVector& Tilt, const FVector& RotationRate, const FVector& Gravity, const FVector& Acceleration, const uint64 Timestamp);
 
 	/** Associate a new UPlayer with this PlayerController. */
 	ENGINE_API virtual void SetPlayer(UPlayer* InPlayer);
diff --git a/Engine/Source/Runtime/Engine/Classes/GameFramework/PlayerInput.h b/Engine/Source/Runtime/Engine/Classes/GameFramework/PlayerInput.h
--- a/Engine/Source/Runtime/Engine/Classes/GameFramework/PlayerInput.h
+++ b/Engine/Source/Runtime/Engine/Classes/GameFramework/PlayerInput.h
@@ -591,14 +591,44 @@ public:
 	ENGINE_API virtual bool InputKey(const FInputKeyEventArgs& Params);
 
 	/** Handles a touch input event.  Returns true. */
+	UE_DEPRECATED(5.6, "Use the version which takes a FInputDeviceId and Timestamp instead.")
 	ENGINE_API bool InputTouch(uint32 Handle, ETouchType::Type Type, const FVector2D& TouchLocation, float Force, FDateTime DeviceTimestamp, uint32 TouchpadIndex);
+	
+	/** Handles a touch input event.  Returns true. */
+	ENGINE_API bool InputTouch(
+		const FInputDeviceId DeviceId,
+		uint32 Handle,
+		ETouchType::Type Type,
+		const FVector2D& TouchLocation,
+		float Force,
+		uint32 TouchpadIndex,
+		const uint64 Timestamp);
 
 	/** Handles a motion input event.  Returns true. */
+	UE_DEPRECATED(5.6, "Use the version which takes in a FInputDeviceId and a timestamp instead.")
 	ENGINE_API bool InputMotion(const FVector& Tilt, const FVector& RotationRate, const FVector& Gravity, const FVector& Acceleration);
 
+	/** Handles a motion input event.  Returns true. */
+	ENGINE_API bool InputMotion(
+		const FInputDeviceId DeviceId,
+		const FVector& Tilt,
+		const FVector& RotationRate,
+		const FVector& Gravity,
+		const FVector& Acceleration,
+		const uint64 Timestamp);
+
 	/** Handles a gesture input event.  Returns true. */
+	UE_DEPRECATED(5.6, "Use the version which takes a device id and timestamp instead.")
 	ENGINE_API bool InputGesture(const FKey Gesture, const EInputEvent Event, const float Value);
 
+	/** Handles a gesture input event.  Returns true. */
+	ENGINE_API bool InputGesture(
+		const FInputDeviceId DeviceId,
+		const FKey Gesture,
+		const EInputEvent Event,
+		const float Value,
+		const uint64 Timestamp);
+
 	/** Manually update the GestureRecognizer AnchorDistance using the current locations of the touches */
 	ENGINE_API void UpdatePinchStartDistance();
 
diff --git a/Engine/Source/Runtime/Engine/Private/GameViewportClient.cpp b/Engine/Source/Runtime/Engine/Private/GameViewportClient.cpp
--- a/Engine/Source/Runtime/Engine/Private/GameViewportClient.cpp
+++ b/Engine/Source/Runtime/Engine/Private/GameViewportClient.cpp
@@ -888,7 +888,15 @@ bool UGameViewportClient::InputChar(FViewport* InViewport, int32 ControllerId, T
 	return bResult;
 }
 
-bool UGameViewportClient::InputTouch(FViewport* InViewport, int32 ControllerId, uint32 Handle, ETouchType::Type Type, const FVector2D& TouchLocation, float Force, FDateTime DeviceTimestamp, uint32 TouchpadIndex)
+bool UGameViewportClient::InputTouch(
+	FViewport* InViewport,
+	const FInputDeviceId DeviceId,
+	uint32 Handle,
+	ETouchType::Type Type,
+	const FVector2D& TouchLocation,
+	float Force,
+	uint32 TouchpadIndex,
+	const uint64 Timestamp)
 {
 	if (IgnoreInput())
 	{
@@ -896,20 +904,20 @@ bool UGameViewportClient::InputTouch(FViewport* InViewport, int32 ControllerId,
 	}
 
 	// route to subsystems that care
-	bool bResult = (ViewportConsole ? ViewportConsole->InputTouch(FInputDeviceId::CreateFromInternalId(ControllerId), Handle, Type, TouchLocation, Force, DeviceTimestamp, TouchpadIndex) : false);
+	bool bResult = (ViewportConsole ? ViewportConsole->InputTouch(DeviceId, Handle, Type, TouchLocation, Force, Timestamp, TouchpadIndex) : false);
 	if (!bResult)
 	{
-		ULocalPlayer* const TargetPlayer = GEngine->GetLocalPlayerFromControllerId(this, ControllerId);
+		ULocalPlayer* const TargetPlayer = GEngine->GetLocalPlayerFromInputDevice(this, DeviceId);
 		if (TargetPlayer && TargetPlayer->PlayerController)
 		{
-			bResult = TargetPlayer->PlayerController->InputTouch(Handle, Type, TouchLocation, Force, DeviceTimestamp, TouchpadIndex);
+			bResult = TargetPlayer->PlayerController->InputTouch(DeviceId, Handle, Type, TouchLocation, Force, TouchpadIndex, Timestamp);
 		}
 	}
 
 	return bResult;
 }
 
-bool UGameViewportClient::InputMotion(FViewport* InViewport, int32 ControllerId, const FVector& Tilt, const FVector& RotationRate, const FVector& Gravity, const FVector& Acceleration)
+bool UGameViewportClient::InputMotion(FViewport* InViewport, const FInputDeviceId DeviceId, const FVector& Tilt, const FVector& RotationRate, const FVector& Gravity, const FVector& Acceleration, const uint64 Timestamp)
 {
 	if (IgnoreInput() || !GetDefault<UInputSettings>()->bEnableMotionControls)
 	{
@@ -919,10 +927,10 @@ bool UGameViewportClient::InputMotion(FViewport* InViewport, int32 ControllerId,
 	// route to subsystems that care
 	bool bResult = false;
 
-	ULocalPlayer* const TargetPlayer = GEngine->GetLocalPlayerFromControllerId(this, ControllerId);
+	ULocalPlayer* const TargetPlayer = GEngine->GetLocalPlayerFromInputDevice(this, DeviceId);
 	if (TargetPlayer && TargetPlayer->PlayerController)
 	{
-		bResult = TargetPlayer->PlayerController->InputMotion(Tilt, RotationRate, Gravity, Acceleration);
+		bResult = TargetPlayer->PlayerController->InputMotion(DeviceId, Tilt, RotationRate, Gravity, Acceleration, Timestamp);
 	}
 
 	return bResult;
diff --git a/Engine/Source/Runtime/Engine/Private/GestureRecognizer.cpp b/Engine/Source/Runtime/Engine/Private/GestureRecognizer.cpp
--- a/Engine/Source/Runtime/Engine/Private/GestureRecognizer.cpp
+++ b/Engine/Source/Runtime/Engine/Private/GestureRecognizer.cpp
@@ -5,6 +5,8 @@
 =============================================================================*/
 
 #include "GestureRecognizer.h"
+
+#include "Engine/LocalPlayer.h"
 #include "GameFramework/PlayerInput.h"
 
 void FGestureRecognizer::DetectGestures(const FVector (&Touches)[EKeys::NUM_TOUCH_KEYS], UPlayerInput* PlayerInput, float DeltaTime)
@@ -136,8 +138,12 @@ void FGestureRecognizer::HandleGesture(UPlayerInput* PlayerInput, FKey Gesture,
 	float* Value = CurrentGestureValues.Find(Gesture);
 	if (Value)
 	{
+		const ULocalPlayer* OwningLP = PlayerInput->GetOwningLocalPlayer();
+		const FInputDeviceId DefaultDevice = IPlatformInputDeviceMapper::Get().GetPrimaryInputDeviceForUser(OwningLP ? OwningLP->GetPlatformUserId() : PLATFORMUSERID_NONE);
+		const uint64 CurrentTimestamp = FPlatformTime::Cycles64();
+		
 		const EInputEvent GestureEvent = (bStarted ? IE_Pressed : bEnded ? IE_Released : IE_Repeat);
-		PlayerInput->InputGesture(Gesture, GestureEvent, *Value);
+		PlayerInput->InputGesture(DefaultDevice, Gesture, GestureEvent, *Value, CurrentTimestamp);
 
 		// remove it if the gesture is complete
 		if (bEnded)
diff --git a/Engine/Source/Runtime/Engine/Private/PlayerController.cpp b/Engine/Source/Runtime/Engine/Private/PlayerController.cpp
--- a/Engine/Source/Runtime/Engine/Private/PlayerController.cpp
+++ b/Engine/Source/Runtime/Engine/Private/PlayerController.cpp
@@ -2438,13 +2438,34 @@ bool APlayerController::InputKey(const FInputKeyEventArgs& Params)
 	return bResult;
 }
 
+PRAGMA_DISABLE_DEPRECATION_WARNINGS
 bool APlayerController::InputTouch(uint32 Handle, ETouchType::Type Type, const FVector2D& TouchLocation, float Force, FDateTime DeviceTimestamp, uint32 TouchpadIndex)
+{
+	// Use the primary input device id for this player controller. The deprecated code path doesnt have any other way to grab it
+	// and this is a sensible default
+	const FInputDeviceId DeviceId = IPlatformInputDeviceMapper::Get().GetPrimaryInputDeviceForUser(GetPlatformUserId());
+
+	// A timestamp of 0 is invalid because we don't have any way to accurately get one here from deprecated callsites.
+	constexpr uint64 Timestamp = 0u;
+	
+	return InputTouch(DeviceId, Handle, Type, TouchLocation, Force, TouchpadIndex, Timestamp);
+}
+PRAGMA_ENABLE_DEPRECATION_WARNINGS
+
+bool APlayerController::InputTouch(
+	const FInputDeviceId DeviceId,
+	uint32 Handle,
+	ETouchType::Type Type,
+	const FVector2D& TouchLocation,
+	float Force,		
+	uint32 TouchpadIndex,
+	const uint64 Timestamp)
 {
 	PRAGMA_DISABLE_DEPRECATION_WARNINGS
 	if (GEngine->XRSystem.IsValid())
 	{
 		IXRInput* XRInput = GEngine->XRSystem->GetXRInput();
-		if(XRInput && XRInput->HandleInputTouch(Handle, Type, TouchLocation, DeviceTimestamp, TouchpadIndex))
+		if(XRInput && XRInput->HandleInputTouch(Handle, Type, TouchLocation, Timestamp, TouchpadIndex))
 		{
 			return true;
 		}
@@ -2454,7 +2475,7 @@ bool APlayerController::InputTouch(uint32 Handle, ETouchType::Type Type, const F
 	bool bResult = false;
 	if (PlayerInput)
 	{
-		bResult = PlayerInput->InputTouch(Handle, Type, TouchLocation, Force, DeviceTimestamp, TouchpadIndex);
+		bResult = PlayerInput->InputTouch(DeviceId, Handle, Type, TouchLocation, Force, TouchpadIndex, Timestamp);
 
 		if (bEnableTouchEvents || bEnableTouchOverEvents)
 		{
@@ -2526,13 +2547,27 @@ bool APlayerController::InputTouch(uint32 Handle, ETouchType::Type Type, const F
 	return bResult;
 }
 
+PRAGMA_DISABLE_DEPRECATION_WARNINGS
 bool APlayerController::InputMotion(const FVector& Tilt, const FVector& RotationRate, const FVector& Gravity, const FVector& Acceleration)
+{
+	// Use the primary input device id for this player controller. The deprecated code path doesnt have any other way to grab it
+	// and this is a sensible default
+	const FInputDeviceId DeviceId = IPlatformInputDeviceMapper::Get().GetPrimaryInputDeviceForUser(GetPlatformUserId());
+
+	// A timestamp of 0 is invalid because we don't have any way to accurately get one here from deprecated callsites.
+	constexpr uint64 Timestamp = 0u;
+
+	return InputMotion(DeviceId, Tilt, RotationRate, Gravity, Acceleration, Timestamp);
+}
+PRAGMA_ENABLE_DEPRECATION_WARNINGS
+
+bool APlayerController::InputMotion(const FInputDeviceId DeviceId, const FVector& Tilt, const FVector& RotationRate, const FVector& Gravity, const FVector& Acceleration, const uint64 Timestamp)
 {
 	bool bResult = false;
 
 	if (PlayerInput && bEnableMotionControls)
 	{
-		bResult = PlayerInput->InputMotion(Tilt, RotationRate, Gravity, Acceleration);
+		bResult = PlayerInput->InputMotion(DeviceId, Tilt, RotationRate, Gravity, Acceleration, Timestamp);
 	}
 
 	return bResult;
diff --git a/Engine/Source/Runtime/Engine/Private/Slate/SceneViewport.cpp b/Engine/Source/Runtime/Engine/Private/Slate/SceneViewport.cpp
--- a/Engine/Source/Runtime/Engine/Private/Slate/SceneViewport.cpp
+++ b/Engine/Source/Runtime/Engine/Private/Slate/SceneViewport.cpp
@@ -856,7 +856,7 @@ FReply FSceneViewport::OnTouchStarted( const FGeometry& MyGeometry, const FPoint
 
 		const FVector2D TouchPosition = CachedCursorPos;
 
-		if(ViewportClient->InputTouch( this, TouchEvent.GetUserIndex(), TouchEvent.GetPointerIndex(), ETouchType::Began, TouchPosition, TouchEvent.GetTouchForce(), FDateTime::Now(), TouchEvent.GetTouchpadIndex()) )
+		if(ViewportClient->InputTouch( this, TouchEvent.GetInputDeviceId(), TouchEvent.GetPointerIndex(), ETouchType::Began, TouchPosition, TouchEvent.GetTouchForce(), TouchEvent.GetTouchpadIndex(), TouchEvent.GetEventTimestamp()) )
 		{
 			const bool bTemporaryCapture = ViewportClient->GetMouseCaptureMode() == EMouseCaptureMode::CaptureDuringMouseDown;
 			if (bTemporaryCapture)
@@ -887,7 +887,7 @@ FReply FSceneViewport::OnTouchMoved( const FGeometry& MyGeometry, const FPointer
 		FScopedConditionalWorldSwitcher WorldSwitcher( ViewportClient );
 
 
-		if( !ViewportClient->InputTouch( this, TouchEvent.GetUserIndex(), TouchEvent.GetPointerIndex(), ETouchType::Moved, CachedCursorPos, TouchEvent.GetTouchForce(), FDateTime::Now(), TouchEvent.GetTouchpadIndex()) )
+		if( !ViewportClient->InputTouch( this, TouchEvent.GetInputDeviceId(), TouchEvent.GetPointerIndex(), ETouchType::Moved, CachedCursorPos, TouchEvent.GetTouchForce(), TouchEvent.GetTouchpadIndex(), TouchEvent.GetEventTimestamp()) )
 		{
 			CurrentReplyState = FReply::Unhandled(); 
 		}
@@ -922,7 +922,7 @@ FReply FSceneViewport::OnTouchEnded( const FGeometry& MyGeometry, const FPointer
 		// Switch to the viewport clients world before processing input
 		FScopedConditionalWorldSwitcher WorldSwitcher( ViewportClient );
 
-		if( !ViewportClient->InputTouch( this, TouchEvent.GetUserIndex(), TouchEvent.GetPointerIndex(), ETouchType::Ended, CurCursorPos, 0.0f, FDateTime::Now(), TouchEvent.GetTouchpadIndex()) )
+		if( !ViewportClient->InputTouch( this, TouchEvent.GetInputDeviceId(), TouchEvent.GetPointerIndex(), ETouchType::Ended, CurCursorPos, 0.0f, TouchEvent.GetTouchpadIndex(), TouchEvent.GetEventTimestamp()) )
 		{
 			CurrentReplyState = FReply::Unhandled(); 
 		}
@@ -951,7 +951,7 @@ FReply FSceneViewport::OnTouchForceChanged(const FGeometry& MyGeometry, const FP
 
 		const FVector2D TouchPosition = MyGeometry.AbsoluteToLocal(TouchEvent.GetScreenSpacePosition()) * MyGeometry.Scale;
 
-		if (!ViewportClient->InputTouch(this, TouchEvent.GetUserIndex(), TouchEvent.GetPointerIndex(), ETouchType::ForceChanged, TouchPosition, TouchEvent.GetTouchForce(), FDateTime::Now(), TouchEvent.GetTouchpadIndex()))
+		if (!ViewportClient->InputTouch(this, TouchEvent.GetInputDeviceId(), TouchEvent.GetPointerIndex(), ETouchType::ForceChanged, TouchPosition, TouchEvent.GetTouchForce(), TouchEvent.GetTouchpadIndex(), TouchEvent.GetEventTimestamp()))
 		{
 			CurrentReplyState = FReply::Unhandled();
 		}
@@ -975,7 +975,7 @@ FReply FSceneViewport::OnTouchFirstMove(const FGeometry& MyGeometry, const FPoin
 
 		const FVector2D TouchPosition = MyGeometry.AbsoluteToLocal(TouchEvent.GetScreenSpacePosition()) * MyGeometry.Scale;
 
-		if (!ViewportClient->InputTouch(this, TouchEvent.GetUserIndex(), TouchEvent.GetPointerIndex(), ETouchType::FirstMove, TouchPosition, TouchEvent.GetTouchForce(), FDateTime::Now(), TouchEvent.GetTouchpadIndex()))
+		if (!ViewportClient->InputTouch(this, TouchEvent.GetInputDeviceId(), TouchEvent.GetPointerIndex(), ETouchType::FirstMove, TouchPosition, TouchEvent.GetTouchForce(), TouchEvent.GetTouchpadIndex(), TouchEvent.GetEventTimestamp()))
 		{
 			CurrentReplyState = FReply::Unhandled();
 		}
@@ -1002,7 +1002,7 @@ FReply FSceneViewport::OnTouchGesture( const FGeometry& MyGeometry, const FPoint
 			FSlateApplication::Get().SetKeyboardFocus(ViewportWidget.Pin());
 		}
 
-		if( !ViewportClient->InputGesture( this, GestureEvent.GetGestureType(), GestureEvent.GetGestureDelta(), GestureEvent.IsDirectionInvertedFromDevice() ) )
+		if(!ViewportClient->InputGesture(this, GestureEvent.GetInputDeviceId(), GestureEvent.GetGestureType(), GestureEvent.GetGestureDelta(), GestureEvent.IsDirectionInvertedFromDevice(), GestureEvent.GetEventTimestamp()))
 		{
 			CurrentReplyState = FReply::Unhandled();
 		}
@@ -1016,12 +1016,12 @@ FReply FSceneViewport::OnMotionDetected( const FGeometry& MyGeometry, const FMot
 	// Start a new reply state
 	CurrentReplyState = FReply::Handled(); 
 
-	if( ViewportClient )
+	if(ViewportClient) 
 	{
 		// Switch to the viewport clients world before processing input
 		FScopedConditionalWorldSwitcher WorldSwitcher( ViewportClient );
 
-		if( !ViewportClient->InputMotion( this, MotionEvent.GetUserIndex(), MotionEvent.GetTilt(), MotionEvent.GetRotationRate(), MotionEvent.GetGravity(), MotionEvent.GetAcceleration()) )
+		if(!ViewportClient->InputMotion( this, MotionEvent.GetInputDeviceId(), MotionEvent.GetTilt(), MotionEvent.GetRotationRate(), MotionEvent.GetGravity(), MotionEvent.GetAcceleration(), MotionEvent.GetEventTimestamp()))
 		{
 			CurrentReplyState = FReply::Unhandled(); 
 		}
diff --git a/Engine/Source/Runtime/Engine/Private/UnrealClient.cpp b/Engine/Source/Runtime/Engine/Private/UnrealClient.cpp
--- a/Engine/Source/Runtime/Engine/Private/UnrealClient.cpp
+++ b/Engine/Source/Runtime/Engine/Private/UnrealClient.cpp
@@ -2557,6 +2557,67 @@ bool FViewportClient::InputAxis(
 
 	return InputAxis(Args);
 }
+
+bool FViewportClient::InputTouch(FViewport* Viewport, int32 ControllerId, uint32 Handle, ETouchType::Type Type, const FVector2D& TouchLocation, float Force, FDateTime DeviceTimestamp, uint32 TouchpadIndex)
+{
+	// Get the newer FInputDeviceId from the legayc ControllerId
+	FPlatformUserId UserId = PLATFORMUSERID_NONE;
+	FInputDeviceId DeviceId = INPUTDEVICEID_NONE;
+	IPlatformInputDeviceMapper::Get().RemapControllerIdToPlatformUserAndDevice(ControllerId, OUT UserId, OUT DeviceId);
+	
+	// Note: Set the timestamp to 0u here to represent an "invalid" timestamp
+	constexpr uint64 Timestamp = 0u;
+	
+	return InputTouch(
+		Viewport,
+		DeviceId,
+		Handle,
+		Type,
+		TouchLocation,
+		Force,
+		TouchpadIndex,
+		Timestamp);
+}
+
+bool FViewportClient::InputGesture(FViewport* Viewport, EGestureEvent GestureType, const FVector2D& GestureDelta, bool bIsDirectionInvertedFromDevice)
+{
+	// Use the default platform user's primary device for any upgrade path to try and get the best behavior.
+	// Legacy behavior is that all touch/gesture events go through this primary user
+	const IPlatformInputDeviceMapper& Mapper = IPlatformInputDeviceMapper::Get();
+	const FPlatformUserId PrimaryUser = Mapper.GetPrimaryPlatformUser();
+	const FInputDeviceId DeviceId = IPlatformInputDeviceMapper::Get().GetPrimaryInputDeviceForUser(PrimaryUser);
+	
+	// Note: Set the timestamp to 0u here to represent an "invalid" timestamp
+	constexpr uint64 Timestamp = 0u;
+	
+	return InputGesture(
+		Viewport,
+		DeviceId,
+		GestureType,
+		GestureDelta,
+		bIsDirectionInvertedFromDevice,
+		Timestamp);
+}
+
+bool FViewportClient::InputMotion(FViewport* Viewport, int32 ControllerId, const FVector& Tilt, const FVector& RotationRate, const FVector& Gravity, const FVector& Acceleration)
+{
+	FPlatformUserId UserId = PLATFORMUSERID_NONE;
+	FInputDeviceId DeviceId = INPUTDEVICEID_NONE;
+	IPlatformInputDeviceMapper::Get().RemapControllerIdToPlatformUserAndDevice(ControllerId, OUT UserId, OUT DeviceId);
+	
+	// Note: Set the timestamp to 0u here to represent an "invalid" timestamp
+	constexpr uint64 Timestamp = 0u;
+
+	return InputMotion(
+		Viewport,
+		DeviceId,
+		Tilt,
+		RotationRate,
+		Gravity,
+		Acceleration,
+		Timestamp);
+}
+
 PRAGMA_ENABLE_DEPRECATION_WARNINGS
 
 TOptional<TSharedRef<SWidget>> FViewportClient::MapCursor(FViewport* Viewport, const FCursorReply& CursorReply)
diff --git a/Engine/Source/Runtime/Engine/Private/UserInterface/PlayerInput.cpp b/Engine/Source/Runtime/Engine/Private/UserInterface/PlayerInput.cpp
--- a/Engine/Source/Runtime/Engine/Private/UserInterface/PlayerInput.cpp
+++ b/Engine/Source/Runtime/Engine/Private/UserInterface/PlayerInput.cpp
@@ -466,7 +466,27 @@ bool UPlayerInput::InputKey(const FInputKeyEventArgs& Params)
 	}
 }
 
+PRAGMA_DISABLE_DEPRECATION_WARNINGS
 bool UPlayerInput::InputTouch(uint32 Handle, ETouchType::Type Type, const FVector2D& TouchLocation, float Force, FDateTime DeviceTimestamp, uint32 TouchpadIndex)
+{
+	// No way to get the input device from deprecated code here
+	const FInputDeviceId DeviceId = INPUTDEVICEID_NONE;
+
+	// A timestamp of 0 is invalid because we don't have any way to accurately get one here from deprecated callsites.
+	constexpr uint64 Timestamp = 0u;
+	
+	return InputTouch(DeviceId, Handle, Type, TouchLocation, Force, TouchpadIndex, Timestamp);
+}
+PRAGMA_ENABLE_DEPRECATION_WARNINGS
+
+bool UPlayerInput::InputTouch(
+		const FInputDeviceId DeviceId,
+		uint32 Handle,
+		ETouchType::Type Type,
+		const FVector2D& TouchLocation,
+		float Force,
+		uint32 TouchpadIndex,
+		const uint64 Timestamp)
 {
 	// if touch disabled from the command line, consume the input and do nothing
 	static bool bTouchDisabled = FParse::Param(FCommandLine::Get(), TEXT("DisableTouch"));
@@ -542,7 +562,26 @@ bool UPlayerInput::InputTouch(uint32 Handle, ETouchType::Type Type, const FVecto
 	return true;
 }
 
-bool UPlayerInput::InputMotion(const FVector& InTilt, const FVector& InRotationRate, const FVector& InGravity, const FVector& InAcceleration)
+PRAGMA_DISABLE_DEPRECATION_WARNINGS
+bool UPlayerInput::InputMotion(const FVector& Tilt, const FVector& RotationRate, const FVector& Gravity, const FVector& Acceleration)
+{
+	// No way to get the input device from deprecated code here
+	const FInputDeviceId DeviceId = INPUTDEVICEID_NONE;
+
+	// A timestamp of 0 is invalid because we don't have any way to accurately get one here from deprecated callsites.
+	constexpr uint64 Timestamp = 0u;
+
+	return InputMotion(DeviceId, Tilt, RotationRate, Gravity, Acceleration, Timestamp);
+}
+PRAGMA_ENABLE_DEPRECATION_WARNINGS
+
+bool UPlayerInput::InputMotion(
+	const FInputDeviceId DeviceId,
+	const FVector& InTilt,
+	const FVector& InRotationRate,
+	const FVector& InGravity,
+	const FVector& InAcceleration,
+	const uint64 Timestamp)
 {
 	FKeyState& KeyStateTilt = KeyStateMap.FindOrAdd(EKeys::Tilt);
 	KeyStateTilt.RawValueAccumulator += InTilt;
@@ -569,7 +608,25 @@ bool UPlayerInput::InputMotion(const FVector& InTilt, const FVector& InRotationR
 	return true;
 }
 
+PRAGMA_DISABLE_DEPRECATION_WARNINGS
 bool UPlayerInput::InputGesture(const FKey Gesture, const EInputEvent Event, const float Value)
+{
+	// No way to get the input device from deprecated code here
+	const FInputDeviceId DeviceId = INPUTDEVICEID_NONE;
+
+	// A timestamp of 0 is invalid because we don't have any way to accurately get one here from deprecated callsites.
+	constexpr uint64 Timestamp = 0u;
+
+	return InputGesture(DeviceId, Gesture, Event, Value, Timestamp);	
+}
+PRAGMA_ENABLE_DEPRECATION_WARNINGS
+
+bool UPlayerInput::InputGesture(
+		const FInputDeviceId DeviceId,
+		const FKey Gesture,
+		const EInputEvent Event,
+		const float Value,
+		const uint64 Timestamp)
 {
 	FKeyState& KeyState = KeyStateMap.FindOrAdd(Gesture);
 
@@ -1786,8 +1843,26 @@ void UPlayerInput::DisplayDebug(class UCanvas* Canvas, const FDebugDisplayInfo&
 			FKeyState const* const KeyState = &It.Value();
 			FKey const Key = It.Key();
 
+			auto GetKeyStateString = [Key, KeyState]() -> FString
+			{
+				if (Key.IsAxis3D())
+				{
+					return FString::Printf(TEXT(" %s: X: %.2f Y: %.2f Z: %.2f (raw X: %.2f Y: %.2f Z: %.2f)"), *(Key.GetDisplayName().ToString()),
+						KeyState->Value.X, KeyState->Value.Y, KeyState->Value.Z,
+						KeyState->RawValue.X, KeyState->RawValue.Y, KeyState->RawValue.Z);
+				}
+				else if (Key.IsAxis2D())
+				{
+					return FString::Printf(TEXT(" %s: X: %.2f Y: %.2f (raw X: %.2f Y: %.2f)"), *(Key.GetDisplayName().ToString()),
+						KeyState->Value.X, KeyState->Value.Y,
+						KeyState->RawValue.X, KeyState->RawValue.Y);
+				}
+				return FString::Printf(TEXT(" %s: %.2f (raw %.2f)"), *(Key.GetDisplayName().ToString()), KeyState->Value.X, KeyState->RawValue.X);
+			};
+
 			// used cached mouse values, since they were flushed already
-			FString Str = FString::Printf(TEXT(" %s: %.2f (raw %.2f)"), *(Key.GetDisplayName().ToString()), KeyState->Value.X, KeyState->RawValue.X);
+			
+			FString Str = GetKeyStateString();
 
 			if ( KeyState->bDown || (KeyState->Value.X != 0.f) )
 			{
diff --git a/Engine/Source/Runtime/Engine/Public/ViewportClient.h b/Engine/Source/Runtime/Engine/Public/ViewportClient.h
--- a/Engine/Source/Runtime/Engine/Public/ViewportClient.h
+++ b/Engine/Source/Runtime/Engine/Public/ViewportClient.h
@@ -83,7 +83,6 @@ public:
 	 * @param	Viewport - The viewport which the axis movement is from.
 	 * @param	ControllerId - The controller which the axis movement is from.
 	 * @param	Character - The character.
-	 * // TODO_BH: Timestamp
 	 * @return	True to consume the character, false to pass it on.
 	 */
 	virtual bool InputChar(FViewport* Viewport,int32 ControllerId,TCHAR Character) { return false; }
@@ -99,10 +98,36 @@ public:
 	 * @param	Force - How hard the touch is
 	 * @param	DeviceTimestamp - Timestamp of the event
 	 * @param	TouchpadIndex - For devices with multiple touchpads, this is the index of which one
-	 * // TODO_BH: Timestamp
 	 * @return	True to consume the key event, false to pass it on.
 	 */
-	virtual bool InputTouch(FViewport* Viewport, int32 ControllerId, uint32 Handle, ETouchType::Type Type, const FVector2D& TouchLocation, float Force, FDateTime DeviceTimestamp, uint32 TouchpadIndex) { return false; }
+	UE_DEPRECATED(5.6, "Use the version which takes a uint64 timestamp instead")
+	ENGINE_API virtual bool InputTouch(FViewport* Viewport, int32 ControllerId, uint32 Handle, ETouchType::Type Type, const FVector2D& TouchLocation, float Force, FDateTime DeviceTimestamp, uint32 TouchpadIndex) final;
+
+	/**
+	 * Check a key event received by the viewport.
+	 * If the viewport client uses the event, it should return true to consume it.
+	 * @param	Viewport - The viewport which the event is from.
+	 * @param	DeviceId - The input device id from which this event is from
+	 * @param	Handle - Identifier unique to this touch event
+	 * @param	Type - What kind of touch event this is (see ETouchType)
+	 * @param	TouchLocation - Screen position of the touch
+	 * @param	Force - How hard the touch is
+	 * @param	TouchpadIndex - For devices with multiple touchpads, this is the index of which one
+	 * @param	Timestamp - The timestamp of when this input event was originally polled in the form of CPU cycles.
+	 * @return	True to consume the key event, false to pass it on.
+	 */
+	virtual bool InputTouch(
+		FViewport* Viewport,
+		const FInputDeviceId DeviceId,
+		uint32 Handle,
+		ETouchType::Type Type,
+		const FVector2D& TouchLocation,
+		float Force,		
+		uint32 TouchpadIndex,
+		const uint64 Timestamp)
+	{
+		return false;
+	}
 
 	/**
 	 * Check a gesture event received by the viewport.
@@ -110,10 +135,33 @@ public:
 	 * @param	Viewport - The viewport which the event is from.
 	 * @param	GestureType - @todo desc
 	 * @param	GestureDelta - @todo desc
-	 * // TODO_BH: Timestamp
 	 * @return	True to consume the gesture event, false to pass it on.
 	 */
-	virtual bool InputGesture(FViewport* Viewport, EGestureEvent GestureType, const FVector2D& GestureDelta, bool bIsDirectionInvertedFromDevice) { return false; }
+	UE_DEPRECATED(5.6, "Use the version which takes in a FInputDeviceId and uint64 Timestamp instead.")
+	ENGINE_API virtual bool InputGesture(FViewport* Viewport, EGestureEvent GestureType, const FVector2D& GestureDelta, bool bIsDirectionInvertedFromDevice) final;
+
+	/**
+	 * Check a gesture event received by the viewport.
+	 * If the viewport client uses the event, it should return true to consume it.
+	 * 
+	 * @param	Viewport - The viewport which the event is from.
+	 * @param	DeviceId - The input device id from which this gesture came from
+	 * @param	GestureType - The type of gesture being input
+	 * @param	GestureDelta - Delta values of this gesture from the previous instance of it
+	 * @param	Timestamp - The timestamp of when this input event was originally polled in the form of CPU cycles.
+	 * 
+	 * @return	True to consume the gesture event, false to pass it on.
+	 */
+	virtual bool InputGesture(
+		FViewport* Viewport,
+		const FInputDeviceId DeviceId,
+		EGestureEvent GestureType,
+		const FVector2D& GestureDelta,
+		bool bIsDirectionInvertedFromDevice,
+		const uint64 Timestamp)
+	{
+		return false;
+	}
 
 	/**
 	 * Each frame, the input system will update the motion data.
@@ -124,10 +172,34 @@ public:
 	 * @param RotationRate	How fast the tilt is changing
 	 * @param Gravity		Describes the current gravity of the device
 	 * @param Acceleration  Describes the acceleration of the device
-	 * // TODO_BH: Timestamp
 	 * @return	True to consume the motion event, false to pass it on.
 	 */
-	virtual bool InputMotion(FViewport* Viewport, int32 ControllerId, const FVector& Tilt, const FVector& RotationRate, const FVector& Gravity, const FVector& Acceleration) { return false; }
+	UE_DEPRECATED(5.6, "Use the version which takes in a FInputDeviceId and a timestamp instead")
+	ENGINE_API virtual bool InputMotion(FViewport* Viewport, int32 ControllerId, const FVector& Tilt, const FVector& RotationRate, const FVector& Gravity, const FVector& Acceleration) final;
+
+	/**
+	 * Each frame, the input system will update the motion data.
+	 *
+	 * @param Viewport - The viewport which the key event is from.
+	 * @param DeviceId - The input device which this event is from
+	 * @param Tilt			The current orientation of the device
+	 * @param RotationRate	How fast the tilt is changing
+	 * @param Gravity		Describes the current gravity of the device
+	 * @param Acceleration  Describes the acceleration of the device
+	 * @param Timestamp The timestamp of when this input event was originally polled in the form of CPU cycles.
+	 * @return	True to consume the motion event, false to pass it on.
+	 */
+	virtual bool InputMotion(
+		FViewport* Viewport,
+		const FInputDeviceId DeviceId,
+		const FVector& Tilt,
+		const FVector& RotationRate,
+		const FVector& Gravity,
+		const FVector& Acceleration,
+		const uint64 Timestamp)
+	{
+		return false;
+	}
 
 	virtual void SetIsSimulateInEditorViewport(bool bInIsSimulateInEditorViewport) { };
 
