Only perform parallel actions on skeletal mesh components outside of a transaction.

diff --git a/Engine/Source/Runtime/Engine/Private/Components/SkeletalMeshComponent.cpp b/Engine/Source/Runtime/Engine/Private/Components/SkeletalMeshComponent.cpp
--- a/Engine/Source/Runtime/Engine/Private/Components/SkeletalMeshComponent.cpp
+++ b/Engine/Source/Runtime/Engine/Private/Components/SkeletalMeshComponent.cpp
@@ -30,6 +30,7 @@
 #include "SkinnedAssetCompiler.h"
 #include "Rendering/SkeletalMeshRenderData.h"
 #include "Rendering/RenderCommandPipes.h"
+#include "AutoRTFM/AutoRTFM.h"
 
 #include "Logging/MessageLog.h"
 #include "Animation/AnimNode_LinkedInputPose.h"
@@ -2829,6 +2830,16 @@ void USkeletalMeshComponent::SwapEvaluationContextBuffers()
 
 void USkeletalMeshComponent::DispatchParallelEvaluationTasks(FActorComponentTickFunction* TickFunction)
 {
+	// If we try and start parallel evaluation in a transaction, just run it immediately instead as it
+	// is unsafe to have a transactional thread interacting with anything async.
+	if (AutoRTFM::IsTransactional())
+	{
+		ensureMsgf(false, TEXT("Should not be calling `DispatchParallelEvaluationTasks` from within a transaction!"));
+		check(!IsValidRef(ParallelAnimationEvaluationTask));
+		DoParallelEvaluationTasks_OnGameThread();
+		return;
+	}
+
 	LLM_SCOPE(ELLMTag::SkeletalMesh);
 	SwapEvaluationContextBuffers();
 
@@ -4734,6 +4745,26 @@ bool USkeletalMeshComponent::HandleExistingParallelEvaluationTask(bool bBlockOnT
 	{
 		if (bBlockOnTask)
 		{
+			// If we try and complete existing parallel evaluation in a transaction, we have to abort the transaction
+			// and retry up to this point. This is because we can only have started parallel evaluation outside of a
+			// transaction, and thus it is only safe to complete the work outside of a transaction too! For instance
+			// if the paralle animation evaluation task depended on something that was trying to take a mutex that we
+			// had also taken in the transaction, we'd have caused a deadlock if we didn't have this here.
+			if (AutoRTFM::IsTransactional())
+			{
+				(void)AutoRTFM::Close([&]
+				{
+					AutoRTFM::CascadingRetryTransaction([strong_this = TObjectPtr<USkeletalMeshComponent>(this), bPerformPostAnimEvaluation]
+					{
+						UE_LOG(LogSkeletalMesh, Display, TEXT("A call to `HandleExistingParallelEvaluationTask` that is blocking is causing Verse to abort and retry."));
+						const bool bBlockOnTask = true; // Was true from the if block above!
+						strong_this->HandleExistingParallelEvaluationTask(bBlockOnTask, bPerformPostAnimEvaluation);
+						UE_LOG(LogSkeletalMesh, Display, TEXT("`HandleExistingParallelEvaluationTask` has completed after Verse aborted, and we are now retrying."));
+					});
+				});
+				return true;
+			}
+
 			TRACE_CPUPROFILER_EVENT_SCOPE(USkeletalMeshComponent::BlockOnParallelEvaluationTask);
 			check(IsInGameThread()); // Only attempt this from game thread!
 			FTaskGraphInterface::Get().WaitUntilTaskCompletes(ParallelAnimationEvaluationTask, ENamedThreads::GameThread);
