VolumetricCloud - changed the material UPROPERTY to be soft reference with synchronous loading on register into ObjectReference to keep the material around once loaded.
This allows to have a material not loaded by default (soft reference) and not cost anything. Only the actual material assigned on the component will be loaded.
Also updated the SetMaterial function to account for the new logic.

Updated the SunSky asset accordingly.

#rnx
#tests PC editor dx12, and game win64, component and component in blueprint.

#virtualized

diff --git a/Engine/Source/Runtime/Engine/Classes/Components/VolumetricCloudComponent.h b/Engine/Source/Runtime/Engine/Classes/Components/VolumetricCloudComponent.h
--- a/Engine/Source/Runtime/Engine/Classes/Components/VolumetricCloudComponent.h
+++ b/Engine/Source/Runtime/Engine/Classes/Components/VolumetricCloudComponent.h
@@ -68,7 +68,7 @@ class UVolumetricCloudComponent : public USceneComponent
 
 	/** The material describing the cloud volume. It must be a Volume domain material. */
 	UPROPERTY(EditAnywhere, BlueprintReadOnly, Category = "Cloud Material")
-	TObjectPtr<UMaterialInterface> Material;
+	TSoftObjectPtr<UMaterialInterface> Material;
 
 	/** Whether to apply atmosphere transmittance per sample, instead of using the light global transmittance. */
 	UPROPERTY(EditAnywhere, BlueprintReadOnly, Category = "Cloud Tracing")
@@ -181,7 +181,7 @@ class UVolumetricCloudComponent : public USceneComponent
 	UFUNCTION(BlueprintCallable, Category = "Rendering")
 	ENGINE_API void SetMaterial(UMaterialInterface* NewValue);
 
-	ENGINE_API UMaterialInterface* GetMaterial() const { return Material; }
+	ENGINE_API UMaterialInterface* GetMaterial() const { return MaterialLifetimePtr.Get(); }
 
 	// Deprecated functions but still valid because they forward data correctly.
 	UE_DEPRECATED(5.0, "This function has been replaced by SetReflectionViewSampleCountScale.")
@@ -211,6 +211,7 @@ public:
 	virtual void Serialize(FArchive& Ar) override;
 #if WITH_EDITOR
 	virtual void PostEditChangeProperty(FPropertyChangedEvent& PropertyChangedEvent) override;
+	virtual void OnCookEvent(UE::Cook::ECookEvent CookEvent, UE::Cook::FCookEventContext& CookContext) override;
 #endif // WITH_EDITOR
 	//~ End UObject Interface
 
@@ -232,6 +233,10 @@ private:
 
 	FVolumetricCloudSceneProxy* VolumetricCloudSceneProxy;
 
+	// The material is assigned to a regular pointer so that the garbage collector to not get ride of it.
+	UPROPERTY()
+	TObjectPtr<UMaterialInterface> MaterialLifetimePtr;
+
 };
 
 
diff --git a/Engine/Source/Runtime/Engine/Private/Components/VolumetricCloudComponent.cpp b/Engine/Source/Runtime/Engine/Private/Components/VolumetricCloudComponent.cpp
--- a/Engine/Source/Runtime/Engine/Private/Components/VolumetricCloudComponent.cpp
+++ b/Engine/Source/Runtime/Engine/Private/Components/VolumetricCloudComponent.cpp
@@ -13,6 +13,13 @@
 #include "SceneInterface.h"
 #include "UObject/UE5MainStreamObjectVersion.h"
 
+#if WITH_EDITOR
+#include "MaterialShared.h"
+#include "Cooker/CookDependency.h"
+#include "Cooker/CookEvents.h"
+#include "Interfaces/ITargetPlatform.h"
+#endif
+
 #include UE_INLINE_GENERATED_CPP_BY_NAME(VolumetricCloudComponent)
 
 #if WITH_EDITOR
@@ -56,7 +63,10 @@ UVolumetricCloudComponent::UVolumetricCloudComponent(const FObjectInitializer& O
 	, bHoldout(false)
 	, bRenderInMainPass(true)
 	, VolumetricCloudSceneProxy(nullptr)
+	, MaterialLifetimePtr(nullptr)
 {
+	// The default material as a soft object pointer, so that the relatively heavy default material with few 3D textures is not always loaded at first.
+	Material = TSoftObjectPtr<UMaterialInterface>(FSoftObjectPath(TEXT("/Engine/EngineSky/VolumetricClouds/m_SimpleVolumetricCloud_Inst.m_SimpleVolumetricCloud_Inst")));
 }
 
 UVolumetricCloudComponent::~UVolumetricCloudComponent()
@@ -108,6 +118,18 @@ void UVolumetricCloudComponent::DestroyRenderState_Concurrent()
 
 #if WITH_EDITOR
 
+void UVolumetricCloudComponent::OnCookEvent(UE::Cook::ECookEvent CookEvent, UE::Cook::FCookEventContext& CookContext)
+{
+	Super::OnCookEvent(CookEvent, CookContext);
+
+	if (CookEvent == UE::Cook::ECookEvent::PlatformCookDependencies && CookContext.IsCooking())
+	{
+		// Add a runtime dependency on the material asset.
+		const FString MaterialLongPackageName = Material.GetLongPackageName();
+		CookContext.AddRuntimeDependency(MaterialLongPackageName);
+	}
+}
+
 void UVolumetricCloudComponent::CheckForErrors()
 {
 	// Clouds with SkyAtmosphere?
@@ -144,14 +166,9 @@ void UVolumetricCloudComponent::OnRegister()
 {
 	Super::OnRegister();
 
-	// Only load and assign a default the material if none is assigned.
-	// This prevent the loading of default the material in game, and provide a default material in editor when drag and dropping the component.
-	if (Material == nullptr)
-	{
-		// This is a similar behavior to material on meshes: a default material is always assigned if none are set.
-		check(IsInGameThread());
-		Material = LoadObject<UMaterialInterface>(NULL, TEXT("/Engine/EngineSky/VolumetricClouds/m_SimpleVolumetricCloud_Inst.m_SimpleVolumetricCloud_Inst"));
-	}
+	// Make sure the assigned TSoftObjectPtr material is loaded by that point so that we can create proxy render side.
+	// And it is assigned to a TObjectPtr to make sure it remains loaded and ignored by the GC.
+	MaterialLifetimePtr = Material.LoadSynchronous();
 }
 
 
@@ -177,7 +194,16 @@ CLOUD_DECLARE_BLUEPRINT_SETFUNCTION(float, ViewSampleCountScale);
 CLOUD_DECLARE_BLUEPRINT_SETFUNCTION(float, ShadowViewSampleCountScale);
 CLOUD_DECLARE_BLUEPRINT_SETFUNCTION(float, ShadowTracingDistance);
 CLOUD_DECLARE_BLUEPRINT_SETFUNCTION(float, StopTracingTransmittanceThreshold);
-CLOUD_DECLARE_BLUEPRINT_SETFUNCTION(UMaterialInterface*, Material);
+
+void UVolumetricCloudComponent::SetMaterial(UMaterialInterface* NewValue)
+{
+	if (AreDynamicDataChangesAllowed() && MaterialLifetimePtr.Get() != NewValue)
+	{
+		Material = NewValue;
+		MaterialLifetimePtr = Material.LoadSynchronous();
+		MarkRenderStateDirty();
+	}
+}
 
 void UVolumetricCloudComponent::SetHoldout(bool bNewHoldout)
 {
diff --git a/Engine/Source/Runtime/Engine/Private/VolumetricCloudProxy.cpp b/Engine/Source/Runtime/Engine/Private/VolumetricCloudProxy.cpp
--- a/Engine/Source/Runtime/Engine/Private/VolumetricCloudProxy.cpp
+++ b/Engine/Source/Runtime/Engine/Private/VolumetricCloudProxy.cpp
@@ -32,7 +32,7 @@ FVolumetricCloudSceneProxy::FVolumetricCloudSceneProxy(const UVolumetricCloudCom
 	, AerialPespectiveMieScatteringFadeDistance(InComponent->AerialPespectiveMieScatteringFadeDistance) 
 	, AerialPespectiveRayleighScatteringStartDistance(InComponent->AerialPespectiveRayleighScatteringStartDistance)
 	, AerialPespectiveRayleighScatteringFadeDistance(InComponent->AerialPespectiveRayleighScatteringFadeDistance)
-	, CloudVolumeMaterial(InComponent->Material)
+	, CloudVolumeMaterial(InComponent->GetMaterial())
 {
 }
 
