Substrate -Fixed fully simplified materials broken after 40645535.

#rnx
#tests PC editor

diff --git a/Engine/Source/Runtime/Core/Private/UObject/FortniteMainBranchObjectVersion.cpp b/Engine/Source/Runtime/Core/Private/UObject/FortniteMainBranchObjectVersion.cpp
--- a/Engine/Source/Runtime/Core/Private/UObject/FortniteMainBranchObjectVersion.cpp
+++ b/Engine/Source/Runtime/Core/Private/UObject/FortniteMainBranchObjectVersion.cpp
@@ -14,6 +14,6 @@ TMap<FGuid, FGuid> FFortniteMainBranchObjectVersion::GetSystemGuids()
 	SystemGuids.Add(DevGuids.Niagara_LatestScriptCompileVersion, FGuid("670A2EBDEE1A4794BB77BBFF331EA7EF"));
 	SystemGuids.Add(DevGuids.SkeletalMeshDerivedDataVersion, FGuid("5E93192FCF3647EEB7F9D4FFF1750050"));
 	SystemGuids.Add(DevGuids.STATICMESH_DERIVEDDATA_VER, FGuid("3ABF32149F9D4D838750368AB95573BE"));
-	SystemGuids.Add(DevGuids.MaterialTranslationDDCVersion, FGuid("9A84A99E461847BB80B67402ED5586F2"));
+	SystemGuids.Add(DevGuids.MaterialTranslationDDCVersion, FGuid("8ED4EEC4900A41E1BD70CB45CCEDB575"));
 	return SystemGuids;
 }
diff --git a/Engine/Source/Runtime/Engine/Private/Materials/HLSLMaterialTranslator.cpp b/Engine/Source/Runtime/Engine/Private/Materials/HLSLMaterialTranslator.cpp
--- a/Engine/Source/Runtime/Engine/Private/Materials/HLSLMaterialTranslator.cpp
+++ b/Engine/Source/Runtime/Engine/Private/Materials/HLSLMaterialTranslator.cpp
@@ -1934,68 +1934,72 @@ void FHLSLMaterialTranslator::TranslateMaterial()
 
 	if (bUseRootNodeToSubstrateHiddenConversion)
 	{
-		UMaterialInterface* MaterialInterface = Material->GetMaterialInterface();
-		UMaterial* BaseMaterial = MaterialInterface ? MaterialInterface->GetMaterial() : nullptr;
-		UMaterialEditorOnlyData* EditorOnlyData = BaseMaterial ? BaseMaterial->GetEditorOnlyData() : nullptr;
-
-		if (BaseMaterial && BaseMaterial->bUseMaterialAttributes)
+		auto CompileFrontMaterial = [&]()
 		{
-			const bool bHasSSS = MaterialInterface && MaterialInterface->SubsurfaceProfile;
+			UMaterialInterface* MaterialInterface = Material->GetMaterialInterface();
+			UMaterial* BaseMaterial = MaterialInterface ? MaterialInterface->GetMaterial() : nullptr;
+			UMaterialEditorOnlyData* EditorOnlyData = BaseMaterial ? BaseMaterial->GetEditorOnlyData() : nullptr;
 
-			Chunk[MP_FrontMaterial] = UMaterialExpressionSubstrateConvertMaterialAttributes::CompileCommon(this, 0,
-				FMaterialAttributeDefinitionMap::GetConnectedMaterialAttributesBitmask(BaseMaterial->GetExpressions()), EditorOnlyData->MaterialAttributes, MaterialInterface->GetShadingModels().GetFirstShadingModel(),
-				*WaterScatteringCoefficients, *WaterAbsorptionCoefficients, *WaterPhaseG, *ColorScaleBehindWater,
-				bHasSSS, bHasSSS ? MaterialInterface->SubsurfaceProfile : nullptr);
+			if (BaseMaterial && BaseMaterial->bUseMaterialAttributes)
+			{
+				const bool bHasSSS = MaterialInterface && MaterialInterface->SubsurfaceProfile;
 
-		}
-		else if ((Material->GetMaterialDomain() == MD_Surface || Material->GetMaterialDomain() == MD_RuntimeVirtualTexture || Material->GetMaterialDomain() == MD_DeferredDecal) && EditorOnlyData)
-		{
-			const bool bHasSSS = MaterialInterface && MaterialInterface->SubsurfaceProfile;
-			const bool bHasAnisotropy = EditorOnlyData->Anisotropy.IsConnected();
+				return UMaterialExpressionSubstrateConvertMaterialAttributes::CompileCommon(this, 0,
+					FMaterialAttributeDefinitionMap::GetConnectedMaterialAttributesBitmask(BaseMaterial->GetExpressions()), EditorOnlyData->MaterialAttributes, MaterialInterface->GetShadingModels().GetFirstShadingModel(),
+					*WaterScatteringCoefficients, *WaterAbsorptionCoefficients, *WaterPhaseG, *ColorScaleBehindWater,
+					bHasSSS, bHasSSS ? MaterialInterface->SubsurfaceProfile : nullptr);
 
-			Chunk[MP_FrontMaterial] = UMaterialExpressionSubstrateShadingModels::CompileCommon(this,
-				EditorOnlyData->BaseColor, EditorOnlyData->Specular, EditorOnlyData->Metallic, EditorOnlyData->Roughness, EditorOnlyData->EmissiveColor,
-				EditorOnlyData->Opacity, EditorOnlyData->SubsurfaceColor, EditorOnlyData->ClearCoat, EditorOnlyData->ClearCoatRoughness,
-				EditorOnlyData->ShadingModelFromMaterialExpression, MaterialInterface->GetShadingModels().GetFirstShadingModel(),
-				*ThinTranslucentTransmittanceColor, *ThinTranslucentSurfaceCoverage,
-				*WaterScatteringCoefficients, *WaterAbsorptionCoefficients, *WaterPhaseG, *ColorScaleBehindWater,
-				bHasAnisotropy, EditorOnlyData->Anisotropy,
-				EditorOnlyData->Normal, EditorOnlyData->Tangent,
-				*ClearCoatNormal, *CustomTangent,
-				bHasSSS, MaterialInterface->SubsurfaceProfile,
-				EditorOnlyData);
-		}
-		else if (Material->GetMaterialDomain() == MD_Volume && EditorOnlyData)
-		{
-			Chunk[MP_FrontMaterial] = UMaterialExpressionSubstrateVolumetricFogCloudBSDF::CompileCommon(
-				this,
-				EditorOnlyData->BaseColor,
-				EditorOnlyData->SubsurfaceColor,
-				EditorOnlyData->EmissiveColor,
-				EditorOnlyData->AmbientOcclusion,
-				MaterialInterface->GetShadingModels().HasOnlyShadingModel(MSM_Unlit),
-				EditorOnlyData);
-		}
-		else if (Material->GetMaterialDomain() == MD_LightFunction && EditorOnlyData)
-		{
-			Chunk[MP_FrontMaterial] = UMaterialExpressionSubstrateLightFunction::CompileCommon(this, EditorOnlyData->EmissiveColor, EditorOnlyData);
-		}
-		else if (Material->GetMaterialDomain() == MD_PostProcess && EditorOnlyData)
-		{
-			Chunk[MP_FrontMaterial] = UMaterialExpressionSubstratePostProcess::CompileCommon(this, EditorOnlyData->EmissiveColor, EditorOnlyData->Opacity, EditorOnlyData);
-		}
-		else if (Material->GetMaterialDomain() == MD_UI && EditorOnlyData)
-		{
-			Chunk[MP_FrontMaterial] = UMaterialExpressionSubstrateUI::CompileCommon(this, EditorOnlyData->EmissiveColor, EditorOnlyData->Opacity, EditorOnlyData);
-		}
-		else
-		{
-			Chunk[MP_FrontMaterial] = SubstrateCreateAndRegisterNullMaterial();
-		}
+			}
+			else if ((Material->GetMaterialDomain() == MD_Surface || Material->GetMaterialDomain() == MD_RuntimeVirtualTexture || Material->GetMaterialDomain() == MD_DeferredDecal) && EditorOnlyData)
+			{
+				const bool bHasSSS = MaterialInterface && MaterialInterface->SubsurfaceProfile;
+				const bool bHasAnisotropy = EditorOnlyData->Anisotropy.IsConnected();
+
+				return UMaterialExpressionSubstrateShadingModels::CompileCommon(this,
+					EditorOnlyData->BaseColor, EditorOnlyData->Specular, EditorOnlyData->Metallic, EditorOnlyData->Roughness, EditorOnlyData->EmissiveColor,
+					EditorOnlyData->Opacity, EditorOnlyData->SubsurfaceColor, EditorOnlyData->ClearCoat, EditorOnlyData->ClearCoatRoughness,
+					EditorOnlyData->ShadingModelFromMaterialExpression, MaterialInterface->GetShadingModels().GetFirstShadingModel(),
+					*ThinTranslucentTransmittanceColor, *ThinTranslucentSurfaceCoverage,
+					*WaterScatteringCoefficients, *WaterAbsorptionCoefficients, *WaterPhaseG, *ColorScaleBehindWater,
+					bHasAnisotropy, EditorOnlyData->Anisotropy,
+					EditorOnlyData->Normal, EditorOnlyData->Tangent,
+					*ClearCoatNormal, *CustomTangent,
+					bHasSSS, MaterialInterface->SubsurfaceProfile,
+					EditorOnlyData);
+			}
+			else if (Material->GetMaterialDomain() == MD_Volume && EditorOnlyData)
+			{
+				return UMaterialExpressionSubstrateVolumetricFogCloudBSDF::CompileCommon(
+					this,
+					EditorOnlyData->BaseColor,
+					EditorOnlyData->SubsurfaceColor,
+					EditorOnlyData->EmissiveColor,
+					EditorOnlyData->AmbientOcclusion,
+					MaterialInterface->GetShadingModels().HasOnlyShadingModel(MSM_Unlit),
+					EditorOnlyData);
+			}
+			else if (Material->GetMaterialDomain() == MD_LightFunction && EditorOnlyData)
+			{
+				return UMaterialExpressionSubstrateLightFunction::CompileCommon(this, EditorOnlyData->EmissiveColor, EditorOnlyData);
+			}
+			else if (Material->GetMaterialDomain() == MD_PostProcess && EditorOnlyData)
+			{
+				return UMaterialExpressionSubstratePostProcess::CompileCommon(this, EditorOnlyData->EmissiveColor, EditorOnlyData->Opacity, EditorOnlyData);
+			}
+			else if (Material->GetMaterialDomain() == MD_UI && EditorOnlyData)
+			{
+				return UMaterialExpressionSubstrateUI::CompileCommon(this, EditorOnlyData->EmissiveColor, EditorOnlyData->Opacity, EditorOnlyData);
+			}
+
+			return SubstrateCreateAndRegisterNullMaterial();
+		};
+		
+		Chunk[MP_FrontMaterial] = CompileFrontMaterial();
 
-		// Make the fully simplified code be the same as the FrontMaterial that is a single BSDF when using hidden conversion.
 		FullySimplifiedFrontMaterialCodeChunkStart = SharedPropertyCodeChunks[FrontMaterialShaderFrequency].Num();
-		FullySimplifiedSubstrateFrontMaterialCodeChunk = Chunk[MP_FrontMaterial];
+		CurrentSubstrateCompilationContext = ESubstrateCompilationContext::SCC_FullySimplified;
+		FullySimplifiedSubstrateFrontMaterialCodeChunk = CompileFrontMaterial();
+		CurrentSubstrateCompilationContext = ESubstrateCompilationContext::SCC_Default;
 		FullySimplifiedFrontMaterialCodeChunkEnd = SharedPropertyCodeChunks[FrontMaterialShaderFrequency].Num();
 	}
 
