Fixed static analysis warning in MaterialIRModuleBuilder.cpp.

diff --git a/Engine/Source/Runtime/Engine/Private/Materials/MaterialIRModuleBuilder.cpp b/Engine/Source/Runtime/Engine/Private/Materials/MaterialIRModuleBuilder.cpp
--- a/Engine/Source/Runtime/Engine/Private/Materials/MaterialIRModuleBuilder.cpp
+++ b/Engine/Source/Runtime/Engine/Private/Materials/MaterialIRModuleBuilder.cpp
@@ -515,17 +515,18 @@ struct FMaterialIRModuleBuilderImpl
 	//
 	void PushConnectionInsight(const UObject* InputObject, int InputIndex, const UMaterialExpression* OutputExpression, int OutputIndex, const MIR::FType* Type)
 	{
-		if (!Builder->TargetInsights || Type->IsPoison())
+		if (!Builder->TargetInsights || !Type || Type->IsPoison())
 		{
 			return;
 		}
 
-		FMaterialInsights::FConnectionInsight Insight;
-		Insight.InputObject = InputObject,
-		Insight.OutputExpression = OutputExpression,
-		Insight.InputIndex = InputIndex,
-		Insight.OutputIndex = OutputIndex,
-		Insight.ValueType = Type ? Type->ToValueType() : UE::Shader::EValueType::Any,
+		FMaterialInsights::FConnectionInsight Insight {
+			.InputObject = InputObject,
+			.OutputExpression = OutputExpression,
+			.InputIndex = InputIndex,
+			.OutputIndex = OutputIndex,
+			.ValueType = Type->ToValueType(),
+		};
 		
 		Builder->TargetInsights->ConnectionInsights.Push(Insight);
 	}
