Ensure condition failed: ExistingAttributeIndex == INDEX_NONE (Animation\AttributesContainer.h)
#fix clear attributes container when performing state-tree active transition blends, as attribute blending expects an empty container to blend into
#misc added attribute container copy to custom blend option in state-tree (matching curves)

diff --git a/Engine/Source/Runtime/Engine/Private/Animation/AnimNode_StateMachine.cpp b/Engine/Source/Runtime/Engine/Private/Animation/AnimNode_StateMachine.cpp
--- a/Engine/Source/Runtime/Engine/Private/Animation/AnimNode_StateMachine.cpp
+++ b/Engine/Source/Runtime/Engine/Private/Animation/AnimNode_StateMachine.cpp
@@ -1075,6 +1075,8 @@ void FAnimNode_StateMachine::EvaluateTransitionStandardBlendInternal(FPoseContex
 	Output.Curve.Override(PreviousStateResult.Curve, 1.0 - Transition.Alpha);
 	Output.Curve.Accumulate(NextStateResult.Curve, Transition.Alpha);
 		
+	// reset attributes container as it can contain a previously handled active transition (and attributes expects blending into an empty container)
+	Output.CustomAttributes.Empty();
 	UE::Anim::Attributes::BlendAttributes({ PreviousStateResult.CustomAttributes, NextStateResult.CustomAttributes }, { 1.0f - Transition.Alpha, Transition.Alpha }, { 0, 1 }, Output.CustomAttributes);
 }
 
@@ -1119,6 +1121,9 @@ void FAnimNode_StateMachine::EvaluateTransitionCustomBlend(FPoseContext& Output,
 
 		// Copy curve over also, replacing current.
 		Output.Curve.CopyFrom(StatePoseResult.Curve);
+
+		// Copy attributes over, replacing current.
+		Output.CustomAttributes.CopyFrom(StatePoseResult.CustomAttributes);
 	}
 }
 
