[HWRT] Fix r.Raytracing.Debug.ForceRuntimeBLAS.
- it wasn't being respected on some streaming code paths.
- modified logic to immediately discard offline data during loading logic.

diff --git a/Engine/Source/Runtime/Engine/Internal/Streaming/SkeletalMeshUpdate.cpp b/Engine/Source/Runtime/Engine/Internal/Streaming/SkeletalMeshUpdate.cpp
--- a/Engine/Source/Runtime/Engine/Internal/Streaming/SkeletalMeshUpdate.cpp
+++ b/Engine/Source/Runtime/Engine/Internal/Streaming/SkeletalMeshUpdate.cpp
@@ -87,19 +87,11 @@ void FSkeletalMeshStreamIn::FIntermediateRayTracingGeometry::CreateFromCPUData(F
 
 	if (RayTracingGeometry.RawData.Num())
 	{
+		check(!RayTracing::ShouldForceRuntimeBLAS());
 		check(Initializer.OfflineData == nullptr);
 		Initializer.OfflineData = &RayTracingGeometry.RawData;
 	}
 
-	static const auto CVarDebugForceRuntimeBLAS = IConsoleManager::Get().FindTConsoleVariableDataInt(TEXT("r.Raytracing.DebugForceRuntimeBLAS"));
-	const bool bDebugForceRuntimeBLAS = (!CVarDebugForceRuntimeBLAS) || (CVarDebugForceRuntimeBLAS->GetValueOnAnyThread() != 0);
-
-	if (bDebugForceRuntimeBLAS && Initializer.OfflineData != nullptr)
-	{
-		Initializer.OfflineData->Discard();
-		Initializer.OfflineData = nullptr;
-	}
-
 	RayTracingGeometryRHI = RHICmdList.CreateRayTracingGeometry(Initializer);
 	bRequiresBuild = Initializer.OfflineData == nullptr || RayTracingGeometryRHI->IsCompressed();
 }
diff --git a/Engine/Source/Runtime/Engine/Private/SkeletalMeshLODRenderData.cpp b/Engine/Source/Runtime/Engine/Private/SkeletalMeshLODRenderData.cpp
--- a/Engine/Source/Runtime/Engine/Private/SkeletalMeshLODRenderData.cpp
+++ b/Engine/Source/Runtime/Engine/Private/SkeletalMeshLODRenderData.cpp
@@ -790,6 +790,11 @@ void FSkeletalMeshLODRenderData::SerializeStreamedData(FArchive& Ar, USkinnedAss
 	Ar << SourceRayTracingGeometry.RawData;
 	SourceRayTracingGeometry.SetResourceName(SkeletalMeshResourceNames::SourceRayTracingGeometryName);
 
+	if (Ar.IsLoading() && RayTracing::ShouldForceRuntimeBLAS())
+	{
+		SourceRayTracingGeometry.RawData.Discard();
+	}
+
 	// Determine if morph target data should be cooked out for this platform
 	EShaderPlatform MorphTargetShaderPlatform = GMaxRHIShaderPlatform;
 	bool bSerializeCompressedMorphTargets = Ar.IsSaving() && FMorphTargetVertexInfoBuffers::IsPlatformShaderSupported(MorphTargetShaderPlatform) && Owner && Owner->GetMorphTargets().Num() > 0;
diff --git a/Engine/Source/Runtime/Engine/Private/StaticMesh.cpp b/Engine/Source/Runtime/Engine/Private/StaticMesh.cpp
--- a/Engine/Source/Runtime/Engine/Private/StaticMesh.cpp
+++ b/Engine/Source/Runtime/Engine/Private/StaticMesh.cpp
@@ -622,6 +622,12 @@ void FStaticMeshLODResources::SerializeBuffers(FArchive& Ar, UStaticMesh* OwnerS
 		Ar << SerializedRayTracingGeometry->RawDataHeader;
 		SerializedRayTracingGeometry->RawData.BulkSerialize(Ar);
 		AccumRayTracingGeometrySize(*SerializedRayTracingGeometry, OutBuffersSize.SerializedBuffersSize);
+
+		if (Ar.IsLoading() && RayTracing::ShouldForceRuntimeBLAS())
+		{
+			SerializedRayTracingGeometry->RawDataHeader = FRayTracingGeometryOfflineDataHeader();
+			SerializedRayTracingGeometry->RawData.Discard();
+		}
 	}
 
 	AreaWeightedSectionSamplers.SetNum(Sections.Num());
@@ -727,6 +733,11 @@ void FStaticMeshLODResources::SerializeAvailabilityInfo(FArchive& Ar)
 		FRayTracingGeometry DummyRayTracingGeometry;
 		FRayTracingGeometry* SerializedRayTracingGeometry = RayTracingGeometry != nullptr ? RayTracingGeometry : &DummyRayTracingGeometry;
 		Ar << SerializedRayTracingGeometry->RawDataHeader;
+
+		if (Ar.IsLoading() && RayTracing::ShouldForceRuntimeBLAS())
+		{
+			SerializedRayTracingGeometry->RawDataHeader = FRayTracingGeometryOfflineDataHeader();
+		}
 	}
 }
 
@@ -1683,7 +1694,7 @@ void FStaticMeshRayTracingProxyLOD::InitResources(UStaticMesh* Parent, int32 LOD
 #if RHI_RAYTRACING
 	if (IsRayTracingAllowed() && RayTracingGeometry != nullptr && bOwnsRayTracingGeometry)
 	{
-		if(RayTracingGeometry->RawDataHeader.IsValid())
+		if (RayTracingGeometry->RawDataHeader.IsValid())
 		{
 			FRayTracingAccelerationStructureOfflineMetadata OfflineMetadata = RHIGetRayTracingGeometryOfflineMetadata(RayTracingGeometry->RawDataHeader);
 
@@ -1898,6 +1909,11 @@ void FStaticMeshRayTracingProxyLOD::SerializeRayTracingGeometry(FArchive& Ar, US
 			OfflineBVHOffset = BeforeOffset;
 			OfflineBVHSize = AfterOffset - BeforeOffset;
 		}
+
+		if (Ar.IsLoading() && RayTracing::ShouldForceRuntimeBLAS())
+		{
+			RayTracingGeometry->RawData.Discard();
+		}
 	}
 }
 
@@ -1920,6 +1936,11 @@ void FStaticMeshRayTracingProxyLOD::SerializeMetaData(FArchive& Ar)
 		Ar << OfflineBVHSize;
 
 		Ar << RayTracingGeometry->RawDataHeader;
+
+		if (Ar.IsLoading() && RayTracing::ShouldForceRuntimeBLAS())
+		{
+			RayTracingGeometry->RawDataHeader = FRayTracingGeometryOfflineDataHeader();
+		}
 	}
 }
 
@@ -1978,6 +1999,11 @@ void FStaticMeshRayTracingProxyLOD::InitWithStreamedData(FRHICommandListBase& RH
 		RayTracingGeometry->RawData.Empty(OfflineBVHSize);
 		RayTracingGeometry->RawData.AddUninitialized(OfflineBVHSize);
 		Ar.Serialize(RayTracingGeometry->RawData.GetData(), OfflineBVHSize);
+
+		if (Ar.IsLoading() && RayTracing::ShouldForceRuntimeBLAS())
+		{
+			RayTracingGeometry->RawData.Discard();
+		}
 	}
 	
 	check(!Ar.IsError());
diff --git a/Engine/Source/Runtime/Engine/Private/Streaming/StaticMeshUpdate.cpp b/Engine/Source/Runtime/Engine/Private/Streaming/StaticMeshUpdate.cpp
--- a/Engine/Source/Runtime/Engine/Private/Streaming/StaticMeshUpdate.cpp
+++ b/Engine/Source/Runtime/Engine/Private/Streaming/StaticMeshUpdate.cpp
@@ -100,19 +100,11 @@ void FStaticMeshStreamIn::FIntermediateRayTracingGeometry::CreateFromCPUData(FRH
 
 	if (RayTracingGeometry.RawData.Num())
 	{
+		check(!RayTracing::ShouldForceRuntimeBLAS());
 		check(Initializer.OfflineData == nullptr);
 		Initializer.OfflineData = &RayTracingGeometry.RawData;
 	}
 
-	static const auto CVarDebugForceRuntimeBLAS = IConsoleManager::Get().FindTConsoleVariableDataInt(TEXT("r.Raytracing.DebugForceRuntimeBLAS"));
-	const bool bDebugForceRuntimeBLAS = (!CVarDebugForceRuntimeBLAS) || (CVarDebugForceRuntimeBLAS->GetValueOnAnyThread() != 0);
-
-	if (bDebugForceRuntimeBLAS && Initializer.OfflineData != nullptr)
-	{
-		Initializer.OfflineData->Discard();
-		Initializer.OfflineData = nullptr;
-	}
-
 	RayTracingGeometryRHI = RHICmdList.CreateRayTracingGeometry(Initializer);
 	bRequiresBuild = Initializer.OfflineData == nullptr || RayTracingGeometryRHI->IsCompressed();
 }
diff --git a/Engine/Source/Runtime/RenderCore/Private/RayTracingGeometry.cpp b/Engine/Source/Runtime/RenderCore/Private/RayTracingGeometry.cpp
--- a/Engine/Source/Runtime/RenderCore/Private/RayTracingGeometry.cpp
+++ b/Engine/Source/Runtime/RenderCore/Private/RayTracingGeometry.cpp
@@ -15,11 +15,17 @@ IRayTracingGeometryManager* GRayTracingGeometryManager = nullptr;
 #endif
 
 static TAutoConsoleVariable<int32> CVarDebugForceRuntimeBLAS(
-	TEXT("r.Raytracing.DebugForceRuntimeBLAS"),
+	TEXT("r.Raytracing.Debug.ForceRuntimeBLAS"),
 	0,
 	TEXT("Force building BLAS at runtime."),
 	ECVF_ReadOnly);
 
+bool RayTracing::ShouldForceRuntimeBLAS()
+{
+	static const bool bDebugForceRuntimeBLAS = CVarDebugForceRuntimeBLAS.GetValueOnAnyThread() != 0;
+	return bDebugForceRuntimeBLAS;
+}
+
 PRAGMA_DISABLE_DEPRECATION_WARNINGS // Remove this in UE 5.6 once RayTracingGeometryRHI is made private
 
 #if RHI_RAYTRACING
@@ -136,16 +142,11 @@ void FRayTracingGeometry::CreateRayTracingGeometry(FRHICommandListBase& RHICmdLi
 
 	if (RawData.Num())
 	{
+		check(!RayTracing::ShouldForceRuntimeBLAS());
 		check(Initializer.OfflineData == nullptr);
 		Initializer.OfflineData = &RawData;
 	}
 
-	if (CVarDebugForceRuntimeBLAS.GetValueOnAnyThread() && Initializer.OfflineData != nullptr)
-	{
-		Initializer.OfflineData->Discard();
-		Initializer.OfflineData = nullptr;
-	}
-
 	if (HasValidInitializer())
 	{
 		// Geometries with StreamingDestination type are initially created in invalid state until they are streamed in (see InitRHIForStreaming).
diff --git a/Engine/Source/Runtime/RenderCore/Public/RayTracingGeometry.h b/Engine/Source/Runtime/RenderCore/Public/RayTracingGeometry.h
--- a/Engine/Source/Runtime/RenderCore/Public/RayTracingGeometry.h
+++ b/Engine/Source/Runtime/RenderCore/Public/RayTracingGeometry.h
@@ -13,6 +13,8 @@ namespace RayTracing
 {
 	using FGeometryGroupHandle = int32;
 	using GeometryGroupHandle UE_DEPRECATED(5.6, "Use FGeometryGroupHandle instead.") = FGeometryGroupHandle;
+
+	RENDERCORE_API bool ShouldForceRuntimeBLAS();
 }
 
 enum class ERTAccelerationStructureBuildPriority : uint8
