Cleanup of deprecated scene view properties & functions.

#rnx

diff --git a/Engine/Source/Runtime/Engine/Private/SceneView.cpp b/Engine/Source/Runtime/Engine/Private/SceneView.cpp
--- a/Engine/Source/Runtime/Engine/Private/SceneView.cpp
+++ b/Engine/Source/Runtime/Engine/Private/SceneView.cpp
@@ -905,9 +905,6 @@ FSceneView::FSceneView(const FSceneViewInitOptions& InitOptions)
 	, bIsLocked(false)
 	, bStaticSceneOnly(false)
 	, bIsInstancedStereoEnabled(false)
-PRAGMA_DISABLE_DEPRECATION_WARNINGS
-	, bIsMultiViewEnabled(false)
-PRAGMA_ENABLE_DEPRECATION_WARNINGS
 	, bIsMultiViewportEnabled(false)
 	, bIsMobileMultiViewEnabled(false)
 	, bShouldBindInstancedViewUB(false)
@@ -953,10 +950,6 @@ PRAGMA_ENABLE_DEPRECATION_WARNINGS
 		bIsMobileMultiViewEnabled = !bIsSingleViewCapture && Family && Family->bRequireMultiView && Aspects.IsMobileMultiViewEnabled();
 		bIsSinglePassStereo = bIsInstancedStereoEnabled || bIsMobileMultiViewEnabled;
 
-		PRAGMA_DISABLE_DEPRECATION_WARNINGS
-		bIsMultiViewEnabled = bIsMultiViewportEnabled;	// temporary, as a graceful way to support plugins/licensee mods
-		PRAGMA_ENABLE_DEPRECATION_WARNINGS
-
 		// If instanced stereo is enabled, we should also have either multiviewport enabled, or mmv fallback enabled.
 		// Assert this since a more graceful handling is done earlier (see PostInitRHI)
 		checkf(!bIsInstancedStereoEnabled || (bIsMultiViewportEnabled || Aspects.IsMobileMultiViewEnabled()),
@@ -2935,9 +2928,6 @@ FSceneViewFamily::ConstructionValues::ConstructionValues(
 ,	Scene(InScene)
 ,	EngineShowFlags(InEngineShowFlags)
 ,	ViewModeParam(-1)
-PRAGMA_DISABLE_DEPRECATION_WARNINGS
-,	GammaCorrection(1.0f)
-PRAGMA_ENABLE_DEPRECATION_WARNINGS
 ,	bAdditionalViewFamily(false)
 ,	bRealtimeUpdate(false)
 ,	bDeferClear(false)
@@ -2964,9 +2954,6 @@ FSceneViewFamily::FSceneViewFamily(const ConstructionValues& CVS)
 	Scene(CVS.Scene),
 	EngineShowFlags(CVS.EngineShowFlags),
 	Time(CVS.Time),
-	CurrentWorldTime(CVS.Time.GetWorldTimeSeconds()),
-	DeltaWorldTime(CVS.Time.GetDeltaWorldTimeSeconds()),
-	CurrentRealTime(CVS.Time.GetRealTimeSeconds()),
 	FrameNumber(UINT_MAX),
 	bAdditionalViewFamily(CVS.bAdditionalViewFamily),
 	bRealtimeUpdate(CVS.bRealtimeUpdate),
@@ -2978,9 +2965,6 @@ FSceneViewFamily::FSceneViewFamily(const ConstructionValues& CVS)
 	bWorldIsPaused(false),
 	bIsHDR(false),
 	bRequireMultiView(CVS.bRequireMultiView),
-PRAGMA_DISABLE_DEPRECATION_WARNINGS
-	GammaCorrection(CVS.GammaCorrection),
-PRAGMA_ENABLE_DEPRECATION_WARNINGS
 	SecondaryViewFraction(1.0f),
 	SecondaryScreenPercentageMethod(ESecondaryScreenPercentageMethod::LowerPixelDensitySimulation),
 	SceneRenderer(nullptr),
@@ -3054,7 +3038,6 @@ PRAGMA_ENABLE_DEPRECATION_WARNINGS
 	EngineShowFlags.ScreenPercentage = SupportsScreenPercentage();
 }
 
-PRAGMA_DISABLE_DEPRECATION_WARNINGS // TOptional can't be deprecated without emitting warnings in destructor
 FSceneViewFamily::~FSceneViewFamily()
 {
 	// If a screen percentage was given for the view family, delete it since any new copy of a view family will Fork it.
@@ -3081,7 +3064,6 @@ FSceneViewFamily::~FSceneViewFamily()
 		delete reinterpret_cast<ISceneViewFamilyExtention*>(SecondarySpatialUpscalerInterface);
 	}
 }
-PRAGMA_ENABLE_DEPRECATION_WARNINGS
 
 FSceneViewFamily::FSceneViewFamily(FSceneViewFamily&& InViewFamily)
 	: FSceneViewFamily(static_cast<const FSceneViewFamily&>(InViewFamily))
diff --git a/Engine/Source/Runtime/Engine/Public/SceneView.h b/Engine/Source/Runtime/Engine/Public/SceneView.h
--- a/Engine/Source/Runtime/Engine/Public/SceneView.h
+++ b/Engine/Source/Runtime/Engine/Public/SceneView.h
@@ -1604,10 +1604,6 @@ public:
 	/** True if instanced stereo is enabled. */
 	bool bIsInstancedStereoEnabled;
 
-	/** OLD variable governing multi-viewport rendering. */
-	UE_DEPRECATED(5.1, "bIsMultiViewEnabled has been deprecated. Use bIsMultiViewportEnabled")
-	bool bIsMultiViewEnabled;
-
 	/** True if multi-viewport instanced stereo rendering is enabled. */
 	bool bIsMultiViewportEnabled;
 
@@ -2185,10 +2181,6 @@ public:
 		/** The current time. */
 		FGameTime Time;
 
-		/** Gamma correction used when rendering this family. Default is 1.0 */
-		UE_DEPRECATED(5.4, "Unused gamma correction.")
-		float GammaCorrection;
-
 		/** Indicates whether the view family is additional. */
 		uint32 bAdditionalViewFamily : 1;
 
@@ -2221,12 +2213,6 @@ public:
 			return *this;
 		}
 
-		UE_DEPRECATED(5.0, "Use FSceneViewFamily::ConstructionValues::SetTime()")
-		ConstructionValues& SetWorldTimes(float InCurrentWorldTime, float InDeltaWorldTime, float InCurrentRealTime)
-		{
-			return SetTime(FGameTime::CreateDilated(InCurrentRealTime, /* InCurrentRealTime = */ InDeltaWorldTime, InCurrentWorldTime, InDeltaWorldTime));
-		}
-
 		/** Set  whether the view family is additional. */
 		ConstructionValues& SetAdditionalViewFamily(const bool Value) { bAdditionalViewFamily = Value; return *this; }
 
@@ -2241,15 +2227,6 @@ public:
 
 		/** Setting to true results in scene color and depth being multiview-allocated. */
 		ConstructionValues& SetRequireMobileMultiView(const bool Value) { bRequireMultiView = Value; return *this; }
-		
-		/** Set Gamma correction used when rendering this family. */
-		UE_DEPRECATED(5.4, "Unused gamma correction.")
-		ConstructionValues& SetGammaCorrection(const float Value)
-		{
-			PRAGMA_DISABLE_DEPRECATION_WARNINGS
-			GammaCorrection = Value; return *this;
-			PRAGMA_ENABLE_DEPRECATION_WARNINGS
-		}
 
 		/** Set the view param. */
 		ConstructionValues& SetViewModeParam(const int InViewModeParam, const FName& InViewModeParamName) { ViewModeParam = InViewModeParam; ViewModeParamName = InViewModeParamName; return *this; }		
@@ -2277,25 +2254,12 @@ public:
 	/** The current time. */
 	FGameTime Time;
 
-	UE_DEPRECATED(5.0, "Use FSceneViewFamily::Time")
-	float CurrentWorldTime;
-
-	UE_DEPRECATED(5.0, "Use FSceneViewFamily::Time")
-	float DeltaWorldTime;
-
-	UE_DEPRECATED(5.0, "Use FSceneViewFamily::Time")
-	float CurrentRealTime;
-
 	/** Copy from main thread GFrameNumber to be accessible on render thread side. UINT_MAX before CreateSceneRenderer() or BeginRenderingViewFamily() was called */
 	uint32 FrameNumber;
 
 	/** Copy from main thread GFrameCounter to be accessible on render thread side. GFrameCounter is incremented once per engine tick, so multi views of the same frame have the same value. */
 	uint64 FrameCounter = 0;
 
-	// When deleted, remove PRAGMA_DISABLE_DEPRECATION_WARNINGS / PRAGMA_ENABLE_DEPRECATION_WARNINGS from ~FSceneViewFamily()
-	UE_DEPRECATED(5.3, "Do not use, has no effect, and will be removed in a future release.")
-	TOptional<uint64> OverrideFrameCounter;
-
 	/** Indicates this view family is an additional one. */
 	bool bAdditionalViewFamily;
 
@@ -2361,10 +2325,6 @@ public:
 
 	/** True if scenecolor and depth should be multiview-allocated */
 	bool bRequireMultiView;
-
-	/** Gamma correction used when rendering this family. Default is 1.0 */
-	UE_DEPRECATED(5.4, "Unused gamma correction.")
-	float GammaCorrection;
 	
 	/** DPI scale to be used for debugging font. */
 	float DebugDPIScale = 1.0f;
