Fix Substrate glint density conversion with new material conversion.

This CL moves the density conversion function into the post load material and add proper object version instead of relying onto Substrate internal version.

diff --git a/Engine/Source/Runtime/Core/Public/UObject/FortniteMainBranchObjectVersions.inl b/Engine/Source/Runtime/Core/Public/UObject/FortniteMainBranchObjectVersions.inl
--- a/Engine/Source/Runtime/Core/Public/UObject/FortniteMainBranchObjectVersions.inl
+++ b/Engine/Source/Runtime/Core/Public/UObject/FortniteMainBranchObjectVersions.inl
@@ -601,5 +601,8 @@ FORTNITE_MAIN_VERSION(PCGAttributeSetToPointAlwaysConverts, 195)
 //Add per material slot overlay material data
 FORTNITE_MAIN_VERSION(MeshMaterialSlotOverlayMaterialAdded, 196)
 
+// Convert Sustrate glint density properly
+FORTNITE_MAIN_VERSION(ConvertGlintDensity, 197)
+
 // -----<new versions can be added above this line>-------------------------------------------------
 #undef FORTNITE_MAIN_VERSION
\ No newline at end of file
diff --git a/Engine/Source/Runtime/Engine/Private/Materials/Material.cpp b/Engine/Source/Runtime/Engine/Private/Materials/Material.cpp
--- a/Engine/Source/Runtime/Engine/Private/Materials/Material.cpp
+++ b/Engine/Source/Runtime/Engine/Private/Materials/Material.cpp
@@ -4134,46 +4134,8 @@ bool UMaterial::ConvertMaterialToSubstrateMaterial(bool bAllowEmptyMaterialUpdat
 	// Glint conversion parametrization
 	if (EditorOnly->SubstrateVersion < 0)
 	{
-		for (UMaterialExpression* Expression : GetExpressions())
-		{
-			if (UMaterialExpressionSubstrateSlabBSDF* SlabExpression = Cast<UMaterialExpressionSubstrateSlabBSDF>(Expression))
-			{
-				if (SlabExpression->GlintValue.IsConnected())
-				{
-					// Gather custom output for thin translucency
-					// Empty material: Create by default a slab node
-					UMaterialFunction* ConversionMF = LoadObject<UMaterialFunction>(nullptr, TEXT("/Engine/Functions/Substrate/SMF_GlintConversion.SMF_GlintConversion")); 
-					if (ConversionMF)
-					{
-						ConversionMF->UpdateFromFunctionResource();
-						ConversionMF->PostEditChange();
-						ConversionMF->ConditionalPostLoad();
-
-						UMaterialExpressionMaterialFunctionCall* MFCallNode = NewObject<UMaterialExpressionMaterialFunctionCall>(this);
-						if (MFCallNode->SetMaterialFunction(ConversionMF))
-						{
-							// This is needed for input/output expressions to be set correctly, otherwise compilation will fail.
-							EditorOnly->ExpressionCollection.AddExpression(MFCallNode);
-							SetPosXAndMoveReferenceToTheRight(MFCallNode);
-
-							// Connect GlintValue graph to the MF input
-							ConnectionToIndexed(SlabExpression->GlintValue, MFCallNode, 0/*Density*/);
-
-							// Connect the MF output to the GlintValue input
-							SlabExpression->GlintValue.Connect(0, MFCallNode);
-
-							MFCallNode->UpdateFromFunctionResource();
-							MFCallNode->PostEditChange();
-							MFCallNode->ConditionalPostLoad();
-
-						}
-					}
-				}
-			}
-		}
-
-		// Even if no glint conversion was one, we invalidate the shader in order to save the new Substrate's version (from -1 to 0)
-		bInvalidateShader = true;
+		// This is now done in ConvertSlabExpressionGlintDensity() during the material's PostLoad() function
+		// This conversion change SubstrateVersion from -1 -> 0
 	}
 	//else if (EditorOnly->SubstrateVersion < 1)
 	//{ [...] }
@@ -4258,6 +4220,61 @@ static void ConvertSlabExpressionMaterialSubSurfaceType(TConstArrayView<TObjectP
 	}
 }
 
+static void ConvertSlabExpressionGlintDensity(UMaterial* InMaterial)
+{
+	// Convert existing Substrate material with glint. 
+	// The FrontMaterial is expected to be connected. We rely on the Substrate version to make different 
+	// between material already converted and not (in addition to the global object version which now insure all new material will have the 
+	// correct glint density parameterization)
+	UMaterialEditorOnlyData* EditorOnly = InMaterial->GetEditorOnlyData();
+	if (EditorOnly && EditorOnly->SubstrateVersion < 0 && EditorOnly->FrontMaterial.IsConnected())
+	{
+		TArray<UMaterialExpressionSubstrateSlabBSDF*> SlabExpressions;
+		InMaterial->GetAllExpressionsInMaterialAndFunctionsOfType<UMaterialExpressionSubstrateSlabBSDF>(SlabExpressions);
+		for (UMaterialExpressionSubstrateSlabBSDF* SlabExpression : SlabExpressions)
+		{
+			if (SlabExpression->GlintValue.IsConnected())
+			{
+				// Gather custom output for thin translucency
+				// Empty material: Create by default a slab node
+				UMaterialFunction* ConversionMF = LoadObject<UMaterialFunction>(nullptr, TEXT("/Engine/Functions/Substrate/SMF_GlintConversion.SMF_GlintConversion"));
+				if (ConversionMF)
+				{
+					ConversionMF->UpdateFromFunctionResource();
+					ConversionMF->PostEditChange();
+					ConversionMF->ConditionalPostLoad();
+
+					UMaterialExpressionMaterialFunctionCall* MFCallNode = NewObject<UMaterialExpressionMaterialFunctionCall>(InMaterial);
+					if (MFCallNode->SetMaterialFunction(ConversionMF))
+					{
+						// This is needed for input/output expressions to be set correctly, otherwise compilation will fail.
+						EditorOnly->ExpressionCollection.AddExpression(MFCallNode);
+						MFCallNode->MaterialExpressionEditorX = SlabExpression->MaterialExpressionEditorX - 200;
+						MFCallNode->MaterialExpressionEditorY = SlabExpression->MaterialExpressionEditorY + 500;
+
+						// Connect GlintValue graph to the MF input
+						if (SlabExpression->GlintValue.IsConnected())
+						{
+							MFCallNode->GetInput(0/*Density*/)->Connect(SlabExpression->GlintValue.OutputIndex, SlabExpression->GlintValue.Expression);
+							SlabExpression->GlintValue.Expression = nullptr;
+						}
+
+						// Connect the MF output to the GlintValue input
+						SlabExpression->GlintValue.Connect(0, MFCallNode);
+
+						MFCallNode->UpdateFromFunctionResource();
+						MFCallNode->PostEditChange();
+						MFCallNode->ConditionalPostLoad();
+					}
+				}
+			}
+		}
+
+		// Update substrate version
+		// Even if no glint conversion was done, we save the new Substrate's version (from -1 to 0)
+		EditorOnly->SubstrateVersion = GetSubstrateVersion();
+	}
+}
 #endif // WITH_EDITOR
 
 TMap<FGuid, UMaterialInterface*> LightingGuidFixupMap;
@@ -4563,6 +4580,11 @@ void UMaterial::PostLoad()
 			ConvertSlabExpressionMaterialSubSurfaceType(GetExpressions(), BlendMode, bIsThinSurface);
 		}
 
+		if (FortMainVer < FFortniteMainBranchObjectVersion::ConvertGlintDensity)
+		{
+			ConvertSlabExpressionGlintDensity(this);
+		}
+
 		// Substrate materials conversion needs to be done after expressions are cached, otherwise material function won't have 
 		// valid inputs in certain cases
 		if (ConvertMaterialToSubstrateMaterial(true /*bAllowEmptyMaterialUpdate*/))
