[Shaders] Honor command line argument -ExclusiveRemoteShaderCompiling when launching local shader compiling thread.

This command line option is currently only honored when launching local workers in UbaController, but the shader compiling manager still launches its own local workers.
This option prevents any local SCW instances which is solely for debugging purposes to test remote shader compilation.

Also rename "bExclusiveDistributedShaderCompiling" field to "bUseOnlyDistributedCompilationThread" to make distinction between distributed- and remote-compilation more clear.

#rnx

diff --git a/Engine/Source/Runtime/Engine/Private/ShaderCompiler/ShaderCompiler.cpp b/Engine/Source/Runtime/Engine/Private/ShaderCompiler/ShaderCompiler.cpp
--- a/Engine/Source/Runtime/Engine/Private/ShaderCompiler/ShaderCompiler.cpp
+++ b/Engine/Source/Runtime/Engine/Private/ShaderCompiler/ShaderCompiler.cpp
@@ -843,7 +843,7 @@ FShaderCompilingManager::FShaderCompilingManager() :
 
 		GConfig->SetBool(TEXT("/Script/UnrealEd.UnrealEdOptions"), TEXT("UsingXGE"), RemoteCompileThread != nullptr, GEditorIni);
 
-		if (!bExclusiveDistributedShaderCompiling)
+		if (!bUseOnlyDistributedCompilationThread)
 		{
 			LaunchLocalShaderCompilingThread(bDelayCompileThreadsExecution);
 		}
@@ -1028,7 +1028,10 @@ FShaderCompileThreadRunnableBase* FShaderCompilingManager::LaunchRemoteShaderCom
 		UE_LOG(LogShaderCompilers, Display, TEXT("Using %s for shader compilation"), *BuildDistributionController->GetName());
 		TUniquePtr<FShaderCompileThreadRunnableBase> RemoteCompileThread = MakeUnique<FShaderCompileDistributedThreadRunnable_Interface>(this, *BuildDistributionController);
 
-		bExclusiveDistributedShaderCompiling = BuildDistributionController->SupportsLocalWorkers();
+		const bool bExclusiveRemoteShaderCompiling = FParse::Param(FCommandLine::Get(), TEXT("ExclusiveRemoteShaderCompiling"));
+		const bool bDistributedControllerSupportsLocalWorkers = BuildDistributionController->SupportsLocalWorkers();
+
+		bUseOnlyDistributedCompilationThread = bDistributedControllerSupportsLocalWorkers || bExclusiveRemoteShaderCompiling;
 		return LaunchShaderCompilingThread(MoveTemp(RemoteCompileThread), bDelayThreadExecution);
 	}
 
@@ -1235,12 +1238,12 @@ void FShaderCompilingManager::OnDistributedShaderCompilingChanged()
 	{
 		// Only update conditions if local shader compiling thread was initially disabled but is now required
 		const bool bDistributedControllerSupportsLocalWorkers = BuildDistributionController->SupportsLocalWorkers();
-		if (bExclusiveDistributedShaderCompiling && !bDistributedControllerSupportsLocalWorkers)
+		if (bUseOnlyDistributedCompilationThread && !bDistributedControllerSupportsLocalWorkers)
 		{
 			// Launch local shader compiling thread, since the distributed controller no longer supports local workers
 			if (LaunchLocalShaderCompilingThread() != nullptr)
 			{
-				bExclusiveDistributedShaderCompiling = false;
+				bUseOnlyDistributedCompilationThread = false;
 			}
 		}
 	}
diff --git a/Engine/Source/Runtime/Engine/Public/ShaderCompiler.h b/Engine/Source/Runtime/Engine/Public/ShaderCompiler.h
--- a/Engine/Source/Runtime/Engine/Public/ShaderCompiler.h
+++ b/Engine/Source/Runtime/Engine/Public/ShaderCompiler.h
@@ -940,7 +940,7 @@ private:
 	/** Whether to allow shaders to compile in the background or to block after each material. */
 	bool bAllowAsynchronousShaderCompiling;
 	/** Whether shaders are compiled exclusively through the distributed shader controller and no local shader compiling thread is allocated. Ignored if 'bAllowCompilingThroughWorkers' is false. */
-	bool bExclusiveDistributedShaderCompiling = false;
+	bool bUseOnlyDistributedCompilationThread = false;
 	/** Whether to ask to retry a failed shader compile error. */
 	bool bPromptToRetryFailedShaderCompiles;
 	/** If enabled when we enter the prompt to retry we will break in the debugger if one is attached rather than prompting. */
@@ -1100,7 +1100,7 @@ public:
 	 */
 	bool IsExclusiveDistributedCompilingEnabled() const
 	{
-		return bExclusiveDistributedShaderCompiling;
+		return bUseOnlyDistributedCompilationThread;
 	}
 
 	/**
