Allow Virtualization Mode to be overriden and set at runtime on ActiveSound

diff --git a/Engine/Source/Runtime/Engine/Private/ActiveSound.cpp b/Engine/Source/Runtime/Engine/Private/ActiveSound.cpp
--- a/Engine/Source/Runtime/Engine/Private/ActiveSound.cpp
+++ b/Engine/Source/Runtime/Engine/Private/ActiveSound.cpp
@@ -82,6 +82,8 @@ FActiveSound::FActiveSound()
 	, PlayOrder(INDEX_NONE)
 	, AudioDevice(nullptr)
 	, SoundClassOverride(nullptr)
+	, VirtualizationModeOverride(EVirtualizationMode::Restart)
+	, bEnableVirtualizationModeOverride(false)
 	, bHasCheckedOcclusion(false)
 	, bAllowSpatialization(true)
 	, bHasAttenuationSettings(false)
@@ -193,7 +195,7 @@ FActiveSound* FActiveSound::CreateVirtualCopy(const FActiveSound& InActiveSoundT
 	// If we are the restart virtual mode, we need to reset our sound cue parse state and our playback time.
 	USoundBase* Sound = InActiveSoundToCopy.GetSound();
 	check(Sound);
-	if (Sound->VirtualizationMode == EVirtualizationMode::Restart)
+	if (InActiveSoundToCopy.GetVirtualizationMode() == EVirtualizationMode::Restart)
 	{
 		ActiveSound->SoundNodeOffsetMap.Reset();
 	}
@@ -387,6 +389,11 @@ bool FActiveSound::IsPlayWhenSilent() const
 		return false;
 	}
 
+	if (bEnableVirtualizationModeOverride)
+	{
+		return VirtualizationModeOverride == EVirtualizationMode::PlayWhenSilent;
+	}
+	
 	return Sound && Sound->IsPlayWhenSilent();
 }
 
@@ -402,6 +409,32 @@ float FActiveSound::GetConcurrencyPriority() const
 	}
 }
 
+EVirtualizationMode FActiveSound::GetVirtualizationMode() const
+{
+	if (bEnableVirtualizationModeOverride)
+	{
+		return VirtualizationModeOverride;
+	}
+
+	if (Sound)
+	{
+		return Sound->VirtualizationMode;
+	}
+
+	return EVirtualizationMode::Disabled;
+}
+
+void FActiveSound::OverrideVirtualizationMode(EVirtualizationMode Mode)
+{
+	VirtualizationModeOverride = Mode;
+	bEnableVirtualizationModeOverride = true;
+}
+
+void FActiveSound::ClearVirtualizationModeOverride()
+{
+	bEnableVirtualizationModeOverride = false;
+}
+
 void FActiveSound::ClearAudioComponent()
 {
 	AudioComponentID = 0;
diff --git a/Engine/Source/Runtime/Engine/Private/AudioVirtualLoop.cpp b/Engine/Source/Runtime/Engine/Private/AudioVirtualLoop.cpp
--- a/Engine/Source/Runtime/Engine/Private/AudioVirtualLoop.cpp
+++ b/Engine/Source/Runtime/Engine/Private/AudioVirtualLoop.cpp
@@ -101,7 +101,7 @@ bool FAudioVirtualLoop::Virtualize(const FActiveSound& InActiveSound, FAudioDevi
 	USoundBase* Sound = InActiveSound.GetSound();
 	check(Sound);
 
-	if (Sound->VirtualizationMode == EVirtualizationMode::Disabled)
+	if (InActiveSound.GetVirtualizationMode() == EVirtualizationMode::Disabled)
 	{
 		return false;
 	}
diff --git a/Engine/Source/Runtime/Engine/Public/ActiveSound.h b/Engine/Source/Runtime/Engine/Public/ActiveSound.h
--- a/Engine/Source/Runtime/Engine/Public/ActiveSound.h
+++ b/Engine/Source/Runtime/Engine/Public/ActiveSound.h
@@ -386,6 +386,14 @@ public:
 	/** The priority used for concurrency calculations */
 	float GetConcurrencyPriority() const;
 
+	EVirtualizationMode GetVirtualizationMode() const;
+
+	/** Override the Virtualization settings set on Audio Component/Sound */
+	void OverrideVirtualizationMode(EVirtualizationMode Mode);
+
+	/** Stop using the Virtualization Mode Override */
+	void ClearVirtualizationModeOverride();
+
 	FAudioDevice* AudioDevice;
 
 	/** The concurrent groups that this sound is actively playing in. */
@@ -408,6 +416,10 @@ private:
 
 	TSharedPtr<Audio::IParameterTransmitter> InstanceTransmitter;
 
+	EVirtualizationMode VirtualizationModeOverride;
+
+	uint8 bEnableVirtualizationModeOverride: 1;
+
 public:
 	Audio::IParameterTransmitter* GetTransmitter()
 	{
@@ -588,7 +600,7 @@ public:
 	FDynamicParameter CurrentOcclusionVolumeAttenuation;
 
 	float SubtitlePriority;
-
+	
 	/** The product of the component priority and the USoundBase priority */
 	float Priority;
 
