[shaders] reset duplicate flag when a duplicate job is promoted to the "main" job during the cancellation process of another job. fixes asserts like this:

Assertion failed: !bCompilationSkipped || SingleJob.Output.CompileTime == 0.0f

diff --git a/Engine/Source/Runtime/Engine/Private/ShaderCompiler/ShaderCompilerJobCache.cpp b/Engine/Source/Runtime/Engine/Private/ShaderCompiler/ShaderCompilerJobCache.cpp
--- a/Engine/Source/Runtime/Engine/Private/ShaderCompiler/ShaderCompilerJobCache.cpp
+++ b/Engine/Source/Runtime/Engine/Private/ShaderCompiler/ShaderCompilerJobCache.cpp
@@ -1275,6 +1275,8 @@ int32 FShaderJobCache::RemoveAllPendingJobsWithId(uint32 InId)
 							Unlink(*DuplicateJob);
 							RemoveDuplicateJob(DuplicateJob);
 
+							DuplicateJob->JobStatusPtr->SetIsDuplicate(false);
+
 							// Add it as pending at the appropriate priority
 							GShaderCompilerStats->RegisterNewPendingJob(*DuplicateJob);
 
diff --git a/Engine/Source/Runtime/RenderCore/Public/ShaderCompilerJobTypes.h b/Engine/Source/Runtime/RenderCore/Public/ShaderCompilerJobTypes.h
--- a/Engine/Source/Runtime/RenderCore/Public/ShaderCompilerJobTypes.h
+++ b/Engine/Source/Runtime/RenderCore/Public/ShaderCompilerJobTypes.h
@@ -110,9 +110,9 @@ public:
 		return Status.load();
 	}
 
-	void SetIsDuplicate()
+	void SetIsDuplicate(bool bDuplicate = true)
 	{
-		IsDuplicate.store(true);
+		IsDuplicate.store(bDuplicate);
 	}
 
 	bool GetIsDuplicate() const
