Add concept of TickGroupWork to the task sync manager, which is a gamethread-only tick function that runs at the end of a tick group.
Change it so RequestWork will push into this global list if the specific task has already been dispatched. This means it will run after any new tick functions that were added during the frame.
Fix threading issues with requesting reserved work, this now works fully.
Fix some editor-only issues with registering sync points.

diff --git a/Engine/Source/Runtime/Engine/Private/TaskSyncManager.cpp b/Engine/Source/Runtime/Engine/Private/TaskSyncManager.cpp
--- a/Engine/Source/Runtime/Engine/Private/TaskSyncManager.cpp
+++ b/Engine/Source/Runtime/Engine/Private/TaskSyncManager.cpp
@@ -91,7 +91,12 @@ FTickFunction* FActiveSyncWorkHandle::GetDependencyTickFunction()
 {
 	if (IsInGameThread())
 	{
-		return SyncPoint.Get();
+		FActiveSyncPoint* SyncPointFunction = SyncPoint.Get();
+		if (SyncPointFunction && !SyncPointFunction->IsTickGroupWork())
+		{
+			// Tick group work can't be used for dependencies
+			return SyncPointFunction;
+		}
 	}
 	return nullptr;
 }
@@ -142,6 +147,44 @@ bool FActiveSyncWorkHandle::RequestWork(FTickFunction* WorkFunction, ESyncWorkRe
 	return false;
 }
 
+/** Task to execute a tick function manually outside of it's normal tick group */
+// TODO move into TickTaskManager?
+class FManualTickFunctionTask
+{
+	FTickFunction* Target;
+	ENamedThreads::Type DesiredThread;
+	float DeltaSeconds;
+	TEnumAsByte<ELevelTick> TickType;
+
+public:
+	FManualTickFunctionTask(FTickFunction* InTarget, ENamedThreads::Type InDesiredThread, float InDeltaSeconds, TEnumAsByte<ELevelTick> InTickType)
+		: Target(InTarget)
+		, DesiredThread(InDesiredThread)
+		, DeltaSeconds(InDeltaSeconds)
+		, TickType(InTickType)
+	{
+	}
+	static TStatId GetStatId()
+	{
+		RETURN_QUICK_DECLARE_CYCLE_STAT(FManualTickFunctionTask, STATGROUP_TaskGraphTasks);
+	}
+	ENamedThreads::Type GetDesiredThread()
+	{
+		return DesiredThread;
+	}
+	static ESubsequentsMode::Type GetSubsequentsMode()
+	{
+		return ESubsequentsMode::TrackSubsequents;
+	}
+
+	void DoTask(ENamedThreads::Type CurrentThread, const FGraphEventRef& MyCompletionGraphEvent)
+	{
+		if (Target && Target->IsTickFunctionEnabled())
+		{
+			Target->ExecuteTick(DeltaSeconds, TickType, CurrentThread, MyCompletionGraphEvent);
+		}
+	}
+};
 
 FActiveSyncPoint::FActiveSyncPoint()
 {
@@ -163,6 +206,7 @@ void FActiveSyncPoint::ResetWorkForFrame()
 		check(!CurrentWork.bCurrentlyExecuting);
 
 		CurrentWork.bAlreadyExecuted = false;
+		CurrentWork.bLateWorkRequested = false;
 
 		if (CurrentWork.bRequestEveryFrame)
 		{
@@ -173,6 +217,12 @@ void FActiveSyncPoint::ResetWorkForFrame()
 		{
 			CurrentWork.bWorkReserved = true;
 		}
+
+		if (!CurrentWork.bWorkRequested && CurrentWork.TickFunction)
+		{
+			// This could delete a wrapper tick function
+			CurrentWork.SetTickFunction(nullptr, false);
+		}
 		// Don't reset reserved/requested in case they were set before the frame started
 	}
 }
@@ -203,6 +253,12 @@ bool FActiveSyncPoint::HandleFirstTickGroup()
 
 bool FActiveSyncPoint::IsReadyToProcessWork()
 {
+	if (FrameStatus < ESyncPointStatus::DispatchWaitingForWork || FrameStatus == ESyncPointStatus::ExecutionComplete)
+	{
+		// Too early or late to process work
+		return false;
+	}
+
 	// See if we should dispatch a process task
 	bool bHasRequest = false;
 	bool bHasReservation = false;
@@ -227,31 +283,33 @@ bool FActiveSyncPoint::IsReadyToProcessWork()
 	return false;
 }
 
-bool FActiveSyncPoint::GetWorkToExecute(int32& NumReserved)
+bool FActiveSyncPoint::GetWorkToExecute(float DeltaTime, ELevelTick TickType, ENamedThreads::Type CurrentThread, const FGraphEventRef& MyCompletionGraphEvent)
 {
 	// Copy work into executing work and possibly delete active work
 	UE::TScopeLock WorkScope(WorkLock);
 
-	check(ExecutingWork.Num() == 0); // ExecuteTick cannot be re-entrant and will not be resized during execution
-
-	check(FrameStatus >= ESyncPointStatus::Dispatched && FrameStatus != ESyncPointStatus::ExecutionComplete);
-	FrameStatus = ESyncPointStatus::Executing;
-	ExecutingWork.Reserve(ActiveWork.Num());
-	NumReserved = 0;
+	if (FrameStatus != ESyncPointStatus::Executing)
+	{
+		check(FrameStatus >= ESyncPointStatus::Dispatched && FrameStatus != ESyncPointStatus::ExecutionComplete);
+		check(ExecutingWork.Num() == 0);
+		FrameStatus = ESyncPointStatus::Executing;
+	}
+	
+	const int32 NumActive = ActiveWork.Num();
+	int32 NumReserved = 0;
+	int32 NumExecuting = 0;
+	
+	// Presize array to make sure there's room, we don't shrink as it will be the same across frames
+	ExecutingWork.SetNumUninitialized(NumActive, EAllowShrinking::No);
 
-	for (int32 i = 0; i < ActiveWork.Num(); i++)
+	for (int32 ActiveIndex = 0; ActiveIndex < NumActive; ActiveIndex++)
 	{
-		FActiveSyncWork& CurrentWork = ActiveWork[i];
+		FActiveSyncWork& CurrentWork = ActiveWork[ActiveIndex];
 		if (CurrentWork.bCurrentlyExecuting)
 		{
 			CurrentWork.bCurrentlyExecuting = false;
 			CurrentWork.bAlreadyExecuted = true;
-			if (!CurrentWork.bRequestEveryFrame)
-			{
-				// Clear function now
-				CurrentWork.TickFunction = nullptr;
-				CurrentWork.bWorkRequested = false;
-			}
+			CurrentWork.bWorkRequested = false;
 		}
 		else if (CurrentWork.bWorkReserved)
 		{
@@ -259,18 +317,34 @@ bool FActiveSyncPoint::GetWorkToExecute(int32& NumReserved)
 		}
 		else if (CurrentWork.bWorkRequested && !CurrentWork.bAlreadyExecuted)
 		{
-			ExecutingWork.Add(CurrentWork);
 			CurrentWork.bCurrentlyExecuting = true;
+
+			ExecutingWork[NumExecuting].ActiveWorkIndex = ActiveIndex;
+			ExecutingWork[NumExecuting].TickFunction = CurrentWork.TickFunction;
+			NumExecuting++;
 		}
 	}
 
-	if (ExecutingWork.Num() > 0)
+	// Set array to number of things copied, this is faster than adding as we go
+	ExecutingWork.SetNumUninitialized(NumExecuting, EAllowShrinking::No);
+
+	if (NumExecuting > 0)
 	{
 		return true;
 	}
 	else if (NumReserved > 0)
 	{
 		FrameStatus = ESyncPointStatus::WaitingForMoreWork;
+
+		check(SyncPointDescription.ActivationRules == ESyncPointActivationRules::ActivateForAnyWork);
+		check(MyCompletionGraphEvent.IsValid() && !ReactivationEvent.IsValid());
+
+		// Create our redo event but don't dispatch it yet
+		// This might get dispatched immediately after leaving the lock here
+		FBaseGraphTask* Task = TGraphTask<FManualTickFunctionTask>::CreateTask(nullptr, CurrentThread).ConstructAndHold(static_cast<FTickFunction*>(this), CurrentThread, DeltaTime, TickType);
+		ReactivationEvent = Task->GetCompletionEvent();
+		MyCompletionGraphEvent->DontCompleteUntil(ReactivationEvent);
+
 		return false;
 	}
 	else
@@ -280,6 +354,38 @@ bool FActiveSyncPoint::GetWorkToExecute(int32& NumReserved)
 	}
 }
 
+void FActiveSyncPoint::CancelRequestedWork(int32 ActiveWorkIndex, FActiveSyncWork& CurrentWork)
+{
+	if (CurrentWork.bLateWorkRequested)
+	{
+		// This turned into a late work request so cancel it
+		FTaskSyncManager* Manager = FTaskSyncManager::Get();
+		Manager->CancelTemporaryWorkRequest(this, ActiveWorkIndex);
+	}
+
+	if (CurrentWork.bWorkRequested && CurrentWork.bCurrentlyExecuting)
+	{
+		// This is called with a WorkLock when another thread may be in execute tick, so we can only make simple value changes
+		for (int32 ExecutingIndex = 0; ExecutingIndex < ExecutingWork.Num(); ExecutingIndex++)
+		{
+			if (ExecutingWork[ExecutingIndex].ActiveWorkIndex == ActiveWorkIndex)
+			{
+				ExecutingWork[ExecutingIndex].Invalidate();
+			}
+		}
+	}
+
+	CurrentWork.bWorkRequested = false;
+	CurrentWork.bRequestEveryFrame = false;
+}
+
+void FActiveSyncPoint::ExecuteFromGameThread(float DeltaTime, ELevelTick TickType)
+{
+	// TODO turn on a mode to disable some logic?
+	FGraphEventRef EmptyRef;
+	ExecuteTick(DeltaTime, TickType, ENamedThreads::GameThread, EmptyRef);
+}
+
 void FActiveSyncPoint::ExecuteTick(float DeltaTime, ELevelTick TickType, ENamedThreads::Type CurrentThread, const FGraphEventRef& MyCompletionGraphEvent)
 {
 	// Only lock during acquisition, GetWorkToExecute protects against recursive ticking
@@ -287,44 +393,23 @@ void FActiveSyncPoint::ExecuteTick(float DeltaTime, ELevelTick TickType, ENamedT
 
 	if (SyncPointDescription.EventType == ESyncPointEventType::SimpleEvent)
 	{
+		FrameStatus = ESyncPointStatus::ExecutionComplete;
 		return;
 	}
 
 	// TODO add pre/post batch function calls
 
-	int32 NumReserved = 0;
-	while (GetWorkToExecute(NumReserved))
+	while (GetWorkToExecute(DeltaTime, TickType, CurrentThread, MyCompletionGraphEvent))
 	{
 		for (int32 i = 0; i < ExecutingWork.Num(); i++)
 		{
-			// Executing work cannot change size during execution
+			// Executing work cannot change size during execution, but TickFunction could be nulled
 			FTickFunction* TickFunction = ExecutingWork[i].TickFunction;
 			if (TickFunction && TickFunction->IsTickFunctionEnabled())
 			{
 				TickFunction->ExecuteNestedTick(DeltaTime, TickType, CurrentThread, MyCompletionGraphEvent);
 			}
 		}
-
-		ExecutingWork.SetNum(0, EAllowShrinking::No);
-	}
-
-	if (NumReserved > 0)
-	{
-		UE::TScopeLock WorkScope(WorkLock);
-
-		check(FrameStatus == ESyncPointStatus::WaitingForMoreWork);
-
-		if (ensure(SyncPointDescription.ActivationRules == ESyncPointActivationRules::ActivateForAnyWork))
-		{
-			ReactivationEvent = FGraphEvent::CreateGraphEvent();
-
-			// TODO create a new task to process the queue, need to add something to TickTaskManager
-		}
-		else
-		{
-			// This is invalid but set us to complete
-			FrameStatus = ESyncPointStatus::ExecutionComplete;
-		}
 	}
 }
 
@@ -372,23 +457,31 @@ FActiveSyncWorkHandle FActiveSyncPoint::RegisterWorkHandle()
 bool FActiveSyncPoint::ReserveFutureWork(FActiveSyncWorkHandle& Handle, ESyncWorkRepetition Repeat)
 {
 	UE::TScopeLock WorkScope(WorkLock);
+	int32 WorkIndex = Handle.GetIndex();
 
-	if (ensure(Handle.IsValid() && ActiveWork.IsValidIndex(Handle.WorkIndex)))
+	if (ensure(Handle.IsValid() && ActiveWork.IsValidIndex(WorkIndex)))
 	{
-		FActiveSyncWork& CurrentWork = ActiveWork[Handle.WorkIndex];
+		FActiveSyncWork& CurrentWork = ActiveWork[WorkIndex];
 
-		if (ensure(CurrentWork.bHasActiveHandle && !CurrentWork.bWorkReserved && !CurrentWork.bWorkRequested))
+		if (ensure(CurrentWork.bHasActiveHandle))
 		{
-			// TODO properly handle the repetition types
-
-			CurrentWork.bWorkReserved = true;
 			CurrentWork.bReserveEveryFrame = (Repeat == ESyncWorkRepetition::EveryFrame);
 
-			Handle.bWorkReserved = true;
+			if (Repeat == ESyncWorkRepetition::Never)
+			{
+				Handle.bWorkReserved = false;
+
+				CurrentWork.bWorkReserved = false;
+			}
+			else
+			{
+				Handle.bWorkReserved = true;
+				// TODO what do we do if you reserve after requesting this frame?
+				CurrentWork.bWorkReserved = true;
+			}
 
-			if (FrameStatus > ESyncPointStatus::DispatchWaitingForWork && IsReadyToProcessWork())
+			if (IsReadyToProcessWork())
 			{
-				WorkScope.Unlock();
 				DispatchWorkTask();
 			}
 
@@ -401,29 +494,48 @@ bool FActiveSyncPoint::ReserveFutureWork(FActiveSyncWorkHandle& Handle, ESyncWor
 bool FActiveSyncPoint::RequestWork(FActiveSyncWorkHandle& Handle, FTickFunction* WorkFunction, ESyncWorkRepetition Repeat)
 {
 	UE::TScopeLock WorkScope(WorkLock);
+	int32 WorkIndex = Handle.GetIndex();
 
-	if (ensure(Handle.IsValid() && ActiveWork.IsValidIndex(Handle.WorkIndex)))
+	if (ensure(Handle.IsValid() && ActiveWork.IsValidIndex(WorkIndex)))
 	{
-		FActiveSyncWork& CurrentWork = ActiveWork[Handle.WorkIndex];
+		FActiveSyncWork& CurrentWork = ActiveWork[WorkIndex];
 
-		if (ensure(CurrentWork.bHasActiveHandle && !CurrentWork.bWorkRequested))
+		if (ensure(CurrentWork.bHasActiveHandle))
 		{
-			// TODO properly handle the repetition types
+			if (Repeat == ESyncWorkRepetition::Never)
+			{
+				Handle.bWorkRequested = false;
 
-			CurrentWork.TickFunction = WorkFunction;
-			CurrentWork.bWorkReserved = false;
-			CurrentWork.bWorkRequested = true;
-			CurrentWork.bRequestEveryFrame = (Repeat == ESyncWorkRepetition::EveryFrame);
+				CancelRequestedWork(WorkIndex, CurrentWork);
+			}
+			else 
+			{
+				Handle.bWorkRequested = true;
+
+				// If the task is completely done or in the middle of processing and we didn't reserve work
+				if (FrameStatus == ESyncPointStatus::ExecutionComplete || (FrameStatus >= ESyncPointStatus::Executing && !CurrentWork.bWorkReserved))
+				{
+					FTaskSyncManager* Manager = FTaskSyncManager::Get();
+					
+					// This may schedule it as part of tick group work
+					CurrentWork.bLateWorkRequested = Manager->HandleLateWorkRequest(this, WorkIndex, WorkFunction);
+				}
 
-			Handle.bWorkReserved = false;
-			Handle.bWorkRequested = true;
+				CurrentWork.bWorkReserved = false;
+				CurrentWork.bWorkRequested = !CurrentWork.bLateWorkRequested;
 
-			if (FrameStatus > ESyncPointStatus::DispatchWaitingForWork && IsReadyToProcessWork())
+				CurrentWork.bRequestEveryFrame = (Repeat == ESyncWorkRepetition::EveryFrame);
+				if (CurrentWork.bRequestEveryFrame || CurrentWork.bWorkRequested)
+				{
+					// Set function ptr for when it executes which may be next frame
+					CurrentWork.SetTickFunction(WorkFunction, false);
+				}
+			}
+
+			if (IsReadyToProcessWork())
 			{
-				WorkScope.Unlock();
 				DispatchWorkTask();
 			}
-
 			return true;
 		}
 	}
@@ -433,28 +545,26 @@ bool FActiveSyncPoint::RequestWork(FActiveSyncWorkHandle& Handle, FTickFunction*
 bool FActiveSyncPoint::AbandonWork(FActiveSyncWorkHandle& Handle)
 {
 	UE::TScopeLock WorkScope(WorkLock);
+	int32 WorkIndex = Handle.GetIndex();
 	
-	if (ensure(Handle.IsValid() && ActiveWork.IsValidIndex(Handle.WorkIndex)))
+	if (ensure(Handle.IsValid() && ActiveWork.IsValidIndex(WorkIndex)))
 	{
-		FActiveSyncWork& CurrentWork = ActiveWork[Handle.WorkIndex];
+		FActiveSyncWork& CurrentWork = ActiveWork[WorkIndex];
 
 		if (ensure(CurrentWork.bHasActiveHandle))
 		{
 			if (CurrentWork.bWorkReserved || CurrentWork.bWorkRequested || CurrentWork.bReserveEveryFrame || CurrentWork.bRequestEveryFrame)
 			{
-				// TODO What do we do if work has already been scheduled but not executed?
+				CancelRequestedWork(WorkIndex, CurrentWork);
 
 				CurrentWork.bWorkReserved = false;
-				CurrentWork.bWorkRequested = false;
 				CurrentWork.bReserveEveryFrame = false;
-				CurrentWork.bRequestEveryFrame = false;
 
 				Handle.bWorkReserved = false;
 				Handle.bWorkRequested = false;
 			
-				if (FrameStatus > ESyncPointStatus::DispatchWaitingForWork && IsReadyToProcessWork())
+				if (IsReadyToProcessWork())
 				{
-					WorkScope.Unlock();
 					DispatchWorkTask();
 				}
 			}
@@ -467,23 +577,28 @@ bool FActiveSyncPoint::AbandonWork(FActiveSyncWorkHandle& Handle)
 bool FActiveSyncPoint::ResetWorkHandle(FActiveSyncWorkHandle& Handle)
 {
 	UE::TScopeLock WorkScope(WorkLock);
+	int32 WorkIndex = Handle.GetIndex();
 
-	if (ensure(Handle.IsValid() && ActiveWork.IsValidIndex(Handle.WorkIndex)))
+	if (ensure(Handle.IsValid() && ActiveWork.IsValidIndex(WorkIndex)))
 	{
-		FActiveSyncWork& CurrentWork = ActiveWork[Handle.WorkIndex];
+		FActiveSyncWork& CurrentWork = ActiveWork[WorkIndex];
 
 		if (ensure(CurrentWork.bHasActiveHandle))
 		{
 			bool bShouldCheckWork = CurrentWork.bWorkReserved || CurrentWork.bWorkRequested;
 
-			CurrentWork.Reset();
-			Handle.ResetInternal();
+			CancelRequestedWork(WorkIndex, CurrentWork);
+			
+			CurrentWork.bWorkReserved = false;
+			CurrentWork.bReserveEveryFrame = false;
+			CurrentWork.bHasActiveHandle = false;
 
-			// TODO What do we do if work has already been scheduled but not executed?
+			// ResetWorkForFrame handles clearing the function
 
-			if (bShouldCheckWork && FrameStatus > ESyncPointStatus::DispatchWaitingForWork && IsReadyToProcessWork())
+			Handle.ResetInternal();
+
+			if (bShouldCheckWork && IsReadyToProcessWork())
 			{
-				WorkScope.Unlock();
 				DispatchWorkTask();
 			}
 			return true;
@@ -515,10 +630,20 @@ void FActiveSyncPoint::DispatchWorkTask()
 	{
 		ensure(DispatchManually());
 	}
+	else if (FrameStatus == ESyncPointStatus::WaitingForMoreWork)
+	{
+		check(ReactivationEvent.IsValid());
+		// Dispatch our reactivation event to let the task activate
+		ReactivationEvent->Unlock();
+		ReactivationEvent.SafeRelease();
+		FrameStatus = ESyncPointStatus::Dispatched; // TODO do we need a different state for dispatched a second time?
+	}
 	else
 	{
-		// TODO Reschedule work with ReactivationEvent
+		// If it gets here the work will get handled by the already dispatched task
+		check(FrameStatus == ESyncPointStatus::Dispatched || FrameStatus == ESyncPointStatus::Executing);
 	}
+	
 }
 
 FTaskSyncManager* FTaskSyncManager::Get()
@@ -605,6 +730,13 @@ void FTaskSyncManager::ReloadRegisteredData()
 
 	for (const FSyncPointDescription& SyncPoint : Settings->RegisteredSyncPoints)
 	{
+		if (SyncPoint.EventType == ESyncPointEventType::Invalid || SyncPoint.ActivationRules == ESyncPointActivationRules::Invalid || !SyncPoint.IsValid())
+		{
+			// Only error outside the editor as this can happen while modifying the project settings
+			UE_CLOG(!GIsEditor, LogTaskSync, Error, TEXT("Cannot register invalid sync point %s from project settings!"), *SyncPoint.RegisteredName.ToString());
+
+			continue;
+		}
 		RegisterSyncPointInternal(SyncPoint);
 	}
 
@@ -702,13 +834,15 @@ bool FTaskSyncManager::UnregisterSyncPoint(FName RegisteredName, FName SourceNam
 	return false;
 }
 
-void FTaskSyncManager::StartFrame(const UWorld* InWorld)
+void FTaskSyncManager::StartFrame(const UWorld* InWorld, float InDeltaSeconds, ELevelTick InTickType)
 {
 	check(InWorld);	
 	check(IsInGameThread() && !IsTicking() && CurrentTickGroup == ETickingGroup::TG_MAX);
 
 	UE::TScopeLock Lock(ManagerLock);
 	CurrentTickWorld = InWorld;
+	CurrentDeltaTime = InDeltaSeconds;
+	CurrentTickType = InTickType;
 	ULevel* PersistentLevel = InWorld->PersistentLevel.Get();
 	FWorldContextId WorldContext(InWorld);
 
@@ -766,13 +900,16 @@ void FTaskSyncManager::InitializeBatchForFrame(FBatchData& BatchData, ULevel* Pe
 					}
 				}
 
-				ensure(SyncPoint->FrameStatus == ESyncPointStatus::TaskNotRegistered);
+				// This may have been unregistered automatically due to level change
+				ensure(SyncPoint->FrameStatus == ESyncPointStatus::TaskNotRegistered || SyncPoint->FrameStatus == ESyncPointStatus::ExecutionComplete);
 				SyncPoint->RegisterTickFunction(PersistentLevel);
 			}
 			SyncPoint->FrameStatus = ESyncPointStatus::TaskNotCreated;
 			SyncPoint->ResetWorkForFrame();
 		}
 	}
+
+	// TODO execute one of the TickGroupWork functions? But should it be at the top or bottom
 }
 
 FActiveSyncPoint* FTaskSyncManager::GetOrCreateSyncPoint(FBatchData& BatchData, FRegisteredSyncPointData& SyncData)
@@ -792,6 +929,7 @@ FActiveSyncPoint* FTaskSyncManager::GetOrCreateSyncPoint(FBatchData& BatchData,
 		NewTickFunction->SyncPointDescription = Description;
 		NewTickFunction->SyncPointId.BatchContext = BatchData.BatchContext;
 		NewTickFunction->SyncPointId.SyncId = SyncData.RegisteredId;
+		NewTickFunction->WorldContextId = BatchData.WorldContext;
 
 		switch (Description.EventType)
 		{
@@ -840,6 +978,33 @@ FActiveSyncPoint* FTaskSyncManager::GetOrCreateSyncPoint(FBatchData& BatchData,
 	return &FoundData->Get();
 }
 
+FActiveSyncPoint* FTaskSyncManager::GetOrCreateTickGroupWork(FBatchData& BatchData, ETickingGroup TickGroup)
+{
+	check(TickGroup >= 0 && TickGroup < ETickingGroup::TG_MAX);
+	if (TickGroup > BatchData.TickGroupWork.Num())
+	{
+		BatchData.TickGroupWork.SetNum(TickGroup);
+	}
+
+	FActiveSyncPoint* FoundPoint = BatchData.TickGroupWork[TickGroup].Get();
+	if (FoundPoint)
+	{
+		return FoundPoint;
+	}
+
+	TSharedRef<FActiveSyncPoint> NewTickFunction = MakeShared<FActiveSyncPoint>();
+	NewTickFunction->TickGroup = TickGroup;
+	NewTickFunction->EndTickGroup = TickGroup;
+	NewTickFunction->WorldContextId = BatchData.WorldContext;
+	NewTickFunction->SyncPointDescription.EventType = ESyncPointEventType::GameThreadTask;
+	NewTickFunction->SyncPointDescription.ActivationRules = ESyncPointActivationRules::AlwaysActivate;
+
+	check(NewTickFunction->IsTickGroupWork());
+	
+	BatchData.TickGroupWork[TickGroup] = NewTickFunction;
+	return &NewTickFunction.Get();
+}
+
 void FTaskSyncManager::StartTickGroup(const UWorld* InWorld, ETickingGroup Group, TArray<FTickFunction*>& TicksToManualDispatch)
 {
 	check(IsInGameThread() && CurrentTickWorld == InWorld && CurrentTickGroup == ETickingGroup::TG_MAX);
@@ -869,6 +1034,22 @@ void FTaskSyncManager::StartTickGroup(const UWorld* InWorld, ETickingGroup Group
 void FTaskSyncManager::EndTickGroup(const UWorld* InWorld, ETickingGroup Group)
 {
 	check(IsInGameThread() && CurrentTickWorld == InWorld && CurrentTickGroup != ETickingGroup::TG_MAX);
+
+	// For all batches that match the world
+	for (FBatchData& BatchData : BatchList)
+	{
+		if (BatchData.TickGroupWork.IsValidIndex(Group))
+		{
+			FActiveSyncPoint* TickGroupWork = BatchData.TickGroupWork[Group].Get();
+			if (TickGroupWork)
+			{
+				// TODO We also may want to run this earlier in the tick group, such as after the first stalled request
+				TickGroupWork->ExecuteFromGameThread(CurrentDeltaTime, CurrentTickType);
+			}
+		}
+	}
+
+
 	CurrentTickGroup = ETickingGroup::TG_MAX;
 }
 
@@ -877,6 +1058,17 @@ void FTaskSyncManager::EndFrame(const UWorld* InWorld)
 	check(IsInGameThread() && CurrentTickWorld == InWorld && CurrentTickGroup == ETickingGroup::TG_MAX);
 	CurrentTickWorld = nullptr;
 	CurrentTickGroup = ETickingGroup::TG_MAX;
+
+	FWorldContextId WorldContext(InWorld);
+
+	// For all batches that match the world
+	for (FBatchData& BatchData : BatchList)
+	{
+		if (BatchData.WorldContext == WorldContext)
+		{
+			BatchData.TemporaryWorkRequests.Reset();
+		}
+	}
 }
 
 void FTaskSyncManager::ReleaseWorldContext(FWorldContextId WorldContext)
@@ -1135,6 +1327,29 @@ FTaskSyncResult FTaskSyncManager::RegisterWorkHandle(FSyncPointId SyncPoint, FAc
 	return ReturnStatus;
 }
 
+FTaskSyncResult FTaskSyncManager::RegisterTickGroupWorkHandle(FWorldContextId WorldContext, ETickingGroup TickGroup, FActiveSyncWorkHandle& OutWorkHandle)
+{
+	UE::TScopeLock ManagerScope(ManagerLock);
+
+	// Find default batch for world
+	for (FBatchData& BatchData : BatchList)
+	{
+		if (BatchData.WorldContext == WorldContext)
+		{
+			FActiveSyncPoint* TickGroupWork = GetOrCreateTickGroupWork(BatchData, TickGroup);
+
+			if (TickGroupWork)
+			{
+				OutWorkHandle = TickGroupWork->RegisterWorkHandle();
+				return FTaskSyncResult(ESyncOperationResult::Success);
+			}
+			return FTaskSyncResult(ESyncOperationResult::EventTypeIncorrect);
+		}
+	}
+
+	return FTaskSyncResult(ESyncOperationResult::WorldNotFound);
+}
+
 FTaskSyncResult FTaskSyncManager::FindActiveSyncPoint(FSyncPointId SyncPoint, TSharedPtr<FActiveSyncPoint>& OutData)
 {
 	// Internal only, does not lock
@@ -1158,3 +1373,61 @@ FTaskSyncResult FTaskSyncManager::FindActiveSyncPoint(FSyncPointId SyncPoint, TS
 	}
 	return ESyncOperationResult::BatchNotFound;
 }
+
+bool FTaskSyncManager::HandleLateWorkRequest(FActiveSyncPoint* RequestedSyncPoint, int32 RequestingHandle, FTickFunction* TickFunction)
+{
+	check(RequestedSyncPoint && RequestedSyncPoint->FrameStatus >= ESyncPointStatus::Dispatched);
+	ETickingGroup RequestedTickGroup = RequestedSyncPoint->TickGroup;
+	FWorldContextId WorldContext = RequestedSyncPoint->GetWorldContextId();
+	
+	UE::TScopeLock ManagerScope(ManagerLock);
+
+	if (CurrentTickGroup == TG_MAX)
+	{
+		// Not ticking any more, just queue for next frame
+		return false;
+	}
+	
+	if (CurrentTickGroup < TG_MAX && CurrentTickGroup > RequestedTickGroup)
+	{
+		// Move to the current tick group, which will be processed at the end
+		// TODO maybe this should be TG_NewlySpawned in some cases? Need to review the logic for newly-registered tick functions
+		RequestedTickGroup = CurrentTickGroup;
+	}
+
+	// Find default batch for world
+	for (FBatchData& BatchData : BatchList)
+	{
+		if (BatchData.WorldContext == WorldContext)
+		{
+			FActiveSyncPoint* TickGroupWork = GetOrCreateTickGroupWork(BatchData, RequestedTickGroup);
+			if (ensure(TickGroupWork))
+			{
+				FTemporaryWorkRequest& WorkRequest = BatchData.TemporaryWorkRequests.Emplace_GetRef(RequestedSyncPoint, RequestingHandle);
+				WorkRequest.WorkHandle = TickGroupWork->RegisterWorkHandle();
+
+				ensure(WorkRequest.WorkHandle.RequestWork(TickFunction, ESyncWorkRepetition::Once));
+				return true;
+			}
+		}
+	}
+
+	return false;
+}
+
+void FTaskSyncManager::CancelTemporaryWorkRequest(FActiveSyncPoint* RequestedSyncPoint, int32 RequestingHandle)
+{
+	UE::TScopeLock ManagerScope(ManagerLock);
+
+	// Check all batches as the world context may have been lost
+	for (FBatchData& BatchData : BatchList)
+	{
+		for (FTemporaryWorkRequest& WorkRequest : BatchData.TemporaryWorkRequests)
+		{
+			if (RequestedSyncPoint == WorkRequest.RequestingSyncPoint && RequestingHandle == WorkRequest.RequestingHandle)
+			{
+				WorkRequest.WorkHandle.Reset();
+			}
+		}
+	}
+}
\ No newline at end of file
diff --git a/Engine/Source/Runtime/Engine/Private/TickTaskManager.cpp b/Engine/Source/Runtime/Engine/Private/TickTaskManager.cpp
--- a/Engine/Source/Runtime/Engine/Private/TickTaskManager.cpp
+++ b/Engine/Source/Runtime/Engine/Private/TickTaskManager.cpp
@@ -1876,7 +1876,7 @@ public:
 		if (FTaskSyncManager* SyncManager = FTaskSyncManager::Get())
 		{
 			// This can create tick functions
-			SyncManager->StartFrame(InWorld);
+			SyncManager->StartFrame(InWorld, InDeltaSeconds, InTickType);
 		}
 
 		Context.TickGroup = ETickingGroup(0); // reset this to the start tick group
diff --git a/Engine/Source/Runtime/Engine/Public/TaskSyncManager.h b/Engine/Source/Runtime/Engine/Public/TaskSyncManager.h
--- a/Engine/Source/Runtime/Engine/Public/TaskSyncManager.h
+++ b/Engine/Source/Runtime/Engine/Public/TaskSyncManager.h
@@ -384,7 +384,6 @@ struct FSyncPointId final
 };
 
 /** Used to specify how many times work should be executed as part of a sync point */
-// TODO May need to change this enum to deal with in-frame updates
 enum class ESyncWorkRepetition : uint8
 {
 	/** Do not perform this work, used to cancel previous requests */
@@ -422,13 +421,16 @@ struct FActiveSyncWorkHandle final
 	 */
 	ENGINE_API FTickFunction* GetDependencyTickFunction();
 
-	/** True if this handle has been used to reserve work. This will still be true after work has been requested */
+	/** True if this handle has been used to reserve work. This will still be true until work is abandoned */
 	ENGINE_API bool HasReservedWork() const;
 
-	/** True if this handle has been used to request work. This will still be true during and after it is executed */
+	/** True if this handle has been used to request work. This will still be true until work is abandoned */
 	ENGINE_API bool HasRequestedWork() const;
 
-	/** Reserve work using this handle, possibly for every frame */
+	/** 
+	 * Request to reserve work using this handle, which must be filled to complete the frame's tasks. 
+	 * If EveryFrame is passed in, it will reserve work at the start of every frame that must be requested or abandoned.
+	 */
 	ENGINE_API bool ReserveFutureWork(ESyncWorkRepetition Repeat);
 
 	/**
@@ -439,7 +441,6 @@ struct FActiveSyncWorkHandle final
 
 	/** 
 	 * Abandon any requested or reserved work and return to how it was at initial registration.
-	 * This will not stop any requested work that has already started on another thread.
 	 * This clears HasReserved/HasRequested, but does not invalidate the handle.
 	 * Return true if there was any requested/reserved work to abandon.
 	 */
@@ -458,14 +459,18 @@ private:
 	/** Max amount of work per sync point */
 	static constexpr uint32 MaxWorkIndex = (2 << 24) - 1;
 
-	// TODO May remove these flags or add requested type
-
-	/** True if this is a reserved work handle that needs to be requested or abandoned */
+	/** True if this may represent active reserved work */
 	uint32 bWorkReserved : 1 = 0;
 
-	/** True if this is a requested work handle that will be executed */
+	/** True if this may represent active requested work */
 	uint32 bWorkRequested : 1 = 0;
 
+	/** Get index as signed int */
+	inline int32 GetIndex() const
+	{
+		return static_cast<int32>(WorkIndex);
+	}
+
 	/** Reset a handle to point to nothing */
 	inline void ResetInternal()
 	{
@@ -490,9 +495,15 @@ struct FActiveSyncPoint final : public TSharedFromThis<FActiveSyncPoint, ESPMode
 	/** Returns the sync point id that created this, including the current batch context */
 	inline const FSyncPointId& GetSyncPointId() const { return SyncPointId; }
 
+	/** Returns the world context */
+	inline const FWorldContextId& GetWorldContextId() const { return WorldContextId; }
+
 	/** Returns current status */
 	inline ESyncPointStatus GetFrameStatus() const { return FrameStatus; }
 
+	/** True if this is a special work function that is never registered */
+	inline bool IsTickGroupWork() const { return !SyncPointId.IsValid(); }
+
 	/** Registers a new work handle that can be used to reserve and request work */
 	ENGINE_API FActiveSyncWorkHandle RegisterWorkHandle();
 
@@ -516,30 +527,36 @@ private:
 		union
 		{
 			/** Default state of 0 means it can be reused immediately */
-			uint8 StateByte;
+			uint32 StateValue;
 
 			struct
 			{
 				/** If true, this work corresponds to an active handle */
-				uint8 bHasActiveHandle : 1;
+				uint32 bHasActiveHandle : 1;
 
 				/** If true, this work has been reserved for execution later in the frame */
-				uint8 bWorkReserved : 1;
+				uint32 bWorkReserved : 1;
 
 				/** If true, this work can be executed this frame */
-				uint8 bWorkRequested : 1;
+				uint32 bWorkRequested : 1;
 
 				/** If true, work will be reserved at the start of every frame */
-				uint8 bReserveEveryFrame : 1;
+				uint32 bReserveEveryFrame : 1;
 
 				/** If true, work will be requested at the start of every frame */
-				uint8 bRequestEveryFrame : 1;
+				uint32 bRequestEveryFrame : 1;
 
 				/** If true, this work is in the middle of being executed and is in the ExecutingWork array */
-				uint8 bCurrentlyExecuting : 1;
+				uint32 bCurrentlyExecuting : 1;
 
 				/** If true, this work has completely finished executing this frame */
-				uint8 bAlreadyExecuted : 1;
+				uint32 bAlreadyExecuted : 1;
+
+				/** If true, a late work request was created for this handle */
+				uint32 bLateWorkRequested : 1;
+
+				/** If true, the tick function allocated by this work and should be freed when reset */
+				uint32 bAllocatedTickFunction : 1;
 			};
 		};
 
@@ -547,23 +564,69 @@ private:
 		FTickFunction* TickFunction;
 
 		FActiveSyncWork()
-			: StateByte(0)
+			: StateValue(0)
 			, TickFunction(nullptr)
 		{
 		}
 
-		void Reset()
+		~FActiveSyncWork()
 		{
-			TickFunction = nullptr;
-			StateByte = 0;
+			if (bAllocatedTickFunction && TickFunction)
+			{
+				delete TickFunction;
+			}
 		}
 
-		bool IsInitialized() const
+		inline void SetTickFunction(FTickFunction* InTickFunction, bool bWasAllocatedInternally)
 		{
-			return StateByte != 0;
+			if (bAllocatedTickFunction && TickFunction)
+			{
+				delete TickFunction;
+			}
+			TickFunction = InTickFunction;
+			bAllocatedTickFunction = bWasAllocatedInternally;
 		}
+
+		inline void Reset()
+		{
+			SetTickFunction(nullptr, false);
+			StateValue = 0;
+		}
+
+		inline bool IsInitialized() const
+		{
+			return StateValue != 0;
+		}
+
+		FActiveSyncWork(const FActiveSyncWork&) = delete;
+		FActiveSyncWork& operator=(const FActiveSyncWork&) = delete;
 	};
 
+	/** Struct used to actually execute callbacks */
+	struct FExecutingSyncWork
+	{
+		/** Source index in the ActiveWork array */
+		int32 ActiveWorkIndex;
+
+		/** Function to run, if this is null it was disabled */
+		FTickFunction* TickFunction;
+
+		FExecutingSyncWork(int32 InIndex, const FActiveSyncWork& InWork)
+			: ActiveWorkIndex(InIndex),
+			TickFunction(InWork.TickFunction)
+		{}
+
+		inline bool IsValid() const
+		{
+			return TickFunction != nullptr;
+		}
+
+		/** Call from outside execution to invalidate even if already in queue */
+		inline void Invalidate()
+		{
+			TickFunction = nullptr;
+		}
+	};
 
 	// FTickFunction interface
 	virtual void ExecuteTick(float DeltaTime, ELevelTick TickType, ENamedThreads::Type CurrentThread, const FGraphEventRef& MyCompletionGraphEvent) override;
@@ -571,6 +634,9 @@ private:
 	virtual FName DiagnosticContext(bool bDetailed) override;
 	virtual bool DispatchManually() override;
 
+	/** Execute work from the game thread instead of from task graph */
+	void ExecuteFromGameThread(float DeltaTime, ELevelTick TickType);
+
 	/** Resets tick function at start of frame, this can automatically reserve/request work */
 	void ResetWorkForFrame();
 
@@ -581,7 +647,10 @@ private:
 	bool IsReadyToProcessWork();
 
 	/** Handle copying of work into / out of ExecutingWork array, returns true if there is anything to do */
-	bool GetWorkToExecute(int32& NumReserved);
+	bool GetWorkToExecute(float DeltaTime, ELevelTick TickType, ENamedThreads::Type CurrentThread, const FGraphEventRef& MyCompletionGraphEvent);
+
+	/** Cancels work that may already be in the execution queue */
+	void CancelRequestedWork(int32 ActiveWorkIndex, FActiveSyncWork& CurrentWork);
 
 	/** Gets index of new work slot, could resize array so must be called locked and not during iteration */
 	uint32 AllocateActiveWork();
@@ -592,9 +661,12 @@ private:
 	/** Registered sync point that created this */
 	FSyncPointDescription SyncPointDescription;
 	
-	/** Batch this is inside */
+	/** Sync point this corresponds to, including batch */
 	FSyncPointId SyncPointId;
 
+	/** World Context this was created in */
+	FWorldContextId WorldContextId;
+
 	/** Current state of this sync point in the current frame */
 	ESyncPointStatus FrameStatus;
 
@@ -604,11 +676,11 @@ private:
 	/** Lock to handle access to work array and other internal structs */
 	mutable FTransactionallySafeCriticalSection WorkLock;
 
-	/** Array of work to execute as part of tick */
+	/** Array of work to execute as part of tick, this can only be read or written with WorkLock */
 	TArray<FActiveSyncWork> ActiveWork;
 
-	/** Second buffer of work used during active execution */
-	TArray<FActiveSyncWork> ExecutingWork;
+	/** Second buffer of work used during active execution, this can only reallocate with WorkLock but is read outside it */
+	TArray<FExecutingSyncWork> ExecutingWork;
 
 	friend class FTaskSyncManager;
 	friend struct FActiveSyncWorkHandle;
@@ -660,11 +732,17 @@ public:
 	ENGINE_API FTaskSyncResult TriggerSyncPointAfterEvent(FSyncPointId SyncPoint, FGraphEventRef EventToWairFor);
 
 	/** 
-	 * Tries to creates a new sync work handle that can be used to reserve or request work to happen as part of a sync point
+	 * Tries to creates a new sync work handle that can be used to reserve or request work to happen as part of a sync point.
 	 * If it was successful, OutWorkHandle will be set to a handle you can pass between threads and later fulfill or abandon.
 	 */
 	ENGINE_API FTaskSyncResult RegisterWorkHandle(FSyncPointId SyncPoint, FActiveSyncWorkHandle& OutWorkHandle);
 
+	/**
+	 * Tries to creates a new sync work handle that can be used to reserve or request work to happen on the game thread during a tick group.
+	 * If it was successful, OutWorkHandle will be set to a handle you can use to request work for this frame or every frame (but not reserve).
+	 */
+	ENGINE_API FTaskSyncResult RegisterTickGroupWorkHandle(FWorldContextId WorldContext, ETickingGroup TickGroup, FActiveSyncWorkHandle& OutWorkHandle);
+
 
 	// Game thread-only functions, in general these will not affect the current frame
 
@@ -690,7 +768,7 @@ public:
 	ENGINE_API FTickFunction* GetTickFunctionForSyncPoint(FSyncPointId SyncPoint);
 
 	/** Tells the manager that it is the start of a frame, which will register the appropriate events */
-	ENGINE_API void StartFrame(const UWorld* InWorld);
+	ENGINE_API void StartFrame(const UWorld* InWorld, float InDeltaSeconds, ELevelTick InTickType);
 
 	/** Tells the manager that a tick group is starting for a specific world */
 	ENGINE_API void StartTickGroup(const UWorld* InWorld, ETickingGroup TickGroup, TArray<FTickFunction*>& TicksToManualDispatch);
@@ -705,15 +783,40 @@ public:
 	ENGINE_API void ReleaseWorldContext(FWorldContextId WorldContext);
 private:
 		
+	struct FTemporaryWorkRequest
+	{
+		FTemporaryWorkRequest(FActiveSyncPoint* InRequestingSyncPoint, int32 InRequestingHandle)
+			: RequestingSyncPoint(InRequestingSyncPoint)
+			, RequestingHandle(InRequestingHandle)
+		{}
+
+		FActiveSyncPoint* RequestingSyncPoint;
+		int32 RequestingHandle;
+		FActiveSyncWorkHandle WorkHandle;
+	};
+
 	struct FBatchData
 	{
+		/** The batch of objects with index in batch array */
 		FBatchContextId BatchContext;
+
+		/** The world this is associated with */
 		FWorldContextId WorldContext;
+
+		/** Array of specific user-defined sync points */
 		TMap<FSyncPointId::FInternalId, TSharedRef<FActiveSyncPoint>> SyncPointData;
 
+		/** Array of general gamethread work per tickgroup, these are not scheduled as real tasks and only exist in default batch */
+		TArray<TSharedPtr<FActiveSyncPoint>> TickGroupWork;
+
+		/** Array of temporary work handles, dropped at end of frame */
+		TArray<FTemporaryWorkRequest> TemporaryWorkRequests;
+
 		inline void Reset()
 		{
 			SyncPointData.Reset();
+			TickGroupWork.Reset();
+			TemporaryWorkRequests.Reset();
 			BatchContext = FBatchContextId();
 			WorldContext = FWorldContextId();
 		}
@@ -737,6 +840,8 @@ private:
 
 	/** World that is currently ticking, cannot modify certain operations during a tick or tick two worlds at once */
 	const UWorld* CurrentTickWorld = nullptr;
+	float CurrentDeltaTime = 0.0f;
+	ELevelTick CurrentTickType = LEVELTICK_All;
 	ETickingGroup CurrentTickGroup = ETickingGroup::TG_MAX;
 
 	/** True if this is currently ticking a world */
@@ -745,10 +850,19 @@ private:
 	/** Internal accessor to get a tick function's raw data, does not lock */
 	FTaskSyncResult FindActiveSyncPoint(FSyncPointId SyncPoint, TSharedPtr<FActiveSyncPoint>& OutData);
 
+	/** Try to handle forwarded work from a SyncPoint, will add to TickGroupWork if it makes sense */
+	bool HandleLateWorkRequest(FActiveSyncPoint* RequestedSyncPoint, int32 RequestingHandle, FTickFunction* TickFunction);
+
+	/** Cancels any scheduled temporary work requests matching the point and handle */
+	void CancelTemporaryWorkRequest(FActiveSyncPoint* RequestedSyncPoint, int32 RequestingHandle);
+
 	void RegisterSyncPointInternal(const FSyncPointDescription& InDescription);
 	void OnWorldContextRemove(FWorldContext& InWorldContext);
 	void InitializeBatchForFrame(FBatchData& BatchData, ULevel* PersistentLevel);
 	FActiveSyncPoint* GetOrCreateSyncPoint(FBatchData& BatchData, FRegisteredSyncPointData& SyncData);
+	FActiveSyncPoint* GetOrCreateTickGroupWork(FBatchData& BatchData, ETickingGroup TickGroup);
+
+	friend FActiveSyncPoint;
 };
 
 } // namespace UE::Task
