Substrate - fix for material compilation with hidden conversion.
Not compiling uneeded legacy input as those are never used by Substrate and sometimes interact in a bad way with substrate code chunks.

#tests PC editor dx12 + mobile vulkan generic preview
#rnx

diff --git a/Engine/Source/Runtime/Core/Private/UObject/FortniteMainBranchObjectVersion.cpp b/Engine/Source/Runtime/Core/Private/UObject/FortniteMainBranchObjectVersion.cpp
--- a/Engine/Source/Runtime/Core/Private/UObject/FortniteMainBranchObjectVersion.cpp
+++ b/Engine/Source/Runtime/Core/Private/UObject/FortniteMainBranchObjectVersion.cpp
@@ -14,6 +14,6 @@ TMap<FGuid, FGuid> FFortniteMainBranchObjectVersion::GetSystemGuids()
 	SystemGuids.Add(DevGuids.Niagara_LatestScriptCompileVersion, FGuid("670A2EBDEE1A4794BB77BBFF331EA7EF"));
 	SystemGuids.Add(DevGuids.SkeletalMeshDerivedDataVersion, FGuid("5E93192FCF3647EEB7F9D4FFF1750050"));
 	SystemGuids.Add(DevGuids.STATICMESH_DERIVEDDATA_VER, FGuid("3ABF32149F9D4D838750368AB95573BE"));
-	SystemGuids.Add(DevGuids.MaterialTranslationDDCVersion, FGuid("8ED4EEC4900A41E1BD70CB45CCEDB575"));
+	SystemGuids.Add(DevGuids.MaterialTranslationDDCVersion, FGuid("0D223CD149C54D5B8B09436D1B8AC0BF"));
 	return SystemGuids;
 }
diff --git a/Engine/Source/Runtime/Engine/Private/Materials/HLSLMaterialTranslator.cpp b/Engine/Source/Runtime/Engine/Private/Materials/HLSLMaterialTranslator.cpp
--- a/Engine/Source/Runtime/Engine/Private/Materials/HLSLMaterialTranslator.cpp
+++ b/Engine/Source/Runtime/Engine/Private/Materials/HLSLMaterialTranslator.cpp
@@ -1832,22 +1832,26 @@ void FHLSLMaterialTranslator::TranslateMaterial()
 		NormalCodeChunkEnd = SharedPropertyCodeChunks[NormalShaderFrequency].Num();
 	}
 
-	// Rest of properties
-	Chunk[MP_EmissiveColor] = Material->CompilePropertyAndSetMaterialProperty(MP_EmissiveColor			,this);
-	Chunk[MP_DiffuseColor] = Material->CompilePropertyAndSetMaterialProperty(MP_DiffuseColor			,this);
-	Chunk[MP_SpecularColor] = Material->CompilePropertyAndSetMaterialProperty(MP_SpecularColor			,this);
-	Chunk[MP_BaseColor] = Material->CompilePropertyAndSetMaterialProperty(MP_BaseColor				,this);
-	Chunk[MP_Metallic] = Material->CompilePropertyAndSetMaterialProperty(MP_Metallic				,this);
-	Chunk[MP_Specular] = Material->CompilePropertyAndSetMaterialProperty(MP_Specular				,this);
-	Chunk[MP_Roughness] = Material->CompilePropertyAndSetMaterialProperty(MP_Roughness				,this);
-	Chunk[MP_Anisotropy] = Material->CompilePropertyAndSetMaterialProperty(MP_Anisotropy				,this);
-	Chunk[MP_Opacity] = Material->CompilePropertyAndSetMaterialProperty(MP_Opacity				,this);
-	Chunk[MP_OpacityMask] = Material->CompilePropertyAndSetMaterialProperty(MP_OpacityMask			,this);
-	Chunk[MP_Tangent] = Material->CompilePropertyAndSetMaterialProperty(MP_Tangent				,this);
-	Chunk[MP_WorldPositionOffset] = Material->CompilePropertyAndSetMaterialProperty(MP_WorldPositionOffset	,this);
-
-	// Make sure to compile this property before using ShadingModelsFromCompilation
-	Chunk[MP_ShadingModel] = Material->CompilePropertyAndSetMaterialProperty(MP_ShadingModel			,this);
+	if (!bSubstrateEnabled || !bUseRootNodeToSubstrateHiddenConversion)
+	{
+		// Rest of properties
+		Chunk[MP_EmissiveColor]	= Material->CompilePropertyAndSetMaterialProperty(MP_EmissiveColor,	this);
+		Chunk[MP_DiffuseColor]	= Material->CompilePropertyAndSetMaterialProperty(MP_DiffuseColor,	this);
+		Chunk[MP_SpecularColor]	= Material->CompilePropertyAndSetMaterialProperty(MP_SpecularColor,	this);
+		Chunk[MP_BaseColor]		= Material->CompilePropertyAndSetMaterialProperty(MP_BaseColor,		this);
+		Chunk[MP_Metallic]		= Material->CompilePropertyAndSetMaterialProperty(MP_Metallic,		this);
+		Chunk[MP_Specular]		= Material->CompilePropertyAndSetMaterialProperty(MP_Specular,		this);
+		Chunk[MP_Roughness]		= Material->CompilePropertyAndSetMaterialProperty(MP_Roughness,		this);
+		Chunk[MP_Anisotropy]	= Material->CompilePropertyAndSetMaterialProperty(MP_Anisotropy,	this);
+		Chunk[MP_Tangent]		= Material->CompilePropertyAndSetMaterialProperty(MP_Tangent,		this);
+
+		// Make sure to compile this property before using ShadingModelsFromCompilation
+		Chunk[MP_ShadingModel]	= Material->CompilePropertyAndSetMaterialProperty(MP_ShadingModel,	this);
+	}
+
+	Chunk[MP_Opacity]			= Material->CompilePropertyAndSetMaterialProperty(MP_Opacity,		this);
+	Chunk[MP_OpacityMask]		= Material->CompilePropertyAndSetMaterialProperty(MP_OpacityMask,	this);
+	Chunk[MP_WorldPositionOffset]= Material->CompilePropertyAndSetMaterialProperty(MP_WorldPositionOffset, this);
 
 
 #if DEBUG_SUBSTRATE_TREE_STACK
@@ -1873,6 +1877,76 @@ void FHLSLMaterialTranslator::TranslateMaterial()
 			FullySimplifiedFrontMaterialCodeChunkEnd = SharedPropertyCodeChunks[FrontMaterialShaderFrequency].Num();
 		}
 	}
+	else // if (bUseRootNodeToSubstrateHiddenConversion)
+	{
+		auto CompileFrontMaterial = [&]()
+			{
+				UMaterialInterface* MaterialInterface = Material->GetMaterialInterface();
+				UMaterial* BaseMaterial = MaterialInterface ? MaterialInterface->GetMaterial() : nullptr;
+				UMaterialEditorOnlyData* EditorOnlyData = BaseMaterial ? BaseMaterial->GetEditorOnlyData() : nullptr;
+
+				if (BaseMaterial && BaseMaterial->bUseMaterialAttributes)
+				{
+					const bool bHasSSS = MaterialInterface && MaterialInterface->SubsurfaceProfile;
+
+					return UMaterialExpressionSubstrateConvertMaterialAttributes::CompileCommon(this, 0,
+						FMaterialAttributeDefinitionMap::GetConnectedMaterialAttributesBitmask(BaseMaterial->GetExpressions()), EditorOnlyData->MaterialAttributes, MaterialInterface->GetShadingModels().GetFirstShadingModel(),
+						*WaterScatteringCoefficients, *WaterAbsorptionCoefficients, *WaterPhaseG, *ColorScaleBehindWater,
+						bHasSSS, bHasSSS ? MaterialInterface->SubsurfaceProfile : nullptr);
+
+				}
+				else if ((Material->GetMaterialDomain() == MD_Surface || Material->GetMaterialDomain() == MD_RuntimeVirtualTexture || Material->GetMaterialDomain() == MD_DeferredDecal) && EditorOnlyData)
+				{
+					const bool bHasSSS = MaterialInterface && MaterialInterface->SubsurfaceProfile;
+					const bool bHasAnisotropy = EditorOnlyData->Anisotropy.IsConnected();
+
+					return UMaterialExpressionSubstrateShadingModels::CompileCommon(this,
+						EditorOnlyData->BaseColor, EditorOnlyData->Specular, EditorOnlyData->Metallic, EditorOnlyData->Roughness, EditorOnlyData->EmissiveColor,
+						EditorOnlyData->Opacity, EditorOnlyData->SubsurfaceColor, EditorOnlyData->ClearCoat, EditorOnlyData->ClearCoatRoughness,
+						EditorOnlyData->ShadingModelFromMaterialExpression, MaterialInterface->GetShadingModels().GetFirstShadingModel(),
+						*ThinTranslucentTransmittanceColor, *ThinTranslucentSurfaceCoverage,
+						*WaterScatteringCoefficients, *WaterAbsorptionCoefficients, *WaterPhaseG, *ColorScaleBehindWater,
+						bHasAnisotropy, EditorOnlyData->Anisotropy,
+						EditorOnlyData->Normal, EditorOnlyData->Tangent,
+						*ClearCoatNormal, *CustomTangent,
+						bHasSSS, MaterialInterface->SubsurfaceProfile,
+						EditorOnlyData);
+				}
+				else if (Material->GetMaterialDomain() == MD_Volume && EditorOnlyData)
+				{
+					return UMaterialExpressionSubstrateVolumetricFogCloudBSDF::CompileCommon(
+						this,
+						EditorOnlyData->BaseColor,
+						EditorOnlyData->SubsurfaceColor,
+						EditorOnlyData->EmissiveColor,
+						EditorOnlyData->AmbientOcclusion,
+						MaterialInterface->GetShadingModels().HasOnlyShadingModel(MSM_Unlit),
+						EditorOnlyData);
+				}
+				else if (Material->GetMaterialDomain() == MD_LightFunction && EditorOnlyData)
+				{
+					return UMaterialExpressionSubstrateLightFunction::CompileCommon(this, EditorOnlyData->EmissiveColor, EditorOnlyData);
+				}
+				else if (Material->GetMaterialDomain() == MD_PostProcess && EditorOnlyData)
+				{
+					return UMaterialExpressionSubstratePostProcess::CompileCommon(this, EditorOnlyData->EmissiveColor, EditorOnlyData->Opacity, EditorOnlyData);
+				}
+				else if (Material->GetMaterialDomain() == MD_UI && EditorOnlyData)
+				{
+					return UMaterialExpressionSubstrateUI::CompileCommon(this, EditorOnlyData->EmissiveColor, EditorOnlyData->Opacity, EditorOnlyData);
+				}
+
+				return SubstrateCreateAndRegisterNullMaterial();
+			};
+
+		Chunk[MP_FrontMaterial] = CompileFrontMaterial();
+
+		FullySimplifiedFrontMaterialCodeChunkStart = SharedPropertyCodeChunks[FrontMaterialShaderFrequency].Num();
+		CurrentSubstrateCompilationContext = ESubstrateCompilationContext::SCC_FullySimplified;
+		FullySimplifiedSubstrateFrontMaterialCodeChunk = CompileFrontMaterial();
+		CurrentSubstrateCompilationContext = ESubstrateCompilationContext::SCC_Default;
+		FullySimplifiedFrontMaterialCodeChunkEnd = SharedPropertyCodeChunks[FrontMaterialShaderFrequency].Num();
+	}
 
 	CHECK_DDC_QUERY_FINISHED_ELSE_RETURN();
 
@@ -1901,8 +1975,11 @@ void FHLSLMaterialTranslator::TranslateMaterial()
 		CHECK_DDC_QUERY_FINISHED_ELSE_RETURN();
 	}
 
-	Chunk[MP_CustomData0] = Material->CompilePropertyAndSetMaterialProperty(MP_CustomData0, this);
-	Chunk[MP_CustomData1] = Material->CompilePropertyAndSetMaterialProperty(MP_CustomData1, this);
+	if (!bSubstrateEnabled || !bUseRootNodeToSubstrateHiddenConversion)
+	{
+		Chunk[MP_CustomData0] = Material->CompilePropertyAndSetMaterialProperty(MP_CustomData0, this);
+		Chunk[MP_CustomData1] = Material->CompilePropertyAndSetMaterialProperty(MP_CustomData1, this);
+	}
 	Chunk[MP_AmbientOcclusion] = Material->CompilePropertyAndSetMaterialProperty(MP_AmbientOcclusion, this);
 
 	if (IsTranslucentBlendMode(BlendMode) || MaterialShadingModels.HasShadingModel(MSM_SingleLayerWater))
@@ -1931,78 +2008,6 @@ void FHLSLMaterialTranslator::TranslateMaterial()
 
 	Chunk[MP_Displacement] = Material->CompilePropertyAndSetMaterialProperty(MP_Displacement, this);
 
-
-	if (bUseRootNodeToSubstrateHiddenConversion)
-	{
-		auto CompileFrontMaterial = [&]()
-		{
-			UMaterialInterface* MaterialInterface = Material->GetMaterialInterface();
-			UMaterial* BaseMaterial = MaterialInterface ? MaterialInterface->GetMaterial() : nullptr;
-			UMaterialEditorOnlyData* EditorOnlyData = BaseMaterial ? BaseMaterial->GetEditorOnlyData() : nullptr;
-
-			if (BaseMaterial && BaseMaterial->bUseMaterialAttributes)
-			{
-				const bool bHasSSS = MaterialInterface && MaterialInterface->SubsurfaceProfile;
-
-				return UMaterialExpressionSubstrateConvertMaterialAttributes::CompileCommon(this, 0,
-					FMaterialAttributeDefinitionMap::GetConnectedMaterialAttributesBitmask(BaseMaterial->GetExpressions()), EditorOnlyData->MaterialAttributes, MaterialInterface->GetShadingModels().GetFirstShadingModel(),
-					*WaterScatteringCoefficients, *WaterAbsorptionCoefficients, *WaterPhaseG, *ColorScaleBehindWater,
-					bHasSSS, bHasSSS ? MaterialInterface->SubsurfaceProfile : nullptr);
-
-			}
-			else if ((Material->GetMaterialDomain() == MD_Surface || Material->GetMaterialDomain() == MD_RuntimeVirtualTexture || Material->GetMaterialDomain() == MD_DeferredDecal) && EditorOnlyData)
-			{
-				const bool bHasSSS = MaterialInterface && MaterialInterface->SubsurfaceProfile;
-				const bool bHasAnisotropy = EditorOnlyData->Anisotropy.IsConnected();
-
-				return UMaterialExpressionSubstrateShadingModels::CompileCommon(this,
-					EditorOnlyData->BaseColor, EditorOnlyData->Specular, EditorOnlyData->Metallic, EditorOnlyData->Roughness, EditorOnlyData->EmissiveColor,
-					EditorOnlyData->Opacity, EditorOnlyData->SubsurfaceColor, EditorOnlyData->ClearCoat, EditorOnlyData->ClearCoatRoughness,
-					EditorOnlyData->ShadingModelFromMaterialExpression, MaterialInterface->GetShadingModels().GetFirstShadingModel(),
-					*ThinTranslucentTransmittanceColor, *ThinTranslucentSurfaceCoverage,
-					*WaterScatteringCoefficients, *WaterAbsorptionCoefficients, *WaterPhaseG, *ColorScaleBehindWater,
-					bHasAnisotropy, EditorOnlyData->Anisotropy,
-					EditorOnlyData->Normal, EditorOnlyData->Tangent,
-					*ClearCoatNormal, *CustomTangent,
-					bHasSSS, MaterialInterface->SubsurfaceProfile,
-					EditorOnlyData);
-			}
-			else if (Material->GetMaterialDomain() == MD_Volume && EditorOnlyData)
-			{
-				return UMaterialExpressionSubstrateVolumetricFogCloudBSDF::CompileCommon(
-					this,
-					EditorOnlyData->BaseColor,
-					EditorOnlyData->SubsurfaceColor,
-					EditorOnlyData->EmissiveColor,
-					EditorOnlyData->AmbientOcclusion,
-					MaterialInterface->GetShadingModels().HasOnlyShadingModel(MSM_Unlit),
-					EditorOnlyData);
-			}
-			else if (Material->GetMaterialDomain() == MD_LightFunction && EditorOnlyData)
-			{
-				return UMaterialExpressionSubstrateLightFunction::CompileCommon(this, EditorOnlyData->EmissiveColor, EditorOnlyData);
-			}
-			else if (Material->GetMaterialDomain() == MD_PostProcess && EditorOnlyData)
-			{
-				return UMaterialExpressionSubstratePostProcess::CompileCommon(this, EditorOnlyData->EmissiveColor, EditorOnlyData->Opacity, EditorOnlyData);
-			}
-			else if (Material->GetMaterialDomain() == MD_UI && EditorOnlyData)
-			{
-				return UMaterialExpressionSubstrateUI::CompileCommon(this, EditorOnlyData->EmissiveColor, EditorOnlyData->Opacity, EditorOnlyData);
-			}
-
-			return SubstrateCreateAndRegisterNullMaterial();
-		};
-		
-		Chunk[MP_FrontMaterial] = CompileFrontMaterial();
-
-		FullySimplifiedFrontMaterialCodeChunkStart = SharedPropertyCodeChunks[FrontMaterialShaderFrequency].Num();
-		CurrentSubstrateCompilationContext = ESubstrateCompilationContext::SCC_FullySimplified;
-		FullySimplifiedSubstrateFrontMaterialCodeChunk = CompileFrontMaterial();
-		CurrentSubstrateCompilationContext = ESubstrateCompilationContext::SCC_Default;
-		FullySimplifiedFrontMaterialCodeChunkEnd = SharedPropertyCodeChunks[FrontMaterialShaderFrequency].Num();
-	}
-
 	CHECK_DDC_QUERY_FINISHED_ELSE_RETURN();
 	ResourcesString = TEXT("");
 
