Substrate - Fix for nullptr CurrentScopeChunks in some case when no other properties have been compiled before.

#rnx
#tests PC editor dx12

diff --git a/Engine/Source/Runtime/Engine/Private/Materials/HLSLMaterialTranslator.cpp b/Engine/Source/Runtime/Engine/Private/Materials/HLSLMaterialTranslator.cpp
--- a/Engine/Source/Runtime/Engine/Private/Materials/HLSLMaterialTranslator.cpp
+++ b/Engine/Source/Runtime/Engine/Private/Materials/HLSLMaterialTranslator.cpp
@@ -1939,6 +1939,9 @@ void FHLSLMaterialTranslator::TranslateMaterial()
 				return SubstrateCreateAndRegisterNullMaterial();
 			};
 
+		// Need to assign a valid scope for Front material to be correct (and avoid nullptr CurrentScopeChunks in case no other properties have been compiled before).
+		SetMaterialProperty(MP_FrontMaterial, FrontMaterialShaderFrequency, false/*bUsePreviousFrameTime*/);
+
 		Chunk[MP_FrontMaterial] = CompileFrontMaterial();
 
 		FullySimplifiedFrontMaterialCodeChunkStart = SharedPropertyCodeChunks[FrontMaterialShaderFrequency].Num();
