PR #12795: Correct grammer of comment above Destroy() in Actor.h

diff --git a/Engine/Source/Runtime/Engine/Classes/GameFramework/Actor.h b/Engine/Source/Runtime/Engine/Classes/GameFramework/Actor.h
--- a/Engine/Source/Runtime/Engine/Classes/GameFramework/Actor.h
+++ b/Engine/Source/Runtime/Engine/Classes/GameFramework/Actor.h
@@ -2305,7 +2305,7 @@ public:
 	ENGINE_API void UserConstructionScript();
 
 	/**
-	 * Destroy this actor. Returns true the actor is destroyed or already marked for destruction, false if indestructible.
+	 * Destroy this actor. Returns true if the actor is destroyed or already marked for destruction, false if indestructible.
 	 * Destruction is latent. It occurs at the end of the tick.
 	 * @param	bNetForce				[opt] Ignored unless called during play.  Default is false.
 	 * @param	bShouldModifyLevel		[opt] If true, Modify() the level before removing the actor.  Default is true.	
