World Partition - HLOD: Added support for Instanced Skinned Mesh Component batching

diff --git a/Engine/Source/Runtime/CoreUObject/Public/Serialization/ArchiveCrc32.h b/Engine/Source/Runtime/CoreUObject/Public/Serialization/ArchiveCrc32.h
--- a/Engine/Source/Runtime/CoreUObject/Public/Serialization/ArchiveCrc32.h
+++ b/Engine/Source/Runtime/CoreUObject/Public/Serialization/ArchiveCrc32.h
@@ -52,6 +52,8 @@ public:
 		return Ar;
 	}
 
+	using FArchiveUObject::operator<<; // For visibility of the overloads we don't override
+
 private:
 	uint32 CRC;
 	UObject* RootObject;
diff --git a/Engine/Source/Runtime/Engine/Classes/Animation/AnimBank.h b/Engine/Source/Runtime/Engine/Classes/Animation/AnimBank.h
--- a/Engine/Source/Runtime/Engine/Classes/Animation/AnimBank.h
+++ b/Engine/Source/Runtime/Engine/Classes/Animation/AnimBank.h
@@ -4,6 +4,7 @@
 
 #include "AnimationAsset.h"
 #include "Experimental/Containers/RobinHoodHashTable.h"
+#include "HLOD/HLODBatchingPolicy.h"
 #include "Interfaces/Interface_AsyncCompilation.h"
 #include "IO/IoHash.h"
 #include "Logging/LogMacros.h"
@@ -347,7 +348,55 @@ inline uint32 GetTypeHash(const TArray<FSoftAnimBankItem>& InBankItems)
 }
 
 USTRUCT()
-struct FSkinnedMeshComponentDescriptor
+struct FSkinnedMeshComponentDescriptorBase
+{
+	GENERATED_BODY()
+	ENGINE_API FSkinnedMeshComponentDescriptorBase();
+	ENGINE_API explicit FSkinnedMeshComponentDescriptorBase(ENoInit);
+	ENGINE_API explicit FSkinnedMeshComponentDescriptorBase(const FSkinnedMeshComponentDescriptorBase&);
+	ENGINE_API virtual ~FSkinnedMeshComponentDescriptorBase();
+
+	ENGINE_API UInstancedSkinnedMeshComponent* CreateComponent(UObject* Outer, FName Name = NAME_None, EObjectFlags ObjectFlags = EObjectFlags::RF_NoFlags) const;
+
+	ENGINE_API virtual void InitFrom(const UInstancedSkinnedMeshComponent* Component, bool bInitBodyInstance = true);
+	ENGINE_API virtual void InitComponent(UInstancedSkinnedMeshComponent* ISMComponent) const;
+
+	ENGINE_API bool operator!=(const FSkinnedMeshComponentDescriptorBase& Other) const;
+	ENGINE_API bool operator==(const FSkinnedMeshComponentDescriptorBase& Other) const;
+
+public:
+	UPROPERTY()
+	mutable uint32 Hash = 0;
+
+	UPROPERTY(EditAnywhere, Category = "Component Settings")
+	TEnumAsByte<EComponentMobility::Type> Mobility;
+
+	UPROPERTY(EditAnywhere, Category = "Component Settings")
+	TSubclassOf<UInstancedSkinnedMeshComponent> ComponentClass;
+
+	UPROPERTY(EditAnywhere, Category = "Component Settings")
+	int32 InstanceMinDrawDistance;
+
+	UPROPERTY(EditAnywhere, Category = "Component Settings")
+	int32 InstanceStartCullDistance;
+
+	UPROPERTY(EditAnywhere, Category = "Component Settings")
+	int32 InstanceEndCullDistance;
+
+	UPROPERTY(EditAnywhere, Category = "Component Settings")
+	uint8 bCastShadow : 1;
+
+#if WITH_EDITORONLY_DATA
+	UPROPERTY(EditAnywhere, Category = "Component Settings")
+	uint8 bIncludeInHLOD : 1;
+
+	UPROPERTY(EditAnywhere, Category = "Component Settings")
+	EHLODBatchingPolicy HLODBatchingPolicy;
+#endif
+};
+
+USTRUCT()
+struct FSkinnedMeshComponentDescriptor : public FSkinnedMeshComponentDescriptorBase
 {
 	GENERATED_BODY()
 
@@ -383,15 +432,6 @@ struct FSkinnedMeshComponentDescriptor
 	}
 
 public:
-	UPROPERTY()
-	mutable uint32 Hash = 0;
-
-	UPROPERTY(EditAnywhere, Category = "Component Settings")
-	TEnumAsByte<EComponentMobility::Type> Mobility;
-
-	UPROPERTY(EditAnywhere, Category = "Component Settings")
-	TSubclassOf<UInstancedSkinnedMeshComponent> ComponentClass;
-
 	UPROPERTY(EditAnywhere, Category = "Component Settings", meta = (DisplayAfter = "ComponentClass"))
 	TObjectPtr<class USkinnedAsset> SkinnedAsset;
 
@@ -400,7 +440,7 @@ public:
 };
 
 USTRUCT()
-struct FSoftSkinnedMeshComponentDescriptor
+struct FSoftSkinnedMeshComponentDescriptor : public FSkinnedMeshComponentDescriptorBase
 {
 	GENERATED_BODY()
 
@@ -436,15 +476,6 @@ struct FSoftSkinnedMeshComponentDescriptor
 	}
 
 public:
-	UPROPERTY()
-	mutable uint32 Hash = 0;
-
-	UPROPERTY(EditAnywhere, Category = "Component Settings")
-	TEnumAsByte<EComponentMobility::Type> Mobility;
-
-	UPROPERTY(EditAnywhere, Category = "Component Settings")
-	TSubclassOf<UInstancedSkinnedMeshComponent> ComponentClass;
-
 	UPROPERTY()
 	TSoftObjectPtr<class USkinnedAsset> SkinnedAsset;
 
diff --git a/Engine/Source/Runtime/Engine/Classes/Components/ActorComponent.h b/Engine/Source/Runtime/Engine/Classes/Components/ActorComponent.h
--- a/Engine/Source/Runtime/Engine/Classes/Components/ActorComponent.h
+++ b/Engine/Source/Runtime/Engine/Classes/Components/ActorComponent.h
@@ -1323,6 +1323,9 @@ public:
 #if WITH_EDITOR
 	/** Override to provide a custom set of components to be used in place of the current component for the HLOD generation. */
 	virtual TArray<UActorComponent*> GetHLODProxyComponents() const { return TArray<UActorComponent*>(); }
+
+	/** Override to allow a component subclass to compute an hash for the purpose of HLOD generation. */
+	virtual void ComputeHLODHash(class FHLODHashBuilder& HashBuilder) const	{}
 #endif
 
 	/** Suffix used to identify template component instances */
diff --git a/Engine/Source/Runtime/Engine/Classes/Components/InstancedSkinnedMeshComponent.h b/Engine/Source/Runtime/Engine/Classes/Components/InstancedSkinnedMeshComponent.h
--- a/Engine/Source/Runtime/Engine/Classes/Components/InstancedSkinnedMeshComponent.h
+++ b/Engine/Source/Runtime/Engine/Classes/Components/InstancedSkinnedMeshComponent.h
@@ -75,6 +75,18 @@ protected:
 	UPROPERTY(EditAnywhere, BlueprintReadOnly, Category = Animation)
 	float AnimationMinScreenSize = 0.0f;
 
+	/** Distance from camera at which each instance begins to draw. */
+	UPROPERTY(EditAnywhere, BlueprintReadOnly, Category = Culling)
+	int32 InstanceMinDrawDistance;
+
+	/** Distance from camera at which each instance begins to fade out. */
+	UPROPERTY(EditAnywhere, BlueprintReadOnly, Category = Culling)
+	int32 InstanceStartCullDistance;
+
+	/** Distance from camera at which each instance completely fades out. */
+	UPROPERTY(EditAnywhere, BlueprintReadOnly, Category = Culling)
+	int32 InstanceEndCullDistance;
+
 	/** If true, this component will avoid serializing its per instance data / those properties will also not be editable */
 	UPROPERTY()
 	uint8 bInheritPerInstanceData : 1;
@@ -103,7 +115,10 @@ public:
 	virtual void OnRegister() override;
 	virtual void SendRenderInstanceData_Concurrent() override;
 	virtual void TickComponent(float DeltaTime, enum ELevelTick TickType, FActorComponentTickFunction* ThisTickFunction) override;
-
+	virtual bool IsHLODRelevant() const override;
+#if WITH_EDITOR
+	virtual void ComputeHLODHash(class FHLODHashBuilder& HashBuilder) const override;
+#endif
 	virtual void CreateRenderState_Concurrent(FRegisterComponentContext* Context) override;
 	virtual void DestroyRenderState_Concurrent() override;
 
@@ -141,17 +156,23 @@ public:
 	UFUNCTION(BlueprintCallable, Category = "Components|InstancedSkinnedMesh")
 	ENGINE_API bool SetCustomDataValue(FPrimitiveInstanceId InstanceId, int32 CustomDataIndex, float CustomDataValue);
 
-	/** Update all custom data values for specific instance, the size of the array view must match the NumCustomDataFloats */
+	/** Update all custom data values for specific instance, the size of the array view must match the NumCustomDataFloats. Returns True on success. */
 	ENGINE_API bool SetCustomData(FPrimitiveInstanceId InstanceId, TArrayView<const float> CustomDataFloats);
 
+	/** Get all custom data values for specific instance, the size of the array view must match the NumCustomDataFloats. Returns True on success. */
+	ENGINE_API bool GetCustomData(FPrimitiveInstanceId InstanceId, TArrayView<float> CustomDataFloats) const;
+
 	/** Update number of custom data entries per instance. This applies to all instances and will reallocate the full custom data buffer and reset all values to 0 */
 	UFUNCTION(BlueprintCallable, Category = "Components|InstancedSkinnedMesh")
 	ENGINE_API void SetNumCustomDataFloats(int32 InNumCustomDataFloats);
 
-	/** Get the transform for the instance specified. Instance is returned in local space of this component unless bWorldSpace is set.  Returns True on success. */
+	/** Get the transform for the instance specified. Instance is returned in local space of this component unless bWorldSpace is set. Returns True on success. */
 	UFUNCTION(BlueprintCallable, Category = "Components|InstancedSkinnedMesh")
 	ENGINE_API bool GetInstanceTransform(FPrimitiveInstanceId InstanceId, FTransform& OutInstanceTransform, bool bWorldSpace = false) const;
 
+	/** Get the bank index for the instance specified. Returns True on success. */
+	ENGINE_API bool GetInstanceBankIndex(FPrimitiveInstanceId InstanceId, int32& OutBankIndex) const;
+
 	/** Remove the instance specified. Returns True on success. */
 	UFUNCTION(BlueprintCallable, Category = "Components|InstancedSkinnedMesh")
 	ENGINE_API bool RemoveInstance(FPrimitiveInstanceId InstanceId);
@@ -183,6 +204,10 @@ public:
 	ENGINE_API void SetAnimBankItems(TArrayView<const struct FAnimBankItem> InAnimBankItems);
 
 	int32 GetNumCustomDataFloats() const { return NumCustomDataFloats; }
+
+	int32 GetMinDrawDistance() const { return InstanceMinDrawDistance; }
+	void GetCullDistances(int32& OutStartCullDistance, int32& OutEndCullDistance) const { OutStartCullDistance = InstanceStartCullDistance; OutEndCullDistance = InstanceEndCullDistance; }
+	ENGINE_API void SetCullDistances(int32 StartCullDistance, int32 EndCullDistance);
 protected:
 
 	/** Handle changes that must happen before the proxy is recreated. */
diff --git a/Engine/Source/Runtime/Engine/Classes/Components/InstancedStaticMeshComponent.h b/Engine/Source/Runtime/Engine/Classes/Components/InstancedStaticMeshComponent.h
--- a/Engine/Source/Runtime/Engine/Classes/Components/InstancedStaticMeshComponent.h
+++ b/Engine/Source/Runtime/Engine/Classes/Components/InstancedStaticMeshComponent.h
@@ -469,6 +469,8 @@ public:
 	ENGINE_API virtual bool IsInstanceTouchingSelectionFrustum(int32 InstanceIndex, const FConvexVolume& InFrustum, const bool bMustEncompassEntireInstance) const;
 
 	ENGINE_API virtual bool CanEditChange(const FProperty* InProperty) const override;
+
+	ENGINE_API virtual void ComputeHLODHash(class FHLODHashBuilder& HashBuilder) const override;
 #endif
 
 	// Helper function to construct a base-set of instance data flags that in
diff --git a/Engine/Source/Runtime/Engine/Classes/Components/PrimitiveComponent.h b/Engine/Source/Runtime/Engine/Classes/Components/PrimitiveComponent.h
--- a/Engine/Source/Runtime/Engine/Classes/Components/PrimitiveComponent.h
+++ b/Engine/Source/Runtime/Engine/Classes/Components/PrimitiveComponent.h
@@ -730,6 +730,10 @@ public:
 	UFUNCTION(BlueprintCallable, Category = "HLOD", meta = (DisplayName = "Set Excluded From HLOD Level"))
 	ENGINE_API void SetExcludedFromHLODLevel(EHLODLevelExclusion HLODLevel, bool bExcluded);
 
+#if WITH_EDITOR
+	ENGINE_API virtual void ComputeHLODHash(class FHLODHashBuilder& HashBuilder) const override;
+#endif
+
 private:
 	UE_DEPRECATED("5.2", "Use SetExcludedFromHLODLevel instead")
 	UFUNCTION(BlueprintCallable, BlueprintSetter, Category = "HLOD", meta = (BlueprintInternalUseOnly="true"))
diff --git a/Engine/Source/Runtime/Engine/Classes/Components/StaticMeshComponent.h b/Engine/Source/Runtime/Engine/Classes/Components/StaticMeshComponent.h
--- a/Engine/Source/Runtime/Engine/Classes/Components/StaticMeshComponent.h
+++ b/Engine/Source/Runtime/Engine/Classes/Components/StaticMeshComponent.h
@@ -582,6 +582,9 @@ public:
 #endif
 	ENGINE_API virtual TStructOnScope<FActorComponentInstanceData> GetComponentInstanceData() const override;
 	ENGINE_API virtual bool IsHLODRelevant() const override;
+#if WITH_EDITOR
+	ENGINE_API virtual void ComputeHLODHash(class FHLODHashBuilder& HashBuilder) const override;
+#endif
 	//~ End UActorComponent Interface.
 
 	//~ Begin UPrimitiveComponent Interface.
diff --git a/Engine/Source/Runtime/Engine/Private/Animation/AnimBank.cpp b/Engine/Source/Runtime/Engine/Private/Animation/AnimBank.cpp
--- a/Engine/Source/Runtime/Engine/Private/Animation/AnimBank.cpp
+++ b/Engine/Source/Runtime/Engine/Private/Animation/AnimBank.cpp
@@ -1385,7 +1385,65 @@ bool FSoftAnimBankItem::operator==(const FSoftAnimBankItem& Other) const
 	return BankAsset == Other.BankAsset && SequenceIndex == Other.SequenceIndex;
 }
 
+FSkinnedMeshComponentDescriptorBase::FSkinnedMeshComponentDescriptorBase()
+{
+	// Note: should not really be used - prefer using FSkinnedMeshComponentDescriptor or FSoftSkinnedMeshComponentDescriptor
+	InitFrom(UInstancedSkinnedMeshComponent::StaticClass()->GetDefaultObject<UInstancedSkinnedMeshComponent>());
+}
+
+FSkinnedMeshComponentDescriptorBase::FSkinnedMeshComponentDescriptorBase(ENoInit) {}
+FSkinnedMeshComponentDescriptorBase::FSkinnedMeshComponentDescriptorBase(const FSkinnedMeshComponentDescriptorBase& InDescriptor) = default;
+FSkinnedMeshComponentDescriptorBase::~FSkinnedMeshComponentDescriptorBase() = default;
+
+void FSkinnedMeshComponentDescriptorBase::InitFrom(const UInstancedSkinnedMeshComponent* Template, bool bInitBodyInstance)
+{
+	check(Template);
+
+	Mobility = Template->Mobility;
+	InstanceMinDrawDistance = Template->GetMinDrawDistance();
+	Template->GetCullDistances(InstanceStartCullDistance, InstanceEndCullDistance);
+	ComponentClass = Template->GetClass();
+	bCastShadow = Template->CastShadow;
+#if WITH_EDITOR
+	HLODBatchingPolicy = Template->HLODBatchingPolicy;
+	bIncludeInHLOD = Template->bEnableAutoLODGeneration;
+#endif	
+}
+
+void FSkinnedMeshComponentDescriptorBase::InitComponent(UInstancedSkinnedMeshComponent* Component) const
+{
+	Component->Mobility = Mobility;
+	Component->MinDrawDistance = InstanceMinDrawDistance;
+	Component->SetCullDistances(InstanceStartCullDistance, InstanceEndCullDistance);
+	Component->CastShadow = bCastShadow;
+#if WITH_EDITOR
+	Component->HLODBatchingPolicy = HLODBatchingPolicy;
+	Component->bEnableAutoLODGeneration = bIncludeInHLOD;
+#endif
+}
+
+bool FSkinnedMeshComponentDescriptorBase::operator!=(const FSkinnedMeshComponentDescriptorBase& Other) const
+{
+	return !(*this == Other);
+}
+
+bool FSkinnedMeshComponentDescriptorBase::operator==(const FSkinnedMeshComponentDescriptorBase& Other) const
+{
+	return
+		ComponentClass == Other.ComponentClass &&
+		Mobility == Other.Mobility &&
+		InstanceMinDrawDistance == Other.InstanceMinDrawDistance &&
+		InstanceStartCullDistance == Other.InstanceStartCullDistance &&
+		InstanceEndCullDistance == Other.InstanceEndCullDistance &&
+#if WITH_EDITOR
+		HLODBatchingPolicy == Other.HLODBatchingPolicy &&
+		bIncludeInHLOD == Other.bIncludeInHLOD &&
+#endif 
+		bCastShadow == Other.bCastShadow;
+}
+
 FSkinnedMeshComponentDescriptor::FSkinnedMeshComponentDescriptor()
+	: FSkinnedMeshComponentDescriptorBase(NoInit)
 {
 	// Make sure we have proper defaults
 	InitFrom(UInstancedSkinnedMeshComponent::StaticClass()->GetDefaultObject<UInstancedSkinnedMeshComponent>());
@@ -1396,6 +1454,7 @@ FSkinnedMeshComponentDescriptor::FSkinnedMeshComponentDescriptor(const FSkinnedM
 FSkinnedMeshComponentDescriptor::~FSkinnedMeshComponentDescriptor() = default;
 
 FSkinnedMeshComponentDescriptor::FSkinnedMeshComponentDescriptor(const FSoftSkinnedMeshComponentDescriptor& InDescriptor)
+	: FSkinnedMeshComponentDescriptorBase(InDescriptor)
 {
 	SkinnedAsset = InDescriptor.SkinnedAsset.LoadSynchronous();
 	Algo::Transform(InDescriptor.BankItems, BankItems, [](const FSoftAnimBankItem& BankItem) { return FAnimBankItem(BankItem); });
@@ -1414,10 +1473,8 @@ UInstancedSkinnedMeshComponent* FSkinnedMeshComponentDescriptor::CreateComponent
 
 void FSkinnedMeshComponentDescriptor::InitFrom(const UInstancedSkinnedMeshComponent* Template, bool bInitBodyInstance)
 {
-	check(Template);
+	Super::InitFrom(Template, bInitBodyInstance);
 
-	ComponentClass = Template->GetClass();
-	Mobility = Template->Mobility;
 	SkinnedAsset = Template->GetSkinnedAsset();
 	BankItems = Template->GetAnimBankItems();
 }
@@ -1436,8 +1493,8 @@ uint32 FSkinnedMeshComponentDescriptor::ComputeHash() const
 
 void FSkinnedMeshComponentDescriptor::InitComponent(UInstancedSkinnedMeshComponent* Component) const
 {
+	Super::InitComponent(Component);
 	Component->SetSkinnedAsset(SkinnedAsset);
-	Component->Mobility = Mobility;
 	Component->SetAnimBankItems(BankItems);
 }
 
@@ -1457,11 +1514,11 @@ bool FSkinnedMeshComponentDescriptor::operator==(const FSkinnedMeshComponentDesc
 		(Hash == 0 || Other.Hash == 0 || Hash == Other.Hash) && // Check hash first, other checks are in case of Hash collision
 		SkinnedAsset == Other.SkinnedAsset &&
 		BankItems == Other.BankItems &&
-		ComponentClass == Other.ComponentClass &&
-		Mobility == Other.Mobility;
+		Super::operator==(Other);
 }
 
 FSoftSkinnedMeshComponentDescriptor::FSoftSkinnedMeshComponentDescriptor()
+	: FSkinnedMeshComponentDescriptorBase(NoInit)
 {
 	// Make sure we have proper defaults
 	InitFrom(UInstancedSkinnedMeshComponent::StaticClass()->GetDefaultObject<UInstancedSkinnedMeshComponent>());
@@ -1472,6 +1529,7 @@ FSoftSkinnedMeshComponentDescriptor::FSoftSkinnedMeshComponentDescriptor(const F
 FSoftSkinnedMeshComponentDescriptor::~FSoftSkinnedMeshComponentDescriptor() = default;
 
 FSoftSkinnedMeshComponentDescriptor::FSoftSkinnedMeshComponentDescriptor(const FSkinnedMeshComponentDescriptor& InDescriptor)
+	: FSkinnedMeshComponentDescriptorBase(InDescriptor)
 {
 	SkinnedAsset = InDescriptor.SkinnedAsset;
 	Algo::Transform(InDescriptor.BankItems, BankItems, [](const FAnimBankItem& BankItem) { return FSoftAnimBankItem(BankItem); });
@@ -1490,10 +1548,7 @@ UInstancedSkinnedMeshComponent* FSoftSkinnedMeshComponentDescriptor::CreateCompo
 
 void FSoftSkinnedMeshComponentDescriptor::InitFrom(const UInstancedSkinnedMeshComponent* Template, bool bInitBodyInstance)
 {
-	check(Template);
-
-	ComponentClass = Template->GetClass();
-	Mobility = Template->Mobility;
+	Super::InitFrom(Template, bInitBodyInstance);
 	SkinnedAsset = Template->GetSkinnedAsset();
 	BankItems = Template->GetAnimBankItems();
 }
@@ -1511,8 +1566,9 @@ uint32 FSoftSkinnedMeshComponentDescriptor::ComputeHash() const
 
 void FSoftSkinnedMeshComponentDescriptor::InitComponent(UInstancedSkinnedMeshComponent* Component) const
 {
+	Super::InitComponent(Component);
+
 	Component->SetSkinnedAsset(SkinnedAsset.LoadSynchronous());
-	Component->Mobility = Mobility;
 	TArray<FAnimBankItem> AnimBankItems;
 	for (const FSoftAnimBankItem& BankItem : BankItems)
 	{
@@ -1537,8 +1593,7 @@ bool FSoftSkinnedMeshComponentDescriptor::operator==(const FSoftSkinnedMeshCompo
 		(Hash == 0 || Other.Hash == 0 || Hash == Other.Hash) && // Check hash first, other checks are in case of Hash collision
 		SkinnedAsset == Other.SkinnedAsset &&
 		BankItems == Other.BankItems &&
-		ComponentClass == Other.ComponentClass &&
-		Mobility == Other.Mobility;
+		Super::operator==(Other);
 }
 
 #undef LOCTEXT_NAMESPACE
diff --git a/Engine/Source/Runtime/Engine/Private/Components/InstancedSkinnedMeshComponent.cpp b/Engine/Source/Runtime/Engine/Private/Components/InstancedSkinnedMeshComponent.cpp
--- a/Engine/Source/Runtime/Engine/Private/Components/InstancedSkinnedMeshComponent.cpp
+++ b/Engine/Source/Runtime/Engine/Private/Components/InstancedSkinnedMeshComponent.cpp
@@ -21,6 +21,10 @@
 #include "InstanceData/InstanceUpdateChangeSet.h"
 #include "InstanceData/InstanceDataUpdateUtils.h"
 
+#if WITH_EDITOR
+#include "WorldPartition/HLOD/HLODHashBuilder.h"
+#endif
+
 #include UE_INLINE_GENERATED_CPP_BY_NAME(InstancedSkinnedMeshComponent)
 
 static TAutoConsoleVariable<int32> CVarInstancedSkinnedMeshesForceRefPose(
@@ -47,10 +51,10 @@ public:
 	using Super = Nanite::FSkinnedSceneProxy;
 
 	FInstancedSkinnedMeshSceneProxy(const Nanite::FMaterialAudit& MaterialAudit, UInstancedSkinnedMeshComponent* InComponent, FSkeletalMeshRenderData* InRenderData)
-	: Super(MaterialAudit, InComponent, InRenderData, false /* bAllowScale */)
-	, AnimationMinScreenSize(InComponent->AnimationMinScreenSize)
+		: Super(MaterialAudit, InComponent, InRenderData, false /* bAllowScale */)
+		, AnimationMinScreenSize(InComponent->AnimationMinScreenSize)
 	{
-	#if WITH_EDITOR
+#if WITH_EDITOR
 		const bool bSupportInstancePicking = HasPerInstanceHitProxies() && SMInstanceElementDataUtil::SMInstanceElementsEnabled();
 		HitProxyMode = bSupportInstancePicking ? EHitProxyMode::PerInstance : EHitProxyMode::MaterialSection;
 
@@ -63,10 +67,14 @@ public:
 				SetSelection_GameThread(true);
 			}
 		}
-	#endif
+#endif
 		const bool bForceRefPose = CVarInstancedSkinnedMeshesForceRefPose.GetValueOnAnyThread() != 0;
 		const bool bUseAnimBank = !bForceRefPose && InComponent->AnimBankItems.Num() > 0;
 
+		InstanceMinDrawDistance = InComponent->InstanceMinDrawDistance;
+		InstanceStartCullDistance = InComponent->InstanceStartCullDistance;
+		InstanceEndCullDistance = InComponent->InstanceEndCullDistance;
+
 		InstanceDataSceneProxy = InComponent->InstanceDataManager.GetOrCreateProxy();
 		SetupInstanceSceneDataBuffers(InstanceDataSceneProxy->GeInstanceSceneDataBuffers());
 
@@ -161,6 +169,26 @@ public:
 	}
 	virtual float GetAnimationMinScreenSize() const  override { return AnimationMinScreenSize; }
 
+	virtual bool GetInstanceDrawDistanceMinMax(FVector2f& OutCullRange) const override
+	{
+		if (InstanceEndCullDistance > 0)
+		{
+			OutCullRange = FVector2f(float(InstanceMinDrawDistance), float(InstanceEndCullDistance));
+			return true;
+		}
+		else
+		{
+			OutCullRange = FVector2f(0.0f);
+			return false;
+		}
+	}
+
+	virtual void SetInstanceCullDistance_RenderThread(float StartCullDistance, float EndCullDistance) override
+	{
+		InstanceStartCullDistance = StartCullDistance;
+		InstanceEndCullDistance = EndCullDistance;
+	}
+
 public:
 	// FPrimitiveSceneProxy interface.
 	virtual SIZE_T GetTypeHash() const override
@@ -176,6 +204,9 @@ private:
 	TArray<FAnimBankItem> AnimBankItems;
 	uint64 BankRegistration = 0;
 	float AnimationMinScreenSize = 0.0f;
+	uint32 InstanceMinDrawDistance = 0;
+	uint32 InstanceStartCullDistance = 0;
+	uint32 InstanceEndCullDistance = 0;
 };
 
 /** 
@@ -1018,6 +1049,67 @@ void UInstancedSkinnedMeshComponent::SendRenderInstanceData_Concurrent()
 	}
 }
 
+bool UInstancedSkinnedMeshComponent::IsHLODRelevant() const
+{
+	if (!CanBeHLODRelevant(this))
+	{
+		return false;
+	}
+
+	if (!GetSkinnedAsset())
+	{
+		return false;
+	}
+
+	if (!IsVisible())
+	{
+		return false;
+	}
+
+	if (Mobility == EComponentMobility::Movable)
+	{
+		return false;
+	}
+
+#if WITH_EDITORONLY_DATA
+	if (!bEnableAutoLODGeneration)
+	{
+		return false;
+	}
+#endif
+
+	return true;
+}
+
+#if WITH_EDITOR
+void UInstancedSkinnedMeshComponent::ComputeHLODHash(FHLODHashBuilder& HashBuilder) const
+{
+	Super::ComputeHLODHash(HashBuilder);
+
+	FHLODHashScope HashScope(HashBuilder, TEXT("UInstancedSkinnedMeshComponent"));
+
+	for (const FSkinnedMeshInstanceData& SkinnedMeshInstanceData : InstanceData)
+	{
+		HashBuilder << FTransform(SkinnedMeshInstanceData.Transform);
+		HashBuilder << SkinnedMeshInstanceData.BankIndex;
+	}
+	HashBuilder << FHLODHashContext(GET_MEMBER_NAME_CHECKED(UInstancedSkinnedMeshComponent, InstanceData));
+
+	for (FAnimBankItem AnimBankItem : AnimBankItems)
+	{
+		HashBuilder << AnimBankItem.BankAsset;
+		HashBuilder << AnimBankItem.SequenceIndex;
+	}
+	HashBuilder << FHLODHashContext(GET_MEMBER_NAME_CHECKED(UInstancedSkinnedMeshComponent, AnimBankItems));
+
+	HashBuilder << InstanceCustomData << FHLODHashContext(GET_MEMBER_NAME_CHECKED(UInstancedSkinnedMeshComponent, InstanceCustomData));
+	HashBuilder << InstanceMinDrawDistance << FHLODHashContext(GET_MEMBER_NAME_CHECKED(UInstancedSkinnedMeshComponent, InstanceMinDrawDistance));
+	HashBuilder << InstanceStartCullDistance << FHLODHashContext(GET_MEMBER_NAME_CHECKED(UInstancedSkinnedMeshComponent, InstanceStartCullDistance));
+	HashBuilder << InstanceEndCullDistance << FHLODHashContext(GET_MEMBER_NAME_CHECKED(UInstancedSkinnedMeshComponent, InstanceEndCullDistance));
+		
+	HashBuilder << GetSkinnedAsset() << FHLODHashContext(TEXT("SkinnedAsset"));
+}
+#endif
 
 void UInstancedSkinnedMeshComponent::CreateRenderState_Concurrent(FRegisterComponentContext* Context)
 {
@@ -1265,6 +1357,20 @@ TStructOnScope<FActorComponentInstanceData> UInstancedSkinnedMeshComponent::GetC
 	return ComponentInstanceData;
 }
 
+void UInstancedSkinnedMeshComponent::SetCullDistances(int32 StartCullDistance, int32 EndCullDistance)
+{
+	if (InstanceStartCullDistance != StartCullDistance || InstanceEndCullDistance != EndCullDistance)
+	{
+		InstanceStartCullDistance = StartCullDistance;
+		InstanceEndCullDistance = EndCullDistance;
+
+		if (GetScene() && SceneProxy)
+		{
+			GetScene()->UpdateInstanceCullDistance(this, StartCullDistance, EndCullDistance);
+		}
+	}
+}
+
 void UInstancedSkinnedMeshComponent::PreApplyComponentInstanceData(struct FInstancedSkinnedMeshComponentInstanceData* InstancedMeshData)
 {
 #if WITH_EDITOR
@@ -1514,6 +1620,19 @@ void UInstancedSkinnedMeshComponent::SetNumCustomDataFloats(int32 InNumCustomDat
 	}
 }
 
+bool UInstancedSkinnedMeshComponent::GetCustomData(FPrimitiveInstanceId InstanceId, TArrayView<float> CustomDataFloats) const
+{
+	int32 InstanceIndex = InstanceDataManager.IdToIndex(InstanceId);
+	if (!InstanceData.IsValidIndex(InstanceIndex))
+	{
+		return false;
+	}
+
+	const int32 NumToCopy = FMath::Min(CustomDataFloats.Num(), NumCustomDataFloats);
+	FMemory::Memcpy(CustomDataFloats.GetData(), &InstanceCustomData[InstanceIndex * NumCustomDataFloats], NumToCopy * CustomDataFloats.GetTypeSize());
+	return true;
+}
+
 bool UInstancedSkinnedMeshComponent::GetInstanceTransform(FPrimitiveInstanceId InstanceId, FTransform& OutInstanceTransform, bool bWorldSpace) const
 {
 	int32 InstanceIndex = InstanceDataManager.IdToIndex(InstanceId);
@@ -1533,6 +1652,18 @@ bool UInstancedSkinnedMeshComponent::GetInstanceTransform(FPrimitiveInstanceId I
 	return true;
 }
 
+bool UInstancedSkinnedMeshComponent::GetInstanceBankIndex(FPrimitiveInstanceId InstanceId, int32& OutBankIndex) const
+{
+	int32 InstanceIndex = InstanceDataManager.IdToIndex(InstanceId);
+	if (!InstanceData.IsValidIndex(InstanceIndex))
+	{
+		return false;
+	}
+
+	OutBankIndex = InstanceData[InstanceIndex].BankIndex;
+	return true;
+}
+
 bool UInstancedSkinnedMeshComponent::RemoveInstance(FPrimitiveInstanceId InstanceId)
 {
 	int32 InstanceIndex = InstanceDataManager.IdToIndex(InstanceId);
diff --git a/Engine/Source/Runtime/Engine/Private/Components/PrimitiveComponent.cpp b/Engine/Source/Runtime/Engine/Private/Components/PrimitiveComponent.cpp
--- a/Engine/Source/Runtime/Engine/Private/Components/PrimitiveComponent.cpp
+++ b/Engine/Source/Runtime/Engine/Private/Components/PrimitiveComponent.cpp
@@ -63,6 +63,7 @@
 #include "Engine/LODActor.h"
 #include "Interfaces/ITargetPlatform.h"
 #include "Rendering/StaticLightingSystemInterface.h"
+#include "WorldPartition/HLOD/HLODHashBuilder.h"
 #else
 #include "Components/InstancedStaticMeshComponent.h"
 #endif // WITH_EDITOR
@@ -5274,6 +5275,30 @@ void UPrimitiveComponent::SetExcludedFromHLODLevel(EHLODLevelExclusion HLODLevel
 	}
 }
 
+#if WITH_EDITOR
+void UPrimitiveComponent::ComputeHLODHash(FHLODHashBuilder& HashBuilder) const
+{
+	Super::ComputeHLODHash(HashBuilder);
+
+	FHLODHashScope HashScope(HashBuilder, TEXT("UPrimitiveComponent"));
+
+	HashBuilder << GetComponentTransform() << FHLODHashContext(TEXT("Transform"));
+	HashBuilder << HLODBatchingPolicy << FHLODHashContext(GET_MEMBER_NAME_CHECKED(UPrimitiveComponent, HLODBatchingPolicy));
+	HashBuilder << CastShadow << FHLODHashContext(GET_MEMBER_NAME_CHECKED(UPrimitiveComponent, CastShadow));
+	HashBuilder << bCastStaticShadow << FHLODHashContext(GET_MEMBER_NAME_CHECKED(UPrimitiveComponent, bCastStaticShadow));
+	HashBuilder << bCastDynamicShadow << FHLODHashContext(GET_MEMBER_NAME_CHECKED(UPrimitiveComponent, bCastDynamicShadow));
+	HashBuilder << bCastFarShadow << FHLODHashContext(GET_MEMBER_NAME_CHECKED(UPrimitiveComponent, bCastFarShadow));
+	HashBuilder << GetCustomPrimitiveData() << FHLODHashContext(TEXT("CustomPrimitiveData"));
+	
+	// Materials
+	const int32 NumMaterials = GetNumMaterials();
+	for (int32 MaterialIndex = 0; MaterialIndex < NumMaterials; ++MaterialIndex)
+	{
+		HashBuilder << GetMaterial(MaterialIndex);
+	}
+}
+#endif
+
 void UPrimitiveComponent::GetPrimitiveStats(FPrimitiveStats& PrimitiveStats) const
 {
 	// no default values returned
diff --git a/Engine/Source/Runtime/Engine/Private/Components/StaticMeshComponent.cpp b/Engine/Source/Runtime/Engine/Private/Components/StaticMeshComponent.cpp
--- a/Engine/Source/Runtime/Engine/Private/Components/StaticMeshComponent.cpp
+++ b/Engine/Source/Runtime/Engine/Private/Components/StaticMeshComponent.cpp
@@ -38,6 +38,7 @@
 #include "Rendering/StaticLightingSystemInterface.h"
 #include "Streaming/ActorTextureStreamingBuildDataComponent.h"
 #include "Templates/UnrealTemplate.h"
+#include "WorldPartition/HLOD/HLODHashBuilder.h"
 #endif // WITH_EDITOR
 #include "LightMap.h"
 #include "ShadowMap.h"
@@ -3015,6 +3016,35 @@ bool UStaticMeshComponent::IsHLODRelevant() const
 	return true;
 }
 
+#if WITH_EDITOR
+void UStaticMeshComponent::ComputeHLODHash(FHLODHashBuilder& HashBuilder) const
+{
+	Super::ComputeHLODHash(HashBuilder);
+
+	FHLODHashScope HashScope(HashBuilder, TEXT("UStaticMeshComponent"));
+
+	HashBuilder << ForcedLodModel << FHLODHashContext(GET_MEMBER_NAME_CHECKED(UStaticMeshComponent, ForcedLodModel));
+	
+	int32 Width, Height;
+	if (GetLightMapResolution(Width, Height))
+	{
+		HashBuilder << Width << Height << FHLODHashContext(TEXT("LightMapResolution"));
+	}
+	
+	for (const FStaticMeshComponentLODInfo& LODInfo : LODData)
+	{
+		if (LODInfo.OverrideVertexColors)
+		{
+			HashBuilder.Serialize((uint8*)LODInfo.OverrideVertexColors->GetVertexData(), LODInfo.OverrideVertexColors->GetNumVertices() * LODInfo.OverrideVertexColors->GetStride());
+			HashBuilder << FHLODHashContext(TEXT("OverrideVertexColors"));
+		}
+	}
+	
+	// Static mesh
+	HashBuilder << GetStaticMesh().Get();
+}
+#endif
+
 bool UStaticMeshComponent::DoCustomNavigableGeometryExport(FNavigableGeometryExport& GeomExport) const
 {
 	const FVector Scale3D = GetComponentToWorld().GetScale3D();
diff --git a/Engine/Source/Runtime/Engine/Private/ISMPartition/ISMComponentBatcher.cpp b/Engine/Source/Runtime/Engine/Private/ISMPartition/ISMComponentBatcher.cpp
--- a/Engine/Source/Runtime/Engine/Private/ISMPartition/ISMComponentBatcher.cpp
+++ b/Engine/Source/Runtime/Engine/Private/ISMPartition/ISMComponentBatcher.cpp
@@ -1,7 +1,7 @@
 // Copyright Epic Games, Inc. All Rights Reserved.
 
 #include "ISMPartition/ISMComponentBatcher.h"
-#include "Components/StaticMeshComponent.h"
+#include "Components/InstancedSkinnedMeshComponent.h"
 #include "Serialization/ArchiveCrc32.h"
 #include "Misc/TransformUtilities.h"
 #include "Templates/TypeHash.h"
@@ -33,6 +33,11 @@ void FISMComponentBatcher::AddInternal(const UActorComponent* InComponent, TOpti
 		NewNumCustomDataFloats = SMC->GetCustomPrimitiveData().Data.Num();
 		NewNumInstances = 1;
 	}
+	else if (const UInstancedSkinnedMeshComponent* ISKMC = Cast<UInstancedSkinnedMeshComponent>(InComponent))
+	{
+		NewNumCustomDataFloats = ISKMC->GetNumCustomDataFloats();
+		NewNumInstances = ISKMC->GetInstanceCount();
+	}
 
 	if (NewNumCustomDataFloats > NumCustomDataFloats)
 	{
@@ -107,6 +112,37 @@ void FISMComponentBatcher::AddInternal(const UActorComponent* InComponent, TOpti
 		InstancesCustomData.Append(SMC->GetCustomPrimitiveData().Data);
 		InstancesCustomData.AddDefaulted(NumCustomDataFloats - SMC->GetCustomPrimitiveData().Data.Num());
 	}
+	else if (const UInstancedSkinnedMeshComponent* ISKMC = Cast<UInstancedSkinnedMeshComponent>(InComponent))
+	{
+		// Add each instance
+		FPrimitiveInstanceId InstanceId;
+		for (InstanceId.Id = 0; InstanceId.Id < ISKMC->GetInstanceCount(); InstanceId.Id++)
+		{
+			// Add instance transform
+			FTransform InstanceTransformWS;
+
+			ISKMC->GetInstanceTransform(InstanceId, InstanceTransformWS, /*bWorldSpace*/ true);
+			if (InTransformFunc.IsSet())
+			{
+				InstanceTransformWS = InTransformFunc->operator()(InstanceTransformWS);
+			}
+			InstancesTransformsWS.Add(InstanceTransformWS);
+
+			ISKMC->GetInstanceBankIndex(InstanceId, BankIndices.AddZeroed_GetRef());
+			
+			// Add per instance custom data, if any
+			if (NumCustomDataFloats > 0)
+			{
+				int32 StartIndex = InstancesCustomData.Num(); // Get the current end index
+				InstancesCustomData.AddDefaulted(NumCustomDataFloats); // Append new elements
+
+				if (ISKMC->GetNumCustomDataFloats() > 0)
+				{
+					ISKMC->GetCustomData(InstanceId, TArrayView<float>(InstancesCustomData.GetData() + StartIndex, NumCustomDataFloats));
+				}
+			}
+		}
+	}
 }
 
 void FISMComponentBatcher::InitComponent(UInstancedStaticMeshComponent* ISMComponent) const
@@ -126,6 +162,21 @@ void FISMComponentBatcher::InitComponent(UInstancedStaticMeshComponent* ISMCompo
 	}		
 }
 
+void FISMComponentBatcher::InitComponent(UInstancedSkinnedMeshComponent* ISMComponent) const
+{
+	ISMComponent->SetNumCustomDataFloats(NumCustomDataFloats);
+	ISMComponent->AddInstances(InstancesTransformsWS, BankIndices, /*bShouldReturnIndices*/false, /*bWorldSpace*/true);
+
+	if (NumCustomDataFloats)
+	{
+		FPrimitiveInstanceId InstanceId;
+		for (InstanceId.Id = 0; InstanceId.Id < ISMComponent->GetInstanceCount(); InstanceId.Id++)
+		{
+			ISMComponent->SetCustomData(InstanceId, TArrayView<const float>(&InstancesCustomData[InstanceId.Id], NumCustomDataFloats));
+		}
+	}
+}
+
 void FISMComponentBatcher::ComputeHash() const
 {
 	uint32 CRC = 0;
diff --git a/Engine/Source/Runtime/Engine/Private/InstancedStaticMesh.cpp b/Engine/Source/Runtime/Engine/Private/InstancedStaticMesh.cpp
--- a/Engine/Source/Runtime/Engine/Private/InstancedStaticMesh.cpp
+++ b/Engine/Source/Runtime/Engine/Private/InstancedStaticMesh.cpp
@@ -64,6 +64,7 @@
 
 #if WITH_EDITOR
 #include "Rendering/StaticLightingSystemInterface.h"
+#include "WorldPartition/HLOD/HLODHashBuilder.h"
 #endif
 
 #if UE_BUILD_SHIPPING && !WITH_EDITORONLY_DATA
@@ -2285,6 +2286,31 @@ bool UInstancedStaticMeshComponent::IsHLODRelevant() const
 	return Super::IsHLODRelevant();
 }
 
+#if WITH_EDITOR
+void UInstancedStaticMeshComponent::ComputeHLODHash(FHLODHashBuilder& HashBuilder) const
+{
+	Super::ComputeHLODHash(HashBuilder);
+
+	FHLODHashScope HashScope(HashBuilder, TEXT("UInstancedStaticMeshComponent"));
+
+	for (const FInstancedStaticMeshInstanceData& InstanceData : PerInstanceSMData)
+	{
+		HashBuilder << FTransform(InstanceData.Transform);
+	}
+	HashBuilder << FHLODHashContext(GET_MEMBER_NAME_CHECKED(UInstancedStaticMeshComponent, PerInstanceSMData));
+
+	HashBuilder << PerInstanceSMCustomData << FHLODHashContext(GET_MEMBER_NAME_CHECKED(UInstancedStaticMeshComponent, PerInstanceSMCustomData));
+	HashBuilder << InstancingRandomSeed << FHLODHashContext(GET_MEMBER_NAME_CHECKED(UInstancedStaticMeshComponent, InstancingRandomSeed));
+		
+	for (const FInstancedStaticMeshRandomSeed& ISMCRandomSeed : AdditionalRandomSeeds)
+	{
+		HashBuilder << ISMCRandomSeed.StartInstanceIndex;
+		HashBuilder << ISMCRandomSeed.RandomSeed;
+	}
+	HashBuilder << FHLODHashContext(GET_MEMBER_NAME_CHECKED(UInstancedStaticMeshComponent, AdditionalRandomSeeds));
+}
+#endif
+
 void UInstancedStaticMeshComponent::SendRenderInstanceData_Concurrent()
 {
 	Super::SendRenderInstanceData_Concurrent();
diff --git a/Engine/Source/Runtime/Engine/Private/WorldPartition/HLOD/HLODBuilder.cpp b/Engine/Source/Runtime/Engine/Private/WorldPartition/HLOD/HLODBuilder.cpp
--- a/Engine/Source/Runtime/Engine/Private/WorldPartition/HLOD/HLODBuilder.cpp
+++ b/Engine/Source/Runtime/Engine/Private/WorldPartition/HLOD/HLODBuilder.cpp
@@ -4,6 +4,7 @@
 
 #include "Components/StaticMeshComponent.h"
 #include "Engine/HLODProxy.h"
+#include "Engine/SkinnedAsset.h"
 #include "Engine/StaticMesh.h"
 #include "Engine/Texture.h"
 #include "HLOD/HLODBatchingPolicy.h"
@@ -12,7 +13,9 @@
 #include "Materials/MaterialInterface.h"
 #include "Misc/ConfigCacheIni.h"
 #include "UObject/Package.h"
+#include "WorldPartition/HLOD/HLODHashBuilder.h"
 #include "WorldPartition/HLOD/HLODInstancedStaticMeshComponent.h"
+#include "WorldPartition/HLOD/HLODInstancedSkinnedMeshComponent.h"
 
 #include "AssetRegistry/AssetData.h"
 #include "AssetRegistry/AssetRegistryModule.h"
@@ -27,6 +30,7 @@ UHLODBuilder::UHLODBuilder(const FObjectInitializer& ObjectInitializer)
 	: Super(ObjectInitializer)
 #if WITH_EDITORONLY_DATA
 	, HLODInstancedStaticMeshComponentClass(UHLODInstancedStaticMeshComponent::StaticClass())
+	, HLODInstancedSkinnedMeshComponentClass(UHLODInstancedSkinnedMeshComponent::StaticClass())
 #endif
 {
 }
@@ -61,70 +65,9 @@ bool UHLODBuilder::RequiresWarmup() const
 
 uint32 UHLODBuilder::ComputeHLODHash(const UActorComponent* InSourceComponent) const
 {
-	uint32 ComponentCRC = 0;
-
-	if (UStaticMeshComponent* StaticMeshComponent = Cast<UStaticMeshComponent>(const_cast<UActorComponent*>(InSourceComponent)))
-	{
-		UE_LOG(LogHLODBuilder, VeryVerbose, TEXT(" - Component \'%s\' from actor \'%s\'"), *StaticMeshComponent->GetName(), *StaticMeshComponent->GetOwner()->GetName());
-
-		// CRC component
-		ComponentCRC = UHLODProxy::GetCRC(StaticMeshComponent);
-		UE_LOG(LogHLODBuilder, VeryVerbose, TEXT("     - Static Mesh Component (%s) = %x"), *StaticMeshComponent->GetName(), ComponentCRC);
-
-		// CRC static mesh
-		if (UStaticMesh* StaticMesh = StaticMeshComponent->GetStaticMesh())
-		{
-			uint32 StaticMeshCRC = UHLODProxy::GetCRC(StaticMesh);
-			UE_LOG(LogHLODBuilder, VeryVerbose, TEXT("     - Static Mesh (%s) = %x"), *StaticMesh->GetName(), StaticMeshCRC);
-			ComponentCRC = HashCombineFast(ComponentCRC, StaticMeshCRC);
-		}
-
-		// CRC materials
-		const int32 NumMaterials = StaticMeshComponent->GetNumMaterials();
-		for (int32 MaterialIndex = 0; MaterialIndex < NumMaterials; ++MaterialIndex)
-		{
-			UMaterialInterface* MaterialInterface = StaticMeshComponent->GetMaterial(MaterialIndex);
-			if (MaterialInterface)
-			{
-				uint32 MaterialInterfaceCRC = UHLODProxy::GetCRC(MaterialInterface);
-				UE_LOG(LogHLODBuilder, VeryVerbose, TEXT("     - Material (%s) = %x"), *MaterialInterface->GetName(), MaterialInterfaceCRC);
-				ComponentCRC = HashCombineFast(ComponentCRC, MaterialInterfaceCRC);
-
-				TArray<UTexture*> Textures;
-				MaterialInterface->GetUsedTextures(Textures, EMaterialQualityLevel::High, true, ERHIFeatureLevel::SM5, true);
-				for (UTexture* Texture : Textures)
-				{
-					uint32 TextureCRC = UHLODProxy::GetCRC(Texture);
-					UE_LOG(LogHLODBuilder, VeryVerbose, TEXT("     - Texture (%s) = %x"), *Texture->GetName(), TextureCRC);
-					ComponentCRC = HashCombineFast(ComponentCRC, TextureCRC);
-				}
-			}
-			UMaterialInterface* NaniteOverride = MaterialInterface ? MaterialInterface->GetNaniteOverride() : nullptr;
-			if (NaniteOverride)
-			{
-				uint32 NaniteOverrideCRC = UHLODProxy::GetCRC(NaniteOverride);
-				UE_LOG(LogHLODBuilder, VeryVerbose, TEXT("     - Nanite Override Material (%s) = %x"), *NaniteOverride->GetName(), NaniteOverrideCRC);
-				ComponentCRC = HashCombineFast(ComponentCRC, NaniteOverrideCRC);
-
-				TArray<UTexture*> Textures;
-				NaniteOverride->GetUsedTextures(Textures, EMaterialQualityLevel::High, true, ERHIFeatureLevel::SM5, true);
-				for (UTexture* Texture : Textures)
-				{
-					uint32 TextureCRC = UHLODProxy::GetCRC(Texture);
-					UE_LOG(LogHLODBuilder, VeryVerbose, TEXT("     - Nanite Override Texture (%s) = %x"), *Texture->GetName(), TextureCRC);
-					ComponentCRC = HashCombineFast(ComponentCRC, TextureCRC);
-				}
-			}
-		}
-	}
-	else
-	{
-		ComponentCRC = FMath::Rand();
-		UE_LOG(LogHLODBuilder, Warning, TEXT("Can't compute HLOD hash for component of type %s, assuming it is dirty."), *InSourceComponent->GetClass()->GetName());
-		
-	}
-
-	return ComponentCRC;
+	FHLODHashBuilder HashBuilder;
+	InSourceComponent->ComputeHLODHash(HashBuilder);
+	return HashBuilder.GetCrc();
 }
 
 uint32 UHLODBuilder::ComputeHLODHash(const TArray<UActorComponent*>& InSourceComponents)
@@ -141,13 +84,19 @@ uint32 UHLODBuilder::ComputeHLODHash(const TArray<UActorComponent*>& InSourceCom
 		}
 
 		uint32 ComponentHash = HLODBuilderClass->GetDefaultObject<UHLODBuilder>()->ComputeHLODHash(SourceComponent);
+		if (ComponentHash == 0)
+		{
+			UE_LOG(LogHLODBuilder, Warning, TEXT("Can't compute HLOD hash for component of type %s, assuming it is dirty."), *SourceComponent->GetClass()->GetName());
+			ComponentHash = FMath::Rand();
+		}
+
 		ComponentsCRCs.Add(ComponentHash);
 	}
 
 	// Sort the components CRCs to ensure the order of components won't have an impact on the final CRC
 	ComponentsCRCs.Sort();
 
-	return FCrc::MemCrc32(ComponentsCRCs.GetData(), ComponentsCRCs.Num() * ComponentsCRCs.GetTypeSize());
+	return FCrc::MemCrc32(ComponentsCRCs.GetData(), ComponentsCRCs.NumBytes());
 }
 
 TSubclassOf<UHLODInstancedStaticMeshComponent> UHLODBuilder::GetInstancedStaticMeshComponentClass()
@@ -165,60 +114,155 @@ TSubclassOf<UHLODInstancedStaticMeshComponent> UHLODBuilder::GetInstancedStaticM
 	return ISMClass;
 }
 
+TSubclassOf<UHLODInstancedSkinnedMeshComponent> UHLODBuilder::GetInstancedSkinnedMeshComponentClass()
+{
+	TSubclassOf<UHLODInstancedSkinnedMeshComponent> ISMClass = StaticClass()->GetDefaultObject<UHLODBuilder>()->HLODInstancedSkinnedMeshComponentClass;
+	if (!ISMClass)
+	{
+		FString ConfigValue;
+		GConfig->GetString(TEXT("/Script/Engine.HLODBuilder"), TEXT("HLODInstancedSkinnedMeshComponentClass"), ConfigValue, GEditorIni);
+		UE_LOG(LogHLODBuilder, Error, TEXT("Could not resolve the class specified for HLODInstancedSkinnedMeshComponentClass. Config value was %s"), *ConfigValue);
+
+		// Fallback to standard HLOD ISMC
+		ISMClass = UHLODInstancedSkinnedMeshComponent::StaticClass();
+	}
+	return ISMClass;
+}
+
 namespace
 {
-	struct FInstanceBatch
+	struct FInstancedStaticMeshBatch
 	{
 		TUniquePtr<FISMComponentDescriptor>		ISMComponentDescriptor;
 		FISMComponentBatcher					ISMComponentBatcher;
 	};
+
+	struct FInstancedSkinnedMeshBatch
+	{
+		TUniquePtr<FSkinnedMeshComponentDescriptor>	ISMComponentDescriptor;
+		FISMComponentBatcher						ISMComponentBatcher;
+	};
 }
 
 TArray<UActorComponent*> UHLODBuilder::BatchInstances(const TArray<UActorComponent*>& InSourceComponents)
 {
 	TRACE_CPUPROFILER_EVENT_SCOPE(UHLODBuilderInstancing::Build);
 
-	TArray<UStaticMeshComponent*> SourceStaticMeshComponents = FilterComponents<UStaticMeshComponent>(InSourceComponents);
+	TArray<UActorComponent*> HLODComponents;
 
-	TSubclassOf<UHLODInstancedStaticMeshComponent> ComponentClass = UHLODBuilder::GetInstancedStaticMeshComponentClass();
+	TArray<UStaticMeshComponent*> SourceStaticMeshComponents;
+	TArray<UInstancedSkinnedMeshComponent*> SourceInstancedSkinnedMeshComponents;
+	TArray<UActorComponent*> FilteredOutComponents;
 
-	// Prepare instance batches
-	TMap<uint32, FInstanceBatch> InstancesData;
-	for (UStaticMeshComponent* SMC : SourceStaticMeshComponents)
+	for (UActorComponent* SourceComponent : InSourceComponents)
 	{
-		const UStaticMesh* StaticMesh = SMC->GetStaticMesh();
-		const bool bPrivateStaticMesh = StaticMesh && !StaticMesh->HasAnyFlags(RF_Public);
-		const bool bTransientStaticMesh = StaticMesh && StaticMesh->HasAnyFlags(RF_Transient);
-		if (!StaticMesh || bPrivateStaticMesh || bTransientStaticMesh)
+		if (UStaticMeshComponent* SourceStaticMeshComponent = Cast<UStaticMeshComponent>(SourceComponent))
 		{
-			UE_LOG(LogHLODBuilder, Warning, TEXT("Instanced HLOD source component %s points to a %s static mesh, ignoring."), *SMC->GetPathName(), !StaticMesh ? TEXT("null") : bPrivateStaticMesh ? TEXT("private") : TEXT("transient"));
-			continue;
+			SourceStaticMeshComponents.Add(SourceStaticMeshComponent);
 		}
+		else if (UInstancedSkinnedMeshComponent* SourceInstancedSkinnedMeshComponent = Cast<UInstancedSkinnedMeshComponent>(SourceComponent))
+		{
+			SourceInstancedSkinnedMeshComponents.Add(SourceInstancedSkinnedMeshComponent);
+		}
+		else if (SourceComponent)
+		{
+			FilteredOutComponents.Add(SourceComponent);
+		}
+	}
 
-		TUniquePtr<FISMComponentDescriptor> ISMComponentDescriptor = ComponentClass->GetDefaultObject<UHLODInstancedStaticMeshComponent>()->AllocateISMComponentDescriptor();
-		ISMComponentDescriptor->InitFrom(SMC, false);
-
-		FInstanceBatch& InstanceBatch = InstancesData.FindOrAdd(ISMComponentDescriptor->GetTypeHash());
-		if (!InstanceBatch.ISMComponentDescriptor.IsValid())
+	// Log filtered out components, this is unexpected
+	if (!FilteredOutComponents.IsEmpty())
+	{
+		UE_LOG(LogHLODBuilder, Warning, TEXT("UHLODBuilder::BatchInstances - Excluding %d unsupported components:"), FilteredOutComponents.Num());
+		for (UActorComponent* FilteredOutComponent : FilteredOutComponents)
 		{
-			InstanceBatch.ISMComponentDescriptor = MoveTemp(ISMComponentDescriptor);
+			UE_LOG(LogHLODBuilder, Warning, TEXT("    -> (%s) %s"), *FilteredOutComponent->GetClass()->GetName(), *FilteredOutComponent->GetName());
 		}
-		
-		InstanceBatch.ISMComponentBatcher.Add(SMC);
 	}
 
-	// Create an ISMC for each SM asset we found
-	TArray<UActorComponent*> HLODComponents;
-	for (auto& Entry : InstancesData)
+	auto IsAssetValidForInstancing = [](UActorComponent* InComponent, UObject* InObj)
 	{
-		const FInstanceBatch& InstanceBatch = Entry.Value;
+		const bool bPrivateAsset = InObj && !InObj->HasAnyFlags(RF_Public);
+		const bool bTransientAsset = InObj && InObj->HasAnyFlags(RF_Transient);
+		if (!InObj || bPrivateAsset || bTransientAsset)
+		{
+			UE_LOG(LogHLODBuilder, Warning, TEXT("Instanced HLOD source component %s points to a %s mesh, ignoring."), *InComponent->GetPathName(), !InObj ? TEXT("null") : bPrivateAsset ? TEXT("private") : TEXT("transient"));
+			return false;
+		}
 
-		UInstancedStaticMeshComponent* ISMComponent = InstanceBatch.ISMComponentDescriptor->CreateComponent(GetTransientPackage());
-		InstanceBatch.ISMComponentBatcher.InitComponent(ISMComponent);
-				
-		HLODComponents.Add(ISMComponent);
+		return true;
 	};
 
+	// Static meshes batching
+	{		
+		TSubclassOf<UHLODInstancedStaticMeshComponent> ComponentClass = UHLODBuilder::GetInstancedStaticMeshComponentClass();
+
+		// Prepare instance batches
+		TMap<uint32, FInstancedStaticMeshBatch> InstancesData;
+		for (UStaticMeshComponent* SMC : SourceStaticMeshComponents)
+		{
+			if (!IsAssetValidForInstancing(SMC, SMC->GetStaticMesh()))
+			{
+				continue;
+			}
+
+			TUniquePtr<FISMComponentDescriptor> ISMComponentDescriptor = ComponentClass->GetDefaultObject<UHLODInstancedStaticMeshComponent>()->AllocateISMComponentDescriptor();
+			ISMComponentDescriptor->InitFrom(SMC, false);
+
+			FInstancedStaticMeshBatch& InstanceBatch = InstancesData.FindOrAdd(ISMComponentDescriptor->GetTypeHash());
+			if (!InstanceBatch.ISMComponentDescriptor.IsValid())
+			{
+				InstanceBatch.ISMComponentDescriptor = MoveTemp(ISMComponentDescriptor);
+			}
+
+			InstanceBatch.ISMComponentBatcher.Add(SMC);
+		}
+
+		// Create an ISMC for each SM asset we found		
+		for (auto& Entry : InstancesData)
+		{
+			const FInstancedStaticMeshBatch& InstanceBatch = Entry.Value;
+			UInstancedStaticMeshComponent* ISMComponent = InstanceBatch.ISMComponentDescriptor->CreateComponent(GetTransientPackage());
+			InstanceBatch.ISMComponentBatcher.InitComponent(ISMComponent);
+			HLODComponents.Add(ISMComponent);
+		};
+	}
+
+	// Skinned meshes batching
+	{
+		TSubclassOf<UHLODInstancedSkinnedMeshComponent> ComponentClass = UHLODBuilder::GetInstancedSkinnedMeshComponentClass();
+
+		// Prepare instance batches
+		TMap<uint32, FInstancedSkinnedMeshBatch> InstancesData;
+		for (UInstancedSkinnedMeshComponent* ISKMC : SourceInstancedSkinnedMeshComponents)
+		{
+			if (!IsAssetValidForInstancing(ISKMC, ISKMC->GetSkinnedAsset()))
+			{
+				continue;
+			}
+
+			TUniquePtr<FSkinnedMeshComponentDescriptor> ISMComponentDescriptor = ComponentClass->GetDefaultObject<UHLODInstancedSkinnedMeshComponent>()->AllocateISMComponentDescriptor();
+			ISMComponentDescriptor->InitFrom(ISKMC, false);
+
+			FInstancedSkinnedMeshBatch& InstanceBatch = InstancesData.FindOrAdd(ISMComponentDescriptor->GetTypeHash());
+			if (!InstanceBatch.ISMComponentDescriptor.IsValid())
+			{
+				InstanceBatch.ISMComponentDescriptor = MoveTemp(ISMComponentDescriptor);
+			}
+
+			InstanceBatch.ISMComponentBatcher.Add(ISKMC);
+		}
+
+		// Create an ISKMC for each SKM asset we found
+		for (auto& Entry : InstancesData)
+		{
+			const FInstancedSkinnedMeshBatch& InstanceBatch = Entry.Value;
+			UInstancedSkinnedMeshComponent* ISMComponent = InstanceBatch.ISMComponentDescriptor->CreateComponent(GetTransientPackage());
+			InstanceBatch.ISMComponentBatcher.InitComponent(ISMComponent);
+			HLODComponents.Add(ISMComponent);
+		};
+	}
+
 	return HLODComponents;
 }
 
diff --git a/Engine/Source/Runtime/Engine/Private/WorldPartition/HLOD/HLODHashBuilder.cpp b/Engine/Source/Runtime/Engine/Private/WorldPartition/HLOD/HLODHashBuilder.cpp
new file mode 100644
--- /dev/null
+++ b/Engine/Source/Runtime/Engine/Private/WorldPartition/HLOD/HLODHashBuilder.cpp
@@ -0,0 +1,117 @@
+// Copyright Epic Games, Inc. All Rights Reserved.
+
+#include "WorldPartition/HLOD/HLODHashBuilder.h"
+
+#if WITH_EDITOR
+
+#include "Engine/HLODProxy.h"
+#include "Engine/SkinnedAsset.h"
+#include "Engine/StaticMesh.h"
+#include "Engine/Texture.h"
+#include "Materials/MaterialInterface.h"
+#include "Misc/TransformUtilities.h"
+#include "WorldPartition/HLOD/HLODBuilder.h"
+
+void FHLODHashBuilder::PushContext(const TCHAR* Context)
+{
+	LogContext(Context, false);
+
+	++IndentationLevel;
+}
+
+void FHLODHashBuilder::PopContext()
+{
+	check(IndentationLevel > 0);
+	--IndentationLevel;
+}
+
+FArchive& FHLODHashBuilder::operator<<(FTransform InTransform)
+{
+	*this << TransformUtilities::GetRoundedTransformCRC32(InTransform);
+	return *this;
+}
+
+FArchive& FHLODHashBuilder::operator<<(UObject*& InObject)
+{
+	if (UMaterialInterface* MaterialInterface = Cast<UMaterialInterface>(InObject))
+	{
+		FHLODHashScope HashScope(*this, TEXT("UMaterialInterface"));
+
+		*this << UHLODProxy::GetCRC(MaterialInterface);
+		*this << FHLODHashContext(TEXT("%s (%s)"), *InObject->GetClass()->GetName(), *InObject->GetName());
+
+		TArray<UTexture*> Textures;
+		MaterialInterface->GetUsedTextures(Textures, EMaterialQualityLevel::High, true, ERHIFeatureLevel::SM5, true);
+		*this << Textures;
+	
+		if (UMaterialInterface* NaniteOverride = MaterialInterface->GetNaniteOverride())
+		{ 
+			*this << UHLODProxy::GetCRC(NaniteOverride);
+			*this << FHLODHashContext(TEXT("%s (%s)"), *NaniteOverride->GetClass()->GetName(), *NaniteOverride->GetName());
+
+			Textures.Reset();
+			NaniteOverride->GetUsedTextures(Textures, EMaterialQualityLevel::High, true, ERHIFeatureLevel::SM5, true);
+			*this << Textures;
+		}
+	}
+	else if (UTexture* Texture = Cast<UTexture>(InObject))
+	{
+		FHLODHashScope HashScope(*this, TEXT("UTexture"));
+
+		*this << UHLODProxy::GetCRC(Texture);
+		*this << FHLODHashContext(TEXT("%s (%s)"), *InObject->GetClass()->GetName(), *InObject->GetName());
+	}
+	else if (UStaticMesh* StaticMesh = Cast<UStaticMesh>(InObject))
+	{
+		FHLODHashScope HashScope(*this, TEXT("UStaticMesh"));
+
+		*this << UHLODProxy::GetCRC(StaticMesh);
+		*this << FHLODHashContext(TEXT("%s (%s)"), *InObject->GetClass()->GetName(), *InObject->GetName());
+	}
+	else
+	{
+		FArchive::operator<<(InObject);
+	}
+	
+	return *this;
+}
+
+FArchive& FHLODHashBuilder::operator<<(UMaterialInterface* InMaterialInterface)
+{
+	UObject* Object = InMaterialInterface;
+	return *this << Object;
+}
+
+FArchive& FHLODHashBuilder::operator<<(UTexture* InTexture)
+{
+	UObject* Object = InTexture;
+	return *this << Object;
+}
+
+FArchive& FHLODHashBuilder::operator<<(UStaticMesh* InStaticMesh)
+{
+	UObject* Object = InStaticMesh;
+	return *this << Object;
+}
+
+FArchive& FHLODHashBuilder::operator<<(USkinnedAsset* InSkinnedAsset)
+{
+	UObject* Object = InSkinnedAsset;
+	return *this << Object;
+}
+
+void FHLODHashBuilder::LogContext(const TCHAR* Context, bool bOutputHash)
+{
+	FString Indentation = FString::ChrN(IndentationLevel * 4, TEXT(' '));
+
+	if (bOutputHash)
+	{
+		UE_LOG(LogHLODBuilder, VeryVerbose, TEXT("%s- %s = %x"), *Indentation, Context, GetCrc());
+	}
+	else
+	{
+		UE_LOG(LogHLODBuilder, VeryVerbose, TEXT("%s- %s"), *Indentation, Context);
+	}
+}
+
+#endif // #if WITH_EDITOR
diff --git a/Engine/Source/Runtime/Engine/Private/WorldPartition/HLOD/HLODInstancedSkinnedMeshComponent.cpp b/Engine/Source/Runtime/Engine/Private/WorldPartition/HLOD/HLODInstancedSkinnedMeshComponent.cpp
new file mode 100644
--- /dev/null
+++ b/Engine/Source/Runtime/Engine/Private/WorldPartition/HLOD/HLODInstancedSkinnedMeshComponent.cpp
@@ -0,0 +1,37 @@
+// Copyright Epic Games, Inc. All Rights Reserved.
+
+#include "WorldPartition/HLOD/HLODInstancedSkinnedMeshComponent.h"
+#include "WorldPartition/HLOD/HLODBuilder.h"
+
+#include "Engine/SkinnedAsset.h"
+
+UHLODInstancedSkinnedMeshComponent::UHLODInstancedSkinnedMeshComponent(const FObjectInitializer& ObjectInitializer)
+	: Super(ObjectInitializer)
+{
+}
+
+#if WITH_EDITOR
+
+TUniquePtr<FSkinnedMeshComponentDescriptor> UHLODInstancedSkinnedMeshComponent::AllocateISMComponentDescriptor() const
+{
+	return MakeUnique<FHLODSkinnedMeshComponentDescriptor>();
+}
+
+FHLODSkinnedMeshComponentDescriptor::FHLODSkinnedMeshComponentDescriptor()
+{
+	ComponentClass = UHLODBuilder::GetInstancedSkinnedMeshComponentClass();
+}
+
+void FHLODSkinnedMeshComponentDescriptor::InitFrom(const UInstancedSkinnedMeshComponent* Component, bool bInitBodyInstance)
+{
+	Super::InitFrom(Component, bInitBodyInstance);
+
+	// Stationnary can be considered as static for the purpose of HLODs
+	if (Mobility == EComponentMobility::Stationary)
+	{
+		Mobility = EComponentMobility::Static;
+	}
+}
+
+
+#endif
diff --git a/Engine/Source/Runtime/Engine/Public/ISMPartition/ISMComponentBatcher.h b/Engine/Source/Runtime/Engine/Public/ISMPartition/ISMComponentBatcher.h
--- a/Engine/Source/Runtime/Engine/Public/ISMPartition/ISMComponentBatcher.h
+++ b/Engine/Source/Runtime/Engine/Public/ISMPartition/ISMComponentBatcher.h
@@ -7,6 +7,8 @@
 #include "Templates/EnableIf.h"
 #include "Templates/UnrealTypeTraits.h"
 
+class UInstancedStaticMeshComponent;
+class UInstancedSkinnedMeshComponent;
 
 /** Struct that allows batching of transforms and custom data of multiple (possibly instanced) static mesh components */
 struct FISMComponentBatcher
@@ -58,11 +60,16 @@ public:
 	int32 GetNumInstances() const { return NumInstances; }
 
 	/**
-	 * Initialize the instances of the provided ISM component using the batched data stored in this class.
+	 * Initialize the instances of the provided instanced static mesh component using the batched data stored in this class.
 	 * @param	ISMComponent	Instanced static mesh component which will be modified.
 	 */
 	ENGINE_API void InitComponent(UInstancedStaticMeshComponent* ISMComponent) const;
 
+	/**
+	 * Initialize the instances of the provided instanced skinned mesh component using the batched data stored in this class.
+	 * @param	ISMComponent	Instanced skinned mesh component which will be modified.
+	 */
+	ENGINE_API void InitComponent(UInstancedSkinnedMeshComponent* ISMComponent) const;
 	
 	inline uint32 GetHash() const
 	{
@@ -83,5 +90,7 @@ protected:
 
 	TArray<FTransform> InstancesTransformsWS;
 	TArray<float> InstancesCustomData;
-	TArray<FInstancedStaticMeshRandomSeed> RandomSeeds;
+
+	TArray<FInstancedStaticMeshRandomSeed> RandomSeeds; // For ISMC only
+	TArray<int32> BankIndices;							// For ISKMC only
 };
diff --git a/Engine/Source/Runtime/Engine/Public/WorldPartition/HLOD/HLODBuilder.h b/Engine/Source/Runtime/Engine/Public/WorldPartition/HLOD/HLODBuilder.h
--- a/Engine/Source/Runtime/Engine/Public/WorldPartition/HLOD/HLODBuilder.h
+++ b/Engine/Source/Runtime/Engine/Public/WorldPartition/HLOD/HLODBuilder.h
@@ -11,6 +11,7 @@
 class AActor;
 class UActorComponent;
 class UHLODInstancedStaticMeshComponent;
+class UHLODInstancedSkinnedMeshComponent;
 
 ENGINE_API DECLARE_LOG_CATEGORY_EXTERN(LogHLODBuilder, Log, All);
 
@@ -155,10 +156,15 @@ public:
 	static ENGINE_API uint32 ComputeHLODHash(const TArray<UActorComponent*>& InSourceComponents);
 
 	/** 
-	 * Get the InstancedStaticMeshComponent subclass that should be used when creating instanced HLODs.
+	 * Get the InstancedStaticMeshComponent subclass that should be used when creating instanced static mesh HLODs.
 	 */
 	static ENGINE_API TSubclassOf<UHLODInstancedStaticMeshComponent> GetInstancedStaticMeshComponentClass();
 
+	/**
+	 * Get the InstancedStaticMeshComponent subclass that should be used when creating instanced skinned mesh HLODs.
+	 */
+	static ENGINE_API TSubclassOf<UHLODInstancedSkinnedMeshComponent> GetInstancedSkinnedMeshComponentClass();
+
 protected:
 	virtual bool ShouldIgnoreBatchingPolicy() const { return false; }
 
@@ -186,6 +192,9 @@ protected:
 private:
 	UPROPERTY(Config)
 	TSubclassOf<UHLODInstancedStaticMeshComponent> HLODInstancedStaticMeshComponentClass;
+
+	UPROPERTY(Config)
+	TSubclassOf<UHLODInstancedSkinnedMeshComponent> HLODInstancedSkinnedMeshComponentClass;
 #endif
 };
 
diff --git a/Engine/Source/Runtime/Engine/Public/WorldPartition/HLOD/HLODHashBuilder.h b/Engine/Source/Runtime/Engine/Public/WorldPartition/HLOD/HLODHashBuilder.h
new file mode 100644
--- /dev/null
+++ b/Engine/Source/Runtime/Engine/Public/WorldPartition/HLOD/HLODHashBuilder.h
@@ -0,0 +1,109 @@
+// Copyright Epic Games, Inc. All Rights Reserved.
+
+#pragma once
+
+#if WITH_EDITOR
+
+#include "Math/MathFwd.h"
+#include "Serialization/ArchiveCrc32.h"
+
+class FHLODHashBuilder : public FArchiveCrc32
+{
+public:
+	/** Push a new context, logging with correct indentation */
+	void PushContext(const TCHAR* Context);
+
+	/** Pop context, ensuring indentation decreases */
+	void PopContext();
+
+	template <typename TParamType>
+	typename TEnableIf<TIsIntegral<TParamType>::Value, FArchive&>::Type operator<<(TParamType InValue)
+	{
+		FArchive& Self = *this;
+		return Self << InValue;
+	}
+
+	FArchive& operator<<(FTransform InTransform);
+	FArchive& operator<<(class UMaterialInterface* InMaterialInterface);
+	FArchive& operator<<(class UTexture* InTexture);
+	FArchive& operator<<(class UStaticMesh* InStaticMesh);
+	FArchive& operator<<(class USkinnedAsset* InSkinnedAsset);
+
+	//~ Begin FArchive Interface
+	virtual FArchive& operator<<(class UObject*& Object) override;
+	//~ End FArchive Interface
+
+	void LogContext(const TCHAR* Context, bool bOutputHash);
+
+	// For visibility of the overloads we don't override
+	using FArchiveCrc32::operator<<;
+
+private:
+	int32 IndentationLevel = 0; // Track indentation depth
+};
+
+// Templated operator overloads for handling arrays
+template <typename TElementType, typename TAllocatorType>
+FArchive& operator<<(FHLODHashBuilder& Ar, const TArray<TElementType, TAllocatorType>& InArray)
+{
+	TArray<TElementType, TAllocatorType>& ArrayMutable = const_cast<TArray<TElementType, TAllocatorType>&>(InArray);
+	FArchive& Self = Ar;
+	return Self << ArrayMutable;
+}
+
+template <typename TElementType, typename TAllocatorType>
+FArchive& operator<<(FHLODHashBuilder& Ar, TArray<TElementType, TAllocatorType>& InArray)
+{
+	FArchive& Self = Ar;
+	return Self << InArray;
+}
+
+class FHLODHashContext
+{
+public:
+	template <typename... TArgs>
+	explicit FHLODHashContext(UE::Core::TCheckedFormatString<FString::FmtCharType, TArgs...> Format, TArgs... Args)
+		: Context(FString::Printf(Format, Args...)) {}
+
+	explicit FHLODHashContext(FName InName)
+		: Context(InName.ToString()) {}
+
+	const TCHAR* GetContext() const { return *Context; }
+
+	/** Overload for FHLODHashContext to log structured properties */
+	friend FArchive& operator<<(FArchive& Ar, FHLODHashContext InContext)
+	{
+		static_cast<FHLODHashBuilder&>(Ar).LogContext(InContext.GetContext(), true);
+		return Ar;
+	}
+
+private:
+	FString Context;
+};
+
+class FHLODHashScope
+{
+public:
+	FHLODHashScope(FHLODHashBuilder& InBuilder, const TCHAR* Format)
+		: Builder(InBuilder)
+	{
+		Builder.PushContext(Format);
+	}
+
+	template <typename... Args>
+	FHLODHashScope(FHLODHashBuilder& InBuilder, const TCHAR* Format, Args... Arguments)
+		: Builder(InBuilder)
+	{
+		Builder.PushContext(*FString::Printf(Format, Arguments...));
+	}
+
+	~FHLODHashScope()
+	{
+		Builder.PopContext();
+	}
+
+private:
+	FHLODHashBuilder& Builder;
+};
+
+#endif
diff --git a/Engine/Source/Runtime/Engine/Public/WorldPartition/HLOD/HLODInstancedSkinnedMeshComponent.h b/Engine/Source/Runtime/Engine/Public/WorldPartition/HLOD/HLODInstancedSkinnedMeshComponent.h
new file mode 100644
--- /dev/null
+++ b/Engine/Source/Runtime/Engine/Public/WorldPartition/HLOD/HLODInstancedSkinnedMeshComponent.h
@@ -0,0 +1,34 @@
+// Copyright Epic Games, Inc. All Rights Reserved.
+
+#pragma once
+
+#include "CoreMinimal.h"
+#include "Components/InstancedSkinnedMeshComponent.h"
+#include "Animation/AnimBank.h"
+#include "HLODInstancedSkinnedMeshComponent.generated.h"
+
+
+UCLASS(HideDropDown, NotPlaceable, MinimalAPI)
+class UHLODInstancedSkinnedMeshComponent : public UInstancedSkinnedMeshComponent
+{
+	GENERATED_UCLASS_BODY()
+
+public:
+#if WITH_EDITOR
+	ENGINE_API virtual TUniquePtr<FSkinnedMeshComponentDescriptor> AllocateISMComponentDescriptor() const;
+#endif
+};
+
+
+// ISM descriptor class based on FSkinnedMeshComponentDescriptor
+USTRUCT()
+struct FHLODSkinnedMeshComponentDescriptor : public FSkinnedMeshComponentDescriptor
+{
+	GENERATED_BODY()
+
+#if WITH_EDITOR
+	ENGINE_API FHLODSkinnedMeshComponentDescriptor();
+
+	ENGINE_API virtual void InitFrom(const UInstancedSkinnedMeshComponent* Component, bool bInitBodyInstance = true) override;
+#endif
+};
