Change FISMComponentDescriptor::StaticMeshComponent into a Transient TWeakObjectPtr instead, to ensure the reference can be GC'd

diff --git a/Engine/Source/Runtime/Engine/Public/ISMPartition/ISMComponentDescriptor.h b/Engine/Source/Runtime/Engine/Public/ISMPartition/ISMComponentDescriptor.h
--- a/Engine/Source/Runtime/Engine/Public/ISMPartition/ISMComponentDescriptor.h
+++ b/Engine/Source/Runtime/Engine/Public/ISMPartition/ISMComponentDescriptor.h
@@ -248,6 +248,11 @@ struct FISMComponentDescriptor : public FISMComponentDescriptorBase
 	ENGINE_API virtual void InitFrom(const UStaticMeshComponent* Component, bool bInitBodyInstance = true) override;
 	ENGINE_API virtual uint32 ComputeHash() const;
 	ENGINE_API virtual void InitComponent(UInstancedStaticMeshComponent* ISMComponent) const override;
+
+	ENGINE_API TWeakObjectPtr<const UStaticMeshComponent> GetStaticMeshComponent() const
+	{
+		return StaticMeshComponent;
+	}
 		
 	ENGINE_API bool operator!=(const FISMComponentDescriptor& Other) const;
 	ENGINE_API bool operator==(const FISMComponentDescriptor& Other) const;
@@ -261,9 +266,6 @@ public:
 	UPROPERTY(EditAnywhere, Category = "Component Settings", meta = (DisplayAfter = "ComponentClass"))
 	TObjectPtr<UStaticMesh> StaticMesh = nullptr;
 
-	UPROPERTY(EditAnywhere, Category = "Component Settings", meta = (DisplayAfter = "ComponentClass"))
-	TObjectPtr<const UStaticMeshComponent> StaticMeshComponent = nullptr;
-
 	UPROPERTY(EditAnywhere, Category = "Component Settings", meta = (DisplayAfter = "ComponentClass"))
 	TArray<TObjectPtr<UMaterialInterface>> OverrideMaterials;
 
@@ -275,6 +277,10 @@ public:
 
 	UPROPERTY(EditAnywhere, Category = "Component Settings", meta = (DisplayAfter = "ComponentClass"))
 	FTransform LocalTransform = FTransform::Identity;
+
+private:
+	UPROPERTY(Transient)
+	TWeakObjectPtr<const UStaticMeshComponent> StaticMeshComponent = nullptr;
 };
 
 USTRUCT()
