Cherry picking engine change from 35.00 to 34.20.
Added a delegate to be able to react when a montage section changes

#rnx

diff --git a/Engine/Source/Runtime/Engine/Private/Animation/AnimInstance.cpp b/Engine/Source/Runtime/Engine/Private/Animation/AnimInstance.cpp
--- a/Engine/Source/Runtime/Engine/Private/Animation/AnimInstance.cpp
+++ b/Engine/Source/Runtime/Engine/Private/Animation/AnimInstance.cpp
@@ -810,7 +810,7 @@ void UAnimInstance::ClearQueuedAnimEvents(bool bShouldUpdateActiveAnimNotifiesSi
 	
 	if (bShouldUpdateActiveAnimNotifiesSinceLastTick)
 	{
-		// Before clearing notify queue we need to ensure we record the notifies that where triggered last tick.
+		// Before clearing notify queue we need to ensure we record the notifies that were triggered last tick.
 		Proxy.UpdateActiveAnimNotifiesSinceLastTick(NotifyQueue);
 	}
 	else
