Mobile: Fixed case where non DefaultLit shading models were not applied if project has r.Mobile.AllowPerPixelShadingModels=0

diff --git a/Engine/Source/Runtime/Engine/Private/Materials/Material.cpp b/Engine/Source/Runtime/Engine/Private/Materials/Material.cpp
--- a/Engine/Source/Runtime/Engine/Private/Materials/Material.cpp
+++ b/Engine/Source/Runtime/Engine/Private/Materials/Material.cpp
@@ -251,7 +251,8 @@ int32 FMaterialResource::CompilePropertyAndSetMaterialProperty(EMaterialProperty
 			}
 			break;
 		case MP_ShadingModel:
-			if (AllowPerPixelShadingModels(Compiler->GetShaderPlatform()))
+			if (AllowPerPixelShadingModels(Compiler->GetShaderPlatform()) 
+				|| !MaterialInterface->IsShadingModelFromMaterialExpression())
 			{
 				Ret = MaterialInterface->CompileProperty(Compiler, Property);
 			}
