Fix link

diff --git a/Engine/Source/Runtime/Engine/Public/Materials/MaterialInterface.h b/Engine/Source/Runtime/Engine/Public/Materials/MaterialInterface.h
--- a/Engine/Source/Runtime/Engine/Public/Materials/MaterialInterface.h
+++ b/Engine/Source/Runtime/Engine/Public/Materials/MaterialInterface.h
@@ -305,7 +305,7 @@ public:
 	{
 		return CachedTexturesSamplingInfo.GetPtrOrNull();
 	}
-	virtual FTextureSamplingInfo CalculateTexturesSamplingInfo(UTexture* Texture);
+	ENGINE_API virtual FTextureSamplingInfo CalculateTexturesSamplingInfo(UTexture* Texture);
 #endif
 
 	/** SubsurfaceProfile, for Screen Space Subsurface Scattering.. */
