Fix uninitialized member FCellDesc.bIs2D for non-spatially loaded cells, causing non-determinism in streaming generation tests under Clang.

diff --git a/Engine/Source/Runtime/Engine/Private/WorldPartition/RuntimeHashSet/RuntimePartitionLHGrid.cpp b/Engine/Source/Runtime/Engine/Private/WorldPartition/RuntimeHashSet/RuntimePartitionLHGrid.cpp
--- a/Engine/Source/Runtime/Engine/Private/WorldPartition/RuntimeHashSet/RuntimePartitionLHGrid.cpp
+++ b/Engine/Source/Runtime/Engine/Private/WorldPartition/RuntimeHashSet/RuntimePartitionLHGrid.cpp
@@ -231,9 +231,9 @@ bool URuntimePartitionLHGrid::GenerateStreaming(const FGenerateStreamingParams&
 				CellDesc.CellBounds.GetValue().Min.Z = -HALF_WORLD_MAX;
 				CellDesc.CellBounds.GetValue().Max.Z = HALF_WORLD_MAX;
 			}
-
-			CellDesc.bIs2D = bIs2D;
 		}
+
+		CellDesc.bIs2D = bIsSpatiallyLoaded && bIs2D;
 	}
 
 	return true;
